Added convenience functions in data collections to simplify blueprint implementation
Added Get Context call in blueprint to be able to remove some context parameters in future CLs.

diff --git a/Engine/Plugins/Experimental/PCG/Config/BasePCG.ini b/Engine/Plugins/Experimental/PCG/Config/BasePCG.ini
--- a/Engine/Plugins/Experimental/PCG/Config/BasePCG.ini
+++ b/Engine/Plugins/Experimental/PCG/Config/BasePCG.ini
@@ -5,6 +5,8 @@
 +FunctionRedirects=(OldName="PCGBlueprintElement.MultiLoopOnPoints",NewName="PCGBlueprintElement.VariableLoop")
 +FunctionRedirects=(OldName="PCGBlueprintElement.LoopOnPointPairs",NewName="PCGBlueprintElement.NestedLoop")
 +FunctionRedirects=(OldName="PCGBlueprintElement.LoopNTimes",NewName="PCGBlueprintElement.IterationLoop")
++FunctionRedirects=(OldName="PCGBlueprintElement.InputLabels",NewName="PCGBlueprintElement.CustomInputLabels")
++FunctionRedirects=(OldName="PCGBlueprintElement.OutputLabels",NewName="PCGBlueprintElement.CustomOutputLabels")
 +PropertyRedirects=(OldName="PCGSplineSamplerSettings.Params",NewName="PCGSplineSamplerSettings.SamplerParams")
 +PropertyRedirects=(OldName="PCGProjectionSettings.Params",NewName="PCGProjectionSettings.ProjectionParams")
 +PropertyRedirects=(OldName="PCGBlueprintElement.IterationLoop.InA",NewName="PCGBlueprintElement.IterationLoop.OptionalA")
diff --git a/Engine/Plugins/Experimental/PCG/Source/PCG/Private/Elements/PCGExecuteBlueprint.cpp b/Engine/Plugins/Experimental/PCG/Source/PCG/Private/Elements/PCGExecuteBlueprint.cpp
--- a/Engine/Plugins/Experimental/PCG/Source/PCG/Private/Elements/PCGExecuteBlueprint.cpp
+++ b/Engine/Plugins/Experimental/PCG/Source/PCG/Private/Elements/PCGExecuteBlueprint.cpp
@@ -103,6 +103,18 @@ void UPCGBlueprintElement::Initialize()
 #endif
 }
 
+FPCGContext& UPCGBlueprintElement::GetContext() const
+{
+	checkf(CurrentContext, TEXT("Execution context is not ready - do not call the GetContext method inside of non-execution methods"));
+	return *CurrentContext;
+}
+
+void UPCGBlueprintElement::SetCurrentContext(FPCGContext* InCurrentContext)
+{
+	ensure(CurrentContext == nullptr || InCurrentContext == nullptr || CurrentContext == InCurrentContext);
+	CurrentContext = InCurrentContext;
+}
+
 #if WITH_EDITOR
 void UPCGBlueprintElement::PostEditChangeProperty(FPropertyChangedEvent& PropertyChangedEvent)
 {
@@ -156,7 +168,7 @@ EPCGSettingsType UPCGBlueprintElement::NodeTypeOverride_Implementation() const
 	return EPCGSettingsType::Blueprint;
 }
 
-TSet<FName> UPCGBlueprintElement::InputLabels() const
+TSet<FName> UPCGBlueprintElement::CustomInputLabels() const
 {
 	TSet<FName> Labels;
 	for (const FPCGPinProperties& PinProperty : CustomInputPins)
@@ -167,7 +179,7 @@ TSet<FName> UPCGBlueprintElement::InputLabels() const
 	return Labels;
 }
 
-TSet<FName> UPCGBlueprintElement::OutputLabels() const
+TSet<FName> UPCGBlueprintElement::CustomOutputLabels() const
 {
 	TSet<FName> Labels;
 	for (const FPCGPinProperties& PinProperty : CustomOutputPins)
@@ -178,6 +190,42 @@ TSet<FName> UPCGBlueprintElement::OutputLabels() const
 	return Labels;
 }
 
+TArray<FPCGPinProperties> UPCGBlueprintElement::GetInputPins() const
+{
+	if (CurrentContext)
+	{
+		if (const UPCGBlueprintSettings* Settings = CurrentContext->GetInputSettings<UPCGBlueprintSettings>())
+		{
+			return Settings->InputPinProperties();
+		}
+	}
+	else if (const UPCGBlueprintSettings* OriginalSettings = Cast<UPCGBlueprintSettings>(GetOuter()))
+	{
+		return OriginalSettings->InputPinProperties();
+	}
+	
+	// Can't retrieve settings - return only custom pins then
+	return CustomInputPins;
+}
+
+TArray<FPCGPinProperties> UPCGBlueprintElement::GetOutputPins() const
+{
+	if (CurrentContext)
+	{
+		if (const UPCGBlueprintSettings* Settings = CurrentContext->GetInputSettings<UPCGBlueprintSettings>())
+		{
+			return Settings->OutputPinProperties();
+		}
+	}
+	else if (const UPCGBlueprintSettings* OriginalSettings = Cast<UPCGBlueprintSettings>(GetOuter()))
+	{
+		return OriginalSettings->OutputPinProperties();
+	}
+
+	// Can't retrieve settings - return only custom pins then
+	return CustomOutputPins;
+}
+
 int UPCGBlueprintElement::GetSeed(FPCGContext& InContext) const 
 {
 	return InContext.GetSeed();
@@ -604,7 +652,9 @@ bool FPCGExecuteBlueprintElement::ExecuteInternal(FPCGContext* InContext) const
 		Context->bShouldUnrootSettingsOnDelete = false;
 
 		/** Finally, execute the actual blueprint */
+		Context->BlueprintElementInstance->SetCurrentContext(Context);
 		Context->BlueprintElementInstance->ExecuteWithContext(*Context, Context->InputData, Context->OutputData);
+		Context->BlueprintElementInstance->SetCurrentContext(nullptr);
 
 		// Put back the proper unroot flag
 		Context->bShouldUnrootSettingsOnDelete = bShouldUnrootSettingsOnDelete;
diff --git a/Engine/Plugins/Experimental/PCG/Source/PCG/Private/PCGData.cpp b/Engine/Plugins/Experimental/PCG/Source/PCG/Private/PCGData.cpp
--- a/Engine/Plugins/Experimental/PCG/Source/PCG/Private/PCGData.cpp
+++ b/Engine/Plugins/Experimental/PCG/Source/PCG/Private/PCGData.cpp
@@ -444,6 +444,63 @@ void FPCGDataCollection::Reset()
 	bCancelExecution = false;
 }
 
+TArray<UPCGData*> UPCGDataFunctionLibrary::GetInputsByPredicate(const FPCGDataCollection& InCollection, TArray<FPCGTaggedData>& OutTaggedData, TFunctionRef<bool(const FPCGTaggedData&)> InPredicate)
+{
+	TArray<UPCGData*> Inputs;
+	OutTaggedData.Reset();
+
+	for (const FPCGTaggedData& TaggedData : InCollection.TaggedData)
+	{
+		if (InPredicate(TaggedData))
+		{
+			Inputs.Add(const_cast<UPCGData*>(TaggedData.Data.Get()));
+			OutTaggedData.Add(TaggedData);
+		}
+	}
+
+	return Inputs;
+}
+
+TArray<UPCGData*> UPCGDataFunctionLibrary::GetTypedInputs(const FPCGDataCollection& InCollection, TArray<FPCGTaggedData>& OutTaggedData, TSubclassOf<UPCGData> InDataTypeClass)
+{
+	return GetInputsByPredicate(InCollection, OutTaggedData, [InDataTypeClass](const FPCGTaggedData& TaggedData)
+	{
+		return TaggedData.Data && (!InDataTypeClass || TaggedData.Data->IsA(InDataTypeClass));
+	});
+}
+
+TArray<UPCGData*> UPCGDataFunctionLibrary::GetTypedInputsByPin(const FPCGDataCollection& InCollection, const FPCGPinProperties& InPin, TArray<FPCGTaggedData>& OutTaggedData, TSubclassOf<UPCGData> InDataTypeClass)
+{
+	return GetTypedInputsByPinLabel(InCollection, InPin.Label, OutTaggedData, InDataTypeClass);
+}
+
+TArray<UPCGData*> UPCGDataFunctionLibrary::GetTypedInputsByPinLabel(const FPCGDataCollection& InCollection, const FName& InPinLabel, TArray<FPCGTaggedData>& OutTaggedData, TSubclassOf<UPCGData> InDataTypeClass)
+{
+	return GetInputsByPredicate(InCollection, OutTaggedData, [&InPinLabel, InDataTypeClass](const FPCGTaggedData& TaggedData)
+	{
+		return TaggedData.Pin == InPinLabel && TaggedData.Data && (!InDataTypeClass || TaggedData.Data->IsA(InDataTypeClass));
+	});
+}
+
+TArray<UPCGData*> UPCGDataFunctionLibrary::GetTypedInputsByTag(const FPCGDataCollection& InCollection, const FString& InTag, TArray<FPCGTaggedData>& OutTaggedData, TSubclassOf<UPCGData> InDataTypeClass)
+{
+	return GetInputsByPredicate(InCollection, OutTaggedData, [&InTag, InDataTypeClass](const FPCGTaggedData& TaggedData)
+	{
+		return TaggedData.Tags.Contains(InTag) && TaggedData.Data && (!InDataTypeClass || TaggedData.Data->IsA(InDataTypeClass));
+	});
+}
+
+void UPCGDataFunctionLibrary::AddToCollection(FPCGDataCollection& InCollection, const UPCGData* InData, FName InPinLabel, TArray<FString> InTags)
+{
+	if (InData)
+	{
+		FPCGTaggedData& TaggedData = InCollection.TaggedData.Emplace_GetRef();
+		TaggedData.Data = InData;
+		TaggedData.Pin = InPinLabel;
+		TaggedData.Tags.Append(InTags);
+	}
+}
+
 TArray<FPCGTaggedData> UPCGDataFunctionLibrary::GetInputs(const FPCGDataCollection& InCollection)
 {
 	return InCollection.GetInputs();
diff --git a/Engine/Plugins/Experimental/PCG/Source/PCG/Public/Elements/PCGExecuteBlueprint.h b/Engine/Plugins/Experimental/PCG/Source/PCG/Public/Elements/PCGExecuteBlueprint.h
--- a/Engine/Plugins/Experimental/PCG/Source/PCG/Public/Elements/PCGExecuteBlueprint.h
+++ b/Engine/Plugins/Experimental/PCG/Source/PCG/Public/Elements/PCGExecuteBlueprint.h
@@ -84,11 +84,19 @@ public:
 	UFUNCTION(BlueprintImplementableEvent, Category = "PCG|Preconfigure Settings", meta = (ForceAsFunction))
 	void ApplyPreconfiguredSettings(UPARAM(ref) const FPCGPreConfiguredSettingsInfo& InPreconfigureInfo);
 
+	// Returns the labels of custom input pins only
 	UFUNCTION(BlueprintCallable, Category = "PCG|Input & Output")
-	TSet<FName> InputLabels() const;
+	TSet<FName> CustomInputLabels() const;
 
+	// Returns the labels of custom output pins only
 	UFUNCTION(BlueprintCallable, Category = "PCG|Input & Output")
-	TSet<FName> OutputLabels() const;
+	TSet<FName> CustomOutputLabels() const;
+
+	UFUNCTION(BlueprintCallable, Category = "PCG|Input & Output", meta = (HideSelfPin = "true"))
+	TArray<FPCGPinProperties> GetInputPins() const;
+
+	UFUNCTION(BlueprintCallable, Category = "PCG|Input & Output", meta = (HideSelfPin = "true"))
+	TArray<FPCGPinProperties> GetOutputPins() const;
 
 	/** Gets the seed from the associated settings & source component */
 	UFUNCTION(BlueprintCallable, Category = "PCG|Random")
@@ -101,6 +109,13 @@ public:
 	/** Called after object creation to setup the object callbacks */
 	void Initialize();
 
+	/** Retrieves the execution context - note that this will not be valid outside of the Execute functions */
+	UFUNCTION(BlueprintCallable, Category = "PCG|Advanced", meta = (HideSelfPin = "true"))
+	FPCGContext& GetContext() const;
+
+	/** Called after the element duplication during execution to be able to get the context easily - internal call only */
+	void SetCurrentContext(FPCGContext* InCurrentContext);
+
 #if WITH_EDITOR
 	// ~Begin UObject interface
 	virtual void PostEditChangeProperty(FPropertyChangedEvent& PropertyChangedEvent) override;
@@ -132,10 +147,10 @@ public:
 	bool bCanBeMultithreaded = false;
 
 #if WITH_EDITORONLY_DATA
-	UPROPERTY(BlueprintGetter=InputLabels, Category = "Settings|Input & Output", meta = (DeprecatedProperty, DeprecatedMessage = "Input Pin Labels are deprecated - use Input Labels instead."))
+	UPROPERTY(BlueprintGetter=CustomInputLabels, Category = "Settings|Input & Output", meta = (DeprecatedProperty, DeprecatedMessage = "Input Pin Labels are deprecated - use Input Labels instead."))
 	TSet<FName> InputPinLabels_DEPRECATED;
 
-	UPROPERTY(BlueprintGetter=OutputLabels, Category = "Settings|Input & Output")
+	UPROPERTY(BlueprintGetter=CustomOutputLabels, Category = "Settings|Input & Output")
 	TSet<FName> OutputPinLabels_DEPRECATED;
 #endif
 
@@ -183,6 +198,10 @@ protected:
 #if !WITH_EDITORONLY_DATA
 	UWorld* InstanceWorld = nullptr;
 #endif
+
+	// Since we duplicate the blueprint elements prior to execution, they will be unique
+	// and have a 1:1 match with their context, which allows us to store it here
+	FPCGContext* CurrentContext = nullptr;
 };
 
 UCLASS(BlueprintType, ClassGroup = (Procedural))
diff --git a/Engine/Plugins/Experimental/PCG/Source/PCG/Public/PCGData.h b/Engine/Plugins/Experimental/PCG/Source/PCG/Public/PCGData.h
--- a/Engine/Plugins/Experimental/PCG/Source/PCG/Public/PCGData.h
+++ b/Engine/Plugins/Experimental/PCG/Source/PCG/Public/PCGData.h
@@ -6,6 +6,7 @@
 #include "PCGCrc.h"
 
 #include "Kismet/BlueprintFunctionLibrary.h"
+#include "Templates/SubclassOf.h"
 
 #include "PCGData.generated.h"
 
@@ -191,6 +192,26 @@ class PCG_API UPCGDataFunctionLibrary : public UBlueprintFunctionLibrary
 	GENERATED_BODY()
 
 public:
+	/** Gets all inputs of the given class type, returning matching tagged data in the OutTaggedData value too */
+	UFUNCTION(BlueprintCallable, Category = Data, meta = (ScriptMethod, DeterminesOutputType = "InDataTypeClass"))
+	static TArray<UPCGData*> GetTypedInputs(const FPCGDataCollection& InCollection, TArray<FPCGTaggedData>& OutTaggedData, TSubclassOf<UPCGData> InDataTypeClass = nullptr);
+
+	/** Gets all inputs of the given class type and on the given pin, returning matching tagged data in the OutTaggedData value too */
+	UFUNCTION(BlueprintCallable, Category = Data, meta = (ScriptMethod, DeterminesOutputType = "InDataTypeClass"))
+	static TArray<UPCGData*> GetTypedInputsByPin(const FPCGDataCollection& InCollection, const FPCGPinProperties& InPin, TArray<FPCGTaggedData>& OutTaggedData, TSubclassOf<UPCGData> InDataTypeClass = nullptr);
+
+	/** Gets all inputs of the given class type and on the given pin label, returning matching tagged data in the OutTaggedData value too */
+	UFUNCTION(BlueprintCallable, Category = Data, meta = (ScriptMethod, DeterminesOutputType = "InDataTypeClass"))
+	static TArray<UPCGData*> GetTypedInputsByPinLabel(const FPCGDataCollection& InCollection, const FName& InPinLabel, TArray<FPCGTaggedData>& OutTaggedData, TSubclassOf<UPCGData> InDataTypeClass = nullptr);
+
+	/** Gets all inputs of the given class type and having the provided tag, returning matching tagged data in the OutTaggedData value too */
+	UFUNCTION(BlueprintCallable, Category = Data, meta = (ScriptMethod, DeterminesOutputType = "InDataTypeClass"))
+	static TArray<UPCGData*> GetTypedInputsByTag(const FPCGDataCollection& InCollection, const FString& InTag, TArray<FPCGTaggedData>& OutTaggedData, TSubclassOf<UPCGData> InDataTypeClass = nullptr);
+
+	/** Adds a data object to a given collection, simpler usage than making a PCGTaggedData object */
+	UFUNCTION(BlueprintCallable, Category = Data, meta = (ScriptMethod))
+	static void AddToCollection(UPARAM(ref) FPCGDataCollection& InCollection, const UPCGData* InData, FName InPinLabel, TArray<FString> InTags);
+
 	// Blueprint methods to support interaction with FPCGDataCollection
 	UFUNCTION(BlueprintCallable, Category = Data)
 	static TArray<FPCGTaggedData> GetInputs(const FPCGDataCollection& InCollection);
@@ -212,6 +233,9 @@ public:
 
 	UFUNCTION(BlueprintCallable, Category = Data)
 	static TArray<FPCGTaggedData> GetAllSettings(const FPCGDataCollection& InCollection);
+
+protected:
+	static TArray<UPCGData*> GetInputsByPredicate(const FPCGDataCollection& InCollection, TArray<FPCGTaggedData>& OutTaggedData, TFunctionRef<bool(const FPCGTaggedData&)> InPredicate);
 };
 
 #if UE_ENABLE_INCLUDE_ORDER_DEPRECATED_IN_5_2
