Improved the behavior of actor selector struct when not using class selection by losing the reference
Added additional information on the Get ... Data nodes regarding to source of data

diff --git a/Engine/Plugins/Experimental/PCG/Source/PCG/Private/Elements/PCGActorSelector.cpp b/Engine/Plugins/Experimental/PCG/Source/PCG/Private/Elements/PCGActorSelector.cpp
--- a/Engine/Plugins/Experimental/PCG/Source/PCG/Private/Elements/PCGActorSelector.cpp
+++ b/Engine/Plugins/Experimental/PCG/Source/PCG/Private/Elements/PCGActorSelector.cpp
@@ -202,3 +202,52 @@ namespace PCGActorSelector
 		return Actors.IsEmpty() ? nullptr : Actors[0];
 	}
 }
+
+void FPCGActorSelectorSettings::PostLoad()
+{
+	// Make sure that when we are not using the 'by class' selection, we don't track that class dependency
+	if (ActorSelection != EPCGActorSelection::ByClass)
+	{
+		ActorSelectionClass = TSubclassOf<AActor>();
+	}
+}
+
+#if WITH_EDITOR
+void FPCGActorSelectorSettings::PostEditChangeProperty(FPropertyChangedEvent& PropertyChangedEvent)
+{
+	if (PropertyChangedEvent.GetPropertyName() == GET_MEMBER_NAME_CHECKED(FPCGActorSelectorSettings, ActorSelection))
+	{
+		// Make sure that when switching away from the 'by class' selection, we actually break that data dependency
+		if (ActorSelection != EPCGActorSelection::ByClass)
+		{
+			ActorSelectionClass = TSubclassOf<AActor>();
+		}
+	}
+}
+
+FText FPCGActorSelectorSettings::GetTaskNameSuffix() const
+{
+	if (ActorFilter == EPCGActorFilter::AllWorldActors)
+	{
+		if (ActorSelection == EPCGActorSelection::ByClass)
+		{
+			return (ActorSelectionClass.Get() ? ActorSelectionClass->GetDisplayNameText() : FText::FromName(NAME_None));
+		}
+		else if (ActorSelection == EPCGActorSelection::ByTag)
+		{
+			return FText::FromName(ActorSelectionTag);
+		}
+	}
+	else if(const UEnum* EnumPtr = StaticEnum<EPCGActorFilter>())
+	{
+		return EnumPtr->GetDisplayNameTextByValue(static_cast<__underlying_type(EPCGActorFilter)>(ActorFilter));
+	}
+
+	return FText();
+}
+
+FName FPCGActorSelectorSettings::GetTaskName(const FText& Prefix) const
+{
+	return FName(FText::Format(NSLOCTEXT("PCGActorSelectorSettings", "NodeTitleFormat", "{0} ({1})"), Prefix, GetTaskNameSuffix()).ToString());
+}
+#endif // WITH_EDITOR
\ No newline at end of file
diff --git a/Engine/Plugins/Experimental/PCG/Source/PCG/Private/Elements/PCGDataFromActor.cpp b/Engine/Plugins/Experimental/PCG/Source/PCG/Private/Elements/PCGDataFromActor.cpp
--- a/Engine/Plugins/Experimental/PCG/Source/PCG/Private/Elements/PCGDataFromActor.cpp
+++ b/Engine/Plugins/Experimental/PCG/Source/PCG/Private/Elements/PCGDataFromActor.cpp
@@ -28,7 +28,32 @@ FText UPCGDataFromActorSettings::GetNodeTooltipText() const
 {
 	return LOCTEXT("DataFromActorTooltip", "Builds a collection of PCG-compatible data from the selected actors.");
 }
+
+void UPCGDataFromActorSettings::PostEditChangeProperty(FPropertyChangedEvent& PropertyChangedEvent)
+{
+	Super::PostEditChangeProperty(PropertyChangedEvent);
+
+	if (PropertyChangedEvent.GetMemberPropertyName() == GET_MEMBER_NAME_CHECKED(UPCGDataFromActorSettings, ActorSelector))
+	{
+		ActorSelector.PostEditChangeProperty(PropertyChangedEvent);
+	}
+}
+#endif
+
+void UPCGDataFromActorSettings::PostLoad()
+{
+	Super::PostLoad();
+	ActorSelector.PostLoad();
+}
+
+FName UPCGDataFromActorSettings::AdditionalTaskName() const
+{
+#if WITH_EDITOR
+	return ActorSelector.GetTaskName(GetDefaultNodeTitle());
+#else
+	return Super::AdditionalTaskName();
 #endif
+}
 
 FPCGElementPtr UPCGDataFromActorSettings::CreateElement() const
 {
diff --git a/Engine/Plugins/Experimental/PCG/Source/PCG/Private/Elements/PCGPropertyToParamData.cpp b/Engine/Plugins/Experimental/PCG/Source/PCG/Private/Elements/PCGPropertyToParamData.cpp
--- a/Engine/Plugins/Experimental/PCG/Source/PCG/Private/Elements/PCGPropertyToParamData.cpp
+++ b/Engine/Plugins/Experimental/PCG/Source/PCG/Private/Elements/PCGPropertyToParamData.cpp
@@ -22,7 +22,16 @@ void UPCGPropertyToParamDataSettings::GetTrackedActorTags(FPCGTagToSettingsMap&
 	{
 		OutTagToSettings.FindOrAdd(ActorSelector.ActorSelectionTag).Emplace({ this, bTrackActorsOnlyWithinBounds });
 	}
+}
+
+void UPCGPropertyToParamDataSettings::PostEditChangeProperty(FPropertyChangedEvent& PropertyChangedEvent)
+{
+	Super::PostEditChangeProperty(PropertyChangedEvent);
 
+	if (PropertyChangedEvent.GetMemberPropertyName() == GET_MEMBER_NAME_CHECKED(UPCGPropertyToParamDataSettings, ActorSelector))
+	{
+		ActorSelector.PostEditChangeProperty(PropertyChangedEvent);
+	}
 }
 #endif // WITH_EDITOR
 
@@ -51,8 +60,20 @@ void UPCGPropertyToParamDataSettings::PostLoad()
 		ActorFilter_DEPRECATED = EPCGActorFilter::Self;
 		bIncludeChildren_DEPRECATED = false;
 	}
+
+	ActorSelector.PostLoad();
 }
 
+FName UPCGPropertyToParamDataSettings::AdditionalTaskName() const
+{
+#if WITH_EDITOR
+	return ActorSelector.GetTaskName(GetDefaultNodeTitle());
+#else
+	return Super::AdditionalTaskName();
+#endif
+}
+
+
 TArray<FPCGPinProperties> UPCGPropertyToParamDataSettings::OutputPinProperties() const
 {
 	TArray<FPCGPinProperties> PinProperties;
diff --git a/Engine/Plugins/Experimental/PCG/Source/PCG/Public/Elements/PCGActorSelector.h b/Engine/Plugins/Experimental/PCG/Source/PCG/Public/Elements/PCGActorSelector.h
--- a/Engine/Plugins/Experimental/PCG/Source/PCG/Public/Elements/PCGActorSelector.h
+++ b/Engine/Plugins/Experimental/PCG/Source/PCG/Public/Elements/PCGActorSelector.h
@@ -66,6 +66,14 @@ struct FPCGActorSelectorSettings
 	/** If true processes all matching actors, otherwise returns data from first match. */
 	UPROPERTY(BlueprintReadWrite, EditAnywhere, Category = Settings, meta = (EditCondition = "ActorFilter==EPCGActorFilter::AllWorldActors && ActorSelection!=EPCGActorSelection::ByName", EditConditionHides))
 	bool bSelectMultiple = false;
+
+	void PostLoad();
+
+#if WITH_EDITOR	
+	void PostEditChangeProperty(FPropertyChangedEvent& PropertyChangedEvent);
+	FText GetTaskNameSuffix() const;
+	FName GetTaskName(const FText& Prefix) const;
+#endif
 };
 
 namespace PCGActorSelector
diff --git a/Engine/Plugins/Experimental/PCG/Source/PCG/Public/Elements/PCGDataFromActor.h b/Engine/Plugins/Experimental/PCG/Source/PCG/Public/Elements/PCGDataFromActor.h
--- a/Engine/Plugins/Experimental/PCG/Source/PCG/Public/Elements/PCGDataFromActor.h
+++ b/Engine/Plugins/Experimental/PCG/Source/PCG/Public/Elements/PCGDataFromActor.h
@@ -37,6 +37,8 @@ public:
 #endif
 	virtual EPCGDataType GetCurrentPinTypes(const UPCGPin* InPin) const override;
 
+	virtual FName AdditionalTaskName() const override;
+
 	virtual TArray<FPCGPinProperties> InputPinProperties() const override { return TArray<FPCGPinProperties>(); }
 	virtual TArray<FPCGPinProperties> OutputPinProperties() const override;
 
@@ -44,6 +46,14 @@ protected:
 	virtual FPCGElementPtr CreateElement() const override;
 	//~End UPCGSettings
 
+public:
+	//~Begin UObject interface
+	virtual void PostLoad() override;
+#if WITH_EDITOR
+	virtual void PostEditChangeProperty(FPropertyChangedEvent& PropertyChangedEvent) override;
+#endif
+	//~End UObject interface
+
 public:
 	/** Override this to filter what kinds of data should be retrieved from the actor(s). */
 	virtual bool DataFilter(EPCGDataType InDataType) const { return true; }
diff --git a/Engine/Plugins/Experimental/PCG/Source/PCG/Public/Elements/PCGPropertyToParamData.h b/Engine/Plugins/Experimental/PCG/Source/PCG/Public/Elements/PCGPropertyToParamData.h
--- a/Engine/Plugins/Experimental/PCG/Source/PCG/Public/Elements/PCGPropertyToParamData.h
+++ b/Engine/Plugins/Experimental/PCG/Source/PCG/Public/Elements/PCGPropertyToParamData.h
@@ -21,6 +21,9 @@ class PCG_API UPCGPropertyToParamDataSettings : public UPCGSettings
 public:
 	//~Begin UObject interface
 	virtual void PostLoad() override;
+#if WITH_EDITOR
+	virtual void PostEditChangeProperty(FPropertyChangedEvent& PropertyChangedEvent) override;
+#endif
 	//~End UObject interface
 
 	//~Begin UPCGSettings interface
@@ -31,6 +34,7 @@ public:
 	virtual void GetTrackedActorTags(FPCGTagToSettingsMap& OutTagToSettings, TArray<TObjectPtr<const UPCGGraph>>& OutVisitedGraphs) const override;
 #endif
 
+	virtual FName AdditionalTaskName() const override;
 	virtual TArray<FPCGPinProperties> OutputPinProperties() const override;
 	virtual TArray<FPCGPinProperties> InputPinProperties() const override { return TArray<FPCGPinProperties>(); }
 
