Added sanity check point data node

diff --git a/Engine/Plugins/Experimental/PCG/Source/PCG/Private/Elements/PCGSanityCheckPointData.cpp b/Engine/Plugins/Experimental/PCG/Source/PCG/Private/Elements/PCGSanityCheckPointData.cpp
new file mode 100644
--- /dev/null
+++ b/Engine/Plugins/Experimental/PCG/Source/PCG/Private/Elements/PCGSanityCheckPointData.cpp
@@ -0,0 +1,87 @@
+// Copyright Epic Games, Inc. All Rights Reserved.
+
+#include "Elements/PCGSanityCheckPointData.h"
+
+#include "PCGComponent.h"
+#include "PCGContext.h"
+#include "PCGPin.h"
+
+#include "Data/PCGPointData.h"
+#include "Data/PCGSpatialData.h"
+
+#define LOCTEXT_NAMESPACE "PCGSanityCheckPointDataElement"
+
+	
+TArray<FPCGPinProperties> UPCGSanityCheckPointDataSettings::InputPinProperties() const
+{
+	TArray<FPCGPinProperties> PinProperties;
+	PinProperties.Emplace(PCGPinConstants::DefaultInputLabel, EPCGDataType::Point);
+
+	return PinProperties;
+}
+
+TArray<FPCGPinProperties> UPCGSanityCheckPointDataSettings::OutputPinProperties() const
+{
+	TArray<FPCGPinProperties> PinProperties;
+	PinProperties.Emplace(PCGPinConstants::DefaultOutputLabel, EPCGDataType::Point);
+
+	return PinProperties;
+}
+
+FPCGElementPtr UPCGSanityCheckPointDataSettings::CreateElement() const
+{
+	return MakeShared<FPCGSanityCheckPointDataElement>();
+}
+
+bool FPCGSanityCheckPointDataElement::ExecuteInternal(FPCGContext* Context) const
+{
+	TRACE_CPUPROFILER_EVENT_SCOPE(FPCGSanityCheckPointDataElement::Execute);
+
+	const UPCGSanityCheckPointDataSettings* Settings = Context->GetInputSettings<UPCGSanityCheckPointDataSettings>();
+	check(Settings);
+
+	TArray<FPCGTaggedData> Sources = Context->InputData.GetInputsByPin(PCGPinConstants::DefaultInputLabel);
+
+	int32 InputPointCount = 0;
+
+	// first find the total Input bounds which will determine the size of each cell
+	for (const FPCGTaggedData& Source : Sources) 
+	{
+		const UPCGPointData* PointData = Cast<UPCGPointData>(Source.Data);
+
+		if (!PointData)
+		{
+			continue;
+		}
+
+		InputPointCount += PointData->GetPoints().Num();
+	}
+
+	if (InputPointCount < Settings->MinPointCount)
+	{
+		PCGE_LOG(Error, GraphAndLog, FText::Format(LOCTEXT("NotEnoughPoints", "Expected at least {0} Points, found {1}"), Settings->MinPointCount, InputPointCount));
+
+		if (Context->SourceComponent.IsValid())
+		{
+			Context->SourceComponent->CancelGeneration();
+		}
+	}
+	else if (InputPointCount > Settings->MaxPointCount)
+	{
+		PCGE_LOG(Error, GraphAndLog, FText::Format(LOCTEXT("TooManyPoints", "Expected less than {0} Points, found {1}"), Settings->MaxPointCount, InputPointCount));
+
+		if (Context->SourceComponent.IsValid())
+		{
+			Context->SourceComponent->CancelGeneration();
+		}
+	}
+	else
+	{
+		// if valid, pass through data
+		Context->OutputData = Context->InputData;
+	}
+
+	return true;
+}
+
+#undef LOCTEXT_NAMESPACE
diff --git a/Engine/Plugins/Experimental/PCG/Source/PCG/Private/Tests/Elements/PCGSanityCheckPointDataTest.cpp b/Engine/Plugins/Experimental/PCG/Source/PCG/Private/Tests/Elements/PCGSanityCheckPointDataTest.cpp
new file mode 100644
--- /dev/null
+++ b/Engine/Plugins/Experimental/PCG/Source/PCG/Private/Tests/Elements/PCGSanityCheckPointDataTest.cpp
@@ -0,0 +1,100 @@
+// Copyright Epic Games, Inc. All Rights Reserved.
+
+#include "Tests/Determinism/PCGDeterminismTestsCommon.h"
+
+#include "PCGComponent.h"
+#include "PCGContext.h"
+#include "Data/PCGPointData.h"
+#include "PCGSubsystem.h"
+#include "PCGGraph.h"
+#include "Utils/PCGNodeVisualLogs.h"
+
+#include "Elements/PCGSanityCheckPointData.h"
+
+IMPLEMENT_CUSTOM_SIMPLE_AUTOMATION_TEST(FPCGSanityCheckPointData_MinPointCount, FPCGTestBaseClass, "pcg.tests.SanityCheckPointData.MinPointCount", PCGTestsCommon::TestFlags)
+
+bool FPCGSanityCheckPointData_MinPointCount::RunTest(const FString& Parameters)
+{
+	PCGTestsCommon::FTestData TestData(PCGDeterminismTests::Defaults::Seed);
+	PCGDeterminismTests::GenerateSettings<UPCGSanityCheckPointDataSettings>(TestData);
+	UPCGSanityCheckPointDataSettings* Settings = CastChecked<UPCGSanityCheckPointDataSettings>(TestData.Settings);
+
+	Settings->MinPointCount = 10;
+	Settings->MaxPointCount = 100;
+
+	FPCGElementPtr TestElement = TestData.Settings->GetElement();
+
+	UPCGPointData *SourceData = PCGTestsCommon::CreateRandomPointData(2, 42);
+
+	FPCGTaggedData& SourcePin = TestData.InputData.TaggedData.Emplace_GetRef();
+	SourcePin.Pin = PCGPinConstants::DefaultInputLabel;
+	SourcePin.Data = SourceData;
+
+	TUniquePtr<FPCGContext> Context = TestData.InitializeTestContext();
+
+	AddExpectedError(TEXT("Expected at least 10 Points, found 2"));
+
+	while (!TestElement->Execute(Context.Get())) {}
+
+	return true;
+}
+
+IMPLEMENT_CUSTOM_SIMPLE_AUTOMATION_TEST(FPCGSanityCheckPointData_MaxPointCount, FPCGTestBaseClass, "pcg.tests.SanityCheckPointData.MaxPointCount", PCGTestsCommon::TestFlags)
+
+bool FPCGSanityCheckPointData_MaxPointCount::RunTest(const FString& Parameters)
+{
+	PCGTestsCommon::FTestData TestData(PCGDeterminismTests::Defaults::Seed);
+	PCGDeterminismTests::GenerateSettings<UPCGSanityCheckPointDataSettings>(TestData);
+	UPCGSanityCheckPointDataSettings* Settings = CastChecked<UPCGSanityCheckPointDataSettings>(TestData.Settings);
+
+	Settings->MinPointCount = 10;
+	Settings->MaxPointCount = 100;
+
+	FPCGElementPtr TestElement = TestData.Settings->GetElement();
+
+	UPCGPointData *SourceData = PCGTestsCommon::CreateRandomPointData(200, 42);
+
+	FPCGTaggedData& SourcePin = TestData.InputData.TaggedData.Emplace_GetRef();
+	SourcePin.Pin = PCGPinConstants::DefaultInputLabel;
+	SourcePin.Data = SourceData;
+
+	TUniquePtr<FPCGContext> Context = TestData.InitializeTestContext(TestData.TestPCGComponent->GetGraph()->AddNode(Settings));
+
+	AddExpectedError(TEXT("Expected less than 100 Points, found 200"));
+
+	while (!TestElement->Execute(Context.Get())) {}
+
+	return true;
+}
+
+IMPLEMENT_CUSTOM_SIMPLE_AUTOMATION_TEST(FPCGSanityCheckPointData_Success, FPCGTestBaseClass, "pcg.tests.SanityCheckPointData.Success", PCGTestsCommon::TestFlags)
+
+bool FPCGSanityCheckPointData_Success::RunTest(const FString& Parameters)
+{
+	PCGTestsCommon::FTestData TestData(PCGDeterminismTests::Defaults::Seed);
+	PCGDeterminismTests::GenerateSettings<UPCGSanityCheckPointDataSettings>(TestData);
+	UPCGSanityCheckPointDataSettings* Settings = CastChecked<UPCGSanityCheckPointDataSettings>(TestData.Settings);
+
+	Settings->MinPointCount = 10;
+	Settings->MaxPointCount = 100;
+
+	FPCGElementPtr TestElement = TestData.Settings->GetElement();
+
+	UPCGPointData *SourceData = PCGTestsCommon::CreateRandomPointData(50, 42);
+
+	FPCGTaggedData& SourcePin = TestData.InputData.TaggedData.Emplace_GetRef();
+	SourcePin.Pin = PCGPinConstants::DefaultInputLabel;
+	SourcePin.Data = SourceData;
+
+	TUniquePtr<FPCGContext> Context = TestData.InitializeTestContext(TestData.TestPCGComponent->GetGraph()->AddNode(Settings));
+
+	while (!TestElement->Execute(Context.Get())) {}
+
+	const TArray<FPCGTaggedData>& Outputs = Context->OutputData.GetInputs();
+
+	UTEST_EQUAL("Output count", Outputs.Num(), 1);
+
+	UTEST_TRUE("Data was forwarded", Outputs[0].Data == SourceData);
+
+	return true;
+}
diff --git a/Engine/Plugins/Experimental/PCG/Source/PCG/Public/Elements/PCGSanityCheckPointData.h b/Engine/Plugins/Experimental/PCG/Source/PCG/Public/Elements/PCGSanityCheckPointData.h
new file mode 100644
--- /dev/null
+++ b/Engine/Plugins/Experimental/PCG/Source/PCG/Public/Elements/PCGSanityCheckPointData.h
@@ -0,0 +1,45 @@
+// Copyright Epic Games, Inc. All Rights Reserved.
+
+#pragma once
+
+#include "PCGSettings.h"
+
+#include "PCGSanityCheckPointData.generated.h"
+
+UCLASS(BlueprintType, ClassGroup = (Procedural))
+class PCG_API UPCGSanityCheckPointDataSettings : public UPCGSettings
+{
+	GENERATED_BODY()
+
+public:
+	//~Begin UPCGSettings interface
+#if WITH_EDITOR
+	virtual FName GetDefaultNodeName() const override { return FName(TEXT("SanityCheckPointData")); }
+	virtual FText GetDefaultNodeTitle() const override { return NSLOCTEXT("PCGSanityCheckPointDataSettings", "NodeTitle", "Sanity Check Point Data"); }
+	virtual EPCGSettingsType GetType() const override { return EPCGSettingsType::Debug; }
+#endif
+
+	virtual TArray<FPCGPinProperties> InputPinProperties() const override;
+	virtual TArray<FPCGPinProperties> OutputPinProperties() const override;
+
+protected:
+	virtual FPCGElementPtr CreateElement() const override;
+	//~End UPCGSettings interface
+
+public:
+	UPROPERTY(BlueprintReadWrite, EditAnywhere, Category = Settings, meta = (PCG_Overridable))
+	int32 MinPointCount = 0;
+
+	UPROPERTY(BlueprintReadWrite, EditAnywhere, Category = Settings, meta = (PCG_Overridable))
+	int32 MaxPointCount = 100;
+};
+
+class FPCGSanityCheckPointDataElement : public FSimplePCGElement
+{
+protected:
+	virtual bool ExecuteInternal(FPCGContext* Context) const;
+};
+
+#if UE_ENABLE_INCLUDE_ORDER_DEPRECATED_IN_5_2
+#include "CoreMinimal.h"
+#endif
