[PCG] Added support to set debug object from level editor selection

diff --git a/Engine/Plugins/Experimental/PCG/Source/PCGEditor/Private/PCGEditorGraph.h b/Engine/Plugins/Experimental/PCG/Source/PCGEditor/Private/PCGEditorGraph.h
--- a/Engine/Plugins/Experimental/PCG/Source/PCGEditor/Private/PCGEditorGraph.h
+++ b/Engine/Plugins/Experimental/PCG/Source/PCGEditor/Private/PCGEditorGraph.h
@@ -32,7 +32,7 @@ public:
 	/** Creates the links for a given node */
 	void CreateLinks(UPCGEditorGraphNodeBase* InGraphNode, bool bCreateInbound, bool bCreateOutbound);
 
-	UPCGGraph* GetPCGGraph() { return PCGGraph; }
+	UPCGGraph* GetPCGGraph() const { return PCGGraph; }
 
 	void SetEditor(TWeakPtr<const FPCGEditor> InEditor) { PCGEditor = InEditor; }
 	TWeakPtr<const FPCGEditor> GetEditor() const { return PCGEditor; }
diff --git a/Engine/Plugins/Experimental/PCG/Source/PCGEditor/Private/SPCGEditorGraphDebugObjectWidget.cpp b/Engine/Plugins/Experimental/PCG/Source/PCGEditor/Private/SPCGEditorGraphDebugObjectWidget.cpp
--- a/Engine/Plugins/Experimental/PCG/Source/PCGEditor/Private/SPCGEditorGraphDebugObjectWidget.cpp
+++ b/Engine/Plugins/Experimental/PCG/Source/PCGEditor/Private/SPCGEditorGraphDebugObjectWidget.cpp
@@ -2,31 +2,62 @@
 
 #include "SPCGEditorGraphDebugObjectWidget.h"
 
-#include "Framework/Views/TableViewMetadata.h"
 #include "PCGComponent.h"
 #include "PCGEditor.h"
 #include "PCGEditorGraph.h"
 
-#include "Editor/UnrealEdEngine.h"
 #include "PropertyCustomizationHelpers.h"
+#include "Selection.h"
 #include "UnrealEdGlobals.h"
+#include "Editor/UnrealEdEngine.h"
 #include "Widgets/Input/SComboBox.h"
 
 #define LOCTEXT_NAMESPACE "PCGEditorGraphDebugObjectWidget"
 
 namespace PCGEditorGraphDebugObjectWidget
 {
-	const FString NoObjectString = TEXT("No debug object selected");
 	const FString SelectionString = TEXT(" (selected)");
 	const FString SeparatorString = TEXT(" / ");
 }
 
+FPCGEditorGraphDebugObjectInstance::FPCGEditorGraphDebugObjectInstance(TWeakObjectPtr<UPCGComponent> InPCGComponent)
+	: PCGComponent(InPCGComponent)
+{
+	SetLabelFromPCGComponent(InPCGComponent);
+}
+
+void FPCGEditorGraphDebugObjectInstance::SetLabelFromPCGComponent(TWeakObjectPtr<UPCGComponent> InPCGComponent)
+{
+	check(InPCGComponent.IsValid());
+	const AActor* Actor = InPCGComponent->GetOwner();
+	check(Actor);
+	FString ActorNameOrLabel = Actor->GetActorNameOrLabel();
+	if (Actor->IsSelected())
+	{
+		ActorNameOrLabel.Append(PCGEditorGraphDebugObjectWidget::SelectionString);
+	}
+			
+	FString ComponentName = InPCGComponent->GetFName().ToString();
+	if (InPCGComponent->IsSelected())
+	{
+		ComponentName.Append(PCGEditorGraphDebugObjectWidget::SelectionString);
+	}
+
+	Label = ActorNameOrLabel + PCGEditorGraphDebugObjectWidget::SeparatorString + ComponentName;
+}
+
 void SPCGEditorGraphDebugObjectWidget::Construct(const FArguments& InArgs, TSharedPtr<FPCGEditor> InPCGEditor)
 {
 	PCGEditorPtr = InPCGEditor;
 
-	DebugObjects.Add(MakeShared<FPCGEditorGraphDebugObjectInstance>(PCGEditorGraphDebugObjectWidget::NoObjectString));
+	DebugObjects.Add(MakeShared<FPCGEditorGraphDebugObjectInstance>());
 
+	const TSharedRef<SWidget> SetButton = PropertyCustomizationHelpers::MakeUseSelectedButton(
+		FSimpleDelegate::CreateSP(this, &SPCGEditorGraphDebugObjectWidget::SetDebugObjectFromSelection_OnClicked),
+		LOCTEXT("SetDebugObject", "Set debug object from Level Editor selection."),
+		TAttribute<bool>::Create(TAttribute<bool>::FGetter::CreateSP(this, &SPCGEditorGraphDebugObjectWidget::IsSetDebugObjectFromSelectionButtonEnabled))
+	);
+	
 	const TSharedRef<SWidget> BrowseButton = PropertyCustomizationHelpers::MakeBrowseButton(
 		FSimpleDelegate::CreateSP(this, &SPCGEditorGraphDebugObjectWidget::SelectedDebugObject_OnClicked),
 		LOCTEXT("DebugSelectActor", "Select and frame the debug actor in the Level Editor."),
@@ -57,6 +88,13 @@ void SPCGEditorGraphDebugObjectWidget::Construct(const FArguments& InArgs, TShar
 		.AutoWidth()
 		.VAlign(VAlign_Center)
 		.Padding(4.0f)
+		[
+			SetButton
+		]
+		+SHorizontalBox::Slot()
+		.AutoWidth()
+		.VAlign(VAlign_Center)
+		.Padding(4.0f)
 		[
 			BrowseButton
 		]
@@ -68,34 +106,23 @@ void SPCGEditorGraphDebugObjectWidget::OnComboBoxOpening()
 	DebugObjects.Empty();
 	DebugObjectsComboBox->RefreshOptions();
 	
-	if (!PCGEditorPtr.IsValid())
-	{
-		return;
-	}
-
-	UPCGEditorGraph* PCGEditorGraph = PCGEditorPtr.Pin()->GetPCGEditorGraph();
-	if (!PCGEditorGraph)
-	{
-		return;
-	}
-
-	const UPCGGraph* PCGGraph = PCGEditorGraph->GetPCGGraph();
+	const UPCGGraph* PCGGraph = GetPCGGraph();
 	if (!PCGGraph)
 	{
 		return;
 	}
 
-	TSharedPtr<FPCGEditorGraphDebugObjectInstance> SelectedItem = DebugObjectsComboBox->GetSelectedItem();
+	const TSharedPtr<FPCGEditorGraphDebugObjectInstance> SelectedItem = DebugObjectsComboBox->GetSelectedItem();
 
-	DebugObjects.Add(MakeShared<FPCGEditorGraphDebugObjectInstance>(PCGEditorGraphDebugObjectWidget::NoObjectString));
+	DebugObjects.Add(MakeShared<FPCGEditorGraphDebugObjectInstance>());
 
-	if (!SelectedItem.IsValid() || SelectedItem->PCGComponent == nullptr)
+	if (!SelectedItem.IsValid() || !SelectedItem->GetPCGComponent().IsValid())
 	{
 		DebugObjectsComboBox->SetSelectedItem(DebugObjects[0]);
 	}
 	
 	TArray<UObject*> PCGComponents;
-	GetObjectsOfClass(UPCGComponent::StaticClass(), PCGComponents, true);
+	GetObjectsOfClass(UPCGComponent::StaticClass(), PCGComponents, /*bIncludeDerivedClasses=*/ true);
 	for (UObject* PCGComponentObject : PCGComponents)
 	{
 		if (!IsValid(PCGComponentObject))
@@ -123,23 +150,10 @@ void SPCGEditorGraphDebugObjectWidget::OnComboBoxOpening()
 
 		if (PCGComponentGraph == PCGGraph)
 		{
-			FString ActorNameOrLabel = Actor->GetActorNameOrLabel();
-			if (Actor->IsSelected())
-			{
-				ActorNameOrLabel.Append(PCGEditorGraphDebugObjectWidget::SelectionString);
-			}
-			
-			FString ComponentName = PCGComponent->GetFName().ToString();
-			if (PCGComponent->IsSelected())
-			{
-				ComponentName.Append(PCGEditorGraphDebugObjectWidget::SelectionString);
-			}
-
-			const FString Label = ActorNameOrLabel + PCGEditorGraphDebugObjectWidget::SeparatorString + ComponentName;
-			TSharedPtr<FPCGEditorGraphDebugObjectInstance> DebugInstance = MakeShared<FPCGEditorGraphDebugObjectInstance>(PCGComponent, Label);
+			TSharedPtr<FPCGEditorGraphDebugObjectInstance> DebugInstance = MakeShared<FPCGEditorGraphDebugObjectInstance>(PCGComponent);
 			DebugObjects.Add(DebugInstance);
 
-			if (SelectedItem.IsValid() && SelectedItem->PCGComponent == PCGComponent)
+			if (SelectedItem.IsValid() && SelectedItem->GetPCGComponent() == PCGComponent)
 			{
 				DebugObjectsComboBox->SetSelectedItem(DebugInstance);
 			}
@@ -151,17 +165,33 @@ void SPCGEditorGraphDebugObjectWidget::OnSelectionChanged(TSharedPtr<FPCGEditorG
 {
 	if (NewSelection.IsValid())
 	{
-		UPCGComponent* PCGComponent = NewSelection->PCGComponent.Get();
+		UPCGComponent* PCGComponent = NewSelection->GetPCGComponent().Get();
 		PCGEditorPtr.Pin()->SetPCGComponentBeingDebugged(PCGComponent);
 	}
 }
 
 TSharedRef<SWidget> SPCGEditorGraphDebugObjectWidget::OnGenerateWidget(TSharedPtr<FPCGEditorGraphDebugObjectInstance> InDebugObjectInstance) const
 {
-	const FString ItemString = InDebugObjectInstance->Label;
+	const FText ItemText = InDebugObjectInstance->GetDebugObjectText();
 
 	return SNew(STextBlock)
-		.Text(FText::FromString(*ItemString));
+		.Text(ItemText);
+}
+
+UPCGGraph* SPCGEditorGraphDebugObjectWidget::GetPCGGraph() const
+{
+	if (!PCGEditorPtr.IsValid())
+	{
+		return nullptr;
+	}
+	
+	const UPCGEditorGraph* PCGEditorGraph = PCGEditorPtr.Pin()->GetPCGEditorGraph();
+	if (!PCGEditorGraph)
+	{
+		return nullptr;
+	}
+
+	return PCGEditorGraph->GetPCGGraph();
 }
 
 FText SPCGEditorGraphDebugObjectWidget::GetSelectedDebugObjectText() const
@@ -193,4 +223,84 @@ bool SPCGEditorGraphDebugObjectWidget::IsSelectDebugObjectButtonEnabled() const
 	return PCGEditorPtr.IsValid() && (PCGEditorPtr.Pin()->GetPCGComponentBeingDebugged() != nullptr);
 }
 
+void SPCGEditorGraphDebugObjectWidget::SetDebugObjectFromSelection_OnClicked()
+{
+	const UPCGGraph* PCGGraph = GetPCGGraph();
+	if (!PCGGraph)
+	{
+		return;
+	}
+	
+	USelection* SelectedActors = GEditor->GetSelectedActors();
+	if (!IsValid(SelectedActors))
+	{
+		return;
+	}
+	
+	for (FSelectionIterator It(*SelectedActors); It; ++It)
+	{
+		const AActor* SelectedActor = Cast<AActor>(*It);
+		if (!IsValid(SelectedActor))
+		{
+			continue;
+		}
+		
+		UPCGComponent* PCGComponent = SelectedActor->GetComponentByClass<UPCGComponent>();
+		if (!IsValid(PCGComponent))
+		{
+			continue;
+		}
+		
+		if (PCGComponent->GetGraph() == PCGGraph)
+		{
+			DebugObjects.Empty();
+			DebugObjectsComboBox->RefreshOptions();
+
+			const TSharedPtr<FPCGEditorGraphDebugObjectInstance> DebugInstance = MakeShared<FPCGEditorGraphDebugObjectInstance>(PCGComponent);
+			DebugObjects.Add(DebugInstance);
+			DebugObjectsComboBox->SetSelectedItem(DebugInstance);
+			
+			PCGEditorPtr.Pin()->SetPCGComponentBeingDebugged(PCGComponent);
+			break;
+		}
+	}
+}
+
+bool SPCGEditorGraphDebugObjectWidget::IsSetDebugObjectFromSelectionButtonEnabled() const
+{
+	const UPCGGraph* PCGGraph = GetPCGGraph();
+	if (!PCGGraph)
+	{
+		return false;
+	}
+	
+	USelection* SelectedActors = GEditor->GetSelectedActors();
+	if (!IsValid(SelectedActors))
+	{
+		return false;
+	}
+	
+	for (FSelectionIterator It(*SelectedActors); It; ++It)
+	{
+		const AActor* SelectedActor = Cast<AActor>(*It);
+		if (!IsValid(SelectedActor))
+		{
+			continue;
+		}
+		
+		const UPCGComponent* PCGComponent = SelectedActor->GetComponentByClass<UPCGComponent>();
+		if (!IsValid(PCGComponent))
+		{
+			continue;
+		}
+		
+		if (PCGComponent->GetGraph() == PCGGraph)
+		{
+			return true;
+		}
+	}
+	
+	return false;
+}
+
 #undef LOCTEXT_NAMESPACE
diff --git a/Engine/Plugins/Experimental/PCG/Source/PCGEditor/Private/SPCGEditorGraphDebugObjectWidget.h b/Engine/Plugins/Experimental/PCG/Source/PCGEditor/Private/SPCGEditorGraphDebugObjectWidget.h
--- a/Engine/Plugins/Experimental/PCG/Source/PCGEditor/Private/SPCGEditorGraphDebugObjectWidget.h
+++ b/Engine/Plugins/Experimental/PCG/Source/PCGEditor/Private/SPCGEditorGraphDebugObjectWidget.h
@@ -9,28 +9,29 @@ template <typename OptionType> class SComboBox;
 
 class FPCGEditor;
 class UPCGComponent;
+class UPCGGraph;
 
-struct FPCGEditorGraphDebugObjectInstance
+class FPCGEditorGraphDebugObjectInstance
 {
-	FPCGEditorGraphDebugObjectInstance(const FString& InLabel)
-		: PCGComponent(nullptr)
-		, Label(InLabel)
-	{
-	}
-	
-	FPCGEditorGraphDebugObjectInstance(TWeakObjectPtr<UPCGComponent> InPCGComponent, const FString& InLabel)
-		: PCGComponent(InPCGComponent)
-		, Label(InLabel)
-	{
-	}
+public:
+	FPCGEditorGraphDebugObjectInstance() = default;
+	FPCGEditorGraphDebugObjectInstance(TWeakObjectPtr<UPCGComponent> InPCGComponent);
+
+	void SetLabelFromPCGComponent(TWeakObjectPtr<UPCGComponent> InPCGComponent);
 
 	FText GetDebugObjectText() const
 	{
 		return FText::FromString(Label);
 	}
 
-	TWeakObjectPtr<UPCGComponent> PCGComponent;
-	FString Label;
+	TWeakObjectPtr<UPCGComponent> GetPCGComponent() const
+	{
+		return PCGComponent;
+	}
+	
+private:
+	TWeakObjectPtr<UPCGComponent> PCGComponent = nullptr;
+	FString Label = TEXT("No debug object selected");
 };
 
 
@@ -47,10 +48,15 @@ private:
 	void OnSelectionChanged(TSharedPtr<FPCGEditorGraphDebugObjectInstance> NewSelection, ESelectInfo::Type SelectInfo) const;
 	TSharedRef<SWidget> OnGenerateWidget(TSharedPtr<FPCGEditorGraphDebugObjectInstance> InDebugObjectInstance) const;
 
+	UPCGGraph* GetPCGGraph() const;
+	
 	FText GetSelectedDebugObjectText() const;
 	void SelectedDebugObject_OnClicked() const;
 	bool IsSelectDebugObjectButtonEnabled() const;
-
+	
+	void SetDebugObjectFromSelection_OnClicked();
+	bool IsSetDebugObjectFromSelectionButtonEnabled() const;
+	
 	/** Pointer back to the PCG editor that owns us */
 	TWeakPtr<FPCGEditor> PCGEditorPtr;
 
