Fixed mismatched copy constructor/assignment operator in FPCGPointRef.

diff --git a/Engine/Plugins/Experimental/PCG/Source/PCG/Private/Data/PCGPointData.cpp b/Engine/Plugins/Experimental/PCG/Source/PCG/Private/Data/PCGPointData.cpp
--- a/Engine/Plugins/Experimental/PCG/Source/PCG/Private/Data/PCGPointData.cpp
+++ b/Engine/Plugins/Experimental/PCG/Source/PCG/Private/Data/PCGPointData.cpp
@@ -235,12 +235,6 @@ FPCGPointRef::FPCGPointRef(const FPCGPoint& InPoint)
 	Bounds = InPoint.GetDensityBounds();
 }
 
-FPCGPointRef::FPCGPointRef(const FPCGPointRef& InPointRef)
-{
-	Point = InPointRef.Point;
-	Bounds = InPointRef.Bounds;
-}
-
 void UPCGPointData::GetResourceSizeEx(FResourceSizeEx& CumulativeResourceSize)
 {
 	Super::GetResourceSizeEx(CumulativeResourceSize);
diff --git a/Engine/Plugins/Experimental/PCG/Source/PCG/Public/Data/PCGPointData.h b/Engine/Plugins/Experimental/PCG/Source/PCG/Public/Data/PCGPointData.h
--- a/Engine/Plugins/Experimental/PCG/Source/PCG/Public/Data/PCGPointData.h
+++ b/Engine/Plugins/Experimental/PCG/Source/PCG/Public/Data/PCGPointData.h
@@ -23,7 +23,6 @@ namespace PCGPointHelpers
 struct PCG_API FPCGPointRef
 {
 	FPCGPointRef(const FPCGPoint& InPoint);
-	FPCGPointRef(const FPCGPointRef& InPointRef);
 
 	const FPCGPoint* Point;
 	FBoxSphereBounds Bounds;
