Improved usage of landscape original guid when using multiple instanced landscapes.

diff --git a/Engine/Plugins/Experimental/PCG/Source/PCG/Private/Data/PCGLandscapeData.cpp b/Engine/Plugins/Experimental/PCG/Source/PCG/Private/Data/PCGLandscapeData.cpp
--- a/Engine/Plugins/Experimental/PCG/Source/PCG/Private/Data/PCGLandscapeData.cpp
+++ b/Engine/Plugins/Experimental/PCG/Source/PCG/Private/Data/PCGLandscapeData.cpp
@@ -151,7 +151,7 @@ bool UPCGLandscapeData::ProjectPoint(const FTransform& InTransform, const FBox&
 	ULandscapeComponent* LandscapeComponent = LandscapeInfo->XYtoComponentMap.FindRef(ComponentMapKey);
 	const FPCGLandscapeCacheEntry * LandscapeCacheEntry = (LandscapeComponent ? LandscapeCache->GetCacheEntry(LandscapeComponent, ComponentMapKey) : nullptr);
 #else
-	const FPCGLandscapeCacheEntry* LandscapeCacheEntry = LandscapeCache->GetCacheEntry(LandscapeInfo->LandscapeGuid, ComponentMapKey);
+	const FPCGLandscapeCacheEntry* LandscapeCacheEntry = LandscapeCache->GetCacheEntry(LandscapeInfo->GetLandscapeProxy()->GetOriginalLandscapeGuid(), ComponentMapKey);
 #endif
 
 	if (!LandscapeCacheEntry)
@@ -263,6 +263,7 @@ const UPCGPointData* UPCGLandscapeData::CreatePointData(FPCGContext* Context, co
 		const int32 MaxComponentX = MaxX / ComponentSizeQuads;
 		const int32 MinComponentY = MinY / ComponentSizeQuads;
 		const int32 MaxComponentY = MaxY / ComponentSizeQuads;
+		const FGuid LandscapeGuid = LandscapeInfo->GetLandscapeProxy()->GetOriginalLandscapeGuid();
 
 		for (int32 ComponentX = MinComponentX; ComponentX <= MaxComponentX; ++ComponentX)
 		{
@@ -273,7 +274,7 @@ const UPCGPointData* UPCGLandscapeData::CreatePointData(FPCGContext* Context, co
 				ULandscapeComponent* LandscapeComponent = LandscapeInfo->XYtoComponentMap.FindRef(ComponentMapKey);
 				const FPCGLandscapeCacheEntry* LandscapeCacheEntry = LandscapeComponent ? LandscapeCache->GetCacheEntry(LandscapeComponent, ComponentMapKey) : nullptr;
 #else
-				const FPCGLandscapeCacheEntry* LandscapeCacheEntry = LandscapeCache->GetCacheEntry(LandscapeInfo->LandscapeGuid, ComponentMapKey);
+				const FPCGLandscapeCacheEntry* LandscapeCacheEntry = LandscapeCache->GetCacheEntry(LandscapeGuid, ComponentMapKey);
 #endif
 
 				if (!LandscapeCacheEntry)
diff --git a/Engine/Plugins/Experimental/PCG/Source/PCG/Private/Grid/PCGLandscapeCache.cpp b/Engine/Plugins/Experimental/PCG/Source/PCG/Private/Grid/PCGLandscapeCache.cpp
--- a/Engine/Plugins/Experimental/PCG/Source/PCG/Private/Grid/PCGLandscapeCache.cpp
+++ b/Engine/Plugins/Experimental/PCG/Source/PCG/Private/Grid/PCGLandscapeCache.cpp
@@ -569,7 +569,7 @@ void UPCGLandscapeCache::PrimeCache()
 			LandscapeInfo->ForEachLandscapeProxy([this, LandscapeInfo, &CacheEntriesToBuild](const ALandscapeProxy* LandscapeProxy)
 			{
 				check(LandscapeProxy);
-				const FGuid LandscapeGuid = LandscapeProxy->GetLandscapeGuid();
+				const FGuid LandscapeGuid = LandscapeProxy->GetOriginalLandscapeGuid();
 
 				for (ULandscapeComponent* LandscapeComponent : LandscapeProxy->LandscapeComponents)
 				{
@@ -628,7 +628,7 @@ void UPCGLandscapeCache::ClearCache()
 #if WITH_EDITOR
 const FPCGLandscapeCacheEntry* UPCGLandscapeCache::GetCacheEntry(ULandscapeComponent* LandscapeComponent, const FIntPoint& ComponentCoordinate)
 {
-	const FGuid LandscapeGuid = (LandscapeComponent && LandscapeComponent->GetLandscapeProxy() ? LandscapeComponent->GetLandscapeProxy()->GetLandscapeGuid() : FGuid());
+	const FGuid LandscapeGuid = (LandscapeComponent && LandscapeComponent->GetLandscapeProxy() ? LandscapeComponent->GetLandscapeProxy()->GetOriginalLandscapeGuid() : FGuid());
 	const FPCGLandscapeCacheEntry* CacheEntry = GetCacheEntry(LandscapeGuid, ComponentCoordinate);
 
 	if (!CacheEntry && LandscapeComponent && LandscapeComponent->GetLandscapeInfo())
@@ -717,7 +717,7 @@ void UPCGLandscapeCache::SampleMetadataOnPoint(ALandscapeProxy* Landscape, FPCGP
 	ULandscapeComponent* LandscapeComponent = LandscapeInfo->XYtoComponentMap.FindRef(ComponentMapKey);
 	const FPCGLandscapeCacheEntry* CacheEntry = GetCacheEntry(LandscapeComponent, ComponentMapKey);
 #else
-	const FPCGLandscapeCacheEntry* CacheEntry = GetCacheEntry(Landscape->GetLandscapeGuid(), ComponentMapKey);
+	const FPCGLandscapeCacheEntry* CacheEntry = GetCacheEntry(Landscape->GetOriginalLandscapeGuid(), ComponentMapKey);
 #endif
 
 	if (!CacheEntry || CacheEntry->LayerData.IsEmpty())
@@ -780,7 +780,7 @@ void UPCGLandscapeCache::OnLandscapeChanged(ALandscapeProxy* Landscape, const FL
 		if (LandscapeComponent)
 		{
 			FIntPoint Coordinate(LandscapeComponent->SectionBaseX / LandscapeComponent->ComponentSizeQuads, LandscapeComponent->SectionBaseY / LandscapeComponent->ComponentSizeQuads);
-			TPair<FGuid, FIntPoint> ComponentKey(Landscape->GetLandscapeGuid(), Coordinate);
+			TPair<FGuid, FIntPoint> ComponentKey(Landscape->GetOriginalLandscapeGuid(), Coordinate);
 
 			if (FPCGLandscapeCacheEntry** FoundEntry = CachedData.Find(ComponentKey))
 			{
diff --git a/Engine/Plugins/Experimental/PCG/Source/PCG/Private/Helpers/PCGBlueprintHelpers.cpp b/Engine/Plugins/Experimental/PCG/Source/PCG/Private/Helpers/PCGBlueprintHelpers.cpp
--- a/Engine/Plugins/Experimental/PCG/Source/PCG/Private/Helpers/PCGBlueprintHelpers.cpp
+++ b/Engine/Plugins/Experimental/PCG/Source/PCG/Private/Helpers/PCGBlueprintHelpers.cpp
@@ -167,7 +167,7 @@ TArray<FPCGLandscapeLayerWeight> UPCGBlueprintHelpers::GetInterpolatedPCGLandsca
 	ULandscapeComponent* LandscapeComponent = LandscapeInfo->XYtoComponentMap.FindRef(ComponentMapKey);
 	const FPCGLandscapeCacheEntry* CacheEntry = LandscapeCache->GetCacheEntry(LandscapeComponent, ComponentMapKey);
 #else
-	const FPCGLandscapeCacheEntry* CacheEntry = LandscapeCache->GetCacheEntry(Landscape->GetLandscapeGuid(), ComponentMapKey);
+	const FPCGLandscapeCacheEntry* CacheEntry = LandscapeCache->GetCacheEntry(Landscape->GetOriginalLandscapeGuid(), ComponentMapKey);
 #endif
 
 	if (!CacheEntry)
