Fixed crash in non-editor builds when using static mesh spawner nodes with no connected output

diff --git a/Engine/Plugins/Experimental/PCG/Source/PCG/Private/Elements/PCGStaticMeshSpawner.cpp b/Engine/Plugins/Experimental/PCG/Source/PCG/Private/Elements/PCGStaticMeshSpawner.cpp
--- a/Engine/Plugins/Experimental/PCG/Source/PCG/Private/Elements/PCGStaticMeshSpawner.cpp
+++ b/Engine/Plugins/Experimental/PCG/Source/PCG/Private/Elements/PCGStaticMeshSpawner.cpp
@@ -168,14 +168,14 @@ bool FPCGStaticMeshSpawnerElement::PrepareDataInternal(FPCGContext* InContext) c
 				Output.Data = OutputPointData;
 				check(!Context->CurrentOutputPointData);
 				Context->CurrentOutputPointData = OutputPointData;
+			}
 
-				// Create an entry in the MeshInstancesData if we're going to keep it
-				if (!bSkippedDueToReuse)
-				{
-					FPCGStaticMeshSpawnerContext::FPackedInstanceListData& InstanceListData = Context->MeshInstancesData.Emplace_GetRef();
-					InstanceListData.TargetActor = TargetActor;
-					InstanceListData.SpatialData = PointData;
-				}
+			// At this point, if we're in a reuse case we don't need to create the instance list here as it won't be processed or spawned
+			if (!bSkippedDueToReuse)
+			{
+				FPCGStaticMeshSpawnerContext::FPackedInstanceListData& InstanceListData = Context->MeshInstancesData.Emplace_GetRef();
+				InstanceListData.TargetActor = TargetActor;
+				InstanceListData.SpatialData = PointData;
 			}
 
 			Context->CurrentPointData = PointData;
