[PCG] Fix crash when using Break Debugger feature without debugger attached

diff --git a/Engine/Plugins/Experimental/PCG/Source/PCG/Private/PCGElement.cpp b/Engine/Plugins/Experimental/PCG/Source/PCG/Private/PCGElement.cpp
--- a/Engine/Plugins/Experimental/PCG/Source/PCG/Private/PCGElement.cpp
+++ b/Engine/Plugins/Experimental/PCG/Source/PCG/Private/PCGElement.cpp
@@ -23,7 +23,7 @@ static TAutoConsoleVariable<bool> CVarPCGValidatePointMetadata(
 #define PCG_ELEMENT_EXECUTION_BREAKPOINT() \
 	if (Context && Context->GetInputSettingsInterface() && Context->GetInputSettingsInterface()->bBreakDebugger) \
 	{ \
-		PLATFORM_BREAK(); \
+		UE_DEBUG_BREAK(); \
 	}
 #else
 #define PCG_ELEMENT_EXECUTION_BREAKPOINT()
diff --git a/Engine/Plugins/Experimental/PCG/Source/PCG/Public/PCGSettings.h b/Engine/Plugins/Experimental/PCG/Source/PCG/Public/PCGSettings.h
--- a/Engine/Plugins/Experimental/PCG/Source/PCG/Public/PCGSettings.h
+++ b/Engine/Plugins/Experimental/PCG/Source/PCG/Public/PCGSettings.h
@@ -137,7 +137,7 @@ public:
 	FPCGDebugVisualizationSettings DebugSettings;
 
 	/** If a debugger is attached, triggers a breakpoint inside IPCGElement::Execute(). Editor only. Transient. */
-	UPROPERTY(Transient, DuplicateTransient, EditAnywhere, BlueprintReadWrite, Category = Debug)
+	UPROPERTY(Transient, DuplicateTransient, EditAnywhere, BlueprintReadWrite, Category = Debug, AdvancedDisplay)
 	bool bBreakDebugger = false;
 #endif
 };
