[PCG] Add Point Filter Range node

* Refactorize code for both Point Filter to use the same underlying logic
* Add another for better properties names
* Add some unit tests too.

diff --git a/Engine/Plugins/Experimental/PCG/Source/PCG/Private/Elements/PCGPointFilter.cpp b/Engine/Plugins/Experimental/PCG/Source/PCG/Private/Elements/PCGPointFilter.cpp
--- a/Engine/Plugins/Experimental/PCG/Source/PCG/Private/Elements/PCGPointFilter.cpp
+++ b/Engine/Plugins/Experimental/PCG/Source/PCG/Private/Elements/PCGPointFilter.cpp
@@ -20,10 +20,17 @@ namespace PCGPointFilterConstants
 {
 	const FName DataToFilterLabel = TEXT("In");
 	const FName FilterLabel = TEXT("Filter");
+	const FName FilterMinLabel = TEXT("FilterMin");
+	const FName FilterMaxLabel = TEXT("FilterMax");
 	const FName InFilterLabel = TEXT("InsideFilter");
 	const FName OutFilterLabel = TEXT("OutsideFilter");
 
 	constexpr int32 ChunkSize = 256;
+
+#if WITH_EDITOR
+	const FText FilterPinTooltip = LOCTEXT("FilterPinTooltip", "This pin accepts Statial data and Attribute Sets. If the data is Spatial, it will automatically sample input points in it. "
+		"If it is points, it will sample if \"Spatial Query\" is enabled, otherwise points number need to match with input.");
+#endif // WITH_EDITOR
 }
 
 namespace PCGPointFilterHelpers
@@ -61,6 +68,159 @@ namespace PCGPointFilterHelpers
 			return false;
 		}
 	}
+
+	template <typename T>
+	bool ApplyRange(const T& Input, const T& InMin, const T& InMax, bool bMinIncluded, bool bMaxIncluded)
+	{
+		if constexpr (PCG::Private::MetadataTraits<T>::CanCompare)
+		{
+			return (bMinIncluded ? PCG::Private::MetadataTraits<T>::GreaterOrEqual(Input, InMin) : PCG::Private::MetadataTraits<T>::Greater(Input, InMin)) &&
+				(bMaxIncluded ? PCG::Private::MetadataTraits<T>::LessOrEqual(Input, InMax) : PCG::Private::MetadataTraits<T>::Less(Input, InMax));
+		}
+		else
+		{
+			return false;
+		}
+	}
+
+	struct ThresholdInfo
+	{
+		TUniquePtr<const IPCGAttributeAccessor> ThresholdAccessor;
+		TUniquePtr<const IPCGAttributeAccessorKeys> ThresholdKeys;
+		bool bUseInputDataForThreshold = false;
+		bool bUseSpatialQuery = false;
+
+		UPCGPointData* ThresholdPointData = nullptr;
+		const UPCGSpatialData* ThresholdSpatialData = nullptr;
+	};
+
+	bool InitialPrepareThresholdInfo(FPCGContext* InContext, TArray<FPCGTaggedData> FilterData, const FPCGPointFilterThresholdSettings& ThresholdSettings, ThresholdInfo& OutThresholdInfo)
+	{
+		if (ThresholdSettings.bUseConstantThreshold)
+		{
+			auto ConstantThreshold = [&OutThresholdInfo](auto&& Value)
+			{
+				using ConstantType = std::decay_t<decltype(Value)>;
+
+				OutThresholdInfo.ThresholdAccessor = MakeUnique<FPCGConstantValueAccessor<ConstantType>>(std::forward<decltype(Value)>(Value));
+				// Dummy keys
+				OutThresholdInfo.ThresholdKeys = MakeUnique<FPCGAttributeAccessorKeysSingleObjectPtr<void>>();
+			};
+
+			ThresholdSettings.AttributeTypes.Dispatcher(ConstantThreshold);
+		}
+		else
+		{
+			if (!FilterData.IsEmpty())
+			{
+				const UPCGData* ThresholdData = FilterData[0].Data;
+
+				if (const UPCGSpatialData* ThresholdSpatialData = Cast<const UPCGSpatialData>(ThresholdData))
+				{
+					// If the threshold is spatial or points (and spatial query is enabled), we'll use spatial query (meaning we'll have to sample points).
+					// Don't create an accessor yet (ThresholdData = nullptr), it will be created further down.
+					OutThresholdInfo.ThresholdSpatialData = ThresholdSpatialData;
+					if (!ThresholdSpatialData->IsA<UPCGPointData>() || ThresholdSettings.bUseSpatialQuery)
+					{
+						OutThresholdInfo.bUseSpatialQuery = true;
+						ThresholdData = nullptr;
+					}
+				}
+
+				if (ThresholdData)
+				{
+					OutThresholdInfo.ThresholdAccessor = PCGAttributeAccessorHelpers::CreateConstAccessor(ThresholdData, ThresholdSettings.ThresholdAttribute);
+					OutThresholdInfo.ThresholdKeys = PCGAttributeAccessorHelpers::CreateConstKeys(ThresholdData, ThresholdSettings.ThresholdAttribute);
+				}
+			}
+			else
+			{
+				OutThresholdInfo.bUseInputDataForThreshold = true;
+			}
+		}
+
+		return true;
+	}
+
+	bool PrepareThresholdInfoFromInput(FPCGContext* InContext, const UPCGPointData* InputData, const FPCGPointFilterThresholdSettings& ThresholdSettings, ThresholdInfo& InOutThresholdInfo, int16 TargetType, bool bCheckCompare, const ThresholdInfo* OtherInfo = nullptr)
+	{
+		check(InContext && InputData);
+
+		if (InOutThresholdInfo.bUseInputDataForThreshold)
+		{
+			// If we have no threshold accessor, we use the same data as input
+			InOutThresholdInfo.ThresholdAccessor = PCGAttributeAccessorHelpers::CreateConstAccessor(InputData, ThresholdSettings.ThresholdAttribute);
+			InOutThresholdInfo.ThresholdKeys = PCGAttributeAccessorHelpers::CreateConstKeys(InputData, ThresholdSettings.ThresholdAttribute);
+		}
+		else if (InOutThresholdInfo.ThresholdSpatialData != nullptr && InOutThresholdInfo.bUseSpatialQuery)
+		{
+			// Don't do 2 spatial query if we are iterating on the same input
+			if (OtherInfo != nullptr && OtherInfo->ThresholdSpatialData == InOutThresholdInfo.ThresholdSpatialData)
+			{
+				InOutThresholdInfo.ThresholdPointData = OtherInfo->ThresholdPointData;
+			}
+			else
+			{
+				// Reset the point data and reserving some points
+				// No need to reserve the full number of points, since we'll go by chunk
+				// Only allocate the chunk size
+				InOutThresholdInfo.ThresholdPointData = NewObject<UPCGPointData>();
+				InOutThresholdInfo.ThresholdPointData->InitializeFromData(InOutThresholdInfo.ThresholdSpatialData);
+				InOutThresholdInfo.ThresholdPointData->GetMutablePoints().SetNum(PCGPointFilterConstants::ChunkSize);
+			}
+
+			// Accessor will be valid, but keys will point to default points. But since it is a view, it will be updated when we sample the points.
+			InOutThresholdInfo.ThresholdAccessor = PCGAttributeAccessorHelpers::CreateConstAccessor(InOutThresholdInfo.ThresholdPointData, ThresholdSettings.ThresholdAttribute);
+			InOutThresholdInfo.ThresholdKeys = PCGAttributeAccessorHelpers::CreateConstKeys(InOutThresholdInfo.ThresholdPointData, ThresholdSettings.ThresholdAttribute);
+		}
+
+		if (!InOutThresholdInfo.ThresholdAccessor.IsValid() || !InOutThresholdInfo.ThresholdKeys.IsValid())
+		{
+			PCGE_LOG_C(Warning, GraphAndLog, InContext, FText::Format(LOCTEXT("AttributeMissingForFilter", "DataToFilter does not have '{0}' threshold attribute/property"), FText::FromName(ThresholdSettings.ThresholdAttribute.GetName())));
+			return false;
+		}
+
+		// Comparison between threshold and target data needs to be of the same type. So we have to make sure that we can
+		// request target type from threshold type. ie. We need to make sure we can broadcast threshold type to target type.
+		// For example: if target is double but threshold is int32, we can broadcast int32 to double, to compare a double with a double.
+		if (!PCG::Private::IsBroadcastable(InOutThresholdInfo.ThresholdAccessor->GetUnderlyingType(), TargetType))
+		{
+			UEnum* PCGDataTypeEnum = StaticEnum<EPCGMetadataTypes>();
+			FText ThresholdTypeName = FText::FromString(PCGDataTypeEnum ? PCGDataTypeEnum->GetNameStringByValue(InOutThresholdInfo.ThresholdAccessor->GetUnderlyingType()) : FString(TEXT("Unknown")));
+			FText InputTypeName = FText::FromString(PCGDataTypeEnum ? PCGDataTypeEnum->GetNameStringByValue(TargetType) : FString(TEXT("Unknown")));
+			PCGE_LOG_C(Warning, GraphAndLog, InContext, FText::Format(LOCTEXT("TypeConversionFailed", "Cannot use threshold type '{0}' on input target type '{1}'"),
+				ThresholdTypeName,
+				InputTypeName));
+			return false;
+		}
+
+		// And also validate that types are comparable
+		if (bCheckCompare)
+		{
+			bool bCanCompare = PCGMetadataAttribute::CallbackWithRightType(InOutThresholdInfo.ThresholdAccessor->GetUnderlyingType(), [](auto Dummy) -> bool
+			{
+				return PCG::Private::MetadataTraits<decltype(Dummy)>::CanCompare;
+			});
+
+			if (!bCanCompare)
+			{
+				PCGE_LOG_C(Warning, GraphAndLog, InContext, LOCTEXT("TypeComparisonFailed", "Cannot compare target type"));
+				return false;
+			}
+		}
+
+		// Check that if we have points as threshold, that the point data has the same number of point that the input data
+		if (InOutThresholdInfo.ThresholdSpatialData != nullptr && !InOutThresholdInfo.bUseSpatialQuery)
+		{
+			if (InOutThresholdInfo.ThresholdKeys->GetNum() != InputData->GetPoints().Num())
+			{
+				PCGE_LOG_C(Warning, GraphAndLog, InContext, FText::Format(LOCTEXT("InvalidNumberOfThresholdPoints", "Threshold point data doesn't have the same number of points ({0}) than the input data ({1})."), InOutThresholdInfo.ThresholdKeys->GetNum(), InputData->GetPoints().Num()));
+				return false;
+			}
+		}
+
+		return true;
+	}
 }
 
 TArray<FPCGPinProperties> UPCGPointFilterSettings::InputPinProperties() const
@@ -71,6 +231,9 @@ TArray<FPCGPinProperties> UPCGPointFilterSettings::InputPinProperties() const
 	if (!bUseConstantThreshold)
 	{
 		PinProperties.Emplace(PCGPointFilterConstants::FilterLabel, EPCGDataType::Any, /*bInAllowMultipleConnections=*/ false);
+#if WITH_EDITOR
+		PinProperties.Last().Tooltip = PCGPointFilterConstants::FilterPinTooltip;
+#endif // WITH_EDITOR
 	}
 
 	return PinProperties;
@@ -135,7 +298,7 @@ PRAGMA_DISABLE_DEPRECATION_WARNINGS
 		}
 	}
 
-	if (bUseConstantThreshold && 
+	if (bUseConstantThreshold &&
 		((ThresholdConstantType_DEPRECATED != EPCGPointFilterConstantType::Float)
 		|| (FloatConstant_DEPRECATED != 0.0f)))
 	{
@@ -177,34 +340,65 @@ PRAGMA_ENABLE_DEPRECATION_WARNINGS
 #endif // WITH_EDITOR
 }
 
-bool FPCGPointFilterElement::ExecuteInternal(FPCGContext* Context) const
+FPCGElementPtr UPCGPointFilterRangeSettings::CreateElement() const
 {
-	TRACE_CPUPROFILER_EVENT_SCOPE(FPCGPointFilterElement::Execute);
-	check(Context);
+	return MakeShared<FPCGPointFilterRangeElement>();
+}
 
-#if !WITH_EDITOR
-	const bool bHasInFilterOutputPin = Context->Node && Context->Node->IsOutputPinConnected(PCGPointFilterConstants::InFilterLabel);
-	const bool bHasOutsideFilterOutputPin = Context->Node && Context->Node->IsOutputPinConnected(PCGPointFilterConstants::OutFilterLabel);
+UPCGPointFilterRangeSettings::UPCGPointFilterRangeSettings()
+	: UPCGSettings()
+{
+	// Previous default object was: density for both selectors
+	// Recreate the same default
+	TargetAttribute.SetPointProperty(EPCGPointProperties::Density);
+	MinThreshold.ThresholdAttribute.SetPointProperty(EPCGPointProperties::Density);
+	MaxThreshold.ThresholdAttribute.SetPointProperty(EPCGPointProperties::Density);
+}
 
-	// Early out - only in non-editor builds, otherwise we will potentially poison the cache, since it is input-driven
-	if (!bHasInFilterOutputPin && !bHasOutsideFilterOutputPin)
+TArray<FPCGPinProperties> UPCGPointFilterRangeSettings::InputPinProperties() const
+{
+	TArray<FPCGPinProperties> PinProperties;
+	PinProperties.Emplace(PCGPointFilterConstants::DataToFilterLabel, EPCGDataType::Point);
+
+	if (!MinThreshold.bUseConstantThreshold)
 	{
-		return true;
+		PinProperties.Emplace(PCGPointFilterConstants::FilterMinLabel, EPCGDataType::Any, /*bInAllowMultipleConnections=*/ false);
+#if WITH_EDITOR
+		PinProperties.Last().Tooltip = PCGPointFilterConstants::FilterPinTooltip;
+#endif // WITH_EDITOR
 	}
-#endif
 
-	const UPCGPointFilterSettings* Settings = Context->GetInputSettings<UPCGPointFilterSettings>();
-	check(Settings);
+	if (!MaxThreshold.bUseConstantThreshold)
+	{
+		PinProperties.Emplace(PCGPointFilterConstants::FilterMaxLabel, EPCGDataType::Any, /*bInAllowMultipleConnections=*/ false);
+#if WITH_EDITOR
+		PinProperties.Last().Tooltip = PCGPointFilterConstants::FilterPinTooltip;
+#endif // WITH_EDITOR
+	}
+
+	return PinProperties;
+}
+
+TArray<FPCGPinProperties> UPCGPointFilterRangeSettings::OutputPinProperties() const
+{
+	TArray<FPCGPinProperties> PinProperties;
+	PinProperties.Emplace(PCGPointFilterConstants::InFilterLabel, EPCGDataType::Point);
+	PinProperties.Emplace(PCGPointFilterConstants::OutFilterLabel, EPCGDataType::Point);
+
+	return PinProperties;
+}
+
+bool FPCGPointFilterElementBase::DoFiltering(FPCGContext* Context, EPCGPointFilterOperator InOperation, const FPCGAttributePropertySelector& TargetAttribute, const FPCGPointFilterThresholdSettings& FirstThreshold, const FPCGPointFilterThresholdSettings* SecondThreshold) const
+{
+	TRACE_CPUPROFILER_EVENT_SCOPE(FPCGPointFilterElementBase::DoFiltering);
+	check(Context);
 
 	TArray<FPCGTaggedData> DataToFilter = Context->InputData.GetInputsByPin(PCGPointFilterConstants::DataToFilterLabel);
 	TArray<FPCGTaggedData>& Outputs = Context->OutputData.TaggedData;
-	TArray<FPCGTaggedData> FilterData = Context->InputData.GetInputsByPin(PCGPointFilterConstants::FilterLabel);
+	TArray<FPCGTaggedData> FirstFilterData = SecondThreshold ? Context->InputData.GetInputsByPin(PCGPointFilterConstants::FilterMinLabel) : Context->InputData.GetInputsByPin(PCGPointFilterConstants::FilterLabel);
+	TArray<FPCGTaggedData> SecondFilterData = SecondThreshold ? Context->InputData.GetInputsByPin(PCGPointFilterConstants::FilterMaxLabel) : TArray<FPCGTaggedData>{};
 
-	const EPCGPointFilterOperator Operator = Settings->Operator;
-	const bool bUseConstantThreshold = Settings->bUseConstantThreshold;
-	// If the filter is a param data, disable bUseSpatialQuery. Therefore make it non-const, to be modified below.
-	bool bUseSpatialQuery = Settings->bUseSpatialQuery;
-	// TODO: allow selector overrides for ThresholdAttribute and TargetAttribute
+	const EPCGPointFilterOperator Operator = InOperation;
 
 	// If there is no input, do nothing
 	if (DataToFilter.IsEmpty())
@@ -212,70 +406,24 @@ bool FPCGPointFilterElement::ExecuteInternal(FPCGContext* Context) const
 		return true;
 	}
 
-	TUniquePtr<const IPCGAttributeAccessor> ThresholdAccessor;
-	TUniquePtr<const IPCGAttributeAccessorKeys> ThresholdKeys;
-	bool bUseInputDataForThreshold = false;
-
-	UPCGPointData* ThresholdPointData = nullptr;
-	const UPCGSpatialData* ThresholdSpatialData = nullptr;
-
-	if (bUseConstantThreshold)
+	// Only support second threshold with the InRange Operation.
+	// We can't have a second threshold if it is not InRange
+	if (!ensure((SecondThreshold != nullptr) == (InOperation == EPCGPointFilterOperator::InRange)))
 	{
-		auto ConstantThreshold = [&ThresholdAccessor, &ThresholdKeys](auto&& Value)
-		{
-			using ConstantType = std::decay_t<decltype(Value)>;
-
-			ThresholdAccessor = MakeUnique<FPCGConstantValueAccessor<ConstantType>>(std::forward<decltype(Value)>(Value));
-			// Dummy keys
-			ThresholdKeys = MakeUnique<FPCGAttributeAccessorKeysSingleObjectPtr<void>>();
-		};
-
-		Settings->AttributeTypes.Dispatcher(ConstantThreshold);
+		return true;
 	}
-	else
-	{
-		if (!FilterData.IsEmpty())
-		{
-			const UPCGData* ThresholdData = FilterData[0].Data;
 
-			if (const UPCGSpatialData* TentativeThresholdSpatialData = Cast<const UPCGSpatialData>(ThresholdData))
-			{
-				// If the threshold is spatial, and we use spatial query, it means we'll have to sample points.
-				// Don't create an accessor yet (ThresholdData = nullptr), it will be created further down.
-				// Otherwise, we convert it to point data and create the accessor.
-				if (bUseSpatialQuery)
-				{
-					ThresholdSpatialData = TentativeThresholdSpatialData;
-					ThresholdData = nullptr;
-				}
-				else
-				{
-					ThresholdSpatialData = TentativeThresholdSpatialData->ToPointData();
-					if (!ThresholdSpatialData)
-					{
-						PCGE_LOG(Error, GraphAndLog, LOCTEXT("CannotGetPointData", "Unable to get Point data from filter input"));
-						return true;
-					}
+	PCGPointFilterHelpers::ThresholdInfo FirstThresholdInfo;
+	PCGPointFilterHelpers::ThresholdInfo SecondThresholdInfo;
 
-					ThresholdData = ThresholdSpatialData;
-				}
-			}
-			else if (ThresholdData->IsA<UPCGParamData>())
-			{
-				// Disable spatial query as it doesn't make sense
-				bUseSpatialQuery = false;
-			}
+	if (!PCGPointFilterHelpers::InitialPrepareThresholdInfo(Context, FirstFilterData, FirstThreshold, FirstThresholdInfo))
+	{
+		return true;
+	}
 
-			if (ThresholdData)
-			{
-				ThresholdAccessor = PCGAttributeAccessorHelpers::CreateConstAccessor(ThresholdData, Settings->ThresholdAttribute);
-				ThresholdKeys = PCGAttributeAccessorHelpers::CreateConstKeys(ThresholdData, Settings->ThresholdAttribute);
-			}
-		}
-		else
-		{
-			bUseInputDataForThreshold = true;
-		}
+	if (SecondThreshold && !PCGPointFilterHelpers::InitialPrepareThresholdInfo(Context, SecondFilterData, *SecondThreshold, SecondThresholdInfo))
+	{
+		return true;
 	}
 
 	for (const FPCGTaggedData& Input : DataToFilter)
@@ -311,67 +459,44 @@ bool FPCGPointFilterElement::ExecuteInternal(FPCGContext* Context) const
 
 		const TArray<FPCGPoint>& OriginalPoints = OriginalData->GetPoints();
 
-		TUniquePtr<const IPCGAttributeAccessor> TargetAccessor = PCGAttributeAccessorHelpers::CreateConstAccessor(OriginalData, Settings->TargetAttribute);
-		TUniquePtr<const IPCGAttributeAccessorKeys> TargetKeys = PCGAttributeAccessorHelpers::CreateConstKeys(OriginalData, Settings->TargetAttribute);
+		TUniquePtr<const IPCGAttributeAccessor> TargetAccessor = PCGAttributeAccessorHelpers::CreateConstAccessor(OriginalData, TargetAttribute);
+		TUniquePtr<const IPCGAttributeAccessorKeys> TargetKeys = PCGAttributeAccessorHelpers::CreateConstKeys(OriginalData, TargetAttribute);
 
 		if (!TargetAccessor.IsValid() || !TargetKeys.IsValid())
 		{
-			PCGE_LOG(Warning, GraphAndLog, FText::Format(LOCTEXT("TargetMissingAttribute", "TargetData doesn't have target attribute/property '{0}'"), FText::FromName(Settings->TargetAttribute.GetName())));
+			PCGE_LOG(Warning, GraphAndLog, FText::Format(LOCTEXT("TargetMissingAttribute", "TargetData doesn't have target attribute/property '{0}'"), FText::FromName(TargetAttribute.GetName())));
 			ForwardInputToOutFilterPin();
 			continue;
 		}
 
-		if (bUseInputDataForThreshold)
-		{
-			// If we have no threshold accessor, we use the same data as input
-			ThresholdAccessor = PCGAttributeAccessorHelpers::CreateConstAccessor(OriginalData, Settings->ThresholdAttribute);
-			ThresholdKeys = PCGAttributeAccessorHelpers::CreateConstKeys(OriginalData, Settings->ThresholdAttribute);
-		}
-		else if (ThresholdSpatialData != nullptr && bUseSpatialQuery)
-		{
-			// Reset the point data and reserving some points
-			// No need to reserve the full number of points, since we'll go by chunk
-			// Only allocate the chunk size
-			ThresholdPointData = NewObject<UPCGPointData>();
-			ThresholdPointData->InitializeFromData(ThresholdSpatialData);
-			ThresholdPointData->GetMutablePoints().SetNum(PCGPointFilterConstants::ChunkSize);
+		const int16 TargetType = TargetAccessor->GetUnderlyingType();
+		const bool bCheckCompare = (Operator != EPCGPointFilterOperator::Equal) && (Operator != EPCGPointFilterOperator::NotEqual);
 
-			// Accessor will be valid, but keys will point to default points. But since it is a view, it will be updated when we sample the points.
-			ThresholdAccessor = PCGAttributeAccessorHelpers::CreateConstAccessor(ThresholdPointData, Settings->ThresholdAttribute);
-			ThresholdKeys = PCGAttributeAccessorHelpers::CreateConstKeys(ThresholdPointData, Settings->ThresholdAttribute);
+		if (!PCGPointFilterHelpers::PrepareThresholdInfoFromInput(Context, OriginalData, FirstThreshold, FirstThresholdInfo, TargetType, bCheckCompare))
+		{
+			ForwardInputToInFilterPin();
+			continue;
 		}
 
-		if (!ThresholdAccessor.IsValid() || !ThresholdKeys.IsValid())
+		if (SecondThreshold && !PCGPointFilterHelpers::PrepareThresholdInfoFromInput(Context, OriginalData, *SecondThreshold, SecondThresholdInfo, TargetType, bCheckCompare, &FirstThresholdInfo))
 		{
-			PCGE_LOG(Warning, GraphAndLog, FText::Format(LOCTEXT("AttributeMissingForFilter", "DataToFilter does not have '{0}' threshold attribute/property"), FText::FromName(Settings->ThresholdAttribute.GetName())));
 			ForwardInputToInFilterPin();
 			continue;
 		}
 
-		// Comparison between threshold and target data needs to be of the same type. So we have to make sure that we can
-		// request target type from threshold type. ie. We need to make sure we can broadcast threshold type to target type.
-		// For example: if target is double but threshold is int32, we can broadcast int32 to double, to compare a double with a double.
-		if (!PCG::Private::IsBroadcastable(ThresholdAccessor->GetUnderlyingType(), TargetAccessor->GetUnderlyingType()))
+
+		if (!PCG::Private::IsBroadcastable(FirstThresholdInfo.ThresholdAccessor->GetUnderlyingType(), TargetAccessor->GetUnderlyingType()))
 		{
 			PCGE_LOG(Warning, GraphAndLog, LOCTEXT("TypeConversionFailed", "Cannot broadcast threshold type to target type"));
 			ForwardInputToInFilterPin();
 			continue;
 		}
 
-		// And also validate that types are comparable
-		if (Operator != EPCGPointFilterOperator::Equal && Operator != EPCGPointFilterOperator::NotEqual)
+		if (SecondThresholdInfo.ThresholdAccessor.IsValid() && !PCG::Private::IsBroadcastable(SecondThresholdInfo.ThresholdAccessor->GetUnderlyingType(), TargetAccessor->GetUnderlyingType()))
 		{
-			bool bCanCompare = PCGMetadataAttribute::CallbackWithRightType(TargetAccessor->GetUnderlyingType(), [](auto Dummy) -> bool
-			{
-				return PCG::Private::MetadataTraits<decltype(Dummy)>::CanCompare;
-			});
-
-			if (!bCanCompare)
-			{
-				PCGE_LOG(Warning, GraphAndLog, LOCTEXT("TypeComparisonFailed", "Cannot compare target type"));
-				ForwardInputToOutFilterPin();
-				continue;
-			}
+			PCGE_LOG(Warning, GraphAndLog, LOCTEXT("TypeConversionFailed", "Cannot broadcast threshold type to target type"));
+			ForwardInputToInFilterPin();
+			continue;
 		}
 
 		UPCGPointData* InFilterData = NewObject<UPCGPointData>();
@@ -385,7 +510,7 @@ bool FPCGPointFilterElement::ExecuteInternal(FPCGContext* Context) const
 		InFilterPoints.Reserve(OriginalPoints.Num());
 		OutFilterPoints.Reserve(OriginalPoints.Num());
 
-		auto Operation = [&InFilterPoints, &OutFilterPoints, &Operator, &ThresholdAccessor, &ThresholdKeys, &TargetAccessor, &TargetKeys, &OriginalPoints, ThresholdPointData, ThresholdSpatialData](auto Dummy) -> bool
+		auto Operation = [&InFilterPoints, &OutFilterPoints, &Operator, &FirstThresholdInfo, &SecondThresholdInfo, &TargetAccessor, &TargetKeys, &OriginalPoints, &FirstThreshold, &SecondThreshold](auto Dummy) -> bool
 		{
 			using Type = decltype(Dummy);
 
@@ -397,11 +522,13 @@ bool FPCGPointFilterElement::ExecuteInternal(FPCGContext* Context) const
 			}
 
 			TArray<Type, TInlineAllocator<PCGPointFilterConstants::ChunkSize>> TargetValues;
-			TArray<Type, TInlineAllocator<PCGPointFilterConstants::ChunkSize>> ThresholdValues;
+			TArray<Type, TInlineAllocator<PCGPointFilterConstants::ChunkSize>> FirstThresholdValues;
+			TArray<Type, TInlineAllocator<PCGPointFilterConstants::ChunkSize>> SecondThresholdValues;
 			TArray<bool, TInlineAllocator<PCGPointFilterConstants::ChunkSize>> SkipTests;
 			TargetValues.SetNum(PCGPointFilterConstants::ChunkSize);
-			ThresholdValues.SetNum(PCGPointFilterConstants::ChunkSize);
-			SkipTests.SetNum(PCGPointFilterConstants::ChunkSize);
+			FirstThresholdValues.SetNum(PCGPointFilterConstants::ChunkSize);
+			SecondThresholdValues.SetNum(PCGPointFilterConstants::ChunkSize);
+			SkipTests.SetNumZeroed(PCGPointFilterConstants::ChunkSize); // All initialized to false
 
 			const int32 NumberOfIterations = (NumberOfEntries + PCGPointFilterConstants::ChunkSize - 1) / PCGPointFilterConstants::ChunkSize;
 
@@ -410,40 +537,63 @@ bool FPCGPointFilterElement::ExecuteInternal(FPCGContext* Context) const
 				const int32 StartIndex = i * PCGPointFilterConstants::ChunkSize;
 				const int32 Range = FMath::Min(NumberOfEntries - StartIndex, PCGPointFilterConstants::ChunkSize);
 				TArrayView<Type> TargetView(TargetValues.GetData(), Range);
-				TArrayView<Type> ThresholdView(ThresholdValues.GetData(), Range);
+				TArrayView<Type> FirstThresholdView(FirstThresholdValues.GetData(), Range);
+				TArrayView<Type> SecondThresholdView(SecondThresholdValues.GetData(), Range);
 
 				// Sampling the points if needed
-				if (ThresholdPointData != nullptr)
+				auto SamplePointData = [Range, StartIndex, &OriginalPoints, &SkipTests](UPCGPointData* InPointData, const UPCGSpatialData* InSpatialData)
 				{
-					// Threshold points only have "ChunkSize" points.
-					TArray<FPCGPoint>& ThresholdPoints = ThresholdPointData->GetMutablePoints();
-					for (int32 j = 0; j < Range; ++j)
+					if (InPointData != nullptr)
 					{
-						FPCGPoint ThresholdPoint;
-						const FPCGPoint& SourcePoint = OriginalPoints[StartIndex + j];
-						if (ThresholdSpatialData->SamplePoint(SourcePoint.Transform, SourcePoint.GetLocalBounds(), ThresholdPoint, ThresholdPointData->Metadata))
-						{
-							ThresholdPoints[j] = ThresholdPoint;
-							SkipTests[j] = false;
-						}
-						else
+						// Threshold points only have "ChunkSize" points.
+						TArray<FPCGPoint>& ThresholdPoints = InPointData->GetMutablePoints();
+						for (int32 j = 0; j < Range; ++j)
 						{
-							SkipTests[j] = true;
+							FPCGPoint ThresholdPoint;
+							const FPCGPoint& SourcePoint = OriginalPoints[StartIndex + j];
+							// If we already mark this point to skip test, don't even try to sample.
+							// If the sample fails, mark the point to skip test.
+							if (!SkipTests[j] && InSpatialData->SamplePoint(SourcePoint.Transform, SourcePoint.GetLocalBounds(), ThresholdPoint, InPointData->Metadata))
+							{
+								ThresholdPoints[j] = ThresholdPoint;
+							}
+							else
+							{
+								SkipTests[j] = true;
+							}
 						}
 					}
+				};
+
+				SamplePointData(FirstThresholdInfo.ThresholdPointData, FirstThresholdInfo.ThresholdSpatialData);
+
+				if (FirstThresholdInfo.ThresholdPointData != SecondThresholdInfo.ThresholdPointData)
+				{
+					SamplePointData(SecondThresholdInfo.ThresholdPointData, SecondThresholdInfo.ThresholdSpatialData);
 				}
 
 				// If ThresholdView point on ThresholdPointData points, there are only "ChunkSize" points in it.
 				// But it wraps around, and since StartIndex is a multiple of "ChunkSize", we'll always start at point 0, as wanted. 
 				if (!TargetAccessor->GetRange(TargetView, StartIndex, *TargetKeys) ||
-					!ThresholdAccessor->GetRange(ThresholdView, StartIndex, *ThresholdKeys, EPCGAttributeAccessorFlags::AllowBroadcast))
+					!FirstThresholdInfo.ThresholdAccessor->GetRange(FirstThresholdView, StartIndex, *FirstThresholdInfo.ThresholdKeys, EPCGAttributeAccessorFlags::AllowBroadcast) ||
+					(SecondThresholdInfo.ThresholdAccessor.IsValid() && !SecondThresholdInfo.ThresholdAccessor->GetRange(SecondThresholdView, StartIndex, *SecondThresholdInfo.ThresholdKeys, EPCGAttributeAccessorFlags::AllowBroadcast)))
 				{
 					return false;
 				}
 
 				for (int32 j = 0; j < Range; ++j)
 				{
-					if (SkipTests[j] || PCGPointFilterHelpers::ApplyCompare(TargetValues[j], ThresholdValues[j], Operator))
+					if (SkipTests[j])
+					{
+						InFilterPoints.Add(OriginalPoints[StartIndex + j]);
+						continue;
+					}
+
+					const bool bShouldKeep = (Operator == EPCGPointFilterOperator::InRange ? 
+						PCGPointFilterHelpers::ApplyRange(TargetValues[j], FirstThresholdValues[j], SecondThresholdValues[j], FirstThreshold.bInclusive, SecondThreshold->bInclusive) : 
+						PCGPointFilterHelpers::ApplyCompare(TargetValues[j], FirstThresholdValues[j], Operator));
+
+					if (bShouldKeep)
 					{
 						InFilterPoints.Add(OriginalPoints[StartIndex + j]);
 					}
@@ -474,4 +624,52 @@ bool FPCGPointFilterElement::ExecuteInternal(FPCGContext* Context) const
 	return true;
 }
 
+bool FPCGPointFilterElement::ExecuteInternal(FPCGContext* Context) const
+{
+	TRACE_CPUPROFILER_EVENT_SCOPE(FPCGPointFilterElement::Execute);
+
+#if !WITH_EDITOR
+	const bool bHasInFilterOutputPin = Context->Node && Context->Node->IsOutputPinConnected(PCGPointFilterConstants::InFilterLabel);
+	const bool bHasOutsideFilterOutputPin = Context->Node && Context->Node->IsOutputPinConnected(PCGPointFilterConstants::OutFilterLabel);
+
+	// Early out - only in non-editor builds, otherwise we will potentially poison the cache, since it is input-driven
+	if (!bHasInFilterOutputPin && !bHasOutsideFilterOutputPin)
+	{
+		return true;
+	}
+#endif
+
+	const UPCGPointFilterSettings* Settings = Context->GetInputSettings<UPCGPointFilterSettings>();
+	check(Settings);
+
+	FPCGPointFilterThresholdSettings ThresholdSettings{};
+	ThresholdSettings.bUseConstantThreshold = Settings->bUseConstantThreshold;
+	ThresholdSettings.bUseSpatialQuery = Settings->bUseSpatialQuery;
+	ThresholdSettings.ThresholdAttribute = Settings->ThresholdAttribute;
+	ThresholdSettings.AttributeTypes = Settings->AttributeTypes;
+
+	return DoFiltering(Context, Settings->Operator, Settings->TargetAttribute, ThresholdSettings);
+}
+
+bool FPCGPointFilterRangeElement::ExecuteInternal(FPCGContext* Context) const
+{
+	TRACE_CPUPROFILER_EVENT_SCOPE(FPCGPointFilterRangeElement::Execute);
+
+#if !WITH_EDITOR
+	const bool bHasInFilterOutputPin = Context->Node && Context->Node->IsOutputPinConnected(PCGPointFilterConstants::InFilterLabel);
+	const bool bHasOutsideFilterOutputPin = Context->Node && Context->Node->IsOutputPinConnected(PCGPointFilterConstants::OutFilterLabel);
+
+	// Early out - only in non-editor builds, otherwise we will potentially poison the cache, since it is input-driven
+	if (!bHasInFilterOutputPin && !bHasOutsideFilterOutputPin)
+	{
+		return true;
+	}
+#endif
+
+	const UPCGPointFilterRangeSettings* Settings = Context->GetInputSettings<UPCGPointFilterRangeSettings>();
+	check(Settings);
+
+	return DoFiltering(Context, EPCGPointFilterOperator::InRange, Settings->TargetAttribute, Settings->MinThreshold, &Settings->MaxThreshold);
+}
+
 #undef LOCTEXT_NAMESPACE
diff --git a/Engine/Plugins/Experimental/PCG/Source/PCG/Private/Tests/Elements/PCGPointFilterTest.cpp b/Engine/Plugins/Experimental/PCG/Source/PCG/Private/Tests/Elements/PCGPointFilterTest.cpp
new file mode 100644
--- /dev/null
+++ b/Engine/Plugins/Experimental/PCG/Source/PCG/Private/Tests/Elements/PCGPointFilterTest.cpp
@@ -0,0 +1,147 @@
+// Copyright Epic Games, Inc. All Rights Reserved.
+
+#include "Tests/Determinism/PCGDeterminismTestsCommon.h"
+
+#include "PCGComponent.h"
+#include "PCGContext.h"
+#include "PCGPoint.h"
+#include "Data/PCGPointData.h"
+#include "Elements/PCGPointFilter.h"
+
+#include "Math/RandomStream.h"
+
+
+IMPLEMENT_CUSTOM_SIMPLE_AUTOMATION_TEST(FPCGPointFilterDensity, FPCGTestBaseClass, "pcg.tests.PointFilter.Density", PCGTestsCommon::TestFlags)
+IMPLEMENT_CUSTOM_SIMPLE_AUTOMATION_TEST(FPCGPointFilterDensityRange, FPCGTestBaseClass, "pcg.tests.PointFilter.DensityRange", PCGTestsCommon::TestFlags)
+
+namespace PCGPointFilterTest
+{
+	const FName InsideFilterLabel = TEXT("InsideFilter");
+	const FName OutsideFilterLabel = TEXT("OutsideFilter");
+
+	UPCGPointData* GeneratePointDataWithRandomDensity(int32 InNumPoints, int32 InRandomSeed)
+	{
+		UPCGPointData* PointData = PCGTestsCommon::CreateEmptyPointData();
+		TArray<FPCGPoint>& Points = PointData->GetMutablePoints();
+
+		FRandomStream RandomSource(InRandomSeed);
+
+		for (int I = 0; I < InNumPoints; ++I)
+		{
+			FPCGPoint& Point = Points.Emplace_GetRef(FTransform(), RandomSource.FRand(), I);
+		}
+
+		return PointData;
+	}
+}
+
+bool FPCGPointFilterDensity::RunTest(const FString& Parameters)
+{
+	PCGTestsCommon::FTestData TestData(PCGDeterminismTests::Defaults::Seed);
+	PCGDeterminismTests::GenerateSettings<UPCGPointFilterSettings>(TestData);
+	UPCGPointFilterSettings* Settings = CastChecked<UPCGPointFilterSettings>(TestData.Settings);
+
+	FPCGElementPtr TestElement = TestData.Settings->GetElement();
+
+	static const int32 NumPoints = 50;
+	static const float DensityThreshold = 0.5f;
+
+	Settings->Operator = EPCGPointFilterOperator::Lesser;
+	Settings->TargetAttribute.SetPointProperty(EPCGPointProperties::Density);
+	Settings->bUseConstantThreshold = true;
+	Settings->AttributeTypes.FloatValue = DensityThreshold;
+	Settings->AttributeTypes.Type = EPCGMetadataTypes::Float;
+
+	FPCGTaggedData& TaggedData = TestData.InputData.TaggedData.Emplace_GetRef(FPCGTaggedData());
+	TaggedData.Pin = PCGPinConstants::DefaultInputLabel;
+	TaggedData.Data = PCGPointFilterTest::GeneratePointDataWithRandomDensity(NumPoints, PCGDeterminismTests::Defaults::Seed);
+
+	TUniquePtr<FPCGContext> Context = TestData.InitializeTestContext();
+
+	while (!TestElement->Execute(Context.Get())) {}
+
+	// We should have outputs on both
+	TArray<FPCGTaggedData> InFilterOutput = Context->OutputData.GetInputsByPin(PCGPointFilterTest::InsideFilterLabel);
+	TArray<FPCGTaggedData> OutFilterOutput = Context->OutputData.GetInputsByPin(PCGPointFilterTest::OutsideFilterLabel);
+
+	UTEST_EQUAL(TEXT("InFilter pin has 1 output"), InFilterOutput.Num(), 1);
+	UTEST_EQUAL(TEXT("OutFilter pin has 1 output"), OutFilterOutput.Num(), 1);
+
+	// Making sure the output is a PointData
+	const UPCGPointData* InFilterPointData = Cast<UPCGPointData>(InFilterOutput[0].Data);
+	const UPCGPointData* OutFilterPointData = Cast<UPCGPointData>(OutFilterOutput[0].Data);
+
+	UTEST_NOT_NULL(TEXT("InFilter data is a point data"), InFilterPointData);
+	UTEST_NOT_NULL(TEXT("OutFilter data is a point data"), OutFilterPointData);
+
+	// Verifying that all points have the right density
+	for (const FPCGPoint& Point : InFilterPointData->GetPoints())
+	{
+		UTEST_TRUE(*FString::Printf(TEXT("Point has a density (%f) lower than the threshold (%f)"), Point.Density, DensityThreshold), Point.Density < DensityThreshold);
+	}
+
+	for (const FPCGPoint& Point : OutFilterPointData->GetPoints())
+	{
+		UTEST_TRUE(*FString::Printf(TEXT("Point has a density (%f) higher or equal than the threshold (%f)"), Point.Density, DensityThreshold), Point.Density >= DensityThreshold);
+	}
+
+	return true;
+}
+
+bool FPCGPointFilterDensityRange::RunTest(const FString& Parameters)
+{
+	PCGTestsCommon::FTestData TestData(PCGDeterminismTests::Defaults::Seed);
+	PCGDeterminismTests::GenerateSettings<UPCGPointFilterRangeSettings>(TestData);
+	UPCGPointFilterRangeSettings* Settings = CastChecked<UPCGPointFilterRangeSettings>(TestData.Settings);
+
+	FPCGElementPtr TestElement = TestData.Settings->GetElement();
+
+	static const int32 NumPoints = 50;
+	static const float DensityMinThreshold = 0.3f;
+	static const float DensityMaxThreshold = 0.8f;
+
+	Settings->TargetAttribute.SetPointProperty(EPCGPointProperties::Density);
+
+	Settings->MinThreshold.bUseConstantThreshold = true;
+	Settings->MinThreshold.AttributeTypes.FloatValue = DensityMinThreshold;
+	Settings->MinThreshold.AttributeTypes.Type = EPCGMetadataTypes::Float;
+
+	Settings->MaxThreshold.bUseConstantThreshold = true;
+	Settings->MaxThreshold.AttributeTypes.FloatValue = DensityMaxThreshold;
+	Settings->MaxThreshold.AttributeTypes.Type = EPCGMetadataTypes::Float;
+
+	FPCGTaggedData& TaggedData = TestData.InputData.TaggedData.Emplace_GetRef(FPCGTaggedData());
+	TaggedData.Pin = PCGPinConstants::DefaultInputLabel;
+	TaggedData.Data = PCGPointFilterTest::GeneratePointDataWithRandomDensity(NumPoints, PCGDeterminismTests::Defaults::Seed);
+
+	TUniquePtr<FPCGContext> Context = TestData.InitializeTestContext();
+
+	while (!TestElement->Execute(Context.Get())) {}
+
+	// We should have outputs on both
+	TArray<FPCGTaggedData> InFilterOutput = Context->OutputData.GetInputsByPin(PCGPointFilterTest::InsideFilterLabel);
+	TArray<FPCGTaggedData> OutFilterOutput = Context->OutputData.GetInputsByPin(PCGPointFilterTest::OutsideFilterLabel);
+
+	UTEST_EQUAL(TEXT("InFilter pin has 1 output"), InFilterOutput.Num(), 1);
+	UTEST_EQUAL(TEXT("OutFilter pin has 1 output"), OutFilterOutput.Num(), 1);
+
+	// Making sure the output is a PointData
+	const UPCGPointData* InFilterPointData = Cast<UPCGPointData>(InFilterOutput[0].Data);
+	const UPCGPointData* OutFilterPointData = Cast<UPCGPointData>(OutFilterOutput[0].Data);
+
+	UTEST_NOT_NULL(TEXT("InFilter data is a point data"), InFilterPointData);
+	UTEST_NOT_NULL(TEXT("OutFilter data is a point data"), OutFilterPointData);
+
+	// Verifying that all points have the right density
+	for (const FPCGPoint& Point : InFilterPointData->GetPoints())
+	{
+		UTEST_TRUE(*FString::Printf(TEXT("Point has a density (%f) within the range ([%f, %f])"), Point.Density, DensityMinThreshold, DensityMaxThreshold), (Point.Density >= DensityMinThreshold) && (Point.Density <= DensityMaxThreshold));
+	}
+
+	for (const FPCGPoint& Point : OutFilterPointData->GetPoints())
+	{
+		UTEST_TRUE(*FString::Printf(TEXT("Point has a density (%f) outside the range ([%f, %f])"), Point.Density, DensityMinThreshold, DensityMaxThreshold), (Point.Density < DensityMinThreshold) || (Point.Density > DensityMaxThreshold));
+	}
+
+	return true;
+}
diff --git a/Engine/Plugins/Experimental/PCG/Source/PCG/Public/Elements/PCGPointFilter.h b/Engine/Plugins/Experimental/PCG/Source/PCG/Public/Elements/PCGPointFilter.h
--- a/Engine/Plugins/Experimental/PCG/Source/PCG/Public/Elements/PCGPointFilter.h
+++ b/Engine/Plugins/Experimental/PCG/Source/PCG/Public/Elements/PCGPointFilter.h
@@ -43,7 +43,31 @@ enum class EPCGPointFilterOperator : uint8
 	Lesser UMETA(DisplayName="<"),
 	LesserOrEqual UMETA(DisplayName="<="),
 	Equal UMETA(DisplayName="="),
-	NotEqual UMETA(DisplayName="!=")
+	NotEqual UMETA(DisplayName="!="),
+	InRange UMETA(Hidden)
+};
+
+USTRUCT(BlueprintType)
+struct PCG_API FPCGPointFilterThresholdSettings
+{
+	GENERATED_BODY()
+
+	/** If the threshold in included or excluded from the range. */
+	UPROPERTY(BlueprintReadWrite, EditAnywhere, Category = Settings, meta = (PCG_Overridable))
+	bool bInclusive = true;
+
+	UPROPERTY(BlueprintReadWrite, EditAnywhere, Category = Settings, meta = (PCG_Overridable))
+	bool bUseConstantThreshold = false;
+
+	UPROPERTY(BlueprintReadWrite, EditAnywhere, Category = Settings, meta = (EditCondition = "!bUseConstantThreshold", EditConditionHides, PCG_NotOverridable))
+	FPCGAttributePropertySelector ThresholdAttribute;
+
+	/** If the threshold data is Point data, it will sample input points in threshold data. Always true with Spatial data.*/
+	UPROPERTY(BlueprintReadWrite, EditAnywhere, Category = Settings, meta = (EditCondition = "!bUseConstantThreshold", EditConditionHides, PCG_Overridable))
+	bool bUseSpatialQuery = true;
+
+	UPROPERTY(BlueprintReadWrite, EditAnywhere, Category = Settings, meta = (EditCondition = "bUseConstantThreshold", EditConditionHides, ShowOnlyInnerProperties, DisplayAfter = "bUseConstantThreshold", PCG_NotOverridable))
+	FPCGMetadataTypesConstantStruct AttributeTypes;
 };
 
 /**
@@ -73,14 +97,13 @@ public:
 	//~Begin UPCGSettings interface
 #if WITH_EDITOR
 	virtual FName GetDefaultNodeName() const override { return FName(TEXT("PointFilter")); }
-	virtual FText GetDefaultNodeTitle() const override { return NSLOCTEXT("PCGPointFilterSettings", "NodeTitle", "Point Filter"); }
+	virtual FText GetDefaultNodeTitle() const override { return NSLOCTEXT("PCGPointFilterElement", "NodeTitle", "Point Filter"); }
 	virtual EPCGSettingsType GetType() const override { return EPCGSettingsType::Filter; }
 #endif
 
+protected:
 	virtual TArray<FPCGPinProperties> InputPinProperties() const override;
 	virtual TArray<FPCGPinProperties> OutputPinProperties() const override;
-
-protected:
 	virtual FPCGElementPtr CreateElement() const override;
 	//~End UPCGSettings interface
 
@@ -96,13 +119,14 @@ public:
 	UPROPERTY(BlueprintReadWrite, EditAnywhere, Category = Settings, meta = (PCG_Overridable))
 	bool bUseConstantThreshold = false;
 
-	UPROPERTY(BlueprintReadWrite, EditAnywhere, Category = Settings, meta = (EditCondition = "!bUseConstantThreshold", EditConditionHides))
+	UPROPERTY(BlueprintReadWrite, EditAnywhere, Category = Settings, meta = (EditCondition = "!bUseConstantThreshold", EditConditionHides, PCG_NotOverridable))
 	FPCGAttributePropertySelector ThresholdAttribute;
 
+	/** If the threshold data is Point data, it will sample input points in threshold data. Always true with Spatial data.*/
 	UPROPERTY(BlueprintReadWrite, EditAnywhere, Category = Settings, meta = (EditCondition = "!bUseConstantThreshold", EditConditionHides, PCG_Overridable))
 	bool bUseSpatialQuery = true;
 
-	UPROPERTY(BlueprintReadWrite, EditAnywhere, Category = Constant, meta = (EditCondition = "bUseConstantThreshold", EditConditionHides, ShowOnlyInnerProperties, DisplayAfter = "bUseConstantThreshold"))
+	UPROPERTY(BlueprintReadWrite, EditAnywhere, Category = Settings, meta = (EditCondition = "bUseConstantThreshold", EditConditionHides, ShowOnlyInnerProperties, DisplayAfter = "bUseConstantThreshold", PCG_NotOverridable))
 	FPCGMetadataTypesConstantStruct AttributeTypes;
 
 #if WITH_EDITORONLY_DATA
@@ -146,8 +170,69 @@ PRAGMA_ENABLE_DEPRECATION_WARNINGS
 #endif 
 };
 
-class FPCGPointFilterElement : public FSimplePCGElement
+
+/**
+* Point filter on range that allows to do "A op B" type filtering, where A is the input spatial data,
+* and B is either a constant, another spatial data, a param data (in filter) or the input itself.
+* The filtering can be done either on properties or attributes.
+* Some examples (that might not make sense, but are valid):
+* - Threshold on property by constant (A.Density in [0.2, 0.5])
+* - Threshold on attribute by constant (A.aaa in [0.4, 0.6])
+* - Threshold on property by metadata attribute(A.density in [B.bbmin, B.bbmax])
+* - Threshold on property by property(A.density in [B.position.x, B.steepness])
+* - Threshold on attribute by metadata attribute(A.aaa in [B.bbmin, B.bbmax])
+* - Threshold on attribute by property(A.aaa in [B.position, B.scale])
+*/
+UCLASS(BlueprintType, ClassGroup = (Procedural))
+class PCG_API UPCGPointFilterRangeSettings : public UPCGSettings
+{
+	GENERATED_BODY()
+
+public:
+	UPCGPointFilterRangeSettings();
+
+
+	//~Begin UPCGSettings interface
+#if WITH_EDITOR
+	virtual FName GetDefaultNodeName() const override { return FName(TEXT("PointFilterRange")); }
+	virtual FText GetDefaultNodeTitle() const override { return NSLOCTEXT("PCGPointFilterElement", "NodeTitleRange", "Point Filter Range"); }
+	virtual EPCGSettingsType GetType() const override { return EPCGSettingsType::Filter; }
+#endif
+
+protected:
+	virtual TArray<FPCGPinProperties> InputPinProperties() const override;
+	virtual TArray<FPCGPinProperties> OutputPinProperties() const override;
+	virtual FPCGElementPtr CreateElement() const override;
+	//~End UPCGSettings interface
+
+public:
+	/** Target property/attribute related properties */
+	UPROPERTY(BlueprintReadWrite, EditAnywhere, Category = Settings)
+	FPCGAttributePropertySelector TargetAttribute;
+
+	/** Threshold property/attribute/constant related properties */
+	UPROPERTY(BlueprintReadWrite, EditAnywhere, Category = Settings, meta = (PCG_Overridable))
+	FPCGPointFilterThresholdSettings MinThreshold;
+
+	UPROPERTY(BlueprintReadWrite, EditAnywhere, Category = Settings, meta = (PCG_Overridable))
+	FPCGPointFilterThresholdSettings MaxThreshold;
+};
+
+class FPCGPointFilterElementBase : public FSimplePCGElement
+{
+protected:
+	bool DoFiltering(FPCGContext* Context, EPCGPointFilterOperator InOperation, const FPCGAttributePropertySelector& TargetAttribute, const FPCGPointFilterThresholdSettings& FirstThreshold, const FPCGPointFilterThresholdSettings* SecondThreshold = nullptr) const;
+};
+
+class FPCGPointFilterElement : public FPCGPointFilterElementBase
 {
 protected:
 	virtual bool ExecuteInternal(FPCGContext* Context) const override;
 };
+
+class FPCGPointFilterRangeElement : public FPCGPointFilterElementBase
+{
+protected:
+	virtual bool ExecuteInternal(FPCGContext* Context) const override;
+};
+
diff --git a/Engine/Plugins/Experimental/PCG/Source/PCG/Public/Metadata/PCGMetadataAttributeTraits.h b/Engine/Plugins/Experimental/PCG/Source/PCG/Public/Metadata/PCGMetadataAttributeTraits.h
--- a/Engine/Plugins/Experimental/PCG/Source/PCG/Public/Metadata/PCGMetadataAttributeTraits.h
+++ b/Engine/Plugins/Experimental/PCG/Source/PCG/Public/Metadata/PCGMetadataAttributeTraits.h
@@ -99,6 +99,12 @@ namespace PCG
 			return IsOfTypes<Types...>(MetadataTypes<T>::Id);
 		}
 
+		inline FString GetTypeName(uint16 InType)
+		{
+			UEnum* PCGDataTypeEnum = StaticEnum<EPCGMetadataTypes>();
+			return PCGDataTypeEnum ? PCGDataTypeEnum->GetNameStringByValue(InType) : FString(TEXT("Unknown"));
+		}
+
 		// Wrapper around a standard 2-dimensional CArray that is constexpr, to know if a type is broadcastable to another.
 		// First index is the original type, second index is the wanted type. Returns true if we can broadcast first type into second type.
 		struct UBroadcastableTypes
