[PCG] Missing include

#rnx

diff --git a/Engine/Plugins/Experimental/PCG/Source/PCG/Public/Metadata/PCGMetadataAttributeTraits.h b/Engine/Plugins/Experimental/PCG/Source/PCG/Public/Metadata/PCGMetadataAttributeTraits.h
--- a/Engine/Plugins/Experimental/PCG/Source/PCG/Public/Metadata/PCGMetadataAttributeTraits.h
+++ b/Engine/Plugins/Experimental/PCG/Source/PCG/Public/Metadata/PCGMetadataAttributeTraits.h
@@ -3,6 +3,8 @@
 #pragma once
 
 #include "Math/Transform.h" // IWYU pragma: keep
+#include "UObject/Class.h"
+
 #include "PCGMetadataAttributeTraits.generated.h"
 
 UENUM(BlueprintType)
