[PCG] Improvement on gathering properties for Overrides

* BP now add the full property path as label for override pin in case of name clash
* Nested properties in BP need to be exposed/visible and editable to be overridable
* Properties for PropertyToParamData need to be exposed/visible to be added to the attribute set

diff --git a/Engine/Plugins/Experimental/PCG/Source/PCG/Private/Elements/PCGExecuteBlueprint.cpp b/Engine/Plugins/Experimental/PCG/Source/PCG/Private/Elements/PCGExecuteBlueprint.cpp
--- a/Engine/Plugins/Experimental/PCG/Source/PCG/Private/Elements/PCGExecuteBlueprint.cpp
+++ b/Engine/Plugins/Experimental/PCG/Source/PCG/Private/Elements/PCGExecuteBlueprint.cpp
@@ -582,7 +582,7 @@ TArray<FPCGSettingsOverridableParam> UPCGBlueprintSettings::GatherOverridablePar
 		{
 			PCGSettingsHelpers::FPCGGetAllOverridableParamsConfig Config;
 			Config.bExcludeSuperProperties = true;
-			Config.ExcludePropertyFlags = CPF_DisableEditOnInstance;
+			Config.ExcludePropertyFlags = CPF_DisableEditOnInstance | CPF_EditConst | CPF_BlueprintReadOnly;
 			OverridableParams.Append(PCGSettingsHelpers::GetAllOverridableParams(BPClass, Config));
 		}
 	}
diff --git a/Engine/Plugins/Experimental/PCG/Source/PCG/Private/Elements/PCGPropertyToParamData.cpp b/Engine/Plugins/Experimental/PCG/Source/PCG/Private/Elements/PCGPropertyToParamData.cpp
--- a/Engine/Plugins/Experimental/PCG/Source/PCG/Private/Elements/PCGPropertyToParamData.cpp
+++ b/Engine/Plugins/Experimental/PCG/Source/PCG/Private/Elements/PCGPropertyToParamData.cpp
@@ -131,7 +131,7 @@ bool FPCGPropertyToParamDataElement::ExecuteInternal(FPCGContext* Context) const
 		ObjectToInspect = FoundActor->GetComponentByClass(Settings->ComponentClass);
 		if (!ObjectToInspect)
 		{
-			PCGE_LOG(Error, GraphAndLog, LOCTEXT("ComponentDoesNotExist", "Component does not exist in the found actor"));
+			PCGE_LOG(Error, GraphAndLog, FText::Format(LOCTEXT("ComponentDoesNotExist", "Component class '{0}' does not exist in the found actor"), FText::FromString(Settings->ComponentClass->GetName())));
 			return true;
 		}
 	}
@@ -140,7 +140,16 @@ bool FPCGPropertyToParamDataElement::ExecuteInternal(FPCGContext* Context) const
 	FProperty* Property = FindFProperty<FProperty>(ObjectToInspect->GetClass(), Settings->PropertyName);
 	if (!Property)
 	{
-		PCGE_LOG(Error, GraphAndLog, LOCTEXT("PropertyDoesNotExist", "Property does not exist in the found actor"));
+		PCGE_LOG(Error, GraphAndLog, FText::Format(LOCTEXT("PropertyDoesNotExist", "Property '{0}' does not exist in the found actor"), FText::FromName(Settings->PropertyName)));
+		return true;
+	}
+
+	// Make sure the property is visible
+	const uint64 ExcludePropertyFlags = CPF_DisableEditOnInstance;
+	const uint64 IncludePropertyFlags = CPF_BlueprintVisible;
+	if (Property->HasAnyPropertyFlags(ExcludePropertyFlags) || !Property->HasAnyPropertyFlags(IncludePropertyFlags))
+	{
+		PCGE_LOG(Error, GraphAndLog, FText::Format(LOCTEXT("PropertyExistsButNotVisible", "Property '{0}' does exist in the found actor, but is not visible."), FText::FromName(Settings->PropertyName)));
 		return true;
 	}
 
@@ -174,6 +183,9 @@ bool FPCGPropertyToParamDataElement::ExecuteInternal(FPCGContext* Context) const
 		Config.bUseSeed = true;
 		Config.bExcludeSuperProperties = true;
 		Config.MaxStructDepth = 0;
+		// Can only get exposed properties and visible
+		Config.ExcludePropertyFlags = ExcludePropertyFlags;
+		Config.IncludePropertyFlags = IncludePropertyFlags;
 		TArray<FPCGSettingsOverridableParam> AllChildProperties = UnderlyingStruct ? PCGSettingsHelpers::GetAllOverridableParams(UnderlyingStruct, Config) : PCGSettingsHelpers::GetAllOverridableParams(UnderlyingClass, Config);
 
 		for (const FPCGSettingsOverridableParam& Param : AllChildProperties)
diff --git a/Engine/Plugins/Experimental/PCG/Source/PCG/Private/Helpers/PCGSettingsHelpers.cpp b/Engine/Plugins/Experimental/PCG/Source/PCG/Private/Helpers/PCGSettingsHelpers.cpp
--- a/Engine/Plugins/Experimental/PCG/Source/PCG/Private/Helpers/PCGSettingsHelpers.cpp
+++ b/Engine/Plugins/Experimental/PCG/Source/PCG/Private/Helpers/PCGSettingsHelpers.cpp
@@ -2,6 +2,7 @@
 
 #include "Helpers/PCGSettingsHelpers.h"
 
+#include "PCGCommon.h"
 #include "PCGComponent.h"
 #include "PCGEdge.h"
 #include "PCGPin.h"
@@ -275,23 +276,10 @@ namespace PCGSettingsHelpers
 	template <typename ClassType>
 	TArray<FPCGSettingsOverridableParam> GetAllOverridableParamsImpl(const ClassType* InClass, const FPCGGetAllOverridableParamsConfig& InConfig)
 	{
-		// TODO: Was not a concern until now, and we didn't have a solution, but this function
-		// only worked if we don't have names clashes in overriable parameters.
-		// The previous override solution was flattening structs, and only override use the struct member name,
-		// not prefixed by the struct name or anything else.
-		// We cannot prefix it now, because it will break existing node that were assuming the flattening.
-		// We'll keep this behavior for now, as it might be solved by passing structs instead of param data,
-		// but we'll still at least raise a warning if there is a clash.
-		TSet<FName> LabelCache;
+		TArray<FName> LabelCache;
 
 		TArray<FPCGSettingsOverridableParam> Res;
 
-		// Can't check metadata in non-editor build
-#if WITH_EDITOR
-		const bool bCheckMetadata = !InConfig.MetadataValues.IsEmpty();
-#endif // WITH_EDITOR
-
-		const bool bCheckExcludePropertyFlags = (InConfig.ExcludePropertyFlags != 0);
 		const EFieldIteratorFlags::SuperClassFlags SuperFlag = InConfig.bExcludeSuperProperties ? EFieldIteratorFlags::ExcludeSuper : EFieldIteratorFlags::IncludeSuper;
 
 		check(InClass);
@@ -307,29 +295,57 @@ namespace PCGSettingsHelpers
 			bool bValid = true;
 
 #if WITH_EDITOR
-			if (bCheckMetadata)
+			auto HasMetadata = [Property](const TArray<FName>& InValues, bool bIsConjunction) -> bool
 			{
-				bool bFoundAny = false;
-				for (const FName& Metadata : InConfig.MetadataValues)
+				bool bFound = false;
+				for (const FName& Metadata : InValues)
 				{
-					if (Property->HasMetaData(Metadata))
+					if (Property->HasMetaData(Metadata) && !bIsConjunction)
 					{
-						bFoundAny = true;
+						bFound = true;
+						if (!bIsConjunction)
+						{
+							break;
+						}
+					}
+					else if (bIsConjunction)
+					{
+						bFound = false;
 						break;
 					}
 				}
 
-				bValid &= bFoundAny;
+				return bFound;
+			};
+
+			if (!InConfig.IncludeMetadataValues.IsEmpty())
+			{
+				bValid &= HasMetadata(InConfig.IncludeMetadataValues, InConfig.bIncludeMetadataIsConjunction);
+			}
+
+			if (!InConfig.ExcludeMetadataValues.IsEmpty())
+			{
+				bValid &= !HasMetadata(InConfig.ExcludeMetadataValues, InConfig.bExcludeMetadataIsConjunction);
 			}
 #endif // WITH_EDITOR
 
-			if (bCheckExcludePropertyFlags)
+			auto HasPropertyFlags = [Property](uint64 InFlags, bool bIsConjunction) -> bool
+			{
+				return bIsConjunction ? Property->HasAllPropertyFlags(InFlags) : Property->HasAnyPropertyFlags(InFlags);
+			};
+
+			if (InConfig.IncludePropertyFlags != 0)
 			{
-				bValid &= !Property->HasAnyPropertyFlags(InConfig.ExcludePropertyFlags);
+				bValid &= HasPropertyFlags(InConfig.IncludePropertyFlags, InConfig.bIncludePropertyFlagsIsConjunction);
+			}
+
+			if (InConfig.ExcludePropertyFlags != 0)
+			{
+				bValid &= !HasPropertyFlags(InConfig.ExcludePropertyFlags, InConfig.bExcludePropertyFlagsIsConjunction);
 			}
 
 			// Don't allow to override the seed if the settings doesn't use the seed.
-			if (!InConfig.bUseSeed)
+			if (!InConfig.bUseSeed && Property->GetOwnerClass()->IsChildOf<UPCGSettings>())
 			{
 				bValid &= (Property->GetFName() != GET_MEMBER_NAME_CHECKED(UPCGSettings, Seed));
 			}
@@ -343,22 +359,30 @@ namespace PCGSettingsHelpers
 			if (PCGAttributeAccessorHelpers::IsPropertyAccessorSupported(Property))
 			{
 				FName Label = NAME_None;
+				bool bHasNameClash = false;
 #if WITH_EDITOR
 				// GetDisplayNameText is not available in non-editor build.
 				Label = *Property->GetDisplayNameText().ToString();
-				if (LabelCache.Contains(Label))
+				const int32 CachedLabelIndex = LabelCache.IndexOfByKey(Label);
+				if (CachedLabelIndex != INDEX_NONE)
 				{
-					UE_LOG(LogPCG, Warning, TEXT("%s property clashes with another property already found. It is a limitation at the moment and this property will be ignored."), *Label.ToString());
-					continue;
+					// If we have a clash, we will use the full path, so mark this param and the other that clashed to use the full path.
+					Res[CachedLabelIndex].bHasNameClash = true;
+					Res[CachedLabelIndex].Label = FName(Res[CachedLabelIndex].GetDisplayPropertyPath());
+					bHasNameClash = true;
 				}
 
-				LabelCache.Add(Label);
+				LabelCache.AddUnique(Label);
 #endif // WITH_EDITOR
 
 				FPCGSettingsOverridableParam& Param = Res.Emplace_GetRef();
-				Param.Label = Label;
 				Param.PropertiesNames.Add(Property->GetFName());
+				Param.Properties.Add(Property);
 				Param.PropertyClass = InClass;
+				Param.bHasNameClash = bHasNameClash;
+#if WITH_EDITOR
+				Param.Label = bHasNameClash ? FName(Param.GetDisplayPropertyPath()) : Label;
+#endif // WITH_EDITOR
 			}
 			else if (const FStructProperty* StructProperty = CastField<FStructProperty>(Property))
 			{
@@ -368,11 +392,11 @@ namespace PCGSettingsHelpers
 					continue;
 				}
 
-				// Use the seed, and don't check metadata.
+				// Use the seed, and don't check metadata for PCG overridable.
 				FPCGGetAllOverridableParamsConfig RecurseConfig = InConfig;
 				RecurseConfig.bUseSeed = true;
 #if WITH_EDITOR
-				RecurseConfig.MetadataValues.Empty();
+				RecurseConfig.IncludeMetadataValues.Remove(PCGObjectMetadata::Overridable);
 #endif // WITH_EDITOR
 
 				if (RecurseConfig.MaxStructDepth > 0)
@@ -383,22 +407,38 @@ namespace PCGSettingsHelpers
 				for (const FPCGSettingsOverridableParam& ChildParam : GetAllOverridableParams(StructProperty->Struct, RecurseConfig))
 				{
 					FName Label = ChildParam.Label;
+					bool bHasNameClash = false;
 #if WITH_EDITOR
 					// Don't check for label clash, as they would all be equal to None in non-editor build
-					if (LabelCache.Contains(Label))
+					const int32 CachedLabelIndex = LabelCache.IndexOfByKey(Label);
+					if (CachedLabelIndex != INDEX_NONE)
 					{
-						UE_LOG(LogPCG, Warning, TEXT("%s property clashes with another property already found. It is a limitation at the moment and this property will be ignored."), *Label.ToString());
-						continue;
+						// If we have a clash, we will use the full path, so mark this param and the other that clashed to use the full path.
+						Res[CachedLabelIndex].bHasNameClash = true;
+						Res[CachedLabelIndex].Label = FName(Res[CachedLabelIndex].GetDisplayPropertyPath());
+						bHasNameClash = true;
 					}
 
-					LabelCache.Add(Label);
+					LabelCache.AddUnique(Label);
 #endif // WITH_EDITOR
 
 					FPCGSettingsOverridableParam& Param = Res.Emplace_GetRef();
-					Param.Label = Label;
 					Param.PropertiesNames.Add(Property->GetFName());
-					Param.PropertiesNames.Append(ChildParam.PropertiesNames);
+					Param.PropertiesNames.Append(std::move(ChildParam.PropertiesNames));
+					Param.Properties.Add(Property);
+					Param.Properties.Append(std::move(ChildParam.Properties));
 					Param.PropertyClass = InClass;
+					if (!ChildParam.bHasNameClash && !bHasNameClash)
+					{
+						Param.Label = Label;
+					}
+					else
+					{
+#if WITH_EDITOR
+						Param.Label = FName(Param.GetDisplayPropertyPath());
+#endif // WITH_EDITOR
+						Param.bHasNameClash = true;
+					}
 				}
 			}
 		}
diff --git a/Engine/Plugins/Experimental/PCG/Source/PCG/Private/Metadata/Accessors/PCGAttributeAccessorHelpers.cpp b/Engine/Plugins/Experimental/PCG/Source/PCG/Private/Metadata/Accessors/PCGAttributeAccessorHelpers.cpp
--- a/Engine/Plugins/Experimental/PCG/Source/PCG/Private/Metadata/Accessors/PCGAttributeAccessorHelpers.cpp
+++ b/Engine/Plugins/Experimental/PCG/Source/PCG/Private/Metadata/Accessors/PCGAttributeAccessorHelpers.cpp
@@ -257,8 +257,10 @@ TUniquePtr<const IPCGAttributeAccessor> PCGAttributeAccessorHelpers::CreateConst
 
 	if (ParamData && ParamData->Metadata && ParamData->Metadata->GetAttributeCount() > 0)
 	{
-		// If the param only has a single attribute and is not from the global Params pin, use this one. Otherwise we need perfect name matching.
-		const FName AttributeName = (ParamData->Metadata->GetAttributeCount() == 1 && !bFromGlobalParamsPin) ? ParamData->Metadata->GetLatestAttributeNameOrNone() : InParam.Properties.Last()->GetFName();
+		// If the param only has a single attribute and is not from the global Params pin, use this one. Otherwise we need perfect name matching, either the property name or its full path if there is a name clash.
+		const FName AttributeName = (ParamData->Metadata->GetAttributeCount() == 1 && !bFromGlobalParamsPin) 
+			? ParamData->Metadata->GetLatestAttributeNameOrNone() 
+			: (!InParam.bHasNameClash ? InParam.PropertiesNames.Last() : FName(InParam.GetPropertyPath()));
 
 		if (OutAttributeName)
 		{
diff --git a/Engine/Plugins/Experimental/PCG/Source/PCG/Private/Metadata/PCGMetadataAttribute.cpp b/Engine/Plugins/Experimental/PCG/Source/PCG/Private/Metadata/PCGMetadataAttribute.cpp
--- a/Engine/Plugins/Experimental/PCG/Source/PCG/Private/Metadata/PCGMetadataAttribute.cpp
+++ b/Engine/Plugins/Experimental/PCG/Source/PCG/Private/Metadata/PCGMetadataAttribute.cpp
@@ -107,7 +107,7 @@ void FPCGMetadataAttributeBase::ClearEntries()
 bool FPCGMetadataAttributeBase::IsValidName(const FString& Name)
 {
 	// A valid name is alphanumeric with some special characters allowed.
-	static const FString AllowedSpecialCharacters = TEXT("_-");
+	static const FString AllowedSpecialCharacters = TEXT("_-/");
 
 	for (int32 i = 0; i < Name.Len(); ++i)
 	{
diff --git a/Engine/Plugins/Experimental/PCG/Source/PCG/Private/PCGSettings.cpp b/Engine/Plugins/Experimental/PCG/Source/PCG/Private/PCGSettings.cpp
--- a/Engine/Plugins/Experimental/PCG/Source/PCG/Private/PCGSettings.cpp
+++ b/Engine/Plugins/Experimental/PCG/Source/PCG/Private/PCGSettings.cpp
@@ -60,6 +60,18 @@ public:
 #endif // WITH_EDITOR
 };
 
+FString FPCGSettingsOverridableParam::GetPropertyPath() const
+{
+	return FString::JoinBy(Properties, TEXT("/"), [](const FProperty* InProperty) { return InProperty ? InProperty->GetName() : FString(); });
+}
+
+#if WITH_EDITOR
+FString FPCGSettingsOverridableParam::GetDisplayPropertyPath() const
+{
+	return FString::JoinBy(Properties, TEXT("/"), [](const FProperty* InProperty) { return InProperty ? InProperty->GetDisplayNameText().ToString() : FString(); });
+}
+#endif // WITH_EDITOR
+
 bool UPCGSettingsInterface::IsInstance() const
 {
 	return this != GetSettings();
@@ -328,7 +340,7 @@ void UPCGSettings::FillOverridableParamsPins(TArray<FPCGPinProperties>& OutPins)
 		ParamPin.Tooltip = FText::Format(LOCTEXT("OverridableParamPinTooltip", "{0}Attribute type is \"{1}\" and its exact name is \"{2}\""),
 			FText::FromString(Tooltip),
 			FText::FromString(Property->GetCPPType()),
-			FText::FromName(Property->GetFName()));
+			FText::FromString(OverridableParam.GetPropertyPath()));
 #endif // WITH_EDITOR
 	}
 }
@@ -511,7 +523,8 @@ TArray<FPCGSettingsOverridableParam> UPCGSettings::GatherOverridableParams() con
 {
 	PCGSettingsHelpers::FPCGGetAllOverridableParamsConfig Config;
 	Config.bUseSeed = bUseSeed;
-	Config.MetadataValues.Add(PCGObjectMetadata::Overridable);
+	Config.IncludeMetadataValues.Add(PCGObjectMetadata::Overridable);
+	Config.ExcludeMetadataValues.Add(PCGObjectMetadata::NotOverridable);
 
 	return PCGSettingsHelpers::GetAllOverridableParams(GetClass(), Config);
 }
diff --git a/Engine/Plugins/Experimental/PCG/Source/PCG/Public/Helpers/PCGSettingsHelpers.h b/Engine/Plugins/Experimental/PCG/Source/PCG/Public/Helpers/PCGSettingsHelpers.h
--- a/Engine/Plugins/Experimental/PCG/Source/PCG/Public/Helpers/PCGSettingsHelpers.h
+++ b/Engine/Plugins/Experimental/PCG/Source/PCG/Public/Helpers/PCGSettingsHelpers.h
@@ -212,12 +212,30 @@ PRAGMA_ENABLE_DEPRECATION_WARNINGS
 
 #if WITH_EDITOR
 		// List of metadata values to find in property metadata. Only works in editor builds as metadata on property is not available elsewise.
-		TArray<FName> MetadataValues{};
+		TArray<FName> IncludeMetadataValues;
+
+		// If the metadata values is a conjunction (all values need to be present to keep), or disjunction (any values needs to be present to keep)
+		bool bIncludeMetadataIsConjunction = false;
+
+		// List of metadata values to find in property metadata. Only works in editor builds as metadata on property is not available elsewise.
+		TArray<FName> ExcludeMetadataValues;
+
+		// If the exclude values is a conjunction (all values need to be present to discard), or disjunction (any values needs to be present to discard)
+		bool bExcludeMetadataIsConjunction = false;
 #endif // WITH_EDITOR
 
 		// Flags to exclude in property flags
 		uint64 ExcludePropertyFlags = 0;
 
+		// If the exclude flags is a conjunction (all flags need to be present to discard), or disjunction (any flag needs to be present to discard)
+		bool bExcludePropertyFlagsIsConjunction = false;
+
+		// Flags to include in property flags. Note that if exclusion says discard, it will be discarded.
+		uint64 IncludePropertyFlags = 0;
+
+		// If the include flags is a conjunction (all flags need to be present to keep), or disjunction (any flag needs to be present to keep)
+		bool bIncludePropertyFlagsIsConjunction = false;
+
 		// Max depth for structs of structs. -1 = no limit
 		int32 MaxStructDepth = -1;
 	};
diff --git a/Engine/Plugins/Experimental/PCG/Source/PCG/Public/PCGCommon.h b/Engine/Plugins/Experimental/PCG/Source/PCG/Public/PCGCommon.h
--- a/Engine/Plugins/Experimental/PCG/Source/PCG/Public/PCGCommon.h
+++ b/Engine/Plugins/Experimental/PCG/Source/PCG/Public/PCGCommon.h
@@ -71,11 +71,15 @@ namespace Private
 namespace PCGObjectMetadata
 {
 	const FName Overridable = TEXT("PCG_Overridable");
+	const FName NotOverridable = TEXT("PCG_NotOverridable");
 
 	// Metadata usable in UPROPERTY for customizing the behavior when displaying the property in a property panel or graph node
 	enum
 	{
 		/// [PropertyMetadata] Indicates that the property is overridable by params.
-		PCG_Overridable
+		PCG_Overridable,
+
+		/// [PropertyMetadata] Indicates that the property is not-overridable by params. Used in structs to hide some parameters
+		PCG_NotOverridable
 	};
 }
diff --git a/Engine/Plugins/Experimental/PCG/Source/PCG/Public/PCGSettings.h b/Engine/Plugins/Experimental/PCG/Source/PCG/Public/PCGSettings.h
--- a/Engine/Plugins/Experimental/PCG/Source/PCG/Public/PCGSettings.h
+++ b/Engine/Plugins/Experimental/PCG/Source/PCG/Public/PCGSettings.h
@@ -65,8 +65,18 @@ struct FPCGSettingsOverridableParam
 	UPROPERTY()
 	TObjectPtr<const UStruct> PropertyClass;
 
+	// If this flag is true, Label will be the full property path.
+	UPROPERTY()
+	bool bHasNameClash = false;
+
 	// Transient
 	TArray<const FProperty*> Properties;
+
+	FString GetPropertyPath() const;
+
+#if WITH_EDITOR
+	FString GetDisplayPropertyPath() const;
+#endif // WITH_EDITOR
 };
 
 /**
@@ -218,6 +228,13 @@ public:
 	virtual FText GetNodeTooltipText() const { return FText::GetEmpty(); }
 	virtual FLinearColor GetNodeTitleColor() const { return FLinearColor::White; }
 	virtual EPCGSettingsType GetType() const { return EPCGSettingsType::Generic; }
+
+	/** Can override the label style for a pin. Return false if no override is available. */
+	virtual bool GetPinLabelStyle(const UPCGPin* InPin, FName& OutLabelStyle) const { return false; }
+
+	/** Can override to add a custom icon next to the pin label (and an optional tooltip). Return false if no override is available. */
+	virtual bool GetPinExtraIcon(const UPCGPin* InPin, FName& OutExtraIcon, FText& OutTooltip) const { return false; }
+
 	/** Derived classes must implement this to communicate dependencies on external actors */
 	virtual void GetTrackedActorTags(FPCGTagToSettingsMap& OutTagToSettings, TArray<TObjectPtr<const UPCGGraph>>& OutVisitedGraphs) const {}
 	/** Override this class to provide an UObject to jump to in case of double click on node
