Added loop node that behaves like a subgraph but will pass on a single data from the first pin & all data from other pins for all data on the first pin. This is especially useful in conjunction with the partition node.

diff --git a/Engine/Plugins/Experimental/PCG/Source/PCG/Private/Elements/PCGLoopElement.cpp b/Engine/Plugins/Experimental/PCG/Source/PCG/Private/Elements/PCGLoopElement.cpp
new file mode 100644
--- /dev/null
+++ b/Engine/Plugins/Experimental/PCG/Source/PCG/Private/Elements/PCGLoopElement.cpp
@@ -0,0 +1,198 @@
+// Copyright Epic Games, Inc. All Rights Reserved.
+
+#include "Elements/PCGLoopElement.h"
+
+#include "PCGComponent.h"
+#include "PCGModule.h"
+#include "PCGPin.h"
+#include "PCGSubsystem.h"
+#include "Helpers/PCGHelpers.h"
+
+#include UE_INLINE_GENERATED_CPP_BY_NAME(PCGLoopElement)
+
+#define LOCTEXT_NAMESPACE "PCGLoopElement"
+
+#if WITH_EDITOR
+FText UPCGLoopSettings::GetDefaultNodeTitle() const
+{
+	return LOCTEXT("NodeTitle", "Loop");
+}
+
+FText UPCGLoopSettings::GetNodeTooltipText() const
+{
+	return LOCTEXT("NodeTooltip", "Executes the specified Subgraph for each data on the loop pins (or on the first pin if no specific loop pins are provided), keeping the rest constant.");
+}
+#endif
+
+FPCGElementPtr UPCGLoopSettings::CreateElement() const
+{
+	return MakeShared<FPCGLoopElement>();
+}
+
+FName UPCGLoopSettings::AdditionalTaskName() const
+{
+	if (UPCGGraph* TargetSubgraph = GetSubgraph())
+	{
+		return FName(FText::Format(LOCTEXT("NodeTitleExtended", "Loop Subgraph - {0}"), FText::FromName(TargetSubgraph->GetFName())).ToString());
+	}
+	else
+	{
+		return FName(LOCTEXT("NodeTitleExtendedInvalidSubgraph", "Loop Subgraph - Invalid Subgraph").ToString());
+	}
+}
+
+bool FPCGLoopElement::ExecuteInternal(FPCGContext* InContext) const
+{
+	TRACE_CPUPROFILER_EVENT_SCOPE(FPCGLoopElement::Execute);
+
+	FPCGSubgraphContext* Context = static_cast<FPCGSubgraphContext*>(InContext);
+
+	const UPCGLoopSettings* Settings = Context->GetInputSettings<UPCGLoopSettings>();
+	check(Settings);
+
+	// If we haven't scheduled the subgraph yet (1st call) - this is part of the FPCGSubgraphContext, see FPCGSubgraphElement for more info
+	if (!Context->bScheduledSubgraph)
+	{
+		UPCGGraph* Subgraph = Settings->GetSubgraph();
+		UPCGSubsystem* Subsystem = Context->SourceComponent.IsValid() ? Context->SourceComponent->GetSubsystem() : nullptr;
+
+		if (!Subsystem || !Subgraph)
+		{
+			// Job cannot run; cancel
+			if (!Subgraph)
+			{
+				PCGE_LOG(Error, GraphAndLog, LOCTEXT("InvalidSubgraph", "Cannot loop on an invalid subgraph"));
+			}
+			
+			Context->OutputData.bCancelExecution = true;
+			return true;
+		}
+
+		TArray<FPCGDataCollection> LoopDataCollection;
+		FPCGDataCollection FixedInputDataCollection;
+		PrepareLoopDataCollections(Context, Settings, LoopDataCollection, FixedInputDataCollection);
+
+		// Early out if there are no data on the loop pin
+		if (LoopDataCollection.IsEmpty() || LoopDataCollection[0].TaggedData.IsEmpty())
+		{
+			PCGE_LOG(Verbose, LogOnly, LOCTEXT("EmptyLoopCollection", "Loop data is empty - will not do anything"));
+			return true;
+		}
+
+		// Early out if the loop collection do not contain the same amount of entries
+		for (int LoopDataIndex = 1; LoopDataIndex < LoopDataCollection.Num(); ++LoopDataIndex)
+		{
+			if (LoopDataCollection[LoopDataIndex].TaggedData.Num() != LoopDataCollection[0].TaggedData.Num())
+			{
+				PCGE_LOG(Error, GraphAndLog, LOCTEXT("MismatchLoopCollections", "Data to loop on does not have the same number of entries for all pins!"));
+				return true;
+			}
+		}
+
+		FPCGDataCollection PreSubgraphDataCollection;
+		PrepareSubgraphUserParameters(Settings, Context, PreSubgraphDataCollection);
+
+		FPCGElementPtr PreGraphElement = MakeShared<FPCGInputForwardingElement>(PreSubgraphDataCollection);
+
+		FPCGDataCollection FixedDataCollection;
+		PrepareSubgraphData(Settings, Context, FixedInputDataCollection, FixedDataCollection);
+
+		// Dispatch the subgraph for each loop entries we have.
+		// Implementation note: even if the execution gets cancelled, these tasks would get cancelled because they are associated to the current source component
+		for(int EntryIndex = 0; EntryIndex < LoopDataCollection[0].TaggedData.Num(); ++EntryIndex)
+		{
+			// The final input for a given loop iteration is going to be the fixed data + the entry at the current loop index in all the gathered loop collections
+			FPCGDataCollection InputDataCollection = FixedDataCollection;
+			for (int LoopCollectionIndex = 0; LoopCollectionIndex < LoopDataCollection.Num(); ++LoopCollectionIndex)
+			{
+				InputDataCollection.TaggedData.Insert(LoopDataCollection[LoopCollectionIndex].TaggedData[EntryIndex], LoopCollectionIndex);
+			}
+
+			FPCGTaskId SubgraphTaskId = Subsystem->ScheduleGraph(Subgraph, Context->SourceComponent.Get(), PreGraphElement, MakeShared<FPCGInputForwardingElement>(InputDataCollection), {});
+
+			Context->SubgraphTaskIds.Add(SubgraphTaskId);
+		}
+
+		Context->bScheduledSubgraph = true;
+		Context->bIsPaused = true;
+
+		// add a trivial task after the output tasks that wakes up this task
+		Subsystem->ScheduleGeneric([Context]()
+		{
+			// Wake up the current task
+			Context->bIsPaused = false;
+			return true;
+		}, Context->SourceComponent.Get(), Context->SubgraphTaskIds);
+
+		return false;
+	}
+	else if (Context->bIsPaused)
+	{
+		// Should not happen once we skip it in the graph executor
+		return false;
+	}
+	else
+	{
+		// when woken up, get the output data from the subgraph
+		// and copy it to the current context output data, and finally return true
+		UPCGSubsystem* Subsystem = Context->SourceComponent->GetSubsystem();
+		if (Subsystem)
+		{
+			for (FPCGTaskId SubgraphTaskId : Context->SubgraphTaskIds)
+			{
+				FPCGDataCollection SubgraphOutput;
+				ensure(Subsystem->GetOutputData(SubgraphTaskId, SubgraphOutput));
+
+				Context->OutputData.TaggedData.Append(SubgraphOutput.TaggedData);
+			}
+		}
+		else
+		{
+			// Job cannot run, cancel
+			Context->OutputData.bCancelExecution = true;
+		}
+
+		return true;
+	}
+}
+
+void FPCGLoopElement::PrepareLoopDataCollections(FPCGContext* Context, const UPCGLoopSettings* Settings, TArray<FPCGDataCollection>& LoopDataCollection, FPCGDataCollection& FixedInputDataCollection) const
+{
+	TArray<FName> LoopPinNames;
+	TArray<FString> PinsFromSettings = PCGHelpers::GetStringArrayFromCommaSeparatedString(Settings->LoopPins);
+
+	for (const FString& PinLabel : PinsFromSettings)
+	{
+		LoopPinNames.Emplace(PinLabel);
+	}
+	
+	// If no named pins have been provided, default to the first pin
+	if (LoopPinNames.IsEmpty())
+	{
+		const UPCGSubgraphNode* LoopNode = Cast<const UPCGSubgraphNode>(Context->Node);
+		check(LoopNode);
+
+		if (const UPCGPin* LoopPin = LoopNode->GetFirstConnectedInputPin())
+		{
+			LoopPinNames.Add(LoopPin->Properties.Label);
+		}
+	}
+
+	LoopDataCollection.Reset();
+	LoopDataCollection.SetNum(LoopPinNames.Num());
+
+	for (const FPCGTaggedData& TaggedData : Context->InputData.TaggedData)
+	{
+		int LoopPinIndex = LoopPinNames.IndexOfByKey(TaggedData.Pin);
+		if (LoopPinIndex != INDEX_NONE)
+		{
+			LoopDataCollection[LoopPinIndex].TaggedData.Add(TaggedData);
+		}
+		else
+		{
+			FixedInputDataCollection.TaggedData.Add(TaggedData);
+		}
+	}
+}
+
+#undef LOCTEXT_NAMESPACE
\ No newline at end of file
diff --git a/Engine/Plugins/Experimental/PCG/Source/PCG/Private/Graph/PCGGraphCompiler.cpp b/Engine/Plugins/Experimental/PCG/Source/PCG/Private/Graph/PCGGraphCompiler.cpp
--- a/Engine/Plugins/Experimental/PCG/Source/PCG/Private/Graph/PCGGraphCompiler.cpp
+++ b/Engine/Plugins/Experimental/PCG/Source/PCG/Private/Graph/PCGGraphCompiler.cpp
@@ -39,6 +39,7 @@ TArray<FPCGGraphTask> FPCGGraphCompiler::CompileGraph(UPCGGraph* InGraph, FPCGTa
 
 		const UPCGBaseSubgraphNode* SubgraphNode = Cast<const UPCGBaseSubgraphNode>(Node);
 		UPCGGraph* Subgraph = SubgraphNode ? SubgraphNode->GetSubgraph() : nullptr;
+		const UPCGBaseSubgraphSettings* SubgraphSettings = SubgraphNode ? Cast<const UPCGBaseSubgraphSettings>(SubgraphNode->GetSettings()) : nullptr;
 
 		// Only catch immediate infinite recursion.
 		// TODO: Add a better mecanism for detecting more complex cyclic recursions
@@ -53,7 +54,7 @@ TArray<FPCGGraphTask> FPCGGraphCompiler::CompileGraph(UPCGGraph* InGraph, FPCGTa
 			return TArray<FPCGGraphTask>();
 		}
 
-		const bool bIsNonDynamicAndNonDisabledSubgraphNode = (SubgraphNode != nullptr && !SubgraphNode->bDynamicGraph && Subgraph && SubgraphNode->GetSettings() && SubgraphNode->GetSettings()->bEnabled);
+		const bool bIsNonDynamicAndNonDisabledSubgraphNode = (SubgraphNode && Subgraph && SubgraphSettings && !SubgraphSettings->IsDynamicGraph() && !SubgraphSettings->bEnabled);
 		if (bIsNonDynamicAndNonDisabledSubgraphNode)
 		{
 			const FPCGTaskId PreId = NextId++;
diff --git a/Engine/Plugins/Experimental/PCG/Source/PCG/Private/Graph/PCGGraphExecutor.cpp b/Engine/Plugins/Experimental/PCG/Source/PCG/Private/Graph/PCGGraphExecutor.cpp
--- a/Engine/Plugins/Experimental/PCG/Source/PCG/Private/Graph/PCGGraphExecutor.cpp
+++ b/Engine/Plugins/Experimental/PCG/Source/PCG/Private/Graph/PCGGraphExecutor.cpp
@@ -89,10 +89,10 @@ FPCGTaskId FPCGGraphExecutor::Schedule(UPCGComponent* Component, const TArray<FP
 	check(Component);
 	UPCGGraph* Graph = Component->GetGraph();
 
-	return Schedule(Graph, Component, GetFetchInputElement(), ExternalDependencies);
+	return Schedule(Graph, Component, FPCGElementPtr(), GetFetchInputElement(), ExternalDependencies);
 }
 
-FPCGTaskId FPCGGraphExecutor::Schedule(UPCGGraph* Graph, UPCGComponent* SourceComponent, FPCGElementPtr InputElement, const TArray<FPCGTaskId>& ExternalDependencies)
+FPCGTaskId FPCGGraphExecutor::Schedule(UPCGGraph* Graph, UPCGComponent* SourceComponent, FPCGElementPtr PreGraphElement, FPCGElementPtr InputElement, const TArray<FPCGTaskId>& ExternalDependencies)
 {
 	if (SourceComponent && IsGraphCacheDebuggingEnabled())
 	{
@@ -143,10 +143,17 @@ FPCGTaskId FPCGGraphExecutor::Schedule(UPCGGraph* Graph, UPCGComponent* SourceCo
 		// Push task (not data) dependencies on the pre-execute task
 		// Note must be done after the offset ids, otherwise we'll break the dependencies
 		check(ScheduledTask.Tasks.Num() >= 2 && ScheduledTask.Tasks[ScheduledTask.Tasks.Num() - 2].Node == nullptr);
+		FPCGGraphTask& PreGraphTask = ScheduledTask.Tasks[ScheduledTask.Tasks.Num() - 2];
+
+		if(PreGraphElement.IsValid())
+		{
+			PreGraphTask.Element = PreGraphElement;
+		}
+
 		for (FPCGTaskId ExternalDependency : ExternalDependencies)
 		{
 			// For the pre-task, we don't consume any input
-			ScheduledTask.Tasks[ScheduledTask.Tasks.Num() - 2].Inputs.Emplace(ExternalDependency, nullptr, nullptr, /*bConsumeData=*/false);
+			PreGraphTask.Inputs.Emplace(ExternalDependency, nullptr, nullptr, /*bConsumeData=*/false);
 		}
 
 		ScheduledTask.FirstTaskIndex = ScheduledTask.Tasks.Num() - 2;
diff --git a/Engine/Plugins/Experimental/PCG/Source/PCG/Private/Graph/PCGGraphExecutor.h b/Engine/Plugins/Experimental/PCG/Source/PCG/Private/Graph/PCGGraphExecutor.h
--- a/Engine/Plugins/Experimental/PCG/Source/PCG/Private/Graph/PCGGraphExecutor.h
+++ b/Engine/Plugins/Experimental/PCG/Source/PCG/Private/Graph/PCGGraphExecutor.h
@@ -83,7 +83,7 @@ public:
 
 	/** Schedules the execution of a given graph with specified inputs. This call is threadsafe */
 	FPCGTaskId Schedule(UPCGComponent* InComponent, const TArray<FPCGTaskId>& TaskDependency = TArray<FPCGTaskId>());
-	FPCGTaskId Schedule(UPCGGraph* Graph, UPCGComponent* InSourceComponent, FPCGElementPtr InputElement, const TArray<FPCGTaskId>& TaskDependency);
+	FPCGTaskId Schedule(UPCGGraph* Graph, UPCGComponent* InSourceComponent, FPCGElementPtr PreGraphElement, FPCGElementPtr InputElement, const TArray<FPCGTaskId>& TaskDependency);
 
 	/** Cancels all tasks originating from the given component */
 	TArray<UPCGComponent*> Cancel(UPCGComponent* InComponent);
diff --git a/Engine/Plugins/Experimental/PCG/Source/PCG/Private/PCGNode.cpp b/Engine/Plugins/Experimental/PCG/Source/PCG/Private/PCGNode.cpp
--- a/Engine/Plugins/Experimental/PCG/Source/PCG/Private/PCGNode.cpp
+++ b/Engine/Plugins/Experimental/PCG/Source/PCG/Private/PCGNode.cpp
@@ -526,6 +526,19 @@ bool UPCGNode::IsEdgeUsedByNodeExecution(const UPCGEdge* InEdge) const
 	return Settings->IsPinUsedByNodeExecution(Pin);
 }
 
+const UPCGPin* UPCGNode::GetFirstConnectedInputPin() const
+{
+	for (const UPCGPin* InputPin : InputPins)
+	{
+		if (InputPin && InputPin->EdgeCount() > 0)
+		{
+			return InputPin;
+		}
+	}
+
+	return nullptr;
+}
+
 void UPCGNode::SetSettingsInterface(UPCGSettingsInterface* InSettingsInterface, bool bUpdatePins)
 {
 	const bool bDifferentInterface = (SettingsInterface.Get() != InSettingsInterface);
diff --git a/Engine/Plugins/Experimental/PCG/Source/PCG/Private/PCGSubgraph.cpp b/Engine/Plugins/Experimental/PCG/Source/PCG/Private/PCGSubgraph.cpp
--- a/Engine/Plugins/Experimental/PCG/Source/PCG/Private/PCGSubgraph.cpp
+++ b/Engine/Plugins/Experimental/PCG/Source/PCG/Private/PCGSubgraph.cpp
@@ -347,6 +347,61 @@ FPCGContext* FPCGSubgraphElement::Initialize(const FPCGDataCollection& InputData
 	return Context;
 }
 
+void FPCGSubgraphElement::PrepareSubgraphData(const UPCGSubgraphSettings* Settings, FPCGSubgraphContext* Context, const FPCGDataCollection& InputData, FPCGDataCollection& OutputData) const
+{
+	// Don't forward overrides
+	if (Settings->HasOverridableParams())
+	{
+		OutputData.TaggedData.Reserve(InputData.TaggedData.Num());
+		const TArray<FPCGPinProperties> InputPins = Settings->DefaultInputPinProperties();
+
+		for (const FPCGTaggedData& Input : InputData.TaggedData)
+		{
+			if (!Input.Data)
+			{
+				continue;
+			}
+
+			// Discard params that don't have a pin on the subgraph input node
+			if (!Input.Data->IsA<UPCGParamData>() || Algo::FindByPredicate(InputPins, [Pin = Input.Pin](const FPCGPinProperties& PinProperty) { return PinProperty.Label == Pin; }))
+			{
+				OutputData.TaggedData.Add(Input);
+			}
+		}
+	}
+	else
+	{
+		OutputData = InputData;
+	}
+}
+
+void FPCGSubgraphElement::PrepareSubgraphUserParameters(const UPCGSubgraphSettings* Settings, FPCGSubgraphContext* Context, FPCGDataCollection& OutputData) const
+{
+	// Also create a new data containing information about the original subgraph and the parameters override
+	// It is used mainly by the UserParameterGetElement to access the correct value.
+	if (const UPCGGraphInterface* SubgraphInterface = Settings->GetSubgraphInterface())
+	{
+		UPCGUserParametersData* UserParamData = NewObject<UPCGUserParametersData>();
+		UserParamData->OriginalGraph = SubgraphInterface;
+		if (Context->GraphInstanceParametersOverride.IsValid())
+		{
+			UserParamData->UserParameters = std::move(Context->GraphInstanceParametersOverride);
+		}
+
+		FPCGTaggedData& TaggedData = OutputData.TaggedData.Emplace_GetRef();
+		TaggedData.Data = UserParamData;
+		TaggedData.Tags.Add(PCGBaseSubgraphConstants::UserParameterTagData);
+		// Mark this data pinless, since it is internal data, not meant to be shown in the graph editor.
+		TaggedData.bPinlessData = true;
+	}
+}
+
+bool FPCGSubgraphElement::IsCacheable(const UPCGSettings* InSettings) const
+{
+	const UPCGSubgraphSettings* Settings = Cast<const UPCGSubgraphSettings>(InSettings);
+	return (!Settings || !Settings->IsDynamicGraph());
+}
+
 bool FPCGSubgraphElement::ExecuteInternal(FPCGContext* InContext) const
 {
 	FPCGSubgraphContext* Context = static_cast<FPCGSubgraphContext*>(InContext);
@@ -354,8 +409,7 @@ bool FPCGSubgraphElement::ExecuteInternal(FPCGContext* InContext) const
 	const UPCGSubgraphSettings* Settings = Context->GetInputSettings<UPCGSubgraphSettings>();
 	check(Settings);
 
-	const UPCGSubgraphNode* SubgraphNode = Cast<const UPCGSubgraphNode>(Context->Node);
-	if (SubgraphNode && SubgraphNode->bDynamicGraph)
+	if(Settings->IsDynamicGraph())
 	{
 		TRACE_CPUPROFILER_EVENT_SCOPE(FPCGSubgraphElement::Execute);
 
@@ -368,20 +422,35 @@ bool FPCGSubgraphElement::ExecuteInternal(FPCGContext* InContext) const
 			{
 				// Dispatch graph to execute with the given information we have
 				// using this node's task id as additional inputs
-				FPCGTaskId SubgraphTaskId = Subsystem->ScheduleGraph(Subgraph, Context->SourceComponent.Get(), MakeShared<FPCGInputForwardingElement>(Context->InputData), {});
+				FPCGDataCollection PreSubgraphInputData;
+				PrepareSubgraphUserParameters(Settings, Context, PreSubgraphInputData);
+
+				FPCGDataCollection SubgraphInputData;
+				PrepareSubgraphData(Settings, Context, Context->InputData, SubgraphInputData);
 
-				Context->SubgraphTaskId = SubgraphTaskId;
-				Context->bScheduledSubgraph = true;
-				Context->bIsPaused = true;
+				FPCGTaskId SubgraphTaskId = Subsystem->ScheduleGraph(Subgraph, Context->SourceComponent.Get(), MakeShared<FPCGInputForwardingElement>(PreSubgraphInputData), MakeShared<FPCGInputForwardingElement>(SubgraphInputData), {});
 
-				// add a trivial task after the output task that wakes up this task
-				Subsystem->ScheduleGeneric([Context]() {
-					// Wake up the current task
-					Context->bIsPaused = false;
+				if (SubgraphTaskId != InvalidPCGTaskId)
+				{
+					Context->SubgraphTaskIds.Add(SubgraphTaskId);
+					Context->bScheduledSubgraph = true;
+					Context->bIsPaused = true;
+					
+					// add a trivial task after the output task that wakes up this task
+					Subsystem->ScheduleGeneric([Context]() {
+						// Wake up the current task
+						Context->bIsPaused = false;
 					return true;
-					}, Context->SourceComponent.Get(), { Context->SubgraphTaskId });
+					}, Context->SourceComponent.Get(), Context->SubgraphTaskIds);
 
-				return false;
+					return false;
+				}
+				else
+				{
+					// Scheduling failed - early out
+					Context->OutputData.bCancelExecution = true;
+					return true;
+				}
 			}
 			else
 			{
@@ -402,7 +471,14 @@ bool FPCGSubgraphElement::ExecuteInternal(FPCGContext* InContext) const
 			UPCGSubsystem* Subsystem = Context->SourceComponent->GetSubsystem();
 			if (Subsystem)
 			{
-				ensure(Subsystem->GetOutputData(Context->SubgraphTaskId, Context->OutputData));
+				if (Context->SubgraphTaskIds.Num() > 0)
+				{
+					// This element does not support multiple results/dispatches. Note that this would be easy to fix;
+					// as we just can call the GetOutputData on a fresh data collection and merge it afterwards in the output data,
+					// but this is not needed here so we will keep the full assignment & benefit from the crc as well.
+					ensure(Context->SubgraphTaskIds.Num() == 1);
+					ensure(Subsystem->GetOutputData(Context->SubgraphTaskIds[0], Context->OutputData));
+				}
 			}
 			else
 			{
@@ -415,49 +491,9 @@ bool FPCGSubgraphElement::ExecuteInternal(FPCGContext* InContext) const
 	}
 	else
 	{
-		// Don't forward overrides
-		if (Settings->HasOverridableParams())
-		{
-			Context->OutputData.TaggedData.Reserve(Context->InputData.TaggedData.Num());
-			const TArray<FPCGPinProperties> InputPins = Settings->DefaultInputPinProperties();
-
-			for (const FPCGTaggedData& InputData : Context->InputData.TaggedData)
-			{
-				if (!InputData.Data)
-				{
-					continue;
-				}
-
-				// Discard params that don't have a pin on the subgraph input node
-				if (!InputData.Data->IsA<UPCGParamData>() || Algo::FindByPredicate(InputPins, [Pin = InputData.Pin](const FPCGPinProperties& PinProperty) { return PinProperty.Label == Pin;}))
-				{
-					Context->OutputData.TaggedData.Add(InputData);
-				}
-			}
-		}
-		else
-		{
-			Context->OutputData = Context->InputData;
-		}
-
-		// Also create a new data containing information about the original subgraph and the parameters override
-		// It is used mainly by the UseParameterGetElement to access the correct value.
-		if (const UPCGGraphInterface* SubgraphInterface = Settings->GetSubgraphInterface())
-		{
-			UPCGUserParametersData* UserParamData = NewObject<UPCGUserParametersData>();
-			UserParamData->OriginalGraph = SubgraphInterface;
-			if (Context->GraphInstanceParametersOverride.IsValid())
-			{
-				UserParamData->UserParameters = std::move(Context->GraphInstanceParametersOverride);
-			}
-
-			FPCGTaggedData& TaggedData = Context->OutputData.TaggedData.Emplace_GetRef();
-			TaggedData.Data = UserParamData;
-			TaggedData.Tags.Add(PCGBaseSubgraphConstants::UserParameterTagData);
-			// Mark this data pinless, since it is internal data, not meant to be shown in the graph editor.
-			TaggedData.bPinlessData = true;
-		}
-
+		// This node acts as both the pre-graph node and the input node so it should have both the user parameters & the actual inputs
+		PrepareSubgraphData(Settings, Context, Context->InputData, Context->OutputData);
+		PrepareSubgraphUserParameters(Settings, Context, Context->OutputData);
 		return true;
 	}
 }
diff --git a/Engine/Plugins/Experimental/PCG/Source/PCG/Private/PCGSubsystem.cpp b/Engine/Plugins/Experimental/PCG/Source/PCG/Private/PCGSubsystem.cpp
--- a/Engine/Plugins/Experimental/PCG/Source/PCG/Private/PCGSubsystem.cpp
+++ b/Engine/Plugins/Experimental/PCG/Source/PCG/Private/PCGSubsystem.cpp
@@ -440,11 +440,11 @@ FPCGTaskId UPCGSubsystem::ScheduleGraph(UPCGComponent* SourceComponent, const TA
 	}
 }
 
-FPCGTaskId UPCGSubsystem::ScheduleGraph(UPCGGraph* Graph, UPCGComponent* SourceComponent, FPCGElementPtr InputElement, const TArray<FPCGTaskId>& Dependencies)
+FPCGTaskId UPCGSubsystem::ScheduleGraph(UPCGGraph* Graph, UPCGComponent* SourceComponent, FPCGElementPtr PreGraphElement, FPCGElementPtr InputElement, const TArray<FPCGTaskId>& Dependencies)
 {
 	if (SourceComponent)
 	{
-		return GraphExecutor->Schedule(Graph, SourceComponent, InputElement, Dependencies);
+		return GraphExecutor->Schedule(Graph, SourceComponent, PreGraphElement, InputElement, Dependencies);
 	}
 	else
 	{
diff --git a/Engine/Plugins/Experimental/PCG/Source/PCG/Public/Elements/PCGLoopElement.h b/Engine/Plugins/Experimental/PCG/Source/PCG/Public/Elements/PCGLoopElement.h
new file mode 100644
--- /dev/null
+++ b/Engine/Plugins/Experimental/PCG/Source/PCG/Public/Elements/PCGLoopElement.h
@@ -0,0 +1,46 @@
+// Copyright Epic Games, Inc. All Rights Reserved.
+
+#pragma once
+
+#include "PCGSubgraph.h"
+
+#include "PCGLoopElement.generated.h"
+
+UCLASS(BlueprintType, ClassGroup=(Procedural))
+class PCG_API UPCGLoopSettings : public UPCGSubgraphSettings
+{
+	GENERATED_BODY()
+
+public:
+	//~Begin UPCGSettings interface implementation
+#if WITH_EDITOR
+	virtual FName GetDefaultNodeName() const override { return FName(TEXT("Loop")); }
+	virtual FText GetDefaultNodeTitle() const override;
+	virtual FText GetNodeTooltipText() const override;
+#endif
+
+	virtual FName AdditionalTaskName() const override;
+
+protected:
+	virtual FPCGElementPtr CreateElement() const override;
+	// ~End UPCGSettings interface implementation
+
+public:
+	//~Begin UPCGBaseSubgraphSettings interface
+	virtual bool IsDynamicGraph() const { return true; }
+	//~End UPCGBaseSubgraphSettings interface
+
+	/** 
+	* Comma-separated list of pin names on which we will loop by-element in a step-wise fashion; if more than one is provided, it is expected that they all have the same number of data. 
+	* If none are provided, the first connected pin will taken as the pin to loop on. 
+	*/
+	UPROPERTY(BlueprintReadWrite, EditAnywhere, Category = Data)
+	FString LoopPins;
+};
+
+class PCG_API FPCGLoopElement : public FPCGSubgraphElement
+{
+protected:
+	virtual bool ExecuteInternal(FPCGContext* Context) const override;
+	void PrepareLoopDataCollections(FPCGContext* Context, const UPCGLoopSettings* Settings, TArray<FPCGDataCollection>& LoopDataCollection, FPCGDataCollection& FixedInputDataCollection) const;
+};
\ No newline at end of file
diff --git a/Engine/Plugins/Experimental/PCG/Source/PCG/Public/PCGNode.h b/Engine/Plugins/Experimental/PCG/Source/PCG/Public/PCGNode.h
--- a/Engine/Plugins/Experimental/PCG/Source/PCG/Public/PCGNode.h
+++ b/Engine/Plugins/Experimental/PCG/Source/PCG/Public/PCGNode.h
@@ -115,6 +115,9 @@ public:
 	/** True if the edge is being used by the node. UI will gray out unused pins. */
 	virtual bool IsEdgeUsedByNodeExecution(const UPCGEdge* InEdge) const;
 
+	/** Returns the first connected pin on the node */
+	const UPCGPin* GetFirstConnectedInputPin() const;
+
 	const TArray<TObjectPtr<UPCGPin>>& GetInputPins() const { return InputPins; }
 	const TArray<TObjectPtr<UPCGPin>>& GetOutputPins() const { return OutputPins; }
 
diff --git a/Engine/Plugins/Experimental/PCG/Source/PCG/Public/PCGSubgraph.h b/Engine/Plugins/Experimental/PCG/Source/PCG/Public/PCGSubgraph.h
--- a/Engine/Plugins/Experimental/PCG/Source/PCG/Public/PCGSubgraph.h
+++ b/Engine/Plugins/Experimental/PCG/Source/PCG/Public/PCGSubgraph.h
@@ -22,6 +22,7 @@ class PCG_API UPCGBaseSubgraphSettings : public UPCGSettings
 public:
 	UPCGGraph* GetSubgraph() const;
 	virtual UPCGGraphInterface* GetSubgraphInterface() const { return nullptr; }
+	virtual bool IsDynamicGraph() const { return false; }
 
 	// Use this method from the outside to set the subgraph, as it will connect editor callbacks
 	virtual void SetSubgraph(UPCGGraphInterface* InGraph);
@@ -115,9 +116,6 @@ class PCG_API UPCGBaseSubgraphNode : public UPCGNode
 	GENERATED_BODY()
 
 public:
-	UPROPERTY(BlueprintReadWrite, EditAnywhere, Category = Settings)
-	bool bDynamicGraph = false;
-
 	TObjectPtr<UPCGGraph> GetSubgraph() const;
 	virtual TObjectPtr<UPCGGraphInterface> GetSubgraphInterface() const { return nullptr; }
 };
@@ -135,7 +133,7 @@ public:
 
 struct PCG_API FPCGSubgraphContext : public FPCGContext
 {
-	FPCGTaskId SubgraphTaskId = InvalidPCGTaskId;
+	TArray<FPCGTaskId> SubgraphTaskIds;
 	bool bScheduledSubgraph = false;
 	FInstancedStruct GraphInstanceParametersOverride;
 
@@ -147,10 +145,13 @@ class PCG_API FPCGSubgraphElement : public IPCGElement
 {
 public:
 	virtual FPCGContext* Initialize(const FPCGDataCollection& InputData, TWeakObjectPtr<UPCGComponent> SourceComponent, const UPCGNode* Node) override;
+	virtual bool IsCacheable(const UPCGSettings* InSettings) const override;
 
 protected:
 	virtual bool ExecuteInternal(FPCGContext* Context) const override;
 	virtual bool IsPassthrough(const UPCGSettings* InSettings) const override { return !InSettings || InSettings->bEnabled; }
+	void PrepareSubgraphData(const UPCGSubgraphSettings* Settings, FPCGSubgraphContext* Context, const FPCGDataCollection& InputData, FPCGDataCollection& OutputData) const;
+	void PrepareSubgraphUserParameters(const UPCGSubgraphSettings* Settings, FPCGSubgraphContext* Context, FPCGDataCollection& OutputData) const;
 };
 
 class PCG_API FPCGInputForwardingElement : public FSimplePCGElement
diff --git a/Engine/Plugins/Experimental/PCG/Source/PCG/Public/PCGSubsystem.h b/Engine/Plugins/Experimental/PCG/Source/PCG/Public/PCGSubsystem.h
--- a/Engine/Plugins/Experimental/PCG/Source/PCG/Public/PCGSubsystem.h
+++ b/Engine/Plugins/Experimental/PCG/Source/PCG/Public/PCGSubsystem.h
@@ -87,7 +87,7 @@ public:
 	FPCGTaskId ScheduleCleanup(UPCGComponent* PCGComponent, bool bRemoveComponents, bool bSave, const TArray<FPCGTaskId>& Dependencies);
 
 	// Schedule graph (used internally for dynamic subgraph execution)
-	FPCGTaskId ScheduleGraph(UPCGGraph* Graph, UPCGComponent* SourceComponent, FPCGElementPtr InputElement, const TArray<FPCGTaskId>& Dependencies);
+	FPCGTaskId ScheduleGraph(UPCGGraph* Graph, UPCGComponent* SourceComponent, FPCGElementPtr PreGraphElement, FPCGElementPtr InputElement, const TArray<FPCGTaskId>& Dependencies);
 
 	// Schedule graph (used internally for dynamic subgraph execution)
 	FPCGTaskId ScheduleGraph(UPCGComponent* SourceComponent, const TArray<FPCGTaskId>& Dependencies);
