Added cancel-on-refresh

diff --git a/Engine/Plugins/Experimental/PCG/Source/PCG/Private/PCGComponent.cpp b/Engine/Plugins/Experimental/PCG/Source/PCG/Private/PCGComponent.cpp
--- a/Engine/Plugins/Experimental/PCG/Source/PCG/Private/PCGComponent.cpp
+++ b/Engine/Plugins/Experimental/PCG/Source/PCG/Private/PCGComponent.cpp
@@ -532,7 +532,8 @@ void UPCGComponent::OnProcessGraphAborted(bool bQuiet)
 
 #if WITH_EDITOR
 	CurrentRefreshTask = InvalidPCGTaskId;
-	bDirtyGenerated = false;
+	// Implementation note: while it may seem logical to clear the bDirtyGenerated flag here, 
+	// the component is still considered dirty if we aborted processing, hence it should stay this way.
 #endif
 }
 
@@ -1709,6 +1710,11 @@ void UPCGComponent::Refresh()
 	// Discard any refresh if have already one scheduled.
 	if (UPCGSubsystem* Subsystem = GetSubsystem())
 	{
+		if (CurrentGenerationTask != InvalidPCGTaskId)
+		{
+			CancelGeneration();
+		}
+
 		if (CurrentRefreshTask == InvalidPCGTaskId)
 		{
 			CurrentRefreshTask = Subsystem->ScheduleRefresh(this);
