[PCG] Improvements on Attribute operations

* All metadata operations will now propagate the first non-param pin by default
* Add a new "OutputDataFromPin" settings, to allow to propagte something else than the previous rule.
* Add option on PCGSettings to set a given pin label style and add icons next to pin labels.
* Pins that flows have an arrow next to their name
* Add "MakeRotFromAngles" in MakeRotator
* Add "DefaultValue" concept, implemented in all Make operations

diff --git a/Engine/Plugins/Experimental/PCG/Source/PCG/Private/Elements/Metadata/PCGMetadataMakeRotator.cpp b/Engine/Plugins/Experimental/PCG/Source/PCG/Private/Elements/Metadata/PCGMetadataMakeRotator.cpp
--- a/Engine/Plugins/Experimental/PCG/Source/PCG/Private/Elements/Metadata/PCGMetadataMakeRotator.cpp
+++ b/Engine/Plugins/Experimental/PCG/Source/PCG/Private/Elements/Metadata/PCGMetadataMakeRotator.cpp
@@ -2,7 +2,9 @@
 
 #include "Elements/Metadata/PCGMetadataMakeRotator.h"
 
+#include "PCGParamData.h"
 #include "Elements/Metadata/PCGMetadataElementCommon.h"
+#include "Metadata/PCGMetadata.h"
 
 #define LOCTEXT_NAMESPACE "PCGMetadataMakeRotatorSettings"
 
@@ -28,9 +30,12 @@ FName UPCGMetadataMakeRotatorSettings::GetInputPinLabel(uint32 Index) const
 		return (Index == 0 ? PCGMetadataMakeRotatorConstants::YLabel : PCGMetadataMakeRotatorConstants::ZLabel);
 	case EPCGMetadataMakeRotatorOp::MakeRotFromZY:
 		return (Index == 1 ? PCGMetadataMakeRotatorConstants::YLabel : PCGMetadataMakeRotatorConstants::ZLabel);
-	case EPCGMetadataMakeRotatorOp::MakeRotFromAxis:
-	default:
+	case EPCGMetadataMakeRotatorOp::MakeRotFromAxes:
 		return (Index == 0 ? PCGMetadataMakeRotatorConstants::ForwardLabel : (Index == 1 ? PCGMetadataMakeRotatorConstants::RightLabel : PCGMetadataMakeRotatorConstants::UpLabel));
+	case EPCGMetadataMakeRotatorOp::MakeRotFromAngles:
+		return (Index == 0 ? PCGMetadataMakeRotatorConstants::RollLabel : (Index == 1 ? PCGMetadataMakeRotatorConstants::PitchLabel : PCGMetadataMakeRotatorConstants::YawLabel));
+	default:
+		return NAME_None;
 	}
 }
 
@@ -49,7 +54,8 @@ uint32 UPCGMetadataMakeRotatorSettings::GetInputPinNum() const
 	case EPCGMetadataMakeRotatorOp::MakeRotFromYZ: // fall-through
 	case EPCGMetadataMakeRotatorOp::MakeRotFromZY:
 		return 2;
-	case EPCGMetadataMakeRotatorOp::MakeRotFromAxis:
+	case EPCGMetadataMakeRotatorOp::MakeRotFromAxes: // fall-through
+	case EPCGMetadataMakeRotatorOp::MakeRotFromAngles:
 	default:
 		return 3;
 	}
@@ -58,7 +64,14 @@ uint32 UPCGMetadataMakeRotatorSettings::GetInputPinNum() const
 bool UPCGMetadataMakeRotatorSettings::IsSupportedInputType(uint16 TypeId, uint32 InputIndex, bool& bHasSpecialRequirement) const
 {
 	bHasSpecialRequirement = false;
-	return PCG::Private::IsOfTypes<FVector, FVector2D, float, double, int32, int64>(TypeId);
+	if (Operation == EPCGMetadataMakeRotatorOp::MakeRotFromAngles)
+	{
+		return PCG::Private::IsOfTypes<float, double, int32, int64>(TypeId);
+	}
+	else
+	{
+		return PCG::Private::IsOfTypes<FVector, FVector2D, float, double, int32, int64>(TypeId);
+	}
 }
 
 FPCGAttributePropertySelector UPCGMetadataMakeRotatorSettings::GetInputSource(uint32 Index) const
@@ -114,6 +127,30 @@ FPCGElementPtr UPCGMetadataMakeRotatorSettings::CreateElement() const
 	return MakeShared<FPCGMetadataMakeRotatorElement>();
 }
 
+bool UPCGMetadataMakeRotatorSettings::DoesInputSupportDefaultValue(uint32 Index) const
+{
+	return Operation == EPCGMetadataMakeRotatorOp::MakeRotFromAngles;
+}
+
+UPCGParamData* UPCGMetadataMakeRotatorSettings::CreateDefaultValueParam(uint32 Index) const
+{
+	if (Operation != EPCGMetadataMakeRotatorOp::MakeRotFromAngles)
+	{
+		return nullptr;
+	}
+
+	UPCGParamData* NewParamData = NewObject<UPCGParamData>();
+	NewParamData->Metadata->CreateAttribute<double>(NAME_None, 0, /*bAllowsInterpolation=*/ true, /*bOverrideParent=*/ false);
+	return NewParamData;
+}
+
+#if WITH_EDITOR
+FString UPCGMetadataMakeRotatorSettings::GetDefaultValueString(uint32 Index) const
+{
+	return (Operation == EPCGMetadataMakeRotatorOp::MakeRotFromAngles ? FString::Printf(TEXT("%f"), 0.0) : FString());
+}
+#endif // WITH_EDITOR
+
 bool FPCGMetadataMakeRotatorElement::DoOperation(FOperationData& OperationData) const
 {
 	TRACE_CPUPROFILER_EVENT_SCOPE(FPCGMetadataMakeRotatorElement::Execute);
@@ -140,11 +177,16 @@ bool FPCGMetadataMakeRotatorElement::DoOperation(FOperationData& OperationData)
 		return DoBinaryOp<FVector, FVector>(OperationData, [](const FVector& Y, const FVector& Z) -> FRotator { return FRotationMatrix::MakeFromYZ(Y, Z).Rotator(); });
 	case EPCGMetadataMakeRotatorOp::MakeRotFromZY:
 		return DoBinaryOp<FVector, FVector>(OperationData, [](const FVector& Z, const FVector& Y) -> FRotator { return FRotationMatrix::MakeFromZY(Z, Y).Rotator(); });
-	case EPCGMetadataMakeRotatorOp::MakeRotFromAxis:
+	case EPCGMetadataMakeRotatorOp::MakeRotFromAxes:
 		return DoTernaryOp<FVector, FVector, FVector>(OperationData, [](const FVector& X, const FVector& Y, const FVector& Z) -> FRotator
 		{
 			return FMatrix(X.GetSafeNormal(), Y.GetSafeNormal(), Z.GetSafeNormal(), FVector::ZeroVector).Rotator();
 		});
+	case EPCGMetadataMakeRotatorOp::MakeRotFromAngles:
+		return DoTernaryOp<double, double, double>(OperationData, [](const double& Roll, const double& Pitch, const double& Yaw) -> FRotator
+		{
+			return FRotator{ Pitch, Yaw, Roll };
+		});
 	}
 
 	return true;
diff --git a/Engine/Plugins/Experimental/PCG/Source/PCG/Private/Elements/Metadata/PCGMetadataMakeTransform.cpp b/Engine/Plugins/Experimental/PCG/Source/PCG/Private/Elements/Metadata/PCGMetadataMakeTransform.cpp
--- a/Engine/Plugins/Experimental/PCG/Source/PCG/Private/Elements/Metadata/PCGMetadataMakeTransform.cpp
+++ b/Engine/Plugins/Experimental/PCG/Source/PCG/Private/Elements/Metadata/PCGMetadataMakeTransform.cpp
@@ -2,6 +2,9 @@
 
 #include "Elements/Metadata/PCGMetadataMakeTransform.h"
 
+#include "PCGParamData.h"
+#include "Elements/Metadata/PCGMetadataElementCommon.h"
+#include "Metadata/PCGMetadata.h"
 #include "Metadata/PCGMetadataAttributeTpl.h"
 
 #include UE_INLINE_GENERATED_CPP_BY_NAME(PCGMetadataMakeTransform)
@@ -118,6 +121,67 @@ FText UPCGMetadataMakeTransformSettings::GetDefaultNodeTitle() const
 }
 #endif // WITH_EDITOR
 
+bool UPCGMetadataMakeTransformSettings::DoesInputSupportDefaultValue(uint32 Index) const
+{
+	return true;
+}
+
+UPCGParamData* UPCGMetadataMakeTransformSettings::CreateDefaultValueParam(uint32 Index) const
+{
+	switch (Index)
+	{
+	case 0:
+	{
+		// Location -> Default is Zero vector
+		UPCGParamData* NewParamData = NewObject<UPCGParamData>();
+		NewParamData->Metadata->CreateAttribute<FVector>(NAME_None, FVector::ZeroVector, /*bAllowsInterpolation=*/ true, /*bOverrideParent=*/ false);
+		return NewParamData;
+	}
+	case 1:
+	{
+		// Rotation -> Default is Zero rotator
+		UPCGParamData* NewParamData = NewObject<UPCGParamData>();
+		NewParamData->Metadata->CreateAttribute<FRotator>(NAME_None, FRotator::ZeroRotator, /*bAllowsInterpolation=*/ true, /*bOverrideParent=*/ false);
+		return NewParamData;
+	}
+	case 2:
+	{
+		// Scale -> Default is Vector (1, 1, 1)
+		UPCGParamData* NewParamData = NewObject<UPCGParamData>();
+		NewParamData->Metadata->CreateAttribute<FVector>(NAME_None, FVector::OneVector, /*bAllowsInterpolation=*/ true, /*bOverrideParent=*/ false);
+		return NewParamData;
+	}
+	default:
+		return nullptr;
+	}
+}
+
+#if WITH_EDITOR
+FString UPCGMetadataMakeTransformSettings::GetDefaultValueString(uint32 Index) const
+{
+	switch (Index)
+	{
+	case 0:
+	{
+		// Location -> Default is Zero vector
+		return FVector::ZeroVector.ToString();
+	}
+	case 1:
+	{
+		// Rotation -> Default is Zero rotator
+		return FRotator::ZeroRotator.ToString();
+	}
+	case 2:
+	{
+		// Scale -> Default is Vector (1, 1, 1)
+		return FVector::OneVector.ToString();
+	}
+	default:
+		return FString();
+	}
+}
+#endif // WITH_EDITOR
+
 FPCGElementPtr UPCGMetadataMakeTransformSettings::CreateElement() const
 {
 	return MakeShared<FPCGMetadataMakeTransformElement>();
diff --git a/Engine/Plugins/Experimental/PCG/Source/PCG/Private/Elements/Metadata/PCGMetadataMakeVector.cpp b/Engine/Plugins/Experimental/PCG/Source/PCG/Private/Elements/Metadata/PCGMetadataMakeVector.cpp
--- a/Engine/Plugins/Experimental/PCG/Source/PCG/Private/Elements/Metadata/PCGMetadataMakeVector.cpp
+++ b/Engine/Plugins/Experimental/PCG/Source/PCG/Private/Elements/Metadata/PCGMetadataMakeVector.cpp
@@ -2,6 +2,9 @@
 
 #include "Elements/Metadata/PCGMetadataMakeVector.h"
 
+#include "PCGParamData.h"
+#include "Elements/Metadata/PCGMetadataElementCommon.h"
+#include "Metadata/PCGMetadata.h"
 
 #include UE_INLINE_GENERATED_CPP_BY_NAME(PCGMetadataMakeVector)
 
@@ -193,6 +196,54 @@ FText UPCGMetadataMakeVectorSettings::GetDefaultNodeTitle() const
 }
 #endif // WITH_EDITOR
 
+bool UPCGMetadataMakeVectorSettings::DoesInputSupportDefaultValue(uint32 Index) const
+{
+	return true;
+}
+
+UPCGParamData* UPCGMetadataMakeVectorSettings::CreateDefaultValueParam(uint32 Index) const
+{
+	// Use labels since the logic is already done there.
+	FName Label = GetInputPinLabel(Index);
+	UPCGParamData* NewParamData = NewObject<UPCGParamData>();
+
+	if (Label == PCGMetadataMakeVectorConstants::XYZLabel)
+	{
+		NewParamData->Metadata->CreateAttribute<FVector>(NAME_None, FVector::ZeroVector, /*bAllowsInterpolation=*/ true, /*bOverrideParent=*/ false);
+	}
+	else if (Label == PCGMetadataMakeVectorConstants::XYLabel || Label == PCGMetadataMakeVectorConstants::ZWLabel)
+	{
+		NewParamData->Metadata->CreateAttribute<FVector2D>(NAME_None, FVector2D::ZeroVector, /*bAllowsInterpolation=*/ true, /*bOverrideParent=*/ false);
+	}
+	else
+	{
+		NewParamData->Metadata->CreateAttribute<double>(NAME_None, 0, /*bAllowsInterpolation=*/ true, /*bOverrideParent=*/ false);
+	}
+	
+	return NewParamData;
+}
+
+#if WITH_EDITOR
+FString UPCGMetadataMakeVectorSettings::GetDefaultValueString(uint32 Index) const
+{
+	// Use labels since the logic is already done there.
+	FName Label = GetInputPinLabel(Index);
+
+	if (Label == PCGMetadataMakeVectorConstants::XYZLabel)
+	{
+		return FVector::ZeroVector.ToString();
+	}
+	else if (Label == PCGMetadataMakeVectorConstants::XYLabel || Label == PCGMetadataMakeVectorConstants::ZWLabel)
+	{
+		return FVector2D::ZeroVector.ToString();
+	}
+	else
+	{
+		return FString::Printf(TEXT("%f"), 0.0);
+	}
+}
+#endif // WITH_EDITOR
+
 FPCGElementPtr UPCGMetadataMakeVectorSettings::CreateElement() const
 {
 	return MakeShared<FPCGMetadataMakeVectorElement>();
diff --git a/Engine/Plugins/Experimental/PCG/Source/PCG/Private/Elements/Metadata/PCGMetadataOpElementBase.cpp b/Engine/Plugins/Experimental/PCG/Source/PCG/Private/Elements/Metadata/PCGMetadataOpElementBase.cpp
--- a/Engine/Plugins/Experimental/PCG/Source/PCG/Private/Elements/Metadata/PCGMetadataOpElementBase.cpp
+++ b/Engine/Plugins/Experimental/PCG/Source/PCG/Private/Elements/Metadata/PCGMetadataOpElementBase.cpp
@@ -28,30 +28,188 @@ void UPCGMetadataSettingsBase::PostLoad()
 #endif // WITH_EDITOR
 }
 
+EPCGDataType UPCGMetadataSettingsBase::GetInputPinType(uint32 Index) const
+{
+	const EPCGDataType FirstPinTypeUnion = GetTypeUnionOfIncidentEdges(GetInputPinLabel(Index));
+
+	// If the pin is not connected but support default value, treat it as a param
+	if (FirstPinTypeUnion == EPCGDataType::None && DoesInputSupportDefaultValue(Index))
+	{
+		return EPCGDataType::Param;
+	}
+
+	return FirstPinTypeUnion != EPCGDataType::None ? FirstPinTypeUnion : EPCGDataType::Any;
+}
+
+TArray<FName> UPCGMetadataSettingsBase::GetOutputDataFromPinOptions() const
+{
+	const uint32 NumInputs = GetInputPinNum();
+
+	TArray<FName> AllOptions;
+	AllOptions.SetNum(NumInputs + 1);
+	AllOptions[0] = PCGMetadataSettingsBaseConstants::DefaultOutputDataFromPinName;
+
+	for (uint32 Index = 0; Index < NumInputs; ++Index)
+	{
+		AllOptions[Index + 1] = GetInputPinLabel(Index);
+	}
+
+	return AllOptions;
+}
+
+uint32 UPCGMetadataSettingsBase::GetInputPinIndex(FName InPinLabel) const
+{
+	if (InPinLabel != PCGMetadataSettingsBaseConstants::DefaultOutputDataFromPinName)
+	{
+		for (uint32 Index = 0; Index < GetInputPinNum(); ++Index)
+		{
+			if (InPinLabel == GetInputPinLabel(Index))
+			{
+				return Index;
+			}
+		}
+	}
+
+	return (uint32)INDEX_NONE;
+}
+
+uint32 UPCGMetadataSettingsBase::GetInputPinToForward() const
+{
+	const uint32 NumberOfInputs = GetInputPinNum();
+	uint32 InputPinToForward = 0;
+
+	// If there is only one input, it is trivial
+	if (NumberOfInputs != 1)
+	{
+		// Heuristic:
+		//	* If OutputDataFromPin is set, use this value
+		//	* If there are connected pins, use the first spatial input (not Any)
+		//	* Otherwise, take the first pin
+		const uint32 OutputDataFromPinIndex = GetInputPinIndex(OutputDataFromPin);
+		if (OutputDataFromPinIndex != (uint32)INDEX_NONE)
+		{
+			InputPinToForward = OutputDataFromPinIndex;
+		}
+		else
+		{
+			for (uint32 InputPinIndex = 0; InputPinIndex < NumberOfInputs; ++InputPinIndex)
+			{
+				const EPCGDataType PinType = GetInputPinType(InputPinIndex);
+
+				if ((PinType != EPCGDataType::Any) && !!(PinType & EPCGDataType::Spatial))
+				{
+					InputPinToForward = InputPinIndex;
+					break;
+				}
+			}
+		}
+	}
+
+	return InputPinToForward;
+}
+
 EPCGDataType UPCGMetadataSettingsBase::GetCurrentPinTypes(const UPCGPin* InPin) const
 {
 	check(InPin);
-	if (!InPin->IsOutputPin() || GetInputPinNum() == 0)
+	const uint32 NumberOfInputs = GetInputPinNum();
+	if (!InPin->IsOutputPin() || NumberOfInputs == 0)
 	{
 		// Fall back to default for input pins, or if no input pins present from which to obtain type
 		return Super::GetCurrentPinTypes(InPin);
 	}
 
-	// Output pin narrows to union of inputs on first pin
-	const EPCGDataType FirstPinTypeUnion = GetTypeUnionOfIncidentEdges(GetInputPinLabel(0));
-	return FirstPinTypeUnion != EPCGDataType::None ? FirstPinTypeUnion : EPCGDataType::Any;
+	// Output pin narrows to union of inputs on pin to forward
+	return GetInputPinType(GetInputPinToForward());
+}
+
+#if WITH_EDITOR
+void UPCGMetadataSettingsBase::PostEditChangeProperty(FPropertyChangedEvent& PropertyChangedEvent)
+{
+	Super::PostEditChangeProperty(PropertyChangedEvent);
+
+	// Make sure the output data from pin value is always valid. Reset it otherwise.
+	if (GetInputPinIndex(OutputDataFromPin) == (uint32)INDEX_NONE)
+	{
+		OutputDataFromPin = PCGMetadataSettingsBaseConstants::DefaultOutputDataFromPinName;
+	}
+}
+
+bool UPCGMetadataSettingsBase::CanEditChange(const FProperty* InProperty) const
+{
+	if (!Super::CanEditChange(InProperty))
+	{
+		return false;
+	}
+
+	return InProperty && ((InProperty->GetFName() != GET_MEMBER_NAME_CHECKED(UPCGMetadataSettingsBase, OutputDataFromPin)) || (GetInputPinNum() != 1));
+}
+
+bool UPCGMetadataSettingsBase::GetPinExtraIcon(const UPCGPin* InPin, FName& OutExtraIcon, FText& OutTooltip) const
+{
+	// Only set the arrow if the output data from pin is forced.
+	if (GetInputPinNum() > 1 && InPin && OutputDataFromPin == InPin->Properties.Label)
+	{
+		OutExtraIcon = TEXT("Icons.ArrowRight");
+		return true;
+	}
+	else
+	{
+		return false;
+	}
 }
 
+FText UPCGMetadataSettingsBase::GetNodeTooltipText() const
+{
+	return LOCTEXT("NodeTooltip", "Metadata operation between Points/Spatial/AttributeSet data.\n"
+		"Output data will be taken from the first spatial data by default, or first pin if all are attribute sets.\n"
+		"It can be overridden in the settings.");
+}
+#endif // WITH_EDITOR
+
 TArray<FPCGPinProperties> UPCGMetadataSettingsBase::InputPinProperties() const
 {
 	TArray<FPCGPinProperties> PinProperties;
 
-	for (uint32 i = 0; i < GetInputPinNum(); ++i)
+	for (uint32 InputPinIndex = 0; InputPinIndex < GetInputPinNum(); ++InputPinIndex)
 	{
-		const FName PinLabel = GetInputPinLabel(i);
+		const FName PinLabel = GetInputPinLabel(InputPinIndex);
 		if (PinLabel != NAME_None)
 		{
 			PinProperties.Emplace(PinLabel, EPCGDataType::Any, /*bAllowMultipleConnections=*/false);
+
+#if WITH_EDITOR
+			TArray<FText> AllTooltips;
+			TArray<FString> SupportedTypes;
+
+			for (uint8 TypeId = 0; TypeId < (uint8)EPCGMetadataTypes::Count; ++TypeId)
+			{
+				bool DummyBool;
+				if (IsSupportedInputType(TypeId, InputPinIndex, DummyBool))
+				{
+					SupportedTypes.Add(PCG::Private::GetTypeName(TypeId));
+				}
+			}
+			
+			if (!SupportedTypes.IsEmpty())
+			{
+				AllTooltips.Add(FText::Format(LOCTEXT("PinTooltipSupportedTypes", "Supported types: {0}"), FText::FromString(FString::Join(SupportedTypes, TEXT(", ")))));
+			}
+
+			if (GetInputPinNum() > 1 && OutputDataFromPin == PinLabel)
+			{
+				AllTooltips.Add(LOCTEXT("PinTooltipForwardInput", "This input will be forwarded to the output."));
+			}
+
+			if (DoesInputSupportDefaultValue(InputPinIndex))
+			{
+				AllTooltips.Add(FText::Format(LOCTEXT("PinTooltipDefaultValue", "Pin is optional, will use default value if not connected ({0})"), FText::FromString(GetDefaultValueString(InputPinIndex))));
+			}
+
+			if (!AllTooltips.IsEmpty())
+			{
+				PinProperties.Last().Tooltip = FText::Join(FText::FromString(TEXT("\n")), AllTooltips);
+			}
+#endif // WITH_EDITOR
 		}
 	}
 
@@ -62,9 +220,9 @@ TArray<FPCGPinProperties> UPCGMetadataSettingsBase::OutputPinProperties() const
 {
 	TArray<FPCGPinProperties> PinProperties;
 	
-	for (uint32 i = 0; i < GetOutputPinNum(); ++i)
+	for (uint32 OutputPinIndex = 0; OutputPinIndex < GetOutputPinNum(); ++OutputPinIndex)
 	{
-		const FName PinLabel = GetOutputPinLabel(i);
+		const FName PinLabel = GetOutputPinLabel(OutputPinIndex);
 		if (PinLabel != NAME_None)
 		{
 			PinProperties.Emplace(PinLabel, EPCGDataType::Any);
@@ -98,10 +256,26 @@ bool FPCGMetadataElementBase::ExecuteInternal(FPCGContext* Context) const
 	SourceAttribute.SetNum(NumberOfInputs);
 	InputTaggedData.SetNum(NumberOfInputs);
 
-	for (uint32 i = 0; i < NumberOfInputs; ++i)
+	uint32 InputPinToForward = Settings->GetInputPinToForward();
+
+	for (uint32 InputPinIndex = 0; InputPinIndex < NumberOfInputs; ++InputPinIndex)
 	{
-		FName PinLabel = Settings->GetInputPinLabel(i);
+		FName PinLabel = Settings->GetInputPinLabel(InputPinIndex);
 		TArray<FPCGTaggedData> InputData = Context->InputData.GetInputsByPin(PinLabel);
+
+		// If input data is empty but we can have default values, add a "fake" param data input.
+		if (InputData.IsEmpty() && Settings->DoesInputSupportDefaultValue(InputPinIndex))
+		{
+			FPCGTaggedData& DummyData = InputData.Emplace_GetRef();
+			DummyData.Pin = PinLabel;
+			DummyData.Data = Settings->CreateDefaultValueParam(InputPinIndex);
+			if (!DummyData.Data)
+			{
+				PCGE_LOG(Error, GraphAndLog, FText::Format(LOCTEXT("CantCreateDefaultValue", "Pin {0} supports default value but we could not create it."), FText::FromName(PinLabel)));
+				return true;
+			}
+		}
+
 		if (InputData.IsEmpty())
 		{
 			// Absence of data not worth broadcasting to user as visual warning
@@ -115,16 +289,16 @@ bool FPCGMetadataElementBase::ExecuteInternal(FPCGContext* Context) const
 		}
 
 		// By construction, there can only be one of then(hence the 0 index)
-		InputTaggedData[i] = MoveTemp(InputData[0]);
+		InputTaggedData[InputPinIndex] = MoveTemp(InputData[0]);
 
 		// Only gather Spacial and Params input. 
-		if (const UPCGSpatialData* SpatialInput = Cast<const UPCGSpatialData>(InputTaggedData[i].Data))
+		if (const UPCGSpatialData* SpatialInput = Cast<const UPCGSpatialData>(InputTaggedData[InputPinIndex].Data))
 		{
-			SourceMetadata[i] = SpatialInput->Metadata;
+			SourceMetadata[InputPinIndex] = SpatialInput->Metadata;
 		}
-		else if (const UPCGParamData* ParamsInput = Cast<const UPCGParamData>(InputTaggedData[i].Data))
+		else if (const UPCGParamData* ParamsInput = Cast<const UPCGParamData>(InputTaggedData[InputPinIndex].Data))
 		{
-			SourceMetadata[i] = ParamsInput->Metadata;
+			SourceMetadata[InputPinIndex] = ParamsInput->Metadata;
 		}
 		else
 		{
@@ -148,45 +322,52 @@ bool FPCGMetadataElementBase::ExecuteInternal(FPCGContext* Context) const
 	// Only set it if the first input is an attribute.
 	FName InputName = NAME_None;
 
-	for (uint32 i = 0; i < NumberOfInputs; ++i)
+	auto CreateInputAccessorAndValidate = [this, Settings, InputPinToForward, &InputTaggedData, Context, &SourceMetadata, &OperationData, &InputName, &NumberOfElements, &bNoOperationNeeded, &SourceAttribute](uint32 Index) -> bool
 	{
-		// First we verify if the input data match the first one.
-		if (i != 0 &&
-			InputTaggedData[0].Data->GetClass() != InputTaggedData[i].Data->GetClass() &&
-			!InputTaggedData[i].Data->IsA<UPCGParamData>())
+		// First we verify if the input data match the first one. If the pin to forward is not connected, behave like a param data
+		const UClass* InputPinToForwardClass = (InputTaggedData[InputPinToForward].Data ? InputTaggedData[InputPinToForward].Data->GetClass() : UPCGParamData::StaticClass());
+
+		if (Index != InputPinToForward && 
+			InputPinToForwardClass != InputTaggedData[Index].Data->GetClass() &&
+			!InputTaggedData[Index].Data->IsA<UPCGParamData>())
 		{
-			PCGE_LOG(Error, GraphAndLog, FText::Format(LOCTEXT("InputTypeMismatch", "Data on pin {0} is not of the same type than on pin {1} and is not an Attribute Set. This is not supported."), FText::FromName(InputTaggedData[i].Pin), FText::FromName(InputTaggedData[0].Pin)));
-			return true;
+			PCGE_LOG(Error, GraphAndLog, FText::Format(LOCTEXT("InputTypeMismatch", "Data on pin {0} is not of the same type than on pin {1} and is not an Attribute Set. This is not supported."), FText::FromName(InputTaggedData[Index].Pin), FText::FromName(InputTaggedData[InputPinToForward].Pin)));
+			return false;
 		}
 
-		FPCGAttributePropertySelector InputSource = Settings->GetInputSource(i);
+		FPCGAttributePropertySelector InputSource = Settings->GetInputSource(Index);
 
-		if (InputSource.Selection == EPCGAttributePropertySelection::Attribute && InputSource.AttributeName == NAME_None)
+		if (InputSource.Selection == EPCGAttributePropertySelection::Attribute && InputSource.AttributeName == NAME_None && SourceMetadata[Index])
 		{
-			InputSource.AttributeName = SourceMetadata[i]->GetLatestAttributeNameOrNone();
+			InputSource.AttributeName = SourceMetadata[Index]->GetLatestAttributeNameOrNone();
 		}
 
-		if (i == 0 && InputSource.Selection == EPCGAttributePropertySelection::Attribute)
+		if (Index == InputPinToForward && InputSource.Selection == EPCGAttributePropertySelection::Attribute)
 		{
 			InputName = InputSource.AttributeName;
 		}
 
-		OperationData.InputAccessors[i] = PCGAttributeAccessorHelpers::CreateConstAccessor(InputTaggedData[i].Data, InputSource);
-		OperationData.InputKeys[i] = PCGAttributeAccessorHelpers::CreateConstKeys(InputTaggedData[i].Data, InputSource);
+		OperationData.InputAccessors[Index] = PCGAttributeAccessorHelpers::CreateConstAccessor(InputTaggedData[Index].Data, InputSource);
+		OperationData.InputKeys[Index] = PCGAttributeAccessorHelpers::CreateConstKeys(InputTaggedData[Index].Data, InputSource);
 
-		if (!OperationData.InputAccessors[i].IsValid() || !OperationData.InputKeys[i].IsValid())
+		if (!OperationData.InputAccessors[Index].IsValid() || !OperationData.InputKeys[Index].IsValid())
 		{
-			PCGE_LOG(Error, GraphAndLog, FText::Format(LOCTEXT("AttributeDoesNotExist", "Attribute/Property '{0}' from pin {1} does not exist"), FText::FromName(InputSource.GetName()), FText::FromName(InputTaggedData[i].Pin)));
-			return true;
+			PCGE_LOG(Error, GraphAndLog, FText::Format(LOCTEXT("AttributeDoesNotExist", "Attribute/Property '{0}' from pin {1} does not exist"), FText::FromName(InputSource.GetName()), FText::FromName(InputTaggedData[Index].Pin)));
+			return false;
 		}
 
-		uint16 AttributeTypeId = OperationData.InputAccessors[i]->GetUnderlyingType();
+		uint16 AttributeTypeId = OperationData.InputAccessors[Index]->GetUnderlyingType();
 
 		// Then verify that the type is OK
 		bool bHasSpecialRequirement = false;
-		if (!Settings->IsSupportedInputType(AttributeTypeId, i, bHasSpecialRequirement))
+		if (!Settings->IsSupportedInputType(AttributeTypeId, Index, bHasSpecialRequirement))
 		{
-			PCGE_LOG(Error, GraphAndLog, FText::Format(LOCTEXT("UnsupportedAttributeType", "Attribute/Property '{0}' from pin {1} is not a supported type"), FText::FromName(InputSource.GetName()), FText::FromName(InputTaggedData[i].Pin)));
+			FText AttributeTypeName = FText::FromString(PCG::Private::GetTypeName(AttributeTypeId));
+			PCGE_LOG(Error, GraphAndLog, FText::Format(LOCTEXT("UnsupportedAttributeType", "Attribute/Property '{0}' from pin {1} is not a supported type ('{2}')"), 
+				FText::FromName(InputSource.GetName()), 
+				FText::FromName(InputTaggedData[Index].Pin),
+				AttributeTypeName
+			));
 			return true;
 		}
 
@@ -199,50 +380,72 @@ bool FPCGMetadataElementBase::ExecuteInternal(FPCGContext* Context) const
 			}
 			else if (OperationData.MostComplexInputType != AttributeTypeId && !PCG::Private::IsBroadcastable(AttributeTypeId, OperationData.MostComplexInputType))
 			{
-				PCGE_LOG(Error, GraphAndLog, FText::Format(LOCTEXT("AttributeCannotBeBroadcasted", "Attribute '{0}' from pin {1} cannot be broadcasted to match types"), FText::FromName(InputSource.GetName()), FText::FromName(InputTaggedData[i].Pin)));
+				FText AttributeTypeName = FText::FromString(PCG::Private::GetTypeName(AttributeTypeId));
+				FText MostComplexTypeName = FText::FromString(PCG::Private::GetTypeName(OperationData.MostComplexInputType));
+				PCGE_LOG(Error, GraphAndLog, FText::Format(LOCTEXT("AttributeCannotBeBroadcasted", "Attribute '{0}' (from pin {1}) of type '{2}' cannot be used for operation with type '{3}'"), 
+					FText::FromName(InputSource.GetName()), 
+					FText::FromName(InputTaggedData[Index].Pin),
+					AttributeTypeName,
+					MostComplexTypeName));
 				return true;
 			}
 		}
 
-		NumberOfElements[i] = OperationData.InputKeys[i]->GetNum();
+		NumberOfElements[Index] = OperationData.InputKeys[Index]->GetNum();
 
-		if (OperationData.NumberOfElementsToProcess == -1)
+		if (Index == InputPinToForward)
 		{
-			OperationData.NumberOfElementsToProcess = NumberOfElements[i];
+			OperationData.NumberOfElementsToProcess = NumberOfElements[Index];
 		}
 
 		// There is nothing to do if one input doesn't have any element to process.
 		// Therefore mark that we have nothing to do and early out.
-		if (NumberOfElements[i] == 0)
+		if (NumberOfElements[Index] == 0)
 		{
-			PCGE_LOG(Verbose, LogOnly, FText::Format(LOCTEXT("NoElementsInInput", "No elements in data from pin {0}"), FText::FromName(InputTaggedData[i].Pin)));
+			PCGE_LOG(Verbose, LogOnly, FText::Format(LOCTEXT("NoElementsInInput", "No elements in data from pin {0}"), FText::FromName(InputTaggedData[Index].Pin)));
 			bNoOperationNeeded = true;
-			break;
+			return true;
 		}
 
 		// Verify that the number of elements makes sense
-		if (OperationData.NumberOfElementsToProcess % NumberOfElements[i] != 0)
+		if (OperationData.NumberOfElementsToProcess % NumberOfElements[Index] != 0)
 		{
-			PCGE_LOG(Error, GraphAndLog, FText::Format(LOCTEXT("MismatchInNumberOfElements", "Mismatch between the number of elements from pin {0} ({1}) and from pin {2} ({3})"), FText::FromName(InputTaggedData[0].Pin), OperationData.NumberOfElementsToProcess, FText::FromName(InputTaggedData[i].Pin), NumberOfElements[i]));
-			return true;
+			PCGE_LOG(Error, GraphAndLog, FText::Format(LOCTEXT("MismatchInNumberOfElements", "Mismatch between the number of elements from pin {0} ({1}) and from pin {2} ({3})"), FText::FromName(InputTaggedData[InputPinToForward].Pin), OperationData.NumberOfElementsToProcess, FText::FromName(InputTaggedData[Index].Pin), NumberOfElements[Index]));
+			return false;
 		}
 
 		if (InputSource.Selection == EPCGAttributePropertySelection::Attribute)
 		{
-			SourceAttribute[i] = SourceMetadata[i]->GetConstAttribute(InputSource.GetName());
+			SourceAttribute[Index] = SourceMetadata[Index]->GetConstAttribute(InputSource.GetName());
 		}
 		else
 		{
-			SourceAttribute[i] = nullptr;
+			SourceAttribute[Index] = nullptr;
+		}
+
+		return true;
+	};
+
+	// First do it for the input to forward (it's our control data)
+	if (!CreateInputAccessorAndValidate(InputPinToForward))
+	{
+		return true;
+	}
+
+	for (uint32 i = 0; i < NumberOfInputs && !bNoOperationNeeded; ++i)
+	{
+		if (i != InputPinToForward && !CreateInputAccessorAndValidate(i))
+		{
+			return true;
 		}
 	}
 
-	// If no operation is needed, just forward input 0
+	// If no operation is needed, just forward input 
 	if (bNoOperationNeeded)
 	{
 		for (uint32 OutputIndex = 0; OutputIndex < Settings->GetOutputPinNum(); ++OutputIndex)
 		{
-			FPCGTaggedData& OutputData = Outputs.Add_GetRef(InputTaggedData[0]);
+			FPCGTaggedData& OutputData = Outputs.Add_GetRef(InputTaggedData[InputPinToForward]);
 			OutputData.Pin = Settings->GetOutputPinLabel(OutputIndex);
 		}
 
@@ -261,16 +464,13 @@ bool FPCGMetadataElementBase::ExecuteInternal(FPCGContext* Context) const
 	{
 		OutputTarget.AttributeName = InputName;
 	}
-	
-	// TODO: Make it an option
-	const int32 InputToForward = 0;
 
 	// Use implicit capture, since we capture a lot
 	auto CreateAttribute = [&](uint32 OutputIndex, auto DummyOutValue) -> bool
 	{
 		using AttributeType = decltype(DummyOutValue);
 
-		FPCGTaggedData& OutputData = Outputs.Add_GetRef(InputTaggedData[InputToForward]);
+		FPCGTaggedData& OutputData = Outputs.Add_GetRef(InputTaggedData[InputPinToForward]);
 		OutputData.Pin = Settings->GetOutputPinLabel(OutputIndex);
 
 		UPCGMetadata* OutMetadata = nullptr;
@@ -283,7 +483,7 @@ bool FPCGMetadataElementBase::ExecuteInternal(FPCGContext* Context) const
 			// and create a new one of the right type.
 			// But if we have extra selectors, we need to handle it the same way as properties.
 			// There is no point of failure before duplicating. So duplicate, create the attribute and then the accessor.
-			PCGMetadataElementCommon::DuplicateTaggedData(InputTaggedData[InputToForward], OutputData, OutMetadata);
+			PCGMetadataElementCommon::DuplicateTaggedData(InputTaggedData[InputPinToForward], OutputData, OutMetadata);
 			FPCGMetadataAttributeBase* OutputAttribute = PCGMetadataElementCommon::ClearOrCreateAttribute(OutMetadata, OutputName, AttributeType{});
 			if (!OutputAttribute)
 			{
@@ -291,9 +491,9 @@ bool FPCGMetadataElementBase::ExecuteInternal(FPCGContext* Context) const
 			}
 
 			// And copy the mapping from the original attribute, if it is not points
-			if (!InputTaggedData[InputToForward].Data->IsA<UPCGPointData>() && SourceMetadata[InputToForward] && SourceAttribute[InputToForward])
+			if (!InputTaggedData[InputPinToForward].Data->IsA<UPCGPointData>() && SourceMetadata[InputPinToForward] && SourceAttribute[InputPinToForward])
 			{
-				PCGMetadataElementCommon::CopyEntryToValueKeyMap(SourceMetadata[InputToForward], SourceAttribute[InputToForward], OutputAttribute);
+				PCGMetadataElementCommon::CopyEntryToValueKeyMap(SourceMetadata[InputPinToForward], SourceAttribute[InputPinToForward], OutputAttribute);
 			}
 
 			OperationData.OutputAccessors[OutputIndex] = PCGAttributeAccessorHelpers::CreateAccessor(Cast<UPCGData>(OutputData.Data), OutputTarget);
@@ -313,7 +513,7 @@ bool FPCGMetadataElementBase::ExecuteInternal(FPCGContext* Context) const
 					return false;
 				}
 
-				PCGMetadataElementCommon::DuplicateTaggedData(InputTaggedData[InputToForward], OutputData, OutMetadata);
+				PCGMetadataElementCommon::DuplicateTaggedData(InputTaggedData[InputPinToForward], OutputData, OutMetadata);
 
 				// Re-create the accessor to point to the right data (since we just duplicated the data)
 				OperationData.OutputAccessors[OutputIndex] = PCGAttributeAccessorHelpers::CreateAccessor(Cast<UPCGData>(OutputData.Data), OutputTarget);
diff --git a/Engine/Plugins/Experimental/PCG/Source/PCG/Public/Elements/Metadata/PCGMetadataMakeRotator.h b/Engine/Plugins/Experimental/PCG/Source/PCG/Public/Elements/Metadata/PCGMetadataMakeRotator.h
--- a/Engine/Plugins/Experimental/PCG/Source/PCG/Public/Elements/Metadata/PCGMetadataMakeRotator.h
+++ b/Engine/Plugins/Experimental/PCG/Source/PCG/Public/Elements/Metadata/PCGMetadataMakeRotator.h
@@ -14,6 +14,9 @@ namespace PCGMetadataMakeRotatorConstants
 	const FName ForwardLabel = TEXT("Forward");
 	const FName RightLabel = TEXT("Right");
 	const FName UpLabel = TEXT("Up");
+	const FName YawLabel = TEXT("Yaw");
+	const FName PitchLabel = TEXT("Pitch");
+	const FName RollLabel = TEXT("Roll");
 }
 
 UENUM()
@@ -28,7 +31,8 @@ enum class EPCGMetadataMakeRotatorOp : uint8
 	MakeRotFromZX,
 	MakeRotFromYZ,
 	MakeRotFromZY,
-	MakeRotFromAxis
+	MakeRotFromAxes,
+	MakeRotFromAngles
 };
 
 /* Create a Rotator from 1, 2 or 3 axis. */
@@ -46,6 +50,11 @@ public:
 	virtual bool OnlyExposePreconfiguredSettings() const override { return true; }
 #endif
 	virtual void ApplyPreconfiguredSettings(const FPCGPreConfiguredSettingsInfo& PreconfiguredInfo) override;
+	virtual bool DoesInputSupportDefaultValue(uint32 Index) const override;
+	virtual UPCGParamData* CreateDefaultValueParam(uint32 Index) const override;
+#if WITH_EDITOR
+	virtual FString GetDefaultValueString(uint32 Index) const override;
+#endif // WITH_EDITOR
 	//~End UPCGSettings interface
 
 	FPCGAttributePropertySelector GetInputSource(uint32 Index) const override;
@@ -68,11 +77,11 @@ public:
 	UPROPERTY(BlueprintReadWrite, EditAnywhere, Category = Input, meta = (EditCondition = "Operation != EPCGMetadataMakeRotatorOp::MakeRotFromX && Operation != EPCGMetadataMakeRotatorOp::MakeRotFromY && Operation != EPCGMetadataMakeRotatorOp::MakeRotFromZ", EditConditionHides))
 	FPCGAttributePropertySelector InputSource2;
 
-	UPROPERTY(BlueprintReadWrite, EditAnywhere, Category = Input, meta = (EditCondition = "Operation == EPCGMetadataMakeRotatorOp::MakeRotFromAxis", EditConditionHides))
+	UPROPERTY(BlueprintReadWrite, EditAnywhere, Category = Input, meta = (EditCondition = "Operation == EPCGMetadataMakeRotatorOp::MakeRotFromAxes || Operation == EPCGMetadataMakeRotatorOp::MakeRotFromAngles", EditConditionHides))
 	FPCGAttributePropertySelector InputSource3;
 
 	UPROPERTY(BlueprintReadWrite, EditAnywhere, Category = Settings)
-	EPCGMetadataMakeRotatorOp Operation = EPCGMetadataMakeRotatorOp::MakeRotFromAxis;
+	EPCGMetadataMakeRotatorOp Operation = EPCGMetadataMakeRotatorOp::MakeRotFromAxes;
 };
 
 class FPCGMetadataMakeRotatorElement : public FPCGMetadataElementBase
diff --git a/Engine/Plugins/Experimental/PCG/Source/PCG/Public/Elements/Metadata/PCGMetadataMakeTransform.h b/Engine/Plugins/Experimental/PCG/Source/PCG/Public/Elements/Metadata/PCGMetadataMakeTransform.h
--- a/Engine/Plugins/Experimental/PCG/Source/PCG/Public/Elements/Metadata/PCGMetadataMakeTransform.h
+++ b/Engine/Plugins/Experimental/PCG/Source/PCG/Public/Elements/Metadata/PCGMetadataMakeTransform.h
@@ -28,8 +28,13 @@ public:
 	//~Begin UPCGSettings interface
 	virtual FName GetDefaultNodeName() const override;
 	virtual FText GetDefaultNodeTitle() const override;
-	//~End UPCGSettings interface
 #endif
+	virtual bool DoesInputSupportDefaultValue(uint32 Index) const override;
+	virtual UPCGParamData* CreateDefaultValueParam(uint32 Index) const override;
+#if WITH_EDITOR
+	virtual FString GetDefaultValueString(uint32 Index) const override;
+#endif // WITH_EDITORs
+	//~End UPCGSettings interface
 
 	FPCGAttributePropertySelector GetInputSource(uint32 Index) const override;
 
diff --git a/Engine/Plugins/Experimental/PCG/Source/PCG/Public/Elements/Metadata/PCGMetadataMakeVector.h b/Engine/Plugins/Experimental/PCG/Source/PCG/Public/Elements/Metadata/PCGMetadataMakeVector.h
--- a/Engine/Plugins/Experimental/PCG/Source/PCG/Public/Elements/Metadata/PCGMetadataMakeVector.h
+++ b/Engine/Plugins/Experimental/PCG/Source/PCG/Public/Elements/Metadata/PCGMetadataMakeVector.h
@@ -50,8 +50,13 @@ public:
 	//~Begin UPCGSettings interface
 	virtual FName GetDefaultNodeName() const override;
 	virtual FText GetDefaultNodeTitle() const override;
-	//~End UPCGSettings interface
 #endif
+	virtual bool DoesInputSupportDefaultValue(uint32 Index) const override;
+	virtual UPCGParamData* CreateDefaultValueParam(uint32 Index) const override;
+#if WITH_EDITOR
+	virtual FString GetDefaultValueString(uint32 Index) const override;
+#endif // WITH_EDITOR
+	//~End UPCGSettings interface
 
 	FPCGAttributePropertySelector GetInputSource(uint32 Index) const override;
 
diff --git a/Engine/Plugins/Experimental/PCG/Source/PCG/Public/Elements/Metadata/PCGMetadataOpElementBase.h b/Engine/Plugins/Experimental/PCG/Source/PCG/Public/Elements/Metadata/PCGMetadataOpElementBase.h
--- a/Engine/Plugins/Experimental/PCG/Source/PCG/Public/Elements/Metadata/PCGMetadataOpElementBase.h
+++ b/Engine/Plugins/Experimental/PCG/Source/PCG/Public/Elements/Metadata/PCGMetadataOpElementBase.h
@@ -23,6 +23,8 @@ namespace PCGMetadataSettingsBaseConstants
 	const FName ClampMaxLabel = TEXT("Max");
 	const FName LerpRatioLabel = TEXT("Ratio");
 	const FName TransformLabel = TEXT("Transform");
+
+	const FName DefaultOutputDataFromPinName = TEXT("Default");
 }
 
 // Defines behavior when number of entries doesn't match in inputs
@@ -45,7 +47,9 @@ enum class EPCGMetadataSettingsBaseTypes
  * Base class for all Metadata operations.
  * Metadata operation can work with attributes or properties. For example you could compute the addition between all points density and a constant from
  * a param data.
- * The output will be the duplication of the first input, with the same metadata + the result of the operation (either in an attribute or a property)
+ * The output will be the duplication of the first spatial input (by default - can be overridden by OutputDataFromPin), 
+ * with the same metadata + the result of the operation (either in an attribute or a property).
+ * 
  * The new attribute can collide with one of the attributes in the incoming metadata. In this case, the attribute value will be overridden by the result
  * of the operation. It will also override the type of the attribute if it doesn't match the original.
  * 
@@ -73,12 +77,18 @@ class PCG_API UPCGMetadataSettingsBase : public UPCGSettings
 public:
 	// ~Begin UObject interface
 	virtual void PostLoad() override;
+#if WITH_EDITOR
+	virtual bool CanEditChange(const FProperty* InProperty) const override;
+	virtual void PostEditChangeProperty(FPropertyChangedEvent& PropertyChangedEvent) override;
+#endif // WITH_EDITOR
 	// ~End UObject interface
 
 	//~Begin UPCGSettings interface
 #if WITH_EDITOR
 	virtual EPCGSettingsType GetType() const override { return EPCGSettingsType::Metadata; }
 	virtual bool HasDynamicPins() const override { return true; }
+	virtual bool GetPinExtraIcon(const UPCGPin* InPin, FName& OutExtraIcon, FText& OutTooltip) const override;
+	virtual FText GetNodeTooltipText() const override;
 #endif
 	virtual EPCGDataType GetCurrentPinTypes(const UPCGPin* InPin) const override;
 
@@ -103,9 +113,23 @@ public:
 	virtual bool HasDifferentOutputTypes() const { return false; }
 	virtual TArray<uint16> GetAllOutputTypes() const { return TArray<uint16>(); };
 
+	/* Can be overriden by child class to support default values on unplugged pins. */
+	virtual bool DoesInputSupportDefaultValue(uint32 Index) const { return false; }
+	virtual UPCGParamData* CreateDefaultValueParam(uint32 Index) const { return nullptr; }
+#if WITH_EDITOR
+	virtual FString GetDefaultValueString(uint32 Index) const { return FString(); }
+#endif // WITH_EDITOR
+
+	/* Return the current input pin to forward to the output. */
+	uint32 GetInputPinToForward() const;
+
 	UPROPERTY(BlueprintReadWrite, EditAnywhere, Category = "Output")
 	FPCGAttributePropertySelector OutputTarget;
 
+	/* By default, output is taken from first non-param pin (aka if the second pin is a point data, the output will be this point data). You can change it to any available input pin. */
+	UPROPERTY(BlueprintReadWrite, EditAnywhere, AdvancedDisplay, Category = "Output", meta = (GetOptions = GetOutputDataFromPinOptions))
+	FName OutputDataFromPin = PCGMetadataSettingsBaseConstants::DefaultOutputDataFromPinName;
+
 #if WITH_EDITORONLY_DATA
 	UPROPERTY()
 	FName OutputAttributeName_DEPRECATED = NAME_None;
@@ -121,6 +145,17 @@ public:
 	// Useful for unit tests. Allow to force a connection to allow the node to do its operation, even if nothing is connected to it.
 	TStaticArray<bool, MaxNumberOfOutputs> ForceOutputConnections{};
 #endif // WITH_EDITORONLY_DATA
+
+protected:
+	/* Return the type union from incident edges with the support for default values. */
+	EPCGDataType GetInputPinType(uint32 Index) const;
+
+	/* Return the index of the given input pin label. INDEX_NONE if not found*/
+	uint32 GetInputPinIndex(FName InPinLabel) const;
+
+	/* Return the list of all the Input pins */
+	UFUNCTION()
+	TArray<FName> GetOutputDataFromPinOptions() const;
 };
 
 
@@ -146,7 +181,6 @@ public:
 
 protected:
 	virtual bool ExecuteInternal(FPCGContext* Context) const override;
-
 	virtual bool DoOperation(FOperationData& InOperationData) const = 0;
 
 	/**
diff --git a/Engine/Plugins/Experimental/PCG/Source/PCGEditor/Private/SPCGEditorGraphNode.cpp b/Engine/Plugins/Experimental/PCG/Source/PCGEditor/Private/SPCGEditorGraphNode.cpp
--- a/Engine/Plugins/Experimental/PCG/Source/PCGEditor/Private/SPCGEditorGraphNode.cpp
+++ b/Engine/Plugins/Experimental/PCG/Source/PCGEditor/Private/SPCGEditorGraphNode.cpp
@@ -41,9 +41,12 @@ public:
 
 	virtual FSlateColor GetPinColor() const override;
 	virtual FSlateColor GetPinTextColor() const override;
+	FName GetLabelStyle(FName DefaultLabelStyle) const;
+	bool GetExtraIcon(FName& OutExtraIcon, FText& OutTooltip) const;
 
 private:
 	void ApplyUnusedPinStyle(FSlateColor& InOutColor) const;
+	void GetPCGNodeAndPin(const UPCGNode*& OutNode, const UPCGPin*& OutPin) const;
 };
 
 void SPCGEditorGraphNodePin::Construct(const FArguments& InArgs, UEdGraphPin* InPin)
@@ -99,13 +102,27 @@ void SPCGEditorGraphNodePin::Construct(const FArguments& InArgs, UEdGraphPin* In
 			.Image(this, &SPCGEditorGraphNodePin::GetPinStatusIcon)
 		];
 
-	TSharedRef<SWidget> LabelWidget = GetLabelWidget(InArgs._PinLabelStyle);
+	TSharedRef<SWidget> LabelWidget = GetLabelWidget(GetLabelStyle(InArgs._PinLabelStyle));
 
 	// Create the widget used for the pin body (status indicator, label, and value)
 	LabelAndValue =
 		SNew(SWrapBox)
 		.PreferredSize(150.f);
 
+	TSharedPtr<SImage> ExtraPinIconWidget;
+	FName ExtraPinIcon;
+	FText ExtraPinIconTooltip;
+	if (GetExtraIcon(ExtraPinIcon, ExtraPinIconTooltip))
+	{
+		ExtraPinIconWidget = SNew(SImage)
+			.Image(FAppStyle::GetBrush(ExtraPinIcon));
+
+		if (!ExtraPinIconTooltip.IsEmpty())
+		{
+			ExtraPinIconWidget->SetToolTipText(ExtraPinIconTooltip);
+		}
+	}
+
 	if (!bIsInput)
 	{
 		LabelAndValue->AddSlot()
@@ -114,6 +131,15 @@ void SPCGEditorGraphNodePin::Construct(const FArguments& InArgs, UEdGraphPin* In
 				PinStatusIndicator
 			];
 
+		if (ExtraPinIconWidget.IsValid())
+		{
+			LabelAndValue->AddSlot()
+				.VAlign(VAlign_Center)
+				[
+					ExtraPinIconWidget.ToSharedRef()
+				];
+		}
+
 		LabelAndValue->AddSlot()
 			.VAlign(VAlign_Center)
 			[
@@ -150,6 +176,15 @@ void SPCGEditorGraphNodePin::Construct(const FArguments& InArgs, UEdGraphPin* In
 			}
 		}
 
+		if (ExtraPinIconWidget.IsValid())
+		{
+			LabelAndValue->AddSlot()
+				.VAlign(VAlign_Center)
+				[
+					ExtraPinIconWidget.ToSharedRef()
+				];
+		}
+
 		LabelAndValue->AddSlot()
 			.VAlign(VAlign_Center)
 			[
@@ -231,22 +266,35 @@ void SPCGEditorGraphNodePin::Construct(const FArguments& InArgs, UEdGraphPin* In
 
 }
 
-void SPCGEditorGraphNodePin::ApplyUnusedPinStyle(FSlateColor& InOutColor) const
+void SPCGEditorGraphNodePin::GetPCGNodeAndPin(const UPCGNode*& OutNode, const UPCGPin*& OutPin) const
 {
 	UEdGraphPin* GraphPin = GetPinObj();
 	if (GraphPin && !GraphPin->IsPendingKill())
 	{
 		const UPCGEditorGraphNodeBase* EditorNode = CastChecked<const UPCGEditorGraphNodeBase>(GraphPinObj->GetOwningNode());
-		const UPCGNode* PCGNode = EditorNode ? EditorNode->GetPCGNode() : nullptr;
-		const UPCGPin* PCGPin = PCGNode ? PCGNode->GetInputPin(GraphPin->GetFName()) : nullptr;
+		OutNode = EditorNode ? EditorNode->GetPCGNode() : nullptr;
+		OutPin = OutNode ? OutNode->GetInputPin(GraphPin->GetFName()) : nullptr;
+	}
+	else
+	{
+		OutNode = nullptr;
+		OutPin = nullptr;
+	}
+}
 
-		// Halve opacity if pin is unused - intended to happen whether disabled or not
-		if (PCGPin && !PCGNode->IsPinUsedByNodeExecution(PCGPin))
-		{
-			FLinearColor Color = InOutColor.GetSpecifiedColor();
-			Color.A *= 0.5;
-			InOutColor = Color;
-		}
+void SPCGEditorGraphNodePin::ApplyUnusedPinStyle(FSlateColor& InOutColor) const
+{
+	const UPCGPin* PCGPin = nullptr;
+	const UPCGNode* PCGNode = nullptr;
+
+	GetPCGNodeAndPin(PCGNode, PCGPin);
+
+	// Halve opacity if pin is unused - intended to happen whether disabled or not
+	if (PCGPin && PCGNode && !PCGNode->IsPinUsedByNodeExecution(PCGPin))
+	{
+		FLinearColor Color = InOutColor.GetSpecifiedColor();
+		Color.A *= 0.5;
+		InOutColor = Color;
 	}
 }
 
@@ -269,6 +317,35 @@ FSlateColor SPCGEditorGraphNodePin::GetPinTextColor() const
 	return Color;
 }
 
+FName SPCGEditorGraphNodePin::GetLabelStyle(FName DefaultLabelStyle) const
+{
+	const UPCGPin* PCGPin = nullptr;
+	const UPCGNode* PCGNode = nullptr;
+	FName LabelStyle = NAME_None;
+
+	GetPCGNodeAndPin(PCGNode, PCGPin);
+
+	const UPCGSettings* Settings = PCGNode ? PCGNode->GetSettings() : nullptr;
+
+	if (!PCGPin || !Settings || !Settings->GetPinLabelStyle(PCGPin, LabelStyle))
+	{
+		LabelStyle = DefaultLabelStyle;
+	}
+
+	return LabelStyle;
+}
+
+bool SPCGEditorGraphNodePin::GetExtraIcon(FName& OutExtraIcon, FText& OutTooltip) const
+{
+	const UPCGPin* PCGPin = nullptr;
+	const UPCGNode* PCGNode = nullptr;
+
+	GetPCGNodeAndPin(PCGNode, PCGPin);
+
+	const UPCGSettings* Settings = PCGNode ? PCGNode->GetSettings() : nullptr;
+	return (PCGPin && Settings) ? Settings->GetPinExtraIcon(PCGPin, OutExtraIcon, OutTooltip) : false;
+}
+
 void SPCGEditorGraphNode::Construct(const FArguments& InArgs, UPCGEditorGraphNodeBase* InNode)
 {
 	GraphNode = InNode;
