[PCG] Fix Get Texture Data failing for textures that aren't already loaded.

Smooth out some other rough edges in the UX.

diff --git a/Engine/Plugins/Experimental/PCG/Source/PCG/Private/Data/PCGTextureData.cpp b/Engine/Plugins/Experimental/PCG/Source/PCG/Private/Data/PCGTextureData.cpp
--- a/Engine/Plugins/Experimental/PCG/Source/PCG/Private/Data/PCGTextureData.cpp
+++ b/Engine/Plugins/Experimental/PCG/Source/PCG/Private/Data/PCGTextureData.cpp
@@ -285,8 +285,8 @@ void UPCGTextureData::Initialize(UTexture2D* InTexture, const FTransform& InTran
 
 				if (const uint8_t* BulkData = reinterpret_cast<const uint8_t*>(PlatformData->Mips[0].BulkData.LockReadOnly()))
 				{
-					Width = Texture->GetSizeX();
-					Height = Texture->GetSizeY();
+					Width = PlatformData->SizeX;
+					Height = PlatformData->SizeY;
 					const int32 PixelCount = Width * Height;
 					ColorData.SetNum(PixelCount);
 
@@ -319,7 +319,7 @@ void UPCGTextureData::Initialize(UTexture2D* InTexture, const FTransform& InTran
 				}
 				else
 				{
-					UE_LOG(LogPCG, Error, TEXT("PCGTextureData unable to get bulk data from %s"), *Texture->GetFName().ToString());
+					UE_LOG(LogPCG, Error, TEXT("PCGTextureData unable to get bulk data from '%s'"), *Texture->GetFName().ToString());
 				}
 
 				PlatformData->Mips[0].BulkData.Unlock();
@@ -327,7 +327,7 @@ void UPCGTextureData::Initialize(UTexture2D* InTexture, const FTransform& InTran
 		}
 		else
 		{
-			UE_LOG(LogPCG, Error, TEXT("PCGTextureData does not support the format of %s"), *Texture->GetFName().ToString());
+			UE_LOG(LogPCG, Error, TEXT("PCGTextureData does not support the format of '%s'"), *Texture->GetFName().ToString());
 		}
 	}
 
diff --git a/Engine/Plugins/Experimental/PCG/Source/PCG/Private/Elements/PCGTextureSampler.cpp b/Engine/Plugins/Experimental/PCG/Source/PCG/Private/Elements/PCGTextureSampler.cpp
--- a/Engine/Plugins/Experimental/PCG/Source/PCG/Private/Elements/PCGTextureSampler.cpp
+++ b/Engine/Plugins/Experimental/PCG/Source/PCG/Private/Elements/PCGTextureSampler.cpp
@@ -36,6 +36,11 @@ bool FPCGTextureSamplerElement::ExecuteInternal(FPCGContext* Context) const
 	const UPCGTextureSamplerSettings* Settings = Context->GetInputSettings<UPCGTextureSamplerSettings>();
 	check(Settings);
 
+	if (Settings->Texture.IsNull())
+	{
+		return true;
+	}
+
 	UTexture2D* Texture = Settings->Texture.LoadSynchronous();
 	if (!Texture)
 	{
@@ -45,7 +50,7 @@ bool FPCGTextureSamplerElement::ExecuteInternal(FPCGContext* Context) const
 
 	if (!UPCGTextureData::IsSupported(Texture))
 	{
-		PCGE_LOG(Error, GraphAndLog, FText::Format(LOCTEXT("UnsupportedTextureFormat", "Texture '{0}' has unsupported texture format, currently supported formats are B8G8R8A8, R8G8B8A8 and G8"), FText::FromName(Texture->GetFName())));
+		PCGE_LOG(Error, GraphAndLog, FText::Format(LOCTEXT("UnsupportedTextureFormat", "Texture '{0}' has unsupported texture format, currently supported formats are B8G8R8A8, R8G8B8A8 and G8. Also ensure mipmap generation is disabled and sRGB is disabled."), FText::FromName(Texture->GetFName())));
 		return true;
 	}
 
@@ -71,9 +76,7 @@ bool FPCGTextureSamplerElement::ExecuteInternal(FPCGContext* Context) const
 	TextureData->TargetActor = Context->SourceComponent->GetOwner();
 
 	AActor* OriginalActor = UPCGBlueprintHelpers::GetOriginalComponent(*Context)->GetOwner();
-	
 	FTransform FinalTransform = Transform;
-
 	if (!bUseAbsoluteTransform)
 	{
 		FTransform OriginalActorTransform = OriginalActor->GetTransform();
