[PCG] Tentative fix for null pointer

#rnx

diff --git a/Engine/Plugins/Experimental/PCG/Source/PCG/Private/Helpers/PCGSettingsHelpers.cpp b/Engine/Plugins/Experimental/PCG/Source/PCG/Private/Helpers/PCGSettingsHelpers.cpp
--- a/Engine/Plugins/Experimental/PCG/Source/PCG/Private/Helpers/PCGSettingsHelpers.cpp
+++ b/Engine/Plugins/Experimental/PCG/Source/PCG/Private/Helpers/PCGSettingsHelpers.cpp
@@ -345,7 +345,7 @@ namespace PCGSettingsHelpers
 			}
 
 			// Don't allow to override the seed if the settings doesn't use the seed.
-			if (!InConfig.bUseSeed && Property->GetOwnerClass()->IsChildOf<UPCGSettings>())
+			if (!InConfig.bUseSeed && Property->GetOwnerClass() && Property->GetOwnerClass()->IsChildOf<UPCGSettings>())
 			{
 				bValid &= (Property->GetFName() != GET_MEMBER_NAME_CHECKED(UPCGSettings, Seed));
 			}
