Fixed issue where we would access pending kill landscapes in some instances when getting landscape proxies

diff --git a/Engine/Plugins/Experimental/PCG/Source/PCG/Private/Helpers/PCGHelpers.cpp b/Engine/Plugins/Experimental/PCG/Source/PCG/Private/Helpers/PCGHelpers.cpp
--- a/Engine/Plugins/Experimental/PCG/Source/PCG/Private/Helpers/PCGHelpers.cpp
+++ b/Engine/Plugins/Experimental/PCG/Source/PCG/Private/Helpers/PCGHelpers.cpp
@@ -219,7 +219,7 @@ namespace PCGHelpers
 
 		for (TObjectIterator<ALandscape> It; It; ++It)
 		{
-			if (It->GetWorld() == InWorld)
+			if (IsValid(*It) && It->GetWorld() == InWorld)
 			{
 				const FBox LandscapeBounds = GetLandscapeBounds(*It);
 				if (LandscapeBounds.IsValid && LandscapeBounds.Intersect(InBounds))
@@ -244,7 +244,7 @@ namespace PCGHelpers
 
 		for (TObjectIterator<ALandscapeProxy> It; It; ++It)
 		{
-			if (It->GetWorld() == InWorld)
+			if (IsValid(*It) && It->GetWorld() == InWorld)
 			{
 				const FBox LandscapeBounds = GetLandscapeBounds(*It);
 				if (LandscapeBounds.IsValid && LandscapeBounds.Intersect(InBounds))
