Fixed issue with PCG partition actors not reacting to landscape changes before a full graph refresh

diff --git a/Engine/Plugins/Experimental/PCG/Source/PCG/Private/PCGComponent.cpp b/Engine/Plugins/Experimental/PCG/Source/PCG/Private/PCGComponent.cpp
--- a/Engine/Plugins/Experimental/PCG/Source/PCG/Private/PCGComponent.cpp
+++ b/Engine/Plugins/Experimental/PCG/Source/PCG/Private/PCGComponent.cpp
@@ -231,6 +231,8 @@ void UPCGComponent::SetPropertiesFromOriginal(const UPCGComponent* Original)
 		RefreshTrackingData();
 	}
 
+	UpdateTrackedLandscape();
+
 	// Note that while we dirty here, we won't trigger a refresh since we don't have the required context
 	if (bIsDirty)
 	{
