Fixed error following loop changes

#rnx

diff --git a/Engine/Plugins/Experimental/PCG/Source/PCG/Private/Graph/PCGGraphCompiler.cpp b/Engine/Plugins/Experimental/PCG/Source/PCG/Private/Graph/PCGGraphCompiler.cpp
--- a/Engine/Plugins/Experimental/PCG/Source/PCG/Private/Graph/PCGGraphCompiler.cpp
+++ b/Engine/Plugins/Experimental/PCG/Source/PCG/Private/Graph/PCGGraphCompiler.cpp
@@ -54,7 +54,7 @@ TArray<FPCGGraphTask> FPCGGraphCompiler::CompileGraph(UPCGGraph* InGraph, FPCGTa
 			return TArray<FPCGGraphTask>();
 		}
 
-		const bool bIsNonDynamicAndNonDisabledSubgraphNode = (SubgraphNode && Subgraph && SubgraphSettings && !SubgraphSettings->IsDynamicGraph() && !SubgraphSettings->bEnabled);
+		const bool bIsNonDynamicAndNonDisabledSubgraphNode = (SubgraphNode && Subgraph && SubgraphSettings && !SubgraphSettings->IsDynamicGraph() && SubgraphSettings->bEnabled);
 		if (bIsNonDynamicAndNonDisabledSubgraphNode)
 		{
 			const FPCGTaskId PreId = NextId++;
