[PCG] Integrate up to CL 25489483

Had to fix up unavailability of spline curve FindNearest (verified with Julien), and also had to remove Ellipsis overflow policy in attribute list (verified with Thomas)

diff --git a/Engine/Plugins/Experimental/PCG/Source/PCG/Private/Data/PCGSplineData.cpp b/Engine/Plugins/Experimental/PCG/Source/PCG/Private/Data/PCGSplineData.cpp
--- a/Engine/Plugins/Experimental/PCG/Source/PCG/Private/Data/PCGSplineData.cpp
+++ b/Engine/Plugins/Experimental/PCG/Source/PCG/Private/Data/PCGSplineData.cpp
@@ -220,7 +220,7 @@ bool UPCGSplineProjectionData::SamplePoint(const FTransform& InTransform, const
 	FVector2D LocalPosition2D = Project(LocalTransform.GetLocation());
 	float Dummy;
 	// Find nearest key on 2D spline
-	float NearestInputKey = ProjectedPosition.FindNearest(LocalPosition2D, Dummy);
+	float NearestInputKey = ProjectedPosition.InaccurateFindNearest(LocalPosition2D, Dummy);
 	// TODO: if we didn't want to hand off density computation to the spline and do it here instead, we could do it in 2D space.
 	// Find point on original spline using the previously found key. Note this is an approximation that might not hold true since
 	// we are changing the curve length. Also, to support surface orientations that are not axis aligned, the project function
diff --git a/Engine/Plugins/Experimental/PCG/Source/PCG/Private/Elements/PCGSplineSampler.cpp b/Engine/Plugins/Experimental/PCG/Source/PCG/Private/Elements/PCGSplineSampler.cpp
--- a/Engine/Plugins/Experimental/PCG/Source/PCG/Private/Elements/PCGSplineSampler.cpp
+++ b/Engine/Plugins/Experimental/PCG/Source/PCG/Private/Elements/PCGSplineSampler.cpp
@@ -978,7 +978,7 @@ namespace PCGSplineSampler
 
 						// if bTreatAsPolyline, then we shouldnt use this, we should use nearest point on the polygon line segments
 						float Dummy;
-						const float NearestSplineKey = bFindNearestSplineKey ? Spline->SplineCurves.Position.FindNearest(SurfaceLocation, Dummy) : 0.f;
+						const float NearestSplineKey = bFindNearestSplineKey ? Spline->SplineCurves.Position.InaccurateFindNearest(SurfaceLocation, Dummy) : 0.f;
 
 						const FVector PointLocationLS = (Params.bProjectOntoSurface ? FVector(SurfaceLocation) : FVector(SampleLocation, MinPoint.Z));
 
