Added support for soft object reference user parameters & soft object pointers internally. Will behave like generic soft object paths but with the appropriate type narrowing.
Added support for overridable texture in the Get Texture Data node
Submitted on behalf of adrien.logut

diff --git a/Engine/Plugins/Experimental/PCG/Source/PCG/Private/Metadata/Accessors/PCGAttributeAccessorHelpers.cpp b/Engine/Plugins/Experimental/PCG/Source/PCG/Private/Metadata/Accessors/PCGAttributeAccessorHelpers.cpp
--- a/Engine/Plugins/Experimental/PCG/Source/PCG/Private/Metadata/Accessors/PCGAttributeAccessorHelpers.cpp
+++ b/Engine/Plugins/Experimental/PCG/Source/PCG/Private/Metadata/Accessors/PCGAttributeAccessorHelpers.cpp
@@ -136,6 +136,10 @@ namespace PCGAttributeAccessorHelpers
 		{
 			return Functor(Signature<FPCGEnumPropertyAccessor>{}, EnumProperty);
 		}
+		else if (const FSoftObjectProperty* SoftObjectProperty = CastField<FSoftObjectProperty>(InProperty))
+		{
+			return Functor(Signature<FPCGPropertySoftPtrAccessor>{}, SoftObjectProperty);
+		}
 		else if (const FStructProperty* StructProperty = CastField<FStructProperty>(InProperty))
 		{
 			if (StructProperty->Struct == TBaseStructure<FVector>::Get())
diff --git a/Engine/Plugins/Experimental/PCG/Source/PCG/Public/Elements/PCGTextureSampler.h b/Engine/Plugins/Experimental/PCG/Source/PCG/Public/Elements/PCGTextureSampler.h
--- a/Engine/Plugins/Experimental/PCG/Source/PCG/Public/Elements/PCGTextureSampler.h
+++ b/Engine/Plugins/Experimental/PCG/Source/PCG/Public/Elements/PCGTextureSampler.h
@@ -36,7 +36,7 @@ public:
 	bool bUseAbsoluteTransform = false;
 
 	// Texture specific parameters
-	UPROPERTY(BlueprintReadWrite, EditAnywhere, Category = Settings)
+	UPROPERTY(BlueprintReadWrite, EditAnywhere, Category = Settings, meta = (PCG_Overridable))
 	TSoftObjectPtr<UTexture2D> Texture = nullptr;
 
 	// Common members in BaseTextureData
diff --git a/Engine/Plugins/Experimental/PCG/Source/PCG/Public/Metadata/Accessors/PCGPropertyAccessor.h b/Engine/Plugins/Experimental/PCG/Source/PCG/Public/Metadata/Accessors/PCGPropertyAccessor.h
--- a/Engine/Plugins/Experimental/PCG/Source/PCG/Public/Metadata/Accessors/PCGPropertyAccessor.h
+++ b/Engine/Plugins/Experimental/PCG/Source/PCG/Public/Metadata/Accessors/PCGPropertyAccessor.h
@@ -256,3 +256,42 @@ public:
 private:
 	const FProperty* Property = nullptr;
 };
+
+/**
+* Templated accessor class for soft object/class ptr properties.
+* Do not instantiate it manually, use PCGAttributeAccessorHelpers::CreatePropertyAccessor.
+* Will always convert to FString for PCG
+* Key supported: Generic object
+*/
+class FPCGPropertySoftPtrAccessor : public IPCGAttributeAccessorT<FPCGPropertySoftPtrAccessor>
+{
+public:
+	using Type = FString;
+	using Super = IPCGAttributeAccessorT<FPCGPropertySoftPtrAccessor>;
+
+	FPCGPropertySoftPtrAccessor(const FSoftObjectProperty* InProperty)
+		: Super(/*bInReadOnly=*/ false)
+		, Property(InProperty)
+	{
+		check(Property);
+	}
+
+	bool GetRangeImpl(TArrayView<Type> OutValues, int32 Index, const IPCGAttributeAccessorKeys& Keys) const
+	{
+		return PCGPropertyAccessor::IterateGet(Property, OutValues, Index, Keys, [this](const void* PropertyAddressData) -> Type
+		{
+			return Property->GetPropertyValue(PropertyAddressData).ToString();
+		});
+	}
+
+	bool SetRangeImpl(TArrayView<const Type> InValues, int32 Index, IPCGAttributeAccessorKeys& Keys, EPCGAttributeAccessorFlags)
+	{
+		return PCGPropertyAccessor::IterateSet(Property, InValues, Index, Keys, [this](void* PropertyAddressData, const Type& Value) -> void
+		{
+			Property->SetPropertyValue(PropertyAddressData, FSoftObjectPtr(FSoftObjectPath(Value)));
+		});
+	}
+
+private:
+	const FSoftObjectProperty* Property = nullptr;
+};
\ No newline at end of file
