Added some flexiblity in the actor selector to customize display depending on user class.
Small refactoring in the Data From Actor class to allow more precise control in derived classes (namely default class, bundling up, etc.)
Major refactoring of the Get Landscape Data node implementation to fix most of the issues encountered.
Added landscape-specific options to the Get Landscape data node, to allow potential optimizations.

diff --git a/Engine/Plugins/Experimental/PCG/Source/PCG/Private/Data/PCGLandscapeData.cpp b/Engine/Plugins/Experimental/PCG/Source/PCG/Private/Data/PCGLandscapeData.cpp
--- a/Engine/Plugins/Experimental/PCG/Source/PCG/Private/Data/PCGLandscapeData.cpp
+++ b/Engine/Plugins/Experimental/PCG/Source/PCG/Private/Data/PCGLandscapeData.cpp
@@ -8,6 +8,7 @@
 #include "Grid/PCGLandscapeCache.h"
 #include "Helpers/PCGHelpers.h"
 
+#include "Landscape.h"
 #include "LandscapeInfo.h"
 #include "LandscapeProxy.h"
 #include "Engine/World.h"
@@ -16,17 +17,34 @@
 
 void UPCGLandscapeData::Initialize(const TArray<TWeakObjectPtr<ALandscapeProxy>>& InLandscapes, const FBox& InBounds, bool bInHeightOnly, bool bInUseMetadata)
 {
+	TSet<ALandscapeProxy*> LandscapesToIgnore;
+
 	for (TWeakObjectPtr<ALandscapeProxy> InLandscape : InLandscapes)
 	{
-		if (InLandscape.IsValid())
+		if (InLandscape.IsValid() && InLandscape->GetLandscapeActor() != InLandscape.Get())
 		{
-			Landscapes.Emplace(InLandscape.Get());
+			LandscapesToIgnore.Add(InLandscape->GetLandscapeActor());
+		}
+	}
 
-			// Build landscape info list
-			LandscapeInfos.Emplace(PCGHelpers::GetLandscapeBounds(InLandscape.Get()), InLandscape->GetLandscapeInfo());
+	TArray<TWeakObjectPtr<ALandscapeProxy>> FilteredLandscapes;
+	for (TWeakObjectPtr<ALandscapeProxy> InLandscape : InLandscapes)
+	{
+		if (InLandscape.IsValid() && !LandscapesToIgnore.Contains(InLandscape.Get()))
+		{
+			FilteredLandscapes.Add(InLandscape);
 		}
 	}
 
+	for (TWeakObjectPtr<ALandscapeProxy> Landscape : FilteredLandscapes)
+	{
+		check(Landscape.IsValid());
+		Landscapes.Emplace(Landscape.Get());
+
+		// Build landscape info list
+		LandscapeInfos.Emplace(PCGHelpers::GetLandscapeBounds(Landscape.Get()), Landscape->GetLandscapeInfo());
+	}
+
 	check(!Landscapes.IsEmpty());
 
 	ALandscapeProxy* FirstLandscape = Landscapes[0].Get();
@@ -46,7 +64,7 @@ void UPCGLandscapeData::Initialize(const TArray<TWeakObjectPtr<ALandscapeProxy>>
 	// TODO: find a better way to do this - maybe there should be a prototype metadata in the landscape cache
 	if (LandscapeCache)
 	{
-		if (!bHeightOnly && bUseMetadata)
+		if (bUseMetadata)
 		{
 			for (TSoftObjectPtr<ALandscapeProxy> Landscape : Landscapes)
 			{
@@ -163,7 +181,7 @@ bool UPCGLandscapeData::ProjectPoint(const FTransform& InTransform, const FBox&
 
 	if (bHeightOnly)
 	{
-		LandscapeCacheEntry->GetInterpolatedPointHeightOnly(ComponentLocalPoint, OutPoint);
+		LandscapeCacheEntry->GetInterpolatedPointHeightOnly(ComponentLocalPoint, OutPoint, bUseMetadata ? OutMetadata : nullptr);
 	}
 	else
 	{
@@ -224,10 +242,16 @@ const UPCGPointData* UPCGLandscapeData::CreatePointData(FPCGContext* Context, co
 
 	UPCGMetadata* OutMetadata = bUseMetadata ? Data->Metadata : nullptr;
 
+	// Most proxies we gathered will have the same landscape info, we shouldn't loop multiple times
+	// on them, unless we add the box filtering - but even then, depending on the transform we could have overlaps
+	TSet<ULandscapeInfo*> AllLandscapeInfos;
 	for (const TPair<FBox, ULandscapeInfo*>& LandscapeInfoPair : LandscapeInfos)
 	{
-		ULandscapeInfo* LandscapeInfo = LandscapeInfoPair.Value;
+		AllLandscapeInfos.Add(LandscapeInfoPair.Value);
+	}
 
+	for(ULandscapeInfo* LandscapeInfo : AllLandscapeInfos)
+	{
 		if (!LandscapeInfo || !LandscapeInfo->GetLandscapeProxy())
 		{
 			continue;
diff --git a/Engine/Plugins/Experimental/PCG/Source/PCG/Private/Elements/PCGActorSelector.cpp b/Engine/Plugins/Experimental/PCG/Source/PCG/Private/Elements/PCGActorSelector.cpp
--- a/Engine/Plugins/Experimental/PCG/Source/PCG/Private/Elements/PCGActorSelector.cpp
+++ b/Engine/Plugins/Experimental/PCG/Source/PCG/Private/Elements/PCGActorSelector.cpp
@@ -203,28 +203,7 @@ namespace PCGActorSelector
 	}
 }
 
-void FPCGActorSelectorSettings::PostLoad()
-{
-	// Make sure that when we are not using the 'by class' selection, we don't track that class dependency
-	if (ActorSelection != EPCGActorSelection::ByClass)
-	{
-		ActorSelectionClass = TSubclassOf<AActor>();
-	}
-}
-
 #if WITH_EDITOR
-void FPCGActorSelectorSettings::PostEditChangeProperty(FPropertyChangedEvent& PropertyChangedEvent)
-{
-	if (PropertyChangedEvent.GetPropertyName() == GET_MEMBER_NAME_CHECKED(FPCGActorSelectorSettings, ActorSelection))
-	{
-		// Make sure that when switching away from the 'by class' selection, we actually break that data dependency
-		if (ActorSelection != EPCGActorSelection::ByClass)
-		{
-			ActorSelectionClass = TSubclassOf<AActor>();
-		}
-	}
-}
-
 FText FPCGActorSelectorSettings::GetTaskNameSuffix() const
 {
 	if (ActorFilter == EPCGActorFilter::AllWorldActors)
diff --git a/Engine/Plugins/Experimental/PCG/Source/PCG/Private/Elements/PCGDataFromActor.cpp b/Engine/Plugins/Experimental/PCG/Source/PCG/Private/Elements/PCGDataFromActor.cpp
--- a/Engine/Plugins/Experimental/PCG/Source/PCG/Private/Elements/PCGDataFromActor.cpp
+++ b/Engine/Plugins/Experimental/PCG/Source/PCG/Private/Elements/PCGDataFromActor.cpp
@@ -35,15 +35,32 @@ void UPCGDataFromActorSettings::PostEditChangeProperty(FPropertyChangedEvent& Pr
 
 	if (PropertyChangedEvent.GetMemberPropertyName() == GET_MEMBER_NAME_CHECKED(UPCGDataFromActorSettings, ActorSelector))
 	{
-		ActorSelector.PostEditChangeProperty(PropertyChangedEvent);
+		if (PropertyChangedEvent.GetPropertyName() == GET_MEMBER_NAME_CHECKED(FPCGActorSelectorSettings, ActorSelection))
+		{
+			// Make sure that when switching away from the 'by class' selection, we actually break that data dependency
+			if (ActorSelector.ActorSelection != EPCGActorSelection::ByClass)
+			{
+				ActorSelector.ActorSelectionClass = GetDefaultActorSelectorClass();
+			}
+		}
 	}
 }
+
 #endif
 
+TSubclassOf<AActor> UPCGDataFromActorSettings::GetDefaultActorSelectorClass() const
+{
+	return TSubclassOf<AActor>();
+}
+
 void UPCGDataFromActorSettings::PostLoad()
 {
 	Super::PostLoad();
-	ActorSelector.PostLoad();
+
+	if (ActorSelector.ActorSelection != EPCGActorSelection::ByClass)
+	{
+		ActorSelector.ActorSelectionClass = GetDefaultActorSelectorClass();
+	}
 }
 
 FName UPCGDataFromActorSettings::AdditionalTaskName() const
@@ -175,10 +192,7 @@ bool FPCGDataFromActorElement::ExecuteInternal(FPCGContext* InContext) const
 
 	if (Context->bPerformedQuery)
 	{
-		for (AActor* Actor : Context->FoundActors)
-		{
-			ProcessActor(Context, Settings, Actor);
-		}
+		ProcessActors(Context, Settings, Context->FoundActors);
 	}
 
 	return true;
@@ -206,6 +220,14 @@ void FPCGDataFromActorElement::GatherWaitTasks(AActor* FoundActor, FPCGContext*
 	}
 }
 
+void FPCGDataFromActorElement::ProcessActors(FPCGContext* Context, const UPCGDataFromActorSettings* Settings, const TArray<AActor*>& FoundActors) const
+{
+	for (AActor* Actor : FoundActors)
+	{
+		ProcessActor(Context, Settings, Actor);
+	}
+}
+
 void FPCGDataFromActorElement::ProcessActor(FPCGContext* Context, const UPCGDataFromActorSettings* Settings, AActor* FoundActor) const
 {
 	check(Context);
diff --git a/Engine/Plugins/Experimental/PCG/Source/PCG/Private/Elements/PCGPropertyToParamData.cpp b/Engine/Plugins/Experimental/PCG/Source/PCG/Private/Elements/PCGPropertyToParamData.cpp
--- a/Engine/Plugins/Experimental/PCG/Source/PCG/Private/Elements/PCGPropertyToParamData.cpp
+++ b/Engine/Plugins/Experimental/PCG/Source/PCG/Private/Elements/PCGPropertyToParamData.cpp
@@ -30,7 +30,13 @@ void UPCGPropertyToParamDataSettings::PostEditChangeProperty(FPropertyChangedEve
 
 	if (PropertyChangedEvent.GetMemberPropertyName() == GET_MEMBER_NAME_CHECKED(UPCGPropertyToParamDataSettings, ActorSelector))
 	{
-		ActorSelector.PostEditChangeProperty(PropertyChangedEvent);
+		if (PropertyChangedEvent.GetPropertyName() == GET_MEMBER_NAME_CHECKED(FPCGActorSelectorSettings, ActorSelection))
+		{
+			if (ActorSelector.ActorSelection != EPCGActorSelection::ByClass)
+			{
+				ActorSelector.ActorSelectionClass = TSubclassOf<AActor>();
+			}
+		}
 	}
 }
 #endif // WITH_EDITOR
@@ -43,7 +49,7 @@ void UPCGPropertyToParamDataSettings::PostLoad()
 	if (ActorSelection_DEPRECATED != EPCGActorSelection::ByTag ||
 		ActorSelectionTag_DEPRECATED != NAME_None ||
 		ActorSelectionName_DEPRECATED != NAME_None ||
-		ActorSelectionClass_DEPRECATED != TSubclassOf<AActor>{} ||
+		ActorSelectionClass_DEPRECATED != TSubclassOf<AActor>() ||
 		ActorFilter_DEPRECATED != EPCGActorFilter::Self ||
 		bIncludeChildren_DEPRECATED != false)
 	{
@@ -56,12 +62,15 @@ void UPCGPropertyToParamDataSettings::PostLoad()
 		ActorSelection_DEPRECATED = EPCGActorSelection::ByTag;
 		ActorSelectionTag_DEPRECATED = NAME_None;
 		ActorSelectionName_DEPRECATED = NAME_None;
-		ActorSelectionClass_DEPRECATED = TSubclassOf<AActor>{};
+		ActorSelectionClass_DEPRECATED = TSubclassOf<AActor>();
 		ActorFilter_DEPRECATED = EPCGActorFilter::Self;
 		bIncludeChildren_DEPRECATED = false;
 	}
 
-	ActorSelector.PostLoad();
+	if (ActorSelector.ActorSelection != EPCGActorSelection::ByClass)
+	{
+		ActorSelector.ActorSelectionClass = TSubclassOf<AActor>();
+	}
 }
 
 FName UPCGPropertyToParamDataSettings::AdditionalTaskName() const
diff --git a/Engine/Plugins/Experimental/PCG/Source/PCG/Private/Elements/PCGTypedGetter.cpp b/Engine/Plugins/Experimental/PCG/Source/PCG/Private/Elements/PCGTypedGetter.cpp
--- a/Engine/Plugins/Experimental/PCG/Source/PCG/Private/Elements/PCGTypedGetter.cpp
+++ b/Engine/Plugins/Experimental/PCG/Source/PCG/Private/Elements/PCGTypedGetter.cpp
@@ -2,6 +2,8 @@
 
 #include "Elements/PCGTypedGetter.h"
 
+#include "PCGComponent.h"
+#include "Data/PCGLandscapeData.h"
 #include "Helpers/PCGHelpers.h"
 
 #include "Landscape.h"
@@ -14,6 +16,11 @@ UPCGGetLandscapeSettings::UPCGGetLandscapeSettings()
 {
 	bDisplayModeSettings = false;
 	Mode = EPCGGetDataFromActorMode::ParseActorComponents;
+	ActorSelector.bShowActorFilter = false;
+	ActorSelector.bIncludeChildren = false;
+	ActorSelector.bShowActorSelectionClass = false;
+	ActorSelector.bSelectMultiple = true;
+	ActorSelector.bShowSelectMultiple = false;
 
 	// We want to apply different defaults to newly placed nodes. We detect new object if they are not a default object/archetype
 	// and/or they do not need load.
@@ -22,20 +29,30 @@ UPCGGetLandscapeSettings::UPCGGetLandscapeSettings()
 		// This setup replicates what was implemented on the Landscape input node pin
 		ActorSelector.ActorFilter = EPCGActorFilter::AllWorldActors;
 		ActorSelector.bMustOverlapSelf = true;
-		ActorSelector.bSelectMultiple = true;
 		ActorSelector.ActorSelection = EPCGActorSelection::ByClass;
 		ActorSelector.ActorSelectionClass = ALandscapeProxy::StaticClass();
 	}
 }
 
+TSubclassOf<AActor> UPCGGetLandscapeSettings::GetDefaultActorSelectorClass() const
+{
+	return ALandscapeProxy::StaticClass();
+}
+
 TArray<FPCGPinProperties> UPCGGetLandscapeSettings::OutputPinProperties() const
 {
 	TArray<FPCGPinProperties> PinProperties;
-	PinProperties.Emplace(PCGPinConstants::DefaultOutputLabel, EPCGDataType::Landscape);
+	PinProperties.Emplace(PCGPinConstants::DefaultOutputLabel, EPCGDataType::Landscape, /*bAllowMultipleConnections=*/true, /*bAllowMultipleData=*/false);
 
 	return PinProperties;
 }
 
+FName UPCGGetLandscapeSettings::AdditionalTaskName() const
+{
+	// Do not use the version from data from actor otherwise we'll show the selected actor class, which serves no purpose
+	return UPCGSettings::AdditionalTaskName();
+}
+
 #if WITH_EDITOR
 FText UPCGGetLandscapeSettings::GetNodeTooltipText() const
 {
@@ -43,6 +60,68 @@ FText UPCGGetLandscapeSettings::GetNodeTooltipText() const
 }
 #endif
 
+FPCGElementPtr UPCGGetLandscapeSettings::CreateElement() const
+{
+	return MakeShared<FPCGGetLandscapeDataElement>();
+}
+
+void FPCGGetLandscapeDataElement::ProcessActors(FPCGContext* Context, const UPCGDataFromActorSettings* InSettings, const TArray<AActor*>& FoundActors) const
+{
+	check(Context);
+	check(InSettings);
+
+	const UPCGGetLandscapeSettings* Settings = CastChecked<UPCGGetLandscapeSettings>(InSettings);
+
+	// In the base class (FPCGDataFromActorElement) we'd go through all actors, one by one and call
+	// UPCGComponent::CreateActorPCGDataCollection, and push the results to the output.
+	// However, in this case, we want to do what's done in UPCGComponent::GetLandscapeData,
+	// which is to create a single UPCGLandscapeData keeping tabs on all selected landscapes.
+	TArray<TWeakObjectPtr<ALandscapeProxy>> Landscapes;
+	FBox LandscapeBounds(EForceInit::ForceInit);
+	TSet<FString> LandscapeTags;
+
+	for (AActor* FoundActor : FoundActors)
+	{
+		if (!FoundActor || !IsValid(FoundActor))
+		{
+			continue;
+		}
+
+		ALandscapeProxy* Landscape = Cast<ALandscapeProxy>(FoundActor);
+		if (ensure(Landscape))
+		{
+			Landscapes.Add(Landscape);
+			LandscapeBounds += PCGHelpers::GetGridBounds(Landscape, nullptr);
+
+			for (FName Tag : Landscape->Tags)
+			{
+				LandscapeTags.Add(Tag.ToString());
+			}
+		}
+	}
+
+	if (!Landscapes.IsEmpty())
+	{
+		UPCGLandscapeData* LandscapeData = NewObject<UPCGLandscapeData>();
+		LandscapeData->Initialize(Landscapes, LandscapeBounds, Settings->bGetHeightOnly, Settings->bGetLayerWeights);
+		
+		// Override target actor for the one we're running on, to make sure we don't output on the landscape directly
+		if (Context->SourceComponent.IsValid())
+		{
+			LandscapeData->TargetActor = Context->SourceComponent->GetOwner();
+		}
+
+		FPCGTaggedData& TaggedData = Context->OutputData.TaggedData.Emplace_GetRef();
+		TaggedData.Data = LandscapeData;
+		TaggedData.Tags = LandscapeTags;
+	}
+}
+
+void FPCGGetLandscapeDataElement::ProcessActor(FPCGContext* Context, const UPCGDataFromActorSettings* Settings, AActor* FoundActor) const
+{
+	checkf(false, TEXT("This should never be called directly"));
+}
+
 UPCGGetSplineSettings::UPCGGetSplineSettings()
 {
 	bDisplayModeSettings = false;
diff --git a/Engine/Plugins/Experimental/PCG/Source/PCG/Private/Grid/PCGLandscapeCache.cpp b/Engine/Plugins/Experimental/PCG/Source/PCG/Private/Grid/PCGLandscapeCache.cpp
--- a/Engine/Plugins/Experimental/PCG/Source/PCG/Private/Grid/PCGLandscapeCache.cpp
+++ b/Engine/Plugins/Experimental/PCG/Source/PCG/Private/Grid/PCGLandscapeCache.cpp
@@ -298,11 +298,12 @@ void FPCGLandscapeCacheEntry::GetInterpolatedPointMetadataOnly(const FVector2D&
 	GetInterpolatedPointMetadataInternal(Indices, OutPoint, OutMetadata);
 }
 
-void FPCGLandscapeCacheEntry::GetInterpolatedPointHeightOnly(const FVector2D& LocalPoint, FPCGPoint& OutPoint) const
+void FPCGLandscapeCacheEntry::GetInterpolatedPointHeightOnly(const FVector2D& LocalPoint, FPCGPoint& OutPoint, UPCGMetadata* OutMetadata) const
 {
 	check(bDataLoaded);
 	const PCGLandscapeCache::FSafeIndices Indices = PCGLandscapeCache::CalcSafeIndices(LocalPoint, Stride);
 	GetInterpolatedPointInternal(Indices, OutPoint, /*bHeightOnly=*/true);
+	GetInterpolatedPointMetadataInternal(Indices, OutPoint, OutMetadata);
 }
 
 bool FPCGLandscapeCacheEntry::TouchAndLoad(int32 InTouch) const
diff --git a/Engine/Plugins/Experimental/PCG/Source/PCG/Public/Elements/PCGActorSelector.h b/Engine/Plugins/Experimental/PCG/Source/PCG/Public/Elements/PCGActorSelector.h
--- a/Engine/Plugins/Experimental/PCG/Source/PCG/Public/Elements/PCGActorSelector.h
+++ b/Engine/Plugins/Experimental/PCG/Source/PCG/Public/Elements/PCGActorSelector.h
@@ -40,37 +40,50 @@ struct FPCGActorSelectorSettings
 	GENERATED_BODY()
 
 	/** Which actors to consider. */
-	UPROPERTY(BlueprintReadWrite, EditAnywhere, Category = Settings)
+	UPROPERTY(BlueprintReadWrite, EditAnywhere, Category = Settings, meta = (EditCondition = "bShowActorFilter", EditConditionHides))
 	EPCGActorFilter ActorFilter = EPCGActorFilter::Self;
 
 	UPROPERTY(BlueprintReadWrite, EditAnywhere, Category = Settings, meta = (EditCondition = "ActorFilter==EPCGActorFilter::AllWorldActors", EditConditionHides))
 	bool bMustOverlapSelf = false;
 
 	/** Whether to consider child actors. */
-	UPROPERTY(BlueprintReadWrite, EditAnywhere, Category = Settings, meta = (EditCondition = "ActorFilter!=EPCGActorFilter::AllWorldActors", EditConditionHides))
+	UPROPERTY(BlueprintReadWrite, EditAnywhere, Category = Settings, meta = (EditCondition = "bShowIncludeChildren && ActorFilter!=EPCGActorFilter::AllWorldActors", EditConditionHides))
 	bool bIncludeChildren = false;
 
 	/** Enables/disables fine-grained actor filtering options. */
 	UPROPERTY(BlueprintReadWrite, EditAnywhere, Category = Settings, meta = (EditCondition = "ActorFilter!=EPCGActorFilter::AllWorldActors && bIncludeChildren", EditConditionHides))
 	bool bDisableFilter = false;
 
-	UPROPERTY(BlueprintReadWrite, EditAnywhere, Category = Settings, meta = (EditCondition = "(ActorFilter==EPCGActorFilter::AllWorldActors || (bIncludeChildren && !bDisableFilter))", EditConditionHides))
+	UPROPERTY(BlueprintReadWrite, EditAnywhere, Category = Settings, meta = (EditCondition = "bShowActorSelection && (ActorFilter==EPCGActorFilter::AllWorldActors || (bIncludeChildren && !bDisableFilter))", EditConditionHides))
 	EPCGActorSelection ActorSelection = EPCGActorSelection::ByTag;
 
-	UPROPERTY(BlueprintReadWrite, EditAnywhere, Category = Settings, meta = (EditCondition = "(ActorFilter==EPCGActorFilter::AllWorldActors || (bIncludeChildren && !bDisableFilter)) && ActorSelection==EPCGActorSelection::ByTag", EditConditionHides))
+	UPROPERTY(BlueprintReadWrite, EditAnywhere, Category = Settings, meta = (EditCondition = "bShowActorSelection && (ActorFilter==EPCGActorFilter::AllWorldActors || (bIncludeChildren && !bDisableFilter)) && ActorSelection==EPCGActorSelection::ByTag", EditConditionHides))
 	FName ActorSelectionTag;
 
-	UPROPERTY(BlueprintReadWrite, EditAnywhere, Category = Settings, meta = (EditCondition = "(ActorFilter==EPCGActorFilter::AllWorldActors || (bIncludeChildren && !bDisableFilter)) && ActorSelection==EPCGActorSelection::ByClass", EditConditionHides, AllowAbstract = "true"))
+	UPROPERTY(BlueprintReadWrite, EditAnywhere, Category = Settings, meta = (EditCondition = "bShowActorSelectionClass && bShowActorSelection && (ActorFilter==EPCGActorFilter::AllWorldActors || (bIncludeChildren && !bDisableFilter)) && ActorSelection==EPCGActorSelection::ByClass", EditConditionHides, AllowAbstract = "true"))
 	TSubclassOf<AActor> ActorSelectionClass;
 
 	/** If true processes all matching actors, otherwise returns data from first match. */
-	UPROPERTY(BlueprintReadWrite, EditAnywhere, Category = Settings, meta = (EditCondition = "ActorFilter==EPCGActorFilter::AllWorldActors && ActorSelection!=EPCGActorSelection::ByName", EditConditionHides))
+	UPROPERTY(BlueprintReadWrite, EditAnywhere, Category = Settings, meta = (EditCondition = "bShowSelectMultiple && ActorFilter==EPCGActorFilter::AllWorldActors && ActorSelection!=EPCGActorSelection::ByName", EditConditionHides))
 	bool bSelectMultiple = false;
 
-	void PostLoad();
+	// Properties used to hide some fields when used in different contexts
+	UPROPERTY(Transient, meta = (EditCondition = false, EditConditionHides))
+	bool bShowActorFilter = true;
+
+	UPROPERTY(Transient, meta = (EditCondition = false, EditConditionHides))
+	bool bShowIncludeChildren = true;
+
+	UPROPERTY(Transient, meta = (EditCondition = false, EditConditionHides))
+	bool bShowActorSelection = true;
+
+	UPROPERTY(Transient, meta = (EditCondition = false, EditConditionHides))
+	bool bShowActorSelectionClass = true;
+
+	UPROPERTY(Transient, meta = (EditCondition = false, EditConditionHides))
+	bool bShowSelectMultiple = true;
 
 #if WITH_EDITOR	
-	void PostEditChangeProperty(FPropertyChangedEvent& PropertyChangedEvent);
 	FText GetTaskNameSuffix() const;
 	FName GetTaskName(const FText& Prefix) const;
 #endif
diff --git a/Engine/Plugins/Experimental/PCG/Source/PCG/Public/Elements/PCGDataFromActor.h b/Engine/Plugins/Experimental/PCG/Source/PCG/Public/Elements/PCGDataFromActor.h
--- a/Engine/Plugins/Experimental/PCG/Source/PCG/Public/Elements/PCGDataFromActor.h
+++ b/Engine/Plugins/Experimental/PCG/Source/PCG/Public/Elements/PCGDataFromActor.h
@@ -39,10 +39,10 @@ public:
 
 	virtual FName AdditionalTaskName() const override;
 
+protected:
 	virtual TArray<FPCGPinProperties> InputPinProperties() const override { return TArray<FPCGPinProperties>(); }
 	virtual TArray<FPCGPinProperties> OutputPinProperties() const override;
 
-protected:
 	virtual FPCGElementPtr CreateElement() const override;
 	//~End UPCGSettings
 
@@ -58,6 +58,9 @@ public:
 	/** Override this to filter what kinds of data should be retrieved from the actor(s). */
 	virtual bool DataFilter(EPCGDataType InDataType) const { return true; }
 
+	/** Override this to change the default value the selector will revert to when changing the actor selection type */
+	virtual TSubclassOf<AActor> GetDefaultActorSelectorClass() const;
+
 	UPROPERTY(BlueprintReadWrite, EditAnywhere, Category = Settings, meta = (ShowOnlyInnerProperties))
 	FPCGActorSelectorSettings ActorSelector;
 
@@ -98,5 +101,6 @@ public:
 protected:
 	virtual bool ExecuteInternal(FPCGContext* Context) const;
 	void GatherWaitTasks(AActor* FoundActor, FPCGContext* Context, TArray<FPCGTaskId>& OutWaitTasks) const;
-	void ProcessActor(FPCGContext* Context, const UPCGDataFromActorSettings* Settings, AActor* FoundActor) const;
+	virtual void ProcessActors(FPCGContext* Context, const UPCGDataFromActorSettings* Settings, const TArray<AActor*>& FoundActors) const;
+	virtual void ProcessActor(FPCGContext* Context, const UPCGDataFromActorSettings* Settings, AActor* FoundActor) const;
 };
diff --git a/Engine/Plugins/Experimental/PCG/Source/PCG/Public/Elements/PCGTypedGetter.h b/Engine/Plugins/Experimental/PCG/Source/PCG/Public/Elements/PCGTypedGetter.h
--- a/Engine/Plugins/Experimental/PCG/Source/PCG/Public/Elements/PCGTypedGetter.h
+++ b/Engine/Plugins/Experimental/PCG/Source/PCG/Public/Elements/PCGTypedGetter.h
@@ -22,14 +22,32 @@ public:
 	virtual FText GetNodeTooltipText() const override;
 #endif
 
+	virtual FName AdditionalTaskName() const override;
+
 protected:
 	virtual TArray<FPCGPinProperties> OutputPinProperties() const override;
+
+	virtual FPCGElementPtr CreateElement() const override;
 	//~End UPCGSettings
 
 public:
 	//~Begin UPCGDataFromActorSettings interface
 	virtual bool DataFilter(EPCGDataType InDataType) const override { return !!(InDataType & EPCGDataType::Landscape); }
+	virtual TSubclassOf<AActor> GetDefaultActorSelectorClass() const override;
 	//~End UPCGDataFromActorSettings
+
+	UPROPERTY(BlueprintReadWrite, EditAnywhere, Category = Settings, meta = (PCG_Overridable))
+	bool bGetHeightOnly = false;
+
+	UPROPERTY(BlueprintReadWrite, EditAnywhere, Category = Settings, meta = (PCG_Overridable))
+	bool bGetLayerWeights = true;
+};
+
+class FPCGGetLandscapeDataElement : public FPCGDataFromActorElement
+{
+protected:
+	virtual void ProcessActors(FPCGContext* Context, const UPCGDataFromActorSettings* Settings, const TArray<AActor*>& FoundActors) const override;
+	virtual void ProcessActor(FPCGContext* Context, const UPCGDataFromActorSettings* Settings, AActor* FoundActor) const override;
 };
 
 /** Builds a collection of spline data from the selected actors. */
diff --git a/Engine/Plugins/Experimental/PCG/Source/PCG/Public/Grid/PCGLandscapeCache.h b/Engine/Plugins/Experimental/PCG/Source/PCG/Public/Grid/PCGLandscapeCache.h
--- a/Engine/Plugins/Experimental/PCG/Source/PCG/Public/Grid/PCGLandscapeCache.h
+++ b/Engine/Plugins/Experimental/PCG/Source/PCG/Public/Grid/PCGLandscapeCache.h
@@ -52,7 +52,7 @@ public:
 	void GetPointHeightOnly(int32 PointIndex, FPCGPoint& OutPoint) const;
 	void GetInterpolatedPoint(const FVector2D& LocalPoint, FPCGPoint& OutPoint, UPCGMetadata* OutMetadata) const;
 	void GetInterpolatedPointMetadataOnly(const FVector2D& LocalPoint, FPCGPoint& OutPoint, UPCGMetadata* OutMetadata) const;
-	void GetInterpolatedPointHeightOnly(const FVector2D& LocalPoint, FPCGPoint& OutPoint) const;
+	void GetInterpolatedPointHeightOnly(const FVector2D& LocalPoint, FPCGPoint& OutPoint, UPCGMetadata* OutMetadata) const;
 	void GetInterpolatedLayerWeights(const FVector2D& LocalPoint, TArray<FPCGLandscapeLayerWeight>& OutLayerWeights) const;
 
 private:
