Fixed issue with PCG picking up a deleted landscape in some instances.
Fixed issue with changing pivots on ISMs placed by PCG - affected mainly scale and rotation

#rnx

diff --git a/Engine/Plugins/Experimental/PCG/Source/PCG/Private/Helpers/PCGActorHelpers.cpp b/Engine/Plugins/Experimental/PCG/Source/PCG/Private/Helpers/PCGActorHelpers.cpp
--- a/Engine/Plugins/Experimental/PCG/Source/PCG/Private/Helpers/PCGActorHelpers.cpp
+++ b/Engine/Plugins/Experimental/PCG/Source/PCG/Private/Helpers/PCGActorHelpers.cpp
@@ -62,7 +62,8 @@ UPCGManagedISMComponent* UPCGActorHelpers::GetOrCreateManagedISMC(AActor* InTarg
 
 				if (UInstancedStaticMeshComponent* ISMC = Resource->GetComponent())
 				{
-					if (ISMC->NumCustomDataFloats == InParams.NumCustomDataFloats &&
+					if (IsValid(ISMC) &&
+						ISMC->NumCustomDataFloats == InParams.NumCustomDataFloats &&
 						Resource->GetDescriptor() == InParams.Descriptor)
 					{
 						MatchingResource = Resource;
@@ -75,6 +76,11 @@ UPCGManagedISMComponent* UPCGActorHelpers::GetOrCreateManagedISMC(AActor* InTarg
 		{
 			TRACE_CPUPROFILER_EVENT_SCOPE(UPCGActorHelpers::GetOrCreateManagedISMC::MarkAsUsed);
 			MatchingResource->MarkAsUsed();
+
+			// Reset the rotation/scale to be identity otherwise if the root component transform has changed, the final transform will be wrong
+			MatchingResource->GetComponent()->SetWorldRotation(FQuat::Identity);
+			MatchingResource->GetComponent()->SetWorldScale3D(FVector::OneVector);
+
 			return MatchingResource;
 		}
 	}
