Added AsyncProcessingEx and AsyncProcessingOneToOneEx which is a more primitive async processor that only uses indexes. This will allow for SOA style processing of multiple output arrays

diff --git a/Engine/Plugins/Experimental/PCG/Source/PCG/Private/Tests/Helpers/PCGAsyncTest.cpp b/Engine/Plugins/Experimental/PCG/Source/PCG/Private/Tests/Helpers/PCGAsyncTest.cpp
new file mode 100644
--- /dev/null
+++ b/Engine/Plugins/Experimental/PCG/Source/PCG/Private/Tests/Helpers/PCGAsyncTest.cpp
@@ -0,0 +1,127 @@
+// Copyright Epic Games, Inc. All Rights Reserved.
+
+#include "Tests/PCGTestsCommon.h"
+
+#include "Helpers/PCGAsync.h"
+
+
+IMPLEMENT_CUSTOM_SIMPLE_AUTOMATION_TEST(FPCGAsyncTest_AsyncProcessing, FPCGTestBaseClass, "pcg.Async.AsyncProcessing", PCGTestsCommon::TestFlags)
+
+bool FPCGAsyncTest_AsyncProcessing::RunTest(const FString& Parameters)
+{
+	FPCGAsyncState AsyncState;
+	AsyncState.NumAvailableTasks = 10;
+
+	TArray<double> TestOutput;
+	const bool bIsDone = FPCGAsync::AsyncProcessing<double>(
+		&AsyncState,
+		6,
+		TestOutput,
+		[Value = MakeUnique<int32>(42)](int32 Index, double& Output)
+		{
+			if ((Index % 2) == 0)
+			{
+				return false; // remove all even
+			}
+
+			Output = double(Index * (*Value));
+
+			return true;
+		},
+		/*bEnableTimeSlicing*/ false,
+		/*ChunkSize*/ 1 // force this to thread 
+	);
+
+	UTEST_TRUE("IsDone", bIsDone);
+
+	UTEST_EQUAL("Wrote Correct Number", TestOutput.Num(), 3);
+
+	UTEST_EQUAL("TestOutput[0]", TestOutput[0], 42.0 * 1);
+	UTEST_EQUAL("TestOutput[1]", TestOutput[1], 42.0 * 3);
+	UTEST_EQUAL("TestOutput[2]", TestOutput[2], 42.0 * 5);
+
+	return true;
+}
+
+IMPLEMENT_CUSTOM_SIMPLE_AUTOMATION_TEST(FPCGAsyncTest_AsyncProcessingEx, FPCGTestBaseClass, "pcg.Async.AsyncProcessingEx", PCGTestsCommon::TestFlags)
+
+bool FPCGAsyncTest_AsyncProcessingEx::RunTest(const FString& Parameters)
+{
+	FPCGAsyncState AsyncState;
+	AsyncState.NumAvailableTasks = 10;
+
+	TArray<double> TestOutput;
+	const bool bIsDone = FPCGAsync::AsyncProcessingEx(
+		&AsyncState,
+		6,
+ 		[&TestOutput]()
+		{
+			TestOutput.SetNumUninitialized(6); 
+		},
+		[Value = MakeUnique<int32>(42), &TestOutput](int32 ReadIndex, int32 WriteIndex)
+		{
+			if ((ReadIndex % 2) == 0)
+			{
+				return false; // remove all even
+			}
+
+			TestOutput[WriteIndex] = double(ReadIndex * (*Value));
+
+			return true;
+		},
+		[&TestOutput](int32 ReadIndex, int32 WriteIndex)
+		{
+			TestOutput[WriteIndex] = TestOutput[ReadIndex];
+		},
+		[&TestOutput](int32 Count)
+		{
+			TestOutput.SetNum(Count);
+		},
+		/*bEnableTimeSlicing*/ false,
+		/*ChunkSize*/ 1 // force this to thrad 
+	);
+
+	UTEST_TRUE("IsDone", bIsDone);
+
+	UTEST_EQUAL("Wrote Correct Number", TestOutput.Num(), 3);
+
+	UTEST_EQUAL("TestOutput[0]", TestOutput[0], 42.0 * 1);
+	UTEST_EQUAL("TestOutput[1]", TestOutput[1], 42.0 * 3);
+	UTEST_EQUAL("TestOutput[2]", TestOutput[2], 42.0 * 5);
+
+	return true;
+}
+
+IMPLEMENT_CUSTOM_SIMPLE_AUTOMATION_TEST(FPCGAsyncTest_AsyncProcessingOneToOneEx, FPCGTestBaseClass, "pcg.Async.AsyncProcessingOneToOneEx", PCGTestsCommon::TestFlags)
+
+bool FPCGAsyncTest_AsyncProcessingOneToOneEx::RunTest(const FString& Parameters)
+{
+	FPCGAsyncState AsyncState;
+	AsyncState.NumAvailableTasks = 10;
+
+	TArray<double> TestOutput;
+	const bool bIsDone = FPCGAsync::AsyncProcessingOneToOneEx(
+		&AsyncState,
+		3,
+		[&TestOutput]()
+		{
+			TestOutput.SetNumUninitialized(3);
+		},
+		[Value = MakeUnique<int32>(42), &TestOutput](int32 ReadIndex, int32 WriteIndex)
+		{
+			TestOutput[WriteIndex] = double(ReadIndex * (*Value));
+		},
+		/*bEnableTimeSlicing*/ false,
+		/*ChunkSize*/ 1 // force this to thread
+	);
+
+	UTEST_TRUE("IsDone", bIsDone);
+
+	UTEST_EQUAL("Wrote Correct Number", TestOutput.Num(), 3);
+
+	UTEST_EQUAL("TestOutput[0]", TestOutput[0], 42.0 * 0);
+	UTEST_EQUAL("TestOutput[1]", TestOutput[1], 42.0 * 1);
+	UTEST_EQUAL("TestOutput[2]", TestOutput[2], 42.0 * 2);
+
+	return true;
+}
diff --git a/Engine/Plugins/Experimental/PCG/Source/PCG/Public/Helpers/PCGAsync.h b/Engine/Plugins/Experimental/PCG/Source/PCG/Public/Helpers/PCGAsync.h
--- a/Engine/Plugins/Experimental/PCG/Source/PCG/Public/Helpers/PCGAsync.h
+++ b/Engine/Plugins/Experimental/PCG/Source/PCG/Public/Helpers/PCGAsync.h
@@ -93,8 +93,8 @@ namespace FPCGAsync
 
 	namespace Private
 	{
-	template <typename OutputType, typename Func>
-	bool AsyncProcessing(FPCGAsyncState& AsyncState, int32 NumIterations, TArray<OutputType>& OutData, Func IterationInnerLoop, const bool bInEnableTimeSlicing, const int32 InChunkSize)
+	template <typename InitializeFunc, typename InnerLoopFunc, typename MoveDataFunc, typename FinishedFunc>
+	bool AsyncProcessing(FPCGAsyncState& AsyncState, int32 NumIterations, InitializeFunc&& Initialize, InnerLoopFunc&& IterationInnerLoop, MoveDataFunc&& MoveData, FinishedFunc&& Finished, const bool bInEnableTimeSlicing, const int32 InChunkSize)
 	{
 		TRACE_CPUPROFILER_EVENT_SCOPE(FPCGAsync::AsyncProcessing);
 
@@ -121,11 +121,11 @@ namespace FPCGAsync
 		const int32 NumTasks = FMath::Min(AsyncState.NumAvailableTasks, ChunksNumber);
 		const int32 NumFutures = NumTasks - 1;
 
-		// Pre-reserve the out data array, only if it is our first time here
 		if (StartIndex == 0)
 		{
-			TRACE_CPUPROFILER_EVENT_SCOPE(FPCGAsync::AsyncProcessing::AllocatingArray);
-			OutData.SetNumUninitialized(NumIterations);
+			TRACE_CPUPROFILER_EVENT_SCOPE(FPCGAsync::Private::AsyncProcessing::Initialize);
+
+			Initialize();
 		}
 
 		// Synchronisation structure to be shared between async tasks and collapsing main thread.
@@ -164,7 +164,14 @@ namespace FPCGAsync
 
 				check(Count > 0);
 
-				const int32 NumElementsWritten = IterationInnerLoop(StartReadIndex, StartWriteIndex, Count);
+				int32 NumElementsWritten;
+
+				{
+					TRACE_CPUPROFILER_EVENT_SCOPE(FPCGAsync::Private::AsyncProcessing::IterationInnerLoop);
+
+					NumElementsWritten = IterationInnerLoop(StartReadIndex, StartWriteIndex, Count);
+				}
+
 				AsyncState.AsyncCurrentWriteIndex += NumElementsWritten;
 				AsyncState.AsyncCurrentReadIndex += ChunkSize;
 				++CurrentChunkToProcess;
@@ -197,13 +204,20 @@ namespace FPCGAsync
 
 				check(Count > 0);
 
-				const int32 NumElementsWritten = IterationInnerLoop(StartReadIndex, StartWriteIndex, Count);
+				int32 NumElementsWritten;
+
+				{
+					TRACE_CPUPROFILER_EVENT_SCOPE(FPCGAsync::Private::AsyncProcessing::InnerLoop);
+
+					NumElementsWritten = IterationInnerLoop(StartReadIndex, StartWriteIndex, Count);
+				}
+
 				SynchroStruct.ChunkProcessedIndexAndNumElementsWrittenQueue.Enqueue(CurrentChunkToProcess, NumElementsWritten);
 			};
 
 			for (int32 TaskIndex = 0; TaskIndex < NumFutures; ++TaskIndex)
 			{
-				TRACE_CPUPROFILER_EVENT_SCOPE(FPCGAsync::AsyncProcessing::StartingTasks);
+				TRACE_CPUPROFILER_EVENT_SCOPE(FPCGAsync::Private::AsyncProcessing::StartingTasks);
 				AsyncTasks.Emplace(UE::Tasks::Launch(UE_SOURCE_LOCATION, [&JobTask, &SynchroStruct]() -> void
 				{
 					// Do at least one run, as we could fall into an infinite loop if we always have to stop before doing anything.
@@ -266,7 +280,7 @@ namespace FPCGAsync
 							// Otherwise collapse
 							for (int32 i = 0; i < NumberOfElementsWritten; ++i)
 							{
-								OutData[AsyncState.AsyncCurrentWriteIndex++] = std::move(OutData[ReadStartIndexForCurrentChunkToCollapse + i]);
+								MoveData(ReadStartIndexForCurrentChunkToCollapse + i, AsyncState.AsyncCurrentWriteIndex++);
 							}
 						}
 
@@ -290,11 +304,10 @@ namespace FPCGAsync
 
 		if (bIsDone)
 		{
-			TRACE_CPUPROFILER_EVENT_SCOPE(FPCGAsync::AsyncProcessing::ShrinkOutData);
-			// Shrinking can have a big impact on the performance, but without it, we can also hold a big chunk of wasted memory.
-			// Might revisit later if the performance impact is too big.
-			OutData.SetNum(AsyncState.AsyncCurrentWriteIndex);
-			
+			TRACE_CPUPROFILER_EVENT_SCOPE(FPCGAsync::Private::AsyncProcessing::Finished);
+
+			Finished(AsyncState.AsyncCurrentWriteIndex);
+
 			// Reset indexes
 			AsyncState.AsyncCurrentReadIndex = 0;
 			AsyncState.AsyncCurrentWriteIndex = 0;
@@ -305,7 +318,9 @@ namespace FPCGAsync
 	}
 
 	/**
-	* Helper for generic parallel loops, with support for timeslicing.
+	* A Helper for generic parallel loops, with support for timeslicing. This version only uses indexes allowing more flexible usage
+	* to process many arrays in parallel or use it for other batch updating.
+	* 
 	* Work will be separated in chunks, that will be processed in parallel. Main thread will then collapse incoming data from async tasks.
 	* Will use AsyncState.ShouldStop() to stop execution if timeslicing is enabled.
 	* Important info: 
@@ -313,25 +328,26 @@ namespace FPCGAsync
 	*   - To avoid infinite loops (when we should stop even before starting working), we will at least process 1 chunk of data per thread.
 	*   - To have async tasks, you need to have at least 3 available threads (main thread + 2 futures). Otherwise, we will only process on the main thread, without collapse.
 	* 
-	* @param Context - The context containing the information about how many tasks we can launch, async read/write index for the current job and a function to know if we need to stop processing.
+	* @param AsyncState - The context containing the information about how many tasks we can launch, async read/write index for the current job and a function to know if we need to stop processing.
 	* @param NumIterations - The number of calls that will be done to the provided function, also an upper bound on the number of data generated.
-	* @param OutData - The array in which the results will be written to. Note that the array will be cleared before execution.
-	* @param Func - Signature: bool(int32, OutputType&). A function that has the index [0; NumIterations] and has to write to some data & return false if the result should be discarded. 
+	* @param Initialize - Signature: void(). A function that will be called once on the first timeslice, where you can reserve data for processing
+	* @param ProcessElement - Signature: bool(int32 ReadIndex, int32 WriteIndex). A function that processes  
+	* @param MoveData - Signature: void(int32 ReadIndex, int32 WriteIndex). If the processing filters points, this will be used to move elements from one index to another 
+	* @param Finished - Signature: void(int32 Count). Called once on finished, and tells you the total count of points written.
 	* @param bEnableTimeSlicing - If false, we will not stop until all the processing is done.
 	* @param ChunkSize - Size of the chunks to cut the input data with
 	* @returns true if the processing is done, false otherwise. Use this to know if you need to reschedule the task.
 	*/
-	template <typename OutputType, typename Func>
-	bool AsyncProcessing(FPCGAsyncState* AsyncState, int32 NumIterations, TArray<OutputType>& OutData, Func InFunc, const bool bEnableTimeSlicing, const int32 ChunkSize = 64)
+	template <typename InitializeFunc, typename ProcessElementFunc, typename MoveFunc, typename FinishedFunc>
+	bool AsyncProcessingEx(FPCGAsyncState* AsyncState, int32 NumIterations, InitializeFunc&& Initialize, ProcessElementFunc&& ProcessElement, MoveFunc&& MoveData, FinishedFunc&& Finished, const bool bEnableTimeSlicing, const int32 ChunkSize = 64)
 	{
-		auto IterationInnerLoop = [&InFunc, &OutData](int32 StartReadIndex, int32 StartWriteIndex, int32 Count) -> int32
+		auto IterationInnerLoop = [Func = MoveTemp(ProcessElement)](int32 StartReadIndex, int32 StartWriteIndex, int32 Count) -> int32
 		{
-			TRACE_CPUPROFILER_EVENT_SCOPE(FPCGAsync::AsyncProcessing::InnerLoop);
 			int32 NumPointsWritten = 0;
 
 			for (int32 i = 0; i < Count; ++i)
 			{
-				if (InFunc(StartReadIndex + i, OutData[StartWriteIndex + NumPointsWritten]))
+				if (Func(StartReadIndex + i, StartWriteIndex + NumPointsWritten))
 				{
 					++NumPointsWritten;
 				}
@@ -343,7 +359,7 @@ namespace FPCGAsync
 		if (AsyncState && !AsyncState->bIsRunningAsyncCall)
 		{
 			AsyncState->bIsRunningAsyncCall = true;
-			bool bIsDone = Private::AsyncProcessing<OutputType>(*AsyncState, NumIterations, OutData, IterationInnerLoop, bEnableTimeSlicing, ChunkSize);
+			const bool bIsDone = Private::AsyncProcessing(*AsyncState, NumIterations, Initialize, IterationInnerLoop, MoveData, Finished, bEnableTimeSlicing, ChunkSize);
 			AsyncState->bIsRunningAsyncCall = false;
 			return bIsDone;
 		}
@@ -354,9 +370,112 @@ namespace FPCGAsync
 			FPCGAsyncState DummyState;
 			DummyState.NumAvailableTasks = 1;
 			DummyState.bIsRunningAsyncCall = true;
-			return Private::AsyncProcessing<OutputType>(DummyState, NumIterations, OutData, IterationInnerLoop, /*bEnableTimeSlicing=*/false, ChunkSize);
+			return Private::AsyncProcessing(DummyState, NumIterations, Initialize, IterationInnerLoop, MoveData, Finished, /*bEnableTimeSlicing=*/false, ChunkSize);
 		}
 	}
+ 
+ 	/**
+	* A Helper for generic parallel loops, with support for timeslicing. This version only uses indexes allowing more flexible usage
+	* to process any data format.
+	* 
+	* This version also does not support removing elements and is expects a 1:1 mapping. This just removes some lambdas required to support
+	* that use case.
+	* 
+	* Work will be separated in chunks, that will be processed in parallel. Main thread will then collapse incoming data from async tasks.
+	* Will use AsyncState.ShouldStop() to stop execution if timeslicing is enabled.
+	* Important info: 
+	*   - We will finish to process and collapse data for all data already in process, even if we need to stop. To mitigate this, try to use small chunk sizes.
+	*   - To avoid infinite loops (when we should stop even before starting working), we will at least process 1 chunk of data per thread.
+	*   - To have async tasks, you need to have at least 3 available threads (main thread + 2 futures). Otherwise, we will only process on the main thread, without collapse.
+	* 
+	* @param AsyncState - The context containing the information about how many tasks we can launch, async read/write index for the current job and a function to know if we need to stop processing.
+	* @param NumIterations - The number of calls that will be done to the provided function, also an upper bound on the number of data generated.
+	* @param Initialize - Signature: void(). A function that will be called once on the first timeslice, where you can reserve data.
+	* @param ProcessElement - Signature: bool(int32 ReadIndex, int32 WriteIndex). A function that processes  
+	* @param bEnableTimeSlicing - If false, we will not stop until all the processing is done.
+	* @param ChunkSize - Size of the chunks to cut the input data with
+	* @returns true if the processing is done, false otherwise. Use this to know if you need to reschedule the task.
+	*/
+ 	template <typename InitializeFunc, typename ProcessElementFunc>
+	bool AsyncProcessingOneToOneEx(FPCGAsyncState* AsyncState, int32 NumIterations, InitializeFunc&& Initialize, ProcessElementFunc&& ProcessElement, const bool bEnableTimeSlicing, const int32 ChunkSize = 64)
+	{
+		auto IterationInnerLoop = [Func = MoveTemp(ProcessElement)](int32 StartReadIndex, int32 StartWriteIndex, int32 Count) -> int32
+		{
+			for (int32 i = 0; i < Count; ++i)
+			{
+				Func(StartReadIndex + i, StartWriteIndex + i);
+			}
+
+			return Count;
+		};
+
+		auto MoveData = [](int32, int32) { ensure(false); };
+		auto Finished = [NumIterations](int32 Count) { ensure(NumIterations == Count); };
+
+		if (AsyncState && !AsyncState->bIsRunningAsyncCall)
+		{
+			AsyncState->bIsRunningAsyncCall = true;
+			const bool bIsDone = Private::AsyncProcessing(*AsyncState, NumIterations, Initialize, IterationInnerLoop, MoveData, Finished, bEnableTimeSlicing, ChunkSize);
+			AsyncState->bIsRunningAsyncCall = false;
+			return bIsDone;
+		}
+		else
+		{
+			// Can't use time slicing without an async state or while running in another async call (it will mess up with async indexes). 
+			// We also force using one thread (the current one).
+			FPCGAsyncState DummyState;
+			DummyState.NumAvailableTasks = 1;
+			DummyState.bIsRunningAsyncCall = true;
+			return Private::AsyncProcessing(DummyState, NumIterations, Initialize, IterationInnerLoop, MoveData, Finished, /*bEnableTimeSlicing=*/false, ChunkSize);
+		}
+	}
+
+	/**
+	* Helper for generic parallel loops, with support for timeslicing. This version takes a data type, expecting processing on an
+	* array of that type.
+	* 
+	* Work will be separated in chunks, that will be processed in parallel. Main thread will then collapse incoming data from async tasks.
+	* Will use AsyncState.ShouldStop() to stop execution if timeslicing is enabled.
+	* Important info: 
+	*   - We will finish to process and collapse data for all data already in process, even if we need to stop. To mitigate this, try to use small chunk sizes.
+	*   - To avoid infinite loops (when we should stop even before starting working), we will at least process 1 chunk of data per thread.
+	*   - To have async tasks, you need to have at least 3 available threads (main thread + 2 futures). Otherwise, we will only process on the main thread, without collapse.
+	* 
+	* @param AsyncState - The context containing the information about how many tasks we can launch, async read/write index for the current job and a function to know if we need to stop processing.
+	* @param NumIterations - The number of calls that will be done to the provided function, also an upper bound on the number of data generated.
+	* @param OutData - The array in which the results will be written to. Note that the array will be cleared before execution.
+	* @param Func - Signature: bool(int32, OutputType&). A function that has the index [0; NumIterations] and has to write to some data & return false if the result should be discarded. 
+	* @param bEnableTimeSlicing - If false, we will not stop until all the processing is done.
+	* @param ChunkSize - Size of the chunks to cut the input data with
+	* @returns true if the processing is done, false otherwise. Use this to know if you need to reschedule the task.
+	*/
+	template <typename OutputType, typename Func>
+	bool AsyncProcessing(FPCGAsyncState* AsyncState, int32 NumIterations, TArray<OutputType>& OutData, Func&& InFunc, const bool bEnableTimeSlicing, const int32 ChunkSize = 64)
+	{
+		auto Initialize = [&OutData, NumIterations]()
+		{
+			OutData.SetNumUninitialized(NumIterations);
+		};
+		
+		auto IterationInnerLoop = [Func = MoveTemp(InFunc), &OutData](int32 ReadIndex, int32 WriteIndex) -> int32
+		{
+			return Func(ReadIndex, OutData[WriteIndex]);
+		};
+
+		auto MoveData = [&OutData](int32 ReadIndex, int32 WriteIndex)
+		{
+			OutData[WriteIndex] = std::move(OutData[ReadIndex]);
+		};
+
+		auto Finished = [&OutData](int32 Count)
+		{
+			// Shrinking can have a big impact on the performance, but without it, we can also hold a big chunk of wasted memory.
+			// Might revisit later if the performance impact is too big.
+			OutData.SetNum(Count);			
+		};
+
+		return AsyncProcessingEx(AsyncState, NumIterations, Initialize, IterationInnerLoop, MoveData, Finished, bEnableTimeSlicing, ChunkSize);
+	}
 }
 
 #if UE_ENABLE_INCLUDE_ORDER_DEPRECATED_IN_5_2
