Fixed SplineSampler ignoring NextDirectionDelta on last point of a closed spline

NextDirectionDelta was not being calculated (left as default initialization) when the spline was closed.

diff --git a/Engine/Plugins/Experimental/PCG/Source/PCG/Private/Elements/PCGSplineSampler.cpp b/Engine/Plugins/Experimental/PCG/Source/PCG/Private/Elements/PCGSplineSampler.cpp
--- a/Engine/Plugins/Experimental/PCG/Source/PCG/Private/Elements/PCGSplineSampler.cpp
+++ b/Engine/Plugins/Experimental/PCG/Source/PCG/Private/Elements/PCGSplineSampler.cpp
@@ -670,6 +670,18 @@ namespace PCGSplineSampler
 	{
 		check(LineData && OutPointData);
 
+		bool bIsClosedSpline = false;
+
+		if (const UPCGSplineData* SplineData = Cast<UPCGSplineData>(LineData))
+		{
+			const FPCGSplineStruct* Spline = &SplineData->SplineStruct;
+
+			if (Spline && Spline->IsClosedLoop())
+			{
+				bIsClosedSpline = true;
+			}
+		}
+
 		FSubdivisionStepSampler SubdivisionSampler(LineData, Params);
 		FDistanceStepSampler DistanceSampler(LineData, Params);
 
@@ -690,6 +702,8 @@ namespace PCGSplineSampler
 			Sampler->Step(*PreviousResult);
 			bHasPreviousPoint = true;
 		}
+
+		const FSamplerResult FirstResult = *PreviousResult;
 		
 		while(!Sampler->IsDone())
 		{
@@ -713,6 +727,15 @@ namespace PCGSplineSampler
 		
 		if (bHasPreviousPoint)
 		{
+			// Get the DeltaAngle between last and first points if we are operating on a closed spline
+			if (bIsClosedSpline)
+			{
+				// Get unsigned angle difference between the two points
+				const FVector::FReal DeltaSinAngle = ((FirstResult.LocalTransform.GetUnitAxis(EAxis::X) ^ PreviousResult->LocalTransform.GetUnitAxis(EAxis::X)) | PreviousResult->LocalTransform.GetUnitAxis(EAxis::Z));
+				// Normalize value to be between -1 and 1
+				PreviousResult->NextDeltaAngle = FMath::Asin(DeltaSinAngle) / UE_HALF_PI;
+			}
+
 			ExtentsSampler->Sample(*PreviousResult, OutPointData);
 		}
 	}
