Rolled back cancel-on-refresh changes until we have something to limit this in some specific cases (namely when we do landscape changes through PCG)

#rnx

diff --git a/Engine/Plugins/Experimental/PCG/Source/PCG/Private/PCGComponent.cpp b/Engine/Plugins/Experimental/PCG/Source/PCG/Private/PCGComponent.cpp
--- a/Engine/Plugins/Experimental/PCG/Source/PCG/Private/PCGComponent.cpp
+++ b/Engine/Plugins/Experimental/PCG/Source/PCG/Private/PCGComponent.cpp
@@ -1712,10 +1712,11 @@ void UPCGComponent::Refresh()
 	// Discard any refresh if have already one scheduled.
 	if (UPCGSubsystem* Subsystem = GetSubsystem())
 	{
-		if (CurrentGenerationTask != InvalidPCGTaskId)
+		//TDOO: uncomment this once we have the proper things in place to prevent cancellation-by-landscape changes when we have a graph that does change the landscape
+		/*if (CurrentGenerationTask != InvalidPCGTaskId)
 		{
 			CancelGeneration();
-		}
+		}*/
 
 		if (CurrentRefreshTask == InvalidPCGTaskId)
 		{
