Allow cleanup after cancellation in cases where we have already resources on the component.

#rnx

diff --git a/Engine/Plugins/Experimental/PCG/Source/PCG/Private/PCGComponent.cpp b/Engine/Plugins/Experimental/PCG/Source/PCG/Private/PCGComponent.cpp
--- a/Engine/Plugins/Experimental/PCG/Source/PCG/Private/PCGComponent.cpp
+++ b/Engine/Plugins/Experimental/PCG/Source/PCG/Private/PCGComponent.cpp
@@ -749,7 +749,7 @@ void UPCGComponent::CleanupLocalImmediate(bool bRemoveComponents)
 
 FPCGTaskId UPCGComponent::CreateCleanupTask(bool bRemoveComponents, const TArray<FPCGTaskId>& Dependencies)
 {
-	if ((!bGenerated && !IsGenerating()) || IsPartitioned() || IsCleaningUp())
+	if ((!bGenerated && GeneratedResources.IsEmpty() && !IsGenerating()) || IsPartitioned() || IsCleaningUp())
 	{
 		return InvalidPCGTaskId;
 	}
