[PCG] Add graph information in UpdatePins warning to help find them quicker.

#rnx

diff --git a/Engine/Plugins/Experimental/PCG/Source/PCG/Private/PCGNode.cpp b/Engine/Plugins/Experimental/PCG/Source/PCG/Private/PCGNode.cpp
--- a/Engine/Plugins/Experimental/PCG/Source/PCG/Private/PCGNode.cpp
+++ b/Engine/Plugins/Experimental/PCG/Source/PCG/Private/PCGNode.cpp
@@ -698,8 +698,10 @@ EPCGChangeType UPCGNode::UpdatePins(TFunctionRef<UPCGPin*(UPCGNode*)> PinAllocat
 			{
 				if (Properties[i].Label == Properties[j].Label)
 				{
-					UE_LOG(LogPCG, Warning, TEXT("UpdatePins: Pin properties from the settings on node '%s' contained a duplicate pin '%s', removing this pin properties."),
-						*GetFName().ToString(), *Properties[i].Label.ToString());
+					const UPCGGraph* PCGGraph = GetGraph();
+					const FString GraphName = PCGGraph ? PCGGraph->GetName() : FString(TEXT("Unknown"));
+					UE_LOG(LogPCG, Warning, TEXT("UpdatePins: Pin properties from the settings on node '%s' in graph '%s' contained a duplicate pin '%s', removing this pin properties."),
+						*GetName(), *GraphName, *Properties[i].Label.ToString());
 
 					// Remove but preserve order
 					Properties.RemoveAt(j);
