Fix localization warnings in PCGPointFilter.cpp
Submitted on behalf of adrien.logut

#rnx

diff --git a/Engine/Plugins/Experimental/PCG/Source/PCG/Private/Elements/PCGPointFilter.cpp b/Engine/Plugins/Experimental/PCG/Source/PCG/Private/Elements/PCGPointFilter.cpp
--- a/Engine/Plugins/Experimental/PCG/Source/PCG/Private/Elements/PCGPointFilter.cpp
+++ b/Engine/Plugins/Experimental/PCG/Source/PCG/Private/Elements/PCGPointFilter.cpp
@@ -484,17 +484,10 @@ bool FPCGPointFilterElementBase::DoFiltering(FPCGContext* Context, EPCGPointFilt
 			continue;
 		}
 
-
-		if (!PCG::Private::IsBroadcastable(FirstThresholdInfo.ThresholdAccessor->GetUnderlyingType(), TargetAccessor->GetUnderlyingType()))
-		{
-			PCGE_LOG(Warning, GraphAndLog, LOCTEXT("TypeConversionFailed", "Cannot broadcast threshold type to target type"));
-			ForwardInputToInFilterPin();
-			continue;
-		}
-
-		if (SecondThresholdInfo.ThresholdAccessor.IsValid() && !PCG::Private::IsBroadcastable(SecondThresholdInfo.ThresholdAccessor->GetUnderlyingType(), TargetAccessor->GetUnderlyingType()))
+		if (!PCG::Private::IsBroadcastable(FirstThresholdInfo.ThresholdAccessor->GetUnderlyingType(), TargetAccessor->GetUnderlyingType())
+			|| (SecondThresholdInfo.ThresholdAccessor.IsValid() && !PCG::Private::IsBroadcastable(SecondThresholdInfo.ThresholdAccessor->GetUnderlyingType(), TargetAccessor->GetUnderlyingType())))
 		{
-			PCGE_LOG(Warning, GraphAndLog, LOCTEXT("TypeConversionFailed", "Cannot broadcast threshold type to target type"));
+			PCGE_LOG(Warning, GraphAndLog, LOCTEXT("TypeCannotBeConverted", "Cannot convert threshold type to target type"));
 			ForwardInputToInFilterPin();
 			continue;
 		}
