[PCG] Introduce concept of 'stack' to graph compilation/execution.

Each task is associated with a stack context and given a stack frame index. This is to enable:

* Subgraph inspection - for a couple of reasons, including to know if a subgraph was called during execution of a graph (and which subgraph node(s) invoked it)
* Dynamic linking of data across LODs - data is stored/retrieve using a unique key built from the stack - so also requires stack information

diff --git a/Engine/Plugins/Experimental/PCG/Source/PCG/Private/Elements/PCGLoopElement.cpp b/Engine/Plugins/Experimental/PCG/Source/PCG/Private/Elements/PCGLoopElement.cpp
--- a/Engine/Plugins/Experimental/PCG/Source/PCG/Private/Elements/PCGLoopElement.cpp
+++ b/Engine/Plugins/Experimental/PCG/Source/PCG/Private/Elements/PCGLoopElement.cpp
@@ -6,6 +6,7 @@
 #include "PCGModule.h"
 #include "PCGPin.h"
 #include "PCGSubsystem.h"
+#include "Graph/PCGStackContext.h"
 #include "Helpers/PCGHelpers.h"
 
 #include UE_INLINE_GENERATED_CPP_BY_NAME(PCGLoopElement)
@@ -99,7 +100,7 @@ bool FPCGLoopElement::ExecuteInternal(FPCGContext* InContext) const
 
 		// Dispatch the subgraph for each loop entries we have.
 		// Implementation note: even if the execution gets cancelled, these tasks would get cancelled because they are associated to the current source component
-		for(int EntryIndex = 0; EntryIndex < LoopDataCollection[0].TaggedData.Num(); ++EntryIndex)
+		for (int EntryIndex = 0; EntryIndex < LoopDataCollection[0].TaggedData.Num(); ++EntryIndex)
 		{
 			// The final input for a given loop iteration is going to be the fixed data + the entry at the current loop index in all the gathered loop collections
 			FPCGDataCollection InputDataCollection = FixedDataCollection;
@@ -108,7 +109,15 @@ bool FPCGLoopElement::ExecuteInternal(FPCGContext* InContext) const
 				InputDataCollection.TaggedData.Insert(LoopDataCollection[LoopCollectionIndex].TaggedData[EntryIndex], LoopCollectionIndex);
 			}
 
-			FPCGTaskId SubgraphTaskId = Subsystem->ScheduleGraph(Subgraph, Context->SourceComponent.Get(), PreGraphElement, MakeShared<FPCGInputForwardingElement>(InputDataCollection), {});
+			// Prepare the invocation stack - which is the stack up to this node, and then this node, then a loop index
+			FPCGStack InvocationStack = ensure(Context->Stack) ? *Context->Stack : FPCGStack();
+
+			TArray<FPCGStackFrame>& StackFrames = InvocationStack.GetStackFramesMutable();
+			StackFrames.Reserve(StackFrames.Num() + 2);
+			StackFrames.Emplace(Context->Node);
+			StackFrames.Emplace(EntryIndex);
+
+			FPCGTaskId SubgraphTaskId = Subsystem->ScheduleGraph(Subgraph, Context->SourceComponent.Get(), PreGraphElement, MakeShared<FPCGInputForwardingElement>(InputDataCollection), {}, &InvocationStack);
 
 			Context->SubgraphTaskIds.Add(SubgraphTaskId);
 		}
diff --git a/Engine/Plugins/Experimental/PCG/Source/PCG/Private/Graph/PCGGraphCompiler.cpp b/Engine/Plugins/Experimental/PCG/Source/PCG/Private/Graph/PCGGraphCompiler.cpp
--- a/Engine/Plugins/Experimental/PCG/Source/PCG/Private/Graph/PCGGraphCompiler.cpp
+++ b/Engine/Plugins/Experimental/PCG/Source/PCG/Private/Graph/PCGGraphCompiler.cpp
@@ -1,21 +1,25 @@
 // Copyright Epic Games, Inc. All Rights Reserved.
 
 #include "PCGGraphCompiler.h"
-#include "Graph/PCGGraphExecutor.h"
-#include "Misc/ScopeRWLock.h"
+
 #include "PCGGraph.h"
 #include "PCGEdge.h"
 #include "PCGModule.h"
-#include "PCGSubgraph.h"
 #include "PCGPin.h"
+#include "PCGSubgraph.h"
+#include "Graph/PCGGraphExecutor.h"
+
+#include "Misc/ScopeRWLock.h"
 
-TArray<FPCGGraphTask> FPCGGraphCompiler::CompileGraph(UPCGGraph* InGraph, FPCGTaskId& NextId)
+TArray<FPCGGraphTask> FPCGGraphCompiler::CompileGraph(UPCGGraph* InGraph, FPCGTaskId& NextId, FPCGStackContext& InOutStackContext)
 {
 	if (!InGraph)
 	{
 		return TArray<FPCGGraphTask>();
 	}
 
+	InOutStackContext.PushFrame(InGraph);
+
 	TArray<FPCGGraphTask> CompiledTasks;
 	TMap<const UPCGNode*, FPCGTaskId> IdMapping;
 	TArray<const UPCGNode*> NodeQueue;
@@ -42,7 +46,7 @@ TArray<FPCGGraphTask> FPCGGraphCompiler::CompileGraph(UPCGGraph* InGraph, FPCGTa
 		const UPCGBaseSubgraphSettings* SubgraphSettings = SubgraphNode ? Cast<const UPCGBaseSubgraphSettings>(SubgraphNode->GetSettings()) : nullptr;
 
 		// Only catch immediate infinite recursion.
-		// TODO: Add a better mecanism for detecting more complex cyclic recursions
+		// TODO: Add a better mechanism for detecting more complex cyclic recursions
 		// Like Graph A has subgraph node with Graph B, and Graph B has a subgraph node with Graph A (A -> B -> A)
 		// or A -> B -> C -> A, etc...
 		// NOTE (for the person that would work on it), keeping a stack of all the subgraphs is not enough, as we could already have compiled graph
@@ -61,7 +65,8 @@ TArray<FPCGGraphTask> FPCGGraphCompiler::CompileGraph(UPCGGraph* InGraph, FPCGTa
 
 			// 1. Compile the subgraph making sure we don't reuse the same ids
 			// Note that we will not consume the pre or post-execute tasks, ergo bIsTopGraph=false
-			TArray<FPCGGraphTask> Subtasks = GetCompiledTasks(Subgraph, /*bIsTopGraph=*/false);
+			FPCGStackContext SubgraphStackContext;
+			TArray<FPCGGraphTask> Subtasks = GetCompiledTasks(Subgraph, SubgraphStackContext, /*bIsTopGraph=*/false);
 
 #if WITH_EDITOR
 			GraphDependenciesLock.Lock();
@@ -69,6 +74,16 @@ TArray<FPCGGraphTask> FPCGGraphCompiler::CompileGraph(UPCGGraph* InGraph, FPCGTa
 			GraphDependenciesLock.Unlock();
 #endif // WITH_EDITOR
 
+			// Append all the stack frames from inside the subgraph to my current graph stack
+			InOutStackContext.PushFrame(Node);
+			const int32 StackOffset = InOutStackContext.GetNumStacks();
+			for (FPCGGraphTask& Subtask : Subtasks)
+			{
+				Subtask.StackIndex += StackOffset;
+			}
+			InOutStackContext.AppendStacks(SubgraphStackContext);
+			InOutStackContext.PopFrame();
+
 			OffsetNodeIds(Subtasks, NextId, PreId);
 			NextId += Subtasks.Num();
 
@@ -88,6 +103,7 @@ TArray<FPCGGraphTask> FPCGGraphCompiler::CompileGraph(UPCGGraph* InGraph, FPCGTa
 			FPCGGraphTask& PreTask = CompiledTasks.Emplace_GetRef();
 			PreTask.Node = Node;
 			PreTask.NodeId = PreId;
+			PreTask.StackIndex = InOutStackContext.GetCurrentStackIndex();
 
 			for (const UPCGPin* InputPin : Node->InputPins)
 			{
@@ -135,6 +151,7 @@ TArray<FPCGGraphTask> FPCGGraphCompiler::CompileGraph(UPCGGraph* InGraph, FPCGTa
 			FPCGGraphTask& PostTask = CompiledTasks.Emplace_GetRef();
 			PostTask.Node = Node;
 			PostTask.NodeId = PostId;
+			PostTask.StackIndex = InOutStackContext.GetCurrentStackIndex();
 
 			// Add subgraph output node task as input to the post-task
 			if (OutputNodeTask)
@@ -151,6 +168,7 @@ TArray<FPCGGraphTask> FPCGGraphCompiler::CompileGraph(UPCGGraph* InGraph, FPCGTa
 			FPCGGraphTask& Task = CompiledTasks.Emplace_GetRef();
 			Task.Node = Node;
 			Task.NodeId = NodeId;
+			Task.StackIndex = InOutStackContext.GetCurrentStackIndex();
 
 			for (const UPCGPin* InputPin : Node->InputPins)
 			{
@@ -164,11 +182,11 @@ TArray<FPCGGraphTask> FPCGGraphCompiler::CompileGraph(UPCGGraph* InGraph, FPCGTa
 
 					if (FPCGTaskId* InboundId = IdMapping.Find(InboundEdge->InputPin->Node))
 					{
-						Task.Inputs.Emplace(IdMapping[InboundEdge->InputPin->Node], InboundEdge->InputPin, InboundEdge->OutputPin); 
+						Task.Inputs.Emplace(*InboundId, InboundEdge->InputPin, InboundEdge->OutputPin); 
 					}
 					else
 					{
-						UE_LOG(LogPCG, Error, TEXT("Inconsistent node linkage on node %s"), *Node->GetFName().ToString());
+						UE_LOG(LogPCG, Error, TEXT("Inconsistent node linkage on node '%s'"), *Node->GetFName().ToString());
 						return TArray<FPCGGraphTask>();
 					}
 				}
@@ -235,7 +253,8 @@ void FPCGGraphCompiler::Compile(UPCGGraph* InGraph)
 	// Otherwise, do the compilation; note that we always start at zero since
 	// the caller will offset the ids as needed
 	FPCGTaskId FirstId = 0;
-	TArray<FPCGGraphTask> CompiledTasks = CompileGraph(InGraph, FirstId);
+	FPCGStackContext StackContext;
+	TArray<FPCGGraphTask> CompiledTasks = CompileGraph(InGraph, FirstId, StackContext);
 
 	// TODO: optimize no-ops, etc.
 
@@ -246,22 +265,34 @@ void FPCGGraphCompiler::Compile(UPCGGraph* InGraph)
 		if (!GraphToTaskMap.Contains(InGraph))
 		{
 			GraphToTaskMap.Add(InGraph, MoveTemp(CompiledTasks));
+			GraphToStackContext.Add(InGraph, StackContext);
 		}
 		GraphToTaskMapLock.WriteUnlock();
 	}
 }
 
-TArray<FPCGGraphTask> FPCGGraphCompiler::GetPrecompiledTasks(UPCGGraph* InGraph, bool bIsTopGraph) const
+TArray<FPCGGraphTask> FPCGGraphCompiler::GetPrecompiledTasks(UPCGGraph* InGraph, FPCGStackContext& OutStackContext, bool bIsTopGraph) const
 {
 	// Get compiled tasks in a threadsafe way
 	FReadScopeLock ReadLock(GraphToTaskMapLock);
 
 	const TArray<FPCGGraphTask>* ExistingTasks = (bIsTopGraph ? TopGraphToTaskMap : GraphToTaskMap).Find(InGraph);
 
+	const FPCGStackContext* StackContext = (bIsTopGraph ? TopGraphToStackContext : GraphToStackContext).Find(InGraph);
+	if (StackContext)
+	{
+		OutStackContext = *StackContext;
+	}
+	else
+	{
+		// If we failed to get a context, output a dummy blank one.
+		new(&OutStackContext) FPCGStackContext();
+	}
+
 	return ExistingTasks ? *ExistingTasks : TArray<FPCGGraphTask>();
 }
 
-TArray<FPCGGraphTask> FPCGGraphCompiler::GetCompiledTasks(UPCGGraph* InGraph, bool bIsTopGraph)
+TArray<FPCGGraphTask> FPCGGraphCompiler::GetCompiledTasks(UPCGGraph* InGraph, FPCGStackContext& OutStackContext, bool bIsTopGraph)
 {
 	TArray<FPCGGraphTask> CompiledTasks;
 
@@ -275,6 +306,7 @@ TArray<FPCGGraphTask> FPCGGraphCompiler::GetCompiledTasks(UPCGGraph* InGraph, bo
 		if (TopGraphToTaskMap.Contains(InGraph))
 		{
 			CompiledTasks = TopGraphToTaskMap[InGraph];
+			OutStackContext = TopGraphToStackContext[InGraph];
 		}
 		GraphToTaskMapLock.ReadUnlock();
 	}
@@ -288,6 +320,7 @@ TArray<FPCGGraphTask> FPCGGraphCompiler::GetCompiledTasks(UPCGGraph* InGraph, bo
 		if (GraphToTaskMap.Contains(InGraph))
 		{
 			CompiledTasks = GraphToTaskMap[InGraph];
+			OutStackContext = GraphToStackContext[InGraph];
 		}
 		GraphToTaskMapLock.ReadUnlock();
 	}
@@ -329,7 +362,8 @@ void FPCGGraphCompiler::CompileTopGraph(UPCGGraph* InGraph)
 	}
 
 	// Build from non-top tasks
-	TArray<FPCGGraphTask> CompiledTasks = GetCompiledTasks(InGraph, /*bIsTopGraph=*/false);
+	FPCGStackContext StackContext;
+	TArray<FPCGGraphTask> CompiledTasks = GetCompiledTasks(InGraph, StackContext, /*bIsTopGraph=*/false);
 
 	// Check that the compilation was valid
 	if (CompiledTasks.Num() == 0)
@@ -390,6 +424,7 @@ void FPCGGraphCompiler::CompileTopGraph(UPCGGraph* InGraph)
 	if (!TopGraphToTaskMap.Contains(InGraph))
 	{
 		TopGraphToTaskMap.Add(InGraph, MoveTemp(CompiledTasks));
+		TopGraphToStackContext.Add(InGraph, StackContext);
 	}
 	GraphToTaskMapLock.WriteUnlock();
 }
diff --git a/Engine/Plugins/Experimental/PCG/Source/PCG/Private/Graph/PCGGraphCompiler.h b/Engine/Plugins/Experimental/PCG/Source/PCG/Private/Graph/PCGGraphCompiler.h
--- a/Engine/Plugins/Experimental/PCG/Source/PCG/Private/Graph/PCGGraphCompiler.h
+++ b/Engine/Plugins/Experimental/PCG/Source/PCG/Private/Graph/PCGGraphCompiler.h
@@ -3,6 +3,7 @@
 #pragma once
 
 #include "PCGCommon.h"
+#include "Graph/PCGStackContext.h"
 
 class UPCGGraph;
 struct FPCGGraphTask;
@@ -15,8 +16,8 @@ class FPCGGraphCompiler
 {
 public:
 	void Compile(UPCGGraph* InGraph);
-	TArray<FPCGGraphTask> GetCompiledTasks(UPCGGraph* InGraph, bool bIsTopGraph = true);
-	TArray<FPCGGraphTask> GetPrecompiledTasks(UPCGGraph* InGraph, bool bIsTopGraph = true) const;
+	TArray<FPCGGraphTask> GetCompiledTasks(UPCGGraph* InGraph, FPCGStackContext& OutStackContext, bool bIsTopGraph = true);
+	TArray<FPCGGraphTask> GetPrecompiledTasks(UPCGGraph* InGraph, FPCGStackContext& OutStackContext, bool bIsTopGraph = true) const;
 
 	static void OffsetNodeIds(TArray<FPCGGraphTask>& Tasks, FPCGTaskId Offset, FPCGTaskId ParentId);
 
@@ -25,12 +26,14 @@ public:
 #endif
 
 private:
-	TArray<FPCGGraphTask> CompileGraph(UPCGGraph* InGraph, FPCGTaskId& NextId);
+	TArray<FPCGGraphTask> CompileGraph(UPCGGraph* InGraph, FPCGTaskId& NextId, FPCGStackContext& InOutStackContext);
 	void CompileTopGraph(UPCGGraph* InGraph);
 
 	mutable FRWLock GraphToTaskMapLock;
 	TMap<UPCGGraph*, TArray<FPCGGraphTask>> GraphToTaskMap;
+	TMap<UPCGGraph*, FPCGStackContext> GraphToStackContext;
 	TMap<UPCGGraph*, TArray<FPCGGraphTask>> TopGraphToTaskMap;
+	TMap<UPCGGraph*, FPCGStackContext> TopGraphToStackContext;
 
 #if WITH_EDITOR
 	void RemoveFromCache(UPCGGraph* InGraph);
diff --git a/Engine/Plugins/Experimental/PCG/Source/PCG/Private/Graph/PCGGraphExecutor.cpp b/Engine/Plugins/Experimental/PCG/Source/PCG/Private/Graph/PCGGraphExecutor.cpp
--- a/Engine/Plugins/Experimental/PCG/Source/PCG/Private/Graph/PCGGraphExecutor.cpp
+++ b/Engine/Plugins/Experimental/PCG/Source/PCG/Private/Graph/PCGGraphExecutor.cpp
@@ -13,6 +13,7 @@
 #include "PCGPin.h"
 #include "PCGSubsystem.h"
 #include "Graph/PCGGraphCache.h"
+#include "Graph/PCGStackContext.h"
 #include "Helpers/PCGHelpers.h"
 #include "Metadata/PCGMetadata.h"
 
@@ -84,15 +85,15 @@ void FPCGGraphExecutor::Compile(UPCGGraph* Graph)
 	GraphCompiler->Compile(Graph);
 }
 
-FPCGTaskId FPCGGraphExecutor::Schedule(UPCGComponent* Component, const TArray<FPCGTaskId>& ExternalDependencies)
+FPCGTaskId FPCGGraphExecutor::Schedule(UPCGComponent* Component, const TArray<FPCGTaskId>& ExternalDependencies, const FPCGStack* InFromStack)
 {
 	check(Component);
 	UPCGGraph* Graph = Component->GetGraph();
 
-	return Schedule(Graph, Component, FPCGElementPtr(), GetFetchInputElement(), ExternalDependencies);
+	return Schedule(Graph, Component, FPCGElementPtr(), GetFetchInputElement(), ExternalDependencies, InFromStack);
 }
 
-FPCGTaskId FPCGGraphExecutor::Schedule(UPCGGraph* Graph, UPCGComponent* SourceComponent, FPCGElementPtr PreGraphElement, FPCGElementPtr InputElement, const TArray<FPCGTaskId>& ExternalDependencies)
+FPCGTaskId FPCGGraphExecutor::Schedule(UPCGGraph* Graph, UPCGComponent* SourceComponent, FPCGElementPtr PreGraphElement, FPCGElementPtr InputElement, const TArray<FPCGTaskId>& ExternalDependencies, const FPCGStack* InFromStack)
 {
 	if (SourceComponent && IsGraphCacheDebuggingEnabled())
 	{
@@ -113,12 +114,20 @@ FPCGTaskId FPCGGraphExecutor::Schedule(UPCGGraph* Graph, UPCGComponent* SourceCo
 	FPCGTaskId ScheduledId = InvalidPCGTaskId;
 
 	// Get compiled tasks from compiler
-	TArray<FPCGGraphTask> CompiledTasks = GraphCompiler->GetCompiledTasks(Graph);
+	TSharedPtr<FPCGStackContext> StackContextPtr = MakeShared<FPCGStackContext>();
+	TArray<FPCGGraphTask> CompiledTasks = GraphCompiler->GetCompiledTasks(Graph, *StackContextPtr);
+
+	// Create the final stack context by including the current stack frames
+	if (InFromStack)
+	{
+		StackContextPtr->PrependParentStack(InFromStack);
+	}
 
 	// Assign this component to the tasks
 	for (FPCGGraphTask& Task : CompiledTasks)
 	{
 		Task.SourceComponent = SourceComponent;
+		Task.StackContext = StackContextPtr;
 	}
 
 	// Prepare scheduled task that will be promoted in the next Execute call.
@@ -648,6 +657,7 @@ void FPCGGraphExecutor::Execute()
 					Task.Context->CompiledTaskId = Task.CompiledTaskId;
 					Task.Context->InputData.Crc = TaskInput.Crc;
 					Task.Context->DependenciesCrc = DependenciesCrc;
+					Task.Context->Stack = (Task.StackContext && Task.StackIndex != INDEX_NONE) ? Task.StackContext->GetStack(Task.StackIndex) : nullptr;
 				}
 
 				// Validate that we can start this task now
@@ -659,6 +669,8 @@ void FPCGGraphExecutor::Execute()
 					ActiveTask.Element = Task.Element;
 					ActiveTask.NodeId = Task.NodeId;
 					ActiveTask.Context = TUniquePtr<FPCGContext>(Task.Context);
+					ActiveTask.StackIndex = Task.StackIndex;
+					ActiveTask.StackContext = Task.StackContext;
 
 #if WITH_EDITOR
 					if (bResultAlreadyInCache)
@@ -1195,7 +1207,9 @@ FPCGElementPtr FPCGGraphExecutor::GetFetchInputElement()
 FPCGTaskId FPCGGraphExecutor::ScheduleDebugWithTaskCallback(UPCGComponent* InComponent, TFunction<void(FPCGTaskId/* TaskId*/, const UPCGNode*/* Node*/, const FPCGDataCollection&/* TaskOutput*/)> TaskCompleteCallback)
 {
 	FPCGTaskId FinalTaskID = Schedule(InComponent, {});
-	TArray<FPCGGraphTask> CompiledTasks = GraphCompiler->GetCompiledTasks(InComponent->GetGraph(), /*bIsTopGraph=*/true);
+
+	FPCGStackContext DummyStackContext;
+	TArray<FPCGGraphTask> CompiledTasks = GraphCompiler->GetCompiledTasks(InComponent->GetGraph(), DummyStackContext, /*bIsTopGraph=*/true);
 	CompiledTasks.Pop(); // Remove the final task
 
 	// Set up all final dependencies for the entire execution
diff --git a/Engine/Plugins/Experimental/PCG/Source/PCG/Private/Graph/PCGGraphExecutor.h b/Engine/Plugins/Experimental/PCG/Source/PCG/Private/Graph/PCGGraphExecutor.h
--- a/Engine/Plugins/Experimental/PCG/Source/PCG/Private/Graph/PCGGraphExecutor.h
+++ b/Engine/Plugins/Experimental/PCG/Source/PCG/Private/Graph/PCGGraphExecutor.h
@@ -2,6 +2,7 @@
 
 #pragma once
 
+#include "PCGContext.h"
 #include "PCGElement.h"
 #include "PCGSubsystem.h"
 #include "Graph/PCGGraphCache.h"
@@ -16,6 +17,8 @@ class UPCGGraph;
 class UPCGNode;
 class UPCGComponent;
 class FPCGGraphCompiler;
+class FPCGStack;
+class FPCGStackContext;
 class FTextFormat;
 
 struct FPCGGraphTaskInput
@@ -51,6 +54,9 @@ struct FPCGGraphTask
 	FPCGTaskId NodeId = InvalidPCGTaskId;
 	FPCGTaskId CompiledTaskId = InvalidPCGTaskId; // the task id as it exists when compiled
 	FPCGTaskId ParentId = InvalidPCGTaskId; // represents the parent sub object graph task, if we were called from one
+
+	int32 StackIndex = INDEX_NONE;
+	TSharedPtr<const FPCGStackContext> StackContext;
 };
 
 struct FPCGGraphScheduleTask
@@ -70,6 +76,8 @@ struct FPCGGraphActiveTask
 #if WITH_EDITOR
 	bool bIsBypassed = false;
 #endif
+	int32 StackIndex = INDEX_NONE;
+	TSharedPtr<const FPCGStackContext> StackContext;
 };
 
 class FPCGGraphExecutor
@@ -82,8 +90,8 @@ public:
 	void Compile(UPCGGraph* InGraph);
 
 	/** Schedules the execution of a given graph with specified inputs. This call is threadsafe */
-	FPCGTaskId Schedule(UPCGComponent* InComponent, const TArray<FPCGTaskId>& TaskDependency = TArray<FPCGTaskId>());
-	FPCGTaskId Schedule(UPCGGraph* Graph, UPCGComponent* InSourceComponent, FPCGElementPtr PreGraphElement, FPCGElementPtr InputElement, const TArray<FPCGTaskId>& TaskDependency);
+	FPCGTaskId Schedule(UPCGComponent* InComponent, const TArray<FPCGTaskId>& TaskDependency = TArray<FPCGTaskId>(), const FPCGStack* InFromStack = nullptr);
+	FPCGTaskId Schedule(UPCGGraph* Graph, UPCGComponent* InSourceComponent, FPCGElementPtr PreGraphElement, FPCGElementPtr InputElement, const TArray<FPCGTaskId>& TaskDependency, const FPCGStack* InFromStack);
 
 	/** Cancels all tasks originating from the given component */
 	TArray<UPCGComponent*> Cancel(UPCGComponent* InComponent);
diff --git a/Engine/Plugins/Experimental/PCG/Source/PCG/Private/Graph/PCGStackContext.cpp b/Engine/Plugins/Experimental/PCG/Source/PCG/Private/Graph/PCGStackContext.cpp
new file mode 100644
--- /dev/null
+++ b/Engine/Plugins/Experimental/PCG/Source/PCG/Private/Graph/PCGStackContext.cpp
@@ -0,0 +1,201 @@
+// Copyright Epic Games, Inc. All Rights Reserved.
+
+#include "Graph/PCGStackContext.h"
+
+#include "PCGGraph.h"
+#include "PCGNode.h"
+#include "PCGPin.h"
+
+#include "Containers/UnrealString.h"
+#include "Misc/StringBuilder.h"
+
+const FPCGStack* FPCGStackContext::GetStack(int32 InStackIndex) const
+{
+	if (ensure(Stacks.IsValidIndex(InStackIndex)))
+	{
+		return &Stacks[InStackIndex];
+	}
+	else
+	{
+		return nullptr;
+	}
+}
+
+void FPCGStack::PopFrame()
+{
+	if (ensure(!StackFrames.IsEmpty()))
+	{
+		StackFrames.Pop();
+	}
+}
+
+int32 FPCGStackContext::PushFrame(const UObject* InFrameObject)
+{
+	if (CurrentStackIndex == INDEX_NONE)
+	{
+		// Create first stack using the given frame.
+		FPCGStack& Stack = Stacks.Emplace_GetRef();
+		Stack.PushFrame(FPCGStackFrame(InFrameObject));
+		CurrentStackIndex = 0;
+	}
+	else
+	{
+		if (!ensure(Stacks.IsValidIndex(CurrentStackIndex)))
+		{
+			return INDEX_NONE;
+		}
+
+		// Append given frame object to current stack. Newly encountered stacks should generally
+		// be unique, so we just commit to creating it immediately rather than searching to see
+		// if it already exists first.
+		FPCGStack CurrentStack = Stacks[CurrentStackIndex];
+		CurrentStack.PushFrame(FPCGStackFrame(InFrameObject));
+		CurrentStackIndex = Stacks.AddUnique(MoveTemp(CurrentStack));
+	}
+
+	return CurrentStackIndex;
+}
+
+int32 FPCGStackContext::PopFrame()
+{
+	if (!ensure(Stacks.IsValidIndex(CurrentStackIndex)))
+	{
+		return INDEX_NONE;
+	}
+
+	// Find the 'parent' callstack (current stack minus latest frame). Can be anywhere in the list of stacks so do a search.
+	CurrentStackIndex = Stacks.IndexOfByPredicate([&CurrentStack = Stacks[CurrentStackIndex]](const FPCGStack& OtherStack)
+	{
+		const int32 RequiredFrameCount = CurrentStack.StackFrames.Num() - 1;
+		if (OtherStack.StackFrames.Num() != RequiredFrameCount)
+		{
+			return false;
+		}
+		
+		for (int32 i = 0; i < RequiredFrameCount; ++i)
+		{
+			if (OtherStack.StackFrames[i] != CurrentStack.StackFrames[i])
+			{
+				return false;
+			}
+		}
+
+		return true;
+	});
+	ensure(CurrentStackIndex != INDEX_NONE);
+
+	return CurrentStackIndex;
+}
+
+void FPCGStackContext::AppendStacks(const FPCGStackContext& InStacks)
+{
+	if (!ensure(Stacks.IsValidIndex(CurrentStackIndex)))
+	{
+		return;
+	}
+
+	for (const FPCGStack& SubgraphStack : InStacks.Stacks)
+	{
+		FPCGStack& NewStack = Stacks.Emplace_GetRef();
+		NewStack.StackFrames.Reserve(Stacks[CurrentStackIndex].StackFrames.Num() + SubgraphStack.StackFrames.Num());
+		
+		NewStack.StackFrames.Append(Stacks[CurrentStackIndex].StackFrames);
+		NewStack.StackFrames.Append(SubgraphStack.StackFrames);
+	}
+}
+
+void FPCGStackContext::PrependParentStack(const FPCGStack* InParentStack)
+{
+	if (!InParentStack || InParentStack->StackFrames.IsEmpty())
+	{
+		return;
+	}
+
+	for (FPCGStack& Stack : Stacks)
+	{
+		Stack.StackFrames.Insert(InParentStack->StackFrames, 0);
+	}
+}
+
+bool FPCGStack::CreateStackFramePath(FString& OutString, const UPCGNode* InNode, const UPCGPin* InPin) const
+{
+	// Give a healthy amount of scratch space on the stack, if it overflows it will use heap.
+	TStringBuilderWithBuffer<TCHAR, 2048> StringBuilder;
+
+	auto AddPathSeparator = [&StringBuilder]()
+	{
+		if (StringBuilder.Len())
+		{
+			StringBuilder << TEXT("/");
+		}
+	};
+
+	for (const FPCGStackFrame& Frame : StackFrames)
+	{
+		if (Frame.Object.IsValid())
+		{
+			AddPathSeparator();
+			const UObject* Object = Frame.Object.Get();
+
+			if (!Object)
+			{
+				// If any object does not resolve, cannot build the string
+				return false;
+			}
+
+			if (Object->IsA<UPCGGraph>())
+			{
+				StringBuilder << TEXT("GRAPH:") << Object->GetFullName();
+			}
+			else if (Object->IsA<UPCGNode>())
+			{
+				StringBuilder << TEXT("NODE:") << Object->GetFName();
+			}
+			else
+			{
+				// Unrecognized type, should not happen
+				ensure(false);
+				StringBuilder << TEXT("UNRECOGNIZED:") << Object->GetFullName();
+			}
+		}
+		else if (Frame.LoopIndex != -1)
+		{
+			AddPathSeparator();
+			StringBuilder << TEXT("LOOP:") << FString::FromInt(Frame.LoopIndex);
+		}
+	}
+
+	if (InNode)
+	{
+		AddPathSeparator();
+		StringBuilder << TEXT("NODE:") << InNode->GetFName().ToString();
+
+		if (InPin)
+		{
+			AddPathSeparator();
+			StringBuilder << TEXT("PIN:") << InPin->GetFName().ToString();
+		}
+	}
+
+	OutString = StringBuilder;
+	return true;
+}
+
+bool FPCGStack::operator==(const FPCGStack& Other) const
+{
+	// Stacks are the same if all stack frames are the same
+	if (StackFrames.Num() != Other.StackFrames.Num())
+	{
+		return false;
+	}
+
+	for (int32 i = 0; i < StackFrames.Num(); i++)
+	{
+		if (StackFrames[i] != Other.StackFrames[i])
+		{
+			return false;
+		}
+	}
+
+	return true;
+}
diff --git a/Engine/Plugins/Experimental/PCG/Source/PCG/Private/PCGSubgraph.cpp b/Engine/Plugins/Experimental/PCG/Source/PCG/Private/PCGSubgraph.cpp
--- a/Engine/Plugins/Experimental/PCG/Source/PCG/Private/PCGSubgraph.cpp
+++ b/Engine/Plugins/Experimental/PCG/Source/PCG/Private/PCGSubgraph.cpp
@@ -6,6 +6,7 @@
 #include "PCGPin.h"
 #include "PCGSubsystem.h"
 #include "Data/PCGUserParametersData.h"
+#include "Graph/PCGStackContext.h"
 #include "Helpers/PCGSettingsHelpers.h"
 
 #include "Algo/Find.h"
@@ -409,7 +410,7 @@ bool FPCGSubgraphElement::ExecuteInternal(FPCGContext* InContext) const
 	const UPCGSubgraphSettings* Settings = Context->GetInputSettings<UPCGSubgraphSettings>();
 	check(Settings);
 
-	if(Settings->IsDynamicGraph())
+	if (Settings->IsDynamicGraph())
 	{
 		TRACE_CPUPROFILER_EVENT_SCOPE(FPCGSubgraphElement::Execute);
 
@@ -428,7 +429,11 @@ bool FPCGSubgraphElement::ExecuteInternal(FPCGContext* InContext) const
 				FPCGDataCollection SubgraphInputData;
 				PrepareSubgraphData(Settings, Context, Context->InputData, SubgraphInputData);
 
-				FPCGTaskId SubgraphTaskId = Subsystem->ScheduleGraph(Subgraph, Context->SourceComponent.Get(), MakeShared<FPCGInputForwardingElement>(PreSubgraphInputData), MakeShared<FPCGInputForwardingElement>(SubgraphInputData), {});
+				// Prepare the invocation stack - which is the stack up to this node, and then this node
+				FPCGStack InvocationStack = ensure(Context->Stack) ? *Context->Stack : FPCGStack();
+				InvocationStack.GetStackFramesMutable().Emplace(Context->Node);
+
+				FPCGTaskId SubgraphTaskId = Subsystem->ScheduleGraph(Subgraph, Context->SourceComponent.Get(), MakeShared<FPCGInputForwardingElement>(PreSubgraphInputData), MakeShared<FPCGInputForwardingElement>(SubgraphInputData), {}, &InvocationStack);
 
 				if (SubgraphTaskId != InvalidPCGTaskId)
 				{
diff --git a/Engine/Plugins/Experimental/PCG/Source/PCG/Private/PCGSubsystem.cpp b/Engine/Plugins/Experimental/PCG/Source/PCG/Private/PCGSubsystem.cpp
--- a/Engine/Plugins/Experimental/PCG/Source/PCG/Private/PCGSubsystem.cpp
+++ b/Engine/Plugins/Experimental/PCG/Source/PCG/Private/PCGSubsystem.cpp
@@ -7,6 +7,7 @@
 #include "PCGWorldActor.h"
 #include "Graph/PCGGraphCache.h"
 #include "Graph/PCGGraphExecutor.h"
+#include "Graph/PCGStackContext.h"
 #include "Grid/PCGLandscapeCache.h"
 #include "Grid/PCGPartitionActor.h"
 #include "Helpers/PCGActorHelpers.h"
@@ -440,11 +441,11 @@ FPCGTaskId UPCGSubsystem::ScheduleGraph(UPCGComponent* SourceComponent, const TA
 	}
 }
 
-FPCGTaskId UPCGSubsystem::ScheduleGraph(UPCGGraph* Graph, UPCGComponent* SourceComponent, FPCGElementPtr PreGraphElement, FPCGElementPtr InputElement, const TArray<FPCGTaskId>& Dependencies)
+FPCGTaskId UPCGSubsystem::ScheduleGraph(UPCGGraph* Graph, UPCGComponent* SourceComponent, FPCGElementPtr PreGraphElement, FPCGElementPtr InputElement, const TArray<FPCGTaskId>& Dependencies, const FPCGStack* InFromStack)
 {
 	if (SourceComponent)
 	{
-		return GraphExecutor->Schedule(Graph, SourceComponent, PreGraphElement, InputElement, Dependencies);
+		return GraphExecutor->Schedule(Graph, SourceComponent, PreGraphElement, InputElement, Dependencies, InFromStack);
 	}
 	else
 	{
diff --git a/Engine/Plugins/Experimental/PCG/Source/PCG/Private/Utils/PCGExtraCapture.cpp b/Engine/Plugins/Experimental/PCG/Source/PCG/Private/Utils/PCGExtraCapture.cpp
--- a/Engine/Plugins/Experimental/PCG/Source/PCG/Private/Utils/PCGExtraCapture.cpp
+++ b/Engine/Plugins/Experimental/PCG/Source/PCG/Private/Utils/PCGExtraCapture.cpp
@@ -169,7 +169,8 @@ PCGUtils::FCallTreeInfo PCGUtils::FExtraCapture::CalculateCallTreeInfo(const UPC
 		return {};
 	}
 
-	TArray<FPCGGraphTask> CompiledTasks = Compiler->GetPrecompiledTasks(Component->GetGraph());
+	FPCGStackContext DummyStackContext;
+	TArray<FPCGGraphTask> CompiledTasks = Compiler->GetPrecompiledTasks(Component->GetGraph(), DummyStackContext);
 	if (CompiledTasks.IsEmpty())
 	{
 		return {};
diff --git a/Engine/Plugins/Experimental/PCG/Source/PCG/Public/Elements/PCGLoopElement.h b/Engine/Plugins/Experimental/PCG/Source/PCG/Public/Elements/PCGLoopElement.h
--- a/Engine/Plugins/Experimental/PCG/Source/PCG/Public/Elements/PCGLoopElement.h
+++ b/Engine/Plugins/Experimental/PCG/Source/PCG/Public/Elements/PCGLoopElement.h
@@ -27,7 +27,7 @@ protected:
 
 public:
 	//~Begin UPCGBaseSubgraphSettings interface
-	virtual bool IsDynamicGraph() const { return true; }
+	virtual bool IsDynamicGraph() const override { return true; }
 	//~End UPCGBaseSubgraphSettings interface
 
 	/** 
@@ -43,4 +43,4 @@ class PCG_API FPCGLoopElement : public FPCGSubgraphElement
 protected:
 	virtual bool ExecuteInternal(FPCGContext* Context) const override;
 	void PrepareLoopDataCollections(FPCGContext* Context, const UPCGLoopSettings* Settings, TArray<FPCGDataCollection>& LoopDataCollection, FPCGDataCollection& FixedInputDataCollection) const;
-};
\ No newline at end of file
+};
diff --git a/Engine/Plugins/Experimental/PCG/Source/PCG/Public/Graph/PCGStackContext.h b/Engine/Plugins/Experimental/PCG/Source/PCG/Public/Graph/PCGStackContext.h
new file mode 100644
--- /dev/null
+++ b/Engine/Plugins/Experimental/PCG/Source/PCG/Public/Graph/PCGStackContext.h
@@ -0,0 +1,84 @@
+// Copyright Epic Games, Inc. All Rights Reserved.
+
+#pragma once
+
+#include "Containers/Array.h"
+#include "UObject/WeakObjectPtr.h"
+
+class FString;
+class UObject;
+class UPCGGraph;
+class UPCGNode;
+class UPCGPin;
+
+/** A single frame of a call stack, represented as a pointer to the associated object (graph/subgraph or node) or a loop index. */
+struct PCG_API FPCGStackFrame
+{
+	explicit FPCGStackFrame(TWeakObjectPtr<const UObject> InObject)
+		: Object(InObject)
+	{}
+
+	explicit FPCGStackFrame(int32 InLoopIndex)
+		: LoopIndex(InLoopIndex)
+	{}
+
+	bool operator==(const FPCGStackFrame& Other) const { return Object == Other.Object && LoopIndex == Other.LoopIndex; }
+	bool operator!=(const FPCGStackFrame& Other) const { return !(*this == Other); };
+
+	// A valid frame should either point to an object or have a loop index >= 0.
+	TWeakObjectPtr<const UObject> Object;
+	int32 LoopIndex = INDEX_NONE;
+};
+
+/** A call stack, represented as an array of stack frames. */
+class PCG_API FPCGStack
+{
+	friend class FPCGStackContext;
+
+public:
+	/** Push frame onto top of stack. */
+	void PushFrame(const FPCGStackFrame& InFrame) { StackFrames.Add(InFrame); }
+
+	/** Pop frame from the stack. */
+	void PopFrame();
+
+	/** Construct a string version of this stack. Postfixed by optional node/pin if provided. */
+	bool CreateStackFramePath(FString& OutString, const UPCGNode* InNode = nullptr, const UPCGPin* InPin = nullptr) const;
+
+	const TArray<FPCGStackFrame>& GetStackFrames() const { return StackFrames; }
+	TArray<FPCGStackFrame>& GetStackFramesMutable() { return StackFrames; }
+
+	bool operator==(const FPCGStack& Other) const;
+	bool operator!=(const FPCGStack& Other) const { return !(*this == Other); }
+
+private:
+	TArray<FPCGStackFrame> StackFrames;
+};
+
+/** A collection of call stacks. */
+class PCG_API FPCGStackContext
+{
+public:
+	int32 GetNumStacks() const { return Stacks.Num(); }
+	int32 GetCurrentStackIndex() const { return CurrentStackIndex; }
+	const FPCGStack* GetStack(int32 InStackIndex) const;
+
+	/** Create a new stack and create a frame from the provided object (typically graph or node pointer). Returns index of newly added stack. */
+	int32 PushFrame(const UObject* InFrameObject);
+
+	/** Remove a frame from the current stack. Returns current stack index. */
+	int32 PopFrame();
+
+	/** Takes the current stack and appends each of the stacks in InStacks. Called during compilation when inlining a static subgraph. */
+	void AppendStacks(const FPCGStackContext& InStacks);
+
+	/** Called during execution when invoking a dynamic subgraph, to prepend the caller stack to form the complete callstacks. */
+	void PrependParentStack(const FPCGStack* InParentStack);
+
+private:
+	/** List of all stacks encountered top graph and all (nested) subgraphs. Order is simply order of encountering during compilation. */
+	TArray<FPCGStack> Stacks;
+
+	/** Index of element in Stacks that is the current stack. */
+	int32 CurrentStackIndex = INDEX_NONE;
+};
diff --git a/Engine/Plugins/Experimental/PCG/Source/PCG/Public/PCGContext.h b/Engine/Plugins/Experimental/PCG/Source/PCG/Public/PCGContext.h
--- a/Engine/Plugins/Experimental/PCG/Source/PCG/Public/PCGContext.h
+++ b/Engine/Plugins/Experimental/PCG/Source/PCG/Public/PCGContext.h
@@ -14,6 +14,7 @@ class UPCGSettingsInterface;
 class UPCGSpatialData;
 struct FPCGGraphCache;
 struct FPCGSettingsOverridableParam;
+class FPCGStack;
 
 namespace PCGContextHelpers
 {
@@ -61,12 +62,15 @@ struct PCG_API FPCGContext
 	FPCGTaskId CompiledTaskId = InvalidPCGTaskId;
 	bool bIsPaused = false;
 
-	// Used to preven settings override being deleted, needs to be false when going through blueprint calls with a context
+	// Used to prevent settings override being deleted, needs to be false when going through blueprint calls with a context
 	bool bShouldUnrootSettingsOnDelete = true;
 
 	EPCGExecutionPhase CurrentPhase = EPCGExecutionPhase::NotExecuted;
 	int32 BypassedOutputCount = 0;
 
+	/** The current call stack. */
+	const FPCGStack* Stack = nullptr;
+
 	const UPCGSettingsInterface* GetInputSettingsInterface() const;
 	
 	// After initializing the context, we can call this method to prepare for parameter override
diff --git a/Engine/Plugins/Experimental/PCG/Source/PCG/Public/PCGSubgraph.h b/Engine/Plugins/Experimental/PCG/Source/PCG/Public/PCGSubgraph.h
--- a/Engine/Plugins/Experimental/PCG/Source/PCG/Public/PCGSubgraph.h
+++ b/Engine/Plugins/Experimental/PCG/Source/PCG/Public/PCGSubgraph.h
@@ -22,6 +22,8 @@ class PCG_API UPCGBaseSubgraphSettings : public UPCGSettings
 public:
 	UPCGGraph* GetSubgraph() const;
 	virtual UPCGGraphInterface* GetSubgraphInterface() const { return nullptr; }
+
+	/** Returns true if the subgraphs nodes were not inlined into the parent graphs tasks during compilation. */
 	virtual bool IsDynamicGraph() const { return false; }
 
 	// Use this method from the outside to set the subgraph, as it will connect editor callbacks
diff --git a/Engine/Plugins/Experimental/PCG/Source/PCG/Public/PCGSubsystem.h b/Engine/Plugins/Experimental/PCG/Source/PCG/Public/PCGSubsystem.h
--- a/Engine/Plugins/Experimental/PCG/Source/PCG/Public/PCGSubsystem.h
+++ b/Engine/Plugins/Experimental/PCG/Source/PCG/Public/PCGSubsystem.h
@@ -24,6 +24,7 @@ class FPCGGraphCompiler;
 class FPCGGraphExecutor;
 struct FPCGContext;
 struct FPCGDataCollection;
+class FPCGStack;
 class UPCGSettings;
 
 class IPCGElement;
@@ -87,7 +88,7 @@ public:
 	FPCGTaskId ScheduleCleanup(UPCGComponent* PCGComponent, bool bRemoveComponents, bool bSave, const TArray<FPCGTaskId>& Dependencies);
 
 	// Schedule graph (used internally for dynamic subgraph execution)
-	FPCGTaskId ScheduleGraph(UPCGGraph* Graph, UPCGComponent* SourceComponent, FPCGElementPtr PreGraphElement, FPCGElementPtr InputElement, const TArray<FPCGTaskId>& Dependencies);
+	FPCGTaskId ScheduleGraph(UPCGGraph* Graph, UPCGComponent* SourceComponent, FPCGElementPtr PreGraphElement, FPCGElementPtr InputElement, const TArray<FPCGTaskId>& Dependencies, const FPCGStack* InFromStack);
 
 	// Schedule graph (used internally for dynamic subgraph execution)
 	FPCGTaskId ScheduleGraph(UPCGComponent* SourceComponent, const TArray<FPCGTaskId>& Dependencies);
