Added PCG Noise node

diff --git a/Engine/Plugins/Experimental/PCG/Source/PCG/Private/Elements/PCGNoise.cpp b/Engine/Plugins/Experimental/PCG/Source/PCG/Private/Elements/PCGNoise.cpp
new file mode 100644
--- /dev/null
+++ b/Engine/Plugins/Experimental/PCG/Source/PCG/Private/Elements/PCGNoise.cpp
@@ -0,0 +1,751 @@
+// Copyright Epic Games, Inc. All Rights Reserved.
+
+#include "Elements/PCGNoise.h"
+
+#include "Kismet/KismetMathLibrary.h"
+
+#include "PCGComponent.h"
+#include "PCGContext.h"
+#include "PCGPin.h"
+
+#include "Data/PCGPointData.h"
+#include "Data/PCGSpatialData.h"
+#include "Helpers/PCGAsync.h"
+#include "Helpers/PCGHelpers.h"
+#include "Helpers/PCGSettingsHelpers.h"
+#include "Metadata/Accessors/IPCGAttributeAccessor.h"
+#include "Metadata/Accessors/PCGAttributeAccessorHelpers.h"
+
+#include UE_INLINE_GENERATED_CPP_BY_NAME(PCGNoise)
+
+#define LOCTEXT_NAMESPACE "PCGNoise"
+
+TArray<FPCGPinProperties> UPCGNoiseSettings::InputPinProperties() const
+{
+	return Super::DefaultPointInputPinProperties();
+}
+
+TArray<FPCGPinProperties> UPCGNoiseSettings::OutputPinProperties() const
+{
+	return Super::DefaultPointOutputPinProperties();
+}
+
+FPCGElementPtr UPCGNoiseSettings::CreateElement() const
+{
+	return MakeShared<FPCGNoise>();
+}
+
+namespace PCGNoise
+{
+	// the useful ranges here are very small so make input easier
+	static const double MAGIC_SCALE_FACTOR = 0.0001;
+
+	inline double Fract(double X)
+	{
+		return X - FMath::Floor(X);
+	}
+
+	inline FVector2D Floor(FVector2D X)
+	{
+		return FVector2D(FMath::Floor(X.X), FMath::Floor(X.Y));
+	}
+
+	inline FInt32Point FloorInt(FVector2D X)
+	{
+		return FInt32Point(FMath::FloorToInt(X.X), FMath::FloorToInt(X.Y));
+	}
+
+	inline FVector2D ToVec2D(FInt32Point V)
+	{
+		return FVector2D(double(V.X), double(V.Y));
+	}
+
+	inline FVector2D Fract(FVector2D V)
+	{
+		return FVector2D(Fract(V.X), Fract(V.Y));
+	}
+
+	inline double ValueHash(FVector2D Position)
+	{
+		Position = 50.0 * Fract(Position * 0.3183099 + FVector2D(0.71, 0.113));
+		return -1.0 + 2.0 * Fract(Position.X * Position.Y * (Position.X + Position.Y));
+	}
+
+	inline double Noise2D(FVector2D Position)
+	{
+		const FVector2D FloorPosition = Floor(Position);
+		const FVector2D Fraction = Position - FloorPosition;
+		const FVector2D U = Fraction * Fraction * (FVector2D(3.0, 3.0) - 2.0 * Fraction);
+
+		return FMath::Lerp(
+			FMath::Lerp(ValueHash(FloorPosition), ValueHash(FloorPosition + FVector2D(1.0, 0.0)), U.X),
+			FMath::Lerp(ValueHash(FloorPosition + FVector2D(0.0, 1.0)), ValueHash(FloorPosition + FVector2D(1.0, 1.0)), U.X),
+			U.Y
+		);
+	}
+
+	inline FVector2D MultiplyMatrix2D(FVector2D Point, const FVector2D (&Mat2)[2])
+	{
+		return Point.X * Mat2[0] + Point.Y * Mat2[1];
+	}
+
+	// takes a cell id and produces a pseudo random vector offset
+	inline FVector2D VoronoiHash2D(FInt32Point Cell)
+	{
+		const FVector2D P = ToVec2D(Cell);
+
+		// this is some arbitrary large random scale+rotation+skew
+		const FVector2D P2 = FVector2D(
+			P.Dot(FVector2D(127.1, 311.7)),
+			P.Dot(FVector2D(269.5, 183.3))
+		);
+
+		// further scale the results by a big number
+		return FVector2D(
+			Fract(FMath::Sin(P2.X) * 17.1717) - 0.5,
+			Fract(FMath::Sin(P2.Y) * 17.1717) - 0.5
+		);
+	}
+
+	struct FVoronoi2DResult
+	{
+		double DistanceToEdge = 8.0;
+		double ID = -1.0;
+
+		FVector2D CellEdgeDirection = FVector2D::ZeroVector;
+	};
+
+	template<typename Func>
+	FVoronoi2DResult CalcVoronoi2D(const FVector2D Position, Func&& HashFunc)
+	{
+		const FInt32Point WorldIdx = FloorInt(Position);
+
+		FVector2D CellHashValue = FVector2D::ZeroVector;
+		FVector2D CellPosition = FVector2D::ZeroVector;
+		FInt32Point CellIdx = FInt32Point(0,0);
+
+		// first pass, find our closest cell point (which is all voronoi really is)
+		{
+			double MinDistanceSq = TNumericLimits<double>::Max();
+			for (int32 Y = WorldIdx.Y-1; Y <= WorldIdx.Y+1; ++Y)
+			{
+				for (int32 X = WorldIdx.X-1; X <= WorldIdx.X+1; ++X)
+				{
+					const FInt32Point ThisCellIdx(X, Y);
+					const FVector2D ThisCellHashValue = HashFunc(ThisCellIdx);
+					const FVector2D ThisCellPosition = ToVec2D(ThisCellIdx) + ThisCellHashValue;
+
+					const double CellDistSquared = FVector2D::DistSquared(ThisCellPosition, Position);
+					if (CellDistSquared > MinDistanceSq)
+					{
+						continue;
+					}
+
+					MinDistanceSq = CellDistSquared;
+					CellHashValue = ThisCellHashValue;
+					CellPosition = ThisCellPosition;
+					CellIdx = ThisCellIdx;
+				}
+			}
+		}
+
+		// second pass, compute distance from Position to the nearest cell edge
+		FVoronoi2DResult Result;
+		Result.ID =  0.5+0.5*FMath::Cos((CellHashValue.X+CellHashValue.Y)*6.2831);
+		Result.DistanceToEdge = TNumericLimits<double>::Max();
+
+		for (int32 Y = CellIdx.Y-2; Y <= CellIdx.Y+2; ++Y)
+		{
+			for (int32 X = CellIdx.X-2; X <= CellIdx.X+2; ++X)
+			{
+				const FInt32Point ThisCellIdx(X, Y);
+
+				if (ThisCellIdx == CellIdx)
+				{
+					continue;
+				}
+
+				const FVector2D ThisCellPosition = ToVec2D(ThisCellIdx) + HashFunc(ThisCellIdx);
+				const FVector2D EdgeDirection = ThisCellPosition - CellPosition;
+
+				if (EdgeDirection.IsNearlyZero())
+				{
+					continue;
+				}
+
+				const FVector2D PlaneNormal = EdgeDirection.GetSafeNormal();
+				const FVector2D PointOnPlane = CellPosition + EdgeDirection * 0.5;
+
+				const double DistanceToEdge = FMath::Abs((Position-PointOnPlane).Dot(PlaneNormal));
+				if (DistanceToEdge > Result.DistanceToEdge)
+				{
+					continue;
+				}
+
+				Result.DistanceToEdge = DistanceToEdge;
+				Result.CellEdgeDirection = PlaneNormal;
+			}
+		}
+
+		return Result;
+	}
+
+	double CalcFractionalBrownian2D(FVector2D Position, int32 Iterations)
+	{
+		double Z = 0.5;
+		double Result = 0.0;
+
+		// just some fixed random rotation and scale numbers
+		const FVector2D RotScale[] = {{1.910673, -0.5910404}, {0.5910404, 1.910673}};
+
+		for (int32 I = 0; I < Iterations; ++I)
+		{
+			Result += FMath::Abs(Noise2D(Position)) * Z;
+			Z *= 0.5;
+			Position = MultiplyMatrix2D(Position, RotScale);
+		}
+
+		return Result;
+	}
+
+	double CalcCaustic2D(FVector2D Position, int32 Iterations)
+	{
+		const FVector2D P = Fract(Position * 0.2) * (PI * 2.0) - FVector2D(250.0, 250.0);
+		FVector2D I = P;
+		double Value = 0.0;
+
+		for (int32 N = 0; N < Iterations; ++N)
+		{
+			const double T = 1.0 - (3.5 / double(N+1));
+			I = P + FVector2D(FMath::Cos(T - I.X) + FMath::Sin(T + I.Y), FMath::Sin(T - I.Y) + FMath::Cos(T + I.X));
+
+			const double S = FMath::Sin(I.X + T);
+			const double C = FMath::Cos(I.Y + T);
+
+			const FVector2D Temp = FVector2D(
+				FMath::Abs(S) > 0.0001 ? (P.X / S) : 1000.0,
+				FMath::Abs(C) > 0.0001 ? (P.Y / C) : 1000.0
+			);
+
+			Value += FMath::InvSqrt(Temp.SquaredLength());
+		}
+
+		return Value / (0.003 * double(Iterations));
+	}
+
+	const FVector2D PerlinM[] = {{1.6, 1.2}, {-1.2, 1.6}};
+
+	double CalcPerlin2D(FVector2D Position, int32 Iterations)
+	{
+		double Value = 0.0;
+
+		double Strength = 1.0;
+
+		for (int32 N = 0; N < Iterations; ++N)
+		{
+			Strength *= 0.5;
+			Value += Strength * Noise2D(Position);
+			Position = MultiplyMatrix2D(Position, PerlinM);
+		}
+
+		return 0.5 + 0.5 * Value;		
+	}
+
+	double ApplyContrast(double Value, double Contrast)
+	{
+		// early out for default 1.0 contrast, the math should be the same
+		if (Contrast == 1.0)
+		{
+			return Value;
+		}
+
+		if (Contrast <= 0.0)
+		{
+			return 0.5;
+		}
+
+		Value = FMath::Clamp(Value, 0.0, 1.0);
+
+		if (Value == 1.0)
+		{
+			return 1.0;
+		}
+
+		return 1.0 / (1.0 + FMath::Pow(Value / (1.0 - Value), -Contrast));
+	}
+
+	FLocalCoordinates2D CalcLocalCoordinates2D(const FBox& ActorLocalBox, const FTransform& ActorTransformInverse, FVector2D Scale, const FVector& InPosition)
+	{
+		if (!ActorLocalBox.IsValid)
+		{
+			return {};
+		}
+
+		const FVector2D LocalPosition = FVector2D(ActorTransformInverse.TransformPosition(InPosition));
+
+		const double LeftDist = LocalPosition.X-ActorLocalBox.Min.X;
+		const double RightDist = LocalPosition.X-ActorLocalBox.Max.X;
+
+		const double TopDist = LocalPosition.Y-ActorLocalBox.Min.Y;
+		const double BottomDist = LocalPosition.Y-ActorLocalBox.Max.Y;
+
+		FLocalCoordinates2D Result;
+
+		Result.X0 = LeftDist * Scale.X;
+		Result.X1 = RightDist * Scale.X;
+		Result.Y0 = TopDist * Scale.Y;
+		Result.Y1 = BottomDist * Scale.Y;
+
+		Result.FracX = FMath::Clamp(LeftDist / (ActorLocalBox.Max.X - ActorLocalBox.Min.X), 0.0, 1.0);
+		Result.FracY = FMath::Clamp(TopDist / (ActorLocalBox.Max.Y - ActorLocalBox.Min.Y), 0.0, 1.0);
+
+		return Result;
+	};
+
+	double CalcEdgeBlendAmount2D(const FLocalCoordinates2D& LocalCoords, double EdgeBlendDistance)
+	{
+		const double UseX = FMath::Abs(LocalCoords.X0) < FMath::Abs(LocalCoords.X1) ? LocalCoords.X0 : LocalCoords.X1;
+		const double UseY = FMath::Abs(LocalCoords.Y0) < FMath::Abs(LocalCoords.Y1) ? LocalCoords.Y0 : LocalCoords.Y1;
+
+		const double CurrentEdgeAmount = FMath::Min(FMath::Abs(UseX), FMath::Abs(UseY));
+
+		return EdgeBlendDistance > CurrentEdgeAmount ? FMath::Clamp((EdgeBlendDistance-CurrentEdgeAmount)/EdgeBlendDistance, 0.0, 1.0) : 0.0;		
+	}
+
+	struct FSharedParams
+	{
+		FPCGContext* Context = nullptr;
+		FTransform ActorTransformInverse;
+		FBox ActorLocalBox;
+		FTransform Transform;
+
+		double Brightness;
+		double Contrast;
+		int32 Iterations;
+		bool bTiling;
+	};
+
+	struct FBufferParams
+	{
+		const UPCGPointData* InputPointData = nullptr;
+		UPCGPointData* OutputPointData = nullptr;
+	};
+
+	template<typename T>
+	void SetAttributeHelper(UPCGSpatialData* Data, const FPCGAttributePropertySelector& PropertySelector, const TArrayView<const T> Values)
+	{
+		if (PropertySelector.Selection == EPCGAttributePropertySelection::Attribute)
+		{
+			Data->Metadata->FindOrCreateAttribute<T>(PropertySelector.AttributeName);
+		}
+
+		TUniquePtr<IPCGAttributeAccessor> Accessor = PCGAttributeAccessorHelpers::CreateAccessor(Data, PropertySelector);
+		if (!Accessor)
+		{
+			return;
+		}
+
+		TUniquePtr<IPCGAttributeAccessorKeys> Keys = PCGAttributeAccessorHelpers::CreateKeys(Data, PropertySelector);
+		if (!Keys)
+		{
+			return;
+		}
+
+		Accessor->SetRange(Values, 0, *Keys);
+	}
+
+	template<typename FractalNoiseFunc>
+	void DoFractal2D(const FSharedParams& SharedParams, const FBufferParams& BufferParams, const UPCGNoiseSettings& Settings, FractalNoiseFunc&& FractalNoise)
+	{
+		const TArray<FPCGPoint>& SrcPoints = BufferParams.InputPointData->GetPoints();
+
+		struct FProcessResults
+		{
+			TArray<FPCGPoint> Points;
+			TArray<double> Values;
+		};
+
+		FProcessResults Results;
+
+		FPCGAsync::AsyncProcessingOneToOneEx(
+			SharedParams.Context ? &SharedParams.Context->AsyncState : nullptr,
+			SrcPoints.Num(),
+			[&Results, Count = SrcPoints.Num()]()
+			{
+				// initialize
+				Results.Points.SetNumUninitialized(Count);
+				Results.Values.SetNumUninitialized(Count);
+			},
+			[
+				&SharedParams,
+				&BufferParams,
+				&Results,
+				&SrcPoints,
+				&FractalNoise
+			](const int32 ReadIndex, const int32 WriteIndex)
+			{
+				const FPCGPoint& InPoint = SrcPoints[ReadIndex];
+				FPCGPoint& OutPoint = Results.Points[WriteIndex];
+
+				OutPoint = InPoint;
+
+				const FVector PointPos = InPoint.Transform.GetTranslation();
+
+				double Value = 0.0;
+
+				if (SharedParams.bTiling)
+				{
+					const FLocalCoordinates2D LocalCoords = CalcLocalCoordinates2D(
+						SharedParams.ActorLocalBox,
+						SharedParams.ActorTransformInverse,
+						FVector2D(SharedParams.Transform.GetScale3D()),
+						PointPos
+					);
+
+					Value = FMath::BiLerp(
+						FractalNoise(FVector2D(LocalCoords.X0, LocalCoords.Y0), SharedParams.Iterations),
+						FractalNoise(FVector2D(LocalCoords.X1, LocalCoords.Y0), SharedParams.Iterations),
+						FractalNoise(FVector2D(LocalCoords.X0, LocalCoords.Y1), SharedParams.Iterations),
+						FractalNoise(FVector2D(LocalCoords.X1, LocalCoords.Y1), SharedParams.Iterations),
+						LocalCoords.FracX,
+						LocalCoords.FracY
+					);
+				}
+				else
+				{
+					Value = FractalNoise(FVector2D(SharedParams.Transform.TransformPosition(PointPos)), SharedParams.Iterations);
+				}
+
+				Value = ApplyContrast(SharedParams.Brightness + Value, SharedParams.Contrast);
+
+				Results.Values[WriteIndex] = Value;
+			},
+			/* bEnableTimeSlicing */ false
+		);
+
+		// now apply these results
+		BufferParams.OutputPointData->GetMutablePoints() = MoveTemp(Results.Points);
+
+		SetAttributeHelper<double>(BufferParams.OutputPointData, Settings.ValueTarget, Results.Values);
+	}
+
+	template<typename FractalNoiseFunc>
+	void DoEdgeMask2D(const FSharedParams& SharedParams, const FBufferParams& BufferParams, const UPCGNoiseSettings& Settings, FractalNoiseFunc&& FractalNoise)
+	{
+		const TArray<FPCGPoint>& SrcPoints = BufferParams.InputPointData->GetPoints();
+
+		struct FProcessResults
+		{
+			TArray<FPCGPoint> Points;
+			TArray<double> Values;
+		};
+
+		FProcessResults Results;
+
+		FPCGAsync::AsyncProcessingOneToOneEx(
+			SharedParams.Context ? &SharedParams.Context->AsyncState : nullptr,
+			SrcPoints.Num(),
+			[&Results, Count = SrcPoints.Num()]()
+			{
+				// initialize
+				Results.Points.SetNumUninitialized(Count);
+				Results.Values.SetNumUninitialized(Count);
+			},
+			[
+				&SharedParams,
+				&BufferParams,
+				&SrcPoints,
+				&Results,
+				FractalNoise,
+				EdgeBlendDistance = Settings.EdgeBlendDistance,
+				EdgeBlendCurveIntensity = Settings.EdgeBlendCurveIntensity,
+				EdgeBlendCurveOffset = Settings.EdgeBlendCurveOffset
+			] (const int32 ReadIndex, const int32 WriteIndex) {
+
+				const FPCGPoint& InPoint = SrcPoints[ReadIndex];
+				FPCGPoint& OutPoint = Results.Points[WriteIndex];
+
+				OutPoint = InPoint;
+
+				const FVector InPosition = InPoint.Transform.GetTranslation();
+
+				const FLocalCoordinates2D LocalCoords = CalcLocalCoordinates2D(
+					SharedParams.ActorLocalBox,
+					SharedParams.ActorTransformInverse,
+					FVector2D(SharedParams.Transform.GetScale3D()),
+					InPosition
+				);
+
+				const double EdgeBlendAmount = CalcEdgeBlendAmount2D(LocalCoords, EdgeBlendDistance);
+
+				double Value = 1.0;
+
+				if (EdgeBlendAmount > 0.0001)
+				{
+					const double NoiseValue = FractalNoise(FVector2D(SharedParams.Transform.TransformPosition(InPosition)), SharedParams.Iterations);
+					const double RemappedNoiseValue = ApplyContrast(NoiseValue, EdgeBlendCurveIntensity);
+					const double OffsetAmount = FMath::Pow(EdgeBlendAmount, EdgeBlendCurveOffset);
+
+					const double NoisedAmount = FMath::Lerp(
+						FMath::Min(OffsetAmount, RemappedNoiseValue*OffsetAmount),
+						1.0 - FMath::Min(1.0-OffsetAmount, (1.0-RemappedNoiseValue) * (1.0-OffsetAmount)),
+						OffsetAmount
+					);
+
+					Value = 1.0 - FMath::Clamp(ApplyContrast(NoisedAmount + SharedParams.Brightness, SharedParams.Contrast), 0.0, 1.0);
+				}
+
+				Results.Values[WriteIndex] = Value;
+			},
+			/* bEnableTimslicing */ false			
+		);
+
+		// now apply these results
+		BufferParams.OutputPointData->GetMutablePoints() = MoveTemp(Results.Points);
+
+		SetAttributeHelper<double>(BufferParams.OutputPointData, Settings.ValueTarget, Results.Values);
+	}
+
+	template<typename VoronoiHashFunc>
+	void DoVoronoi2DInternal(const FSharedParams& SharedParams, const FBufferParams& BufferParams, const UPCGNoiseSettings& Settings, VoronoiHashFunc&& HashFunc)
+	{
+		const TArray<FPCGPoint>& SrcPoints = BufferParams.InputPointData->GetPoints();
+
+		struct FProcessResults
+		{
+			TArray<FPCGPoint> Points;
+			TArray<double> DistanceToEdge;
+			TArray<double> CellID;
+		};
+
+		FProcessResults Results;
+
+		FPCGAsync::AsyncProcessingOneToOneEx(
+			SharedParams.Context ? &SharedParams.Context->AsyncState : nullptr,
+			SrcPoints.Num(),
+			// initialize
+			[&Results, Count = SrcPoints.Num()]()
+			{
+				// initialize
+				Results.Points.SetNumUninitialized(Count);
+				Results.DistanceToEdge.SetNumUninitialized(Count);
+				Results.CellID.SetNumUninitialized(Count);
+			},
+			// process
+			[
+				&SharedParams,
+				&BufferParams,
+				&HashFunc,
+				&SrcPoints,
+				&Results,
+				bVoronoiOrientSamplesToCellEdge = Settings.bVoronoiOrientSamplesToCellEdge
+			](const int32 ReadIndex, const int32 WriteIndex)
+			{
+				const FPCGPoint& InPoint = SrcPoints[ReadIndex];
+				FPCGPoint& OutPoint = Results.Points[WriteIndex];
+
+				OutPoint = InPoint;
+				const FVector2D Position = FVector2D(SharedParams.Transform.TransformPosition(InPoint.Transform.GetTranslation()));
+
+				FVoronoi2DResult Result = CalcVoronoi2D(Position, HashFunc);
+
+				const double Value = ApplyContrast(SharedParams.Brightness + Result.DistanceToEdge / MAGIC_SCALE_FACTOR, SharedParams.Contrast);
+
+				Results.DistanceToEdge[WriteIndex] = Value;
+				Results.CellID[WriteIndex] = Result.ID;
+
+				if (bVoronoiOrientSamplesToCellEdge)
+				{
+					OutPoint.Transform.SetRotation(
+						FQuat(UKismetMathLibrary::FindLookAtRotation(
+							FVector::ZeroVector, 
+							FVector(Result.CellEdgeDirection.X, Result.CellEdgeDirection.Y, 0.0)
+						))
+					);
+				}
+			},
+			/* bEnableTimslicing */ false
+		);
+
+		BufferParams.OutputPointData->GetMutablePoints() = MoveTemp(Results.Points);
+
+		SetAttributeHelper<double>(BufferParams.OutputPointData, Settings.ValueTarget, Results.DistanceToEdge);
+		SetAttributeHelper<double>(BufferParams.OutputPointData, Settings.VoronoiCellIDTarget, Results.CellID);
+	}
+
+	void DoVoronoi2D(const FSharedParams& SharedParams, const FBufferParams& BufferParams, const UPCGNoiseSettings& Settings)
+	{
+		if (SharedParams.bTiling)
+		{
+			/**
+
+			 This voronoi method uses a fixed integer cell subdivision which means it can tile seamlessly
+			 it also has the ability to randomize the interior cell positions with a different seed so the 
+			 end result is variations of things that still tile with each other.
+
+			 */
+			const FInt32Point InteriorOffset = FloorInt(FVector2D(SharedParams.Transform.GetTranslation()));
+			const int32 TiledVoronoiResolution = FMath::Max(1, Settings.TiledVoronoiResolution);
+			const int32 MaxBlendingTileIndex = TiledVoronoiResolution - Settings.TiledVoronoiEdgeBlendCellCount;
+			const FVector LocalBoxSize = SharedParams.ActorLocalBox.GetSize();
+
+			// copy this so the new transform converts to the cell index
+			FSharedParams UseSharedParams = SharedParams;
+
+			UseSharedParams.Transform = SharedParams.ActorTransformInverse * FTransform(
+				FRotator::ZeroRotator,
+				FVector(double(TiledVoronoiResolution) * 0.5 + 0.5, double(TiledVoronoiResolution) * 0.5 + 0.5, 1.0), // offset to the center
+				FVector(double(TiledVoronoiResolution)/LocalBoxSize.X, double(TiledVoronoiResolution)/LocalBoxSize.Y, 1.0) // scale to the range of 0 .. TiledVoronoiResolution
+			);
+
+			DoVoronoi2DInternal(UseSharedParams, BufferParams, Settings,
+				[
+					InteriorOffset,
+					MaxBlendingTileIndex,
+					TiledVoronoiResolution,
+					EdgeOffset = FloorInt(FVector2D(Settings.Transform.GetTranslation())),
+					EdgeBlendCellCount = Settings.TiledVoronoiEdgeBlendCellCount,
+					VoronoiCellRandomness = Settings.VoronoiCellRandomness					
+				] (const FInt32Point Cell)
+				{
+					if (EdgeBlendCellCount > 0 && (Cell.X < EdgeBlendCellCount || Cell.Y < EdgeBlendCellCount || Cell.X >= MaxBlendingTileIndex || Cell.Y >= MaxBlendingTileIndex))
+					{
+						// global seed for edge cells
+						return VoronoiHash2D(FInt32Point((Cell.X + EdgeOffset.X) % TiledVoronoiResolution, (Cell.Y + EdgeOffset.Y) % TiledVoronoiResolution)) * VoronoiCellRandomness;
+					}
+					else
+					{
+						return VoronoiHash2D(Cell + InteriorOffset) * VoronoiCellRandomness;
+					}
+				}
+			);	
+		}
+		else
+		{
+			DoVoronoi2DInternal(SharedParams, BufferParams, Settings,
+				[VoronoiCellRandomness = Settings.VoronoiCellRandomness](const FInt32Point CellIdx) {
+					return VoronoiHash2D(CellIdx) * VoronoiCellRandomness;
+				}
+			);
+		}		
+	}
+}
+
+bool FPCGNoise::ExecuteInternal(FPCGContext* Context) const
+{
+	TRACE_CPUPROFILER_EVENT_SCOPE(FPCGNoise::Execute);
+
+	const UPCGNoiseSettings* Settings = Context->GetInputSettings<UPCGNoiseSettings>();
+	check(Settings);
+
+	FRandomStream RandomSource(Context->GetSeed());
+
+	const FVector RandomOffset = Settings->RandomOffset * FVector(RandomSource.GetFraction(), RandomSource.GetFraction(), RandomSource.GetFraction());
+
+	PCGNoise::FSharedParams SharedParams;
+	SharedParams.Context = Context;
+
+	if (!Context->SourceComponent.IsValid())
+	{
+		PCGE_LOG(Error, GraphAndLog, LOCTEXT("NoSourceComponent", "No Source Component."));
+		return true;		
+	}
+
+	{
+		AActor* Actor = Context->SourceComponent->GetOwner();
+		if (!Actor)
+		{
+			PCGE_LOG(Error, GraphAndLog, LOCTEXT("NoActor", "Source Component has no actor"));
+			return true;		
+		}
+
+		SharedParams.ActorLocalBox = PCGHelpers::GetActorLocalBounds(Actor);
+		SharedParams.ActorLocalBox.Min *= Actor->GetTransform().GetScale3D();
+		SharedParams.ActorLocalBox.Max *= Actor->GetTransform().GetScale3D();
+
+		const FTransform ActorTransform = FTransform(Actor->GetTransform().Rotator(), Actor->GetTransform().GetTranslation(), FVector::One());
+		SharedParams.ActorTransformInverse = ActorTransform.Inverse();
+	}
+
+	SharedParams.Transform = Settings->Transform * FTransform(FRotator::ZeroRotator, RandomOffset, FVector(PCGNoise::MAGIC_SCALE_FACTOR));
+	SharedParams.Brightness = Settings->Brightness;
+	SharedParams.Contrast = Settings->Contrast;
+	SharedParams.bTiling = Settings->bTiling;
+	SharedParams.Iterations = FMath::Max(1, Settings->Iterations); // clamped in meta properties but things will crash if it's < 1
+	
+	TArray<FPCGTaggedData> Inputs = Context->InputData.GetInputsByPin(PCGPinConstants::DefaultInputLabel);	
+	for (const FPCGTaggedData& Input : Inputs)
+	{
+		PCGNoise::FBufferParams BufferParams;
+
+		BufferParams.InputPointData = Cast<UPCGPointData>(Input.Data);
+
+		if (!BufferParams.InputPointData)
+		{
+			PCGE_LOG(Error, GraphAndLog, LOCTEXT("InvalidInputData", "Invalid input data (only supports point data)."));
+			continue;
+		}
+
+		BufferParams.OutputPointData = NewObject<UPCGPointData>();
+		BufferParams.OutputPointData->InitializeFromData(BufferParams.InputPointData);
+		Context->OutputData.TaggedData.Add_GetRef(Input).Data = BufferParams.OutputPointData;
+
+		switch (Settings->Mode)
+		{
+		case PCGNoiseMode::Caustic2D:
+			{
+				DoFractal2D(SharedParams, BufferParams, *Settings, &PCGNoise::CalcCaustic2D);
+			}
+			break;
+
+		case PCGNoiseMode::Perlin2D:
+			{
+				DoFractal2D(SharedParams, BufferParams, *Settings, &PCGNoise::CalcPerlin2D);
+			}
+			break;
+
+		case PCGNoiseMode::FractionalBrownian2D:
+			{
+				DoFractal2D(SharedParams, BufferParams, *Settings, &PCGNoise::CalcFractionalBrownian2D);
+			}
+			break;
+
+		case PCGNoiseMode::Voronoi2D:
+			{
+				DoVoronoi2D(SharedParams, BufferParams, *Settings);
+			}
+			break;
+
+		case PCGNoiseMode::EdgeMask2D:
+			{
+				switch (Settings->EdgeMask2DMode)
+				{
+				case PCGNoiseMask2DMode::Perlin:
+					DoEdgeMask2D(SharedParams, BufferParams, *Settings, &PCGNoise::CalcPerlin2D);
+					break;
+
+				case PCGNoiseMask2DMode::Caustic:
+					DoEdgeMask2D(SharedParams, BufferParams, *Settings, &PCGNoise::CalcCaustic2D);
+					break;
+
+				case PCGNoiseMask2DMode::FractionalBrownian:
+					DoEdgeMask2D(SharedParams, BufferParams, *Settings, &PCGNoise::CalcFractionalBrownian2D);
+					break;
+
+				default:
+					checkNoEntry();
+					break;
+				}
+			}
+			break;
+
+		default:
+			checkNoEntry();
+			break;
+		}
+	}
+
+	return true;
+}
+
+#undef LOCTEXT_NAMESPACE
\ No newline at end of file
diff --git a/Engine/Plugins/Experimental/PCG/Source/PCG/Private/Tests/Elements/PCGNoiseTest.cpp b/Engine/Plugins/Experimental/PCG/Source/PCG/Private/Tests/Elements/PCGNoiseTest.cpp
new file mode 100644
--- /dev/null
+++ b/Engine/Plugins/Experimental/PCG/Source/PCG/Private/Tests/Elements/PCGNoiseTest.cpp
@@ -0,0 +1,315 @@
+// Copyright Epic Games, Inc. All Rights Reserved.
+
+#include "Tests/Determinism/PCGDeterminismTestsCommon.h"
+
+#include "PCGContext.h"
+#include "Data/PCGPointData.h"
+#include "Elements/PCGNoise.h"
+
+IMPLEMENT_SIMPLE_AUTOMATION_TEST(FPCGNoise_CalcLocalCoordinates2D, "pcg.tests.Noise.CalcLocalCoordinates2D", EAutomationTestFlags::EditorContext | EAutomationTestFlags::ProductFilter)
+
+bool FPCGNoise_CalcLocalCoordinates2D::RunTest(const FString& Parameters)
+{
+	const FBox TestBox(-FVector::One(), FVector::One());
+	const FVector2D Scale = FVector2D::One();
+
+	UTEST_EQUAL("Disabled Edge value", CalcEdgeBlendAmount2D(PCGNoise::CalcLocalCoordinates2D(TestBox, FTransform::Identity, Scale, FVector::One()), 0.0), 0.0);
+	UTEST_EQUAL("Disabled Edge value", CalcEdgeBlendAmount2D(PCGNoise::CalcLocalCoordinates2D(TestBox, FTransform::Identity, Scale, FVector::Zero()), 0.0), 0.0);
+
+	PCGNoise::FLocalCoordinates2D Test0 = PCGNoise::CalcLocalCoordinates2D(TestBox, FTransform::Identity, Scale, FVector::One());
+	const double EdgeBlendAmount0 = CalcEdgeBlendAmount2D(Test0, 1.0);
+	UTEST_EQUAL_TOLERANCE("Enabled Edge Value", EdgeBlendAmount0, 1.0, 0.001);
+	UTEST_EQUAL_TOLERANCE("Edge Sample X0", Test0.X0, 2.0, 0.001);
+	UTEST_EQUAL_TOLERANCE("Edge Sample Y0", Test0.Y0, 2.0, 0.001);
+	UTEST_EQUAL_TOLERANCE("Edge Sample X1", Test0.X1, 0.0, 0.001);
+	UTEST_EQUAL_TOLERANCE("Edge Sample Y1", Test0.Y1, 0.0, 0.001);
+
+	UTEST_EQUAL_TOLERANCE("Enabled Edge Value Center", CalcEdgeBlendAmount2D(PCGNoise::CalcLocalCoordinates2D(TestBox, FTransform::Identity, Scale, FVector::Zero()), 1.0), 0.0, 0.001);
+
+	PCGNoise::FLocalCoordinates2D Test1 = PCGNoise::CalcLocalCoordinates2D(TestBox, FTransform::Identity, Scale, FVector(0.9));
+	const double EdgeBlendAmount1 = CalcEdgeBlendAmount2D(Test1, 0.5);
+	UTEST_EQUAL_TOLERANCE("Blended Edge Value", EdgeBlendAmount1, 0.8, 0.001);
+	UTEST_EQUAL_TOLERANCE("Blended Edge Sample X", Test1.X0, 1.9, 0.001);
+	UTEST_EQUAL_TOLERANCE("Blended Edge Sample Y", Test1.Y0, 1.9, 0.001);
+	UTEST_EQUAL_TOLERANCE("Blended Edge Sample X", Test1.X1, -0.1, 0.001);
+	UTEST_EQUAL_TOLERANCE("Blended Edge Sample Y", Test1.Y1, -0.1, 0.001);
+
+	return true;
+}
+
+// just make sure stuff doesn't crash
+
+IMPLEMENT_CUSTOM_SIMPLE_AUTOMATION_TEST(FPCGNoise_Perlin2D, FPCGTestBaseClass, "pcg.tests.Noise.Perlin2D", PCGTestsCommon::TestFlags)
+
+bool FPCGNoise_Perlin2D::RunTest(const FString& Parameters)
+{
+	PCGTestsCommon::FTestData TestData(PCGDeterminismTests::Defaults::Seed);
+	PCGDeterminismTests::GenerateSettings<UPCGNoiseSettings>(TestData);
+	UPCGNoiseSettings* Settings = CastChecked<UPCGNoiseSettings>(TestData.Settings);
+
+	Settings->Mode = PCGNoiseMode::Perlin2D;
+	Settings->ValueTarget.AttributeName = TEXT("Noise");
+
+	FPCGElementPtr TestElement = TestData.Settings->GetElement();
+
+	{
+		FPCGTaggedData& SourcePin = TestData.InputData.TaggedData.Emplace_GetRef();
+		SourcePin.Pin = PCGPinConstants::DefaultInputLabel;
+		SourcePin.Data = PCGTestsCommon::CreateRandomPointData(100, 42);
+	}
+
+	TUniquePtr<FPCGContext> Context = TestData.InitializeTestContext();
+
+	while (!TestElement->Execute(Context.Get())) {}	
+
+	const TArray<FPCGTaggedData>& Outputs = Context->OutputData.GetInputs();
+
+	UTEST_EQUAL("Output count", Outputs.Num(), 1);
+
+	UPCGPointData* OutputData = Cast<UPCGPointData>(Outputs[0].Data);
+
+	UTEST_NOT_NULL("Output data", OutputData);
+	UTEST_EQUAL("Output point data count", OutputData->GetPoints().Num(), 100);
+
+	const FPCGMetadataAttribute<double>* NoiseAttribute = OutputData->Metadata->GetConstTypedAttribute<double>(Settings->ValueTarget.AttributeName);
+	UTEST_NOT_NULL("Noise Attribute Created", NoiseAttribute);
+
+	return true;
+}
+
+IMPLEMENT_CUSTOM_SIMPLE_AUTOMATION_TEST(FPCGNoise_Caustic2D, FPCGTestBaseClass, "pcg.tests.Noise.Caustic2D", PCGTestsCommon::TestFlags)
+
+bool FPCGNoise_Caustic2D::RunTest(const FString& Parameters)
+{
+	PCGTestsCommon::FTestData TestData(PCGDeterminismTests::Defaults::Seed);
+	PCGDeterminismTests::GenerateSettings<UPCGNoiseSettings>(TestData);
+	UPCGNoiseSettings* Settings = CastChecked<UPCGNoiseSettings>(TestData.Settings);
+
+	Settings->Mode = PCGNoiseMode::Caustic2D;
+	Settings->ValueTarget.AttributeName = TEXT("Noise");
+
+	FPCGElementPtr TestElement = TestData.Settings->GetElement();
+
+	{
+		FPCGTaggedData& SourcePin = TestData.InputData.TaggedData.Emplace_GetRef();
+		SourcePin.Pin = PCGPinConstants::DefaultInputLabel;
+		SourcePin.Data = PCGTestsCommon::CreateRandomPointData(100, 42);
+	}
+
+	TUniquePtr<FPCGContext> Context = TestData.InitializeTestContext();
+
+	while (!TestElement->Execute(Context.Get())) {}	
+
+	const TArray<FPCGTaggedData>& Outputs = Context->OutputData.GetInputs();
+
+	UTEST_EQUAL("Output count", Outputs.Num(), 1);
+
+	UPCGPointData* OutputData = Cast<UPCGPointData>(Outputs[0].Data);
+
+	UTEST_NOT_NULL("Output data", OutputData);
+	UTEST_EQUAL("Output point data count", OutputData->GetPoints().Num(), 100);
+
+	const FPCGMetadataAttribute<double>* NoiseAttribute = OutputData->Metadata->GetConstTypedAttribute<double>(Settings->ValueTarget.AttributeName);
+	UTEST_NOT_NULL("Noise Attribute Created", NoiseAttribute);
+
+	return true;
+}
+
+IMPLEMENT_CUSTOM_SIMPLE_AUTOMATION_TEST(FPCGNoise_Voronoi2D, FPCGTestBaseClass, "pcg.tests.Noise.Voronoi2D", PCGTestsCommon::TestFlags)
+
+bool FPCGNoise_Voronoi2D::RunTest(const FString& Parameters)
+{
+	PCGTestsCommon::FTestData TestData(PCGDeterminismTests::Defaults::Seed);
+	PCGDeterminismTests::GenerateSettings<UPCGNoiseSettings>(TestData);
+	UPCGNoiseSettings* Settings = CastChecked<UPCGNoiseSettings>(TestData.Settings);
+
+	Settings->Mode = PCGNoiseMode::Voronoi2D;
+	Settings->ValueTarget.AttributeName = TEXT("Distance");
+	Settings->VoronoiCellIDTarget.AttributeName = TEXT("CellID");
+
+	FPCGElementPtr TestElement = TestData.Settings->GetElement();
+
+	{
+		FPCGTaggedData& SourcePin = TestData.InputData.TaggedData.Emplace_GetRef();
+		SourcePin.Pin = PCGPinConstants::DefaultInputLabel;
+		SourcePin.Data = PCGTestsCommon::CreateRandomPointData(100, 42);
+	}
+
+	TUniquePtr<FPCGContext> Context = TestData.InitializeTestContext();
+
+	while (!TestElement->Execute(Context.Get())) {}	
+
+	const TArray<FPCGTaggedData>& Outputs = Context->OutputData.GetInputs();
+
+	UTEST_EQUAL("Output count", Outputs.Num(), 1);
+
+	UPCGPointData* OutputData = Cast<UPCGPointData>(Outputs[0].Data);
+
+	UTEST_NOT_NULL("Output data", OutputData);
+	UTEST_EQUAL("Output point data count", OutputData->GetPoints().Num(), 100);
+
+	const FPCGMetadataAttribute<double>* NoiseAttribute = OutputData->Metadata->GetConstTypedAttribute<double>(Settings->ValueTarget.AttributeName);
+	UTEST_NOT_NULL("Distance Attribute Created", NoiseAttribute);
+
+	const FPCGMetadataAttribute<double>* CellIDAttribute = OutputData->Metadata->GetConstTypedAttribute<double>(Settings->VoronoiCellIDTarget.AttributeName);
+	UTEST_NOT_NULL("Cell ID Attribute Created", CellIDAttribute);
+
+
+	return true;
+}
+
+IMPLEMENT_CUSTOM_SIMPLE_AUTOMATION_TEST(FPCGNoise_FractionalBrownian2D, FPCGTestBaseClass, "pcg.tests.Noise.FractionalBrownian2D", PCGTestsCommon::TestFlags)
+
+bool FPCGNoise_FractionalBrownian2D::RunTest(const FString& Parameters)
+{
+	PCGTestsCommon::FTestData TestData(PCGDeterminismTests::Defaults::Seed);
+	PCGDeterminismTests::GenerateSettings<UPCGNoiseSettings>(TestData);
+	UPCGNoiseSettings* Settings = CastChecked<UPCGNoiseSettings>(TestData.Settings);
+
+	Settings->Mode = PCGNoiseMode::FractionalBrownian2D;
+	Settings->ValueTarget.AttributeName = TEXT("Noise");
+
+	FPCGElementPtr TestElement = TestData.Settings->GetElement();
+
+	{
+		FPCGTaggedData& SourcePin = TestData.InputData.TaggedData.Emplace_GetRef();
+		SourcePin.Pin = PCGPinConstants::DefaultInputLabel;
+		SourcePin.Data = PCGTestsCommon::CreateRandomPointData(100, 42);
+	}
+
+	TUniquePtr<FPCGContext> Context = TestData.InitializeTestContext();
+
+	while (!TestElement->Execute(Context.Get())) {}	
+
+	const TArray<FPCGTaggedData>& Outputs = Context->OutputData.GetInputs();
+
+	UTEST_EQUAL("Output count", Outputs.Num(), 1);
+
+	UPCGPointData* OutputData = Cast<UPCGPointData>(Outputs[0].Data);
+
+	UTEST_NOT_NULL("Output data", OutputData);
+	UTEST_EQUAL("Output point data count", OutputData->GetPoints().Num(), 100);
+
+	const FPCGMetadataAttribute<double>* NoiseAttribute = OutputData->Metadata->GetConstTypedAttribute<double>(Settings->ValueTarget.AttributeName);
+	UTEST_NOT_NULL("Noise Attribute Created", NoiseAttribute);
+
+	return true;
+}
+
+IMPLEMENT_CUSTOM_SIMPLE_AUTOMATION_TEST(FPCGNoise_EdgeMask2D, FPCGTestBaseClass, "pcg.tests.Noise.EdgeMask2D", PCGTestsCommon::TestFlags)
+
+bool FPCGNoise_EdgeMask2D::RunTest(const FString& Parameters)
+{
+	PCGTestsCommon::FTestData TestData(PCGDeterminismTests::Defaults::Seed);
+	PCGDeterminismTests::GenerateSettings<UPCGNoiseSettings>(TestData);
+	UPCGNoiseSettings* Settings = CastChecked<UPCGNoiseSettings>(TestData.Settings);
+
+	Settings->Mode = PCGNoiseMode::EdgeMask2D;
+	Settings->ValueTarget.AttributeName = TEXT("Noise");
+
+	FPCGElementPtr TestElement = TestData.Settings->GetElement();
+
+	{
+		FPCGTaggedData& SourcePin = TestData.InputData.TaggedData.Emplace_GetRef();
+		SourcePin.Pin = PCGPinConstants::DefaultInputLabel;
+		SourcePin.Data = PCGTestsCommon::CreateRandomPointData(100, 42);
+	}
+
+	TUniquePtr<FPCGContext> Context = TestData.InitializeTestContext();
+
+	while (!TestElement->Execute(Context.Get())) {}	
+
+	const TArray<FPCGTaggedData>& Outputs = Context->OutputData.GetInputs();
+
+	UTEST_EQUAL("Output count", Outputs.Num(), 1);
+
+	UPCGPointData* OutputData = Cast<UPCGPointData>(Outputs[0].Data);
+
+	UTEST_NOT_NULL("Output data", OutputData);
+	UTEST_EQUAL("Output point data count", OutputData->GetPoints().Num(), 100);
+
+	const FPCGMetadataAttribute<double>* NoiseAttribute = OutputData->Metadata->GetConstTypedAttribute<double>(Settings->ValueTarget.AttributeName);
+	UTEST_NOT_NULL("Noise Attribute Created", NoiseAttribute);
+
+	return true;
+}
+
+IMPLEMENT_CUSTOM_SIMPLE_AUTOMATION_TEST(FPCGNoise_TilingPerlin2D, FPCGTestBaseClass, "pcg.tests.Noise.TilingPerlin2D", PCGTestsCommon::TestFlags)
+bool FPCGNoise_TilingPerlin2D::RunTest(const FString& Parameters)
+{
+	PCGTestsCommon::FTestData TestData(PCGDeterminismTests::Defaults::Seed);
+	PCGDeterminismTests::GenerateSettings<UPCGNoiseSettings>(TestData);
+	UPCGNoiseSettings* Settings = CastChecked<UPCGNoiseSettings>(TestData.Settings);
+
+	Settings->Mode = PCGNoiseMode::Perlin2D;
+	Settings->bTiling = true;
+	Settings->ValueTarget.AttributeName = TEXT("Noise");
+
+	FPCGElementPtr TestElement = TestData.Settings->GetElement();
+
+	{
+		FPCGTaggedData& SourcePin = TestData.InputData.TaggedData.Emplace_GetRef();
+		SourcePin.Pin = PCGPinConstants::DefaultInputLabel;
+		SourcePin.Data = PCGTestsCommon::CreateRandomPointData(100, 42);
+	}
+
+	TUniquePtr<FPCGContext> Context = TestData.InitializeTestContext();
+
+	while (!TestElement->Execute(Context.Get())) {}	
+
+	const TArray<FPCGTaggedData>& Outputs = Context->OutputData.GetInputs();
+
+	UTEST_EQUAL("Output count", Outputs.Num(), 1);
+
+	UPCGPointData* OutputData = Cast<UPCGPointData>(Outputs[0].Data);
+
+	UTEST_NOT_NULL("Output data", OutputData);
+	UTEST_EQUAL("Output point data count", OutputData->GetPoints().Num(), 100);
+
+	const FPCGMetadataAttribute<double>* NoiseAttribute = OutputData->Metadata->GetConstTypedAttribute<double>(Settings->ValueTarget.AttributeName);
+	UTEST_NOT_NULL("Noise Attribute Created", NoiseAttribute);
+
+	return true;
+}
+
+IMPLEMENT_CUSTOM_SIMPLE_AUTOMATION_TEST(FPCGNoise_TilingVoronoi2D, FPCGTestBaseClass, "pcg.tests.Noise.TilingVoronoi2D", PCGTestsCommon::TestFlags)
+bool FPCGNoise_TilingVoronoi2D::RunTest(const FString& Parameters)
+{
+	PCGTestsCommon::FTestData TestData(PCGDeterminismTests::Defaults::Seed);
+	PCGDeterminismTests::GenerateSettings<UPCGNoiseSettings>(TestData);
+	UPCGNoiseSettings* Settings = CastChecked<UPCGNoiseSettings>(TestData.Settings);
+
+	Settings->Mode = PCGNoiseMode::Voronoi2D;
+	Settings->bTiling = true;
+	Settings->ValueTarget.AttributeName = TEXT("Distance");
+	Settings->VoronoiCellIDTarget.AttributeName = TEXT("CellID");
+
+	FPCGElementPtr TestElement = TestData.Settings->GetElement();
+
+	{
+		FPCGTaggedData& SourcePin = TestData.InputData.TaggedData.Emplace_GetRef();
+		SourcePin.Pin = PCGPinConstants::DefaultInputLabel;
+		SourcePin.Data = PCGTestsCommon::CreateRandomPointData(100, 42);
+	}
+
+	TUniquePtr<FPCGContext> Context = TestData.InitializeTestContext();
+
+	while (!TestElement->Execute(Context.Get())) {}	
+
+	const TArray<FPCGTaggedData>& Outputs = Context->OutputData.GetInputs();
+
+	UTEST_EQUAL("Output count", Outputs.Num(), 1);
+
+	UPCGPointData* OutputData = Cast<UPCGPointData>(Outputs[0].Data);
+
+	UTEST_NOT_NULL("Output data", OutputData);
+	UTEST_EQUAL("Output point data count", OutputData->GetPoints().Num(), 100);
+
+	const FPCGMetadataAttribute<double>* NoiseAttribute = OutputData->Metadata->GetConstTypedAttribute<double>(Settings->ValueTarget.AttributeName);
+	UTEST_NOT_NULL("Distance Attribute Created", NoiseAttribute);
+
+	const FPCGMetadataAttribute<double>* CellIDAttribute = OutputData->Metadata->GetConstTypedAttribute<double>(Settings->VoronoiCellIDTarget.AttributeName);
+	UTEST_NOT_NULL("Cell ID Attribute Created", CellIDAttribute);
+
+	return true;
+}
diff --git a/Engine/Plugins/Experimental/PCG/Source/PCG/Public/Elements/PCGNoise.h b/Engine/Plugins/Experimental/PCG/Source/PCG/Public/Elements/PCGNoise.h
new file mode 100644
--- /dev/null
+++ b/Engine/Plugins/Experimental/PCG/Source/PCG/Public/Elements/PCGNoise.h
@@ -0,0 +1,143 @@
+// Copyright Epic Games, Inc. All Rights Reserved.
+
+#pragma once
+
+#include "PCGSettings.h"
+#include "Elements/PCGPointProcessingElementBase.h"
+#include "Metadata/PCGAttributePropertySelector.h"
+
+#include "PCGNoise.generated.h"
+
+UENUM()
+enum class PCGNoiseMode
+{
+	Perlin2D, // your classic perlin noise
+	Caustic2D, // based on underwater fake custic rendering, gives swirly look
+	Voronoi2D, // voronoi noise, result a the distance to edge and cell ID
+	FractionalBrownian2D, // based on fractional brownian motion
+	EdgeMask2D, // used to create masks to blend out edges
+};
+
+UENUM()
+enum class PCGNoiseMask2DMode
+{
+	Perlin, // your classic perlin noise
+	Caustic, // based on underwater fake custic rendering, gives swirly look
+	FractionalBrownian, // based on fractional brownian motion
+};
+
+namespace PCGNoise
+{
+	struct FLocalCoordinates2D
+	{
+		// coordinates to sample for the repeating edges (not 0 to 1, is scaled according to the settings scale)
+		double X0;
+		double Y0;
+		double X1;
+		double Y1;
+
+		// how much to interpolate between the corners
+		double FracX;
+		double FracY;
+	};
+
+	FLocalCoordinates2D CalcLocalCoordinates2D(const FBox& ActorLocalBox, const FTransform& ActorTransformInverse, FVector2D Scale, const FVector& Position);
+	double CalcEdgeBlendAmount2D(const FLocalCoordinates2D& LocalCoords, double EdgeBlendDistance);
+}
+
+/**
+ * Various fractal noises that can be used to filter points
+ */
+UCLASS(BlueprintType, ClassGroup = (Procedural))
+class UPCGNoiseSettings : public UPCGSettings
+{
+	GENERATED_BODY()
+
+public:
+	//~Begin UPCGSettings interface
+#if WITH_EDITOR
+	virtual FName GetDefaultNodeName() const override { return FName(TEXT("Noise")); }
+	virtual FText GetDefaultNodeTitle() const override { return NSLOCTEXT("PCGNoise", "NodeTitle", "Noise"); }
+	virtual EPCGSettingsType GetType() const override { return EPCGSettingsType::Spatial; }
+#endif
+
+protected:
+	virtual TArray<FPCGPinProperties> InputPinProperties() const override;
+	virtual TArray<FPCGPinProperties> OutputPinProperties() const override;
+
+	virtual FPCGElementPtr CreateElement() const override;
+	//~End UPCGSettings interface
+
+public:
+
+	// The noise method used
+	UPROPERTY(BlueprintReadWrite, EditAnywhere, Category = Settings, meta = (PCG_Overridable))
+	PCGNoiseMode Mode = PCGNoiseMode::Perlin2D;
+
+	UPROPERTY(BlueprintReadWrite, EditAnywhere, Category = Settings, meta = (EditCondition = "Mode == PCGNoiseMode::EdgeMask2D", EditConditionHides, PCG_Overridable))
+	PCGNoiseMask2DMode EdgeMask2DMode = PCGNoiseMask2DMode::Perlin;
+
+	// this is how many times the fractal method recurses. A higher number will mean more detail
+	UPROPERTY(BlueprintReadWrite, EditAnywhere, Category = Settings, meta = (ClampMin = "1", ClampMax = "100", EditCondition = "Mode != PCGNoiseMode::Voronoi2D", EditConditionHides, PCG_Overridable))
+	int32 Iterations = 4;
+
+	// if true, will generate results that tile along the bounding box size of the 
+	UPROPERTY(BlueprintReadWrite, EditAnywhere, Category = Settings, meta = (PCG_Overridable))
+	bool bTiling = false;
+
+	UPROPERTY(BlueprintReadWrite, EditAnywhere, Category = Settings, meta = (PCG_Overridable))
+	float Brightness = 0.0;
+
+	UPROPERTY(BlueprintReadWrite, EditAnywhere, Category = Settings, meta = (PCG_Overridable))
+	float Contrast = 1.0;
+
+	// The output attribute name to write, if not 'None'
+	UPROPERTY(BlueprintReadWrite, EditAnywhere, Category = Settings)
+	FPCGAttributePropertySelector ValueTarget;
+
+	// Adds a random amount of offset up to this amount
+	UPROPERTY(BlueprintReadWrite, EditAnywhere, Category = Settings, meta = (PCG_Overridable))
+	FVector RandomOffset = FVector(100000.0);
+
+	// this will apply a transform to the points before calculating noise
+	UPROPERTY(BlueprintReadWrite, EditAnywhere, Category = Settings, meta = (EditCondition = "Mode != PCGNoiseMode::Voronoi2D || !bTiling", EditConditionHides))
+	FTransform Transform = FTransform::Identity;
+
+	// the less random this is, the more it returns to being a grid
+	UPROPERTY(BlueprintReadWrite, EditAnywhere, Category = Settings, meta = (ClampMin = "0.0", ClampMax = "1.0", EditCondition = "Mode == PCGNoiseMode::Voronoi2D", EditConditionHides, PCG_Overridable))
+	double VoronoiCellRandomness = 1.0;
+
+	// The output attribute name to write the voronoi cell id, if not 'None'
+	UPROPERTY(BlueprintReadWrite, EditAnywhere, Category = Settings, meta = (EditCondition = "Mode == PCGNoiseMode::Voronoi2D", EditConditionHides, PCG_Overridable))
+	FPCGAttributePropertySelector VoronoiCellIDTarget;
+	
+	// If true it will orient the output points to point towards the cell edges, which can be used for effects
+	UPROPERTY(BlueprintReadWrite, EditAnywhere, Category = Settings, meta = (EditCondition = "Mode == PCGNoiseMode::Voronoi2D", EditConditionHides, PCG_Overridable))
+	bool bVoronoiOrientSamplesToCellEdge = false;
+	
+	// The cell resolution of the tiled voronoi (across the bounds)
+	UPROPERTY(BlueprintReadWrite, EditAnywhere, Category = Settings, meta = (EditCondition = "Mode == PCGNoiseMode::Voronoi2D && bTiling", EditConditionHides, PCG_Overridable, ClampMin = "1"))
+	int32 TiledVoronoiResolution = 8;
+
+	// how many cells around the edge will tile
+	UPROPERTY(BlueprintReadWrite, EditAnywhere, Category = Settings, meta = (EditCondition = "Mode == PCGNoiseMode::Voronoi2D && bTiling", EditConditionHides, PCG_Overridable))
+	int32 TiledVoronoiEdgeBlendCellCount = 2;
+
+	// if > 0, we blend to a tiling edge value
+	UPROPERTY(BlueprintReadWrite, EditAnywhere, Category = Settings, meta = (EditCondition = "Mode == PCGNoiseMode::EdgeMask2D", EditConditionHides, PCG_Overridable))
+	float EdgeBlendDistance = 1.0;
+
+	// Adjust the center point of the curve (where x = curve(x) crosses over)
+	UPROPERTY(BlueprintReadWrite, EditAnywhere, Category = Settings, meta = (EditCondition = "Mode == PCGNoiseMode::EdgeMask2D", EditConditionHides, ClampMin = "0", PCG_Overridable))
+	float EdgeBlendCurveOffset = 1.0;
+
+	// will makes the falloff harsher or softer
+	UPROPERTY(BlueprintReadWrite, EditAnywhere, Category = Settings, meta = (EditCondition = "Mode == PCGNoiseMode::EdgeMask2D", EditConditionHides, ClampMin = "0", PCG_Overridable))
+	float EdgeBlendCurveIntensity = 1.0;
+};
+
+class FPCGNoise : public FPCGPointProcessingElementBase
+{
+protected:
+	virtual bool ExecuteInternal(FPCGContext* Context) const override;
+};
