Fix for non-unity build error

diff --git a/Engine/Plugins/Experimental/PCG/Source/PCG/Private/Elements/PCGNoise.cpp b/Engine/Plugins/Experimental/PCG/Source/PCG/Private/Elements/PCGNoise.cpp
--- a/Engine/Plugins/Experimental/PCG/Source/PCG/Private/Elements/PCGNoise.cpp
+++ b/Engine/Plugins/Experimental/PCG/Source/PCG/Private/Elements/PCGNoise.cpp
@@ -3,6 +3,7 @@
 #include "Elements/PCGNoise.h"
 
 #include "Kismet/KismetMathLibrary.h"
+#include "GameFramework/Actor.h"
 
 #include "PCGComponent.h"
 #include "PCGContext.h"
