Fix for divide by 0 CIS error

diff --git a/Engine/Plugins/Experimental/PCG/Source/PCG/Private/Elements/PCGNoise.cpp b/Engine/Plugins/Experimental/PCG/Source/PCG/Private/Elements/PCGNoise.cpp
--- a/Engine/Plugins/Experimental/PCG/Source/PCG/Private/Elements/PCGNoise.cpp
+++ b/Engine/Plugins/Experimental/PCG/Source/PCG/Private/Elements/PCGNoise.cpp
@@ -458,7 +458,7 @@ namespace PCGNoise
 				&SrcPoints,
 				&Results,
 				FractalNoise,
-				EdgeBlendDistance = Settings.EdgeBlendDistance,
+				EdgeBlendDistance = FMath::Max(Settings.EdgeBlendDistance, 0.01),
 				EdgeBlendCurveIntensity = Settings.EdgeBlendCurveIntensity,
 				EdgeBlendCurveOffset = Settings.EdgeBlendCurveOffset
 			] (const int32 ReadIndex, const int32 WriteIndex) {
