Preconditions:
- Open QAGame in Unreal Editor
- From the File Menu bar, select File > New Level
- Create a new Basic Template level

Steps:
1) Within the viewport, select the Floor mesh

2) Within the Details Panel's Components section, select Add

3) Search for PCG

4)The PCG Component is exposed within an Actor's Component section within the viewport

diff --git a/Engine/Plugins/Experimental/PCG/Source/PCGEditor/PCGEditor.Build.cs b/Engine/Plugins/Experimental/PCG/Source/PCGEditor/PCGEditor.Build.cs
--- a/Engine/Plugins/Experimental/PCG/Source/PCGEditor/PCGEditor.Build.cs
+++ b/Engine/Plugins/Experimental/PCG/Source/PCGEditor/PCGEditor.Build.cs
@@ -17,6 +17,17 @@ namespace UnrealBuildTool.Rules
 			"PlacementMode",
 		});
 
+			if (Target.WithAutomationTests)
+			{
+
+				PublicDependencyModuleNames.AddRange(
+					new string[]
+					{
+						"LevelEditor"
+					}
+				);
+			}
+
 			PrivateDependencyModuleNames.AddRange(
 				new string[]
 				{
diff --git a/Engine/Plugins/Experimental/PCG/Source/PCGEditor/Private/Tests/Display/PCGInActorAddComponentField.cpp b/Engine/Plugins/Experimental/PCG/Source/PCGEditor/Private/Tests/Display/PCGInActorAddComponentField.cpp
new file mode 100644
--- /dev/null
+++ b/Engine/Plugins/Experimental/PCG/Source/PCGEditor/Private/Tests/Display/PCGInActorAddComponentField.cpp
@@ -0,0 +1,105 @@
+// Copyright Epic Games, Inc. All Rights Reserved.
+
+#include "Editor.h" 
+#include "Editor/EditorEngine.h"
+#include "Editor/TemplateMapInfo.h"
+#include "Editor/UnrealEdEngine.h"
+#include "EditorLevelLibrary.h"
+#include "EditorLevelUtils.h"
+
+#include "Engine/StaticMeshActor.h"
+#include "LevelEditorSubsystem.h"
+#include "UnrealEdGlobals.h"
+#include "EngineUtils.h"
+
+#include "Tests/AutomationEditorCommon.h"
+#include "Tests/PCGTestsCommon.h"
+#include "PCGComponent.h"
+
+
+#if WITH_AUTOMATION_TESTS
+
+IMPLEMENT_CUSTOM_SIMPLE_AUTOMATION_TEST(FPCGInActorAddComponentField, FPCGTestBaseClass, "Editor.Plugins.Tools.PCG.PCGInActorAddComponentField", PCGTestsCommon::TestFlags)
+
+ULevel* LoadLevelByTemplateName(const FString& TemplateName)
+{
+	TArray<FTemplateMapInfo> TemplateMaps;
+	TemplateMaps = GUnrealEd ? GUnrealEd->GetTemplateMapInfos() : TemplateMaps;
+	FText TargetDisplayName = FText::FromString(TemplateName);
+
+	// Getting the basic template
+	FTemplateMapInfo* FoundTemplate = TemplateMaps.FindByPredicate([&](const FTemplateMapInfo& TemplateMapInfo) {
+		return TemplateMapInfo.DisplayName.EqualTo(TargetDisplayName);
+		});
+
+	ULevelEditorSubsystem* LevelEditorSubsystem = GEditor->GetEditorSubsystem<ULevelEditorSubsystem>();
+
+	ULevel* NewLevel = nullptr;
+
+	if (FoundTemplate)
+	{
+		// Basic template asset path
+		const FString BasicAssetPath = FoundTemplate->Map.ToString();
+
+		// Create and open a new level using the template
+		bool bSuccess = LevelEditorSubsystem->LoadLevel(BasicAssetPath);
+
+		if (bSuccess)
+		{
+			NewLevel = LevelEditorSubsystem->GetCurrentLevel();
+		}
+		else {
+			// Log an error if the level fails to load
+			UE_LOG(LogTemp, Error, TEXT("Failed to load level from template: %s"), *TemplateName);
+		}
+	}
+
+	return NewLevel;
+}
+
+bool FPCGInActorAddComponentField::RunTest(const FString& Parameters)
+{
+	UWorld* World = GEditor->GetEditorWorldContext().World();
+	TestNotNull(TEXT("World should not be null."), World);
+
+	ULevel* NewLevel = LoadLevelByTemplateName(TEXT("Basic"));
+
+	if (NewLevel)
+	{
+		AStaticMeshActor* FloorMeshActor = nullptr;
+
+		for (AActor* Actor : NewLevel->Actors)
+		{
+			AStaticMeshActor* MeshActor = Cast<AStaticMeshActor>(Actor);
+			if (MeshActor && MeshActor->GetName().Contains(TEXT("Floor")))
+			{
+				FloorMeshActor = MeshActor;
+				break;
+			}
+		}
+
+		TestNotNull(TEXT("FloorMeshActor doesn't exist."), FloorMeshActor);
+		if (FloorMeshActor)
+		{
+			// Add a PCGComponent to the FloorMeshActor , instead of using UI details panel 
+			UPCGComponent* NewPCGComponent = NewObject<UPCGComponent>(FloorMeshActor);
+			UActorComponent* ActorComponent = Cast<UActorComponent>(NewPCGComponent);
+
+			TestNotNull(TEXT("PCGComponent was not created successfully."), ActorComponent);
+
+			//Testing that PCGComponent can be added programmatically instead of UI add button
+			if (ActorComponent)
+			{
+				FloorMeshActor->AddOwnedComponent(ActorComponent);
+				ActorComponent->RegisterComponent();
+
+				// Check if the PCGComponent was added successfully.
+				UPCGComponent* PCGComponent = FloorMeshActor->FindComponentByClass<UPCGComponent>();
+				TestNotNull(TEXT("PCGComponent was not added successfully."), PCGComponent);
+			}
+		}
+	}
+
+	return true;
+}
+#endif 
diff --git a/Engine/Source/Programs/UnrealBuildTool/Configuration/TargetRules.cs b/Engine/Source/Programs/UnrealBuildTool/Configuration/TargetRules.cs
--- a/Engine/Source/Programs/UnrealBuildTool/Configuration/TargetRules.cs
+++ b/Engine/Source/Programs/UnrealBuildTool/Configuration/TargetRules.cs
@@ -4093,5 +4093,39 @@ namespace UnrealBuildTool
 		{
 			return Inner.OptedInModulePlatforms == null || Inner.OptedInModulePlatforms.Contains(Platform);
 		}
+
+		/// <summary>
+		/// Determines if the automation tests should be compiled based on the current configuration and optional forced settings.
+		/// If either the development tests or performance tests are compiled (unless explicitly disabled), this property returns true.
+		/// </summary>
+		private bool? bWithAutomationTestsPrivate = null;
+		public bool WithAutomationTests
+		{
+			get
+			{
+				if (!bWithAutomationTestsPrivate.HasValue)
+				{
+					bool bCompileDevTests = Configuration != UnrealTargetConfiguration.Shipping;
+					bool bCompilePerfTests = bCompileDevTests;
+
+					if (bForceCompileDevelopmentAutomationTests)
+					{
+						bCompileDevTests = true;
+					}
+					if (bForceCompilePerformanceAutomationTests)
+					{
+						bCompilePerfTests = true;
+					}
+					if (bForceDisableAutomationTests)
+					{
+						bCompileDevTests = bCompilePerfTests = false;
+					}
+
+					bWithAutomationTestsPrivate = bCompileDevTests || bCompilePerfTests;
+				}
+
+				return bWithAutomationTestsPrivate.Value;
+			}
+		}
 	}
 }
