[Backout] - CL25590145 (backed out due to compile errors)
Original CL Desc
-----------------------------------------------------------------
Fix for non-unity build error

diff --git a/Engine/Plugins/Experimental/PCG/Source/PCG/Private/Elements/PCGNoise.cpp b/Engine/Plugins/Experimental/PCG/Source/PCG/Private/Elements/PCGNoise.cpp
--- a/Engine/Plugins/Experimental/PCG/Source/PCG/Private/Elements/PCGNoise.cpp
+++ b/Engine/Plugins/Experimental/PCG/Source/PCG/Private/Elements/PCGNoise.cpp
@@ -3,7 +3,6 @@
 #include "Elements/PCGNoise.h"
 
 #include "Kismet/KismetMathLibrary.h"
-#include "GameFramework/Actor.h"
 
 #include "PCGComponent.h"
 #include "PCGContext.h"
