[PCG] Remove placement-new as it would not call the destructor on the previous object.

diff --git a/Engine/Plugins/Experimental/PCG/Source/PCG/Private/Graph/PCGGraphCompiler.cpp b/Engine/Plugins/Experimental/PCG/Source/PCG/Private/Graph/PCGGraphCompiler.cpp
--- a/Engine/Plugins/Experimental/PCG/Source/PCG/Private/Graph/PCGGraphCompiler.cpp
+++ b/Engine/Plugins/Experimental/PCG/Source/PCG/Private/Graph/PCGGraphCompiler.cpp
@@ -286,7 +286,7 @@ TArray<FPCGGraphTask> FPCGGraphCompiler::GetPrecompiledTasks(UPCGGraph* InGraph,
 	else
 	{
 		// If we failed to get a context, output a dummy blank one.
-		new(&OutStackContext) FPCGStackContext();
+		OutStackContext = FPCGStackContext();
 	}
 
 	return ExistingTasks ? *ExistingTasks : TArray<FPCGGraphTask>();
