[PCG] Various spline sampler improvements

* Use double precision for accumulating weights when point sampling point sets. Removes annoying stragglers when subtracting one set from another.
* Improve "row failed to sample" error, move the error message out of parallel code so it doesn't trigger a crash in slate code when OnNodeChanged is fired from a worker thread.
* Added a IsInGameThread() guard to avoid firing that delegate on worker threads
* Minor re-order of spline settings to put the dominant setting at the top

diff --git a/Engine/Plugins/Experimental/PCG/Source/PCG/Private/Data/PCGPointData.cpp b/Engine/Plugins/Experimental/PCG/Source/PCG/Private/Data/PCGPointData.cpp
--- a/Engine/Plugins/Experimental/PCG/Source/PCG/Private/Data/PCGPointData.cpp
+++ b/Engine/Plugins/Experimental/PCG/Source/PCG/Private/Data/PCGPointData.cpp
@@ -488,7 +488,7 @@ bool UPCGPointData::ProjectPoint(const FTransform& InTransform, const FBox& InBo
 	FVector WeightedPosition = FVector::ZeroVector;
 	FQuat WeightedQuat = FQuat::Identity;
 	FVector WeightedScale = FVector::ZeroVector;
-	float WeightedDensity = 0;
+	FVector::FReal WeightedDensity = 0;
 	FVector WeightedBoundsMin = FVector::ZeroVector;
 	FVector WeightedBoundsMax = FVector::ZeroVector;
 	FVector4 WeightedColor = FVector4::Zero();
@@ -549,7 +549,7 @@ bool UPCGPointData::ProjectPoint(const FTransform& InTransform, const FBox& InBo
 		OutPoint.Transform.SetScale3D(InTransform.GetScale3D());
 	}
 
-	OutPoint.Density = WeightedDensity;
+	OutPoint.Density = static_cast<float>(WeightedDensity);
 	OutPoint.BoundsMin = WeightedBoundsMin;
 	OutPoint.BoundsMax = WeightedBoundsMax;
 	OutPoint.Color = WeightedColor;
diff --git a/Engine/Plugins/Experimental/PCG/Source/PCG/Private/Elements/PCGSplineSampler.cpp b/Engine/Plugins/Experimental/PCG/Source/PCG/Private/Elements/PCGSplineSampler.cpp
--- a/Engine/Plugins/Experimental/PCG/Source/PCG/Private/Elements/PCGSplineSampler.cpp
+++ b/Engine/Plugins/Experimental/PCG/Source/PCG/Private/Elements/PCGSplineSampler.cpp
@@ -914,6 +914,7 @@ namespace PCGSplineSampler
 		InteriorSplinePointData.SetNum(NumDispatch);
 
 		const FTransform LineDataTransform = LineData->GetTransform();
+		bool bAnyRowFailedToSample = false;
 
 		ParallelFor(NumDispatch, [&](int32 DispatchIndex)
 		{
@@ -938,7 +939,7 @@ namespace PCGSplineSampler
 
 				if (Intersections.Num() % 2 != 0)
 				{
-					PCGE_LOG_C(Error, GraphAndLog, Context, LOCTEXT("IntersectionTestFailed", "Intersection test failed, skipping samples for this row"));
+					bAnyRowFailedToSample = true;
 					continue;
 				}
 
@@ -1079,6 +1080,11 @@ namespace PCGSplineSampler
 			}
 		});
 
+		if (bAnyRowFailedToSample)
+		{
+			PCGE_LOG_C(Error, GraphAndLog, Context, LOCTEXT("IntersectionTestFailed", "One or more rows of the spline interior failed to sample (intersection test failed). Ensure the spline points are not overlapping."));
+		}
+
 		// Finally, gather the data and push to the points
 		int32 PointCount = 0;
 		for (const TArray<TTuple<FTransform, FVector, float>>& InteriorData : InteriorSplinePointData)
diff --git a/Engine/Plugins/Experimental/PCG/Source/PCG/Private/Utils/PCGNodeVisualLogs.cpp b/Engine/Plugins/Experimental/PCG/Source/PCG/Private/Utils/PCGNodeVisualLogs.cpp
--- a/Engine/Plugins/Experimental/PCG/Source/PCG/Private/Utils/PCGNodeVisualLogs.cpp
+++ b/Engine/Plugins/Experimental/PCG/Source/PCG/Private/Utils/PCGNodeVisualLogs.cpp
@@ -40,7 +40,7 @@ void FPCGNodeVisualLogs::Log(TWeakObjectPtr<const UPCGNode> InNode, TWeakObjectP
 	}
 
 	// Broadcast outside of write scope lock
-	if (bAdded)
+	if (bAdded && IsInGameThread())
 	{
 		InNode->OnNodeChangedDelegate.Broadcast(const_cast<UPCGNode*>(InNode.Get()), EPCGChangeType::Cosmetic);
 	}
