[PCG] Added inspection support for subgraphs

diff --git a/Engine/Plugins/Experimental/PCG/Source/PCG/Private/Graph/PCGGraphExecutor.cpp b/Engine/Plugins/Experimental/PCG/Source/PCG/Private/Graph/PCGGraphExecutor.cpp
--- a/Engine/Plugins/Experimental/PCG/Source/PCG/Private/Graph/PCGGraphExecutor.cpp
+++ b/Engine/Plugins/Experimental/PCG/Source/PCG/Private/Graph/PCGGraphExecutor.cpp
@@ -630,10 +630,13 @@ void FPCGGraphExecutor::Execute()
 				{
 #if WITH_EDITOR
 					// doing this now since we're about to modify ReadyTasks potentially reallocating while Task is a reference. 
-					UPCGComponent* SourceComponent = Task.SourceComponent.Get();
-					if (SourceComponent && SourceComponent->IsInspecting())
+					if (UPCGComponent* SourceComponent = Task.SourceComponent.Get())
 					{
-						SourceComponent->StoreInspectionData(Task.Node, CachedOutput);
+						if (Task.StackIndex != INDEX_NONE)
+						{
+							const FPCGStack* Stack = Task.StackContext->GetStack(Task.StackIndex);
+							SourceComponent->StoreInspectionData(Stack, Task.Node, CachedOutput);
+						}
 					}
 #endif
 
@@ -787,10 +790,13 @@ void FPCGGraphExecutor::Execute()
 			// Additional note: this needs to be executed before the StoreResults since debugging might cancel further tasks
 			ActiveTask.Element->DebugDisplay(ActiveTask.Context.Get());
 
-			UPCGComponent* SourceComponent = ActiveTask.Context->SourceComponent.Get();
-			if (SourceComponent && SourceComponent->IsInspecting())
+			if (UPCGComponent* SourceComponent = ActiveTask.Context->SourceComponent.Get())
 			{
-				SourceComponent->StoreInspectionData(ActiveTask.Context->Node, ActiveTask.Context->OutputData);
+				if (ActiveTask.StackIndex != INDEX_NONE)
+				{
+					const FPCGStack* Stack = ActiveTask.StackContext->GetStack(ActiveTask.StackIndex);
+					SourceComponent->StoreInspectionData(Stack, ActiveTask.Context->Node, ActiveTask.Context->OutputData);
+				}
 			}
 #endif
 
diff --git a/Engine/Plugins/Experimental/PCG/Source/PCG/Private/Graph/PCGStackContext.cpp b/Engine/Plugins/Experimental/PCG/Source/PCG/Private/Graph/PCGStackContext.cpp
--- a/Engine/Plugins/Experimental/PCG/Source/PCG/Private/Graph/PCGStackContext.cpp
+++ b/Engine/Plugins/Experimental/PCG/Source/PCG/Private/Graph/PCGStackContext.cpp
@@ -5,6 +5,7 @@
 #include "PCGGraph.h"
 #include "PCGNode.h"
 #include "PCGPin.h"
+#include "PCGSubgraph.h"
 
 #include "Containers/UnrealString.h"
 #include "Misc/StringBuilder.h"
@@ -29,6 +30,89 @@
 	}
 }
 
+bool FPCGStack::CreateStackFramePath(FString& OutString, const UPCGNode* InNode, const UPCGPin* InPin) const
+{
+	// Give a healthy amount of scratch space on the stack, if it overflows it will use heap.
+	TStringBuilderWithBuffer<TCHAR, 2048> StringBuilder;
+
+	auto AddPathSeparator = [&StringBuilder]()
+	{
+		if (StringBuilder.Len())
+		{
+			StringBuilder << TEXT("/");
+		}
+	};
+
+	for (const FPCGStackFrame& Frame : StackFrames)
+	{
+		if (Frame.Object.IsValid())
+		{
+			AddPathSeparator();
+			const UObject* Object = Frame.Object.Get();
+
+			if (!Object)
+			{
+				// If any object does not resolve, cannot build the string
+				return false;
+			}
+
+			if (Object->IsA<UPCGGraph>())
+			{
+				StringBuilder << TEXT("GRAPH:") << Object->GetFullName();
+			}
+			else if (Object->IsA<UPCGNode>())
+			{
+				StringBuilder << TEXT("NODE:") << Object->GetFName();
+			}
+			else
+			{
+				// Unrecognized type, should not happen
+				ensure(false);
+				StringBuilder << TEXT("UNRECOGNIZED:") << Object->GetFullName();
+			}
+		}
+		else if (Frame.LoopIndex != -1)
+		{
+			AddPathSeparator();
+			StringBuilder << TEXT("LOOP:") << FString::FromInt(Frame.LoopIndex);
+		}
+	}
+
+	if (InNode)
+	{
+		AddPathSeparator();
+		StringBuilder << TEXT("NODE:") << InNode->GetFName().ToString();
+
+		if (InPin)
+		{
+			AddPathSeparator();
+			StringBuilder << TEXT("PIN:") << InPin->GetFName().ToString();
+		}
+	}
+
+	OutString = StringBuilder;
+	return true;
+}
+
+bool FPCGStack::operator==(const FPCGStack& Other) const
+{
+	// Stacks are the same if all stack frames are the same
+	if (StackFrames.Num() != Other.StackFrames.Num())
+	{
+		return false;
+	}
+
+	for (int32 i = 0; i < StackFrames.Num(); i++)
+	{
+		if (StackFrames[i] != Other.StackFrames[i])
+		{
+			return false;
+		}
+	}
+
+	return true;
+}
+
 int32 FPCGStackContext::PushFrame(const UObject* InFrameObject)
 {
 	if (CurrentStackIndex == INDEX_NONE)
@@ -117,85 +201,42 @@
 	}
 }
 
-bool FPCGStack::CreateStackFramePath(FString& OutString, const UPCGNode* InNode, const UPCGPin* InPin) const
+FPCGStackContext FPCGStackContext::CreateStackContextFromGraph(const UPCGGraph* InPCGGraph)
 {
-	// Give a healthy amount of scratch space on the stack, if it overflows it will use heap.
-	TStringBuilderWithBuffer<TCHAR, 2048> StringBuilder;
+	FPCGStackContext StackContext;
 
-	auto AddPathSeparator = [&StringBuilder]()
+	TFunction<void(const UPCGGraph*, FPCGStackContext&)> ParseGraphRecursive;
+	ParseGraphRecursive = [&ParseGraphRecursive] (const UPCGGraph* InPCGGraph, FPCGStackContext& InStackContext)
 	{
-		if (StringBuilder.Len())
+		if (!InPCGGraph)
 		{
-			StringBuilder << TEXT("/");
+			return;
 		}
+
+		InStackContext.PushFrame(InPCGGraph);
+
+		for (const UPCGNode* PCGNode : InPCGGraph->GetNodes())
+		{
+			if (PCGNode)
+			{
+				if (const UPCGBaseSubgraphSettings* SubgraphSettings = Cast<UPCGBaseSubgraphSettings>(PCGNode->GetSettings()))
+				{
+					if (const UPCGGraph* PCGSubgraph = SubgraphSettings->GetSubgraph())
+					{
+						InStackContext.PushFrame(PCGNode);
+
+						FPCGStackContext SubgraphStackContext;
+						ParseGraphRecursive(PCGSubgraph, SubgraphStackContext);
+						InStackContext.AppendStacks(SubgraphStackContext);
+
+						InStackContext.PopFrame();
+					}
+				}
+			}
+		}
 	};
 
-	for (const FPCGStackFrame& Frame : StackFrames)
-	{
-		if (Frame.Object.IsValid())
-		{
-			AddPathSeparator();
-			const UObject* Object = Frame.Object.Get();
+	ParseGraphRecursive(InPCGGraph, StackContext);
 
-			if (!Object)
-			{
-				// If any object does not resolve, cannot build the string
-				return false;
-			}
-
-			if (Object->IsA<UPCGGraph>())
-			{
-				StringBuilder << TEXT("GRAPH:") << Object->GetFullName();
-			}
-			else if (Object->IsA<UPCGNode>())
-			{
-				StringBuilder << TEXT("NODE:") << Object->GetFName();
-			}
-			else
-			{
-				// Unrecognized type, should not happen
-				ensure(false);
-				StringBuilder << TEXT("UNRECOGNIZED:") << Object->GetFullName();
-			}
-		}
-		else if (Frame.LoopIndex != -1)
-		{
-			AddPathSeparator();
-			StringBuilder << TEXT("LOOP:") << FString::FromInt(Frame.LoopIndex);
-		}
-	}
-
-	if (InNode)
-	{
-		AddPathSeparator();
-		StringBuilder << TEXT("NODE:") << InNode->GetFName().ToString();
-
-		if (InPin)
-		{
-			AddPathSeparator();
-			StringBuilder << TEXT("PIN:") << InPin->GetFName().ToString();
-		}
-	}
-
-	OutString = StringBuilder;
-	return true;
+	return StackContext;
 }
-
-bool FPCGStack::operator==(const FPCGStack& Other) const
-{
-	// Stacks are the same if all stack frames are the same
-	if (StackFrames.Num() != Other.StackFrames.Num())
-	{
-		return false;
-	}
-
-	for (int32 i = 0; i < StackFrames.Num(); i++)
-	{
-		if (StackFrames[i] != Other.StackFrames[i])
-		{
-			return false;
-		}
-	}
-
-	return true;
-}
diff --git a/Engine/Plugins/Experimental/PCG/Source/PCG/Private/PCGComponent.cpp b/Engine/Plugins/Experimental/PCG/Source/PCG/Private/PCGComponent.cpp
--- a/Engine/Plugins/Experimental/PCG/Source/PCG/Private/PCGComponent.cpp
+++ b/Engine/Plugins/Experimental/PCG/Source/PCG/Private/PCGComponent.cpp
@@ -21,6 +21,7 @@
 #include "Data/PCGSplineData.h"
 #include "Data/PCGUnionData.h"
 #include "Data/PCGVolumeData.h"
+#include "Graph/PCGStackContext.h"
 #include "Grid/PCGPartitionActor.h"
 #include "Helpers/PCGActorHelpers.h"
 #include "Helpers/PCGHelpers.h"
@@ -1675,30 +1676,52 @@ void UPCGComponent::ResetLastGeneratedBounds()
 {
 	LastGeneratedBounds = FBox(EForceInit::ForceInit);
 }
+
+bool UPCGComponent::IsInspecting() const
+{
+	return InspectionCounter > 0;
+}
+
+void UPCGComponent::EnableInspection()
+{
+	if (!ensure(InspectionCounter >= 0))
+	{
+		InspectionCounter = 0;
+	}
+	
+	InspectionCounter++;
+}
+
 void UPCGComponent::DisableInspection()
 {
-	bIsInspecting = false;
-	InspectionCache.Empty();
+	if (ensure(InspectionCounter > 0))
+	{
+		InspectionCounter--;
+	}
+	
+	if (InspectionCounter == 0)
+	{
+		InspectionCache.Empty();
+	}
 };
 
-void UPCGComponent::StoreInspectionData(const UPCGNode* InNode, const FPCGDataCollection& InInspectionData)
+void UPCGComponent::StoreInspectionData(const FPCGStack* InStack, const UPCGNode* InNode, const FPCGDataCollection& InInspectionData)
 {
-	if (!InNode)
+	if (!IsInspecting())
 	{
 		return;
 	}
 
-	if (GetGraph() != InNode->GetGraph())
+	FString StackFramePath;
+	if (ensure(InStack) && ensure(InStack->CreateStackFramePath(StackFramePath, InNode)))
 	{
-		return;
+		InspectionCache.Add(StackFramePath, InInspectionData);
 	}
-
-	InspectionCache.Add(InNode, InInspectionData);
 }
 
-const FPCGDataCollection* UPCGComponent::GetInspectionData(const UPCGNode* InNode) const
+const FPCGDataCollection* UPCGComponent::GetInspectionData(const FString& InStackPath) const
 {
-	return InspectionCache.Find(InNode);
+	return InspectionCache.Find(InStackPath);
 }
 
 void UPCGComponent::Refresh()
diff --git a/Engine/Plugins/Experimental/PCG/Source/PCG/Public/Graph/PCGStackContext.h b/Engine/Plugins/Experimental/PCG/Source/PCG/Public/Graph/PCGStackContext.h
--- a/Engine/Plugins/Experimental/PCG/Source/PCG/Public/Graph/PCGStackContext.h
+++ b/Engine/Plugins/Experimental/PCG/Source/PCG/Public/Graph/PCGStackContext.h
@@ -62,6 +62,7 @@ public:
 	int32 GetNumStacks() const { return Stacks.Num(); }
 	int32 GetCurrentStackIndex() const { return CurrentStackIndex; }
 	const FPCGStack* GetStack(int32 InStackIndex) const;
+	const TArray<FPCGStack>& GetStacks() const { return Stacks; }
 
 	/** Create a new stack and create a frame from the provided object (typically graph or node pointer). Returns index of newly added stack. */
 	int32 PushFrame(const UObject* InFrameObject);
@@ -75,6 +76,8 @@ public:
 	/** Called during execution when invoking a dynamic subgraph, to prepend the caller stack to form the complete callstacks. */
 	void PrependParentStack(const FPCGStack* InParentStack);
 
+	static FPCGStackContext CreateStackContextFromGraph(const UPCGGraph* InPCGGraph);
+
 private:
 	/** List of all stacks encountered top graph and all (nested) subgraphs. Order is simply order of encountering during compilation. */
 	TArray<FPCGStack> Stacks;
diff --git a/Engine/Plugins/Experimental/PCG/Source/PCG/Public/PCGComponent.h b/Engine/Plugins/Experimental/PCG/Source/PCG/Public/PCGComponent.h
--- a/Engine/Plugins/Experimental/PCG/Source/PCG/Public/PCGComponent.h
+++ b/Engine/Plugins/Experimental/PCG/Source/PCG/Public/PCGComponent.h
@@ -14,6 +14,7 @@ namespace EEndPlayReason { enum Type : int; }
 
 class APCGPartitionActor;
 struct FPCGContext;
+class FPCGStack;
 class UPCGComponent;
 class UPCGGraph;
 class UPCGGraphInterface;
@@ -224,11 +225,11 @@ public:
 	void ResetLastGeneratedBounds();
 
 	/** Functions for managing the node inspection cache */
-	bool IsInspecting() const { return bIsInspecting; }
-	void EnableInspection() { bIsInspecting = true; };
+	bool IsInspecting() const;
+	void EnableInspection();
 	void DisableInspection();
-	void StoreInspectionData(const UPCGNode* InNode, const FPCGDataCollection& InInspectionData);
-	const FPCGDataCollection* GetInspectionData(const UPCGNode* InNode) const;
+	void StoreInspectionData(const FPCGStack* InStack, const UPCGNode* InNode, const FPCGDataCollection& InInspectionData);
+	const FPCGDataCollection* GetInspectionData(const FString& InStackPath) const;
 
 	/** Used by the tracking system to know if the component need to track actors. Not enabled for now, tracking is still done on the component.*/
 	bool ShouldTrackActors() const { return false; }
@@ -391,7 +392,7 @@ private:
 	bool bIsComponentLocal = false;
 
 #if WITH_EDITOR
-	bool bIsInspecting = false;
+	int32 InspectionCounter = 0;
 	FBox LastGeneratedBoundsPriorToUndo = FBox(EForceInit::ForceInit);
 	FPCGTagToSettingsMap CachedTrackedTagsToSettings;
 
@@ -418,7 +419,7 @@ private:
 	bool bActorToTagsMapPopulated = false;
 
 	UPROPERTY(Transient)
-	TMap<TObjectPtr<const UPCGNode>, FPCGDataCollection> InspectionCache;
+	TMap<FString, FPCGDataCollection> InspectionCache;
 #endif
 
 	mutable FCriticalSection GeneratedResourcesLock;
diff --git a/Engine/Plugins/Experimental/PCG/Source/PCGEditor/Private/PCGEditor.cpp b/Engine/Plugins/Experimental/PCG/Source/PCGEditor/Private/PCGEditor.cpp
--- a/Engine/Plugins/Experimental/PCG/Source/PCGEditor/Private/PCGEditor.cpp
+++ b/Engine/Plugins/Experimental/PCG/Source/PCGEditor/Private/PCGEditor.cpp
@@ -159,18 +159,27 @@ UPCGEditorGraph* FPCGEditor::GetPCGEditorGraph()
 	return PCGEditorGraph;
 }
 
-void FPCGEditor::SetPCGComponentBeingDebugged(UPCGComponent* InPCGComponent)
+void FPCGEditor::SetComponentAndStackBeingInspected(UPCGComponent* InPCGComponent, const FPCGStack& InPCGStack)
 {
-	if (PCGComponentBeingDebugged != InPCGComponent)
+	if (PCGComponentBeingInspected != InPCGComponent)
 	{
-		PCGComponentBeingDebugged = InPCGComponent;
+		if (PCGComponentBeingInspected.IsValid())
+		{
+			PCGComponentBeingInspected->DisableInspection();
+		}
 
-		OnDebugObjectChangedDelegate.Broadcast(PCGComponentBeingDebugged.Get());
+		PCGComponentBeingInspected = InPCGComponent;
+		OnInspectedComponentChangedDelegate.Broadcast(PCGComponentBeingInspected.Get());
 
-		if (PCGComponentBeingDebugged.IsValid())
+		if (PCGComponentBeingInspected.IsValid())
 		{
-			// We need to force generation so that we create debug data
-			PCGComponentBeingDebugged->GenerateLocal(/*bForce=*/true);
+			if (!PCGComponentBeingInspected->IsInspecting())
+			{
+				// We need to force generation so that we create debug data
+				PCGComponentBeingInspected->GenerateLocal(/*bForce=*/true);
+			}
+			
+			PCGComponentBeingInspected->EnableInspection();
 		}
 
 		for (UEdGraphNode* Node : PCGEditorGraph->Nodes)
@@ -182,6 +191,12 @@ void FPCGEditor::SetPCGComponentBeingDebugged(UPCGComponent* InPCGComponent)
 			}
 		}
 	}
+
+	if(InPCGStack != StackBeingInspected)
+	{
+		StackBeingInspected = InPCGStack;
+		OnInspectedStackChangedDelegate.Broadcast(StackBeingInspected);
+	}
 }
 
 void FPCGEditor::JumpToNode(const UEdGraphNode* InNode)
@@ -632,19 +647,19 @@ void FPCGEditor::OnDeterminismNodeTest()
 
 bool FPCGEditor::CanRunDeterminismGraphTest() const
 {
-	return PCGEditorGraph && PCGComponentBeingDebugged.IsValid();
+	return PCGEditorGraph && PCGComponentBeingInspected.IsValid();
 }
 
 void FPCGEditor::OnDeterminismGraphTest()
 {
 	check(GraphEditorWidget.IsValid());
 
-	if (!DeterminismWidget.IsValid() || !DeterminismWidget->WidgetIsConstructed() || !PCGGraphBeingEdited || !PCGComponentBeingDebugged.IsValid())
+	if (!DeterminismWidget.IsValid() || !DeterminismWidget->WidgetIsConstructed() || !PCGGraphBeingEdited || !PCGComponentBeingInspected.IsValid())
 	{
 		return;
 	}
 
-	if (PCGComponentBeingDebugged->GetGraph() != PCGGraphBeingEdited)
+	if (PCGComponentBeingInspected->GetGraph() != PCGGraphBeingEdited)
 	{
 		// TODO: Should we alert the user more directly or disable this altogether?
 		UE_LOG(LogPCGEditor, Warning, TEXT("Running Determinism on a PCG Component with different/no attached PCG Graph"));
@@ -659,9 +674,9 @@ void FPCGEditor::OnDeterminismGraphTest()
 	TSharedPtr<FDeterminismTestResult> TestResult = MakeShared<FDeterminismTestResult>();
 	TestResult->TestResultTitle = TEXT("Full Graph Test");
 	TestResult->TestResultName = PCGGraphBeingEdited->GetName();
-	TestResult->Seed = PCGComponentBeingDebugged->Seed;
+	TestResult->Seed = PCGComponentBeingInspected->Seed;
 
-	PCGDeterminismTests::RunDeterminismTest(PCGGraphBeingEdited, PCGComponentBeingDebugged.Get(), *TestResult);
+	PCGDeterminismTests::RunDeterminismTest(PCGGraphBeingEdited, PCGComponentBeingInspected.Get(), *TestResult);
 
 	DeterminismWidget->AddItem(TestResult);
 	DeterminismWidget->AddDetailsColumn();
diff --git a/Engine/Plugins/Experimental/PCG/Source/PCGEditor/Private/PCGEditor.h b/Engine/Plugins/Experimental/PCG/Source/PCGEditor/Private/PCGEditor.h
--- a/Engine/Plugins/Experimental/PCG/Source/PCGEditor/Private/PCGEditor.h
+++ b/Engine/Plugins/Experimental/PCG/Source/PCGEditor/Private/PCGEditor.h
@@ -2,9 +2,10 @@
 
 #pragma once
 
+#include "PCGContext.h"
+#include "Graph/PCGStackContext.h"
 
 #include "EditorUndoClient.h"
-#include "PCGContext.h"
 #include "Toolkits/AssetEditorToolkit.h"
 
 class FSpawnTabArgs;
@@ -28,7 +29,8 @@ class UPCGEditorGraphNodeBase;
 class UPCGGraph;
 class UPCGNode;
 
-DECLARE_MULTICAST_DELEGATE_OneParam(FOnDebugObjectChanged, UPCGComponent*);
+DECLARE_MULTICAST_DELEGATE_OneParam(FOnInspectedComponentChanged, UPCGComponent*);
+DECLARE_MULTICAST_DELEGATE_OneParam(FOnInspectedStackChanged, const FPCGStack&);
 DECLARE_MULTICAST_DELEGATE_OneParam(FOnInspectedNodeChanged, UPCGEditorGraphNodeBase*);
 
 class FPCGEditor : public FAssetEditorToolkit, public FGCObject, public FSelfRegisteringEditorUndoClient
@@ -40,11 +42,14 @@ public:
 	/** Get the PCG graph being edited */
 	UPCGEditorGraph* GetPCGEditorGraph();
 
-	/** Sets the PCG component we want to debug */
-	void SetPCGComponentBeingDebugged(UPCGComponent* InPCGComponent);
+	/** Sets the PCG component and stack that we want to inspect */
+	void SetComponentAndStackBeingInspected(UPCGComponent* InPCGComponent, const FPCGStack& InPCGStack);
 
 	/** Gets the PCG component we are debugging */
-	UPCGComponent* GetPCGComponentBeingDebugged() const { return PCGComponentBeingDebugged.Get(); }
+	UPCGComponent* GetPCGComponentBeingInspected() const { return PCGComponentBeingInspected.Get(); }
+	
+	/** Gets the PCG stack we are inspecting */
+	const FPCGStack& GetStackBeingInspected() const { return StackBeingInspected; }
 
 	/** Focus the graph view on a specific node */
 	void JumpToNode(const UEdGraphNode* InNode);
@@ -77,10 +82,8 @@ public:
 	virtual void InitToolMenuContext(FToolMenuContext& MenuContext) override;
 	// ~End FAssetEditorToolkit interface
 
-	// Methods used by schema actions
-	UEdGraphNode* AddNode(UPCGSettings* InSettings, bool bIsInstance, const FVector2D& InLocation);
-
-	FOnDebugObjectChanged OnDebugObjectChangedDelegate;
+	FOnInspectedComponentChanged OnInspectedComponentChangedDelegate;
+	FOnInspectedStackChanged OnInspectedStackChangedDelegate;
 	FOnInspectedNodeChanged OnInspectedNodeChangedDelegate;
 
 protected:
@@ -287,6 +290,7 @@ private:
 	TObjectPtr<UPCGGraph> PCGGraphBeingEdited = nullptr;
 	UPCGEditorGraph* PCGEditorGraph = nullptr;
 
-	TWeakObjectPtr<UPCGComponent> PCGComponentBeingDebugged;
+	TWeakObjectPtr<UPCGComponent> PCGComponentBeingInspected;
+	FPCGStack StackBeingInspected;
 	UPCGEditorGraphNodeBase* PCGGraphNodeBeingInspected = nullptr;
 };
diff --git a/Engine/Plugins/Experimental/PCG/Source/PCGEditor/Private/PCGEditorGraphNodeBase.cpp b/Engine/Plugins/Experimental/PCG/Source/PCGEditor/Private/PCGEditorGraphNodeBase.cpp
--- a/Engine/Plugins/Experimental/PCG/Source/PCGEditor/Private/PCGEditorGraphNodeBase.cpp
+++ b/Engine/Plugins/Experimental/PCG/Source/PCGEditor/Private/PCGEditorGraphNodeBase.cpp
@@ -299,7 +299,7 @@ void UPCGEditorGraphNodeBase::UpdateErrorsAndWarnings()
 		{
 			const UPCGEditorGraph* EditorGraph = CastChecked<UPCGEditorGraph>(GetGraph());
 			const FPCGEditor* Editor = (EditorGraph && EditorGraph->GetEditor().IsValid()) ? EditorGraph->GetEditor().Pin().Get() : nullptr;
-			ComponentBeingDebugged = Editor ? Editor->GetPCGComponentBeingDebugged() : nullptr;
+			ComponentBeingDebugged = Editor ? Editor->GetPCGComponentBeingInspected() : nullptr;
 		}
 
 		const bool bOldHasCompilerMessage = bHasCompilerMessage;
diff --git a/Engine/Plugins/Experimental/PCG/Source/PCGEditor/Private/SPCGEditorGraphAttributeListView.cpp b/Engine/Plugins/Experimental/PCG/Source/PCGEditor/Private/SPCGEditorGraphAttributeListView.cpp
--- a/Engine/Plugins/Experimental/PCG/Source/PCGEditor/Private/SPCGEditorGraphAttributeListView.cpp
+++ b/Engine/Plugins/Experimental/PCG/Source/PCGEditor/Private/SPCGEditorGraphAttributeListView.cpp
@@ -187,7 +187,7 @@ SPCGEditorGraphAttributeListView::~SPCGEditorGraphAttributeListView()
 {
 	if (PCGEditorPtr.IsValid())
 	{
-		PCGEditorPtr.Pin()->OnDebugObjectChangedDelegate.RemoveAll(this);
+		PCGEditorPtr.Pin()->OnInspectedStackChangedDelegate.RemoveAll(this);
 		PCGEditorPtr.Pin()->OnInspectedNodeChangedDelegate.RemoveAll(this);
 	}
 }
@@ -197,7 +197,8 @@ void SPCGEditorGraphAttributeListView::Construct(const FArguments& InArgs, TShar
 	PCGEditorPtr = InPCGEditor;
 	SortingColumn = PCGEditorGraphAttributeListView::NAME_IndexColumn;
 
-	PCGEditorPtr.Pin()->OnDebugObjectChangedDelegate.AddSP(this, &SPCGEditorGraphAttributeListView::OnDebugObjectChanged);
+	PCGEditorPtr.Pin()->OnInspectedComponentChangedDelegate.AddSP(this, &SPCGEditorGraphAttributeListView::OnInspectedComponentChanged);
+	PCGEditorPtr.Pin()->OnInspectedStackChangedDelegate.AddSP(this, &SPCGEditorGraphAttributeListView::OnInspectedStackChanged);
 	PCGEditorPtr.Pin()->OnInspectedNodeChangedDelegate.AddSP(this, &SPCGEditorGraphAttributeListView::OnInspectedNodeChanged);
 
 	ListViewHeader = CreateHeaderRowWidget();
@@ -279,7 +280,7 @@ void SPCGEditorGraphAttributeListView::Construct(const FArguments& InArgs, TShar
 				.Text(PCGEditorGraphAttributeListView::NoNodeInspectedText)
 				.ToolTipText(PCGEditorGraphAttributeListView::NoNodeInspectedToolTip)
 			]
-			+ SHorizontalBox::Slot()
+			+SHorizontalBox::Slot()
 			.FillWidth(1.0f)
 			+SHorizontalBox::Slot()
 			.AutoWidth()
@@ -323,35 +324,51 @@ void SPCGEditorGraphAttributeListView::Construct(const FArguments& InArgs, TShar
 	];
 }
 
+void SPCGEditorGraphAttributeListView::Tick(const FGeometry& AllottedGeometry, const double InCurrentTime, const float InDeltaTime)
+{
+	SCompoundWidget::Tick(AllottedGeometry, InCurrentTime, InDeltaTime);
+
+	if (bNeedsRefresh)
+	{
+		bNeedsRefresh = false;
+
+		RefreshDataComboBox();
+		RefreshAttributeList();
+	}
+}
+
 TSharedRef<SHeaderRow> SPCGEditorGraphAttributeListView::CreateHeaderRowWidget() const
 {
 	return SNew(SHeaderRow);
 }
 
-void SPCGEditorGraphAttributeListView::OnDebugObjectChanged(UPCGComponent* InPCGComponent)
+void SPCGEditorGraphAttributeListView::OnInspectedComponentChanged(UPCGComponent* InPCGComponent)
 {
 	if (PCGComponent.IsValid())
 	{
 		PCGComponent->OnPCGGraphGeneratedDelegate.RemoveAll(this);
 		PCGComponent->OnPCGGraphCleanedDelegate.RemoveAll(this);
-		PCGComponent->DisableInspection();
 	}
 
 	PCGComponent = InPCGComponent;
 
 	if (PCGComponent.IsValid())
 	{
-		PCGComponent->EnableInspection();
 		PCGComponent->OnPCGGraphGeneratedDelegate.AddSP(this, &SPCGEditorGraphAttributeListView::OnGenerateUpdated);
 		PCGComponent->OnPCGGraphCleanedDelegate.AddSP(this, &SPCGEditorGraphAttributeListView::OnGenerateUpdated);
 	}
 	else
 	{
-		RefreshDataComboBox();
-		RefreshAttributeList();
+		// Refresh if PCGComponent is cleared since we wont get a refresh after generate/cleaned
+		RequestRefresh();
 	}
 }
 
+void SPCGEditorGraphAttributeListView::OnInspectedStackChanged(const FPCGStack& InPCGStack)
+{
+	RequestRefresh();
+}
+
 void SPCGEditorGraphAttributeListView::OnInspectedNodeChanged(UPCGEditorGraphNodeBase* InPCGEditorGraphNode)
 {
 	if (PCGEditorGraphNode == InPCGEditorGraphNode)
@@ -371,47 +388,59 @@ void SPCGEditorGraphAttributeListView::OnInspectedNodeChanged(UPCGEditorGraphNod
 		NodeNameTextBlock->SetText(PCGEditorGraphAttributeListView::NoNodeInspectedText);
 		NodeNameTextBlock->SetToolTipText(PCGEditorGraphAttributeListView::NoNodeInspectedToolTip);
 	}
-	
-	RefreshDataComboBox();
-	RefreshAttributeList();
+
+	RequestRefresh();
 }
 
 void SPCGEditorGraphAttributeListView::OnGenerateUpdated(UPCGComponent* /*InPCGComponent*/)
 {
-	RefreshDataComboBox();
-	RefreshAttributeList();
+	RequestRefresh();
 }
 
-void SPCGEditorGraphAttributeListView::RefreshAttributeList()
+const FPCGDataCollection* SPCGEditorGraphAttributeListView::GetInspectionData()
 {
-	HiddenAttributes = ListViewHeader->GetHiddenColumnIds();
-	
-	// Swapping to an empty item list to force a widget clear, otherwise the widgets will try to update during add column and access invalid data
-	static const TArray<PCGListviewItemPtr> EmptyList;
-	ListView->SetItemsSource(&EmptyList);
-
-	PCGColumnData.Empty();
-	ListViewItems.Empty();
-	ListViewHeader->ClearColumns();
-	InfoTextBlock->SetText(FText::GetEmpty());
-	
 	if (!PCGComponent.IsValid())
 	{
-		return;
+		return nullptr;
 	}
 
 	if (!PCGEditorGraphNode.IsValid())
 	{
-		return;
+		return nullptr;
 	}
 
 	const UPCGNode* PCGNode = PCGEditorGraphNode->GetPCGNode();
 	if (!PCGNode)
 	{
-		return;
+		return nullptr;
+	}
+
+	const TSharedPtr<FPCGEditor> PCGEditor = PCGEditorPtr.Pin();
+	const FPCGStack& PCGStack = PCGEditor->GetStackBeingInspected();
+
+	FString ResourcePath;
+	if (!PCGStack.CreateStackFramePath(ResourcePath, PCGNode))
+	{
+		return nullptr;
 	}
 
-	const FPCGDataCollection* InspectionData = PCGComponent->GetInspectionData(PCGNode);
+	return PCGComponent->GetInspectionData(ResourcePath);
+}
+
+void SPCGEditorGraphAttributeListView::RefreshAttributeList()
+{
+	HiddenAttributes = ListViewHeader->GetHiddenColumnIds();
+
+	// Swapping to an empty item list to force a widget clear, otherwise the widgets will try to update during add column and access invalid data
+	static const TArray<PCGListviewItemPtr> EmptyList;
+	ListView->SetItemsSource(&EmptyList);
+
+	PCGColumnData.Empty();
+	ListViewItems.Empty();
+	ListViewHeader->ClearColumns();
+	InfoTextBlock->SetText(FText::GetEmpty());
+
+	const FPCGDataCollection* InspectionData = GetInspectionData();
 	if (!InspectionData)
 	{
 		return;
@@ -497,23 +526,7 @@ void SPCGEditorGraphAttributeListView::RefreshDataComboBox()
 	DataComboBox->ClearSelection();
 	DataComboBox->RefreshOptions();
 
-	if (!PCGComponent.IsValid())
-	{
-		return;
-	}
-
-	if (!PCGEditorGraphNode.IsValid())
-	{
-		return;
-	}
-
-	const UPCGNode* PCGNode = PCGEditorGraphNode->GetPCGNode();
-	if (!PCGNode)
-	{
-		return;
-	}
-
-	const FPCGDataCollection* InspectionData = PCGComponent->GetInspectionData(PCGNode);
+	const FPCGDataCollection* InspectionData = GetInspectionData();
 	if (!InspectionData)
 	{
 		return;
@@ -864,7 +877,7 @@ void SPCGEditorGraphAttributeListView::RefreshSorting()
 						if (PCG::Private::MetadataTraits<ValueType>::Equal(LHSValue, RHSValue))
 						{
 							return LHS->Index > RHS->Index;
-						}	
+						}
 
 						return PCG::Private::MetadataTraits<ValueType>::Greater(LHSValue, RHSValue);
 					};
@@ -976,11 +989,11 @@ void SPCGEditorGraphAttributeListView::AddPointDataColumns(const UPCGPointData*
 	{
 		AddColumn(InPCGPointData, PCGEditorGraphAttributeListView::NAME_PointMetadataEntry, PCGEditorGraphAttributeListView::TEXT_PointMetadataEntryLabel);
 		AddColumn(InPCGPointData, PCGEditorGraphAttributeListView::NAME_PointMetadataEntryParent, PCGEditorGraphAttributeListView::TEXT_PointMetadataEntryParentLabel);
-	}	
+	}
 }
 
 void SPCGEditorGraphAttributeListView::AddMetadataColumn(const UPCGData* InPCGData, const FName& InColumnId, EPCGMetadataTypes InMetadataType, const TCHAR* PostFix)
-{	
+{
 	FString OriginalColumnIdString = InColumnId.ToString();
 	if (PostFix)
 	{
@@ -1010,7 +1023,7 @@ void SPCGEditorGraphAttributeListView::AddMetadataColumn(const UPCGData* InPCGDa
 
 	if (ColumnLabel.IsEmpty())
 	{
-		ColumnLabel = FText::FromName(ColumnId);	
+		ColumnLabel = FText::FromName(ColumnId);
 	}
 	
 	float ColumnWidth = 0.0f;
diff --git a/Engine/Plugins/Experimental/PCG/Source/PCGEditor/Private/SPCGEditorGraphAttributeListView.h b/Engine/Plugins/Experimental/PCG/Source/PCGEditor/Private/SPCGEditorGraphAttributeListView.h
--- a/Engine/Plugins/Experimental/PCG/Source/PCGEditor/Private/SPCGEditorGraphAttributeListView.h
+++ b/Engine/Plugins/Experimental/PCG/Source/PCGEditor/Private/SPCGEditorGraphAttributeListView.h
@@ -2,6 +2,7 @@
 
 #pragma once
 
+#include "Graph/PCGStackContext.h"
 #include "Metadata/Accessors/IPCGAttributeAccessor.h"
 #include "Metadata/Accessors/PCGAttributeAccessorKeys.h"
 
@@ -21,6 +22,7 @@ class UPCGData;
 class UPCGEditorGraphNodeBase;
 class UPCGMetadata;
 class UPCGParamData;
+struct FPCGDataCollection;
 struct FPCGPoint;
 enum class EPCGMetadataTypes : uint8;
 
@@ -78,14 +80,21 @@ public:
 
 	void Construct(const FArguments& InArgs, TSharedPtr<FPCGEditor> InPCGEditor);
 
+	virtual void Tick(const FGeometry& AllottedGeometry, const double InCurrentTime, const float InDeltaTime) override;
+
+	void RequestRefresh() { bNeedsRefresh = true; }
+	
 private:
 	TSharedRef<SHeaderRow> CreateHeaderRowWidget() const;
 
-	void OnDebugObjectChanged(UPCGComponent* InPCGComponent);
+	void OnInspectedComponentChanged(UPCGComponent* InPCGComponent);
+	void OnInspectedStackChanged(const FPCGStack& InPCGStack);
 	void OnInspectedNodeChanged(UPCGEditorGraphNodeBase* InPCGEditorGraphNode);
 
 	void OnGenerateUpdated(UPCGComponent* InPCGComponent);
 
+	const FPCGDataCollection* GetInspectionData();
+	
 	void RefreshAttributeList();
 	void RefreshDataComboBox();
 
@@ -140,4 +149,6 @@ private:
 
 	FName SortingColumn = NAME_None;
 	EColumnSortMode::Type SortMode = EColumnSortMode::Type::Ascending;
+
+	bool bNeedsRefresh = false;
 };
diff --git a/Engine/Plugins/Experimental/PCG/Source/PCGEditor/Private/SPCGEditorGraphDebugObjectWidget.cpp b/Engine/Plugins/Experimental/PCG/Source/PCGEditor/Private/SPCGEditorGraphDebugObjectWidget.cpp
--- a/Engine/Plugins/Experimental/PCG/Source/PCGEditor/Private/SPCGEditorGraphDebugObjectWidget.cpp
+++ b/Engine/Plugins/Experimental/PCG/Source/PCGEditor/Private/SPCGEditorGraphDebugObjectWidget.cpp
@@ -11,39 +11,27 @@
 #include "UnrealEdGlobals.h"
 #include "Editor/UnrealEdEngine.h"
 #include "Widgets/Input/SComboBox.h"
+#include "Widgets/Input/SComboButton.h"
 
 #define LOCTEXT_NAMESPACE "PCGEditorGraphDebugObjectWidget"
 
-namespace PCGEditorGraphDebugObjectWidget
+FPCGEditorGraphDebugObjectInstance::FPCGEditorGraphDebugObjectInstance()
+	: Label(LOCTEXT("NoDebugObject", "No debug object selected"))
 {
-	const FString SelectionString = TEXT(" (selected)");
-	const FString SeparatorString = TEXT(" / ");
 }
 
-FPCGEditorGraphDebugObjectInstance::FPCGEditorGraphDebugObjectInstance(TWeakObjectPtr<UPCGComponent> InPCGComponent)
+FPCGEditorGraphDebugObjectInstance::FPCGEditorGraphDebugObjectInstance(TWeakObjectPtr<UPCGComponent> InPCGComponent, const FPCGStack& InPCGStack)
 	: PCGComponent(InPCGComponent)
+	, PCGStack(InPCGStack)
 {
-	SetLabelFromPCGComponent(InPCGComponent);
-}
+	FString PathLabel = InPCGComponent->GetOwner()->GetActorNameOrLabel() / InPCGComponent->GetName();
 
-void FPCGEditorGraphDebugObjectInstance::SetLabelFromPCGComponent(TWeakObjectPtr<UPCGComponent> InPCGComponent)
-{
-	check(InPCGComponent.IsValid());
-	const AActor* Actor = InPCGComponent->GetOwner();
-	check(Actor);
-	FString ActorNameOrLabel = Actor->GetActorNameOrLabel();
-	if (Actor->IsSelected())
-	{
-		ActorNameOrLabel.Append(PCGEditorGraphDebugObjectWidget::SelectionString);
-	}
-			
-	FString ComponentName = InPCGComponent->GetFName().ToString();
-	if (InPCGComponent->IsSelected())
+	for (const FPCGStackFrame& StackFrame : PCGStack.GetStackFrames())
 	{
-		ComponentName.Append(PCGEditorGraphDebugObjectWidget::SelectionString);
+		PathLabel /= StackFrame.Object->GetName();
 	}
 
-	Label = ActorNameOrLabel + PCGEditorGraphDebugObjectWidget::SeparatorString + ComponentName;
+	Label = FText::FromString(PathLabel);
 }
 
 void SPCGEditorGraphDebugObjectWidget::Construct(const FArguments& InArgs, TSharedPtr<FPCGEditor> InPCGEditor)
@@ -57,13 +45,13 @@ void SPCGEditorGraphDebugObjectWidget::Construct(const FArguments& InArgs, TShar
 		LOCTEXT("SetDebugObject", "Set debug object from Level Editor selection."),
 		TAttribute<bool>::Create(TAttribute<bool>::FGetter::CreateSP(this, &SPCGEditorGraphDebugObjectWidget::IsSetDebugObjectFromSelectionButtonEnabled))
 	);
-	
+
 	const TSharedRef<SWidget> BrowseButton = PropertyCustomizationHelpers::MakeBrowseButton(
 		FSimpleDelegate::CreateSP(this, &SPCGEditorGraphDebugObjectWidget::SelectedDebugObject_OnClicked),
 		LOCTEXT("DebugSelectActor", "Select and frame the debug actor in the Level Editor."),
 		TAttribute<bool>::Create(TAttribute<bool>::FGetter::CreateSP(this, &SPCGEditorGraphDebugObjectWidget::IsSelectDebugObjectButtonEnabled))
 	);
-	
+
 	DebugObjectsComboBox = SNew(SComboBox<TSharedPtr<FPCGEditorGraphDebugObjectInstance>>)
 		.OptionsSource(&DebugObjects)
 		.InitiallySelectedItem(DebugObjects[0])
@@ -105,7 +93,7 @@ void SPCGEditorGraphDebugObjectWidget::OnComboBoxOpening()
 {
 	DebugObjects.Empty();
 	DebugObjectsComboBox->RefreshOptions();
-	
+
 	const UPCGGraph* PCGGraph = GetPCGGraph();
 	if (!PCGGraph)
 	{
@@ -120,7 +108,7 @@ void SPCGEditorGraphDebugObjectWidget::OnComboBoxOpening()
 	{
 		DebugObjectsComboBox->SetSelectedItem(DebugObjects[0]);
 	}
-	
+
 	TArray<UObject*> PCGComponents;
 	GetObjectsOfClass(UPCGComponent::StaticClass(), PCGComponents, /*bIncludeDerivedClasses=*/ true);
 	for (UObject* PCGComponentObject : PCGComponents)
@@ -129,13 +117,13 @@ void SPCGEditorGraphDebugObjectWidget::OnComboBoxOpening()
 		{
 			continue;
 		}
-		
+
 		UPCGComponent* PCGComponent = Cast<UPCGComponent>(PCGComponentObject);
 		if (!PCGComponent)
 		{
 			continue;
 		}
-		
+
 		const AActor* Actor = PCGComponent->GetOwner();
 		if (!Actor)
 		{
@@ -148,14 +136,23 @@ void SPCGEditorGraphDebugObjectWidget::OnComboBoxOpening()
 			continue;
 		}
 
-		if (PCGComponentGraph == PCGGraph)
-		{
-			TSharedPtr<FPCGEditorGraphDebugObjectInstance> DebugInstance = MakeShared<FPCGEditorGraphDebugObjectInstance>(PCGComponent);
-			DebugObjects.Add(DebugInstance);
+		FPCGStackContext StackContext = FPCGStackContext::CreateStackContextFromGraph(PCGComponentGraph);
 
-			if (SelectedItem.IsValid() && SelectedItem->GetPCGComponent() == PCGComponent)
+		for (const FPCGStack& Stack : StackContext.GetStacks())
+		{
+			const FPCGStackFrame& StackFrame = Stack.GetStackFrames().Top();
+			if (const UPCGGraph* StackGraph = Cast<const UPCGGraph>(StackFrame.Object))
 			{
-				DebugObjectsComboBox->SetSelectedItem(DebugInstance);
+				if (StackGraph == PCGGraph)
+				{
+					const TSharedPtr<FPCGEditorGraphDebugObjectInstance> DebugInstance = MakeShared<FPCGEditorGraphDebugObjectInstance>(PCGComponent, Stack);
+					DebugObjects.Add(DebugInstance);
+
+					if (SelectedItem.IsValid() && SelectedItem->GetPCGComponent() == PCGComponent && SelectedItem->GetStack() == Stack)
+					{
+						DebugObjectsComboBox->SetSelectedItem(DebugInstance);
+					}
+				}
 			}
 		}
 	}
@@ -166,7 +163,7 @@ void SPCGEditorGraphDebugObjectWidget::OnSelectionChanged(TSharedPtr<FPCGEditorG
 	if (NewSelection.IsValid())
 	{
 		UPCGComponent* PCGComponent = NewSelection->GetPCGComponent().Get();
-		PCGEditorPtr.Pin()->SetPCGComponentBeingDebugged(PCGComponent);
+		PCGEditorPtr.Pin()->SetComponentAndStackBeingInspected(PCGComponent, NewSelection->GetStack());
 	}
 }
 
@@ -184,7 +181,7 @@ UPCGGraph* SPCGEditorGraphDebugObjectWidget::GetPCGGraph() const
 	{
 		return nullptr;
 	}
-	
+
 	const UPCGEditorGraph* PCGEditorGraph = PCGEditorPtr.Pin()->GetPCGEditorGraph();
 	if (!PCGEditorGraph)
 	{
@@ -206,7 +203,7 @@ FText SPCGEditorGraphDebugObjectWidget::GetSelectedDebugObjectText() const
 
 void SPCGEditorGraphDebugObjectWidget::SelectedDebugObject_OnClicked() const
 {
-	if (UPCGComponent* PCGComponent = PCGEditorPtr.Pin()->GetPCGComponentBeingDebugged())
+	if (UPCGComponent* PCGComponent = PCGEditorPtr.Pin()->GetPCGComponentBeingInspected())
 	{
 		if (AActor* Actor = PCGComponent->GetOwner())
 		{
@@ -220,7 +217,7 @@ void SPCGEditorGraphDebugObjectWidget::SelectedDebugObject_OnClicked() const
 
 bool SPCGEditorGraphDebugObjectWidget::IsSelectDebugObjectButtonEnabled() const
 {
-	return PCGEditorPtr.IsValid() && (PCGEditorPtr.Pin()->GetPCGComponentBeingDebugged() != nullptr);
+	return PCGEditorPtr.IsValid() && (PCGEditorPtr.Pin()->GetPCGComponentBeingInspected() != nullptr);
 }
 
 void SPCGEditorGraphDebugObjectWidget::SetDebugObjectFromSelection_OnClicked()
@@ -230,13 +227,13 @@ void SPCGEditorGraphDebugObjectWidget::SetDebugObjectFromSelection_OnClicked()
 	{
 		return;
 	}
-	
+
 	USelection* SelectedActors = GEditor->GetSelectedActors();
 	if (!IsValid(SelectedActors))
 	{
 		return;
 	}
-	
+
 	for (FSelectionIterator It(*SelectedActors); It; ++It)
 	{
 		const AActor* SelectedActor = Cast<AActor>(*It);
@@ -244,24 +241,32 @@ void SPCGEditorGraphDebugObjectWidget::SetDebugObjectFromSelection_OnClicked()
 		{
 			continue;
 		}
-		
+
 		UPCGComponent* PCGComponent = SelectedActor->GetComponentByClass<UPCGComponent>();
 		if (!IsValid(PCGComponent))
 		{
 			continue;
 		}
-		
-		if (PCGComponent->GetGraph() == PCGGraph)
+
+		FPCGStackContext StackContext = FPCGStackContext::CreateStackContextFromGraph(PCGComponent->GetGraph());
+
+		for (const FPCGStack& Stack : StackContext.GetStacks())
 		{
-			DebugObjects.Empty();
-			DebugObjectsComboBox->RefreshOptions();
-
-			const TSharedPtr<FPCGEditorGraphDebugObjectInstance> DebugInstance = MakeShared<FPCGEditorGraphDebugObjectInstance>(PCGComponent);
-			DebugObjects.Add(DebugInstance);
-			DebugObjectsComboBox->SetSelectedItem(DebugInstance);
-			
-			PCGEditorPtr.Pin()->SetPCGComponentBeingDebugged(PCGComponent);
-			break;
+			const FPCGStackFrame& StackFrame = Stack.GetStackFrames().Top();
+			if (const UPCGGraph* StackGraph = Cast<const UPCGGraph>(StackFrame.Object))
+			{
+				if (StackGraph == PCGGraph)
+				{
+					DebugObjects.Empty();
+					DebugObjectsComboBox->RefreshOptions();
+
+					const TSharedPtr<FPCGEditorGraphDebugObjectInstance> DebugInstance = MakeShared<FPCGEditorGraphDebugObjectInstance>(PCGComponent, Stack);
+					DebugObjects.Add(DebugInstance);
+					DebugObjectsComboBox->SetSelectedItem(DebugInstance);
+					PCGEditorPtr.Pin()->SetComponentAndStackBeingInspected(PCGComponent, Stack);
+					break;
+				}
+			}
 		}
 	}
 }
@@ -273,13 +278,13 @@ bool SPCGEditorGraphDebugObjectWidget::IsSetDebugObjectFromSelectionButtonEnable
 	{
 		return false;
 	}
-	
+
 	USelection* SelectedActors = GEditor->GetSelectedActors();
 	if (!IsValid(SelectedActors))
 	{
 		return false;
 	}
-	
+
 	for (FSelectionIterator It(*SelectedActors); It; ++It)
 	{
 		const AActor* SelectedActor = Cast<AActor>(*It);
@@ -293,13 +298,21 @@ bool SPCGEditorGraphDebugObjectWidget::IsSetDebugObjectFromSelectionButtonEnable
 		{
 			continue;
 		}
-		
-		if (PCGComponent->GetGraph() == PCGGraph)
+
+		FPCGStackContext StackContext = FPCGStackContext::CreateStackContextFromGraph(PCGComponent->GetGraph());
+		for (const FPCGStack& Stack : StackContext.GetStacks())
 		{
-			return true;
+			const FPCGStackFrame& StackFrame = Stack.GetStackFrames().Top();
+			if (const UPCGGraph* StackGraph = Cast<const UPCGGraph>(StackFrame.Object))
+			{
+				if (StackGraph == PCGGraph)
+				{
+					return true;
+				}
+			}
 		}
 	}
-	
+
 	return false;
 }
 
diff --git a/Engine/Plugins/Experimental/PCG/Source/PCGEditor/Private/SPCGEditorGraphDebugObjectWidget.h b/Engine/Plugins/Experimental/PCG/Source/PCGEditor/Private/SPCGEditorGraphDebugObjectWidget.h
--- a/Engine/Plugins/Experimental/PCG/Source/PCGEditor/Private/SPCGEditorGraphDebugObjectWidget.h
+++ b/Engine/Plugins/Experimental/PCG/Source/PCGEditor/Private/SPCGEditorGraphDebugObjectWidget.h
@@ -2,6 +2,7 @@
 
 #pragma once
 
+#include "Graph/PCGStackContext.h"
 #include "Widgets/SCompoundWidget.h"
 
 namespace ESelectInfo { enum Type : int; }
@@ -14,24 +15,22 @@ class UPCGGraph;
 class FPCGEditorGraphDebugObjectInstance
 {
 public:
-	FPCGEditorGraphDebugObjectInstance() = default;
-	FPCGEditorGraphDebugObjectInstance(TWeakObjectPtr<UPCGComponent> InPCGComponent);
+	FPCGEditorGraphDebugObjectInstance();
+	FPCGEditorGraphDebugObjectInstance(TWeakObjectPtr<UPCGComponent> InPCGComponent, const FPCGStack& InPCGStack);
 
-	void SetLabelFromPCGComponent(TWeakObjectPtr<UPCGComponent> InPCGComponent);
+	FText GetDebugObjectText() const { return Label; }
 
-	FText GetDebugObjectText() const
-	{
-		return FText::FromString(Label);
-	}
-
-	TWeakObjectPtr<UPCGComponent> GetPCGComponent() const
-	{
-		return PCGComponent;
-	}
+	TWeakObjectPtr<UPCGComponent> GetPCGComponent() const { return PCGComponent; }
+	const FPCGStack& GetStack() const { return PCGStack; }
 	
 private:
+	FText Label;
+
+	/** Component containing the inspection cache */
 	TWeakObjectPtr<UPCGComponent> PCGComponent = nullptr;
-	FString Label = TEXT("No debug object selected");
+
+	/** Stack to identify graph or subgraph instance */
+	FPCGStack PCGStack;
 };
 
 
diff --git a/Engine/Plugins/Experimental/PCG/Source/PCGEditor/Private/SPCGEditorGraphLogView.cpp b/Engine/Plugins/Experimental/PCG/Source/PCGEditor/Private/SPCGEditorGraphLogView.cpp
--- a/Engine/Plugins/Experimental/PCG/Source/PCGEditor/Private/SPCGEditorGraphLogView.cpp
+++ b/Engine/Plugins/Experimental/PCG/Source/PCGEditor/Private/SPCGEditorGraphLogView.cpp
@@ -107,7 +107,7 @@ SPCGEditorGraphLogView::~SPCGEditorGraphLogView()
 {
 	if (PCGEditorPtr.IsValid())
 	{
-		PCGEditorPtr.Pin()->OnDebugObjectChangedDelegate.RemoveAll(this);
+		PCGEditorPtr.Pin()->OnInspectedComponentChangedDelegate.RemoveAll(this);
 	}
 }
 
@@ -119,9 +119,9 @@ void SPCGEditorGraphLogView::Construct(const FArguments& InArgs, TSharedPtr<FPCG
 	if (PCGEditor)
 	{
 		PCGEditorGraph = PCGEditor->GetPCGEditorGraph();
-		PCGComponent = PCGEditor->GetPCGComponentBeingDebugged();
+		PCGComponent = PCGEditor->GetPCGComponentBeingInspected();
 
-		PCGEditor->OnDebugObjectChangedDelegate.AddSP(this, &SPCGEditorGraphLogView::OnDebugObjectChanged);
+		PCGEditor->OnInspectedComponentChangedDelegate.AddSP(this, &SPCGEditorGraphLogView::OnDebugObjectChanged);
 	}
 
 	ListViewHeader = CreateHeaderRowWidget();
diff --git a/Engine/Plugins/Experimental/PCG/Source/PCGEditor/Private/SPCGEditorGraphProfilingView.cpp b/Engine/Plugins/Experimental/PCG/Source/PCGEditor/Private/SPCGEditorGraphProfilingView.cpp
--- a/Engine/Plugins/Experimental/PCG/Source/PCGEditor/Private/SPCGEditorGraphProfilingView.cpp
+++ b/Engine/Plugins/Experimental/PCG/Source/PCGEditor/Private/SPCGEditorGraphProfilingView.cpp
@@ -113,7 +113,7 @@ SPCGEditorGraphProfilingView::~SPCGEditorGraphProfilingView()
 {
 	if (PCGEditorPtr.IsValid())
 	{
-		PCGEditorPtr.Pin()->OnDebugObjectChangedDelegate.RemoveAll(this);
+		PCGEditorPtr.Pin()->OnInspectedComponentChangedDelegate.RemoveAll(this);
 	}
 }
 
@@ -125,9 +125,9 @@ void SPCGEditorGraphProfilingView::Construct(const FArguments& InArgs, TSharedPt
 	if (PCGEditor)
 	{
 		PCGEditorGraph = PCGEditor->GetPCGEditorGraph();
-		PCGComponent = PCGEditor->GetPCGComponentBeingDebugged();
+		PCGComponent = PCGEditor->GetPCGComponentBeingInspected();
 
-		PCGEditor->OnDebugObjectChangedDelegate.AddSP(this, &SPCGEditorGraphProfilingView::OnDebugObjectChanged);
+		PCGEditor->OnInspectedComponentChangedDelegate.AddSP(this, &SPCGEditorGraphProfilingView::OnDebugObjectChanged);
 	}
 
 	SortingColumn = PCGEditorGraphProfilingView::NAME_TotalExecutionTime;
