Replaced "Creates Artifacts" by "Is Cacheable" in PCG blueprints to be clearer.
Added Is Cacheable Override to allow more fine-grained control when needed.
Replaced a ref to name to a name in a BP call for simplicity.
Added ability to fully compute CRC out of a PCG blueprint element, needed in some instances to help caching
Improved the pivot update mechanism in PCG managed ISMs to properly work on reuse

diff --git a/Engine/Plugins/Experimental/PCG/Source/PCG/Private/Elements/PCGExecuteBlueprint.cpp b/Engine/Plugins/Experimental/PCG/Source/PCG/Private/Elements/PCGExecuteBlueprint.cpp
--- a/Engine/Plugins/Experimental/PCG/Source/PCG/Private/Elements/PCGExecuteBlueprint.cpp
+++ b/Engine/Plugins/Experimental/PCG/Source/PCG/Private/Elements/PCGExecuteBlueprint.cpp
@@ -168,6 +168,11 @@ EPCGSettingsType UPCGBlueprintElement::NodeTypeOverride_Implementation() const
 	return EPCGSettingsType::Blueprint;
 }
 
+bool UPCGBlueprintElement::IsCacheableOverride_Implementation() const
+{
+	return bIsCacheable;
+}
+
 TSet<FName> UPCGBlueprintElement::CustomInputLabels() const
 {
 	TSet<FName> Labels;
@@ -319,13 +324,11 @@ void UPCGBlueprintSettings::PostLoad()
 		BlueprintElementInstance->ConditionalPostLoad();
 		BlueprintElementInstance->SetFlags(RF_Transactional);
 #if WITH_EDITOR
-		BlueprintElementInstance->bCreatesArtifacts |= bCreatesArtifacts_DEPRECATED;
 		BlueprintElementInstance->bCanBeMultithreaded |= bCanBeMultithreaded_DEPRECATED;
 #endif
 	}
 
 #if WITH_EDITOR
-	bCreatesArtifacts_DEPRECATED = false;
 	bCanBeMultithreaded_DEPRECATED = false;
 #endif
 }
@@ -881,7 +884,21 @@ bool FPCGExecuteBlueprintElement::IsCacheable(const UPCGSettings* InSettings) co
 	const UPCGBlueprintSettings* BPSettings = Cast<const UPCGBlueprintSettings>(InSettings);
 	if (BPSettings && BPSettings->BlueprintElementInstance)
 	{
-		return !BPSettings->BlueprintElementInstance->bCreatesArtifacts;
+		return BPSettings->BlueprintElementInstance->IsCacheableOverride();
+	}
+	else
+	{
+		return false;
+	}
+}
+
+bool FPCGExecuteBlueprintElement::ShouldComputeFullOutputDataCrc(FPCGContext* Context) const
+{
+	check(Context);
+	const UPCGBlueprintSettings* BPSettings = Context->GetInputSettings<UPCGBlueprintSettings>();
+	if (BPSettings && BPSettings->BlueprintElementInstance)
+	{
+		return !IsCacheable(BPSettings) && BPSettings->BlueprintElementInstance->bComputeFullDataCrc;
 	}
 	else
 	{
diff --git a/Engine/Plugins/Experimental/PCG/Source/PCG/Private/Helpers/PCGActorHelpers.cpp b/Engine/Plugins/Experimental/PCG/Source/PCG/Private/Helpers/PCGActorHelpers.cpp
--- a/Engine/Plugins/Experimental/PCG/Source/PCG/Private/Helpers/PCGActorHelpers.cpp
+++ b/Engine/Plugins/Experimental/PCG/Source/PCG/Private/Helpers/PCGActorHelpers.cpp
@@ -77,10 +77,6 @@ UPCGManagedISMComponent* UPCGActorHelpers::GetOrCreateManagedISMC(AActor* InTarg
 			TRACE_CPUPROFILER_EVENT_SCOPE(UPCGActorHelpers::GetOrCreateManagedISMC::MarkAsUsed);
 			MatchingResource->MarkAsUsed();
 
-			// Reset the rotation/scale to be identity otherwise if the root component transform has changed, the final transform will be wrong
-			MatchingResource->GetComponent()->SetWorldRotation(FQuat::Identity);
-			MatchingResource->GetComponent()->SetWorldScale3D(FVector::OneVector);
-
 			return MatchingResource;
 		}
 	}
@@ -125,6 +121,7 @@ UPCGManagedISMComponent* UPCGActorHelpers::GetOrCreateManagedISMC(AActor* InTarg
 	UPCGManagedISMComponent* Resource = NewObject<UPCGManagedISMComponent>(InSourceComponent);
 	Resource->SetComponent(ISMC);
 	Resource->SetDescriptor(InParams.Descriptor);
+	Resource->SetRootLocation(InTargetActor->GetRootComponent()->GetComponentLocation());
 	Resource->SetSettingsUID(SettingsUID);
 	InSourceComponent->AddToManagedResources(Resource);
 
diff --git a/Engine/Plugins/Experimental/PCG/Source/PCG/Private/PCGData.cpp b/Engine/Plugins/Experimental/PCG/Source/PCG/Private/PCGData.cpp
--- a/Engine/Plugins/Experimental/PCG/Source/PCG/Private/PCGData.cpp
+++ b/Engine/Plugins/Experimental/PCG/Source/PCG/Private/PCGData.cpp
@@ -474,7 +474,7 @@ TArray<UPCGData*> UPCGDataFunctionLibrary::GetTypedInputsByPin(const FPCGDataCol
 	return GetTypedInputsByPinLabel(InCollection, InPin.Label, OutTaggedData, InDataTypeClass);
 }
 
-TArray<UPCGData*> UPCGDataFunctionLibrary::GetTypedInputsByPinLabel(const FPCGDataCollection& InCollection, const FName& InPinLabel, TArray<FPCGTaggedData>& OutTaggedData, TSubclassOf<UPCGData> InDataTypeClass)
+TArray<UPCGData*> UPCGDataFunctionLibrary::GetTypedInputsByPinLabel(const FPCGDataCollection& InCollection, FName InPinLabel, TArray<FPCGTaggedData>& OutTaggedData, TSubclassOf<UPCGData> InDataTypeClass)
 {
 	return GetInputsByPredicate(InCollection, OutTaggedData, [&InPinLabel, InDataTypeClass](const FPCGTaggedData& TaggedData)
 	{
diff --git a/Engine/Plugins/Experimental/PCG/Source/PCG/Private/PCGElement.cpp b/Engine/Plugins/Experimental/PCG/Source/PCG/Private/PCGElement.cpp
--- a/Engine/Plugins/Experimental/PCG/Source/PCG/Private/PCGElement.cpp
+++ b/Engine/Plugins/Experimental/PCG/Source/PCG/Private/PCGElement.cpp
@@ -207,7 +207,7 @@ void IPCGElement::PostExecute(FPCGContext* Context) const
 		// Some nodes benefit from computing an actual CRC from the data. This can halt the propagation of change/executions through the graph. For
 		// data like landscapes we will never have a full accurate data crc for it so we'll tend to assume changed which triggers downstream
 		// execution. Performing a detailed CRC of output data can detect real change in the data and halt the cascade of execution.
-		const bool bShouldComputeFullOutputDataCrc = ShouldComputeFullOutputDataCrc();
+		const bool bShouldComputeFullOutputDataCrc = ShouldComputeFullOutputDataCrc(Context);
 
 		// Compute Crc from output data
 		Context->OutputData.Crc = Context->OutputData.ComputeCrc(bShouldComputeFullOutputDataCrc);
diff --git a/Engine/Plugins/Experimental/PCG/Source/PCG/Private/PCGManagedResource.cpp b/Engine/Plugins/Experimental/PCG/Source/PCG/Private/PCGManagedResource.cpp
--- a/Engine/Plugins/Experimental/PCG/Source/PCG/Private/PCGManagedResource.cpp
+++ b/Engine/Plugins/Experimental/PCG/Source/PCG/Private/PCGManagedResource.cpp
@@ -348,13 +348,63 @@ bool UPCGManagedISMComponent::ReleaseIfUnused(TSet<TSoftObjectPtr<AActor>>& OutA
 
 void UPCGManagedISMComponent::ResetComponent()
 {
-	if (UInstancedStaticMeshComponent * ISMC = GetComponent())
+	if (UInstancedStaticMeshComponent* ISMC = GetComponent())
 	{
 		ISMC->ClearInstances();
 		ISMC->UpdateBounds();
 	}
 }
 
+void UPCGManagedISMComponent::MarkAsUsed()
+{
+	Super::MarkAsUsed();
+
+	if (UInstancedStaticMeshComponent* ISMC = GetComponent())
+	{
+		// Keep track of the current root location so if we reuse this later we are able to update this appropriately
+		if (USceneComponent* RootComponent = ISMC->GetAttachmentRoot())
+		{
+			bHasRootLocation = true;
+			RootLocation = RootComponent->GetComponentLocation();
+		}
+		else
+		{
+			bHasRootLocation = false;
+			RootLocation = FVector::ZeroVector;
+		}
+
+		// Reset the rotation/scale to be identity otherwise if the root component transform has changed, the final transform will be wrong
+		ISMC->SetWorldTransform(FTransform(FQuat::Identity, RootLocation, FVector::OneVector));
+	}
+}
+
+void UPCGManagedISMComponent::MarkAsReused()
+{
+	Super::MarkAsReused();
+
+	if (UInstancedStaticMeshComponent* ISMC = GetComponent())
+	{
+		// Reset the rotation/scale to be identity otherwise if the root component transform has changed, the final transform will be wrong
+		FVector TentativeRootLocation = RootLocation;
+
+		if (!bHasRootLocation)
+		{
+			if (USceneComponent* RootComponent = ISMC->GetAttachmentRoot())
+			{
+				TentativeRootLocation = RootComponent->GetComponentLocation();
+			}
+		}
+
+		ISMC->SetWorldTransform(FTransform(FQuat::Identity, TentativeRootLocation, FVector::OneVector));
+	}
+}
+
+void UPCGManagedISMComponent::SetRootLocation(const FVector& InRootLocation)
+{
+	bHasRootLocation = true;
+	RootLocation = InRootLocation;
+}
+
 UInstancedStaticMeshComponent* UPCGManagedISMComponent::GetComponent() const
 {
 	if (!CachedRawComponentPtr)
diff --git a/Engine/Plugins/Experimental/PCG/Source/PCG/Public/Elements/PCGCollapseElement.h b/Engine/Plugins/Experimental/PCG/Source/PCG/Public/Elements/PCGCollapseElement.h
--- a/Engine/Plugins/Experimental/PCG/Source/PCG/Public/Elements/PCGCollapseElement.h
+++ b/Engine/Plugins/Experimental/PCG/Source/PCG/Public/Elements/PCGCollapseElement.h
@@ -33,7 +33,7 @@ class FPCGCollapseElement : public FSimplePCGElement
 {
 public:
 	// Might be sampling spline/landscape or other external data, worth computing a full CRC in case we can halt change propagation/re-executions
-	virtual bool ShouldComputeFullOutputDataCrc() const override { return true; }
+	virtual bool ShouldComputeFullOutputDataCrc(FPCGContext* Context) const override { return true; }
 
 protected:
 	virtual bool ExecuteInternal(FPCGContext* Context) const;
diff --git a/Engine/Plugins/Experimental/PCG/Source/PCG/Public/Elements/PCGExecuteBlueprint.h b/Engine/Plugins/Experimental/PCG/Source/PCG/Public/Elements/PCGExecuteBlueprint.h
--- a/Engine/Plugins/Experimental/PCG/Source/PCG/Public/Elements/PCGExecuteBlueprint.h
+++ b/Engine/Plugins/Experimental/PCG/Source/PCG/Public/Elements/PCGExecuteBlueprint.h
@@ -81,6 +81,10 @@ public:
 	UFUNCTION(BlueprintNativeEvent, Category = "PCG|Node Customization")
 	EPCGSettingsType NodeTypeOverride() const;
 
+	/** Override for the IsCacheable node property when it depends on the settings in your node */
+	UFUNCTION(BlueprintNativeEvent, Category = "PCG|Execution")
+	bool IsCacheableOverride() const;
+
 	UFUNCTION(BlueprintImplementableEvent, Category = "PCG|Preconfigure Settings", meta = (ForceAsFunction))
 	void ApplyPreconfiguredSettings(UPARAM(ref) const FPCGPreConfiguredSettingsInfo& InPreconfigureInfo);
 
@@ -136,13 +140,18 @@ public:
 	FOnPCGBlueprintChanged OnBlueprintChangedDelegate;
 #endif
 
-	UPROPERTY(BlueprintReadWrite, EditDefaultsOnly, Category = Settings)
-	bool bCreatesArtifacts = false;
+	/** Controls whether results can be cached so we can bypass execution if the inputs & settings are the same in a subsequent execution.
+	* Note that if your node relies on data that is not directly tracked by PCG or creates any kind of artifact (adds components, creates actors, etc.) then it is not cacheable.
+	*/
+	UPROPERTY(BlueprintReadWrite, EditDefaultsOnly, AdvancedDisplay, Category = Settings)
+	bool bIsCacheable = false;
 
-	/** Not yet hooked up, currently work in progress. True if output data can be cached and reused if inputs do not change, false if blueprint should be executed every time. */
+	/** In cases where your node is non-cacheable but is likely to yield the same results on subsequent executions, this controls whether we will do a deep & computationally intensive CRC computation (true), 
+	* which will allow cache usage in downstream nodes in your graph, or, by default (false), a shallow but quick crc computation which will not be cache-friendly. */
 	UPROPERTY(BlueprintReadWrite, EditDefaultsOnly, AdvancedDisplay, Category = Settings)
-	bool bCacheable = false;
+	bool bComputeFullDataCrc = false;
 
+	/** Controls whether this node execution can be run from a non-game thread. This is not related to the Loop functions provided/implemented in this class, which should always run on any thread. */
 	UPROPERTY(BlueprintReadWrite, EditDefaultsOnly, Category = Settings)
 	bool bCanBeMultithreaded = false;
 
@@ -185,7 +194,7 @@ public:
 	UPROPERTY(BlueprintReadWrite, EditDefaultsOnly, Category = AssetInfo, AssetRegistrySearchable)
 	FText Description;
 
-	UPROPERTY(BlueprintReadWrite, EditDefaultsOnly, Category = "Settings|Advanced")
+	UPROPERTY(BlueprintReadWrite, EditDefaultsOnly, AdvancedDisplay, Category = "Settings")
 	int32 DependencyParsingDepth = 1;
 #endif
 
@@ -282,9 +291,6 @@ protected:
 	UPROPERTY(BlueprintReadOnly, EditAnywhere, Category = Settings)
 	bool bTrackActorsOnlyWithinBounds = false;
 
-	UPROPERTY()
-	bool bCreatesArtifacts_DEPRECATED = false;
-
 	UPROPERTY()
 	bool bCanBeMultithreaded_DEPRECATED = false;
 #endif
@@ -317,6 +323,7 @@ class FPCGExecuteBlueprintElement : public IPCGElement
 public:
 	virtual bool CanExecuteOnlyOnMainThread(FPCGContext* Context) const override;
 	virtual bool IsCacheable(const UPCGSettings* InSettings) const override;
+	virtual bool ShouldComputeFullOutputDataCrc(FPCGContext* Context) const override;
 
 protected:
 	virtual bool ExecuteInternal(FPCGContext* Context) const override;
diff --git a/Engine/Plugins/Experimental/PCG/Source/PCG/Public/Elements/PCGSplineSampler.h b/Engine/Plugins/Experimental/PCG/Source/PCG/Public/Elements/PCGSplineSampler.h
--- a/Engine/Plugins/Experimental/PCG/Source/PCG/Public/Elements/PCGSplineSampler.h
+++ b/Engine/Plugins/Experimental/PCG/Source/PCG/Public/Elements/PCGSplineSampler.h
@@ -166,7 +166,7 @@ class FPCGSplineSamplerElement : public FSimplePCGElement
 {
 public:
 	// Worth computing a full CRC in case we can halt change propagation/re-executions
-	virtual bool ShouldComputeFullOutputDataCrc() const override { return true; }
+	virtual bool ShouldComputeFullOutputDataCrc(FPCGContext* Context) const override { return true; }
 
 protected:
 	virtual bool ExecuteInternal(FPCGContext* Context) const override;
diff --git a/Engine/Plugins/Experimental/PCG/Source/PCG/Public/Elements/PCGSurfaceSampler.h b/Engine/Plugins/Experimental/PCG/Source/PCG/Public/Elements/PCGSurfaceSampler.h
--- a/Engine/Plugins/Experimental/PCG/Source/PCG/Public/Elements/PCGSurfaceSampler.h
+++ b/Engine/Plugins/Experimental/PCG/Source/PCG/Public/Elements/PCGSurfaceSampler.h
@@ -146,7 +146,7 @@ public:
 	virtual void GetDependenciesCrc(const FPCGDataCollection& InInput, const UPCGSettings* InSettings, UPCGComponent* InComponent, FPCGCrc& OutCrc) const override;
 	virtual FPCGContext* Initialize(const FPCGDataCollection& InputData, TWeakObjectPtr<UPCGComponent> SourceComponent, const UPCGNode* Node) override;
 	// Might be sampling landscape or other external data, worth computing a full CRC in case we can halt change propagation/re-executions
-	virtual bool ShouldComputeFullOutputDataCrc() const override { return true; }
+	virtual bool ShouldComputeFullOutputDataCrc(FPCGContext* Context) const override { return true; }
 
 protected:
 	virtual bool PrepareDataInternal(FPCGContext* InContext) const override;
diff --git a/Engine/Plugins/Experimental/PCG/Source/PCG/Public/Elements/PCGVolumeSampler.h b/Engine/Plugins/Experimental/PCG/Source/PCG/Public/Elements/PCGVolumeSampler.h
--- a/Engine/Plugins/Experimental/PCG/Source/PCG/Public/Elements/PCGVolumeSampler.h
+++ b/Engine/Plugins/Experimental/PCG/Source/PCG/Public/Elements/PCGVolumeSampler.h
@@ -63,7 +63,7 @@ class FPCGVolumeSamplerElement : public FSimplePCGElement
 {
 public:
 	// Might be sampling external data like brush, worth computing a full CRC in case we can halt change propagation/re-executions
-	virtual bool ShouldComputeFullOutputDataCrc() const override { return true; }
+	virtual bool ShouldComputeFullOutputDataCrc(FPCGContext* Context) const override { return true; }
 
 protected:
 	bool ExecuteInternal(FPCGContext* Context) const override;
diff --git a/Engine/Plugins/Experimental/PCG/Source/PCG/Public/PCGData.h b/Engine/Plugins/Experimental/PCG/Source/PCG/Public/PCGData.h
--- a/Engine/Plugins/Experimental/PCG/Source/PCG/Public/PCGData.h
+++ b/Engine/Plugins/Experimental/PCG/Source/PCG/Public/PCGData.h
@@ -202,7 +202,7 @@ public:
 
 	/** Gets all inputs of the given class type and on the given pin label, returning matching tagged data in the OutTaggedData value too */
 	UFUNCTION(BlueprintCallable, Category = Data, meta = (ScriptMethod, DeterminesOutputType = "InDataTypeClass"))
-	static TArray<UPCGData*> GetTypedInputsByPinLabel(const FPCGDataCollection& InCollection, const FName& InPinLabel, TArray<FPCGTaggedData>& OutTaggedData, TSubclassOf<UPCGData> InDataTypeClass = nullptr);
+	static TArray<UPCGData*> GetTypedInputsByPinLabel(const FPCGDataCollection& InCollection, FName InPinLabel, TArray<FPCGTaggedData>& OutTaggedData, TSubclassOf<UPCGData> InDataTypeClass = nullptr);
 
 	/** Gets all inputs of the given class type and having the provided tag, returning matching tagged data in the OutTaggedData value too */
 	UFUNCTION(BlueprintCallable, Category = Data, meta = (ScriptMethod, DeterminesOutputType = "InDataTypeClass"))
diff --git a/Engine/Plugins/Experimental/PCG/Source/PCG/Public/PCGElement.h b/Engine/Plugins/Experimental/PCG/Source/PCG/Public/PCGElement.h
--- a/Engine/Plugins/Experimental/PCG/Source/PCG/Public/PCGElement.h
+++ b/Engine/Plugins/Experimental/PCG/Source/PCG/Public/PCGElement.h
@@ -82,7 +82,7 @@ public:
 	virtual bool IsCacheable(const UPCGSettings* InSettings) const { return true; }
 
 	/** Whether to do a 'deep' fine-grained CRC of the output data to pass to downstream nodes. Can be expensive so should be used sparingly. */
-	virtual bool ShouldComputeFullOutputDataCrc() const { return false; }
+	virtual bool ShouldComputeFullOutputDataCrc(FPCGContext* Context) const { return false; }
 
 	/**
 	 * Calculate a Crc that provides a receipt for the input data that can be paired with output data from the cache. If any dependency (setting, node input or
diff --git a/Engine/Plugins/Experimental/PCG/Source/PCG/Public/PCGManagedResource.h b/Engine/Plugins/Experimental/PCG/Source/PCG/Public/PCGManagedResource.h
--- a/Engine/Plugins/Experimental/PCG/Source/PCG/Public/PCGManagedResource.h
+++ b/Engine/Plugins/Experimental/PCG/Source/PCG/Public/PCGManagedResource.h
@@ -118,6 +118,8 @@ public:
 	//~Begin UPCGManagedComponents interface
 	virtual void ResetComponent() override;
 	virtual bool SupportsComponentReset() const override{ return true; }
+	virtual void MarkAsUsed() override;
+	virtual void MarkAsReused() override;
 	virtual void ForgetComponent() override;
 	//~End UPCGManagedComponents interface
 
@@ -127,6 +129,8 @@ public:
 	void SetDescriptor(const FISMComponentDescriptor& InDescriptor);
 	const FISMComponentDescriptor& GetDescriptor() const { return Descriptor; }
 
+	void SetRootLocation(const FVector& InRootLocation);
+
 	uint64 GetSettingsUID() const { return SettingsUID; }
 	void SetSettingsUID(uint64 InSettingsUID) { SettingsUID = InSettingsUID; }
 
@@ -137,6 +141,12 @@ protected:
 	UPROPERTY()
 	FISMComponentDescriptor Descriptor;
 
+	UPROPERTY()
+	bool bHasRootLocation = false;
+
+	UPROPERTY()
+	FVector RootLocation = FVector::ZeroVector;
+
 	UPROPERTY(Transient)
 	uint64 SettingsUID = -1; // purposefully a value that will never happen in data
 
