[PCG] PCGEditor toolbar improvements
*Fixed an issue were the PCGEditor would register the toolbar multiple times
*Fixed an issue were the debug object dropdowns in open PCGEditor would clear it's selection when opening a new PCGEditor

diff --git a/Engine/Plugins/Experimental/PCG/Source/PCGEditor/Private/PCGEditor.cpp b/Engine/Plugins/Experimental/PCG/Source/PCGEditor/Private/PCGEditor.cpp
--- a/Engine/Plugins/Experimental/PCG/Source/PCGEditor/Private/PCGEditor.cpp
+++ b/Engine/Plugins/Experimental/PCG/Source/PCGEditor/Private/PCGEditor.cpp
@@ -331,10 +331,17 @@ FString FPCGEditor::GetWorldCentricTabPrefix() const
 void FPCGEditor::RegisterToolbar() const
 {
 	UToolMenus* ToolMenus = UToolMenus::Get();
+	UToolMenu* ToolBar;
 	FName ParentName;
 	const FName MenuName = GetToolMenuToolbarName(ParentName);
-
-	UToolMenu* ToolBar = ToolMenus->RegisterMenu(MenuName, ParentName, EMultiBoxType::ToolBar);
+	if (ToolMenus->IsMenuRegistered(MenuName))
+	{
+		ToolBar = ToolMenus->ExtendMenu(MenuName);
+	}
+	else
+	{
+		ToolBar = ToolMenus->RegisterMenu(MenuName, ParentName, EMultiBoxType::ToolBar);
+	}
 
 	const FPCGEditorCommands& PCGEditorCommands = FPCGEditorCommands::Get();
 	const FToolMenuInsert InsertAfterAssetSection("Asset", EToolMenuInsertType::After);
diff --git a/Engine/Plugins/Experimental/PCG/Source/PCGEditor/Private/SPCGEditorGraphDebugObjectWidget.cpp b/Engine/Plugins/Experimental/PCG/Source/PCGEditor/Private/SPCGEditorGraphDebugObjectWidget.cpp
--- a/Engine/Plugins/Experimental/PCG/Source/PCGEditor/Private/SPCGEditorGraphDebugObjectWidget.cpp
+++ b/Engine/Plugins/Experimental/PCG/Source/PCGEditor/Private/SPCGEditorGraphDebugObjectWidget.cpp
@@ -37,8 +37,16 @@ FPCGEditorGraphDebugObjectInstance::FPCGEditorGraphDebugObjectInstance(TWeakObje
 void SPCGEditorGraphDebugObjectWidget::Construct(const FArguments& InArgs, TSharedPtr<FPCGEditor> InPCGEditor)
 {
 	PCGEditorPtr = InPCGEditor;
-
-	DebugObjects.Add(MakeShared<FPCGEditorGraphDebugObjectInstance>());
+	
+	if (UPCGComponent* PCGComponent = InPCGEditor->GetPCGComponentBeingInspected())
+	{
+		const FPCGStack& PCGStack = InPCGEditor->GetStackBeingInspected();
+		DebugObjects.Add(MakeShared<FPCGEditorGraphDebugObjectInstance>(PCGComponent, PCGStack));
+	}
+	else
+	{
+		DebugObjects.Add(MakeShared<FPCGEditorGraphDebugObjectInstance>());
+	}
 
 	const TSharedRef<SWidget> SetButton = PropertyCustomizationHelpers::MakeUseSelectedButton(
 		FSimpleDelegate::CreateSP(this, &SPCGEditorGraphDebugObjectWidget::SetDebugObjectFromSelection_OnClicked),
