[PCG] Fix compact nodes not visually respecting 'force disabled' state properly.

Mirrors similar code in SGraphNode::UpdateGraphNode().

diff --git a/Engine/Plugins/Experimental/PCG/Source/PCGEditor/Private/SPCGEditorGraphNode.cpp b/Engine/Plugins/Experimental/PCG/Source/PCGEditor/Private/SPCGEditorGraphNode.cpp
--- a/Engine/Plugins/Experimental/PCG/Source/PCGEditor/Private/SPCGEditorGraphNode.cpp
+++ b/Engine/Plugins/Experimental/PCG/Source/PCGEditor/Private/SPCGEditorGraphNode.cpp
@@ -536,6 +536,8 @@ void SPCGEditorGraphNode::UpdateCompactNode()
 	TSharedPtr<SNodeTitle> NodeTitle = SNew(SNodeTitle, GraphNode);
 	TSharedRef<SOverlay> NodeOverlay = SNew(SOverlay);
 	
+	IconColor = FLinearColor::White;
+
 	// Add optional node specific widget to the overlay:
 	TSharedPtr<SWidget> OverlayWidget = GraphNode->CreateNodeImage();
 	if(OverlayWidget.IsValid())
@@ -568,6 +570,7 @@ void SPCGEditorGraphNode::UpdateCompactNode()
 					.TextStyle(FAppStyle::Get(), "Graph.Node.NodeTitle" )
 					.Text( NodeTitle.Get(), &SNodeTitle::GetHeadTitle )
 					.WrapTextAt(128.0f)
+					.ColorAndOpacity(this, &SGraphNode::GetNodeTitleIconColor)
 			]
 			+SVerticalBox::Slot()
 			.AutoHeight()
@@ -615,11 +618,13 @@ void SPCGEditorGraphNode::UpdateCompactNode()
 			[
 				SNew(SImage)
 				.Image( FAppStyle::GetBrush("Graph.VarNode.Body") )
+				.ColorAndOpacity(this, &SGraphNode::GetNodeBodyColor)
 			]
 			+ SOverlay::Slot()
 			[
 				SNew(SImage)
 				.Image( FAppStyle::GetBrush("Graph.VarNode.Gloss") )
+				.ColorAndOpacity(this, &SGraphNode::GetNodeBodyColor)
 			]
 			+SOverlay::Slot()
 			.Padding( FMargin(0,3) )
