Added getters for input/output pin properties by label
Renamed getters on data collection (Get Inputs By Pin -> Get Inputs By [REDACTED_PERSON], GetTaggedInputs -> GetInputsByTag, ...)

diff --git a/Engine/Plugins/Experimental/PCG/Config/BasePCG.ini b/Engine/Plugins/Experimental/PCG/Config/BasePCG.ini
--- a/Engine/Plugins/Experimental/PCG/Config/BasePCG.ini
+++ b/Engine/Plugins/Experimental/PCG/Config/BasePCG.ini
@@ -7,6 +7,11 @@
 +FunctionRedirects=(OldName="PCGBlueprintElement.LoopNTimes",NewName="PCGBlueprintElement.IterationLoop")
 +FunctionRedirects=(OldName="PCGBlueprintElement.InputLabels",NewName="PCGBlueprintElement.CustomInputLabels")
 +FunctionRedirects=(OldName="PCGBlueprintElement.OutputLabels",NewName="PCGBlueprintElement.CustomOutputLabels")
++FunctionRedirects=(OldName="PCGDataFunctionLibrary.GetInputsByPin",NewName="PCGDataFunctionLibrary.GetInputsByPinLabel")
++FunctionRedirects=(OldName="PCGDataFunctionLibrary.GetTaggedInputs",NewName="PCGDataFunctionLibrary.GetInputsByTag")
++FunctionRedirects=(OldName="PCGDataFunctionLibrary.GetParamsByPin",NewName="PCGDataFunctionLibrary.GetParamsByPinLabel")
++FunctionRedirects=(OldName="PCGDataFunctionLibrary.GetTaggedParams",NewName="PCGDataFunctionLibrary.GetParamsByTag")
+
 +PropertyRedirects=(OldName="PCGSplineSamplerSettings.Params",NewName="PCGSplineSamplerSettings.SamplerParams")
 +PropertyRedirects=(OldName="PCGProjectionSettings.Params",NewName="PCGProjectionSettings.ProjectionParams")
 +PropertyRedirects=(OldName="PCGBlueprintElement.IterationLoop.InA",NewName="PCGBlueprintElement.IterationLoop.OptionalA")
diff --git a/Engine/Plugins/Experimental/PCG/Source/PCG/Private/Elements/PCGExecuteBlueprint.cpp b/Engine/Plugins/Experimental/PCG/Source/PCG/Private/Elements/PCGExecuteBlueprint.cpp
--- a/Engine/Plugins/Experimental/PCG/Source/PCG/Private/Elements/PCGExecuteBlueprint.cpp
+++ b/Engine/Plugins/Experimental/PCG/Source/PCG/Private/Elements/PCGExecuteBlueprint.cpp
@@ -231,6 +231,38 @@ TArray<FPCGPinProperties> UPCGBlueprintElement::GetOutputPins() const
 	return CustomOutputPins;
 }
 
+bool UPCGBlueprintElement::GetInputPinByLabel(FName InPinLabel, FPCGPinProperties& OutFoundPin) const
+{
+	TArray<FPCGPinProperties> InputPinProperties = GetInputPins();
+	for (const FPCGPinProperties& InputPin : InputPinProperties)
+	{
+		if (InputPin.Label == InPinLabel)
+		{
+			OutFoundPin = InputPin;
+			return true;
+		}
+	}
+
+	OutFoundPin = FPCGPinProperties();
+	return false;;
+}
+
+bool UPCGBlueprintElement::GetOutputPinByLabel(FName InPinLabel, FPCGPinProperties& OutFoundPin) const
+{
+	TArray<FPCGPinProperties> OutputPinProperties = GetOutputPins();
+	for (const FPCGPinProperties& OutputPin : OutputPinProperties)
+	{
+		if (OutputPin.Label == InPinLabel)
+		{
+			OutFoundPin = OutputPin;
+			return true;
+		}
+	}
+
+	OutFoundPin = FPCGPinProperties();
+	return false;
+}
+
 int UPCGBlueprintElement::GetSeed(FPCGContext& InContext) const 
 {
 	return InContext.GetSeed();
diff --git a/Engine/Plugins/Experimental/PCG/Source/PCG/Private/PCGData.cpp b/Engine/Plugins/Experimental/PCG/Source/PCG/Private/PCGData.cpp
--- a/Engine/Plugins/Experimental/PCG/Source/PCG/Private/PCGData.cpp
+++ b/Engine/Plugins/Experimental/PCG/Source/PCG/Private/PCGData.cpp
@@ -506,12 +506,12 @@ TArray<FPCGTaggedData> UPCGDataFunctionLibrary::GetInputs(const FPCGDataCollecti
 	return InCollection.GetInputs();
 }
 
-TArray<FPCGTaggedData> UPCGDataFunctionLibrary::GetInputsByPin(const FPCGDataCollection& InCollection, const FName& InPinLabel)
+TArray<FPCGTaggedData> UPCGDataFunctionLibrary::GetInputsByPinLabel(const FPCGDataCollection& InCollection, const FName& InPinLabel)
 {
 	return InCollection.GetInputsByPin(InPinLabel);
 }
 
-TArray<FPCGTaggedData> UPCGDataFunctionLibrary::GetTaggedInputs(const FPCGDataCollection& InCollection, const FString& InTag)
+TArray<FPCGTaggedData> UPCGDataFunctionLibrary::GetInputsByTag(const FPCGDataCollection& InCollection, const FString& InTag)
 {
 	return InCollection.GetTaggedInputs(InTag);
 }
@@ -521,12 +521,12 @@ TArray<FPCGTaggedData> UPCGDataFunctionLibrary::GetParams(const FPCGDataCollecti
 	return InCollection.GetAllParams();
 }
 
-TArray<FPCGTaggedData> UPCGDataFunctionLibrary::GetParamsByPin(const FPCGDataCollection& InCollection, const FName& InPinLabel)
+TArray<FPCGTaggedData> UPCGDataFunctionLibrary::GetParamsByPinLabel(const FPCGDataCollection& InCollection, const FName& InPinLabel)
 {
 	return InCollection.GetParamsByPin(InPinLabel);
 }
 
-TArray<FPCGTaggedData> UPCGDataFunctionLibrary::GetTaggedParams(const FPCGDataCollection& InCollection, const FString& InTag)
+TArray<FPCGTaggedData> UPCGDataFunctionLibrary::GetParamsByTag(const FPCGDataCollection& InCollection, const FString& InTag)
 {
 	return InCollection.GetTaggedParams(InTag);
 }
diff --git a/Engine/Plugins/Experimental/PCG/Source/PCG/Public/Elements/PCGExecuteBlueprint.h b/Engine/Plugins/Experimental/PCG/Source/PCG/Public/Elements/PCGExecuteBlueprint.h
--- a/Engine/Plugins/Experimental/PCG/Source/PCG/Public/Elements/PCGExecuteBlueprint.h
+++ b/Engine/Plugins/Experimental/PCG/Source/PCG/Public/Elements/PCGExecuteBlueprint.h
@@ -102,6 +102,14 @@ public:
 	UFUNCTION(BlueprintCallable, Category = "PCG|Input & Output", meta = (HideSelfPin = "true"))
 	TArray<FPCGPinProperties> GetOutputPins() const;
 
+	/** Returns true if there is an input pin with the matching label. If found, will copy the pin properties in OutFoundPin */
+	UFUNCTION(BlueprintCallable, Category = "PCG|Input & Output", meta = (HideSelfPin = "true"))
+	bool GetInputPinByLabel(FName InPinLabel, FPCGPinProperties& OutFoundPin) const;
+
+	/** Returns true if there is an output pin with the matching label. If found, will copy the pin properties in OutFoundPin */
+	UFUNCTION(BlueprintCallable, Category = "PCG|Input & Output", meta = (HideSelfPin = "true"))
+	bool GetOutputPinByLabel(FName InPinLabel, FPCGPinProperties& OutFoundPin) const;
+
 	/** Gets the seed from the associated settings & source component */
 	UFUNCTION(BlueprintCallable, Category = "PCG|Random")
 	int GetSeed(UPARAM(ref) FPCGContext& InContext) const;
diff --git a/Engine/Plugins/Experimental/PCG/Source/PCG/Public/PCGCommon.h b/Engine/Plugins/Experimental/PCG/Source/PCG/Public/PCGCommon.h
--- a/Engine/Plugins/Experimental/PCG/Source/PCG/Public/PCGCommon.h
+++ b/Engine/Plugins/Experimental/PCG/Source/PCG/Public/PCGCommon.h
@@ -20,6 +20,10 @@ enum class EPCGChangeType : uint8
 };
 ENUM_CLASS_FLAGS(EPCGChangeType);
 
+// Bitmask containing the various data types supported in PCG. Note that this enum cannot be a blueprint type because
+// enums have to be uint8 for blueprint, and we already use more than 8 bits in the bitmask.
+// This is why we have a parallel enum just below that must match on a name basis 1:1 to allow the make/break functions to work properly
+// in blueprint.
 UENUM(meta = (Bitflags))
 enum class EPCGDataType : uint32
 {
@@ -55,6 +59,31 @@ enum class EPCGDataType : uint32
 };
 ENUM_CLASS_FLAGS(EPCGDataType);
 
+// As discussed just before, a parallel version for "exclusive" (as in only type) of the EPCGDataType enum. Needed for blueprint compatibility.
+UENUM(BlueprintType, meta=(DisplayName="PCG Data Type"))
+enum class EPCGExclusiveDataType : uint8
+{
+	None = 0 UMETA(Hidden),
+	Point,
+	Spline,
+	LandscapeSpline,
+	PolyLine UMETA(DisplayName = "Curve"),
+	Landscape,
+	Texture,
+	RenderTarget,
+	BaseTexture UMETA(Hidden),
+	Surface,
+	Volume,
+	Primitive,
+	Concrete,
+	Composite UMETA(Hidden),
+	Spatial,
+	Param UMETA(DisplayName = "Attribute Set"),
+	Settings UMETA(Hidden),
+	Other,
+	Any
+};
+
 namespace PCGPinConstants
 {
 	const FName DefaultInputLabel = TEXT("In");
diff --git a/Engine/Plugins/Experimental/PCG/Source/PCG/Public/PCGData.h b/Engine/Plugins/Experimental/PCG/Source/PCG/Public/PCGData.h
--- a/Engine/Plugins/Experimental/PCG/Source/PCG/Public/PCGData.h
+++ b/Engine/Plugins/Experimental/PCG/Source/PCG/Public/PCGData.h
@@ -217,19 +217,19 @@ public:
 	static TArray<FPCGTaggedData> GetInputs(const FPCGDataCollection& InCollection);
 
 	UFUNCTION(BlueprintCallable, Category = Data)
-	static TArray<FPCGTaggedData> GetInputsByPin(const FPCGDataCollection& InCollection, const FName& InPinLabel);
+	static TArray<FPCGTaggedData> GetInputsByPinLabel(const FPCGDataCollection& InCollection, const FName& InPinLabel);
 
 	UFUNCTION(BlueprintCallable, Category = Data)
-	static TArray<FPCGTaggedData> GetTaggedInputs(const FPCGDataCollection& InCollection, const FString& InTag);
+	static TArray<FPCGTaggedData> GetInputsByTag(const FPCGDataCollection& InCollection, const FString& InTag);
 
 	UFUNCTION(BlueprintCallable, Category = Data)
 	static TArray<FPCGTaggedData> GetParams(const FPCGDataCollection& InCollection);
 
 	UFUNCTION(BlueprintCallable, Category = Data)
-	static TArray<FPCGTaggedData> GetParamsByPin(const FPCGDataCollection& InCollection, const FName& InPinLabel);
+	static TArray<FPCGTaggedData> GetParamsByPinLabel(const FPCGDataCollection& InCollection, const FName& InPinLabel);
 
 	UFUNCTION(BlueprintCallable, Category = Data)
-	static TArray<FPCGTaggedData> GetTaggedParams(const FPCGDataCollection& InCollection, const FString& InTag);
+	static TArray<FPCGTaggedData> GetParamsByTag(const FPCGDataCollection& InCollection, const FString& InTag);
 
 	UFUNCTION(BlueprintCallable, Category = Data)
 	static TArray<FPCGTaggedData> GetAllSettings(const FPCGDataCollection& InCollection);
diff --git a/Engine/Plugins/Experimental/PCG/Source/PCG/Public/PCGPin.h b/Engine/Plugins/Experimental/PCG/Source/PCG/Public/PCGPin.h
--- a/Engine/Plugins/Experimental/PCG/Source/PCG/Public/PCGPin.h
+++ b/Engine/Plugins/Experimental/PCG/Source/PCG/Public/PCGPin.h
@@ -4,12 +4,14 @@
 
 #include "PCGCommon.h"
 
+#include "Kismet/BlueprintFunctionLibrary.h"
+
 #include "PCGPin.generated.h"
 
 class UPCGNode;
 class UPCGEdge;
 
-USTRUCT(BlueprintType)
+USTRUCT(BlueprintType, meta=(HasNativeBreak="/Script/PCG.PCGBlueprintPinHelpers.BreakPinProperty", HasNativeMake="/Script/PCG.PCGBlueprintPinHelpers.MakePinProperty"))
 struct PCG_API FPCGPinProperties
 {
 	GENERATED_BODY()
@@ -41,6 +43,69 @@ struct PCG_API FPCGPinProperties
 #endif
 };
 
+UCLASS()
+class PCG_API UPCGBlueprintPinHelpers : public UBlueprintFunctionLibrary
+{
+	GENERATED_BODY()
+
+public:
+	UFUNCTION(BlueprintPure, Category="PCG|Pins", meta = (NativeBreakFunc))
+	static void BreakPinProperty(const FPCGPinProperties& PinProperty, FName& Label, EPCGExclusiveDataType& AllowedType, bool& bAllowMultipleData, bool& bAllowMultipleConnections, bool& bAdvancedPin)
+	{
+		Label = PinProperty.Label;
+
+		const UEnum* DataTypeEnum = StaticEnum<EPCGDataType>();
+		const UEnum* ExclusiveDataTypeEnum = StaticEnum<EPCGExclusiveDataType>();
+
+		AllowedType = EPCGExclusiveDataType::Other;
+
+		if (DataTypeEnum && ExclusiveDataTypeEnum)
+		{
+			FName DataTypeName = DataTypeEnum->GetNameByValue(static_cast<__underlying_type(EPCGDataType)>(PinProperty.AllowedTypes));
+			if (DataTypeName != NAME_None)
+			{
+				const int64 MatchingType = ExclusiveDataTypeEnum->GetValueByName(DataTypeName);
+
+				if (MatchingType != INDEX_NONE)
+				{
+					AllowedType = static_cast<EPCGExclusiveDataType>(MatchingType);
+				}
+			}
+		}
+
+		bAllowMultipleData = PinProperty.bAllowMultipleData;
+		bAllowMultipleConnections = PinProperty.bAllowMultipleConnections;
+		bAdvancedPin = PinProperty.bAdvancedPin;
+	}
+
+	UFUNCTION(BlueprintPure, Category="PCG|Pins", meta = (NativeMakeFunc))
+	static FPCGPinProperties MakePinProperty(FName Label, EPCGExclusiveDataType AllowedType, bool bAllowMultipleData, bool bAllowMultipleConnections, bool bAdvancedPin)
+	{
+		const UEnum* DataTypeEnum = StaticEnum<EPCGDataType>();
+		const UEnum* ExclusiveDataTypeEnum = StaticEnum<EPCGExclusiveDataType>();
+
+		EPCGDataType DataType = EPCGDataType::Other;
+
+		if (DataTypeEnum && ExclusiveDataTypeEnum)
+		{
+			FName ExclusiveDataTypeName = ExclusiveDataTypeEnum->GetNameByValue(static_cast<__underlying_type(EPCGExclusiveDataType)>(AllowedType));
+			if (ensure(ExclusiveDataTypeName != NAME_None))
+			{
+				const int64 MatchingType = DataTypeEnum->GetValueByName(ExclusiveDataTypeName);
+				if (ensure(MatchingType != INDEX_NONE))
+				{
+					DataType = static_cast<EPCGDataType>(MatchingType);
+				}
+			}
+		}
+
+		FPCGPinProperties PinProperties = FPCGPinProperties(Label, DataType, bAllowMultipleConnections, bAllowMultipleData);
+		PinProperties.bAdvancedPin = bAdvancedPin;
+
+		return PinProperties;
+	}
+};
+
 UCLASS(ClassGroup = (Procedural))
 class PCG_API UPCGPin : public UObject
 {
