[PCG] Transform DensityNoise node to AttributeNoise to support any property/attribute

* Added redirectors to keep the old node compatible.

diff --git a/Engine/Plugins/Experimental/PCG/Config/BasePCG.ini b/Engine/Plugins/Experimental/PCG/Config/BasePCG.ini
--- a/Engine/Plugins/Experimental/PCG/Config/BasePCG.ini
+++ b/Engine/Plugins/Experimental/PCG/Config/BasePCG.ini
@@ -24,3 +24,6 @@
 +PropertyRedirects=(OldName="PCGBlueprintElement.NestedLoopBody.InPointB",NewName="PCGBlueprintElement.NestedLoopBody.InInnerPoint")
 
 +ClassRedirects=(OldName="PCGPointSamplerSettings",NewName="/Script/PCG.PCGSelectPointsSettings")
++ClassRedirects=(OldName="PCGDensityNoiseSettings",NewName="/Script/PCG.PCGAttributeNoiseSettings")
+
++EnumRedirects=(OldName="EPCGDensityNoiseMode",NewName="EPCGAttributeNoiseMode")
diff --git a/Engine/Plugins/Experimental/PCG/Source/PCG/Private/Elements/PCGAttributeNoise.cpp b/Engine/Plugins/Experimental/PCG/Source/PCG/Private/Elements/PCGAttributeNoise.cpp
new file mode 100644
--- /dev/null
+++ b/Engine/Plugins/Experimental/PCG/Source/PCG/Private/Elements/PCGAttributeNoise.cpp
@@ -0,0 +1,364 @@
+// Copyright Epic Games, Inc. All Rights Reserved.
+
+#include "Elements/PCGAttributeNoise.h"
+
+#include "PCGCommon.h"
+#include "PCGContext.h"
+#include "Data/PCGPointData.h"
+#include "Helpers/PCGAsync.h"
+#include "Helpers/PCGHelpers.h"
+#include "Metadata/PCGMetadataAttribute.h"
+#include "Metadata/PCGMetadataAttributeTraits.h"
+#include "Metadata/Accessors/IPCGAttributeAccessor.h"
+#include "Metadata/Accessors/PCGAttributeAccessorKeys.h"
+#include "Metadata/Accessors/PCGAttributeAccessorHelpers.h"
+
+#include "Math/RandomStream.h"
+#include "PCGPoint.h"
+
+#include UE_INLINE_GENERATED_CPP_BY_NAME(PCGAttributeNoise)
+
+#define LOCTEXT_NAMESPACE "PCGAttributeNoiseSettings"
+
+namespace PCGAttributeNoiseSettings
+{
+	template <typename T>
+	void ProcessNoise(T& InOutValue, FRandomStream& InRandomSource, const UPCGAttributeNoiseSettings* InSettings, const bool bClampResult)
+	{
+		if constexpr (std::is_arithmetic_v<T> && !std::is_same_v<T, bool>)
+		{
+			const EPCGAttributeNoiseMode Mode = InSettings->Mode;
+			const float NoiseMin = InSettings->NoiseMin;
+			const float NoiseMax = InSettings->NoiseMax;
+			const bool bInvertSource = InSettings->bInvertSource;
+
+			const double Noise = InRandomSource.FRandRange(NoiseMin, NoiseMax);
+
+			if (bInvertSource)
+			{
+				InOutValue = static_cast<T>(1.0 - InOutValue);
+			}
+
+			if (Mode == EPCGAttributeNoiseMode::Minimum)
+			{
+				InOutValue = FMath::Min<T>(InOutValue, Noise);
+			}
+			else if (Mode == EPCGAttributeNoiseMode::Maximum)
+			{
+				InOutValue = FMath::Max<T>(InOutValue, Noise);
+			}
+			else if (Mode == EPCGAttributeNoiseMode::Add)
+			{
+				InOutValue = static_cast<T>(InOutValue + Noise);
+			}
+			else if (Mode == EPCGAttributeNoiseMode::Multiply)
+			{
+				InOutValue = static_cast<T>(InOutValue * Noise);
+			}
+			else //if (Mode == EPCGAttributeNoiseMode::Set)
+			{
+				InOutValue = static_cast<T>(Noise);
+			}
+
+			if (bClampResult)
+			{
+				InOutValue = FMath::Clamp<T>(InOutValue, 0, 1);
+			}
+		}
+		else if constexpr (std::is_same_v<FVector2D, T>)
+		{
+			ProcessNoise(InOutValue.X, InRandomSource, InSettings, bClampResult);
+			ProcessNoise(InOutValue.Y, InRandomSource, InSettings, bClampResult);
+		}
+		else if constexpr (std::is_same_v<FVector, T>)
+		{
+			ProcessNoise(InOutValue.X, InRandomSource, InSettings, bClampResult);
+			ProcessNoise(InOutValue.Y, InRandomSource, InSettings, bClampResult);
+			ProcessNoise(InOutValue.Z, InRandomSource, InSettings, bClampResult);
+		}
+		else if constexpr (std::is_same_v<FVector4, T> || std::is_same_v<FQuat, T>)
+		{
+			ProcessNoise(InOutValue.X, InRandomSource, InSettings, bClampResult);
+			ProcessNoise(InOutValue.Y, InRandomSource, InSettings, bClampResult);
+			ProcessNoise(InOutValue.Z, InRandomSource, InSettings, bClampResult);
+			ProcessNoise(InOutValue.W, InRandomSource, InSettings, bClampResult);
+		}
+		else if constexpr (std::is_same_v<FRotator, T>)
+		{
+			ProcessNoise(InOutValue.Roll, InRandomSource, InSettings, bClampResult);
+			ProcessNoise(InOutValue.Pitch, InRandomSource, InSettings, bClampResult);
+			ProcessNoise(InOutValue.Yaw, InRandomSource, InSettings, bClampResult);
+		}
+	}
+}
+
+UPCGAttributeNoiseSettings::UPCGAttributeNoiseSettings()
+{
+	bUseSeed = true;
+	InputSource.SetPointProperty(EPCGPointProperties::Density);
+}
+
+void UPCGAttributeNoiseSettings::PostLoad()
+{
+	Super::PostLoad();
+
+#if WITH_EDITOR
+	if (DensityMode_DEPRECATED != EPCGAttributeNoiseMode::Set)
+	{
+		Mode = DensityMode_DEPRECATED;
+		DensityMode_DEPRECATED = EPCGAttributeNoiseMode::Set;
+	}
+
+	if (DensityNoiseMin_DEPRECATED != 0.f)
+	{
+		NoiseMin = DensityNoiseMin_DEPRECATED;
+		DensityNoiseMin_DEPRECATED = 0.f;
+	}
+
+	if (DensityNoiseMax_DEPRECATED != 1.f)
+	{
+		NoiseMax = DensityNoiseMax_DEPRECATED;
+		DensityNoiseMax_DEPRECATED = 1.f;
+	}
+
+	if (bInvertSourceDensity_DEPRECATED)
+	{
+		bInvertSource = bInvertSourceDensity_DEPRECATED;
+		bInvertSourceDensity_DEPRECATED = false;
+	}
+#endif // WITH_EDITOR
+}
+
+#if WITH_EDITOR
+void UPCGAttributeNoiseSettings::ApplyDeprecationBeforeUpdatePins(UPCGNode* InOutNode, TArray<TObjectPtr<UPCGPin>>& InputPins, TArray<TObjectPtr<UPCGPin>>& OutputPins)
+{
+	Super::ApplyDeprecationBeforeUpdatePins(InOutNode, InputPins, OutputPins);
+
+	// Overridable properties have been renamed, rename all pins by their counterpart, to avoid breaking existing graphs.
+	const TArray<TPair<FName, FName>> OldToNewPinNames =
+	{
+		{TEXT("Density Mode"), TEXT("Mode")},
+		{TEXT("Density Noise Min"), TEXT("Noise Min")},
+		{TEXT("Density Noise Max"),TEXT("Noise Max")},
+		{TEXT("Invert Source Density"), TEXT("Invert Source")}
+	};
+
+	for (const TPair<FName, FName>& OldToNew : OldToNewPinNames)
+	{
+		InOutNode->RenameInputPin(OldToNew.Key, OldToNew.Value);
+	}
+}
+#endif // WITH_EDITOR
+
+FPCGElementPtr UPCGAttributeNoiseSettings::CreateElement() const
+{
+	return MakeShared<FPCGAttributeNoiseElement>();
+}
+
+FPCGContext* FPCGAttributeNoiseElement::Initialize(const FPCGDataCollection& InInputData, TWeakObjectPtr<UPCGComponent> InSourceComponent, const UPCGNode* InNode)
+{
+	FPCGAttributeNoiseContext* Context = new FPCGAttributeNoiseContext();
+	Context->InputData = InInputData;
+	Context->SourceComponent = InSourceComponent;
+	Context->Node = InNode;
+
+	return Context;
+}
+
+bool FPCGAttributeNoiseElement::ExecuteInternal(FPCGContext* InContext) const
+{
+	TRACE_CPUPROFILER_EVENT_SCOPE(FPCGAttributeNoiseElement::Execute);
+
+	FPCGAttributeNoiseContext* Context = static_cast<FPCGAttributeNoiseContext*>(InContext);
+
+	const UPCGAttributeNoiseSettings* Settings = Context->GetInputSettings<UPCGAttributeNoiseSettings>();
+	check(Settings);
+
+	TArray<FPCGTaggedData> Inputs = Context->InputData.GetInputs();
+	TArray<FPCGTaggedData>& Outputs = Context->OutputData.TaggedData;
+
+	// Precompute a seed based on the settings one and the component one
+	const int Seed = Context->GetSeed();
+
+	while (Context->CurrentInput < Inputs.Num())
+	{
+		int32 CurrentInput = Context->CurrentInput;
+		const FPCGTaggedData& Input = Inputs[CurrentInput];
+
+		TRACE_CPUPROFILER_EVENT_SCOPE(FPCGAttributeNoiseElement::InputLoop);
+		
+		const TArray<FPCGPoint>* InputPoints = nullptr;
+		TArray<FPCGPoint>* OutputPoints = nullptr;
+
+		if (!Context->bDataPreparedForCurrentInput)
+		{
+			const UPCGSpatialData* SpatialData = Cast<UPCGSpatialData>(Input.Data);
+
+			if (!SpatialData)
+			{
+				PCGE_LOG(Error, GraphAndLog, LOCTEXT("InputMissingSpatialData", "Unable to get Spatial data from input"));
+				Context->CurrentInput++;
+				continue;
+			}
+
+			const UPCGPointData* PointData = SpatialData->ToPointData(Context);
+
+			if (!PointData)
+			{
+				PCGE_LOG(Error, GraphAndLog, LOCTEXT("InputMissingPointData", "Unable to get Point data from input"));
+				Context->CurrentInput++;
+				continue;
+			}
+
+			InputPoints = &PointData->GetPoints();
+
+			// Create a dummy accessor on the input before allocating the output, to avoid doing useless allocation.
+			TUniquePtr<const IPCGAttributeAccessor> TempInputAccessor = PCGAttributeAccessorHelpers::CreateConstAccessor(PointData, Settings->InputSource);
+			if (!TempInputAccessor)
+			{
+				Outputs.RemoveAt(Outputs.Num() - 1);
+				PCGE_LOG(Error, GraphAndLog, FText::Format(LOCTEXT("CantCreateAccessor", "Could not find Attribute/Property {0}"), FText::FromName(Settings->InputSource.GetName())));
+				Context->CurrentInput++;
+				continue;
+			}
+
+			// Also need to make sure the accessor is a "noisable" type
+			if (!PCG::Private::IsOfTypes<int32, int64, float, double, FVector, FVector2D, FVector4, FRotator, FQuat>(TempInputAccessor->GetUnderlyingType()))
+			{
+				Outputs.RemoveAt(Outputs.Num() - 1);
+				PCGE_LOG(Error, GraphAndLog, FText::Format(LOCTEXT("AttributeIsNotANumericalType", "Attribute/Property {0} is not a numerical type, we can't apply noise to it."), FText::FromName(Settings->InputSource.GetName())));
+				Context->CurrentInput++;
+				continue;
+			}
+
+			TempInputAccessor.Reset();
+
+			FPCGTaggedData& Output = Outputs.Add_GetRef(Input);
+			UPCGPointData* OutputData = NewObject<UPCGPointData>();
+			OutputData->InitializeFromData(PointData);
+			OutputPoints = &OutputData->GetMutablePoints();
+			OutputPoints->SetNumUninitialized(InputPoints->Num());
+			Output.Data = OutputData;
+
+			// Then create the accessor/keys
+			Context->InputAccessor = PCGAttributeAccessorHelpers::CreateAccessor(OutputData, Settings->InputSource);
+			Context->Keys = PCGAttributeAccessorHelpers::CreateKeys(OutputData, Settings->InputSource);
+
+			// It won't fail because we validated on the input data and output is initialized from input.
+			check(Context->InputAccessor && Context->Keys);
+
+			// Allocate temp buffer and create output accessor if needed
+			const bool bValid = PCGMetadataAttribute::CallbackWithRightType(Context->InputAccessor->GetUnderlyingType(), [this, Context, Settings, OutputData](auto&& Dummy) -> bool
+			{
+				using AttributeType = std::decay_t<decltype(Dummy)>;
+				int32 NumPoints = Context->Keys->GetNum();
+				Context->TempValuesBuffer.SetNumUninitialized(sizeof(AttributeType) * NumPoints);
+
+				// If we have a different output target, create accessor (and also create the attribute if it doesn't exist), and check that we can broadcast to it
+				if (Settings->bOutputTargetDifferentFromInputSource)
+				{
+					Context->OptionalOutputAccessor = PCGAttributeAccessorHelpers::CreateAccessor(OutputData, Settings->OutputTarget);
+					if (!Context->OptionalOutputAccessor && Settings->OutputTarget.Selection == EPCGAttributePropertySelection::Attribute)
+					{
+						OutputData->Metadata->CreateAttribute<AttributeType>(Settings->OutputTarget.GetName(), AttributeType{}, /*bAllowsInterpolation=*/ true, /*bOverrideParent=*/false);
+						Context->OptionalOutputAccessor = PCGAttributeAccessorHelpers::CreateAccessor(OutputData, Settings->OutputTarget);
+					}
+
+					if (!Context->OptionalOutputAccessor)
+					{
+						PCGE_LOG(Error, GraphAndLog, FText::Format(LOCTEXT("OutputTargetInvalid", "Failed to find/create Attribute/Property {0}."), FText::FromName(Settings->OutputTarget.GetName())));
+						return false;
+					}
+
+					if (!PCG::Private::IsBroadcastable(Context->InputAccessor->GetUnderlyingType(), Context->OptionalOutputAccessor->GetUnderlyingType()))
+					{
+						PCGE_LOG(Error, GraphAndLog, FText::Format(LOCTEXT("CantBroadcast", "Cannot convert Attribute {0} ({1}) into Attribute {2} ({3})."),
+							FText::FromName(Settings->InputSource.GetName()),
+							FText::FromString(PCG::Private::GetTypeName(Context->InputAccessor->GetUnderlyingType())),
+							FText::FromName(Settings->OutputTarget.GetName()),
+							FText::FromString(PCG::Private::GetTypeName(Context->OptionalOutputAccessor->GetUnderlyingType()))));
+						return false;
+					}
+				}
+
+				return true;
+			});
+
+			if (!bValid)
+			{
+				Outputs.RemoveAt(Outputs.Num() - 1);
+				Context->CurrentInput++;
+				continue;
+			}
+
+			Context->bDataPreparedForCurrentInput = true;
+		}
+		else
+		{
+			OutputPoints = &CastChecked<UPCGPointData>(Outputs[CurrentInput].Data)->GetMutablePoints();
+			InputPoints = &CastChecked<const UPCGPointData>(Inputs[CurrentInput].Data)->GetPoints();
+		}
+
+		check(InputPoints && OutputPoints);
+
+		// Force clamp on Density
+		const FPCGAttributePropertySelector& Selector = Settings->bOutputTargetDifferentFromInputSource ? Settings->OutputTarget : Settings->InputSource;
+		const bool bClampResult = Settings->bClampResult || (Selector.Selection == EPCGAttributePropertySelection::PointProperty && Selector.PointProperty == EPCGPointProperties::Density);
+
+		// Dummy Initialize, we already initialized the output points before.
+		auto Initialize = []() {};
+
+		const bool bDone = FPCGAsync::AsyncProcessingOneToOneEx(&Context->AsyncState, OutputPoints->Num(), Initialize, [InputPoints, OutputPoints, Settings, Seed, bClampResult, Context](int32 StartReadIndex, int32 StartWriteIndex) -> int32
+		{
+			PCGMetadataAttribute::CallbackWithRightType(Context->InputAccessor->GetUnderlyingType(), [Settings, StartReadIndex, StartWriteIndex, InputPoints, OutputPoints, Context, Seed, bClampResult](auto&& Dummy)
+			{
+				using AttributeType = std::decay_t<decltype(Dummy)>;
+
+				// Copying the point
+				FPCGPoint& OutPoint = (*OutputPoints)[StartWriteIndex];
+				OutPoint = (*InputPoints)[StartReadIndex];
+
+				// Reinterpret the buffer to store our temporary value.
+				AttributeType& Value = *(reinterpret_cast<AttributeType*>(Context->TempValuesBuffer.GetData()) + StartReadIndex);
+
+				if (Context->InputAccessor->Get<AttributeType>(Value, StartReadIndex, *Context->Keys))
+				{
+					FRandomStream RandomSource(PCGHelpers::ComputeSeed(Seed, OutPoint.Seed));
+					PCGAttributeNoiseSettings::ProcessNoise(Value, RandomSource, Settings, bClampResult);
+				}
+			});
+
+			return true;
+		}, /*bEnableTimeSlicing=*/true);
+
+		if (bDone)
+		{
+			PCGMetadataAttribute::CallbackWithRightType(Context->InputAccessor->GetUnderlyingType(), [Context](auto&& Value)
+			{
+				using AttributeType = std::decay_t<decltype(Value)>;
+				int32 NumPoints = Context->Keys->GetNum();
+				TArrayView<AttributeType> Values(reinterpret_cast<AttributeType*>(Context->TempValuesBuffer.GetData()), NumPoints);
+				if (Context->OptionalOutputAccessor)
+				{
+					Context->OptionalOutputAccessor->SetRange<AttributeType>(Values, 0, *Context->Keys, EPCGAttributeAccessorFlags::AllowBroadcast);
+				}
+				else
+				{
+					Context->InputAccessor->SetRange<AttributeType>(Values, 0, *Context->Keys);
+				}
+			});
+			
+			Context->CurrentInput++;
+			Context->bDataPreparedForCurrentInput = false;
+		}
+		
+		if (!bDone || Context->ShouldStop())
+		{
+			return false;
+		}
+	}
+
+	return true;
+}
+
+#undef LOCTEXT_NAMESPACE
diff --git a/Engine/Plugins/Experimental/PCG/Source/PCG/Private/Elements/PCGDensityNoise.cpp b/Engine/Plugins/Experimental/PCG/Source/PCG/Private/Elements/PCGDensityNoise.cpp
deleted file mode 100644
diff --git a/Engine/Plugins/Experimental/PCG/Source/PCG/Private/Tests/Elements/PCGAttributeNoiseTest.cpp b/Engine/Plugins/Experimental/PCG/Source/PCG/Private/Tests/Elements/PCGAttributeNoiseTest.cpp
new file mode 100644
--- /dev/null
+++ b/Engine/Plugins/Experimental/PCG/Source/PCG/Private/Tests/Elements/PCGAttributeNoiseTest.cpp
@@ -0,0 +1,233 @@
+// Copyright Epic Games, Inc. All Rights Reserved.
+
+#include "Tests/Determinism/PCGDeterminismTestsCommon.h"
+
+#include "Data/PCGSpatialData.h"
+#include "PCGComponent.h"
+
+#include "Data/PCGPointData.h"
+#include "Elements/PCGAttributeNoise.h"
+#include "PCGContext.h"
+
+IMPLEMENT_CUSTOM_SIMPLE_AUTOMATION_TEST(FPCGAttributeNoiseDensityTest, FPCGTestBaseClass, "pcg.tests.AttributeNoise.Density", PCGTestsCommon::TestFlags)
+
+bool FPCGAttributeNoiseDensityTest::RunTest(const FString& Parameters)
+{
+	PCGTestsCommon::FTestData TestData(PCGDeterminismTests::Defaults::Seed);
+	PCGDeterminismTests::GenerateSettings<UPCGAttributeNoiseSettings>(TestData);
+	UPCGAttributeNoiseSettings* Settings = CastChecked<UPCGAttributeNoiseSettings>(TestData.Settings);
+	Settings->InputSource.SetPointProperty(EPCGPointProperties::Density);
+	FPCGElementPtr NoiseElement = TestData.Settings->GetElement();
+
+	TObjectPtr<UPCGPointData> PointData = PCGTestsCommon::CreateEmptyPointData();
+	TArray<FPCGPoint>& Points = PointData->GetMutablePoints();
+
+	FRandomStream RandomSource(TestData.Seed);
+	const int PointCount = 5;
+	for (int I = 0; I < PointCount; ++I)
+	{
+		FPCGPoint& Point = Points.Emplace_GetRef(FTransform(), 1, I);
+		Point.Density = RandomSource.GetFraction();
+	}
+
+	FPCGTaggedData& TaggedData = TestData.InputData.TaggedData.Emplace_GetRef(FPCGTaggedData());
+	TaggedData.Data = PointData;
+
+	auto ValidateDensityNoise = [this, &TestData, NoiseElement, Settings](TArray<float> ExpectedOutput) -> bool
+	{
+		TUniquePtr<FPCGContext> Context = TestData.InitializeTestContext();
+
+		while (!NoiseElement->Execute(Context.Get()))
+		{}
+
+		const TArray<FPCGTaggedData>& Inputs = Context->InputData.GetInputs();
+		const TArray<FPCGTaggedData>& Outputs = Context->OutputData.GetInputs();
+
+		if (!TestEqual("Valid number of outputs", Inputs.Num(), Outputs.Num()))
+		{
+			return false;
+		}
+
+		bool bTestPassed = true;
+
+		for (int DataIndex = 0; DataIndex < Inputs.Num(); ++DataIndex)
+		{
+			const FPCGTaggedData& Input = Inputs[DataIndex];
+			const FPCGTaggedData& Output = Outputs[DataIndex];
+
+			const UPCGSpatialData* InSpatialData = Cast<UPCGSpatialData>(Input.Data);
+			check(InSpatialData);
+
+			const UPCGPointData* InPointData = InSpatialData->ToPointData(Context.Get());
+			check(InPointData);
+
+			const UPCGSpatialData* OutSpatialData = Cast<UPCGSpatialData>(Output.Data);
+
+			if (!TestNotNull("Valid output SpatialData", OutSpatialData))
+			{
+				bTestPassed = false;
+				continue;
+			}
+
+			const UPCGPointData* OutPointData = OutSpatialData->ToPointData(Context.Get());
+
+			if (!TestNotNull("Valid output PointData", OutPointData))
+			{
+				bTestPassed = false;
+				continue;
+			}
+
+			const TArray<FPCGPoint>& InPoints = InPointData->GetPoints();
+			const TArray<FPCGPoint>& OutPoints = OutPointData->GetPoints();
+
+			if (!TestEqual("Input and output point counts match", InPoints.Num(), OutPoints.Num()))
+			{ 
+				bTestPassed = false;
+				continue;
+			}
+
+			for (int PointIndex = 0; PointIndex < ExpectedOutput.Num(); ++PointIndex)
+			{
+				bTestPassed &= TestEqual("Correct density", OutPoints[PointIndex].Density, ExpectedOutput[PointIndex]);
+			}
+		}
+
+		return bTestPassed;
+	};
+
+	bool bTestPassed = true;
+
+	auto ValidateDensityNoiseForAllDensityModes = [this, &bTestPassed, ValidateDensityNoise, Settings]()
+	{
+		Settings->Mode = EPCGAttributeNoiseMode::Set;
+		bTestPassed &= ValidateDensityNoise({});
+
+		Settings->Mode = EPCGAttributeNoiseMode::Minimum;
+		bTestPassed &= ValidateDensityNoise({});
+
+		Settings->Mode = EPCGAttributeNoiseMode::Maximum;
+		bTestPassed &= ValidateDensityNoise({});
+
+		Settings->Mode = EPCGAttributeNoiseMode::Add;
+		bTestPassed &= ValidateDensityNoise({});
+
+		Settings->Mode = EPCGAttributeNoiseMode::Multiply;
+		bTestPassed &= ValidateDensityNoise({});
+	};
+
+	// Test [0-1]
+	{
+		Settings->NoiseMin = 0.f;
+		Settings->NoiseMax = 1.f;
+		Settings->bInvertSource = false;
+		ValidateDensityNoiseForAllDensityModes();
+
+		Settings->bInvertSource = true;
+		ValidateDensityNoiseForAllDensityModes();
+	}
+
+	// Test [0-0.5]
+	{
+		Settings->NoiseMin = 0.f;
+		Settings->NoiseMax = 0.5f;
+		Settings->bInvertSource = false;
+		ValidateDensityNoiseForAllDensityModes();
+		
+		Settings->bInvertSource = true;
+		ValidateDensityNoiseForAllDensityModes();
+	}
+
+	// Test [0.5-1]
+	{
+		Settings->NoiseMin = 0.5f;
+		Settings->NoiseMax = 1.f;
+		Settings->bInvertSource = false;
+		ValidateDensityNoiseForAllDensityModes();
+
+		Settings->bInvertSource = true;
+		ValidateDensityNoiseForAllDensityModes();
+	}
+
+	// Test [1-0]
+	{
+		Settings->NoiseMin = 1.f;
+		Settings->NoiseMax = 0.f;
+		Settings->bInvertSource = false;
+		ValidateDensityNoiseForAllDensityModes();
+
+		Settings->bInvertSource = true;
+		ValidateDensityNoiseForAllDensityModes();
+	}
+
+	// Test [0.5-0]
+	{
+		Settings->NoiseMin = 0.5f;
+		Settings->NoiseMax = 0.f;
+		Settings->bInvertSource = false;
+		ValidateDensityNoiseForAllDensityModes();
+		
+		Settings->bInvertSource = true;
+		ValidateDensityNoiseForAllDensityModes();
+	}
+
+	// Test [1-0.5]
+	{
+		Settings->NoiseMin = 1.f;
+		Settings->NoiseMax = 0.5f;
+		Settings->bInvertSource = false;
+		ValidateDensityNoiseForAllDensityModes();
+
+		Settings->bInvertSource = true;
+		ValidateDensityNoiseForAllDensityModes();
+	}
+
+	Settings->bInvertSource = false;
+
+	// Test expected values for Set
+	{
+		Settings->NoiseMin = 0.5f;
+		Settings->NoiseMax = 0.5f;
+		Settings->Mode = EPCGAttributeNoiseMode::Set;
+		bTestPassed &= ValidateDensityNoise({ 0.5f, 0.5f, 0.5f, 0.5f, 0.5f });
+	}
+
+	// Test expected values for Minimum
+	{
+		Settings->NoiseMin = 0.f;
+		Settings->NoiseMax = 0.f;
+		Settings->Mode = EPCGAttributeNoiseMode::Minimum;
+		bTestPassed &= ValidateDensityNoise({ 0.0f, 0.0f, 0.0f, 0.0f, 0.0f });
+	}
+
+	// Test expected values for Maximum
+	{
+		Settings->NoiseMin = 1.f;
+		Settings->NoiseMax = 1.f;
+		Settings->Mode = EPCGAttributeNoiseMode::Maximum;
+		bTestPassed &= ValidateDensityNoise({ 1.f, 1.f, 1.f, 1.f, 1.f });
+	}
+
+	// Test expected values for Add
+	{
+		Settings->NoiseMin = 0.f;
+		Settings->NoiseMax = 0.f;
+		Settings->Mode = EPCGAttributeNoiseMode::Add;
+		bTestPassed &= ValidateDensityNoise({ 
+			Points[0].Density,
+			Points[1].Density,
+			Points[2].Density,
+			Points[3].Density,
+			Points[4].Density,
+		});
+	}
+
+	// Test expected values for Multiply
+	{
+		Settings->NoiseMin = 0.f;
+		Settings->NoiseMax = 0.f;
+		Settings->Mode = EPCGAttributeNoiseMode::Multiply;
+		bTestPassed &= ValidateDensityNoise({ 0.0f, 0.0f, 0.0f, 0.0f, 0.0f });
+	}
+
+	return bTestPassed;
+}
diff --git a/Engine/Plugins/Experimental/PCG/Source/PCG/Private/Tests/Elements/PCGDensityNoiseTest.cpp b/Engine/Plugins/Experimental/PCG/Source/PCG/Private/Tests/Elements/PCGDensityNoiseTest.cpp
deleted file mode 100644
diff --git a/Engine/Plugins/Experimental/PCG/Source/PCG/Public/Elements/PCGAttributeNoise.h b/Engine/Plugins/Experimental/PCG/Source/PCG/Public/Elements/PCGAttributeNoise.h
new file mode 100644
--- /dev/null
+++ b/Engine/Plugins/Experimental/PCG/Source/PCG/Public/Elements/PCGAttributeNoise.h
@@ -0,0 +1,122 @@
+// Copyright Epic Games, Inc. All Rights Reserved.
+
+#pragma once
+
+#include "PCGContext.h"
+#include "PCGElement.h"
+#include "PCGPin.h"
+#include "PCGSettings.h"
+#include "Metadata/PCGAttributePropertySelector.h"
+
+#include "PCGAttributeNoise.generated.h"
+
+class IPCGAttributeAccessor;
+class IPCGAttributeAccessorKeys;
+
+UENUM()
+enum class EPCGAttributeNoiseMode : uint8
+{
+	Set,
+	Minimum,
+	Maximum,
+	Add,
+	Multiply
+};
+
+/**
+* Apply some noise to an attribute/property. You can select the mode you want and a noise range.
+* Support all numerical types and vectors/rotators.
+*/
+UCLASS(BlueprintType, ClassGroup = (Procedural))
+class PCG_API UPCGAttributeNoiseSettings : public UPCGSettings
+{
+	GENERATED_BODY()
+
+public:
+	UPCGAttributeNoiseSettings();
+
+	// ~Begin UObject interface
+	virtual void PostLoad() override;
+	// ~End UObject interface
+
+	// ~Begin UPCGSettings interface
+#if WITH_EDITOR
+	virtual void ApplyDeprecationBeforeUpdatePins(UPCGNode* InOutNode, TArray<TObjectPtr<UPCGPin>>& InputPins, TArray<TObjectPtr<UPCGPin>>& OutputPins);
+	virtual FName GetDefaultNodeName() const override { return FName(TEXT("AttributeNoise")); }
+	virtual FText GetDefaultNodeTitle() const override { return NSLOCTEXT("PCGAttributeNoiseSettings", "NodeTitle", "Attribute Noise"); }
+	virtual EPCGSettingsType GetType() const override { return EPCGSettingsType::Metadata; }
+#endif
+
+protected:
+	virtual TArray<FPCGPinProperties> InputPinProperties() const override { return Super::DefaultPointInputPinProperties(); }
+	virtual TArray<FPCGPinProperties> OutputPinProperties() const override { return Super::DefaultPointOutputPinProperties(); }
+	virtual FPCGElementPtr CreateElement() const override;
+	//~End UPCGSettings interface
+
+public:
+	UPROPERTY(BlueprintReadWrite, EditAnywhere, Category = Settings)
+	FPCGAttributePropertySelector InputSource;
+
+	UPROPERTY(BlueprintReadWrite, EditAnywhere, Category = Settings)
+	bool bOutputTargetDifferentFromInputSource = false;
+
+	UPROPERTY(BlueprintReadWrite, EditAnywhere, Category = Settings, meta = (EditCondition = "bOutputTargetDifferentFromInputSource", EditConditionHides))
+	FPCGAttributePropertySelector OutputTarget;
+
+	/** Attribute = (Original op Noise), Noise in [NoiseMin, NoiseMax] */
+	UPROPERTY(BlueprintReadWrite, EditAnywhere, Category = Settings, meta = (PCG_Overridable))
+	EPCGAttributeNoiseMode Mode = EPCGAttributeNoiseMode::Set;
+
+	UPROPERTY(BlueprintReadWrite, EditAnywhere, Category = Settings, meta = (PCG_Overridable))
+	float NoiseMin = 0.f;
+
+	UPROPERTY(BlueprintReadWrite, EditAnywhere, Category = Settings, meta = (PCG_Overridable))
+	float NoiseMax = 1.f;
+
+	/** Attribute = 1 - Attribute before applying the operation */
+	UPROPERTY(BlueprintReadWrite, EditAnywhere, Category = Settings, meta = (PCG_Overridable))
+	bool bInvertSource = false;
+
+	// Clamp the result between 0 and 1. Always applied if we apply noise to the density.
+	UPROPERTY(BlueprintReadWrite, EditAnywhere, Category = Settings, meta = (PCG_Overridable))
+	bool bClampResult = false;
+
+	// Previous names
+#if WITH_EDITORONLY_DATA
+	UPROPERTY()
+	EPCGAttributeNoiseMode DensityMode_DEPRECATED = EPCGAttributeNoiseMode::Set;
+
+	UPROPERTY()
+	float DensityNoiseMin_DEPRECATED = 0.f;
+
+	UPROPERTY()
+	float DensityNoiseMax_DEPRECATED = 1.f;
+
+	UPROPERTY()
+	bool bInvertSourceDensity_DEPRECATED = false;
+#endif // WITH_EDITORDATA_ONLY
+};
+
+struct FPCGAttributeNoiseContext : public FPCGContext
+{
+	int32 CurrentInput = 0;
+	bool bDataPreparedForCurrentInput = false;
+	TUniquePtr<IPCGAttributeAccessor> InputAccessor;
+	TUniquePtr<IPCGAttributeAccessor> OptionalOutputAccessor;
+	TUniquePtr<IPCGAttributeAccessorKeys> Keys;
+
+	TArray<uint8> TempValuesBuffer;
+};
+
+class FPCGAttributeNoiseElement : public IPCGElement
+{
+public:
+	virtual FPCGContext* Initialize(const FPCGDataCollection& InputData, TWeakObjectPtr<UPCGComponent> SourceComponent, const UPCGNode* Node) override;
+protected:
+	virtual bool ExecuteInternal(FPCGContext* Context) const;
+};
+
+#if UE_ENABLE_INCLUDE_ORDER_DEPRECATED_IN_5_2
+#include "CoreMinimal.h"
+#include "PCGNode.h"
+#endif
diff --git a/Engine/Plugins/Experimental/PCG/Source/PCG/Public/Elements/PCGDensityNoise.h b/Engine/Plugins/Experimental/PCG/Source/PCG/Public/Elements/PCGDensityNoise.h
deleted file mode 100644
