[PCG] Tentative fix for static analysis + missing include

#rnx

diff --git a/Engine/Plugins/Experimental/PCG/Source/PCG/Private/Elements/PCGAttributeNoise.cpp b/Engine/Plugins/Experimental/PCG/Source/PCG/Private/Elements/PCGAttributeNoise.cpp
--- a/Engine/Plugins/Experimental/PCG/Source/PCG/Private/Elements/PCGAttributeNoise.cpp
+++ b/Engine/Plugins/Experimental/PCG/Source/PCG/Private/Elements/PCGAttributeNoise.cpp
@@ -9,8 +9,6 @@
 #include "Helpers/PCGHelpers.h"
 #include "Metadata/PCGMetadataAttribute.h"
 #include "Metadata/PCGMetadataAttributeTraits.h"
-#include "Metadata/Accessors/IPCGAttributeAccessor.h"
-#include "Metadata/Accessors/PCGAttributeAccessorKeys.h"
 #include "Metadata/Accessors/PCGAttributeAccessorHelpers.h"
 
 #include "Math/RandomStream.h"
@@ -170,6 +168,7 @@ bool FPCGAttributeNoiseElement::ExecuteInternal(FPCGContext* InContext) const
 	TRACE_CPUPROFILER_EVENT_SCOPE(FPCGAttributeNoiseElement::Execute);
 
 	FPCGAttributeNoiseContext* Context = static_cast<FPCGAttributeNoiseContext*>(InContext);
+	check(Context);
 
 	const UPCGAttributeNoiseSettings* Settings = Context->GetInputSettings<UPCGAttributeNoiseSettings>();
 	check(Settings);
@@ -250,6 +249,8 @@ bool FPCGAttributeNoiseElement::ExecuteInternal(FPCGContext* InContext) const
 			// Allocate temp buffer and create output accessor if needed
 			const bool bValid = PCGMetadataAttribute::CallbackWithRightType(Context->InputAccessor->GetUnderlyingType(), [this, Context, Settings, OutputData](auto&& Dummy) -> bool
 			{
+				check(Context);
+
 				using AttributeType = std::decay_t<decltype(Dummy)>;
 				int32 NumPoints = Context->Keys->GetNum();
 				Context->TempValuesBuffer.SetNumUninitialized(sizeof(AttributeType) * NumPoints);
diff --git a/Engine/Plugins/Experimental/PCG/Source/PCG/Public/Elements/PCGAttributeNoise.h b/Engine/Plugins/Experimental/PCG/Source/PCG/Public/Elements/PCGAttributeNoise.h
--- a/Engine/Plugins/Experimental/PCG/Source/PCG/Public/Elements/PCGAttributeNoise.h
+++ b/Engine/Plugins/Experimental/PCG/Source/PCG/Public/Elements/PCGAttributeNoise.h
@@ -7,11 +7,12 @@
 #include "PCGPin.h"
 #include "PCGSettings.h"
 #include "Metadata/PCGAttributePropertySelector.h"
+#include "Metadata/Accessors/IPCGAttributeAccessor.h"
+#include "Metadata/Accessors/PCGAttributeAccessorKeys.h"
 
-#include "PCGAttributeNoise.generated.h"
+#include "Templates/UniquePtr.h"
 
-class IPCGAttributeAccessor;
-class IPCGAttributeAccessorKeys;
+#include "PCGAttributeNoise.generated.h"
 
 UENUM()
 enum class EPCGAttributeNoiseMode : uint8
