Small adjustment to the default type when calling Make Pin Property to make it Any (like before) vs None.

#rnx

diff --git a/Engine/Plugins/Experimental/PCG/Source/PCG/Public/PCGPin.h b/Engine/Plugins/Experimental/PCG/Source/PCG/Public/PCGPin.h
--- a/Engine/Plugins/Experimental/PCG/Source/PCG/Public/PCGPin.h
+++ b/Engine/Plugins/Experimental/PCG/Source/PCG/Public/PCGPin.h
@@ -50,7 +50,7 @@ class PCG_API UPCGBlueprintPinHelpers : public UBlueprintFunctionLibrary
 
 public:
 	UFUNCTION(BlueprintPure, Category="PCG|Pins", meta = (NativeBreakFunc))
-	static void BreakPinProperty(const FPCGPinProperties& PinProperty, FName& Label, EPCGExclusiveDataType& AllowedType, bool& bAllowMultipleData, bool& bAllowMultipleConnections, bool& bAdvancedPin)
+	static void BreakPinProperty(const FPCGPinProperties& PinProperty, FName& Label, bool& bAllowMultipleData, bool& bAllowMultipleConnections, bool& bAdvancedPin, EPCGExclusiveDataType& AllowedType)
 	{
 		Label = PinProperty.Label;
 
@@ -79,7 +79,7 @@ public:
 	}
 
 	UFUNCTION(BlueprintPure, Category="PCG|Pins", meta = (NativeMakeFunc))
-	static FPCGPinProperties MakePinProperty(FName Label, EPCGExclusiveDataType AllowedType, bool bAllowMultipleData, bool bAllowMultipleConnections, bool bAdvancedPin)
+	static FPCGPinProperties MakePinProperty(FName Label, bool bAllowMultipleData, bool bAllowMultipleConnections, bool bAdvancedPin, EPCGExclusiveDataType AllowedType = EPCGExclusiveDataType::Any)
 	{
 		const UEnum* DataTypeEnum = StaticEnum<EPCGDataType>();
 		const UEnum* ExclusiveDataTypeEnum = StaticEnum<EPCGExclusiveDataType>();
