Fixed index usage error in the graph cancellation for sleeping tasks

#rnx

diff --git a/Engine/Plugins/Experimental/PCG/Source/PCG/Private/Graph/PCGGraphExecutor.cpp b/Engine/Plugins/Experimental/PCG/Source/PCG/Private/Graph/PCGGraphExecutor.cpp
--- a/Engine/Plugins/Experimental/PCG/Source/PCG/Private/Graph/PCGGraphExecutor.cpp
+++ b/Engine/Plugins/Experimental/PCG/Source/PCG/Private/Graph/PCGGraphExecutor.cpp
@@ -335,7 +335,7 @@ TSet<UPCGComponent*> FPCGGraphExecutor::Cancel(TFunctionRef<bool(TWeakObjectPtr<
 			if(Task.Context && CancelledComponents.Contains(Task.Context->SourceComponent.Get()))
 			{
 				FPCGTaskId CancelledTaskId = Task.NodeId;
-				SleepingTasks.RemoveAtSwap(CancelledTaskId);
+				SleepingTasks.RemoveAtSwap(SleepingTaskIndex);
 				bStableCancellationSet &= !CancelNextTasks(CancelledTaskId, CancelledComponents);
 			}
 		}
