[PCG] Minor tidy up in graph execution cancelling logic

diff --git a/Engine/Plugins/Experimental/PCG/Source/PCG/Private/Graph/PCGGraphExecutor.cpp b/Engine/Plugins/Experimental/PCG/Source/PCG/Private/Graph/PCGGraphExecutor.cpp
--- a/Engine/Plugins/Experimental/PCG/Source/PCG/Private/Graph/PCGGraphExecutor.cpp
+++ b/Engine/Plugins/Experimental/PCG/Source/PCG/Private/Graph/PCGGraphExecutor.cpp
@@ -271,11 +271,6 @@ TSet<UPCGComponent*> FPCGGraphExecutor::Cancel(TFunctionRef<bool(TWeakObjectPtr<
 			}
 		}
 
-		auto RemoveScheduledTasks = [this, &bStableCancellationSet, &CancelledComponents, &CancelledScheduledTasks](FPCGTaskId EndTaskId)
-		{
-
-		};
-
 		// WARNING: variable upper bound
 		for (int32 CancelledTaskIdIndex = 0; CancelledTaskIdIndex < CancelledScheduledTasks.Num(); ++CancelledTaskIdIndex)
 		{
@@ -320,7 +315,7 @@ TSet<UPCGComponent*> FPCGGraphExecutor::Cancel(TFunctionRef<bool(TWeakObjectPtr<
 		for (int32 ActiveTaskIndex = ActiveTasks.Num() - 1; ActiveTaskIndex >= 0; --ActiveTaskIndex)
 		{
 			FPCGGraphActiveTask& Task = ActiveTasks[ActiveTaskIndex];
-			if(Task.Context && CancelledComponents.Contains(Task.Context->SourceComponent.Get()))
+			if (Task.Context && CancelledComponents.Contains(Task.Context->SourceComponent.Get()))
 			{
 				FPCGTaskId CancelledTaskId = Task.NodeId;
 				Task.bWasCancelled = true;
@@ -332,7 +327,7 @@ TSet<UPCGComponent*> FPCGGraphExecutor::Cancel(TFunctionRef<bool(TWeakObjectPtr<
 		for (int32 SleepingTaskIndex = SleepingTasks.Num() - 1; SleepingTaskIndex >= 0; --SleepingTaskIndex)
 		{
 			FPCGGraphActiveTask& Task = SleepingTasks[SleepingTaskIndex];
-			if(Task.Context && CancelledComponents.Contains(Task.Context->SourceComponent.Get()))
+			if (Task.Context && CancelledComponents.Contains(Task.Context->SourceComponent.Get()))
 			{
 				FPCGTaskId CancelledTaskId = Task.NodeId;
 				SleepingTasks.RemoveAtSwap(SleepingTaskIndex);
