Landscapes would not get tracked properly, causing potential missing landscape layers when calling build landscape cache

#rnx

diff --git a/Engine/Plugins/Experimental/PCG/Source/PCG/Private/Grid/PCGLandscapeCache.cpp b/Engine/Plugins/Experimental/PCG/Source/PCG/Private/Grid/PCGLandscapeCache.cpp
--- a/Engine/Plugins/Experimental/PCG/Source/PCG/Private/Grid/PCGLandscapeCache.cpp
+++ b/Engine/Plugins/Experimental/PCG/Source/PCG/Private/Grid/PCGLandscapeCache.cpp
@@ -567,9 +567,10 @@ void UPCGLandscapeCache::PrimeCache()
 		if (IsValid(LandscapeInfo))
 		{
 			// Build per-component information
-			LandscapeInfo->ForEachLandscapeProxy([this, LandscapeInfo, &CacheEntriesToBuild](const ALandscapeProxy* LandscapeProxy)
+			LandscapeInfo->ForEachLandscapeProxy([this, LandscapeInfo, &CacheEntriesToBuild](ALandscapeProxy* LandscapeProxy)
 			{
 				check(LandscapeProxy);
+				Landscapes.Add(LandscapeProxy);
 				const FGuid LandscapeGuid = LandscapeProxy->GetOriginalLandscapeGuid();
 
 				for (ULandscapeComponent* LandscapeComponent : LandscapeProxy->LandscapeComponents)
