Changed initialization order for partitioned graphs to prevent ordering issue wrt landscape - fixes a recent regression

#rnx

diff --git a/Engine/Plugins/Experimental/PCG/Source/PCG/Private/Grid/PCGPartitionActor.cpp b/Engine/Plugins/Experimental/PCG/Source/PCG/Private/Grid/PCGPartitionActor.cpp
--- a/Engine/Plugins/Experimental/PCG/Source/PCG/Private/Grid/PCGPartitionActor.cpp
+++ b/Engine/Plugins/Experimental/PCG/Source/PCG/Private/Grid/PCGPartitionActor.cpp
@@ -312,15 +312,18 @@ void APCGPartitionActor::AddGraphInstance(UPCGComponent* OriginalComponent)
 
 	// Create a new local component
 	LocalComponent = NewObject<UPCGComponent>(this);
+	LocalComponent->MarkAsLocalComponent();
+
+	// Note: we'll place the local component prior to the SetPropertiesFromOriginal so that any code that relies on the parent-child relationship works here
+	OriginalToLocal.Emplace(OriginalComponent, LocalComponent);
+	LocalToOriginal.Emplace(LocalComponent, OriginalComponent);
+
 	LocalComponent->SetPropertiesFromOriginal(OriginalComponent);
 
 	LocalComponent->RegisterComponent();
-	LocalComponent->MarkAsLocalComponent();
+
 	// TODO: check if we should use a non-instanced component?
 	AddInstanceComponent(LocalComponent);
-
-	OriginalToLocal.Emplace(OriginalComponent, LocalComponent);
-	LocalToOriginal.Emplace(LocalComponent, OriginalComponent);
 }
 
 void APCGPartitionActor::RemapGraphInstance(const UPCGComponent* OldOriginalComponent, UPCGComponent* NewOriginalComponent)
