[PCG] Fix issue with OverridableParams that were not serialized correctly

* Also removing any editor only properties from override list at runtime

diff --git a/Engine/Plugins/Experimental/PCG/Source/PCG/Private/PCGSettings.cpp b/Engine/Plugins/Experimental/PCG/Source/PCG/Private/PCGSettings.cpp
--- a/Engine/Plugins/Experimental/PCG/Source/PCG/Private/PCGSettings.cpp
+++ b/Engine/Plugins/Experimental/PCG/Source/PCG/Private/PCGSettings.cpp
@@ -174,12 +174,6 @@ void UPCGSettings::PostInitProperties()
 
 void UPCGSettings::Serialize(FArchive& Ar)
 {
-	// Don't serialize overridable params in non-cooked builds
-	if (Ar.IsSaving() && !Ar.IsCooking())
-	{
-		CachedOverridableParams.Empty();
-	}
-
 	Super::Serialize(Ar);
 
 	Ar.UsingCustomVersion(FPCGCustomVersion::GUID);
@@ -211,12 +205,6 @@ void UPCGSettings::Serialize(FArchive& Ar)
 		}
 #endif // WITH_EDITOR
 	}
-
-	// Reconstruct it at the end
-	if (Ar.IsSaving() && !Ar.IsCooking())
-	{
-		InitializeCachedOverridableParams(/*bReset=*/true);
-	}
 }
 
 void UPCGSettings::PostSaveRoot(FObjectPostSaveRootContext ObjectSaveContext)
@@ -573,22 +561,25 @@ void UPCGSettings::InitializeCachedOverridableParams(bool bReset)
 
 		Param.Properties.Reset(Param.PropertiesNames.Num());
 
-		// Some properties might not be available at runtime. Ignore them.
+		// Some properties might not be available at runtime. Remove them from the list, since they won't be overridable anymore.
 		const FProperty* CurrentProperty = Param.PropertyClass->FindPropertyByName(Param.PropertiesNames[0]);
-		if (CurrentProperty)
+		if (!CurrentProperty)
 		{
-			Param.Properties.Add(CurrentProperty);
+			CachedOverridableParams.RemoveAt(i--);
+			continue;
+		}
+
+		Param.Properties.Add(CurrentProperty);
 
-			for (int32 j = 1; j < Param.PropertiesNames.Num(); ++j)
+		for (int32 j = 1; j < Param.PropertiesNames.Num(); ++j)
+		{
+			// If we have multiple depth properties, it should be Struct properties by construction
+			const FStructProperty* StructProperty = CastField<FStructProperty>(CurrentProperty);
+			if (ensure(StructProperty))
 			{
-				// If we have multiple depth properties, it should be Struct properties by construction
-				const FStructProperty* StructProperty = CastField<FStructProperty>(CurrentProperty);
-				if (ensure(StructProperty))
-				{
-					CurrentProperty = StructProperty->Struct->FindPropertyByName(Param.PropertiesNames[j]);
-					check(CurrentProperty);
-					Param.Properties.Add(CurrentProperty);
-				}
+				CurrentProperty = StructProperty->Struct->FindPropertyByName(Param.PropertiesNames[j]);
+				check(CurrentProperty);
+				Param.Properties.Add(CurrentProperty);
 			}
 		}
 	}
