Tentative fix for crash on update pins related to overridable params

#rnx

diff --git a/Engine/Plugins/Experimental/PCG/Source/PCG/Private/PCGSettings.cpp b/Engine/Plugins/Experimental/PCG/Source/PCG/Private/PCGSettings.cpp
--- a/Engine/Plugins/Experimental/PCG/Source/PCG/Private/PCGSettings.cpp
+++ b/Engine/Plugins/Experimental/PCG/Source/PCG/Private/PCGSettings.cpp
@@ -316,19 +316,23 @@ void UPCGSettings::FillOverridableParamsPins(TArray<FPCGPinProperties>& OutPins)
 		FPCGPinProperties& ParamPin = OutPins.Emplace_GetRef(OverridableParam.Label, EPCGDataType::Param, /*bInAllowMultipleConnections=*/ false, /*bAllowMultipleData=*/ false);
 		ParamPin.bAdvancedPin = true;
 #if WITH_EDITOR
-		const FProperty* Property = OverridableParam.Properties.Last();
-		check(Property);
-		static FName TooltipMetadata("Tooltip");
-		FString Tooltip;
-		if (const FString* TooltipPtr = Property->FindMetaData(TooltipMetadata))
+
+		if (!OverridableParam.Properties.IsEmpty())
 		{
-			Tooltip = *TooltipPtr + TEXT("\n");
-		}
+			const FProperty* Property = OverridableParam.Properties.Last();
+			check(Property);
+			static FName TooltipMetadata("Tooltip");
+			FString Tooltip;
+			if (const FString* TooltipPtr = Property->FindMetaData(TooltipMetadata))
+			{
+				Tooltip = *TooltipPtr + TEXT("\n");
+			}
 
-		ParamPin.Tooltip = FText::Format(LOCTEXT("OverridableParamPinTooltip", "{0}Attribute type is \"{1}\" and its exact name is \"{2}\""),
-			FText::FromString(Tooltip),
-			FText::FromString(Property->GetCPPType()),
-			FText::FromString(OverridableParam.GetPropertyPath()));
+			ParamPin.Tooltip = FText::Format(LOCTEXT("OverridableParamPinTooltip", "{0}Attribute type is \"{1}\" and its exact name is \"{2}\""),
+				FText::FromString(Tooltip),
+				FText::FromString(Property->GetCPPType()),
+				FText::FromString(OverridableParam.GetPropertyPath()));
+		}
 #endif // WITH_EDITOR
 	}
 }
