[PCG] Rename InstancePacker to InstanceDataPacker

diff --git a/Engine/Plugins/Experimental/PCG/Config/BasePCG.ini b/Engine/Plugins/Experimental/PCG/Config/BasePCG.ini
--- a/Engine/Plugins/Experimental/PCG/Config/BasePCG.ini
+++ b/Engine/Plugins/Experimental/PCG/Config/BasePCG.ini
@@ -22,8 +22,14 @@
 +PropertyRedirects=(OldName="PCGBlueprintElement.NestedLoopBody.InB",NewName="PCGBlueprintElement.NestedLoopBody.InInnerData")
 +PropertyRedirects=(OldName="PCGBlueprintElement.NestedLoopBody.InPointA",NewName="PCGBlueprintElement.NestedLoopBody.InOuterPoint")
 +PropertyRedirects=(OldName="PCGBlueprintElement.NestedLoopBody.InPointB",NewName="PCGBlueprintElement.NestedLoopBody.InInnerPoint")
++PropertyRedirects=(OldName="PCGStaticMeshSpawnerSettings.MeshSelectorInstance",NewName="PCGStaticMeshSpawnerSettings.MeshSelectorParameters")
++PropertyRedirects=(OldName="PCGStaticMeshSpawnerSettings.InstancePackerType",NewName="PCGStaticMeshSpawnerSettings.InstanceDataPackerType")
++PropertyRedirects=(OldName="PCGStaticMeshSpawnerSettings.InstancePackerInstance",NewName="PCGStaticMeshSpawnerSettings.InstanceDataPackerParameters")
 
 +ClassRedirects=(OldName="PCGPointSamplerSettings",NewName="/Script/PCG.PCGSelectPointsSettings")
 +ClassRedirects=(OldName="PCGDensityNoiseSettings",NewName="/Script/PCG.PCGAttributeNoiseSettings")
++ClassRedirects=(OldName="PCGInstancePackerBase",NewName="/Script/PCG.PCGInstanceDataPackerBase")
++ClassRedirects=(OldName="PCGInstancePackerByAttribute",NewName="/Script/PCG.PCGInstanceDataPackerByAttribute")
++ClassRedirects=(OldName="PCGInstancePackerByRegex",NewName="/Script/PCG.PCGInstanceDataPackerByRegex")
 
 +EnumRedirects=(OldName="EPCGDensityNoiseMode",NewName="EPCGAttributeNoiseMode")
diff --git a/Engine/Plugins/Experimental/PCG/Source/PCG/Private/Elements/PCGStaticMeshSpawner.cpp b/Engine/Plugins/Experimental/PCG/Source/PCG/Private/Elements/PCGStaticMeshSpawner.cpp
--- a/Engine/Plugins/Experimental/PCG/Source/PCG/Private/Elements/PCGStaticMeshSpawner.cpp
+++ b/Engine/Plugins/Experimental/PCG/Source/PCG/Private/Elements/PCGStaticMeshSpawner.cpp
@@ -8,7 +8,7 @@
 #include "Data/PCGSpatialData.h"
 #include "Elements/PCGStaticMeshSpawnerContext.h"
 #include "Helpers/PCGActorHelpers.h"
-#include "InstancePackers/PCGInstancePackerBase.h"
+#include "InstanceDataPackers/PCGInstanceDataPackerBase.h"
 #include "MeshSelectors/PCGMeshSelectorBase.h"
 #include "MeshSelectors/PCGMeshSelectorWeighted.h"
 
@@ -36,7 +36,7 @@ UPCGStaticMeshSpawnerSettings::UPCGStaticMeshSpawnerSettings(const FObjectInitia
 	// However, removing it makes it that any object actually using the instance created by default would be lost.
 	if (!this->HasAnyFlags(RF_ClassDefaultObject))
 	{
-		MeshSelectorInstance = ObjectInitializer.CreateDefaultSubobject<UPCGMeshSelectorWeighted>(this, TEXT("DefaultSelectorInstance"));
+		MeshSelectorParameters = ObjectInitializer.CreateDefaultSubobject<UPCGMeshSelectorWeighted>(this, TEXT("DefaultSelectorInstance"));
 	}
 }
 
@@ -53,9 +53,9 @@ bool FPCGStaticMeshSpawnerElement::PrepareDataInternal(FPCGContext* InContext) c
 	const UPCGStaticMeshSpawnerSettings* Settings = Context->GetInputSettings<UPCGStaticMeshSpawnerSettings>();
 	check(Settings);
 
-	if (!Settings->MeshSelectorInstance)
+	if (!Settings->MeshSelectorParameters)
 	{
-		PCGE_LOG(Error, GraphAndLog, LOCTEXT("InvalidMeshSelectorInstance", "Invalid MatchAndSet instance, try recreating this node from the node palette"));
+		PCGE_LOG(Error, GraphAndLog, LOCTEXT("InvalidMeshSelectorInstance", "Invalid MeshSelector instance, try reselecting the MeshSelector type"));
 		return true;
 	}
 
@@ -188,7 +188,7 @@ bool FPCGStaticMeshSpawnerElement::PrepareDataInternal(FPCGContext* InContext) c
 			TArray<FPCGMeshInstanceList>& MeshInstances = (bSkippedDueToReuse ? DummyMeshInstances : Context->MeshInstancesData.Last().MeshInstances);
 
 			check(Context->CurrentPointData);
-			Context->bSelectionDone = Settings->MeshSelectorInstance->SelectInstances(*Context, Settings, Context->CurrentPointData, MeshInstances, Context->CurrentOutputPointData);
+			Context->bSelectionDone = Settings->MeshSelectorParameters->SelectInstances(*Context, Settings, Context->CurrentPointData, MeshInstances, Context->CurrentOutputPointData);
 		}
 
 		if (!Context->bSelectionDone)
@@ -207,11 +207,11 @@ bool FPCGStaticMeshSpawnerElement::PrepareDataInternal(FPCGContext* InContext) c
 				PackedCustomData.SetNum(MeshInstances.Num());
 			}
 
-			if (Settings->InstancePackerInstance)
+			if (Settings->InstanceDataPackerParameters)
 			{
 				for (int32 InstanceListIndex = 0; InstanceListIndex < MeshInstances.Num(); ++InstanceListIndex)
 				{
-					Settings->InstancePackerInstance->PackInstances(*Context, Context->CurrentPointData, MeshInstances[InstanceListIndex], PackedCustomData[InstanceListIndex]);
+					Settings->InstanceDataPackerParameters->PackInstances(*Context, Context->CurrentPointData, MeshInstances[InstanceListIndex], PackedCustomData[InstanceListIndex]);
 				}
 			}
 		}
@@ -382,7 +382,7 @@ void UPCGStaticMeshSpawnerSettings::PostLoad()
 	{
 		SetMeshSelectorType(UPCGMeshSelectorWeighted::StaticClass());
 
-		UPCGMeshSelectorWeighted* MeshSelector = CastChecked<UPCGMeshSelectorWeighted>(MeshSelectorInstance);
+		UPCGMeshSelectorWeighted* MeshSelector = CastChecked<UPCGMeshSelectorWeighted>(MeshSelectorParameters);
 
 		for (const FPCGStaticMeshSpawnerEntry& Entry : Meshes_DEPRECATED)
 		{
@@ -399,22 +399,22 @@ void UPCGStaticMeshSpawnerSettings::PostLoad()
 
 	const EObjectFlags Flags = GetMaskedFlags(RF_PropagateToSubObjects) | RF_Transactional;
 	
-	if (!MeshSelectorInstance)
+	if (!MeshSelectorParameters)
 	{
 		RefreshMeshSelector();
 	}
 	else
 	{
-		MeshSelectorInstance->SetFlags(Flags);
+		MeshSelectorParameters->SetFlags(Flags);
 	}
 
-	if (!InstancePackerInstance)
+	if (!InstanceDataPackerParameters)
 	{
 		RefreshInstancePacker();
 	}
 	else
 	{
-		InstancePackerInstance->SetFlags(Flags);
+		InstanceDataPackerParameters->SetFlags(Flags);
 	}
 }
 
@@ -429,7 +429,7 @@ void UPCGStaticMeshSpawnerSettings::PostEditChangeProperty(FPropertyChangedEvent
 		{
 			RefreshMeshSelector();
 		} 
-		else if (PropertyName == GET_MEMBER_NAME_CHECKED(UPCGStaticMeshSpawnerSettings, InstancePackerType))
+		else if (PropertyName == GET_MEMBER_NAME_CHECKED(UPCGStaticMeshSpawnerSettings, InstanceDataPackerType))
 		{
 			RefreshInstancePacker();
 		}
@@ -441,7 +441,7 @@ void UPCGStaticMeshSpawnerSettings::PostEditChangeProperty(FPropertyChangedEvent
 
 void UPCGStaticMeshSpawnerSettings::SetMeshSelectorType(TSubclassOf<UPCGMeshSelectorBase> InMeshSelectorType) 
 {
-	if (!MeshSelectorInstance || InMeshSelectorType != MeshSelectorType)
+	if (!MeshSelectorParameters || InMeshSelectorType != MeshSelectorType)
 	{
 		if (InMeshSelectorType != MeshSelectorType)
 		{
@@ -452,13 +452,13 @@ void UPCGStaticMeshSpawnerSettings::SetMeshSelectorType(TSubclassOf<UPCGMeshSele
 	}
 }
 
-void UPCGStaticMeshSpawnerSettings::SetInstancePackerType(TSubclassOf<UPCGInstancePackerBase> InInstancePackerType) 
+void UPCGStaticMeshSpawnerSettings::SetInstancePackerType(TSubclassOf<UPCGInstanceDataPackerBase> InInstancePackerType) 
 {
-	if (!InstancePackerInstance || InInstancePackerType != InstancePackerType)
+	if (!InstanceDataPackerParameters || InInstancePackerType != InstanceDataPackerType)
 	{
-		if (InInstancePackerType != InstancePackerType)
+		if (InInstancePackerType != InstanceDataPackerType)
 		{
-			InstancePackerType = InInstancePackerType;
+			InstanceDataPackerType = InInstancePackerType;
 		}
 		
 		RefreshInstancePacker();
@@ -469,39 +469,39 @@ void UPCGStaticMeshSpawnerSettings::RefreshMeshSelector()
 {
 	if (MeshSelectorType)
 	{
-		if (MeshSelectorInstance)
+		if (MeshSelectorParameters)
 		{
-			MeshSelectorInstance->Rename(nullptr, GetTransientPackage(), REN_DontCreateRedirectors | REN_ForceNoResetLoaders);
-			MeshSelectorInstance->MarkAsGarbage();
-			MeshSelectorInstance = nullptr;
+			MeshSelectorParameters->Rename(nullptr, GetTransientPackage(), REN_DontCreateRedirectors | REN_ForceNoResetLoaders);
+			MeshSelectorParameters->MarkAsGarbage();
+			MeshSelectorParameters = nullptr;
 		}
 
 		const EObjectFlags Flags = GetMaskedFlags(RF_PropagateToSubObjects);
-		MeshSelectorInstance = NewObject<UPCGMeshSelectorBase>(this, MeshSelectorType, NAME_None, Flags);
+		MeshSelectorParameters = NewObject<UPCGMeshSelectorBase>(this, MeshSelectorType, NAME_None, Flags);
 	}
 	else
 	{
-		MeshSelectorInstance = nullptr;
+		MeshSelectorParameters = nullptr;
 	}
 }
 
 void UPCGStaticMeshSpawnerSettings::RefreshInstancePacker()
 {
-	if (InstancePackerType)
+	if (InstanceDataPackerType)
 	{
-		if (InstancePackerInstance)
+		if (InstanceDataPackerParameters)
 		{
-			InstancePackerInstance->Rename(nullptr, GetTransientPackage(), REN_DontCreateRedirectors | REN_ForceNoResetLoaders);
-			InstancePackerInstance->MarkAsGarbage();
-			InstancePackerInstance = nullptr;
+			InstanceDataPackerParameters->Rename(nullptr, GetTransientPackage(), REN_DontCreateRedirectors | REN_ForceNoResetLoaders);
+			InstanceDataPackerParameters->MarkAsGarbage();
+			InstanceDataPackerParameters = nullptr;
 		}
 
 		const EObjectFlags Flags = GetMaskedFlags(RF_PropagateToSubObjects);
-		InstancePackerInstance = NewObject<UPCGInstancePackerBase>(this, InstancePackerType, NAME_None, Flags);
+		InstanceDataPackerParameters = NewObject<UPCGInstanceDataPackerBase>(this, InstanceDataPackerType, NAME_None, Flags);
 	}
 	else
 	{
-		InstancePackerInstance = nullptr;
+		InstanceDataPackerParameters = nullptr;
 	}
 }
 
diff --git a/Engine/Plugins/Experimental/PCG/Source/PCG/Private/InstanceDataPackers/PCGInstanceDataPackerBase.cpp b/Engine/Plugins/Experimental/PCG/Source/PCG/Private/InstanceDataPackers/PCGInstanceDataPackerBase.cpp
new file mode 100644
--- /dev/null
+++ b/Engine/Plugins/Experimental/PCG/Source/PCG/Private/InstanceDataPackers/PCGInstanceDataPackerBase.cpp
@@ -0,0 +1,140 @@
+// Copyright Epic Games, Inc. All Rights Reserved.
+
+#include "InstanceDataPackers/PCGInstanceDataPackerBase.h"
+
+#include "MeshSelectors/PCGMeshSelectorBase.h"
+#include "Metadata/PCGMetadata.h"
+
+#include UE_INLINE_GENERATED_CPP_BY_NAME(PCGInstanceDataPackerBase)
+
+bool UPCGInstanceDataPackerBase::AddTypeToPacking(int TypeId, FPCGPackedCustomData& OutPackedCustomData) const
+{
+	if (TypeId == PCG::Private::MetadataTypes<float>::Id)
+	{
+		OutPackedCustomData.NumCustomDataFloats += 1;
+	}
+	else if (TypeId == PCG::Private::MetadataTypes<double>::Id)
+	{
+		OutPackedCustomData.NumCustomDataFloats += 1;
+	}
+	else if (TypeId == PCG::Private::MetadataTypes<int32>::Id)
+	{
+		OutPackedCustomData.NumCustomDataFloats += 1;
+	}
+	else if (TypeId == PCG::Private::MetadataTypes<int64>::Id)
+	{
+		OutPackedCustomData.NumCustomDataFloats += 1;
+	}
+	else if (TypeId == PCG::Private::MetadataTypes<FRotator>::Id)
+	{
+		OutPackedCustomData.NumCustomDataFloats += 3;
+	}
+	else if (TypeId == PCG::Private::MetadataTypes<FVector>::Id)
+	{
+		OutPackedCustomData.NumCustomDataFloats += 3;
+	}
+	else if (TypeId == PCG::Private::MetadataTypes<FVector4>::Id)
+	{
+		OutPackedCustomData.NumCustomDataFloats += 4;
+	}
+	else
+	{
+		return false;
+	}
+
+	return true;
+}
+
+void UPCGInstanceDataPackerBase::PackCustomDataFromAttributes(const FPCGMeshInstanceList& InstanceList, const UPCGMetadata* Metadata, const TArray<FName>& AttributeNames, FPCGPackedCustomData& OutPackedCustomData) const
+{
+	if (!Metadata)
+	{
+		UE_LOG(LogPCG, Error, TEXT("Invalid metadata"));
+		return;
+	}
+
+	TArray<const FPCGMetadataAttributeBase*> Attributes;
+	for (const FName& AttributeName : AttributeNames)
+	{
+		Attributes.Add(Metadata->GetConstAttribute(AttributeName));
+	}
+
+	PackCustomDataFromAttributes(InstanceList, Attributes, OutPackedCustomData);
+}
+
+void UPCGInstanceDataPackerBase::PackCustomDataFromAttributes(const FPCGMeshInstanceList& InstanceList, const TArray<const FPCGMetadataAttributeBase*>& Attributes, FPCGPackedCustomData& OutPackedCustomData) const
+{
+	for(uint64 PointMetadataEntry : InstanceList.InstancesMetadataEntry)
+	{
+		for (const FPCGMetadataAttributeBase* AttributeBase : Attributes)
+		{
+			check(AttributeBase);
+
+			if (AttributeBase->GetTypeId() == PCG::Private::MetadataTypes<float>::Id)
+			{
+				const FPCGMetadataAttribute<float>* Attribute = static_cast<const FPCGMetadataAttribute<float>*>(AttributeBase);
+				check(Attribute);
+
+				const float Value = Attribute->GetValueFromItemKey(PointMetadataEntry);
+				OutPackedCustomData.CustomData.Add(Value);
+			}
+			else if (AttributeBase->GetTypeId() == PCG::Private::MetadataTypes<double>::Id)
+			{
+				const FPCGMetadataAttribute<double>* Attribute = static_cast<const FPCGMetadataAttribute<double>*>(AttributeBase);
+				check(Attribute);
+
+				const double Value = Attribute->GetValueFromItemKey(PointMetadataEntry);
+				OutPackedCustomData.CustomData.Add(Value);
+			}
+			else if (AttributeBase->GetTypeId() == PCG::Private::MetadataTypes<int32>::Id)
+			{
+				const FPCGMetadataAttribute<int32>* Attribute = static_cast<const FPCGMetadataAttribute<int32>*>(AttributeBase);
+				check(Attribute);
+
+				const float Value = static_cast<float>(Attribute->GetValueFromItemKey(PointMetadataEntry));
+				OutPackedCustomData.CustomData.Add(Value);
+			}
+			else if (AttributeBase->GetTypeId() == PCG::Private::MetadataTypes<int64>::Id)
+			{
+				const FPCGMetadataAttribute<int64>* Attribute = static_cast<const FPCGMetadataAttribute<int64>*>(AttributeBase);
+				check(Attribute);
+
+				const float Value = static_cast<float>(Attribute->GetValueFromItemKey(PointMetadataEntry));
+				OutPackedCustomData.CustomData.Add(Value);
+			}
+			else if (AttributeBase->GetTypeId() == PCG::Private::MetadataTypes<FRotator>::Id)
+			{
+				const FPCGMetadataAttribute<FRotator>* Attribute = static_cast<const FPCGMetadataAttribute<FRotator>*>(AttributeBase);
+				check(Attribute);
+
+				const FRotator Value = Attribute->GetValueFromItemKey(PointMetadataEntry);
+				OutPackedCustomData.CustomData.Add(Value.Roll);
+				OutPackedCustomData.CustomData.Add(Value.Pitch);
+				OutPackedCustomData.CustomData.Add(Value.Yaw);
+			}
+			else if (AttributeBase->GetTypeId() == PCG::Private::MetadataTypes<FVector>::Id)
+			{
+				const FPCGMetadataAttribute<FVector>* Attribute = static_cast<const FPCGMetadataAttribute<FVector>*>(AttributeBase);
+				check(Attribute);
+
+				const FVector Value = Attribute->GetValueFromItemKey(PointMetadataEntry);
+				OutPackedCustomData.CustomData.Add(Value.X);
+				OutPackedCustomData.CustomData.Add(Value.Y);
+				OutPackedCustomData.CustomData.Add(Value.Z);
+			}
+			else if (AttributeBase->GetTypeId() == PCG::Private::MetadataTypes<FVector4>::Id)
+			{
+				const FPCGMetadataAttribute<FVector4>* Attribute = static_cast<const FPCGMetadataAttribute<FVector4>*>(AttributeBase);
+				check(Attribute);
+
+				const FVector4 Value = Attribute->GetValueFromItemKey(PointMetadataEntry);
+				OutPackedCustomData.CustomData.Add(Value.X);
+				OutPackedCustomData.CustomData.Add(Value.Y);
+				OutPackedCustomData.CustomData.Add(Value.Z);
+				OutPackedCustomData.CustomData.Add(Value.W);
+			}
+		}
+	}
+}
+
+
diff --git a/Engine/Plugins/Experimental/PCG/Source/PCG/Private/InstanceDataPackers/PCGInstanceDataPackerByAttribute.cpp b/Engine/Plugins/Experimental/PCG/Source/PCG/Private/InstanceDataPackers/PCGInstanceDataPackerByAttribute.cpp
new file mode 100644
--- /dev/null
+++ b/Engine/Plugins/Experimental/PCG/Source/PCG/Private/InstanceDataPackers/PCGInstanceDataPackerByAttribute.cpp
@@ -0,0 +1,48 @@
+// Copyright Epic Games, Inc. All Rights Reserved.
+
+#include "InstanceDataPackers/PCGInstanceDataPackerByAttribute.h"
+
+#include "PCGContext.h"
+#include "PCGElement.h"
+#include "Data/PCGSpatialData.h"
+#include "InstanceDataPackers/PCGInstanceDataPackerBase.h"
+
+#include UE_INLINE_GENERATED_CPP_BY_NAME(PCGInstanceDataPackerByAttribute)
+
+#define LOCTEXT_NAMESPACE "PCGInstanceDataPackerByAttribute"
+
+void UPCGInstanceDataPackerByAttribute::PackInstances_Implementation(FPCGContext& Context, const UPCGSpatialData* InSpatialData, const FPCGMeshInstanceList& InstanceList, FPCGPackedCustomData& OutPackedCustomData) const
+{
+	if (!InSpatialData || !InSpatialData->Metadata)
+	{
+		PCGE_LOG_C(Error, GraphAndLog, &Context, LOCTEXT("InvalidInputData", "Invalid input data"));
+		return;
+	}
+
+	TArray<const FPCGMetadataAttributeBase*> SelectedAttributes;
+
+	// Find Attributes by name and calculate NumCustomDataFloats
+	for (const FName& AttributeName : AttributeNames)
+	{
+		if (!InSpatialData->Metadata->HasAttribute(AttributeName)) 
+		{
+			PCGE_LOG_C(Warning, GraphAndLog, &Context, FText::Format(LOCTEXT("AttributeNotInMetadata", "Attribute '{0}' is not in the metadata"), FText::FromName(AttributeName)));
+			continue;
+		}
+
+		const FPCGMetadataAttributeBase* AttributeBase = InSpatialData->Metadata->GetConstAttribute(AttributeName);
+		check(AttributeBase);
+
+		if (!AddTypeToPacking(AttributeBase->GetTypeId(), OutPackedCustomData))
+		{
+			PCGE_LOG_C(Warning, GraphAndLog, &Context, FText::Format(LOCTEXT("AttributeInvalidType", "Attribute name '{0}' is not a valid type"), FText::FromName(AttributeName)));
+			continue;
+		}
+
+		SelectedAttributes.Add(AttributeBase);
+	}
+
+	PackCustomDataFromAttributes(InstanceList, SelectedAttributes, OutPackedCustomData);
+ }
+
+#undef LOCTEXT_NAMESPACE
diff --git a/Engine/Plugins/Experimental/PCG/Source/PCG/Private/InstanceDataPackers/PCGInstanceDataPackerByRegex.cpp b/Engine/Plugins/Experimental/PCG/Source/PCG/Private/InstanceDataPackers/PCGInstanceDataPackerByRegex.cpp
new file mode 100644
--- /dev/null
+++ b/Engine/Plugins/Experimental/PCG/Source/PCG/Private/InstanceDataPackers/PCGInstanceDataPackerByRegex.cpp
@@ -0,0 +1,70 @@
+// Copyright Epic Games, Inc. All Rights Reserved.
+
+#include "InstanceDataPackers/PCGInstanceDataPackerByRegex.h"
+
+#include "PCGContext.h"
+#include "PCGElement.h"
+#include "Data/PCGSpatialData.h"
+#include "InstanceDataPackers/PCGInstanceDataPackerBase.h"
+
+#include "Internationalization/Regex.h"
+
+#include UE_INLINE_GENERATED_CPP_BY_NAME(PCGInstanceDataPackerByRegex)
+
+#define LOCTEXT_NAMESPACE "PCGInstanceDataPackerByRegex"
+
+void UPCGInstanceDataPackerByRegex::PackInstances_Implementation(FPCGContext& Context, const UPCGSpatialData* InSpatialData, const FPCGMeshInstanceList& InstanceList, FPCGPackedCustomData& OutPackedCustomData) const
+{
+	if (!InSpatialData || !InSpatialData->Metadata)
+	{
+		PCGE_LOG_C(Error, GraphAndLog, &Context, LOCTEXT("InvalidInputData", "Invalid input data"));
+		return;
+	}
+	
+	TArray<FName> AttributeNames;
+	TArray<EPCGMetadataTypes> AttributeTypes;
+	InSpatialData->Metadata->GetAttributes(AttributeNames, AttributeTypes);
+
+	TArray<const FPCGMetadataAttributeBase*> SelectedAttributes;
+	TSet<PCGMetadataAttributeKey> FoundAttributes;
+
+	for (const FString& Regex : RegexPatterns)
+	{
+		const FRegexPattern Pattern(Regex);
+	
+		// Match our Regex Pattern against every Attribute
+		for (int Index = 0; Index < AttributeNames.Num(); ++Index)
+		{
+			const FName& AttributeName = AttributeNames[Index];
+			FRegexMatcher RegexMatcher(Pattern, AttributeName.ToString());
+	
+			if (!RegexMatcher.FindNext()) 
+			{
+				continue;
+			}
+	
+			const FPCGMetadataAttributeBase* AttributeBase = InSpatialData->Metadata->GetConstAttribute(AttributeName);
+			check(AttributeBase);
+
+			// Avoid adding the same attribute multiple times if it is captured by different regex patterns
+			if (FoundAttributes.Contains(AttributeBase->AttributeId))
+			{
+				continue;
+			}
+
+			FoundAttributes.Add(AttributeBase->AttributeId);
+	
+			if (!AddTypeToPacking(AttributeBase->GetTypeId(), OutPackedCustomData))
+			{
+				PCGE_LOG_C(Warning, GraphAndLog, &Context, FText::Format(LOCTEXT("InvalidAttributeType", "Attribute name '{0}' is not of a valid type"), FText::FromName(AttributeName)));
+				continue;
+			}
+	
+			SelectedAttributes.Add(AttributeBase);
+		}
+	}
+	
+	PackCustomDataFromAttributes(InstanceList, SelectedAttributes, OutPackedCustomData);
+}
+
+#undef LOCTEXT_NAMESPACE
diff --git a/Engine/Plugins/Experimental/PCG/Source/PCG/Private/InstancePackers/PCGInstancePackerBase.cpp b/Engine/Plugins/Experimental/PCG/Source/PCG/Private/InstancePackers/PCGInstancePackerBase.cpp
deleted file mode 100644
diff --git a/Engine/Plugins/Experimental/PCG/Source/PCG/Private/InstancePackers/PCGInstancePackerByAttribute.cpp b/Engine/Plugins/Experimental/PCG/Source/PCG/Private/InstancePackers/PCGInstancePackerByAttribute.cpp
deleted file mode 100644
diff --git a/Engine/Plugins/Experimental/PCG/Source/PCG/Private/InstancePackers/PCGInstancePackerByRegex.cpp b/Engine/Plugins/Experimental/PCG/Source/PCG/Private/InstancePackers/PCGInstancePackerByRegex.cpp
deleted file mode 100644
diff --git a/Engine/Plugins/Experimental/PCG/Source/PCG/Private/Tests/Elements/PCGInstancePackerTest.cpp b/Engine/Plugins/Experimental/PCG/Source/PCG/Private/Tests/Elements/PCGInstancePackerTest.cpp
--- a/Engine/Plugins/Experimental/PCG/Source/PCG/Private/Tests/Elements/PCGInstancePackerTest.cpp
+++ b/Engine/Plugins/Experimental/PCG/Source/PCG/Private/Tests/Elements/PCGInstancePackerTest.cpp
@@ -9,9 +9,9 @@
 #include "Data/PCGPointData.h"
 #include "Elements/PCGStaticMeshSpawner.h"
 #include "Elements/PCGStaticMeshSpawnerContext.h"
-#include "InstancePackers/PCGInstancePackerBase.h"
-#include "InstancePackers/PCGInstancePackerByAttribute.h"
-#include "InstancePackers/PCGInstancePackerByRegex.h"
+#include "InstanceDataPackers/PCGInstanceDataPackerBase.h"
+#include "InstanceDataPackers/PCGInstanceDataPackerByAttribute.h"
+#include "InstanceDataPackers/PCGInstanceDataPackerByRegex.h"
 #include "MeshSelectors/PCGMeshSelectorWeighted.h"
 
 #include "Components/InstancedStaticMeshComponent.h"
@@ -169,9 +169,9 @@ bool FPCGStaticMeshSpawnerInstancePackerByAttributeTest::RunTest(const FString&
 		UPCGMetadataAccessorHelpers::SetRotatorAttribute(SourcePoints[I], SourceData->Metadata, RotatorName, RotatorValue);
 	}
 
-	Settings->SetInstancePackerType(UPCGInstancePackerByAttribute::StaticClass());
-	UPCGInstancePackerByAttribute* InstancePacker = CastChecked<UPCGInstancePackerByAttribute>(Settings->InstancePackerInstance);
-	UPCGMeshSelectorWeighted* MeshSelector = CastChecked<UPCGMeshSelectorWeighted>(Settings->MeshSelectorInstance);
+	Settings->SetInstancePackerType(UPCGInstanceDataPackerByAttribute::StaticClass());
+	UPCGInstanceDataPackerByAttribute* InstancePacker = CastChecked<UPCGInstanceDataPackerByAttribute>(Settings->InstanceDataPackerParameters);
+	UPCGMeshSelectorWeighted* MeshSelector = CastChecked<UPCGMeshSelectorWeighted>(Settings->MeshSelectorParameters);
 
 	const FString CubePath = TEXT("StaticMesh'/Engine/BasicShapes/Cube.Cube'");
 	const UStaticMesh* CubeMesh = Cast<UStaticMesh>(StaticLoadObject(UStaticMesh::StaticClass(), NULL, *CubePath));
@@ -318,9 +318,9 @@ bool FPCGStaticMeshSpawnerInstancePackerByRegexTest::RunTest(const FString& Para
 		UPCGMetadataAccessorHelpers::SetVector4Attribute(SourcePoints[I], SourceData->Metadata, NameB2, Vec4Value);
 	}
 
-	Settings->SetInstancePackerType(UPCGInstancePackerByRegex::StaticClass());
-	UPCGInstancePackerByRegex* InstancePacker = CastChecked<UPCGInstancePackerByRegex>(Settings->InstancePackerInstance);
-	UPCGMeshSelectorWeighted* MeshSelector = CastChecked<UPCGMeshSelectorWeighted>(Settings->MeshSelectorInstance);
+	Settings->SetInstancePackerType(UPCGInstanceDataPackerByRegex::StaticClass());
+	UPCGInstanceDataPackerByRegex* InstancePacker = CastChecked<UPCGInstanceDataPackerByRegex>(Settings->InstanceDataPackerParameters);
+	UPCGMeshSelectorWeighted* MeshSelector = CastChecked<UPCGMeshSelectorWeighted>(Settings->MeshSelectorParameters);
 
 	const FString CubePath = TEXT("StaticMesh'/Engine/BasicShapes/Cube.Cube'");
 	const UStaticMesh* CubeMesh = Cast<UStaticMesh>(StaticLoadObject(UStaticMesh::StaticClass(), NULL, *CubePath));
diff --git a/Engine/Plugins/Experimental/PCG/Source/PCG/Private/Tests/Elements/PCGStaticMeshSpawnerTest.cpp b/Engine/Plugins/Experimental/PCG/Source/PCG/Private/Tests/Elements/PCGStaticMeshSpawnerTest.cpp
--- a/Engine/Plugins/Experimental/PCG/Source/PCG/Private/Tests/Elements/PCGStaticMeshSpawnerTest.cpp
+++ b/Engine/Plugins/Experimental/PCG/Source/PCG/Private/Tests/Elements/PCGStaticMeshSpawnerTest.cpp
@@ -134,7 +134,7 @@ void TestMeshSelectorByAttribute(
 	// set the MeshSelector
 	UPCGStaticMeshSpawnerSettings* Settings = CastChecked<UPCGStaticMeshSpawnerSettings>(TestData.Settings);
 	Settings->SetMeshSelectorType(UPCGMeshSelectorByAttribute::StaticClass());
-	UPCGMeshSelectorByAttribute* MeshSelector = CastChecked<UPCGMeshSelectorByAttribute>(Settings->MeshSelectorInstance);
+	UPCGMeshSelectorByAttribute* MeshSelector = CastChecked<UPCGMeshSelectorByAttribute>(Settings->MeshSelectorParameters);
 	MeshSelector->AttributeName = AttributeName;
 
 	// initialize and execute the StaticMeshSpawner
@@ -257,7 +257,7 @@ void TestMeshSelectorWeighted(
 	// set the MeshSelector
 	UPCGStaticMeshSpawnerSettings* Settings = CastChecked<UPCGStaticMeshSpawnerSettings>(TestData.Settings);
 	Settings->SetMeshSelectorType(UPCGMeshSelectorWeighted::StaticClass());
-	UPCGMeshSelectorWeighted* MeshSelector = CastChecked<UPCGMeshSelectorWeighted>(Settings->MeshSelectorInstance);
+	UPCGMeshSelectorWeighted* MeshSelector = CastChecked<UPCGMeshSelectorWeighted>(Settings->MeshSelectorParameters);
 	MeshSelector->MeshEntries = Entries;
 
 	// initialize and execute the StaticMeshSpawner
@@ -343,7 +343,7 @@ void TestMeshSelectorWeightedByCategory(
 	// set the MeshSelector
 	UPCGStaticMeshSpawnerSettings* Settings = CastChecked<UPCGStaticMeshSpawnerSettings>(TestData.Settings);
 	Settings->SetMeshSelectorType(UPCGMeshSelectorWeightedByCategory::StaticClass());
-	UPCGMeshSelectorWeightedByCategory* MeshSelector = CastChecked<UPCGMeshSelectorWeightedByCategory>(Settings->MeshSelectorInstance);
+	UPCGMeshSelectorWeightedByCategory* MeshSelector = CastChecked<UPCGMeshSelectorWeightedByCategory>(Settings->MeshSelectorParameters);
 	MeshSelector->CategoryAttribute = AttributeName;
 	MeshSelector->Entries = Entries;
 
diff --git a/Engine/Plugins/Experimental/PCG/Source/PCG/Public/Elements/PCGStaticMeshSpawner.h b/Engine/Plugins/Experimental/PCG/Source/PCG/Public/Elements/PCGStaticMeshSpawner.h
--- a/Engine/Plugins/Experimental/PCG/Source/PCG/Public/Elements/PCGStaticMeshSpawner.h
+++ b/Engine/Plugins/Experimental/PCG/Source/PCG/Public/Elements/PCGStaticMeshSpawner.h
@@ -9,7 +9,7 @@
 #include "PCGPin.h"
 #include "PCGStaticMeshSpawner.generated.h"
 
-class UPCGInstancePackerBase;
+class UPCGInstanceDataPackerBase;
 class UPCGMeshSelectorBase;
 class UPCGSpatialData;
 struct FPCGContext;
@@ -69,7 +69,7 @@ public:
 	void SetMeshSelectorType(TSubclassOf<UPCGMeshSelectorBase> InMeshSelectorType);
 
 	UFUNCTION(BlueprintCallable, Category = Settings)
-	void SetInstancePackerType(TSubclassOf<UPCGInstancePackerBase> InInstancePackerType);
+	void SetInstancePackerType(TSubclassOf<UPCGInstanceDataPackerBase> InInstancePackerType);
 
 public:
 	/** Defines the method of mesh selection per input data */
@@ -77,20 +77,20 @@ public:
 	TSubclassOf<UPCGMeshSelectorBase> MeshSelectorType;
 
 	UPROPERTY(BlueprintReadOnly, VisibleAnywhere, Instanced, Category = Settings)
-	TObjectPtr<UPCGMeshSelectorBase> MeshSelectorInstance;
+	TObjectPtr<UPCGMeshSelectorBase> MeshSelectorParameters;
 
 	/** Defines the method of custom data packing for spawned (H)ISMCs */
 	UPROPERTY(BlueprintReadOnly, EditAnywhere, Category = Settings)
-	TSubclassOf<UPCGInstancePackerBase> InstancePackerType;
+	TSubclassOf<UPCGInstanceDataPackerBase> InstanceDataPackerType;
 
 	UPROPERTY(BlueprintReadOnly, VisibleAnywhere, Instanced, Category = Settings)
-	TObjectPtr<UPCGInstancePackerBase> InstancePackerInstance = nullptr;
+	TObjectPtr<UPCGInstanceDataPackerBase> InstanceDataPackerParameters = nullptr;
 
 	/** Attribute name to store mesh SoftObjectPaths inside if the output pin is connected. Note: Will overwrite existing data if the attribute name already exists. */
 	UPROPERTY(BlueprintReadWrite, EditAnywhere, Category = Settings)
 	FName OutAttributeName = NAME_None;
 
-	UPROPERTY(meta=(DeprecatedProperty, DeprecationMessage="Use MeshSelectorType and MeshSelectorInstance instead."))
+	UPROPERTY(meta=(DeprecatedProperty, DeprecationMessage="Use MeshSelectorType and MeshSelectorParameters instead."))
 	TArray<FPCGStaticMeshSpawnerEntry> Meshes_DEPRECATED;
 
 protected:
@@ -114,6 +114,6 @@ protected:
 
 #if UE_ENABLE_INCLUDE_ORDER_DEPRECATED_IN_5_2
 #include "CoreMinimal.h"
-#include "InstancePackers/PCGInstancePackerBase.h"
+#include "InstanceDataPackers/PCGInstanceDataPackerBase.h"
 #include "MeshSelectors/PCGMeshSelectorBase.h"
 #endif
diff --git a/Engine/Plugins/Experimental/PCG/Source/PCG/Public/Elements/PCGStaticMeshSpawnerContext.h b/Engine/Plugins/Experimental/PCG/Source/PCG/Public/Elements/PCGStaticMeshSpawnerContext.h
--- a/Engine/Plugins/Experimental/PCG/Source/PCG/Public/Elements/PCGStaticMeshSpawnerContext.h
+++ b/Engine/Plugins/Experimental/PCG/Source/PCG/Public/Elements/PCGStaticMeshSpawnerContext.h
@@ -4,7 +4,7 @@
 
 #include "PCGContext.h"
 #include "MeshSelectors/PCGMeshSelectorBase.h"
-#include "InstancePackers/PCGInstancePackerBase.h"
+#include "InstanceDataPackers/PCGInstanceDataPackerBase.h"
 
 #include "PCGStaticMeshSpawnerContext.generated.h"
 
diff --git a/Engine/Plugins/Experimental/PCG/Source/PCG/Public/InstanceDataPackers/PCGInstanceDataPackerBase.h b/Engine/Plugins/Experimental/PCG/Source/PCG/Public/InstanceDataPackers/PCGInstanceDataPackerBase.h
new file mode 100644
--- /dev/null
+++ b/Engine/Plugins/Experimental/PCG/Source/PCG/Public/InstanceDataPackers/PCGInstanceDataPackerBase.h
@@ -0,0 +1,54 @@
+// Copyright Epic Games, Inc. All Rights Reserved.
+
+#pragma once
+
+#include "PCGInstanceDataPackerBase.generated.h"
+
+class UPCGMetadata;
+struct FPCGContext;
+struct FPCGMeshInstanceList;
+
+class UPCGSpatialData;
+class FPCGMetadataAttributeBase;
+
+USTRUCT(BlueprintType)
+struct FPCGPackedCustomData
+{
+	GENERATED_BODY()
+
+	UPROPERTY(BlueprintReadWrite, EditAnywhere, Category = Settings)
+	int NumCustomDataFloats = 0;
+
+	UPROPERTY(BlueprintReadWrite, EditAnywhere, Category = Settings)
+    TArray<float> CustomData;
+};
+
+UCLASS(Abstract, BlueprintType, Blueprintable, ClassGroup = (Procedural))
+class PCG_API UPCGInstanceDataPackerBase : public UObject 
+{
+	GENERATED_BODY()
+
+public:
+	/** Defines the strategy for (H)ISM custom float data packing */
+	UFUNCTION(BlueprintNativeEvent, Category = InstancePacking)
+	void PackInstances(FPCGContext& Context, const UPCGSpatialData* InSpatialData, const FPCGMeshInstanceList& InstanceList, FPCGPackedCustomData& OutPackedCustomData) const;
+
+	virtual void PackInstances_Implementation(FPCGContext& Context, const UPCGSpatialData* InSpatialData, const FPCGMeshInstanceList& InstanceList, FPCGPackedCustomData& OutPackedCustomData) const PURE_VIRTUAL(UPCGInstanceDataPackerBase::PackInstances_Implementation);
+
+	/** Interprets Metadata TypeId and increments OutPackedCustomData.NumCustomDataFloats appropriately. Returns false if the type could not be interpreted. */
+	UFUNCTION(BlueprintCallable, Category = InstancePacking)
+	bool AddTypeToPacking(int TypeId, FPCGPackedCustomData& OutPackedCustomData) const;
+
+	/** Build a PackedCustomData by processing each attribute in order for each point in the InstanceList */
+	UFUNCTION(BlueprintCallable, Category = InstancePacking) 
+	void PackCustomDataFromAttributes(const FPCGMeshInstanceList& InstanceList, const UPCGMetadata* Metadata, const TArray<FName>& AttributeNames, FPCGPackedCustomData& OutPackedCustomData) const;
+
+	/** Build a PackedCustomData by processing each attribute in order for each point in the InstanceList */
+	void PackCustomDataFromAttributes(const FPCGMeshInstanceList& InstanceList, const TArray<const FPCGMetadataAttributeBase*>& Attributes, FPCGPackedCustomData& OutPackedCustomData) const;
+};
+
+#if UE_ENABLE_INCLUDE_ORDER_DEPRECATED_IN_5_2
+#include "CoreMinimal.h"
+#include "MeshSelectors/PCGMeshSelectorBase.h"
+#include "PCGElement.h"
+#endif
diff --git a/Engine/Plugins/Experimental/PCG/Source/PCG/Public/InstanceDataPackers/PCGInstanceDataPackerByAttribute.h b/Engine/Plugins/Experimental/PCG/Source/PCG/Public/InstanceDataPackers/PCGInstanceDataPackerByAttribute.h
new file mode 100644
--- /dev/null
+++ b/Engine/Plugins/Experimental/PCG/Source/PCG/Public/InstanceDataPackers/PCGInstanceDataPackerByAttribute.h
@@ -0,0 +1,20 @@
+// Copyright Epic Games, Inc. All Rights Reserved.
+
+#pragma once
+
+#include "PCGInstanceDataPackerBase.h"
+
+#include "PCGInstanceDataPackerByAttribute.generated.h"
+
+UCLASS(BlueprintType, ClassGroup = (Procedural))
+class PCG_API UPCGInstanceDataPackerByAttribute : public UPCGInstanceDataPackerBase 
+{
+	GENERATED_BODY()
+
+public:
+	virtual void PackInstances_Implementation(UPARAM(ref) FPCGContext& Context, const UPCGSpatialData* InSpatialData, UPARAM(ref) const FPCGMeshInstanceList& InstanceList, FPCGPackedCustomData& OutPackedCustomData) const override;
+
+public:
+	UPROPERTY(BlueprintReadWrite, EditAnywhere, Category = Settings)
+	TArray<FName> AttributeNames;
+};
diff --git a/Engine/Plugins/Experimental/PCG/Source/PCG/Public/InstanceDataPackers/PCGInstanceDataPackerByRegex.h b/Engine/Plugins/Experimental/PCG/Source/PCG/Public/InstanceDataPackers/PCGInstanceDataPackerByRegex.h
new file mode 100644
--- /dev/null
+++ b/Engine/Plugins/Experimental/PCG/Source/PCG/Public/InstanceDataPackers/PCGInstanceDataPackerByRegex.h
@@ -0,0 +1,20 @@
+// Copyright Epic Games, Inc. All Rights Reserved.
+
+#pragma once
+
+#include "PCGInstanceDataPackerBase.h"
+
+#include "PCGInstanceDataPackerByRegex.generated.h"
+
+UCLASS(BlueprintType, ClassGroup = (Procedural))
+class PCG_API UPCGInstanceDataPackerByRegex : public UPCGInstanceDataPackerBase 
+{
+	GENERATED_BODY()
+
+public:
+	virtual void PackInstances_Implementation(UPARAM(ref) FPCGContext& Context, const UPCGSpatialData* InSpatialData, UPARAM(ref) const FPCGMeshInstanceList& InstanceList, FPCGPackedCustomData& OutPackedCustomData) const override;
+
+public:
+	UPROPERTY(BlueprintReadWrite, EditAnywhere, Category = Settings)
+	TArray<FString> RegexPatterns; 
+};
diff --git a/Engine/Plugins/Experimental/PCG/Source/PCG/Public/InstancePackers/PCGInstancePackerBase.h b/Engine/Plugins/Experimental/PCG/Source/PCG/Public/InstancePackers/PCGInstancePackerBase.h
deleted file mode 100644
diff --git a/Engine/Plugins/Experimental/PCG/Source/PCG/Public/InstancePackers/PCGInstancePackerByAttribute.h b/Engine/Plugins/Experimental/PCG/Source/PCG/Public/InstancePackers/PCGInstancePackerByAttribute.h
deleted file mode 100644
diff --git a/Engine/Plugins/Experimental/PCG/Source/PCG/Public/InstancePackers/PCGInstancePackerByRegex.h b/Engine/Plugins/Experimental/PCG/Source/PCG/Public/InstancePackers/PCGInstancePackerByRegex.h
deleted file mode 100644
