[PCG] Add Graph Parameter getters/setters for BP

* BP functions are name based
* Relied a some changes in struct utils that are commented out. Workaround in the meantime

diff --git a/Engine/Plugins/Experimental/PCG/Source/PCG/Private/Helpers/PCGGraphParametersHelpers.cpp b/Engine/Plugins/Experimental/PCG/Source/PCG/Private/Helpers/PCGGraphParametersHelpers.cpp
new file mode 100644
--- /dev/null
+++ b/Engine/Plugins/Experimental/PCG/Source/PCG/Private/Helpers/PCGGraphParametersHelpers.cpp
@@ -0,0 +1,362 @@
+// Copyright Epic Games, Inc. All Rights Reserved.
+
+#include "Helpers/PCGGraphParametersHelpers.h"
+
+#include "PCGGraph.h"
+#include "PropertyBag.h"
+
+#include "Templates/ValueOrError.h"
+#include "UObject/Script.h"
+#include "UObject/Stack.h"
+
+#define LOCTEXT_NAMESPACE "PCGGraphParametersHelpers"
+
+namespace PCGGraphParametersHelpers
+{
+	static const FText InvalidGraphInstanceError = LOCTEXT("InvalidGraphInstance", "Invalid graph instance.");
+
+	void OnException(EPropertyBagResult InResult, FName InPropertyName)
+	{
+		FText ErrorMessage;
+		switch (InResult)
+		{
+		case EPropertyBagResult::TypeMismatch:
+			ErrorMessage = FText::Format(LOCTEXT("TypeMismatch", "Parameter {0} is not of the right type."), FText::FromName(InPropertyName));
+			break;
+		case EPropertyBagResult::PropertyNotFound:
+			ErrorMessage = FText::Format(LOCTEXT("PropertyNotFound", "Parameter {0} does not exist."), FText::FromName(InPropertyName));
+			break;
+		default:
+			break;
+		}
+
+		const FBlueprintExceptionInfo ExceptionInfo(EBlueprintExceptionType::FatalError, ErrorMessage);
+		FBlueprintCoreDelegates::ThrowScriptException(FFrame::GetThreadLocalTopStackFrame()->Object, *FFrame::GetThreadLocalTopStackFrame(), ExceptionInfo);
+
+		if (FPlatformMisc::IsDebuggerPresent())
+		{
+			ensureMsgf(false, TEXT("%s"), *ErrorMessage.ToString());
+		}
+	}
+
+	void UpdateOverrides(UPCGGraphInstance* GraphInstance, FName InPropertyName)
+	{
+		if (!GraphInstance)
+		{
+			return;
+		}
+
+		if (const FPropertyBagPropertyDesc* Desc = GraphInstance->ParametersOverrides.Parameters.FindPropertyDescByName(InPropertyName))
+		{
+			GraphInstance->ParametersOverrides.PropertiesIDsOverridden.Add(Desc->ID);
+		}
+	}
+
+	// Get a parameter from the Instanced Property Bag using its getters. It usually returns T, but in some cases (like soft objects or struct), it returns a ptr or something
+	// else. That's what U is for, and there are strict conditions on the relation between T and U.
+	template <typename T, typename U = T>
+	T GetParameter(UPCGGraphInstance* GraphInstance, FName PropertyName, TFunctionRef<TValueOrError<U, EPropertyBagResult>(FInstancedPropertyBag&)> Getter)
+	{
+		if (!GraphInstance)
+		{
+			const FBlueprintExceptionInfo ExceptionInfo(EBlueprintExceptionType::FatalError, InvalidGraphInstanceError);
+			FBlueprintCoreDelegates::ThrowScriptException(FFrame::GetThreadLocalTopStackFrame()->Object, *FFrame::GetThreadLocalTopStackFrame(), ExceptionInfo);
+
+			return T{};
+		}
+
+		TValueOrError<U, EPropertyBagResult> Result = Getter(GraphInstance->ParametersOverrides.Parameters);
+
+		if (Result.HasError())
+		{
+			PCGGraphParametersHelpers::OnException(Result.GetError(), PropertyName);
+			return T{};
+		}
+		else
+		{
+			if constexpr (std::is_pointer_v<U>)
+			{
+				static_assert(std::is_same_v<T, std::remove_pointer_t<U>>);
+				return Result.GetValue() ? *Result.GetValue() : T{};
+			}
+			else if constexpr (!std::is_same_v<T, U>)
+			{
+				return T(Result.GetValue());
+			}
+			else
+			{
+				return Result.GetValue();
+			}
+		}
+	}
+
+	// Convenient alias when T and U are the same.
+	template <typename T>
+	T GetParameter(UPCGGraphInstance* GraphInstance, FName PropertyName, TFunctionRef<TValueOrError<T, EPropertyBagResult>(FInstancedPropertyBag&)> Getter)
+	{
+		return GetParameter<T, T>(GraphInstance, PropertyName, Getter);
+	}
+
+	// Get a parameter from the Instanced Property Bag using its setters. Value to set must be captured in the Setter.
+	// If the set succeeded, it will force this parameter to be overridden.
+	void SetParameter(UPCGGraphInstance* GraphInstance, FName PropertyName, TFunctionRef<EPropertyBagResult(FInstancedPropertyBag&)> Setter)
+	{
+		if (!GraphInstance)
+		{
+			const FBlueprintExceptionInfo ExceptionInfo(EBlueprintExceptionType::FatalError, InvalidGraphInstanceError);
+			FBlueprintCoreDelegates::ThrowScriptException(FFrame::GetThreadLocalTopStackFrame()->Object, *FFrame::GetThreadLocalTopStackFrame(), ExceptionInfo);
+
+			return;
+		}
+
+		EPropertyBagResult Result = Setter(GraphInstance->ParametersOverrides.Parameters);
+
+		if (Result != EPropertyBagResult::Success)
+		{
+			PCGGraphParametersHelpers::OnException(Result, PropertyName);
+		}
+		else
+		{
+			PCGGraphParametersHelpers::UpdateOverrides(GraphInstance, PropertyName);
+		}
+	}
+}
+
+bool UPCGGraphParametersHelpers::IsOverridden(UPCGGraphInstance* GraphInstance, FName InPropertyName)
+{
+	if (!GraphInstance)
+	{
+		return false;
+	}
+
+	const FPropertyBagPropertyDesc* Desc = GraphInstance->ParametersOverrides.Parameters.FindPropertyDescByName(InPropertyName);
+	return Desc ? GraphInstance->ParametersOverrides.PropertiesIDsOverridden.Contains(Desc->ID) : false;
+}
+
+////////////
+// Getters
+////////////
+
+float UPCGGraphParametersHelpers::GetFloatParameter(UPCGGraphInstance* GraphInstance, const FName Name)
+{
+	return PCGGraphParametersHelpers::GetParameter<float>(GraphInstance, Name, [Name](FInstancedPropertyBag& Bag) { return Bag.GetValueFloat(Name); });
+}
+
+double UPCGGraphParametersHelpers::GetDoubleParameter(UPCGGraphInstance* GraphInstance, const FName Name)
+{
+	return PCGGraphParametersHelpers::GetParameter<double>(GraphInstance, Name, [Name](FInstancedPropertyBag& Bag) { return Bag.GetValueDouble(Name); });
+}
+
+bool UPCGGraphParametersHelpers::GetBoolParameter(UPCGGraphInstance* GraphInstance, const FName Name)
+{
+	return PCGGraphParametersHelpers::GetParameter<bool>(GraphInstance, Name, [Name](FInstancedPropertyBag& Bag) { return Bag.GetValueBool(Name); });
+}
+
+uint8 UPCGGraphParametersHelpers::GetByteParameter(UPCGGraphInstance* GraphInstance, const FName Name)
+{
+	return PCGGraphParametersHelpers::GetParameter<uint8>(GraphInstance, Name, [Name](FInstancedPropertyBag& Bag) { return Bag.GetValueByte(Name); });
+}
+
+int32 UPCGGraphParametersHelpers::GetInt32Parameter(UPCGGraphInstance* GraphInstance, const FName Name)
+{
+	return PCGGraphParametersHelpers::GetParameter<int32>(GraphInstance, Name, [Name](FInstancedPropertyBag& Bag) { return Bag.GetValueInt32(Name); });
+}
+
+int64 UPCGGraphParametersHelpers::GetInt64Parameter(UPCGGraphInstance* GraphInstance, const FName Name)
+{
+	return PCGGraphParametersHelpers::GetParameter<int64>(GraphInstance, Name, [Name](FInstancedPropertyBag& Bag) { return Bag.GetValueInt64(Name); });
+}
+
+FName UPCGGraphParametersHelpers::GetNameParameter(UPCGGraphInstance* GraphInstance, const FName Name)
+{
+	return PCGGraphParametersHelpers::GetParameter<FName>(GraphInstance, Name, [Name](FInstancedPropertyBag& Bag) { return Bag.GetValueName(Name); });
+}
+
+FString UPCGGraphParametersHelpers::GetStringParameter(UPCGGraphInstance* GraphInstance, const FName Name)
+{
+	return PCGGraphParametersHelpers::GetParameter<FString>(GraphInstance, Name, [Name](FInstancedPropertyBag& Bag) { return Bag.GetValueString(Name); });
+}
+
+TSoftObjectPtr<UObject> UPCGGraphParametersHelpers::GetSoftObjectParameter(UPCGGraphInstance* GraphInstance, const FName Name)
+{
+	// TODO: Uncomment when StructUtils is updated
+	// return PCGGraphParametersHelpers::GetParameter<TSoftObjectPtr<UObject>, FSoftObjectPath>(GraphInstance, Name, [Name](FInstancedPropertyBag& Bag) { return Bag.GetValueSoftPath(Name); });
+
+	// Workaround
+	return PCGGraphParametersHelpers::GetParameter<TSoftObjectPtr<UObject>>(GraphInstance, Name, [Name](FInstancedPropertyBag& Bag) -> TValueOrError<TSoftObjectPtr<UObject>, EPropertyBagResult>
+	{
+		const FPropertyBagPropertyDesc* Desc = Bag.FindPropertyDescByName(Name);
+		if (!Desc)
+		{
+			return MakeError(EPropertyBagResult::PropertyNotFound);
+		}
+
+		const FSoftObjectProperty* Property = CastField<FSoftObjectProperty>(Desc->CachedProperty);
+		if (!Property)
+		{
+			return MakeError(EPropertyBagResult::TypeMismatch);
+		}
+
+		const FSoftObjectPtr& SoftObjectPtr = Property->GetPropertyValue_InContainer(Bag.GetValue().GetMemory());
+		return MakeValue(TSoftObjectPtr<UObject>(SoftObjectPtr.ToSoftObjectPath()));
+	});
+}
+
+TSoftClassPtr<UObject> UPCGGraphParametersHelpers::GetSoftClassParameter(UPCGGraphInstance* GraphInstance, const FName Name)
+{
+	// TODO: Uncomment when StructUtils is updated
+	// return PCGGraphParametersHelpers::GetParameter<TSoftClassPtr<UObject>, FSoftObjectPath>(GraphInstance, Name, [Name](FInstancedPropertyBag& Bag) { return Bag.GetValueSoftPath(Name); });
+
+	// Workaround
+	return PCGGraphParametersHelpers::GetParameter<TSoftClassPtr<UObject>>(GraphInstance, Name, [Name](FInstancedPropertyBag& Bag) -> TValueOrError<TSoftClassPtr<UObject>, EPropertyBagResult>
+	{
+		const FPropertyBagPropertyDesc* Desc = Bag.FindPropertyDescByName(Name);
+		if (!Desc)
+		{
+			return MakeError(EPropertyBagResult::PropertyNotFound);
+		}
+
+		const FSoftClassProperty* Property = CastField<FSoftClassProperty>(Desc->CachedProperty);
+		if (!Property)
+		{
+			return MakeError(EPropertyBagResult::TypeMismatch);
+		}
+
+		const FSoftObjectPtr& SoftObjectPtr = Property->GetPropertyValue_InContainer(Bag.GetValue().GetMemory());
+		return MakeValue(TSoftClassPtr<UObject>(SoftObjectPtr.ToSoftObjectPath()));
+	});
+}
+
+FVector UPCGGraphParametersHelpers::GetVectorParameter(UPCGGraphInstance* GraphInstance, const FName Name)
+{
+	return PCGGraphParametersHelpers::GetParameter<FVector, FVector*>(GraphInstance, Name, [Name](FInstancedPropertyBag& Bag) { return Bag.GetValueStruct<FVector>(Name); });
+}
+
+FRotator UPCGGraphParametersHelpers::GetRotatorParameter(UPCGGraphInstance* GraphInstance, const FName Name)
+{
+	return PCGGraphParametersHelpers::GetParameter<FRotator, FRotator*>(GraphInstance, Name, [Name](FInstancedPropertyBag& Bag) { return Bag.GetValueStruct<FRotator>(Name); });
+}
+
+FTransform UPCGGraphParametersHelpers::GetTransformParameter(UPCGGraphInstance* GraphInstance, const FName Name)
+{
+	return PCGGraphParametersHelpers::GetParameter<FTransform, FTransform*>(GraphInstance, Name, [Name](FInstancedPropertyBag& Bag) { return Bag.GetValueStruct<FTransform>(Name); });
+}
+
+
+////////////
+// Setters
+////////////
+
+void UPCGGraphParametersHelpers::SetFloatParameter(UPCGGraphInstance* GraphInstance, const FName Name, const float Value)
+{
+	PCGGraphParametersHelpers::SetParameter(GraphInstance, Name, [Name, Value](FInstancedPropertyBag& Bag) { return Bag.SetValueFloat(Name, Value); });
+}
+
+void UPCGGraphParametersHelpers::SetDoubleParameter(UPCGGraphInstance* GraphInstance, const FName Name, const double Value)
+{
+	PCGGraphParametersHelpers::SetParameter(GraphInstance, Name, [Name, Value](FInstancedPropertyBag& Bag) { return Bag.SetValueDouble(Name, Value); });
+}
+
+void UPCGGraphParametersHelpers::SetBoolParameter(UPCGGraphInstance* GraphInstance, const FName Name, const bool bValue)
+{
+	PCGGraphParametersHelpers::SetParameter(GraphInstance, Name, [Name, bValue](FInstancedPropertyBag& Bag) { return Bag.SetValueBool(Name, bValue); });
+}
+
+void UPCGGraphParametersHelpers::SetByteParameter(UPCGGraphInstance* GraphInstance, const FName Name, const uint8 Value)
+{
+	PCGGraphParametersHelpers::SetParameter(GraphInstance, Name, [Name, Value](FInstancedPropertyBag& Bag) { return Bag.SetValueByte(Name, Value); });
+}
+
+void UPCGGraphParametersHelpers::SetInt32Parameter(UPCGGraphInstance* GraphInstance, const FName Name, const int32 Value)
+{
+	PCGGraphParametersHelpers::SetParameter(GraphInstance, Name, [Name, Value](FInstancedPropertyBag& Bag) { return Bag.SetValueInt32(Name, Value); });
+}
+
+void UPCGGraphParametersHelpers::SetInt64Parameter(UPCGGraphInstance* GraphInstance, const FName Name, const int64 Value)
+{
+	PCGGraphParametersHelpers::SetParameter(GraphInstance, Name, [Name, Value](FInstancedPropertyBag& Bag) { return Bag.SetValueInt64(Name, Value); });
+}
+
+void UPCGGraphParametersHelpers::SetNameParameter(UPCGGraphInstance* GraphInstance, const FName Name, const FName Value)
+{
+	PCGGraphParametersHelpers::SetParameter(GraphInstance, Name, [Name, Value](FInstancedPropertyBag& Bag) { return Bag.SetValueName(Name, Value); });
+}
+
+void UPCGGraphParametersHelpers::SetStringParameter(UPCGGraphInstance* GraphInstance, const FName Name, const FString& Value)
+{
+	PCGGraphParametersHelpers::SetParameter(GraphInstance, Name, [Name, Value](FInstancedPropertyBag& Bag) { return Bag.SetValueString(Name, Value); });
+}
+
+void UPCGGraphParametersHelpers::SetEnumParameter(UPCGGraphInstance* GraphInstance, const FName Name, const UEnum* Enum, const uint8 Value)
+{
+	PCGGraphParametersHelpers::SetParameter(GraphInstance, Name, [Name, Value, Enum](FInstancedPropertyBag& Bag) { return Bag.SetValueEnum(Name, Value, Enum); });
+}
+
+void UPCGGraphParametersHelpers::SetSoftObjectParameter(UPCGGraphInstance* GraphInstance, const FName Name, const TSoftObjectPtr<UObject>& Value)
+{
+	// TODO: Uncomment when StructUtils is updated
+	// PCGGraphParametersHelpers::SetParameter(GraphInstance, Name, [Name, Value](FInstancedPropertyBag& Bag) { return Bag.SetValueSoftPath(Name, Value.ToSoftObjectPath()); });
+
+	// Workaround
+	PCGGraphParametersHelpers::SetParameter(GraphInstance, Name, [Name, Value](FInstancedPropertyBag& Bag)
+	{
+		const FPropertyBagPropertyDesc* Desc = Bag.FindPropertyDescByName(Name);
+		if (!Desc)
+		{
+			return EPropertyBagResult::PropertyNotFound;
+		}
+
+		const FSoftObjectProperty* Property = CastField<FSoftObjectProperty>(Desc->CachedProperty);
+		if (!Property)
+		{
+			return EPropertyBagResult::TypeMismatch;
+		}
+
+		const FSoftObjectPtr SoftObjectPtr(Value.ToSoftObjectPath());
+		Property->SetPropertyValue_InContainer(Bag.GetMutableValue().GetMemory(), SoftObjectPtr);
+		return EPropertyBagResult::Success;
+	});
+}
+
+void UPCGGraphParametersHelpers::SetSoftClassParameter(UPCGGraphInstance* GraphInstance, const FName Name, const TSoftClassPtr<UObject>& Value)
+{
+	// TODO: Uncomment when StructUtils is updated
+	// PCGGraphParametersHelpers::SetParameter(GraphInstance, Name, [Name, Value](FInstancedPropertyBag& Bag) { return Bag.SetValueSoftPath(Name, Value.ToSoftObjectPath()); });
+
+	// Workaround
+	PCGGraphParametersHelpers::SetParameter(GraphInstance, Name, [Name, Value](FInstancedPropertyBag& Bag)
+	{
+		const FPropertyBagPropertyDesc* Desc = Bag.FindPropertyDescByName(Name);
+		if (!Desc)
+		{
+			return EPropertyBagResult::PropertyNotFound;
+		}
+
+		const FSoftClassProperty* Property = CastField<FSoftClassProperty>(Desc->CachedProperty);
+		if (!Property)
+		{
+			return EPropertyBagResult::TypeMismatch;
+		}
+
+		const FSoftObjectPtr SoftObjectPtr(Value.ToSoftObjectPath());
+		Property->SetPropertyValue_InContainer(Bag.GetMutableValue().GetMemory(), SoftObjectPtr);
+		return EPropertyBagResult::Success;
+	});
+}
+
+void UPCGGraphParametersHelpers::SetVectorParameter(UPCGGraphInstance* GraphInstance, const FName Name, const FVector& Value)
+{
+	PCGGraphParametersHelpers::SetParameter(GraphInstance, Name, [Name, Value](FInstancedPropertyBag& Bag) { return Bag.SetValueStruct<FVector>(Name, Value); });
+}
+
+void UPCGGraphParametersHelpers::SetRotatorParameter(UPCGGraphInstance* GraphInstance, const FName Name, const FRotator& Value)
+{
+	PCGGraphParametersHelpers::SetParameter(GraphInstance, Name, [Name, Value](FInstancedPropertyBag& Bag) { return Bag.SetValueStruct<FRotator>(Name, Value); });
+}
+
+void UPCGGraphParametersHelpers::SetTransformParameter(UPCGGraphInstance* GraphInstance, const FName Name, const FTransform& Value)
+{
+	PCGGraphParametersHelpers::SetParameter(GraphInstance, Name, [Name, Value](FInstancedPropertyBag& Bag) { return Bag.SetValueStruct<FTransform>(Name, Value); });
+}
+
+#undef LOCTEXT_NAMESPACE
\ No newline at end of file
diff --git a/Engine/Plugins/Experimental/PCG/Source/PCG/Public/Helpers/PCGGraphParametersHelpers.h b/Engine/Plugins/Experimental/PCG/Source/PCG/Public/Helpers/PCGGraphParametersHelpers.h
new file mode 100644
--- /dev/null
+++ b/Engine/Plugins/Experimental/PCG/Source/PCG/Public/Helpers/PCGGraphParametersHelpers.h
@@ -0,0 +1,111 @@
+// Copyright Epic Games, Inc. All Rights Reserved.
+
+#pragma once
+
+#include "Kismet/BlueprintFunctionLibrary.h"
+
+#include "PCGGraphParametersHelpers.generated.h"
+
+class UPCGGraphInstance;
+
+/**
+* Helpers to dynamically set parameters on graph instances.
+*/
+UCLASS()
+class PCG_API UPCGGraphParametersHelpers : public UBlueprintFunctionLibrary
+{
+	GENERATED_BODY()
+
+public:
+	UFUNCTION(BlueprintCallable, Category = "PCG|Graph Parameters")
+	static bool IsOverridden(UPCGGraphInstance* GraphInstance, const FName Name);
+
+	////////////
+	// Getters
+	////////////
+
+	UFUNCTION(BlueprintCallable, Category = "PCG|Graph Parameters")
+	static float GetFloatParameter(UPCGGraphInstance* GraphInstance, const FName Name);
+
+	UFUNCTION(BlueprintCallable, Category = "PCG|Graph Parameters")
+	static double GetDoubleParameter(UPCGGraphInstance* GraphInstance, const FName Name);
+
+	UFUNCTION(BlueprintCallable, Category = "PCG|Graph Parameters")
+	static bool GetBoolParameter(UPCGGraphInstance* GraphInstance, const FName Name);
+
+	UFUNCTION(BlueprintCallable, Category = "PCG|Graph Parameters")
+	static uint8 GetByteParameter(UPCGGraphInstance* GraphInstance, const FName Name);
+
+	UFUNCTION(BlueprintCallable, Category = "PCG|Graph Parameters")
+	static int32 GetInt32Parameter(UPCGGraphInstance* GraphInstance, const FName Name);
+
+	UFUNCTION(BlueprintCallable, Category = "PCG|Graph Parameters")
+	static int64 GetInt64Parameter(UPCGGraphInstance* GraphInstance, const FName Name);
+
+	UFUNCTION(BlueprintCallable, Category = "PCG|Graph Parameters")
+	static FName GetNameParameter(UPCGGraphInstance* GraphInstance, const FName Name);
+
+	UFUNCTION(BlueprintCallable, Category = "PCG|Graph Parameters")
+	static FString GetStringParameter(UPCGGraphInstance* GraphInstance, const FName Name);
+
+	UFUNCTION(BlueprintCallable, Category = "PCG|Graph Parameters")
+	static TSoftObjectPtr<UObject> GetSoftObjectParameter(UPCGGraphInstance* GraphInstance, const FName Name);
+
+	UFUNCTION(BlueprintCallable, Category = "PCG|Graph Parameters")
+	static TSoftClassPtr<UObject> GetSoftClassParameter(UPCGGraphInstance* GraphInstance, const FName Name);
+
+	UFUNCTION(BlueprintCallable, Category = "PCG|Graph Parameters")
+	static FVector GetVectorParameter(UPCGGraphInstance* GraphInstance, const FName Name);
+
+	UFUNCTION(BlueprintCallable, Category = "PCG|Graph Parameters")
+	static FRotator GetRotatorParameter(UPCGGraphInstance* GraphInstance, const FName Name);
+
+	UFUNCTION(BlueprintCallable, Category = "PCG|Graph Parameters")
+	static FTransform GetTransformParameter(UPCGGraphInstance* GraphInstance, const FName Name);
+
+	////////////
+	// Setters
+	////////////
+
+	UFUNCTION(BlueprintCallable, Category = "PCG|Graph Parameters")
+	static void SetFloatParameter(UPCGGraphInstance* GraphInstance, const FName Name, const float Value);
+
+	UFUNCTION(BlueprintCallable, Category = "PCG|Graph Parameters")
+	static void SetDoubleParameter(UPCGGraphInstance* GraphInstance, const FName Name, const double Value);
+
+	UFUNCTION(BlueprintCallable, Category = "PCG|Graph Parameters")
+	static void SetBoolParameter(UPCGGraphInstance* GraphInstance, const FName Name, const bool bValue);
+
+	UFUNCTION(BlueprintCallable, Category = "PCG|Graph Parameters")
+	static void SetByteParameter(UPCGGraphInstance* GraphInstance, const FName Name, const uint8 Value);
+
+	UFUNCTION(BlueprintCallable, Category = "PCG|Graph Parameters")
+	static void SetInt32Parameter(UPCGGraphInstance* GraphInstance, const FName Name, const int32 Value);
+
+	UFUNCTION(BlueprintCallable, Category = "PCG|Graph Parameters")
+	static void SetInt64Parameter(UPCGGraphInstance* GraphInstance, const FName Name, const int64 Value);
+
+	UFUNCTION(BlueprintCallable, Category = "PCG|Graph Parameters")
+	static void SetNameParameter(UPCGGraphInstance* GraphInstance, const FName Name, const FName Value);
+
+	UFUNCTION(BlueprintCallable, Category = "PCG|Graph Parameters")
+	static void SetStringParameter(UPCGGraphInstance* GraphInstance, const FName Name, const FString& Value);
+
+	UFUNCTION(BlueprintCallable, Category = "PCG|Graph Parameters")
+	static void SetEnumParameter(UPCGGraphInstance* GraphInstance, const FName Name, const UEnum* Enum, const uint8 Value);
+
+	UFUNCTION(BlueprintCallable, Category = "PCG|Graph Parameters")
+	static void SetSoftObjectParameter(UPCGGraphInstance* GraphInstance, const FName Name, UPARAM(ref) const TSoftObjectPtr<UObject>& Value);
+
+	UFUNCTION(BlueprintCallable, Category = "PCG|Graph Parameters")
+	static void SetSoftClassParameter(UPCGGraphInstance* GraphInstance, const FName Name, UPARAM(ref) const TSoftClassPtr<UObject>& Value);
+
+	UFUNCTION(BlueprintCallable, Category = "PCG|Graph Parameters")
+	static void SetVectorParameter(UPCGGraphInstance* GraphInstance, const FName Name, UPARAM(ref) const FVector& Value);
+
+	UFUNCTION(BlueprintCallable, Category = "PCG|Graph Parameters")
+	static void SetRotatorParameter(UPCGGraphInstance* GraphInstance, const FName Name, UPARAM(ref) const FRotator& Value);
+
+	UFUNCTION(BlueprintCallable, Category = "PCG|Graph Parameters")
+	static void SetTransformParameter(UPCGGraphInstance* GraphInstance, const FName Name, UPARAM(ref) const FTransform& Value);
+};
