Fixed issue where some samplers would perform work (which could take a long time/run out of memory) even with nothing plugged in.
Added very basic memory checks in samplers as a first step.

diff --git a/Engine/Plugins/Experimental/PCG/Source/PCG/Private/Elements/PCGSurfaceSampler.cpp b/Engine/Plugins/Experimental/PCG/Source/PCG/Private/Elements/PCGSurfaceSampler.cpp
--- a/Engine/Plugins/Experimental/PCG/Source/PCG/Private/Elements/PCGSurfaceSampler.cpp
+++ b/Engine/Plugins/Experimental/PCG/Source/PCG/Private/Elements/PCGSurfaceSampler.cpp
@@ -2,6 +2,7 @@
 
 #include "Elements/PCGSurfaceSampler.h"
 
+#include "PCGCommon.h"
 #include "PCGComponent.h"
 #include "PCGCustomVersion.h"
 #include "PCGEdge.h"
@@ -13,6 +14,7 @@
 #include "Helpers/PCGHelpers.h"
 #include "Helpers/PCGSettingsHelpers.h"
 
+#include "HAL/UnrealMemory.h"
 #include "Math/RandomStream.h"
 
 #include UE_INLINE_GENERATED_CPP_BY_NAME(PCGSurfaceSampler)
@@ -77,7 +79,9 @@ namespace PCGSurfaceSampler
 			}
 
 			const int64 CellCount64 = CellCountX * CellCountY;
-			if (CellCount64 <= 0 || CellCount64 >= MAX_int32)
+			if (CellCount64 <= 0 || 
+				CellCount64 >= MAX_int32 ||
+				(PCGFeatureSwitches::CVarCheckSamplerMemory.GetValueOnAnyThread() && FPlatformMemory::GetStats().AvailablePhysical < sizeof(FPCGPoint) * CellCount64))
 			{
 				if (Context)
 				{
@@ -354,6 +358,30 @@ bool FPCGSurfaceSamplerElement::AddGeneratingShapesToContext(FPCGSurfaceSamplerC
 		return false;
 	}
 
+	TArray<FPCGTaggedData> SurfaceInputs = InContext->InputData.GetInputsByPin(PCGSurfaceSamplerConstants::SurfaceLabel);
+
+	// Construct a list of shapes to generate samples from. Prefer to get these directly from the first input pin.
+	for (FPCGTaggedData& TaggedData : SurfaceInputs)
+	{
+		if (const UPCGSpatialData* SpatialData = Cast<UPCGSpatialData>(TaggedData.Data))
+		{
+			// Find a concrete shape for sampling. Prefer a 2D surface if we can find one.
+			if (const UPCGSpatialData* SurfaceData = SpatialData->FindShapeFromNetwork(/*InDimension=*/2))
+			{
+				InContext->GeneratingShapes.Add(SurfaceData);
+				Outputs.Add(TaggedData);
+			}
+			else if (const UPCGSpatialData* ConcreteData = SpatialData->FindFirstConcreteShapeFromNetwork())
+			{
+				// Alternatively surface-sample any concrete data - can be used to sprinkle samples down onto shapes like volumes.
+				// Searching like this allows the user to plonk in any composite network and it will often find the shape of interest.
+				// A potential extension would be to find all (unique?) concrete shapes and use all of them rather than just the first.
+				InContext->GeneratingShapes.Add(ConcreteData);
+				Outputs.Add(TaggedData);
+			}
+		}
+	}
+
 	// Grab the Bounding Shape input if there is one.
 	TArray<FPCGTaggedData> BoundingShapeInputs = InContext->InputData.GetInputsByPin(PCGSurfaceSamplerConstants::BoundingShapeLabel);
 
@@ -367,9 +395,9 @@ bool FPCGSurfaceSamplerElement::AddGeneratingShapesToContext(FPCGSurfaceSamplerC
 			const_cast<UPCGSpatialData*>(InContext->BoundingShapeSpatialInput)->AddToRoot();
 		}
 
-		if (!InContext->BoundingShapeSpatialInput && InContext->SourceComponent.IsValid())
+		// Fallback to getting bounds from actor but only if we actually have some inputs
+		if (!InContext->BoundingShapeSpatialInput && InContext->SourceComponent.IsValid() && !InContext->GeneratingShapes.IsEmpty())
 		{
-			// Fallback to getting bounds from actor
 			InContext->BoundingShapeSpatialInput = Cast<UPCGSpatialData>(InContext->SourceComponent->GetActorPCGData());
 		}
 	}
@@ -383,30 +411,6 @@ bool FPCGSurfaceSamplerElement::AddGeneratingShapesToContext(FPCGSurfaceSamplerC
 		InContext->BoundingShapeBounds = InContext->BoundingShapeSpatialInput->GetBounds();
 	}
 
-	TArray<FPCGTaggedData> SurfaceInputs = InContext->InputData.GetInputsByPin(PCGSurfaceSamplerConstants::SurfaceLabel);
-
-	// Construct a list of shapes to generate samples from. Prefer to get these directly from the first input pin.
-	for (FPCGTaggedData& TaggedData : SurfaceInputs)
-	{
-		if (const UPCGSpatialData* SpatialData = Cast<UPCGSpatialData>(TaggedData.Data))
-		{
-			// Find a concrete shape for sampling. Prefer a 2D surface if we can find one.
-			if (const UPCGSpatialData* SurfaceData = SpatialData->FindShapeFromNetwork(/*InDimension=*/2))
-			{
-				InContext->GeneratingShapes.Add(SurfaceData);
-				Outputs.Add(TaggedData);
-			}
-			else if (const UPCGSpatialData* ConcreteData = SpatialData->FindFirstConcreteShapeFromNetwork())
-			{
-				// Alternatively surface-sample any concrete data - can be used to sprinkle samples down onto shapes like volumes.
-				// Searching like this allows the user to plonk in any composite network and it will often find the shape of interest.
-				// A potential extension would be to find all (unique?) concrete shapes and use all of them rather than just the first.
-				InContext->GeneratingShapes.Add(ConcreteData);
-				Outputs.Add(TaggedData);
-			}
-		}
-	}
-
 	// If no shapes were obtained from the first input pin, try to find a shape to sample from nodes connected to the second pin.
 	if (InContext->GeneratingShapes.Num() == 0 && InContext->BoundingShapeSpatialInput)
 	{
diff --git a/Engine/Plugins/Experimental/PCG/Source/PCG/Private/Elements/PCGVolumeSampler.cpp b/Engine/Plugins/Experimental/PCG/Source/PCG/Private/Elements/PCGVolumeSampler.cpp
--- a/Engine/Plugins/Experimental/PCG/Source/PCG/Private/Elements/PCGVolumeSampler.cpp
+++ b/Engine/Plugins/Experimental/PCG/Source/PCG/Private/Elements/PCGVolumeSampler.cpp
@@ -2,6 +2,7 @@
 
 #include "Elements/PCGVolumeSampler.h"
 
+#include "PCGCommon.h"
 #include "PCGComponent.h"
 #include "PCGContext.h"
 #include "PCGCustomVersion.h"
@@ -10,6 +11,8 @@
 #include "Helpers/PCGAsync.h"
 #include "Helpers/PCGHelpers.h"
 
+#include "HAL/UnrealMemory.h"
+
 #include UE_INLINE_GENERATED_CPP_BY_NAME(PCGVolumeSampler)
 
 #define LOCTEXT_NAMESPACE "PCGVolumeSamplerElement"
@@ -78,7 +81,11 @@ namespace PCGVolumeSampler
 				return;
 			}
 
-			if (NumIterationsXY64 > 0 && NumIterationsXY64 < MAX_int32 && NumIterations64 > 0 && NumIterations64 < MAX_int32)
+			if (NumIterationsXY64 > 0 && 
+				NumIterationsXY64 < MAX_int32 && 
+				NumIterations64 > 0 && 
+				NumIterations64 < MAX_int32 &&
+				(!PCGFeatureSwitches::CVarCheckSamplerMemory.GetValueOnAnyThread() || FPlatformMemory::GetStats().AvailablePhysical >= sizeof(FPCGPoint) * NumIterations64))
 			{
 				NumIterations = static_cast<int32>(NumIterations64);
 			}
@@ -182,6 +189,7 @@ bool FPCGVolumeSamplerElement::ExecuteInternal(FPCGContext* Context) const
 	// Grab the Bounding Shape input if there is one.
 	TArray<FPCGTaggedData> BoundingShapeInputs = Context->InputData.GetInputsByPin(PCGVolumeSamplerConstants::BoundingShapeLabel);
 	const UPCGSpatialData* BoundingShape = nullptr;
+	bool bUsedDefaultBoundingShape = false;
 
 	if (!Settings->bUnbounded)
 	{
@@ -191,6 +199,7 @@ bool FPCGVolumeSamplerElement::ExecuteInternal(FPCGContext* Context) const
 		if (!BoundingShape && Context->SourceComponent.IsValid())
 		{
 			BoundingShape = Cast<UPCGSpatialData>(Context->SourceComponent->GetActorPCGData());
+			bUsedDefaultBoundingShape = true;
 		}
 	}
 	else if (BoundingShapeInputs.Num() > 0)
@@ -217,7 +226,7 @@ bool FPCGVolumeSamplerElement::ExecuteInternal(FPCGContext* Context) const
 	}
 
 	// If no shapes were obtained from the first input pin, try to find a shape to sample from nodes connected to the second pin.
-	if (GeneratingShapes.Num() == 0 && BoundingShape)
+	if (GeneratingShapes.Num() == 0 && BoundingShape && !bUsedDefaultBoundingShape)
 	{
 		GeneratingShapes.Add(BoundingShape);
 
diff --git a/Engine/Plugins/Experimental/PCG/Source/PCG/Private/PCGCommon.cpp b/Engine/Plugins/Experimental/PCG/Source/PCG/Private/PCGCommon.cpp
new file mode 100644
--- /dev/null
+++ b/Engine/Plugins/Experimental/PCG/Source/PCG/Private/PCGCommon.cpp
@@ -0,0 +1,14 @@
+// Copyright Epic Games, Inc. All Rights Reserved.
+
+#include "PCGCommon.h"
+
+namespace PCGFeatureSwitches
+{
+
+	TAutoConsoleVariable<bool> CVarCheckSamplerMemory{
+		TEXT("pcg.CheckSamplerMemory"),
+		true,
+		TEXT("Checks expected memory size consumption prior to performing sampling operations")
+	};
+
+}
\ No newline at end of file
diff --git a/Engine/Plugins/Experimental/PCG/Source/PCG/Public/PCGCommon.h b/Engine/Plugins/Experimental/PCG/Source/PCG/Public/PCGCommon.h
--- a/Engine/Plugins/Experimental/PCG/Source/PCG/Public/PCGCommon.h
+++ b/Engine/Plugins/Experimental/PCG/Source/PCG/Public/PCGCommon.h
@@ -2,6 +2,8 @@
 
 #pragma once
 
+#include "HAL/IConsoleManager.h"
+
 #include "PCGCommon.generated.h"
 
 using FPCGTaskId = uint64;
@@ -112,3 +114,8 @@ namespace PCGObjectMetadata
 		PCG_NotOverridable
 	};
 }
+
+namespace PCGFeatureSwitches
+{
+	extern PCG_API TAutoConsoleVariable<bool> CVarCheckSamplerMemory;
+}
\ No newline at end of file
