[PCG] Modified SplineSampler node to allow users to toggle between Unbounded:

* [DEPREC] New samplers will be created to be bounded by default

diff --git a/Engine/Plugins/Experimental/PCG/Source/PCG/Private/Elements/PCGSplineSampler.cpp b/Engine/Plugins/Experimental/PCG/Source/PCG/Private/Elements/PCGSplineSampler.cpp
--- a/Engine/Plugins/Experimental/PCG/Source/PCG/Private/Elements/PCGSplineSampler.cpp
+++ b/Engine/Plugins/Experimental/PCG/Source/PCG/Private/Elements/PCGSplineSampler.cpp
@@ -2,6 +2,7 @@
 
 #include "Elements/PCGSplineSampler.h"
 
+#include "PCGComponent.h"
 #include "PCGContext.h"
 #include "PCGCustomVersion.h"
 #include "Data/PCGIntersectionData.h"
@@ -545,8 +546,9 @@ namespace PCGSplineSampler
 				OutPoint.Transform = Transform;
 			}
 
+			// Only add valid points within bounds
 			FPCGPoint BoundsTestPoint;
-			if (bValid && (!BoundingShapeData || BoundingShapeData->SamplePoint(Transform, InResult.Box, BoundsTestPoint, nullptr)))
+			if (bValid && (BoundingShapeData || BoundingShapeData->SamplePoint(Transform, InResult.Box, BoundsTestPoint, nullptr)))
 			{
 				SetSeed(OutPoint, InResult.LocalTransform.GetLocation(), Params);
 				SetMetadata(InResult, OutPoint, OutPointData->Metadata);
@@ -645,7 +647,7 @@ namespace PCGSplineSampler
 						OutPoint.Transform = ProjPoint.Transform;
 					}
 
-					// Test point against bounds.
+					// Prune points outside of bounds
 					FPCGPoint BoundsTestPoint;
 					if (BoundingShapeData && !BoundingShapeData->SamplePoint(OutPoint.Transform, OutPoint.GetLocalBounds(), BoundsTestPoint, nullptr))
 					{
@@ -1076,6 +1078,13 @@ namespace PCGSplineSampler
 							TransformWS.SetLocation(ProjectedPoint.Transform.GetLocation());
 						}
 
+						// Prune points outside of bounds
+						FPCGPoint BoundsTestPoint;
+						if (InBoundingShape && !InBoundingShape->SamplePoint(TransformWS, SplineLocalBounds.GetBox(), BoundsTestPoint, nullptr))
+						{
+							continue;
+						}
+						
 						InteriorSplinePointData[DispatchIndex].Emplace(TransformWS, PointLocationLS, Density);
 					}
 				}
@@ -1176,16 +1185,28 @@ bool FPCGSplineSamplerElement::ExecuteInternal(FPCGContext* Context) const
 	TArray<FPCGTaggedData> SplineInputs = Context->InputData.GetInputsByPin(PCGSplineSamplerConstants::SplineLabel);
 	TArray<FPCGTaggedData> BoundingShapeInputs = Context->InputData.GetInputsByPin(PCGSplineSamplerConstants::BoundingShapeLabel);
 
-	// Grab the Bounding Shape input if there is one.
-	// TODO: Once we support time-slicing, put this in the context and root (see FPCGSurfaceSamplerContext)
-	bool bUnionCreated = false;
-	const UPCGSpatialData* BoundingShape = Context->InputData.GetSpatialUnionOfInputsByPin(PCGSplineSamplerConstants::BoundingShapeLabel, bUnionCreated);
-	if (!BoundingShape && BoundingShapeInputs.Num() > 0)
+	const FPCGSplineSamplerParams& SamplerParams = Settings->SamplerParams;
+
+	const UPCGSpatialData* BoundingShape = nullptr;
+
+	// If unbounded, specific samplers will not be constrained by a nullptr BoundingShape
+	if(!SamplerParams.bUnbounded)
 	{
-		PCGE_LOG(Warning, GraphAndLog, LOCTEXT("BoundingShapeMissing", "Bounding Shape input is missing or of unsupported type and will not be used"));
-	}
+		// TODO: Once we support time-slicing, put this in the context and root (see FPCGSurfaceSamplerContext)
+		bool bUnionCreated = false;
+		// Grab the Bounding Shape input if there is one.
+		BoundingShape = Context->InputData.GetSpatialUnionOfInputsByPin(PCGSplineSamplerConstants::BoundingShapeLabel, bUnionCreated);
 
-	const FPCGSplineSamplerParams& SamplerParams = Settings->SamplerParams;
+		// Fallback to getting bounds from actor
+		if (!BoundingShape && Context->SourceComponent.IsValid())
+		{
+			BoundingShape = Cast<UPCGSpatialData>(Context->SourceComponent->GetActorPCGData());
+		}
+	}
+	else if (BoundingShapeInputs.Num() > 0)
+	{
+		PCGE_LOG(Verbose, LogOnly, LOCTEXT("BoundsIgnored", "The bounds of the Bounding Shape input pin will be ignored because the Unbounded option is enabled."));
+	}
 
 	TArray<FPCGTaggedData>& Outputs = Context->OutputData.TaggedData;
 
@@ -1233,6 +1254,19 @@ bool FPCGSplineSamplerElement::ExecuteInternal(FPCGContext* Context) const
 }
 
 #if WITH_EDITOR
+void UPCGSplineSamplerSettings::ApplyDeprecation(UPCGNode* InOutNode)
+{
+	check(InOutNode);
+	
+	if (DataVersion < FPCGCustomVersion::SplineSamplerBoundedByDefault)
+	{
+		UE_LOG(LogPCG, Log, TEXT("Spline Sampler node migrated from an older version. Defaulting to 'Unbounded' to match previous behavior."));
+		SamplerParams.bUnbounded = true;
+	}
+	
+	Super::ApplyDeprecation(InOutNode);
+}
+
 void UPCGSplineSamplerSettings::ApplyDeprecationBeforeUpdatePins(UPCGNode* InOutNode, TArray<TObjectPtr<UPCGPin>>& InputPins, TArray<TObjectPtr<UPCGPin>>& OutputPins)
 {
 	if (DataVersion < FPCGCustomVersion::SplineSamplerUpdatedNodeInputs && ensure(InOutNode))
diff --git a/Engine/Plugins/Experimental/PCG/Source/PCG/Public/Elements/PCGSplineSampler.h b/Engine/Plugins/Experimental/PCG/Source/PCG/Public/Elements/PCGSplineSampler.h
--- a/Engine/Plugins/Experimental/PCG/Source/PCG/Public/Elements/PCGSplineSampler.h
+++ b/Engine/Plugins/Experimental/PCG/Source/PCG/Public/Elements/PCGSplineSampler.h
@@ -85,6 +85,13 @@ struct PCG_API FPCGSplineSamplerParams
 	UPROPERTY(BlueprintReadWrite, EditAnywhere, Category = Settings, meta = (EditCondition = "bComputeCurvature"))
 	FName CurvatureAttribute = "Curvature";
 
+	/** If no Bounding Shape input is provided, the actor bounds are used to limit the sample generation domain.
+	* This option allows ignoring the actor bounds and generating over the entire spline. Use with caution as this
+	* may generate a lot of points.
+	*/
+	UPROPERTY(BlueprintReadWrite, EditAnywhere, Category = Settings, meta = (PCG_Overridable))
+	bool bUnbounded = false; 
+	
 	/** The space between each sample point */
 	UPROPERTY(BlueprintReadWrite, EditAnywhere, Category = Settings, meta = (ClampMin = "0.1", EditCondition = "Dimension==EPCGSplineSamplingDimension::OnInterior"))
 	float InteriorSampleSpacing = 100.0f;
@@ -156,6 +163,7 @@ public:
 #endif
 
 #if WITH_EDITOR
+	virtual void ApplyDeprecation(UPCGNode* InOutNode) override;
 	virtual void ApplyDeprecationBeforeUpdatePins(UPCGNode* InOutNode, TArray<TObjectPtr<UPCGPin>>& InputPins, TArray<TObjectPtr<UPCGPin>>& OutputPins) override;
 #endif
 
diff --git a/Engine/Plugins/Experimental/PCG/Source/PCG/Public/Elements/PCGSurfaceSampler.h b/Engine/Plugins/Experimental/PCG/Source/PCG/Public/Elements/PCGSurfaceSampler.h
--- a/Engine/Plugins/Experimental/PCG/Source/PCG/Public/Elements/PCGSurfaceSampler.h
+++ b/Engine/Plugins/Experimental/PCG/Source/PCG/Public/Elements/PCGSurfaceSampler.h
@@ -57,7 +57,7 @@ public:
 	UPROPERTY(BlueprintReadWrite, EditAnywhere, Category = Settings, meta=(ClampMin="0", PCG_Overridable))
 	float Looseness = 1.0f;
 
-	/** If no Bounding Shape input is provided the actor bounds are used to limit the sample generation area.
+	/** If no Bounding Shape input is provided, the actor bounds are used to limit the sample generation domain.
 	* This option allows ignoring the actor bounds and generating over the entire surface. Use with caution as this
 	* may generate a lot of points.
 	*/
diff --git a/Engine/Plugins/Experimental/PCG/Source/PCG/Public/Elements/PCGVolumeSampler.h b/Engine/Plugins/Experimental/PCG/Source/PCG/Public/Elements/PCGVolumeSampler.h
--- a/Engine/Plugins/Experimental/PCG/Source/PCG/Public/Elements/PCGVolumeSampler.h
+++ b/Engine/Plugins/Experimental/PCG/Source/PCG/Public/Elements/PCGVolumeSampler.h
@@ -34,8 +34,8 @@ public:
 	UPROPERTY(BlueprintReadWrite, EditAnywhere, Category = "Data", meta = (PCG_Overridable))
 	FVector VoxelSize = FVector(100.0, 100.0, 100.0);
 
-	/** If no Bounding Shape input is provided the actor bounds are used to limit the sample generation area.
-	* This option allows ignoring the actor bounds and generating over the entire volume provided. Use with caution as this
+	/** If no Bounding Shape input is provided, the actor bounds are used to limit the sample generation domain.
+	* This option allows ignoring the actor bounds and generating over the entire volume. Use with caution as this
 	* may generate a lot of points.
 	*/
 	UPROPERTY(BlueprintReadWrite, EditAnywhere, Category = Settings, meta = (PCG_Overridable))
diff --git a/Engine/Plugins/Experimental/PCG/Source/PCG/Public/PCGCustomVersion.h b/Engine/Plugins/Experimental/PCG/Source/PCG/Public/PCGCustomVersion.h
--- a/Engine/Plugins/Experimental/PCG/Source/PCG/Public/PCGCustomVersion.h
+++ b/Engine/Plugins/Experimental/PCG/Source/PCG/Public/PCGCustomVersion.h
@@ -38,6 +38,9 @@ struct PCG_API FPCGCustomVersion
 		// Renamed params to override.
 		RenameDefaultParamsToOverride = 9,
 
+		// Behavior change for SplineSampler which now defaults to being bounded
+		SplineSamplerBoundedByDefault = 10,
+
 		// -----<new versions can be added above this line>-------------------------------------------------
 		VersionPlusOne,
 		LatestVersion = VersionPlusOne - 1
