[PCG] Fix crash on creating new InstanceDataPacker and PointMatchAndSet Blueprint Native Events

diff --git a/Engine/Plugins/Experimental/PCG/Source/PCG/Private/InstanceDataPackers/PCGInstanceDataPackerBase.cpp b/Engine/Plugins/Experimental/PCG/Source/PCG/Private/InstanceDataPackers/PCGInstanceDataPackerBase.cpp
--- a/Engine/Plugins/Experimental/PCG/Source/PCG/Private/InstanceDataPackers/PCGInstanceDataPackerBase.cpp
+++ b/Engine/Plugins/Experimental/PCG/Source/PCG/Private/InstanceDataPackers/PCGInstanceDataPackerBase.cpp
@@ -2,11 +2,16 @@
 
 #include "InstanceDataPackers/PCGInstanceDataPackerBase.h"
 
-#include "MeshSelectors/PCGMeshSelectorBase.h"
+#include "PCGElement.h"
 #include "Metadata/PCGMetadata.h"
 
 #include UE_INLINE_GENERATED_CPP_BY_NAME(PCGInstanceDataPackerBase)
 
+void UPCGInstanceDataPackerBase::PackInstances_Implementation(FPCGContext& Context, const UPCGSpatialData* InSpatialData, const FPCGMeshInstanceList& InstanceList, FPCGPackedCustomData& OutPackedCustomData) const
+{
+	PCGE_LOG_C(Error, GraphAndLog, &Context, NSLOCTEXT("PCGInstanceDataPackerBase", "InstanceDataPackerBaseFailed", "Unable to execute InstanceDataPacker pure virtual base function, override the PackInstances function or use a default implementation."));
+}
+
 bool UPCGInstanceDataPackerBase::AddTypeToPacking(int TypeId, FPCGPackedCustomData& OutPackedCustomData) const
 {
 	if (TypeId == PCG::Private::MetadataTypes<float>::Id)
diff --git a/Engine/Plugins/Experimental/PCG/Source/PCG/Private/MatchAndSet/PCGMatchAndSetBase.cpp b/Engine/Plugins/Experimental/PCG/Source/PCG/Private/MatchAndSet/PCGMatchAndSetBase.cpp
--- a/Engine/Plugins/Experimental/PCG/Source/PCG/Private/MatchAndSet/PCGMatchAndSetBase.cpp
+++ b/Engine/Plugins/Experimental/PCG/Source/PCG/Private/MatchAndSet/PCGMatchAndSetBase.cpp
@@ -17,6 +17,15 @@ void UPCGMatchAndSetBase::SetType(EPCGMetadataTypes InType, EPCGMetadataTypesCon
 	StringMode = InStringMode;
 }
 
+void UPCGMatchAndSetBase::MatchAndSet_Implementation(
+	FPCGContext& Context,
+	const UPCGPointMatchAndSetSettings* InSettings,
+	const UPCGPointData* InPointData,
+	UPCGPointData* OutPointData) const
+{
+	PCGE_LOG_C(Error, GraphAndLog, &Context, LOCTEXT("MatchAndSetBaseFailed", "Unable to execute MatchAndSet pure virtual base function, override the MatchAndSet function or use a default implementation."));
+}
+
 bool UPCGMatchAndSetBase::CreateAttributeIfNeeded(FPCGContext& Context, const FPCGAttributePropertySelector& Selector, const FPCGMetadataTypesConstantStruct& ConstantValue, UPCGPointData* OutPointData, const UPCGPointMatchAndSetSettings* InSettings) const
 {
 	check(OutPointData && OutPointData->Metadata);
diff --git a/Engine/Plugins/Experimental/PCG/Source/PCG/Public/InstanceDataPackers/PCGInstanceDataPackerBase.h b/Engine/Plugins/Experimental/PCG/Source/PCG/Public/InstanceDataPackers/PCGInstanceDataPackerBase.h
--- a/Engine/Plugins/Experimental/PCG/Source/PCG/Public/InstanceDataPackers/PCGInstanceDataPackerBase.h
+++ b/Engine/Plugins/Experimental/PCG/Source/PCG/Public/InstanceDataPackers/PCGInstanceDataPackerBase.h
@@ -33,7 +33,7 @@ public:
 	UFUNCTION(BlueprintNativeEvent, Category = InstancePacking)
 	void PackInstances(FPCGContext& Context, const UPCGSpatialData* InSpatialData, const FPCGMeshInstanceList& InstanceList, FPCGPackedCustomData& OutPackedCustomData) const;
 
-	virtual void PackInstances_Implementation(FPCGContext& Context, const UPCGSpatialData* InSpatialData, const FPCGMeshInstanceList& InstanceList, FPCGPackedCustomData& OutPackedCustomData) const PURE_VIRTUAL(UPCGInstanceDataPackerBase::PackInstances_Implementation);
+	virtual void PackInstances_Implementation(FPCGContext& Context, const UPCGSpatialData* InSpatialData, const FPCGMeshInstanceList& InstanceList, FPCGPackedCustomData& OutPackedCustomData) const;
 
 	/** Interprets Metadata TypeId and increments OutPackedCustomData.NumCustomDataFloats appropriately. Returns false if the type could not be interpreted. */
 	UFUNCTION(BlueprintCallable, Category = InstancePacking)
diff --git a/Engine/Plugins/Experimental/PCG/Source/PCG/Public/MatchAndSet/PCGMatchAndSetBase.h b/Engine/Plugins/Experimental/PCG/Source/PCG/Public/MatchAndSet/PCGMatchAndSetBase.h
--- a/Engine/Plugins/Experimental/PCG/Source/PCG/Public/MatchAndSet/PCGMatchAndSetBase.h
+++ b/Engine/Plugins/Experimental/PCG/Source/PCG/Public/MatchAndSet/PCGMatchAndSetBase.h
@@ -53,7 +53,7 @@ public:
 		FPCGContext& Context,
 		const UPCGPointMatchAndSetSettings* InSettings,
 		const UPCGPointData* InPointData,
-		UPCGPointData* OutPointData) const PURE_VIRTUAL(UPCGMatchAndSetBase::MatchAndSet_Implementation);
+		UPCGPointData* OutPointData) const;
 
 protected:
 	bool CreateAttributeIfNeeded(
