[PCG] Fixed a crash in the spline sampler when missing bounds information

diff --git a/Engine/Plugins/Experimental/PCG/Source/PCG/Private/Elements/PCGSplineSampler.cpp b/Engine/Plugins/Experimental/PCG/Source/PCG/Private/Elements/PCGSplineSampler.cpp
--- a/Engine/Plugins/Experimental/PCG/Source/PCG/Private/Elements/PCGSplineSampler.cpp
+++ b/Engine/Plugins/Experimental/PCG/Source/PCG/Private/Elements/PCGSplineSampler.cpp
@@ -548,7 +548,7 @@ namespace PCGSplineSampler
 
 			// Only add valid points within bounds
 			FPCGPoint BoundsTestPoint;
-			if (bValid && (BoundingShapeData || BoundingShapeData->SamplePoint(Transform, InResult.Box, BoundsTestPoint, nullptr)))
+			if (bValid && (!BoundingShapeData || BoundingShapeData->SamplePoint(Transform, InResult.Box, BoundsTestPoint, nullptr)))
 			{
 				SetSeed(OutPoint, InResult.LocalTransform.GetLocation(), Params);
 				SetMetadata(InResult, OutPoint, OutPointData->Metadata);
