[PCG] Fixed an issue in the attribute list view where metadata would not display properly for spatial data

diff --git a/Engine/Plugins/Experimental/PCG/Source/PCGEditor/Private/SPCGEditorGraphAttributeListView.cpp b/Engine/Plugins/Experimental/PCG/Source/PCGEditor/Private/SPCGEditorGraphAttributeListView.cpp
--- a/Engine/Plugins/Experimental/PCG/Source/PCGEditor/Private/SPCGEditorGraphAttributeListView.cpp
+++ b/Engine/Plugins/Experimental/PCG/Source/PCGEditor/Private/SPCGEditorGraphAttributeListView.cpp
@@ -489,7 +489,7 @@ void SPCGEditorGraphAttributeListView::RefreshAttributeList()
 			const UPCGMetadata* PCGMetadata = PCGPointData->ConstMetadata();
 
 			AddPointDataColumns(PCGPointData);
-			GenerateColumnsFromMetadata(PCGData, PCGMetadata);
+			GenerateColumnsFromMetadata(PCGPointData, PCGMetadata);
 
 			const TArray<FPCGPoint>& PCGPoints = PCGPointData->GetPoints();
 			const int32 NumPoints = PCGPoints.Num();
