[PCG] Add $Index as an accessor

* Also refactor the whole CreateAccessor helper to remove duplicated code for const and non-const versions

diff --git a/Engine/Plugins/Experimental/PCG/Source/PCG/Private/Metadata/Accessors/PCGAttributeAccessorHelpers.cpp b/Engine/Plugins/Experimental/PCG/Source/PCG/Private/Metadata/Accessors/PCGAttributeAccessorHelpers.cpp
--- a/Engine/Plugins/Experimental/PCG/Source/PCG/Private/Metadata/Accessors/PCGAttributeAccessorHelpers.cpp
+++ b/Engine/Plugins/Experimental/PCG/Source/PCG/Private/Metadata/Accessors/PCGAttributeAccessorHelpers.cpp
@@ -10,6 +10,7 @@
 #include "Metadata/PCGAttributePropertySelector.h"
 #include "Metadata/Accessors/IPCGAttributeAccessorTpl.h"
 #include "Metadata/Accessors/PCGAttributeAccessor.h"
+#include "Metadata/Accessors/PCGCustomAccessor.h"
 #include "Metadata/Accessors/PCGPropertyAccessor.h"
 #include "Metadata/Accessors/PCGAttributeExtractor.h"
 
@@ -43,13 +44,167 @@ namespace PCGAttributeAccessorHelpers
 		}
 	}
 
-	void ExtractMetadataAtribute(const UPCGData* InData, FName Name, UPCGMetadata const*& OutMetadata, FPCGMetadataAttributeBase const*& OutAttribute, bool bForceGetNone = false)
+	// Don't be afraid of this enormous template!
+	// We want to be able to create a TUniquePtr<const IPCGAttributeAccessor> for const accessors and TUniquePtr<IPCGAttributeAccessor> for mutable accessors.
+	// This template just says that:
+	// * AccessorType needs to be IPCGAttributeAccessor or const IPCGAttributeAccessor
+	// * DataType needs to be UPCGData or const UPCGData
+	// * And if one of them is const, they both need to be const.
+	// It allows us to factorise the same logic for const and non-const version (with a few const_cast in the process, but we still have something const at the end)
+	template <
+		typename AccessorType, 
+		typename DataType, 
+		typename = typename std::enable_if_t<
+			std::is_same_v<IPCGAttributeAccessor, std::remove_const_t<AccessorType>> &&
+			std::is_same_v<UPCGData, std::remove_const_t<DataType>> &&
+			std::is_const_v<AccessorType> == std::is_const_v<DataType>
+		>
+	>
+	TUniquePtr<AccessorType> CreateAccessorImpl(DataType* InData, const FPCGAttributePropertySelector& InSelector)
 	{
+		const FName Name = InSelector.GetName();
+		TUniquePtr<IPCGAttributeAccessor> Accessor;
+
+		if (InSelector.Selection == EPCGAttributePropertySelection::PointProperty)
+		{
+			if (!InData || InData->template IsA<UPCGPointData>())
+			{
+				if (const FProperty* Property = FPCGPoint::StaticStruct()->FindPropertyByName(Name))
+				{
+					Accessor = CreatePropertyAccessor(Property);
+				}
+				else if (FPCGPoint::HasCustomPropertyGetterSetter(Name))
+				{
+					Accessor = FPCGPoint::CreateCustomPropertyAccessor(Name);
+				}
+			}
+		}
+		else if (InSelector.Selection == EPCGAttributePropertySelection::ExtraProperty)
+		{
+			Accessor = CreateExtraAccessor(InSelector.ExtraProperty);
+
+			if (!Accessor.IsValid())
+			{
+				UE_LOG(LogPCG, Error, TEXT("[PCGAttributeAccessorHelpers::ConstAccessor] Expected to select an extra property but the data doesn't support this property."));
+				return TUniquePtr<IPCGAttributeAccessor>();
+			}
+		}
+
+		// At this point, it is not a point data or we didn't find a property.
+		// We can't continue if it is a property wanted.
+		if (InSelector.Selection == EPCGAttributePropertySelection::PointProperty && !Accessor.IsValid())
+		{
+			UE_LOG(LogPCG, Error, TEXT("[PCGAttributeAccessorHelpers::CreateAccessor] Expected to select a property but the data doesn't support this property."));
+			return TUniquePtr<IPCGAttributeAccessor>();
+		}
+
+		if (InSelector.Selection == EPCGAttributePropertySelection::Attribute)
+		{
+			UPCGMetadata* Metadata = nullptr;
+			FPCGMetadataAttributeBase* Attribute = nullptr;
+
+			// It is OK to const_cast here, since we will create a const accessor if the input is const.
+			ExtractMetadataAtribute(const_cast<UPCGData*>(InData), Name, Metadata, Attribute);
+
+			auto CreateTypedAccessor = [Attribute, Metadata](auto Dummy) -> TUniquePtr<IPCGAttributeAccessor>
+			{
+				using AttributeType = decltype(Dummy);
+				if constexpr (std::is_const_v<AccessorType>)
+				{
+					return MakeUnique<FPCGAttributeAccessor<AttributeType>>(static_cast<const FPCGMetadataAttribute<AttributeType>*>(Attribute), Metadata);
+				}
+				else
+				{
+					return MakeUnique<FPCGAttributeAccessor<AttributeType>>(static_cast<FPCGMetadataAttribute<AttributeType>*>(Attribute), Metadata);
+				}
+			};
+
+			if (Attribute && Metadata)
+			{
+				Accessor = PCGMetadataAttribute::CallbackWithRightType(Attribute->GetTypeId(), CreateTypedAccessor);
+			}
+			else
+			{
+				return TUniquePtr<AccessorType>();
+			}
+		}
+
+		if (!Accessor.IsValid())
+		{
+			return TUniquePtr<AccessorType>();
+		}
+
+		// At this point, check if we have chain accessors
+		for (const FString& ExtraName : InSelector.ExtraNames)
+		{
+			bool bSuccess = false;
+			Accessor = CreateChainAccessor(std::move(Accessor), FName(ExtraName), bSuccess);
+			if (!bSuccess)
+			{
+				UE_LOG(LogPCG, Error, TEXT("[PCGAttributeAccessorHelpers::CreateAccessor] Extra selectors don't match existing properties."));
+				return TUniquePtr<AccessorType>();
+			}
+		}
+
+		// Can't create mutable accessor that are read only.
+		if constexpr (!std::is_const_v<AccessorType>)
+		{
+			if (Accessor && Accessor->IsReadOnly())
+			{
+				UE_LOG(LogPCG, Error, TEXT("[PCGAttributeAccessorHelpers::CreateAccessor] Attribute can not be written into, since it is read-only."));
+				return TUniquePtr<AccessorType>();
+			}
+		}
+
+		return Accessor;
+	}
+
+	// Same thing for the keys. Check CreateAccessorImpl for more info about this big template.
+	template <
+		typename KeysType,
+		typename DataType,
+		typename = typename std::enable_if_t<
+			std::is_same_v<IPCGAttributeAccessorKeys, std::remove_const_t<KeysType>> &&
+			std::is_same_v<UPCGData, std::remove_const_t<DataType>> &&
+			std::is_const_v<KeysType> == std::is_const_v<DataType>
+		>
+	>
+	TUniquePtr<KeysType> CreateKeysImpl(DataType* InData, const FPCGAttributePropertySelector& InSelector)
+	{
+		if (UPCGPointData* PointData = Cast<UPCGPointData>(const_cast<UPCGData*>(InData)))
+		{
+			if constexpr (std::is_const_v<KeysType>)
+			{
+				return MakeUnique<FPCGAttributeAccessorKeysPoints>(PointData->GetPoints());
+			}
+			else
+			{
+				TArrayView<FPCGPoint> View(PointData->GetMutablePoints());
+				return MakeUnique<FPCGAttributeAccessorKeysPoints>(View);
+			}
+		}
+
 		UPCGMetadata* Metadata = nullptr;
 		FPCGMetadataAttributeBase* Attribute = nullptr;
-		ExtractMetadataAtribute(const_cast<UPCGData*>(InData), Name, Metadata, Attribute, bForceGetNone);
-		OutMetadata = Metadata;
-		OutAttribute = Attribute;
+
+		// It is OK to const_cast here, since we will create const keys if the input is const.
+		ExtractMetadataAtribute(const_cast<UPCGData*>(InData), InSelector.GetName(), Metadata, Attribute);
+
+		if (Attribute)
+		{
+			if constexpr (std::is_const_v<KeysType>)
+			{
+				return MakeUnique<FPCGAttributeAccessorKeysEntries>(static_cast<const FPCGMetadataAttributeBase*>(Attribute));
+			}
+			else
+			{
+				return MakeUnique<FPCGAttributeAccessorKeysEntries>(Attribute);
+			}
+		}
+		else
+		{
+			return TUniquePtr<KeysType>();
+		}
 	}
 
 	TUniquePtr<IPCGAttributeAccessor> CreateChainAccessor(TUniquePtr<IPCGAttributeAccessor> InAccessor, FName Name, bool& bOutSuccess)
@@ -246,6 +401,17 @@ bool PCGAttributeAccessorHelpers::IsPropertyAccessorSupported(const FName InProp
 	});
 }
 
+TUniquePtr<IPCGAttributeAccessor> PCGAttributeAccessorHelpers::CreateExtraAccessor(EPCGExtraProperties InExtraProperties)
+{
+	switch (InExtraProperties)
+	{
+	case EPCGExtraProperties::Index:
+		return MakeUnique<FPCGIndexAccessor>();
+	default:
+		return TUniquePtr<IPCGAttributeAccessor>();
+	}
+}
+
 TUniquePtr<const IPCGAttributeAccessor> PCGAttributeAccessorHelpers::CreateConstAccessorForOverrideParam(const FPCGDataCollection& InInputData, const FPCGSettingsOverridableParam& InParam, FName* OutAttributeName)
 {
 	bool bFromGlobalParamsPin = false;
@@ -281,194 +447,20 @@ TUniquePtr<const IPCGAttributeAccessor> PCGAttributeAccessorHelpers::CreateConst
 
 TUniquePtr<const IPCGAttributeAccessor> PCGAttributeAccessorHelpers::CreateConstAccessor(const UPCGData* InData, const FPCGAttributePropertySelector& InSelector)
 {
-	const FName Name = InSelector.GetName();
-	TUniquePtr<IPCGAttributeAccessor> Accessor;
-
-	if (InSelector.Selection == EPCGAttributePropertySelection::PointProperty)
-	{
-		if(!InData || Cast<const UPCGPointData>(InData))
-		{
-			if (const FProperty* Property = FPCGPoint::StaticStruct()->FindPropertyByName(Name))
-			{
-				Accessor = CreatePropertyAccessor(Property);
-			}
-			else if (FPCGPoint::HasCustomPropertyGetterSetter(Name))
-			{
-				Accessor = FPCGPoint::CreateCustomPropertyAccessor(Name);
-			}
-		}
-	}
-
-	// At this point, it is not a point data or we didn't find a property.
-	// We can't continue if it is a property wanted.
-	if (InSelector.Selection == EPCGAttributePropertySelection::PointProperty && !Accessor.IsValid())
-	{
-		UE_LOG(LogPCG, Error, TEXT("[PCGAttributeAccessorHelpers::CreateConstAccessor] Expected to select a property but the data doesn't support this property."));
-		return TUniquePtr<const IPCGAttributeAccessor>();
-	}
-
-	if (InSelector.Selection == EPCGAttributePropertySelection::Attribute)
-	{
-		const UPCGMetadata* Metadata = nullptr;
-		const FPCGMetadataAttributeBase* Attribute = nullptr;
-
-		ExtractMetadataAtribute(InData, Name, Metadata, Attribute, InSelector.bForceGetNone);
-
-		auto CreateTypedAccessor = [Attribute, Metadata](auto Dummy) -> TUniquePtr<IPCGAttributeAccessor>
-		{
-			using AttributeType = decltype(Dummy);
-			return MakeUnique<FPCGAttributeAccessor<AttributeType>>(static_cast<const FPCGMetadataAttribute<AttributeType>*>(Attribute), Metadata);
-		};
-
-		if (Attribute && Metadata)
-		{
-			Accessor = PCGMetadataAttribute::CallbackWithRightType(Attribute->GetTypeId(), CreateTypedAccessor);
-		}
-		else
-		{
-			return TUniquePtr<const IPCGAttributeAccessor>();
-		}
-	}
-
-	if (!Accessor.IsValid())
-	{
-		return TUniquePtr<const IPCGAttributeAccessor>();
-	}
-
-	// At this point, check if we have chain accessors
-	for (const FString& ExtraName : InSelector.ExtraNames)
-	{
-		bool bSuccess = false;
-		Accessor = CreateChainAccessor(std::move(Accessor), FName(ExtraName), bSuccess);
-		if (!bSuccess)
-		{
-			UE_LOG(LogPCG, Error, TEXT("[PCGAttributeAccessorHelpers::CreateConstAccessor] Extra selectors don't match existing properties."));
-			return TUniquePtr<const IPCGAttributeAccessor>();
-		}
-	}
-
-	return Accessor;
+	return CreateAccessorImpl<const IPCGAttributeAccessor, const UPCGData>(InData, InSelector);
 }
 
 TUniquePtr<IPCGAttributeAccessor> PCGAttributeAccessorHelpers::CreateAccessor(UPCGData* InData, const FPCGAttributePropertySelector& InSelector)
 {
-	const FName Name = InSelector.GetName();
-	TUniquePtr<IPCGAttributeAccessor> Accessor;
-
-	if (InSelector.Selection == EPCGAttributePropertySelection::PointProperty)
-	{
-		if (!InData || Cast<const UPCGPointData>(InData))
-		{
-			if (const FProperty* Property = FPCGPoint::StaticStruct()->FindPropertyByName(Name))
-			{
-				Accessor = CreatePropertyAccessor(Property);
-			}
-			else if (FPCGPoint::HasCustomPropertyGetterSetter(Name))
-			{
-				Accessor = FPCGPoint::CreateCustomPropertyAccessor(Name);
-			}
-		}
-	}
-
-	// At this point, it is not a point data or we didn't find a property.
-	// We can't continue if it is a property wanted.
-	if (InSelector.Selection == EPCGAttributePropertySelection::PointProperty && !Accessor.IsValid())
-	{
-		UE_LOG(LogPCG, Error, TEXT("[PCGAttributeAccessorHelpers::CreateAccessor] Expected to select a property but the data doesn't support this property."));
-		return TUniquePtr<IPCGAttributeAccessor>();
-	}
-
-	if (InSelector.Selection == EPCGAttributePropertySelection::Attribute)
-	{
-		UPCGMetadata* Metadata = nullptr;
-		FPCGMetadataAttributeBase* Attribute = nullptr;
-
-		ExtractMetadataAtribute(InData, Name, Metadata, Attribute);
-
-		auto CreateTypedAccessor = [Attribute, Metadata](auto Dummy) -> TUniquePtr<IPCGAttributeAccessor>
-		{
-			using AttributeType = decltype(Dummy);
-			return MakeUnique<FPCGAttributeAccessor<AttributeType>>(static_cast<FPCGMetadataAttribute<AttributeType>*>(Attribute), Metadata);
-		};
-
-		if (Attribute && Metadata)
-		{
-			Accessor = PCGMetadataAttribute::CallbackWithRightType(Attribute->GetTypeId(), CreateTypedAccessor);
-		}
-		else
-		{
-			return TUniquePtr<IPCGAttributeAccessor>();
-		}
-	}
-
-	if (!Accessor.IsValid())
-	{
-		return TUniquePtr<IPCGAttributeAccessor>();
-	}
-
-	// At this point, check if we have chain accessors
-	for (const FString& ExtraName : InSelector.ExtraNames)
-	{
-		bool bSuccess = false;
-		Accessor = CreateChainAccessor(std::move(Accessor), FName(ExtraName), bSuccess);
-		if (!bSuccess)
-		{
-			UE_LOG(LogPCG, Error, TEXT("[PCGAttributeAccessorHelpers::CreateAccessor] Extra selectors don't match existing properties."));
-			return TUniquePtr<IPCGAttributeAccessor>();
-		}
-	}
-
-	// Can't create mutable accessor that are read only.
-	if (Accessor && Accessor->IsReadOnly())
-	{
-		UE_LOG(LogPCG, Error, TEXT("[PCGAttributeAccessorHelpers::CreateAccessor] Accessor is read-only, therefore not mutable."));
-		return TUniquePtr<IPCGAttributeAccessor>();
-	}
-
-	return Accessor;
+	return CreateAccessorImpl<IPCGAttributeAccessor, UPCGData>(InData, InSelector);
 }
 
 TUniquePtr<const IPCGAttributeAccessorKeys> PCGAttributeAccessorHelpers::CreateConstKeys(const UPCGData* InData, const FPCGAttributePropertySelector& InSelector)
 {
-	if (const UPCGPointData* PointData = Cast<const UPCGPointData>(InData))
-	{
-		return MakeUnique<FPCGAttributeAccessorKeysPoints>(PointData->GetPoints());
-	}
-
-	const UPCGMetadata* Metadata = nullptr;
-	const FPCGMetadataAttributeBase* Attribute = nullptr;
-
-	ExtractMetadataAtribute(InData, InSelector.GetName(), Metadata, Attribute);
-
-	if (Attribute)
-	{
-		return MakeUnique<FPCGAttributeAccessorKeysEntries>(Attribute);
-	}
-	else
-	{
-		return TUniquePtr<IPCGAttributeAccessorKeys>();
-	}
+	return CreateKeysImpl<const IPCGAttributeAccessorKeys, const UPCGData>(InData, InSelector);
 }
 
 TUniquePtr<IPCGAttributeAccessorKeys> PCGAttributeAccessorHelpers::CreateKeys(UPCGData* InData, const FPCGAttributePropertySelector& InSelector)
 {
-	if (UPCGPointData* PointData = Cast<UPCGPointData>(InData))
-	{
-		TArrayView<FPCGPoint> View(PointData->GetMutablePoints());
-		return MakeUnique<FPCGAttributeAccessorKeysPoints>(View);
-	}
-
-	UPCGMetadata* Metadata = nullptr;
-	FPCGMetadataAttributeBase* Attribute = nullptr;
-
-	ExtractMetadataAtribute(InData, InSelector.GetName(), Metadata, Attribute);
-
-	if (Attribute)
-	{
-		return MakeUnique<FPCGAttributeAccessorKeysEntries>(Attribute);
-	}
-	else
-	{
-		return TUniquePtr<IPCGAttributeAccessorKeys>();
-	}
+	return CreateKeysImpl<IPCGAttributeAccessorKeys, UPCGData>(InData, InSelector);
 }
diff --git a/Engine/Plugins/Experimental/PCG/Source/PCG/Private/Metadata/PCGAttributePropertySelector.cpp b/Engine/Plugins/Experimental/PCG/Source/PCG/Private/Metadata/PCGAttributePropertySelector.cpp
--- a/Engine/Plugins/Experimental/PCG/Source/PCG/Private/Metadata/PCGAttributePropertySelector.cpp
+++ b/Engine/Plugins/Experimental/PCG/Source/PCG/Private/Metadata/PCGAttributePropertySelector.cpp
@@ -27,6 +27,17 @@ FName FPCGAttributePropertySelector::GetName() const
 			return NAME_None;
 		}
 	}
+	case EPCGAttributePropertySelection::ExtraProperty:
+	{
+		if (const UEnum* EnumPtr = StaticEnum<EPCGExtraProperties>())
+		{
+			return FName(EnumPtr->GetNameStringByValue((int64)ExtraProperty));
+		}
+		else
+		{
+			return NAME_None;
+		}
+	}
 	default:
 		return AttributeName;
 	}
@@ -74,13 +85,34 @@ bool FPCGAttributePropertySelector::SetAttributeName(FName InAttributeName, bool
 	}
 }
 
+bool FPCGAttributePropertySelector::SetExtraProperty(EPCGExtraProperties InExtraProperty, bool bResetExtraNames)
+{
+	const bool bHasExtraNames = !ExtraNames.IsEmpty();
+	if (bResetExtraNames)
+	{
+		ExtraNames.Empty();
+	}
+
+	if (Selection == EPCGAttributePropertySelection::ExtraProperty && InExtraProperty == ExtraProperty)
+	{
+		// Nothing changed, except perhaps the extra names
+		return (bHasExtraNames && bResetExtraNames);
+	}
+	else
+	{
+		Selection = EPCGAttributePropertySelection::ExtraProperty;
+		ExtraProperty = InExtraProperty;
+		return true;
+	}
+}
+
 FText FPCGAttributePropertySelector::GetDisplayText() const
 {
 	FString Res;
 	const FName Name = GetName();
 
 	// Add a '$' if it is a property
-	if (Selection == EPCGAttributePropertySelection::PointProperty && (Name != NAME_None))
+	if (Selection != EPCGAttributePropertySelection::Attribute && (Name != NAME_None))
 	{
 		Res = FString(PCGAttributePropertySelectorConstants::PropertyPrefix) + Name.ToString();
 	}
@@ -130,14 +162,26 @@ bool FPCGAttributePropertySelector::Update(FString NewValue)
 
 	if (!NewName.IsEmpty() && NewName[0] == PCGAttributePropertySelectorConstants::PropertyPrefixChar)
 	{
+		// Remove the first character of the name, where the first character is $.
+		const FString NewNameWithoutPrefix = NewName.RightChop(1);
+
 		if (const UEnum* EnumPtr = StaticEnum<EPCGPointProperties>())
 		{
-			int32 Index = EnumPtr->GetIndexByNameString(NewName.Mid(/*Start=*/1));
+			int32 Index = EnumPtr->GetIndexByNameString(NewNameWithoutPrefix);
 			if (Index != INDEX_NONE)
 			{
 				return SetPointProperty(static_cast<EPCGPointProperties>(EnumPtr->GetValueByIndex(Index)), /*bResetExtraNames=*/ false) || bExtraChanged;
 			}
 		}
+
+		if (const UEnum* EnumPtr = StaticEnum<EPCGExtraProperties>())
+		{
+			int32 Index = EnumPtr->GetIndexByNameString(NewNameWithoutPrefix);
+			if (Index != INDEX_NONE)
+			{
+				return SetExtraProperty(static_cast<EPCGExtraProperties>(EnumPtr->GetValueByIndex(Index)), /*bResetExtraNames=*/ false) || bExtraChanged;
+			}
+		}
 	}
 
 	return SetAttributeName(NewName.IsEmpty() ? NAME_None : FName(NewName), /*bResetExtraNames=*/ false) || bExtraChanged;
@@ -156,6 +200,11 @@ bool UPCGAttributePropertySelectorBlueprintHelpers::SetAttributeName(FPCGAttribu
 	return Selector.SetAttributeName(InAttributeName);
 }
 
+bool UPCGAttributePropertySelectorBlueprintHelpers::SetExtraProperty(FPCGAttributePropertySelector& Selector, EPCGExtraProperties InExtraProperty)
+{
+	return Selector.SetExtraProperty(InExtraProperty);
+}
+
 FName UPCGAttributePropertySelectorBlueprintHelpers::GetName(const FPCGAttributePropertySelector& Selector)
 {
 	return Selector.GetName();
diff --git a/Engine/Plugins/Experimental/PCG/Source/PCG/Public/Metadata/Accessors/PCGAttributeAccessorHelpers.h b/Engine/Plugins/Experimental/PCG/Source/PCG/Public/Metadata/Accessors/PCGAttributeAccessorHelpers.h
--- a/Engine/Plugins/Experimental/PCG/Source/PCG/Public/Metadata/Accessors/PCGAttributeAccessorHelpers.h
+++ b/Engine/Plugins/Experimental/PCG/Source/PCG/Public/Metadata/Accessors/PCGAttributeAccessorHelpers.h
@@ -4,6 +4,7 @@
 
 #include "Templates/UniquePtr.h"
 
+enum class EPCGExtraProperties : uint8;
 class IPCGAttributeAccessor;
 class IPCGAttributeAccessorKeyIterator;
 class IPCGAttributeAccessorKeys;
@@ -26,6 +27,8 @@ namespace PCGAttributeAccessorHelpers
 	PCG_API TUniquePtr<IPCGAttributeAccessor> CreatePropertyAccessor(const FName InPropertyName, const UClass* InClass);
 	PCG_API TUniquePtr<IPCGAttributeAccessor> CreatePropertyAccessor(const FName InPropertyName, const UStruct* InStruct);
 
+	PCG_API TUniquePtr<IPCGAttributeAccessor> CreateExtraAccessor(EPCGExtraProperties InExtraProperties);
+
 	PCG_API TUniquePtr<IPCGAttributeAccessor> CreateChainAccessor(TUniquePtr<IPCGAttributeAccessor> InAccessor, FName Name, bool& bOutSuccess);
 
 	/**
diff --git a/Engine/Plugins/Experimental/PCG/Source/PCG/Public/Metadata/Accessors/PCGCustomAccessor.h b/Engine/Plugins/Experimental/PCG/Source/PCG/Public/Metadata/Accessors/PCGCustomAccessor.h
--- a/Engine/Plugins/Experimental/PCG/Source/PCG/Public/Metadata/Accessors/PCGCustomAccessor.h
+++ b/Engine/Plugins/Experimental/PCG/Source/PCG/Public/Metadata/Accessors/PCGCustomAccessor.h
@@ -172,3 +172,40 @@ private:
 	ChainGetter Getter;
 	ChainSetter Setter;
 };
+
+/**
+* Very simple accessor that returns the index. Read only
+* Key supported: All
+*/
+class FPCGIndexAccessor : public IPCGAttributeAccessorT<FPCGIndexAccessor>
+{
+public:
+	using Type = int32;
+	using Super = IPCGAttributeAccessorT<FPCGIndexAccessor>;
+
+	FPCGIndexAccessor()
+		: Super(/*bInReadOnly=*/ true)
+	{}
+
+	bool GetRangeImpl(TArrayView<int32> OutValues, int32 Index, const IPCGAttributeAccessorKeys& InKeys) const
+	{
+		const int32 NumKeys = InKeys.GetNum();
+		int32 Counter = Index;
+
+		for (int32 i = 0; i < OutValues.Num(); ++i)
+		{
+			OutValues[i] = Counter++;
+			if (Counter >= NumKeys)
+			{
+				Counter = 0;
+			}
+		}
+
+		return true;
+	}
+
+	bool SetRangeImpl(TArrayView<const int32>, int32, IPCGAttributeAccessorKeys&, EPCGAttributeAccessorFlags)
+	{
+		return false;
+	}
+};
\ No newline at end of file
diff --git a/Engine/Plugins/Experimental/PCG/Source/PCG/Public/Metadata/PCGAttributePropertySelector.h b/Engine/Plugins/Experimental/PCG/Source/PCG/Public/Metadata/PCGAttributePropertySelector.h
--- a/Engine/Plugins/Experimental/PCG/Source/PCG/Public/Metadata/PCGAttributePropertySelector.h
+++ b/Engine/Plugins/Experimental/PCG/Source/PCG/Public/Metadata/PCGAttributePropertySelector.h
@@ -13,6 +13,13 @@ enum class EPCGAttributePropertySelection
 {
 	Attribute,
 	PointProperty,
+	ExtraProperty
+};
+
+UENUM()
+enum class EPCGExtraProperties : uint8
+{
+	Index
 };
 
 /**
@@ -35,6 +42,9 @@ public:
 	UPROPERTY(BlueprintReadOnly, VisibleAnywhere, Category = "Settings", meta = (EditCondition = "Selection == EPCGAttributePropertySelection::PointProperty", EditConditionHides))
 	EPCGPointProperties PointProperty = EPCGPointProperties::Position;
 
+	UPROPERTY(BlueprintReadOnly, VisibleAnywhere, Category = "Settings", meta = (EditCondition = "Selection == EPCGAttributePropertySelection::ExtraProperty", EditConditionHides))
+	EPCGExtraProperties ExtraProperty = EPCGExtraProperties::Index;
+
 	UPROPERTY()
 	TArray<FString> ExtraNames;
 
@@ -45,6 +55,7 @@ public:
 	// Setters, retrurn true if something changed.
 	bool SetPointProperty(EPCGPointProperties InPointProperty, bool bResetExtraNames = true);
 	bool SetAttributeName(FName InAttributeName, bool bResetExtraNames = true);
+	bool SetExtraProperty(EPCGExtraProperties InExtraProperty, bool bResetExtraNames = true);
 
 	FName GetName() const;
 
@@ -75,6 +86,9 @@ public:
 	UFUNCTION(BlueprintCallable, Category = "PCG|PCGAttributePropertySelector", meta = (ScriptMethod))
 	static bool SetAttributeName(UPARAM(ref) FPCGAttributePropertySelector& Selector, FName InAttributeName);
 
+	UFUNCTION(BlueprintCallable, Category = "PCG|PCGAttributePropertySelector", meta = (ScriptMethod))
+	static bool SetExtraProperty(UPARAM(ref) FPCGAttributePropertySelector& Selector, EPCGExtraProperties InExtraProperty);
+
 	UFUNCTION(BlueprintCallable, Category = "PCG|PCGAttributePropertySelector", meta = (ScriptMethod))
 	static FName GetName(UPARAM(const, ref) const FPCGAttributePropertySelector& Selector);
 };
\ No newline at end of file
diff --git a/Engine/Plugins/Experimental/PCG/Source/PCGEditor/Private/Details/PCGAttributePropertySelectorDetails.cpp b/Engine/Plugins/Experimental/PCG/Source/PCGEditor/Private/Details/PCGAttributePropertySelectorDetails.cpp
--- a/Engine/Plugins/Experimental/PCG/Source/PCGEditor/Private/Details/PCGAttributePropertySelectorDetails.cpp
+++ b/Engine/Plugins/Experimental/PCG/Source/PCGEditor/Private/Details/PCGAttributePropertySelectorDetails.cpp
@@ -16,6 +16,52 @@
 
 #define LOCTEXT_NAMESPACE "PCGAttributePropertySelectorDetails"
 
+namespace PCGAttributePropertySelectorDetails
+{
+	static const FText SetAttributePropertyNameTransaction = LOCTEXT("SetAttributePropertyName", "[PCG] Set Attribute/Property Name");
+
+	template <typename EnumType>
+	void BuildMenuSectionFromEnum(FMenuBuilder& InMenuBuilder, FPCGAttributePropertySelectorDetails* InDetailsObject, void(FPCGAttributePropertySelectorDetails::*InCallback)(EnumType))
+	{
+		if (const UEnum* EnumPtr = StaticEnum<EnumType>())
+		{
+			for (int32 i = 0; i < EnumPtr->NumEnums() - 1; ++i)
+			{
+				FString EnumName = EnumPtr->GetDisplayNameTextByIndex(i).ToString();
+				EnumType EnumValue = (EnumType)EnumPtr->GetValueByIndex(i);
+
+				InMenuBuilder.AddMenuEntry(
+					FText::FromString(EnumName),
+					FText(),
+					FSlateIcon(),
+					FExecuteAction::CreateSP(InDetailsObject, InCallback, EnumValue),
+					NAME_None,
+					EUserInterfaceActionType::Button
+				);
+			}
+		}
+	}
+
+	// Func signature: bool(FPCGAttributePropertySelector*, T)
+	template <typename T, typename Func>
+	void SetValue(FPCGAttributePropertySelector* InSelector, TSharedPtr<IPropertyHandle>& InPropertyHandle, T InValue, Func InCallback)
+	{
+		if (InSelector && InCallback)
+		{
+			FScopedTransaction Transaction(SetAttributePropertyNameTransaction);
+
+			InPropertyHandle->NotifyPreChange();
+			if (!InCallback(InSelector, std::move(InValue)))
+			{
+				Transaction.Cancel();
+				return;
+			}
+
+			InPropertyHandle->NotifyPostChange(EPropertyChangeType::ValueSet);
+		}
+	}
+}
+
 void FPCGAttributePropertySelectorDetails::CustomizeHeader(TSharedRef<IPropertyHandle> InPropertyHandle, FDetailWidgetRow& HeaderRow, IPropertyTypeCustomizationUtils& CustomizationUtils)
 {
 	PropertyHandle = InPropertyHandle;
@@ -86,23 +132,13 @@ TSharedRef<SWidget> FPCGAttributePropertySelectorDetails::GenerateExtraMenu()
 
 	MenuBuilder.BeginSection("PointProperties", LOCTEXT("PointPropertiesHeader", "Point Properties"));
 	{
-		if (const UEnum* EnumPtr = StaticEnum<EPCGPointProperties>())
-		{
-			for (int32 i = 0; i < EnumPtr->NumEnums() - 1; ++i)
-			{
-				FString EnumName = EnumPtr->GetNameStringByIndex(i);
-				EPCGPointProperties EnumValue = (EPCGPointProperties)EnumPtr->GetValueByIndex(i);
+		PCGAttributePropertySelectorDetails::BuildMenuSectionFromEnum<EPCGPointProperties>(MenuBuilder, this, &FPCGAttributePropertySelectorDetails::SetPointProperty);
+	}
+	MenuBuilder.EndSection();
 
-				MenuBuilder.AddMenuEntry(
-					FText::FromString(EnumName),
-					FText(),
-					FSlateIcon(),
-					FExecuteAction::CreateSP(this, &FPCGAttributePropertySelectorDetails::SetPointProperty, EnumValue),
-					NAME_None,
-					EUserInterfaceActionType::Button
-				);
-			}
-		}
+	MenuBuilder.BeginSection("OtherProperties", LOCTEXT("OtherPropertiesHeader", "Other Properties"));
+	{
+		PCGAttributePropertySelectorDetails::BuildMenuSectionFromEnum<EPCGExtraProperties>(MenuBuilder, this, &FPCGAttributePropertySelectorDetails::SetExtraProperty);
 	}
 	MenuBuilder.EndSection();
 
@@ -139,53 +175,22 @@ void FPCGAttributePropertySelectorDetails::SetText(const FText& NewText, ETextCo
 		return;
 	}
 
-	if (FPCGAttributePropertySelector* Selector = GetStruct())
-	{
-		FScopedTransaction Transaction(LOCTEXT("SetAttributePropertyName", "[PCG] Set Attribute/Property Name"));
-
-		PropertyHandle->NotifyPreChange();
-		if (!Selector->Update(NewText.ToString()))
-		{
-			Transaction.Cancel();
-			return;
-		}
-
-		PropertyHandle->NotifyPostChange(EPropertyChangeType::ValueSet);
-	}
+	PCGAttributePropertySelectorDetails::SetValue(GetStruct(), PropertyHandle, NewText.ToString(), [](FPCGAttributePropertySelector* InStruct, FString InValue) -> bool { return InStruct && InStruct->Update(std::move(InValue)); });
 }
 
 void FPCGAttributePropertySelectorDetails::SetPointProperty(EPCGPointProperties EnumValue)
 {
-	if (FPCGAttributePropertySelector* Selector = GetStruct())
-	{
-		FScopedTransaction Transaction(LOCTEXT("SetAttributePropertyName", "[PCG] Set Attribute/Property Name"));
-
-		PropertyHandle->NotifyPreChange();
-		if (!Selector->SetPointProperty(EnumValue))
-		{
-			Transaction.Cancel();
-			return;
-		}
-
-		PropertyHandle->NotifyPostChange(EPropertyChangeType::ValueSet);
-	}
+	PCGAttributePropertySelectorDetails::SetValue(GetStruct(), PropertyHandle, EnumValue, [](FPCGAttributePropertySelector* InStruct, EPCGPointProperties InValue) -> bool { return InStruct && InStruct->SetPointProperty(InValue); });
 }
 
 void FPCGAttributePropertySelectorDetails::SetAttributeName(FName NewName)
 {
-	if (FPCGAttributePropertySelector* Selector = GetStruct())
-	{
-		FScopedTransaction Transaction(LOCTEXT("SetAttributePropertyName", "[PCG] Set Attribute/Property Name"));
-
-		PropertyHandle->NotifyPreChange();
-		if (!Selector->SetAttributeName(NewName))
-		{
-			Transaction.Cancel();
-			return;
-		}
+	PCGAttributePropertySelectorDetails::SetValue(GetStruct(), PropertyHandle, NewName, [](FPCGAttributePropertySelector* InStruct, FName InValue) -> bool { return InStruct && InStruct->SetAttributeName(InValue); });
+}
 
-		PropertyHandle->NotifyPostChange(EPropertyChangeType::ValueSet);
-	}
+void FPCGAttributePropertySelectorDetails::SetExtraProperty(EPCGExtraProperties EnumValue)
+{
+	PCGAttributePropertySelectorDetails::SetValue(GetStruct(), PropertyHandle, EnumValue, [](FPCGAttributePropertySelector* InStruct, EPCGExtraProperties InValue) -> bool { return InStruct && InStruct->SetExtraProperty(InValue); });
 }
 
 #undef LOCTEXT_NAMESPACE
diff --git a/Engine/Plugins/Experimental/PCG/Source/PCGEditor/Public/Details/PCGAttributePropertySelectorDetails.h b/Engine/Plugins/Experimental/PCG/Source/PCGEditor/Public/Details/PCGAttributePropertySelectorDetails.h
--- a/Engine/Plugins/Experimental/PCG/Source/PCGEditor/Public/Details/PCGAttributePropertySelectorDetails.h
+++ b/Engine/Plugins/Experimental/PCG/Source/PCGEditor/Public/Details/PCGAttributePropertySelectorDetails.h
@@ -6,6 +6,7 @@
 
 
 enum class EPCGPointProperties : uint8;
+enum class EPCGExtraProperties : uint8;
 namespace ETextCommit { enum Type : int; }
 
 struct FPCGAttributePropertySelector;
@@ -34,6 +35,7 @@ protected:
 	void SetText(const FText& NewText, ETextCommit::Type CommitInfo);
 	void SetPointProperty(EPCGPointProperties EnumValue);
 	void SetAttributeName(FName NewName);
+	void SetExtraProperty(EPCGExtraProperties EnumValue);
 
 	TSharedPtr<IPropertyHandle> PropertyHandle;
 };
