[PCG] Tentative fix for VS2019 build fail

#rnx

diff --git a/Engine/Plugins/Experimental/PCG/Source/PCGEditor/Private/Details/PCGAttributePropertySelectorDetails.cpp b/Engine/Plugins/Experimental/PCG/Source/PCGEditor/Private/Details/PCGAttributePropertySelectorDetails.cpp
--- a/Engine/Plugins/Experimental/PCG/Source/PCGEditor/Private/Details/PCGAttributePropertySelectorDetails.cpp
+++ b/Engine/Plugins/Experimental/PCG/Source/PCGEditor/Private/Details/PCGAttributePropertySelectorDetails.cpp
@@ -46,7 +46,7 @@ namespace PCGAttributePropertySelectorDetails
 	template <typename T, typename Func>
 	void SetValue(FPCGAttributePropertySelector* InSelector, TSharedPtr<IPropertyHandle>& InPropertyHandle, T InValue, Func InCallback)
 	{
-		if (InSelector && InCallback)
+		if (InSelector)
 		{
 			FScopedTransaction Transaction(SetAttributePropertyNameTransaction);
 
