Fixed outstanding bug in the metadata partition when using non-compressed values

diff --git a/Engine/Plugins/Experimental/PCG/Source/PCG/Private/Elements/Metadata/PCGMetadataPartition.cpp b/Engine/Plugins/Experimental/PCG/Source/PCG/Private/Elements/Metadata/PCGMetadataPartition.cpp
--- a/Engine/Plugins/Experimental/PCG/Source/PCG/Private/Elements/Metadata/PCGMetadataPartition.cpp
+++ b/Engine/Plugins/Experimental/PCG/Source/PCG/Private/Elements/Metadata/PCGMetadataPartition.cpp
@@ -80,14 +80,15 @@ bool FPCGMetadataPartitionElement::ExecuteInternal(FPCGContext* Context) const
 			{
 				ValueKeyMapping.Add(-1);
 			}
-			else if (bUsesValueKeys)
+			else if (!bUsesValueKeys)
 			{
 				PCGMetadataValueKey* MatchingVK = Algo::FindByPredicate(ValueKeyMapping, [ValueKey, AttributeBase](const PCGMetadataValueKey& Key)
 				{
 					return AttributeBase->AreValuesEqual(ValueKey, Key);
 				});
 
-				ValueKeyMapping.Add(MatchingVK ? *MatchingVK : ValueKey);
+				const PCGMetadataValueKey TentativeKey = MatchingVK ? *MatchingVK : ValueKey;
+				ValueKeyMapping.Add(TentativeKey);
 			}
 			else
 			{
