[PCG] Check SourceComponent valid as we could crash it by spamming IsPartitioned and Undo/Redo

#rnx

diff --git a/Engine/Plugins/Experimental/PCG/Source/PCG/Private/PCGContext.cpp b/Engine/Plugins/Experimental/PCG/Source/PCG/Private/PCGContext.cpp
--- a/Engine/Plugins/Experimental/PCG/Source/PCG/Private/PCGContext.cpp
+++ b/Engine/Plugins/Experimental/PCG/Source/PCG/Private/PCGContext.cpp
@@ -229,6 +229,11 @@ void FPCGContext::OverrideSettings()
 #if WITH_EDITOR
 void FPCGContext::LogVisual(ELogVerbosity::Type InVerbosity, const FText& InMessage) const
 {
+	if (!SourceComponent.IsValid())
+	{
+		return;
+	}
+
 	if (UPCGSubsystem* Subsystem = UPCGSubsystem::GetInstance(SourceComponent->GetWorld()))
 	{
 		Subsystem->GetNodeVisualLogsMutable().Log(Node, SourceComponent, InVerbosity, InMessage);
@@ -237,6 +242,11 @@ void FPCGContext::LogVisual(ELogVerbosity::Type InVerbosity, const FText& InMess
 
 bool FPCGContext::HasVisualLogs() const
 {
+	if (!SourceComponent.IsValid())
+	{
+		return false;
+	}
+
 	if (UPCGSubsystem* Subsystem = UPCGSubsystem::GetInstance(SourceComponent->GetWorld()))
 	{
 		return Subsystem->GetNodeVisualLogs().HasLogs(Node, SourceComponent.Get());
