[PCG] Hierarchical generation option added to graph (default off).

Adds "Hi Gen Grid Size" nodes that allow prescribing what grid size the downstream nodes should execute on. Linkage tasks are inserted on any edge that crosses a grid size boundary. Linkage will either store data in a map if executing on the "source" grid level, or retrieve data from the map if executing on the "sink" grid level. If source data not found, will try once to schedule the source graph.

Editor features - if debug object set in graph editor, nodes that will execute on a finer grid are drawn semitransparent/disabled, and nodes that will execute on a coarser grid are highlighted to indicate their data is available for use.

Added option to supply context allocator to generic task, allowing creating generic tasks with execution state.

Future steps:

* Auto-intersect data pulled from larger grids with the smaller grid domain when there is risk of mixing together data from multiple levels and repeating work unnecessarily
* Execute once / unbounded - have an execution level above grids, that executes on the original component

diff --git a/Engine/Plugins/Experimental/PCG/Source/PCG/Private/Elements/PCGHiGenGridSize.cpp b/Engine/Plugins/Experimental/PCG/Source/PCG/Private/Elements/PCGHiGenGridSize.cpp
new file mode 100644
--- /dev/null
+++ b/Engine/Plugins/Experimental/PCG/Source/PCG/Private/Elements/PCGHiGenGridSize.cpp
@@ -0,0 +1,117 @@
+// Copyright Epic Games, Inc. All Rights Reserved.
+
+#include "Elements/PCGHiGenGridSize.h"
+
+#include "PCGComponent.h"
+#include "PCGContext.h"
+#include "PCGGraph.h"
+#include "PCGModule.h"
+#include "PCGPin.h"
+
+#define LOCTEXT_NAMESPACE "PCGHiGenGridSizeElement"
+
+namespace PCGHiGenGridSizeConstants
+{
+	const FName NodeName = TEXT("HiGenGridSize");
+	const FText NodeTitle = LOCTEXT("NodeTitle", "Grid Size");
+}
+
+#if WITH_EDITOR
+FName UPCGHiGenGridSizeSettings::GetDefaultNodeName() const
+{
+	return PCGHiGenGridSizeConstants::NodeName;
+}
+
+FText UPCGHiGenGridSizeSettings::GetDefaultNodeTitle() const
+{
+	return PCGHiGenGridSizeConstants::NodeTitle;
+}
+
+FText UPCGHiGenGridSizeSettings::GetNodeTooltipText() const
+{
+	return LOCTEXT("NodeTooltip", "Set the execution grid size for downstream nodes. Enables executing a single graph across a hierarchy of grids. Has no effect if generating component is not partitioned.");
+}
+#endif
+
+EPCGDataType UPCGHiGenGridSizeSettings::GetCurrentPinTypes(const UPCGPin* InPin) const
+{
+	check(InPin);
+	if (!InPin->IsOutputPin())
+	{
+		return Super::GetCurrentPinTypes(InPin);
+	}
+
+	// Output pin narrows to union of inputs on first pin
+	const EPCGDataType InputTypeUnion = GetTypeUnionOfIncidentEdges(PCGPinConstants::DefaultInputLabel);
+	return InputTypeUnion != EPCGDataType::None ? InputTypeUnion : EPCGDataType::Any;
+}
+
+TArray<FPCGPinProperties> UPCGHiGenGridSizeSettings::InputPinProperties() const
+{
+	TArray<FPCGPinProperties> PinProperties;
+	PinProperties.Emplace(PCGPinConstants::DefaultInputLabel, EPCGDataType::Any);
+
+	return PinProperties;
+}
+
+TArray<FPCGPinProperties> UPCGHiGenGridSizeSettings::OutputPinProperties() const
+{
+	TArray<FPCGPinProperties> PinProperties;
+	PinProperties.Emplace(PCGPinConstants::DefaultOutputLabel, EPCGDataType::Any);
+
+	return PinProperties;
+}
+
+FPCGElementPtr UPCGHiGenGridSizeSettings::CreateElement() const
+{
+	return MakeShared<FPCGHiGenGridSizeElement>();
+}
+
+FName UPCGHiGenGridSizeSettings::AdditionalTaskName() const
+{
+	if (HasAnyFlags(RF_ClassDefaultObject))
+	{
+		return PCGHiGenGridSizeConstants::NodeName;
+	}
+
+	const UEnum* EnumPtr = StaticEnum<EPCGHiGenGrid>();
+	if (ensure(EnumPtr))
+	{
+		FText GridSizeDisplayName = EnumPtr->GetDisplayNameTextByValue(static_cast<int64>(HiGenGridSize));
+		return FName(FString::Printf(TEXT("%s: %s"), *PCGHiGenGridSizeConstants::NodeTitle.ToString(), *GridSizeDisplayName.ToString()));
+	}
+	else
+	{
+		return FName(FString::Printf(TEXT("%s: %d"), *PCGHiGenGridSizeConstants::NodeTitle.ToString(), static_cast<int32>(HiGenGridSize)));
+	}
+}
+
+#if WITH_EDITOR
+bool UPCGHiGenGridSizeSettings::IsStructuralProperty(const FName& InPropertyName) const
+{
+	// Grid sizes are processed during graph compilation and is part of the graph structure.
+	return InPropertyName == GET_MEMBER_NAME_CHECKED(UPCGHiGenGridSizeSettings, bEnabled)
+		|| InPropertyName == GET_MEMBER_NAME_CHECKED(UPCGHiGenGridSizeSettings, HiGenGridSize);
+}
+#endif
+
+bool FPCGHiGenGridSizeElement::ExecuteInternal(FPCGContext* Context) const
+{
+	// Trivial pass through. Will only execute on the prescribed grid.
+	Context->OutputData = Context->InputData;
+
+	if (Context->Node && Context->Node->GetGraph() && !Context->Node->GetGraph()->IsHierarchicalGenerationEnabled())
+	{
+		PCGE_LOG(Warning, GraphAndLog, LOCTEXT("GridSizeUsedInNonHiGenGraph", "Grid Size node used in a non-hierarchical graph. Enable hierarchical generation in the graph settings or remove this node."));
+	}
+	else if (!Context->SourceComponent->IsPartitioned() && !Context->SourceComponent->IsLocalComponent())
+	{
+		// Warning if component is not partitioned (and not local component) as this node will otherwise be silently ignored.
+		// Also serves as a hint if the user forgot to enable higen for this graph.
+		PCGE_LOG(Warning, GraphAndLog, LOCTEXT("NonPartitionedComponent", "Grid Size node used on a non-partitioned component and will have no effect. Is Partitioned must be enabled on the component."));
+	}
+
+	return true;
+}
+
+#undef LOCTEXT_NAMESPACE
diff --git a/Engine/Plugins/Experimental/PCG/Source/PCG/Private/Graph/PCGGraphCompiler.cpp b/Engine/Plugins/Experimental/PCG/Source/PCG/Private/Graph/PCGGraphCompiler.cpp
--- a/Engine/Plugins/Experimental/PCG/Source/PCG/Private/Graph/PCGGraphCompiler.cpp
+++ b/Engine/Plugins/Experimental/PCG/Source/PCG/Private/Graph/PCGGraphCompiler.cpp
@@ -2,11 +2,12 @@
 
 #include "PCGGraphCompiler.h"
 
-#include "PCGGraph.h"
 #include "PCGEdge.h"
+#include "PCGGraph.h"
 #include "PCGModule.h"
 #include "PCGPin.h"
 #include "PCGSubgraph.h"
+#include "Elements/PCGHiGenGridSize.h"
 #include "Graph/PCGGraphExecutor.h"
 
 #include "Misc/ScopeRWLock.h"
@@ -292,6 +293,117 @@ TArray<FPCGGraphTask> FPCGGraphCompiler::GetPrecompiledTasks(UPCGGraph* InGraph,
 	return ExistingTasks ? *ExistingTasks : TArray<FPCGGraphTask>();
 }
 
+void FPCGGraphCompiler::ResolveGridSizes(TArray<FPCGGraphTask>& InOutCompiledTasks, const FPCGStackContext& InStackContext) const
+{
+	// The stack is used to form the ResourceKey - a string that provides a path to the data from top graph down to specific pin.
+	// This will be used by link tasks as store/retrieve keys to marshal data for edges that cross grid size boundaries.
+	const FPCGStack* CurrentStack = InStackContext.GetStack(InStackContext.GetCurrentStackIndex());
+	if (InOutCompiledTasks.Num() == 0 || !ensure(CurrentStack))
+	{
+		return;
+	}
+
+	// Calculate execution grid values for each task.
+	for (FPCGGraphTask& Task : InOutCompiledTasks)
+	{
+		CalculateGridRecursive(Task.NodeId, InStackContext, InOutCompiledTasks);
+	}
+
+	// Now add link tasks - if a Grid256 task depends on data from a Grid512 task, inject a link
+	// task that looks up the Grid512 component, schedules its execution if it does not have data, and
+	// then uses its output data.
+	const int32 NumCompiledTasksBefore = InOutCompiledTasks.Num();
+	for (FPCGTaskId TaskId = 0; TaskId < NumCompiledTasksBefore; ++TaskId)
+	{
+		const EPCGHiGenGrid GraphGenerationGrid = InOutCompiledTasks[TaskId].GraphGenerationGrid;
+		for (FPCGGraphTaskInput& TaskInput : InOutCompiledTasks[TaskId].Inputs)
+		{
+			if (!TaskInput.InPin)
+			{
+				// Don't link if we don't have a upstream pin to retrieve data from
+				continue;
+			}
+
+			const EPCGHiGenGrid InputGraphGenerationGrid = InOutCompiledTasks[TaskInput.TaskId].GraphGenerationGrid;
+			// Register linkage task if grid sizes don't match - either way! This allows us to generate an execution-time error if going from
+			// small grid to large grid.
+			if (InputGraphGenerationGrid != EPCGHiGenGrid::Uninitialized && InputGraphGenerationGrid > GraphGenerationGrid)
+			{
+				// Build a string identifier for the data
+				FString ResourceKey;
+				if (!ensure(CurrentStack->CreateStackFramePath(ResourceKey, TaskInput.InPin->Node, TaskInput.InPin)))
+				{
+					continue;
+				}
+
+				// Build task & element to hold the operation to perform
+				FPCGGraphTask& LinkTask = InOutCompiledTasks.Emplace_GetRef();
+				LinkTask.NodeId = InOutCompiledTasks.Num() - 1;
+				LinkTask.StackIndex = InOutCompiledTasks[TaskId].StackIndex;
+
+				LinkTask.Inputs.Emplace(TaskInput.TaskId, TaskInput.InPin, nullptr, /*bConsumeInputData=*/true);
+
+				const EPCGHiGenGrid FromGrid = InOutCompiledTasks[TaskInput.TaskId].GraphGenerationGrid;
+				const EPCGHiGenGrid ToGrid = InOutCompiledTasks[TaskId].GraphGenerationGrid;
+
+				// This lambda runs at execution time and attempts to retrieve the data from a larger grid. Capture by value is intentional.
+				auto GridLinkageOperation = [FromGrid, ToGrid, ResourceKey, OutputPinLabel = TaskInput.InPin->Properties.Label,
+					DownstreamNode = InOutCompiledTasks[TaskId].Node](FPCGContext* InContext)
+				{
+					return PCGGraphExecutor::ExecuteGridLinkage(FromGrid, ToGrid, ResourceKey, OutputPinLabel, DownstreamNode, static_cast<FPCGGridLinkageContext*>(InContext));
+				};
+				FPCGGenericElement::FContextAllocator ContextAllocator = [](const FPCGDataCollection&, TWeakObjectPtr<UPCGComponent>, const UPCGNode*)
+				{
+					return new FPCGGridLinkageContext();
+				};
+				LinkTask.Element = MakeShared<FPCGGenericElement>(GridLinkageOperation, ContextAllocator);
+
+				// Now splice in the new task - redirect the downstream task to grab its input from the link task.
+				TaskInput.TaskId = LinkTask.NodeId;
+
+				// The link needs to execute at both FROM grid size (store) and TO grid size (retrieve).
+				LinkTask.GraphGenerationGrid = FromGrid | ToGrid;
+			}
+		}
+	}
+}
+
+EPCGHiGenGrid FPCGGraphCompiler::CalculateGridRecursive(FPCGTaskId InTaskId, const FPCGStackContext& InStackContext, TArray<FPCGGraphTask>& InOutCompiledTasks) const
+{
+	if (InOutCompiledTasks[InTaskId].GraphGenerationGrid != EPCGHiGenGrid::Uninitialized)
+	{
+		return InOutCompiledTasks[InTaskId].GraphGenerationGrid;
+	}
+
+	// GenerationDefault as default - means we don't know at compilation-time. Any tasks with this unknown value
+	// will receive the graph default grid size at schedule-time.
+	EPCGHiGenGrid Grid = EPCGHiGenGrid::GenerationDefault;
+
+	const UPCGNode* Node = InOutCompiledTasks[InTaskId].Node;
+	const UPCGSettings* Settings = Node ? Node->GetSettings() : nullptr;
+	const UPCGHiGenGridSizeSettings* Gate = Cast<UPCGHiGenGridSizeSettings>(Settings);
+	if (Gate && Gate->bEnabled)
+	{
+		Grid = Gate->GetGrid();
+	}
+	else
+	{
+		// Grid of this task is minimum of all input grids. We can link in data from a larger grid, but not from a finer grid (this goes against hierarchy).
+		for (FPCGGraphTaskInput InputTask : InOutCompiledTasks[InTaskId].Inputs)
+		{
+			const EPCGHiGenGrid InputGrid = CalculateGridRecursive(InputTask.TaskId, InStackContext, InOutCompiledTasks);
+			if (PCGHiGenGrid::IsValidGrid(InputGrid))
+			{
+				Grid = FMath::Min(InputGrid, Grid);
+			}
+		}
+	}
+
+	InOutCompiledTasks[InTaskId].GraphGenerationGrid = Grid;
+
+	return Grid;
+}
+
 TArray<FPCGGraphTask> FPCGGraphCompiler::GetCompiledTasks(UPCGGraph* InGraph, FPCGStackContext& OutStackContext, bool bIsTopGraph)
 {
 	TArray<FPCGGraphTask> CompiledTasks;
@@ -371,6 +483,13 @@ void FPCGGraphCompiler::CompileTopGraph(UPCGGraph* InGraph)
 		return;
 	}
 
+	if (InGraph->IsHierarchicalGenerationEnabled())
+	{
+		// Operates on compiled tasks before pre/post task added, as it may add tasks for linkages
+		// and the pre/post task need to reside at the end of the compiled task array.
+		ResolveGridSizes(CompiledTasks, StackContext);
+	}
+
 	const int TaskNum = CompiledTasks.Num();
 	const FPCGTaskId PreExecuteTaskId = FPCGTaskId(TaskNum);
 	const FPCGTaskId PostExecuteTaskId = PreExecuteTaskId + 1;
diff --git a/Engine/Plugins/Experimental/PCG/Source/PCG/Private/Graph/PCGGraphCompiler.h b/Engine/Plugins/Experimental/PCG/Source/PCG/Private/Graph/PCGGraphCompiler.h
--- a/Engine/Plugins/Experimental/PCG/Source/PCG/Private/Graph/PCGGraphCompiler.h
+++ b/Engine/Plugins/Experimental/PCG/Source/PCG/Private/Graph/PCGGraphCompiler.h
@@ -21,6 +21,9 @@ public:
 
 	static void OffsetNodeIds(TArray<FPCGGraphTask>& Tasks, FPCGTaskId Offset, FPCGTaskId ParentId);
 
+	/** Propagates grid sizes through a graph's compiled tasks. */
+	void ResolveGridSizes(TArray<FPCGGraphTask>& InOutCompiledTasks, const FPCGStackContext& InStackContext) const;
+
 #if WITH_EDITOR
 	void NotifyGraphChanged(UPCGGraph* InGraph);
 #endif
@@ -29,6 +32,9 @@ private:
 	TArray<FPCGGraphTask> CompileGraph(UPCGGraph* InGraph, FPCGTaskId& NextId, FPCGStackContext& InOutStackContext);
 	void CompileTopGraph(UPCGGraph* InGraph);
 
+	/** Returns the grid for InCompiledTaskId which will be a concrete grid size or GenerationDefault if it cannot be determined statically. */
+	EPCGHiGenGrid CalculateGridRecursive(FPCGTaskId InTaskId, const FPCGStackContext& InStackContext, TArray<FPCGGraphTask>& InOutCompiledTasks) const;
+
 	mutable FRWLock GraphToTaskMapLock;
 	TMap<UPCGGraph*, TArray<FPCGGraphTask>> GraphToTaskMap;
 	TMap<UPCGGraph*, FPCGStackContext> GraphToStackContext;
diff --git a/Engine/Plugins/Experimental/PCG/Source/PCG/Private/Graph/PCGGraphExecutor.cpp b/Engine/Plugins/Experimental/PCG/Source/PCG/Private/Graph/PCGGraphExecutor.cpp
--- a/Engine/Plugins/Experimental/PCG/Source/PCG/Private/Graph/PCGGraphExecutor.cpp
+++ b/Engine/Plugins/Experimental/PCG/Source/PCG/Private/Graph/PCGGraphExecutor.cpp
@@ -2,6 +2,7 @@
 
 #include "PCGGraphExecutor.h"
 
+#include "PCGCommon.h"
 #include "PCGComponent.h"
 #include "PCGContext.h"
 #include "PCGCrc.h"
@@ -12,10 +13,13 @@
 #include "PCGParamData.h"
 #include "PCGPin.h"
 #include "PCGSubsystem.h"
+#include "PCGWorldActor.h"
 #include "Graph/PCGGraphCache.h"
 #include "Graph/PCGStackContext.h"
+#include "Grid/PCGPartitionActor.h"
 #include "Helpers/PCGHelpers.h"
 #include "Metadata/PCGMetadata.h"
+#include "Utils/PCGGraphExecutionLogging.h"
 
 #include "Algo/AnyOf.h"
 #include "Async/Async.h"
@@ -100,8 +104,9 @@ FPCGTaskId FPCGGraphExecutor::Schedule(UPCGGraph* Graph, UPCGComponent* SourceCo
 		UE_LOG(LogPCG, Log, TEXT("[%s] --- SCHEDULE GRAPH ---"), *SourceComponent->GetOwner()->GetName());
 	}
 
+	UPCGSubsystem* Subsystem = SourceComponent ? UPCGSubsystem::GetInstance(SourceComponent->GetWorld()) : nullptr;
 #if WITH_EDITOR
-	if (UPCGSubsystem* Subsystem = SourceComponent ? UPCGSubsystem::GetInstance(SourceComponent->GetWorld()) : nullptr)
+	if (Subsystem)
 	{
 		for (const UPCGNode* Node : Graph->GetNodes())
 		{
@@ -130,6 +135,40 @@ FPCGTaskId FPCGGraphExecutor::Schedule(UPCGGraph* Graph, UPCGComponent* SourceCo
 		Task.StackContext = StackContextPtr;
 	}
 
+	// Finalize grid sizes on tasks if graph has hierarchical generation enabled. Some tasks do not have a concrete grid size
+	// assigned at compile time, such as nodes outside of any authored grid size range. Also if we are generating on a non-partitioned
+	// component, knock out the grid sizes.
+	if (SourceComponent && SourceComponent->GetGraph() && SourceComponent->GetGraph()->IsHierarchicalGenerationEnabled())
+	{
+		const int32 ComponentGridSize = SourceComponent->GetGenerationGridSize();
+		const int32 DefaultGridSize = SourceComponent->GetGraph()->GetDefaultGridSize();
+		const EPCGHiGenGrid DefaultGrid = PCGHiGenGrid::GridSizeToGrid(DefaultGridSize);
+		const EPCGHiGenGrid Grid = PCGHiGenGrid::GridSizeToGrid(ensure(PCGHiGenGrid::IsValidGridSize(ComponentGridSize)) ? ComponentGridSize : DefaultGridSize);
+		const bool bNonPartitionedComponent = !SourceComponent->IsLocalComponent() && !SourceComponent->IsPartitioned();
+
+		for (FPCGGraphTask& Task : CompiledTasks)
+		{
+			// Make a copy of this because many things can be scheduled over many frames
+			Task.GenerationGrid = Grid;
+
+			// Set graph generation grid size for this task
+			if (bNonPartitionedComponent)
+			{
+				// Non-partitioned component - throw away grid from all tasks which will force them all to run
+				Task.GraphGenerationGrid = EPCGHiGenGrid::Uninitialized;
+			}
+			else if (!!(Task.GraphGenerationGrid & EPCGHiGenGrid::GenerationDefault))
+			{
+				// Remove the generation default flag (but leave any other grid sizes that might have been set - linkage tasks
+				// will store the To grid size in the lower bits).
+				Task.GraphGenerationGrid &= ~EPCGHiGenGrid::GenerationDefault;
+
+				// Add default grid size
+				Task.GraphGenerationGrid |= DefaultGrid;
+			}
+		}
+	}
+
 	// Prepare scheduled task that will be promoted in the next Execute call.
 	if (CompiledTasks.Num() > 0)
 	{ 
@@ -376,6 +415,7 @@ FPCGTaskId FPCGGraphExecutor::ScheduleGenericWithContext(TFunction<bool(FPCGCont
 	FPCGGraphTask Task;
 	for (FPCGTaskId TaskDependency : TaskDependencies)
 	{
+		ensure(TaskDependency != InvalidPCGTaskId);
 		Task.Inputs.Emplace(TaskDependency, nullptr, nullptr, bConsumeInputData);
 	}
 
@@ -595,31 +635,52 @@ void FPCGGraphExecutor::Execute()
 					continue;
 				}
 
-				// If a task is cacheable and has been cached, then we don't need to create an active task for it unless
-				// there is an execution mode that would prevent us from doing so.
-				const UPCGSettingsInterface* TaskSettingsInterface = TaskInput.GetSettingsInterface(Task.Node ? Task.Node->GetSettingsInterface() : nullptr);
-				const UPCGSettings* TaskSettings = TaskSettingsInterface ? TaskSettingsInterface->GetSettings() : nullptr;
-				const bool bCacheable = Task.Element->IsCacheableInstance(TaskSettingsInterface);
-
-				// Calculate Crc of dependencies (input data Crcs, settings) and use this as the key in the cache lookup
-				FPCGCrc DependenciesCrc;
-				if (TaskSettings && bCacheable)
+				bool bNeedsToCreateActiveTask = true;
+				if (Task.GraphGenerationGrid != EPCGHiGenGrid::Uninitialized)
 				{
-					Task.Element->GetDependenciesCrc(TaskInput, TaskSettings, Task.SourceComponent.Get(), DependenciesCrc);
-				}
+					// Graph generation sizes should be resolved by now
+					ensure(PCGHiGenGrid::IsValidGrid(Task.GraphGenerationGrid) || Task.GraphGenerationGrid == EPCGHiGenGrid::Unbounded);
 
-				if (bGraphCacheDebuggingEnabled && !bCacheable && Task.SourceComponent.Get() && Task.Node)
-				{
-					UE_LOG(LogPCG, Warning, TEXT("[%s] %s\t\tCACHING DISABLED"), *Task.SourceComponent->GetOwner()->GetName(), *Task.Node->GetNodeTitle().ToString());
+					const bool bGraphHasNoGridSizes = Task.GraphGenerationGrid == EPCGHiGenGrid::Unbounded;
+					const bool bAtFromOrToGridSize = !!(Task.GenerationGrid & Task.GraphGenerationGrid);
+					const bool bInRange = !bGraphHasNoGridSizes && PCGHiGenGrid::IsValidGrid(Task.GenerationGrid) && bAtFromOrToGridSize;
+					bNeedsToCreateActiveTask = bGraphHasNoGridSizes || bInRange;
+
+					if (bGraphCacheDebuggingEnabled && !bNeedsToCreateActiveTask && Task.SourceComponent.Get() && Task.SourceComponent->GetOwner() && Task.Node)
+					{
+						UE_LOG(LogPCG, Warning, TEXT("[%s] %s\t\tGRID SIZE %d SKIPPED"), *Task.SourceComponent->GetOwner()->GetName(), *Task.Node->GetNodeTitle().ToString(), PCGHiGenGrid::GridToGridSize(Task.GraphGenerationGrid));
+					}
 				}
 
+				FPCGCrc DependenciesCrc;
 				FPCGDataCollection CachedOutput;
-				const bool bResultAlreadyInCache = bCacheable && DependenciesCrc.IsValid() && GraphCache.GetFromCache(Task.Node, Task.Element.Get(), DependenciesCrc, TaskInput, TaskSettings, Task.SourceComponent.Get(), CachedOutput);
+				bool bResultAlreadyInCache = false;
+				if (bNeedsToCreateActiveTask)
+				{
+					// If a task is cacheable and has been cached, then we don't need to create an active task for it unless
+					// there is an execution mode that would prevent us from doing so.
+					const UPCGSettingsInterface* TaskSettingsInterface = TaskInput.GetSettingsInterface(Task.Node ? Task.Node->GetSettingsInterface() : nullptr);
+					const UPCGSettings* TaskSettings = TaskSettingsInterface ? TaskSettingsInterface->GetSettings() : nullptr;
+					const bool bCacheable = Task.Element->IsCacheableInstance(TaskSettingsInterface);
+
+					// Calculate Crc of dependencies (input data Crcs, settings) and use this as the key in the cache lookup
+					if (TaskSettings && bCacheable)
+					{
+						Task.Element->GetDependenciesCrc(TaskInput, TaskSettings, Task.SourceComponent.Get(), DependenciesCrc);
+					}
+
+					if (bGraphCacheDebuggingEnabled && !bCacheable && Task.SourceComponent.Get() && Task.Node)
+					{
+						UE_LOG(LogPCG, Warning, TEXT("[%s] %s\t\tCACHING DISABLED"), *Task.SourceComponent->GetOwner()->GetName(), *Task.Node->GetNodeTitle().ToString());
+					}
+
+					bResultAlreadyInCache = bCacheable && DependenciesCrc.IsValid() && GraphCache.GetFromCache(Task.Node, Task.Element.Get(), DependenciesCrc, TaskInput, TaskSettings, Task.SourceComponent.Get(), CachedOutput);
 #if WITH_EDITOR
-				const bool bNeedsToCreateActiveTask = !bResultAlreadyInCache || TaskSettingsInterface->bDebug;
+					bNeedsToCreateActiveTask = !bResultAlreadyInCache || TaskSettingsInterface->bDebug;
 #else
-				const bool bNeedsToCreateActiveTask = !bResultAlreadyInCache;
+					bNeedsToCreateActiveTask = !bResultAlreadyInCache;
 #endif
+				}
 
 				if (!bNeedsToCreateActiveTask)
 				{
@@ -655,6 +716,7 @@ void FPCGGraphExecutor::Execute()
 					Task.Context->CompiledTaskId = Task.CompiledTaskId;
 					Task.Context->InputData.Crc = TaskInput.Crc;
 					Task.Context->DependenciesCrc = DependenciesCrc;
+					Task.Context->GenerationGrid = Task.GenerationGrid;
 					Task.Context->Stack = (Task.StackContext && Task.StackIndex != INDEX_NONE) ? Task.StackContext->GetStack(Task.StackIndex) : nullptr;
 				}
 
@@ -1406,13 +1468,184 @@ bool FPCGFetchInputElement::ExecuteInternal(FPCGContext* Context) const
 	return true;
 }
 
-FPCGGenericElement::FPCGGenericElement(TFunction<bool(FPCGContext*)> InOperation)
+FPCGGenericElement::FPCGGenericElement(TFunction<bool(FPCGContext*)> InOperation, const FContextAllocator& InContextAllocator)
 	: Operation(InOperation)
+	, ContextAllocator(InContextAllocator)
 {
 }
 
+FPCGContext* FPCGGenericElement::Initialize(const FPCGDataCollection& InputData, TWeakObjectPtr<UPCGComponent> SourceComponent, const UPCGNode* Node)
+{
+	check(ContextAllocator);
+	FPCGContext* Context = ContextAllocator(InputData, SourceComponent, Node);
+	Context->InputData = InputData;
+	Context->SourceComponent = SourceComponent;
+	Context->Node = Node;
+
+	return Context;
+}
+
 bool FPCGGenericElement::ExecuteInternal(FPCGContext* Context) const
 {
 	TRACE_CPUPROFILER_EVENT_SCOPE(FPCGGenericElement::Execute);
 	return Operation(Context);
 }
+
+namespace PCGGraphExecutor
+{
+	bool ExecuteGridLinkage(EPCGHiGenGrid InFromGrid, EPCGHiGenGrid InToGrid, const FString& InResourceKey, const FName& InOutputPinLabel, const UPCGNode* InDownstreamNode, FPCGGridLinkageContext* InContext)
+	{
+		// Non partitioned graph - no linkage required - data should just pass through.
+		if (!InContext->SourceComponent->IsLocalComponent() || !InContext->SourceComponent->GetGraph()->IsHierarchicalGenerationEnabled())
+		{
+			InContext->OutputData = InContext->InputData;
+			return true;
+		}
+
+		EPCGHiGenGrid FromGrid = InFromGrid;
+		// If grid was not determined, apply default generation grid size.
+		if (!PCGHiGenGrid::IsValidGrid(FromGrid) && ensure(FromGrid == EPCGHiGenGrid::GenerationDefault))
+		{
+			FromGrid = InContext->SourceComponent->GetGraph()->GetDefaultGrid();
+		}
+
+		ensure(PCGHiGenGrid::IsValidGrid(FromGrid));
+		const int32 FromGridSize = PCGHiGenGrid::GridToGridSize(FromGrid);
+		ensure(PCGHiGenGrid::IsValidGrid(InToGrid));
+		const int32 ToGridSize = PCGHiGenGrid::GridToGridSize(InToGrid);
+
+		// Never allow a large grid to read data from small grid - this violates hierarchy.
+		if (FromGridSize < ToGridSize)
+		{
+#if WITH_EDITOR
+			if (UPCGSubsystem* Subsystem = UPCGSubsystem::GetInstance(InContext->SourceComponent->GetWorld()))
+			{
+				// Using the low level logging call because we have only a node pointer for the downstream node. Note that InContext
+				// is the context for the linkage element/task, which is not represented on the graph and cannot receive graph warnings/errors.
+				Subsystem->GetNodeVisualLogsMutable().Log(
+					InDownstreamNode, InContext->SourceComponent, ELogVerbosity::Error,
+					FText::Format(NSLOCTEXT("PCGGraphCompiler", "InvalidLinkage", "Could not read data across grid levels - origin grid size {0} must be greater than destination grid size {1}. Graph default grid size may need increasing."),
+						FromGridSize, ToGridSize
+					));
+
+			}
+#endif
+
+			return true;
+		}
+
+		if (!!(FromGrid & InContext->GenerationGrid) && FromGridSize != ToGridSize)
+		{
+			PCGGraphExecutionLogging::LogGridLinkageTaskExecuteStore(InContext, FromGridSize, ToGridSize, InResourceKey);
+
+			FPCGDataCollection Data;
+			Data.TaggedData = InContext->InputData.GetInputsByPin(InOutputPinLabel);
+			InContext->SourceComponent->StoreOutputDataForPin(InResourceKey, Data);
+		}
+		else if (InToGrid == InContext->GenerationGrid && FromGridSize != ToGridSize)
+		{
+			PCGGraphExecutionLogging::LogGridLinkageTaskExecuteRetrieve(InContext, FromGridSize, ToGridSize, InResourceKey);
+
+			UPCGSubsystem* Subsystem = UPCGSubsystem::GetInstance(InContext->SourceComponent->GetWorld());
+			if (!ensure(Subsystem))
+			{
+				return false;
+			}
+
+			UPCGComponent* ComponentWithData = nullptr;
+			APCGPartitionActor* PartitionActor = Cast<APCGPartitionActor>(InContext->SourceComponent->GetOwner());
+			UPCGComponent* OriginalComponent = PartitionActor ? PartitionActor->GetOriginalComponent(InContext->SourceComponent.Get()) : nullptr;
+			if (OriginalComponent)
+			{
+				// Now iterate through local components and find a component from the correct from-grid, and overlapping this component.
+				// Would be more efficient if we had a spatial index. The current component octree seems to only store original components.
+				Subsystem->ForAllRegisteredLocalComponents(OriginalComponent, [FromGridSize, &ComponentWithData, InContext](UPCGComponent* InLocalComponent)
+				{
+					if (InLocalComponent->GetGenerationGridSize() == FromGridSize)
+					{
+						const FBox OtherBounds = InLocalComponent->GetGridBounds();
+						const FBox ThisBounds = InContext->SourceComponent->GetGridBounds();
+						const FBox Overlap = OtherBounds.Overlap(ThisBounds);
+						if (Overlap.GetVolume() > 0)
+						{
+							ComponentWithData = InLocalComponent;
+						}
+					}
+				});
+			}
+
+			if (!ComponentWithData)
+			{
+				// Nothing we can do currently if PCG component not present. One idea is to schedule an artifact-less execution but that
+				// comes with complications - artifacts/side effects are an integral part of execution. Most likely we'll do a cleanup
+				// pass of any unwanted artifacts/local-components later.
+				PCGGraphExecutionLogging::LogGridLinkageTaskExecuteRetrieveNoLocalComponent(InContext, InResourceKey);
+				return true;
+			}
+
+			// Once we've found our component, try to retrieve the data.
+			if (const FPCGDataCollection* Data = ComponentWithData->RetrieveOutputDataForPin(InResourceKey))
+			{
+				PCGGraphExecutionLogging::LogGridLinkageTaskExecuteRetrieveSuccess(InContext, ComponentWithData, InResourceKey, Data->TaggedData.Num());
+				InContext->OutputData = *Data;
+
+				return true;
+			}
+
+			// At this point we could not get to the data, so we'll try executing the graph if we did not do that already.
+
+			auto WakeUpLambda = [InContext, ComponentWithData]()
+			{
+				PCGGraphExecutionLogging::LogGridLinkageTaskExecuteRetrieveWakeUp(InContext, ComponentWithData);
+				InContext->bIsPaused = false;
+				return true;
+			};
+
+			if (ComponentWithData->IsGenerating())
+			{
+				PCGGraphExecutionLogging::LogGridLinkageTaskExecuteRetrieveWaitOnScheduledGraph(InContext, ComponentWithData, InResourceKey);
+
+				// The component was already generating, but we were not asleep. Not really clear what's happening here,
+				// but in any case go to sleep and wake up when it's done.
+				InContext->bIsPaused = true;
+
+				// Wake up this task after graph has generated.
+				FPCGTaskId GenerationTask = ComponentWithData->GetGenerationTaskId();
+				if (ensure(GenerationTask != InvalidPCGTaskId))
+				{
+					Subsystem->ScheduleGeneric(WakeUpLambda, InContext->SourceComponent.Get(), { ComponentWithData->GetGenerationTaskId() });
+				}
+
+				return false;
+			}
+
+			// Graph is not currently generating. If we have not already tried generating, try it once now.
+			if (!InContext->bScheduledGraph)
+			{
+				PCGGraphExecutionLogging::LogGridLinkageTaskExecuteRetrieveScheduleGraph(InContext, ComponentWithData, InResourceKey);
+
+				// Wake up this task after graph has generated.
+				const FPCGTaskId GraphTaskId = ComponentWithData->GenerateLocalGetTaskId(/*bForce=*/true);
+				Subsystem->ScheduleGeneric(WakeUpLambda, InContext->SourceComponent.Get(), { GraphTaskId });
+
+				// Update state and go to sleep.
+				InContext->bScheduledGraph = true;
+				InContext->bIsPaused = true;
+				return false;
+			}
+			else
+			{
+				// We tried generating but no luck, at this point give up.
+				PCGGraphExecutionLogging::LogGridLinkageTaskExecuteRetrieveNoData(InContext, ComponentWithData, InResourceKey);
+				return true;
+			}
+		}
+		else
+		{
+			// Graceful no op - no linkage required.
+			InContext->OutputData = InContext->InputData;
+		}
+
+		return true;
+	}
+}
diff --git a/Engine/Plugins/Experimental/PCG/Source/PCG/Private/Graph/PCGGraphExecutor.h b/Engine/Plugins/Experimental/PCG/Source/PCG/Private/Graph/PCGGraphExecutor.h
--- a/Engine/Plugins/Experimental/PCG/Source/PCG/Private/Graph/PCGGraphExecutor.h
+++ b/Engine/Plugins/Experimental/PCG/Source/PCG/Private/Graph/PCGGraphExecutor.h
@@ -6,6 +6,7 @@
 #include "PCGElement.h"
 #include "PCGSubsystem.h"
 #include "Graph/PCGGraphCache.h"
+#include "Graph/PCGStackContext.h"
 
 #if WITH_EDITOR
 #include "AsyncCompilationHelpers.h"
@@ -55,6 +56,13 @@ struct FPCGGraphTask
 	FPCGTaskId CompiledTaskId = InvalidPCGTaskId; // the task id as it exists when compiled
 	FPCGTaskId ParentId = InvalidPCGTaskId; // represents the parent sub object graph task, if we were called from one
 
+	// Required purely to cache across multiple frames and allow different grids to be scheduled simultaneously. Set up
+	// at schedule-time, used during execution.
+	EPCGHiGenGrid GenerationGrid = EPCGHiGenGrid::Uninitialized;
+	
+	// The generation level for this node within the graph.
+	EPCGHiGenGrid GraphGenerationGrid = EPCGHiGenGrid::Uninitialized;
+
 	int32 StackIndex = INDEX_NONE;
 	TSharedPtr<const FPCGStackContext> StackContext;
 };
@@ -222,7 +230,17 @@ protected:
 class FPCGGenericElement : public FSimplePCGElement
 {
 public:
-	FPCGGenericElement(TFunction<bool(FPCGContext*)> InOperation);
+	using FContextAllocator = TFunction<FPCGContext*(const FPCGDataCollection&, TWeakObjectPtr<UPCGComponent>, const UPCGNode*)>;
+
+	FPCGGenericElement(
+		TFunction<bool(FPCGContext*)> InOperation,
+		const FContextAllocator& InContextAllocator = (FContextAllocator)[](const FPCGDataCollection&, TWeakObjectPtr<UPCGComponent>, const UPCGNode*)
+	{
+		return new FPCGContext();
+	});
+	
+	virtual FPCGContext* Initialize(const FPCGDataCollection& InputData, TWeakObjectPtr<UPCGComponent> SourceComponent, const UPCGNode* Node) override;
+
 	virtual bool IsCacheable(const UPCGSettings* InSettings) const override { return false; }
 	virtual bool CanExecuteOnlyOnMainThread(FPCGContext* Context) const override { return true; }
 
@@ -238,4 +256,26 @@ protected:
 
 private:
 	TFunction<bool(FPCGContext*)> Operation;
+
+	/** Creates a context object for this element. */
+	FContextAllocator ContextAllocator;
+};
+
+/** Context for linkage element which marshalls data across hierarchical generation grids. */
+struct FPCGGridLinkageContext : public FPCGContext
+{
+	/** If we require data from a component that is not generated, we schedule it once to see if we can get the data later. */
+	bool bScheduledGraph = false;
 };
+
+namespace PCGGraphExecutor
+{
+	/** Compares InFromGrid and InToGrid and performs data storage/retrieval as necessary to marshal data across execution grids. */
+	bool ExecuteGridLinkage(
+		EPCGHiGenGrid InFromGrid,
+		EPCGHiGenGrid InToGrid,
+		const FString& InResourceKey,
+		const FName& InOutputPinLabel,
+		const UPCGNode* InDownstreamNode,
+		FPCGGridLinkageContext* InContext);
+}
diff --git a/Engine/Plugins/Experimental/PCG/Source/PCG/Private/Grid/PCGPartitionActor.cpp b/Engine/Plugins/Experimental/PCG/Source/PCG/Private/Grid/PCGPartitionActor.cpp
--- a/Engine/Plugins/Experimental/PCG/Source/PCG/Private/Grid/PCGPartitionActor.cpp
+++ b/Engine/Plugins/Experimental/PCG/Source/PCG/Private/Grid/PCGPartitionActor.cpp
@@ -305,6 +305,7 @@ void APCGPartitionActor::AddGraphInstance(UPCGComponent* OriginalComponent)
 		// Update properties as needed and early out
 		LocalComponent->SetPropertiesFromOriginal(OriginalComponent);
 		LocalComponent->MarkAsLocalComponent();
+		LocalComponent->SetGenerationGridSize(PCGGridSize);
 		return;
 	}
 
@@ -313,6 +314,7 @@ void APCGPartitionActor::AddGraphInstance(UPCGComponent* OriginalComponent)
 	// Create a new local component
 	LocalComponent = NewObject<UPCGComponent>(this);
 	LocalComponent->MarkAsLocalComponent();
+	LocalComponent->SetGenerationGridSize(PCGGridSize);
 
 	// Note: we'll place the local component prior to the SetPropertiesFromOriginal so that any code that relies on the parent-child relationship works here
 	OriginalToLocal.Emplace(OriginalComponent, LocalComponent);
@@ -417,9 +419,10 @@ AActor* APCGPartitionActor::GetSceneOutlinerParent() const
 	}	
 }
 
-void APCGPartitionActor::PostCreation()
+void APCGPartitionActor::PostCreation(const FGuid& InGridGUID)
 {
 	PCGGridSize = GridSize;
+	PCGGuid = InGridGUID;
 
 	// Put in cache if we use the 2D grid or not.
 	if (APCGWorldActor* PCGActor = PCGHelpers::GetPCGWorldActor(GetWorld()))
diff --git a/Engine/Plugins/Experimental/PCG/Source/PCG/Private/Grid/PCGPartitionActor.h b/Engine/Plugins/Experimental/PCG/Source/PCG/Private/Grid/PCGPartitionActor.h
--- a/Engine/Plugins/Experimental/PCG/Source/PCG/Private/Grid/PCGPartitionActor.h
+++ b/Engine/Plugins/Experimental/PCG/Source/PCG/Private/Grid/PCGPartitionActor.h
@@ -45,6 +45,7 @@ public:
 	//~Begin APartitionActor Interface
 	virtual uint32 GetDefaultGridSize(UWorld* InWorld) const override;
 	virtual FGuid GetGridGuid() const override { return PCGGuid; }
+	virtual bool ShouldIncludeGridSizeInLabel() const override { return true; }
 	//~End APartitionActor Interface
 #endif
 
@@ -64,7 +65,7 @@ public:
 
 #if WITH_EDITOR
 	/** To be called after the creation of a new actor to copy the GridSize property (Editor only) into the PCGGridSize property */
-	void PostCreation();
+	void PostCreation(const FGuid& InGridGUID);
 
 	/** [Game thread only] Return if the actor is safe for deletion, meaning no generation is currently running on all original components. */
 	bool IsSafeForDeletion() const;
diff --git a/Engine/Plugins/Experimental/PCG/Source/PCG/Private/Helpers/PCGHelpers.cpp b/Engine/Plugins/Experimental/PCG/Source/PCG/Private/Helpers/PCGHelpers.cpp
--- a/Engine/Plugins/Experimental/PCG/Source/PCG/Private/Helpers/PCGHelpers.cpp
+++ b/Engine/Plugins/Experimental/PCG/Source/PCG/Private/Helpers/PCGHelpers.cpp
@@ -3,8 +3,10 @@
 #include "Helpers/PCGHelpers.h"
 
 #include "PCGComponent.h"
+#include "PCGGraph.h"
 #include "PCGModule.h"
 #include "PCGSubsystem.h"
+#include "PCGWorldActor.h"
 #include "Grid/PCGPartitionActor.h"
 
 #include "Landscape.h"
@@ -405,4 +407,24 @@ namespace PCGHelpers
 		// even if it is loading, so we use the outer found above.
 		return CurrentInspectedObject && !CurrentInspectedObject->HasAnyFlags(RF_ClassDefaultObject | RF_NeedLoad | RF_NeedPostLoad);
 	}
+
+	bool GetGenerationGridSizes(const UPCGGraph* InGraph, const APCGWorldActor* InWorldActor, PCGHiGenGrid::FSizeArray& OutGridSizes)
+	{
+		if (InGraph && InGraph->IsHierarchicalGenerationEnabled())
+		{
+			InGraph->GetGridSizes(OutGridSizes);
+			return true;
+		}
+		else if (InWorldActor)
+		{
+			OutGridSizes.Add(InWorldActor->PartitionGridSize);
+			return true;
+		}
+		else
+		{
+			ensure(false);
+			OutGridSizes.Add(PCGHiGenGrid::UninitializedGridSize());
+			return false;
+		}
+	}
 }
diff --git a/Engine/Plugins/Experimental/PCG/Source/PCG/Private/PCGActorAndComponentMapping.cpp b/Engine/Plugins/Experimental/PCG/Source/PCG/Private/PCGActorAndComponentMapping.cpp
--- a/Engine/Plugins/Experimental/PCG/Source/PCG/Private/PCGActorAndComponentMapping.cpp
+++ b/Engine/Plugins/Experimental/PCG/Source/PCG/Private/PCGActorAndComponentMapping.cpp
@@ -3,6 +3,7 @@
 #include "PCGActorAndComponentMapping.h"
 
 #include "PCGComponent.h"
+#include "PCGGraph.h"
 #include "PCGModule.h"
 #include "PCGSubsystem.h"
 #include "PCGWorldActor.h"
@@ -139,7 +140,9 @@ bool UPCGActorAndComponentMapping::RegisterOrUpdatePartitionedPCGComponent(UPCGC
 	// sure to update the runtime flow, for them to also create PA if they need to.
 	if (bComponentHasChanged || bComponentWasAdded)
 	{
-		PCGSubsystem->CreatePartitionActorsWithinBounds(Bounds);
+		PCGHiGenGrid::FSizeArray GridSizes;
+		ensure(PCGHelpers::GetGenerationGridSizes(InComponent ? InComponent->GetGraph() : nullptr, PCGSubsystem->GetPCGWorldActor(), GridSizes));
+		PCGSubsystem->CreatePartitionActorsWithinBounds(Bounds, GridSizes);
 	}
 #endif // WITH_EDITOR
 
@@ -327,12 +330,13 @@ void UPCGActorAndComponentMapping::RegisterPartitionActor(APCGPartitionActor* Ac
 	{
 		FWriteScopeLock WriteLock(PartitionActorsMapLock);
 
-		if (PartitionActorsMap.Contains(GridCoord))
+		TMap<FIntVector, TObjectPtr<APCGPartitionActor>>& PartitionActorsMapGrid = PartitionActorsMap.FindOrAdd(Actor->GetGridSize());
+		if (PartitionActorsMapGrid.Contains(GridCoord))
 		{
 			return;
 		}
 
-		PartitionActorsMap.Add(GridCoord, Actor);
+		PartitionActorsMapGrid.Add(GridCoord, Actor);
 	}
 
 	// For deprecration: bUse2DGrid is now true by default. But if we already have Partition Actors that were created when the flag was false by default,
@@ -372,9 +376,10 @@ void UPCGActorAndComponentMapping::UnregisterPartitionActor(APCGPartitionActor*
 
 	FIntVector GridCoord = Actor->GetGridCoord();
 
+	if (TMap<FIntVector, TObjectPtr<APCGPartitionActor>>* PartitionActorsMapGrid = PartitionActorsMap.Find(Actor->GetGridSize()))
 	{
 		FWriteScopeLock WriteLock(PartitionActorsMapLock);
-		PartitionActorsMap.Remove(GridCoord);
+		PartitionActorsMapGrid->Remove(GridCoord);
 	}
 
 	// And then unregister itself to all the components that intersect with it
@@ -401,17 +406,22 @@ void UPCGActorAndComponentMapping::ForAllIntersectingPartitionActors(const FBox&
 		return;
 	}
 
-	const uint32 GridSize = PCGWorldActor->PartitionGridSize;
-	const bool bUse2DGrid = PCGWorldActor->bUse2DGrid;
-	FIntVector MinCellCoords = UPCGActorHelpers::GetCellCoord(InBounds.Min, GridSize, bUse2DGrid);
-	FIntVector MaxCellCoords = UPCGActorHelpers::GetCellCoord(InBounds.Max, GridSize, bUse2DGrid);
-
+	PCGHiGenGrid::FSizeToGuidMap GridSizeToGuid;
+	PCGWorldActor->GetGridGuids(GridSizeToGuid);
+	for (const TPair<uint32, FGuid>& SizeAndGuid : GridSizeToGuid)
 	{
+		const uint32 GridSize = SizeAndGuid.Key;
+
+		const bool bUse2DGrid = PCGWorldActor->bUse2DGrid;
+		FIntVector MinCellCoords = UPCGActorHelpers::GetCellCoord(InBounds.Min, GridSize, bUse2DGrid);
+		FIntVector MaxCellCoords = UPCGActorHelpers::GetCellCoord(InBounds.Max, GridSize, bUse2DGrid);
+
 		FReadScopeLock ReadLock(PartitionActorsMapLock);
 
-		if (PartitionActorsMap.IsEmpty())
+		const TMap<FIntVector, TObjectPtr<APCGPartitionActor>>* PartitionActorsMapGrid = PartitionActorsMap.Find(GridSize);
+		if (!PartitionActorsMapGrid || PartitionActorsMapGrid->IsEmpty())
 		{
-			return;
+			continue;
 		}
 
 		for (int32 z = MinCellCoords.Z; z <= MaxCellCoords.Z; z++)
@@ -421,7 +431,7 @@ void UPCGActorAndComponentMapping::ForAllIntersectingPartitionActors(const FBox&
 				for (int32 x = MinCellCoords.X; x <= MaxCellCoords.X; x++)
 				{
 					FIntVector CellCoords(x, y, z);
-					if (const TObjectPtr<APCGPartitionActor>* ActorPtr = PartitionActorsMap.Find(CellCoords))
+					if (const TObjectPtr<APCGPartitionActor>* ActorPtr = PartitionActorsMapGrid->Find(CellCoords))
 					{
 						if (APCGPartitionActor* Actor = ActorPtr->Get())
 						{
diff --git a/Engine/Plugins/Experimental/PCG/Source/PCG/Private/PCGCommon.cpp b/Engine/Plugins/Experimental/PCG/Source/PCG/Private/PCGCommon.cpp
--- a/Engine/Plugins/Experimental/PCG/Source/PCG/Private/PCGCommon.cpp
+++ b/Engine/Plugins/Experimental/PCG/Source/PCG/Private/PCGCommon.cpp
@@ -11,4 +11,39 @@ namespace PCGFeatureSwitches
 		TEXT("Checks expected memory size consumption prior to performing sampling operations")
 	};
 
-}
\ No newline at end of file
+}
+
+namespace PCGHiGenGrid
+{
+	bool IsValidGridSize(uint32 InGridSize)
+	{
+		// Must be a power of 2 (in m) and within the valid range
+		// TODO: support other units
+		return FMath::IsPowerOfTwo(InGridSize / 100)
+			&& InGridSize >= GridToGridSize(EPCGHiGenGrid::GridMin)
+			&& InGridSize <= GridToGridSize(EPCGHiGenGrid::GridMax);
+	}
+
+	bool IsValidGrid(EPCGHiGenGrid InGrid)
+	{
+		// Check the bitmask value is within range
+		return InGrid >= EPCGHiGenGrid::GridMin && static_cast<uint32>(InGrid) < 2 * static_cast<uint32>(EPCGHiGenGrid::GridMax);
+	}
+
+	uint32 GridToGridSize(EPCGHiGenGrid InGrid)
+	{
+		const uint32 GridAsUint = static_cast<uint32>(InGrid);
+		ensure(FMath::IsPowerOfTwo(GridAsUint));
+		// TODO: support other units
+		return IsValidGrid(InGrid) ? (GridAsUint * 100) : UninitializedGridSize();
+	}
+
+	EPCGHiGenGrid GridSizeToGrid(uint32 InGridSize)
+	{
+		// TODO: support other units
+		return ensure(IsValidGridSize(InGridSize)) ? static_cast<EPCGHiGenGrid>(InGridSize / 100) : EPCGHiGenGrid::Uninitialized;
+	}
+
+	uint32 UninitializedGridSize() { return static_cast<uint32>(EPCGHiGenGrid::Uninitialized); }
+	uint32 UnboundedGridSize() { return static_cast<uint32>(EPCGHiGenGrid::Unbounded); }
+}
diff --git a/Engine/Plugins/Experimental/PCG/Source/PCG/Private/PCGComponent.cpp b/Engine/Plugins/Experimental/PCG/Source/PCG/Private/PCGComponent.cpp
--- a/Engine/Plugins/Experimental/PCG/Source/PCG/Private/PCGComponent.cpp
+++ b/Engine/Plugins/Experimental/PCG/Source/PCG/Private/PCGComponent.cpp
@@ -267,6 +267,11 @@ void UPCGComponent::GenerateLocal(bool bForce)
 	GenerateInternal(bForce, EPCGComponentGenerationTrigger::GenerateOnDemand, {});
 }
 
+FPCGTaskId UPCGComponent::GenerateLocalGetTaskId(bool bForce)
+{
+	return GenerateInternal(bForce, EPCGComponentGenerationTrigger::GenerateOnDemand, {});
+}
+
 FPCGTaskId UPCGComponent::GenerateInternal(bool bForce, EPCGComponentGenerationTrigger RequestedGenerationTrigger, const TArray<FPCGTaskId>& Dependencies)
 {
 	if (IsGenerating() || !GetSubsystem() || !ShouldGenerate(bForce, RequestedGenerationTrigger))
@@ -647,6 +652,24 @@ AActor* UPCGComponent::ClearPCGLink(UClass* TemplateActor)
 	return NewActor;
 }
 
+void UPCGComponent::StoreOutputDataForPin(const FString& InResourceKey, const FPCGDataCollection& InData)
+{
+	FReadScopeLock ScopedWriteLock(PerPinGeneratedOutputLock);
+	PerPinGeneratedOutput.FindOrAdd(InResourceKey) = InData;
+}
+
+const FPCGDataCollection* UPCGComponent::RetrieveOutputDataForPin(const FString& InResourceKey)
+{
+	FReadScopeLock ScopedReadLock(PerPinGeneratedOutputLock);
+	return PerPinGeneratedOutput.Find(InResourceKey);
+}
+
+void UPCGComponent::ClearPerPinGeneratedOutput()
+{
+	FReadScopeLock ScopedWriteLock(PerPinGeneratedOutputLock);
+	PerPinGeneratedOutput.Reset();
+}
+
 bool UPCGComponent::MoveResourcesToNewActor(AActor* InNewActor, bool bCreateChild)
 {
 	// Don't move resources if we are generating or cleaning up
@@ -1624,6 +1647,8 @@ void UPCGComponent::DirtyGenerated(EPCGComponentDirtyFlag DirtyFlag, const bool
 
 	bDirtyGenerated = true;
 
+	ClearPerPinGeneratedOutput();
+
 	// Dirty data as a waterfall from basic values
 	if (!!(DirtyFlag & EPCGComponentDirtyFlag::Actor))
 	{
diff --git a/Engine/Plugins/Experimental/PCG/Source/PCG/Private/PCGGraph.cpp b/Engine/Plugins/Experimental/PCG/Source/PCG/Private/PCGGraph.cpp
--- a/Engine/Plugins/Experimental/PCG/Source/PCG/Private/PCGGraph.cpp
+++ b/Engine/Plugins/Experimental/PCG/Source/PCG/Private/PCGGraph.cpp
@@ -10,6 +10,7 @@
 #include "PCGPin.h"
 #include "PCGSubsystem.h"
 #include "PCGSubgraph.h"
+#include "Elements/PCGHiGenGridSize.h"
 #include "Elements/PCGUserParameterGet.h"
 
 #include UE_INLINE_GENERATED_CPP_BY_NAME(PCGGraph)
@@ -676,6 +677,45 @@ void UPCGGraph::PostNodeUndo(UPCGNode* InPCGNode)
 }
 #endif
 
+void UPCGGraph::GetGridSizes(PCGHiGenGrid::FSizeArray& OutGridSizes) const
+{
+	const uint32 GraphDefaultGridSize = GetDefaultGridSize();
+	if (!IsHierarchicalGenerationEnabled())
+	{
+		OutGridSizes.Add(GraphDefaultGridSize);
+		return;
+	}
+
+	bool bHasUninitialized = false;
+	for (const UPCGNode* Node : Nodes)
+	{
+		const uint32 GridSize = GetNodeGenerationGridSize(Node, GraphDefaultGridSize);
+		if (PCGHiGenGrid::IsValidGridSize(GridSize))
+		{
+			if (!OutGridSizes.Contains(GridSize))
+			{
+				OutGridSizes.Add(GridSize);
+			}
+		}
+		else if (GridSize == PCGHiGenGrid::UninitializedGridSize())
+		{
+			// Outside nodes will not have a concrete grid set
+			bHasUninitialized = true;
+		}
+	}
+
+	if (bHasUninitialized)
+	{
+		// Nodes outside grid ranges will execute at graph default
+		OutGridSizes.Add(GraphDefaultGridSize);
+	}
+
+	// Descending
+	OutGridSizes.Sort([](const uint32& A, const uint32& B) { return A > B; });
+
+	return;
+}
+
 #if WITH_EDITOR
 void UPCGGraph::DisableNotificationsForEditor()
 {
@@ -803,7 +843,12 @@ void UPCGGraph::NotifyGraphParametersChanged(EPCGGraphParameterEvent InChangeTyp
 
 void UPCGGraph::OnNodeChanged(UPCGNode* InNode, EPCGChangeType ChangeType)
 {
-	if((ChangeType & ~EPCGChangeType::Cosmetic) != EPCGChangeType::None)
+	if (ChangeType != EPCGChangeType::Cosmetic)
+	{
+		NodeToGridSize.Reset();
+	}
+
+	if ((ChangeType & ~EPCGChangeType::Cosmetic) != EPCGChangeType::None)
 	{
 		NotifyGraphChanged(ChangeType);
 	}
@@ -865,6 +910,12 @@ void UPCGGraph::PostEditChangeProperty(FPropertyChangedEvent& PropertyChangedEve
 	{
 		OnGraphParametersChanged(EPCGGraphParameterEvent::ValueModifiedLocally, MemberPropertyName);
 	}
+	else if (PropertyName == GET_MEMBER_NAME_CHECKED(UPCGGraph, HiGenGridSize)
+		|| PropertyName == GET_MEMBER_NAME_CHECKED(UPCGGraph, bHierarchicalGenerationEnabled))
+	{
+		// The higen settings change the structure of the graph (presence or absence of links between grid levels).
+		NotifyGraphChanged(EPCGChangeType::Structural);
+	}
 
 	NumberOfUserParametersPreEdit = 0;
 }
@@ -931,6 +982,63 @@ void UPCGGraph::OnGraphParametersChanged(EPCGGraphParameterEvent InChangeType, F
 
 #endif // WITH_EDITOR
 
+uint32 UPCGGraph::GetNodeGenerationGridSize(const UPCGNode* InNode, uint32 InDefaultGridSize) const
+{
+	{
+		FReadScopeLock ScopedReadLock(NodeToGridSizeLock);
+		if (const uint32* CachedGridSize = NodeToGridSize.Find(InNode))
+		{
+			return *CachedGridSize;
+		}
+	}
+
+	{
+		FWriteScopeLock ScopedWriteLock(NodeToGridSizeLock);
+		return CalculateNodeGridSizeRecursive_Unsafe(InNode, InDefaultGridSize);
+	}
+}
+
+uint32 UPCGGraph::CalculateNodeGridSizeRecursive_Unsafe(const UPCGNode* InNode, uint32 InDefaultGridSize) const
+{
+	uint32 GridSize = PCGHiGenGrid::UninitializedGridSize();
+	if (const UPCGHiGenGridSizeSettings* GridSizeSettings = Cast<UPCGHiGenGridSizeSettings>(InNode->GetSettings()))
+	{
+		GridSize = GridSizeSettings->GetGridSize();
+	}
+	else
+	{
+		// Grid size for a node is the minimum of the grid sizes of connected upstream nodes.
+		uint32 MinGenerationSize = std::numeric_limits<uint32>::max();
+		for (const UPCGPin* Pin : InNode->GetInputPins())
+		{
+			if (Pin)
+			{
+				for (const UPCGEdge* Edge : Pin->Edges)
+				{
+					const UPCGPin* OtherPin = Edge ? Edge->InputPin : nullptr;
+					if (OtherPin && OtherPin->Node.Get())
+					{
+						const uint32 InputGridSize = CalculateNodeGridSizeRecursive_Unsafe(OtherPin->Node, InDefaultGridSize);
+						if (PCGHiGenGrid::IsValidGridSize(InputGridSize))
+						{
+							MinGenerationSize = FMath::Min(MinGenerationSize, InputGridSize);
+						}
+					}
+				}
+			}
+		}
+
+		GridSize = (MinGenerationSize == std::numeric_limits<uint32>::max()) ? InDefaultGridSize : MinGenerationSize;
+	}
+
+	if (GridSize != PCGHiGenGrid::UninitializedGridSize())
+	{
+		NodeToGridSize.Add(InNode, GridSize);
+	}
+
+	return GridSize;
+}
+
 /****************************
 * UPCGGraphInstance
 ****************************/
diff --git a/Engine/Plugins/Experimental/PCG/Source/PCG/Private/PCGSubsystem.cpp b/Engine/Plugins/Experimental/PCG/Source/PCG/Private/PCGSubsystem.cpp
--- a/Engine/Plugins/Experimental/PCG/Source/PCG/Private/PCGSubsystem.cpp
+++ b/Engine/Plugins/Experimental/PCG/Source/PCG/Private/PCGSubsystem.cpp
@@ -3,6 +3,7 @@
 #include "PCGSubsystem.h"
 
 #include "PCGComponent.h"
+#include "PCGGraph.h"
 #include "PCGModule.h"
 #include "PCGWorldActor.h"
 #include "Graph/PCGGraphCache.h"
@@ -99,7 +100,7 @@ namespace PCGSubsystemConsole
 #if WITH_EDITOR
 namespace PCGSubsystem
 {
-	FPCGTaskId ForEachIntersectingCell(FPCGGraphExecutor* GraphExecutor, UWorld* World, const FBox& InBounds, bool bCreateActor, bool bLoadCell, bool bSaveActors, TFunctionRef<FPCGTaskId(APCGPartitionActor*, const FBox&, const TArray<FPCGTaskId>&)> InOperation);
+	FPCGTaskId ForEachIntersectingCell(FPCGGraphExecutor* GraphExecutor, UWorld* World, const FBox& InBounds, const PCGHiGenGrid::FSizeArray& InGridSizes, bool bCreateActor, bool bLoadCell, bool bSaveActors, TFunctionRef<FPCGTaskId(APCGPartitionActor*, const FBox&, const TArray<FPCGTaskId>&)> InOperation);
 }
 #endif
 
@@ -278,7 +279,9 @@ FPCGTaskId UPCGSubsystem::ScheduleComponent(UPCGComponent* PCGComponent, bool bS
 		// Note: This is an immediate operation, as we need the PA for the generation.
 		auto ScheduleTask = [](APCGPartitionActor* PCGActor, const FBox& InIntersectedBounds, const TArray<FPCGTaskId>& TaskDependencies) { return InvalidPCGTaskId; };
 
-		PCGSubsystem::ForEachIntersectingCell(GraphExecutor, PCGComponent->GetWorld(), PCGComponent->GetGridBounds(), /*bCreateActor=*/true, /*bLoadCell=*/false, /*bSave=*/false, ScheduleTask);
+		PCGHiGenGrid::FSizeArray GridSizes;
+		ensure(PCGHelpers::GetGenerationGridSizes(PCGComponent->GetGraph(), GetPCGWorldActor(), GridSizes));
+		PCGSubsystem::ForEachIntersectingCell(GraphExecutor, PCGComponent->GetWorld(), PCGComponent->GetGridBounds(), GridSizes, /*bCreateActor=*/true, /*bLoadCell=*/false, /*bSave=*/false, ScheduleTask);
 
 		ActorAndComponentMapping.UpdateMappingPCGComponentPartitionActor(PCGComponent);
 	}
@@ -570,7 +573,7 @@ bool UPCGSubsystem::GetOutputData(FPCGTaskId TaskId, FPCGDataCollection& OutData
 
 namespace PCGSubsystem
 {
-	FPCGTaskId ForEachIntersectingCell(FPCGGraphExecutor* GraphExecutor, UWorld* World, const FBox& InBounds, bool bCreateActor, bool bLoadCell, bool bSaveActors, TFunctionRef<FPCGTaskId(APCGPartitionActor*, const FBox&, const TArray<FPCGTaskId>&)> InOperation)
+	FPCGTaskId ForEachIntersectingCell(FPCGGraphExecutor* GraphExecutor, UWorld* World, const FBox& InBounds, const PCGHiGenGrid::FSizeArray& InGridSizes, bool bCreateActor, bool bLoadCell, bool bSaveActors, TFunctionRef<FPCGTaskId(APCGPartitionActor*, const FBox&, const TArray<FPCGTaskId>&)> InOperation)
 	{
 		if (!GraphExecutor || !World)
 		{
@@ -578,154 +581,182 @@ namespace PCGSubsystem
 			return InvalidPCGTaskId;
 		}
 
-		APCGWorldActor* PCGWorldActor = PCGHelpers::GetPCGWorldActor(World);
+		PCGHiGenGrid::FSizeArray GridSizes = InGridSizes;
+		// We have no use for unbounded grids as this is a grid-centric function.
+		GridSizes.Remove(static_cast<uint32>(EPCGHiGenGrid::Unbounded));
+		ensure(!GridSizes.IsEmpty());
+		if (GridSizes.IsEmpty())
+		{
+			return InvalidPCGTaskId;
+		}
+		// Descend down hierarchy
+		GridSizes.Sort([](const uint32& A, const uint32& B) { return A > B; });
 
+		APCGWorldActor* PCGWorldActor = PCGHelpers::GetPCGWorldActor(World);
 		check(PCGWorldActor);
 
-		// In case of 2D grid, we are clamping our bounds in Z to be within 0 and PartitionGridSize
-		// By doing so, WP will tie all the partition actors to [0, PartitionGridSize[ interval, generating a "2D grid" instead of 3D.
-		FBox ModifiedInBounds = InBounds;
-		if (PCGWorldActor->bUse2DGrid)
-		{
-			FVector MinBounds = InBounds.Min;
-			FVector MaxBounds = InBounds.Max;
-
-			MinBounds.Z = 0;
-			MaxBounds.Z = FMath::Max<FVector::FReal>(0.0, PCGWorldActor->PartitionGridSize - 1.0); // -1 to be just below the PartitionGridSize
-			ModifiedInBounds = FBox(MinBounds, MaxBounds);
-		}
+		PCGWorldActor->CreateGridGuidsIfNecessary(GridSizes);
 
+		PCGHiGenGrid::FSizeToGuidMap GridSizeToGuid;
+		PCGWorldActor->GetGridGuids(GridSizeToGuid);
+		
 		TArray<FPCGTaskId> CellTasks;
-
-		auto CellLambda = [&CellTasks, GraphExecutor, World, bCreateActor, bLoadCell, bSaveActors, ModifiedInBounds, &InOperation](const UActorPartitionSubsystem::FCellCoord& CellCoord, const FBox& CellBounds)
+		for (uint32 GridSize : GridSizes)
 		{
-			UActorPartitionSubsystem* PartitionSubsystem = UWorld::GetSubsystem<UActorPartitionSubsystem>(World);
-			FBox IntersectedBounds = ModifiedInBounds.Overlap(CellBounds);
-
-			if (IntersectedBounds.IsValid)
+			ensure(PCGHiGenGrid::IsValidGridSize(GridSize));
+			const FGuid* GuidPtr = GridSizeToGuid.Find(GridSize);
+			if (!ensure(GuidPtr))
 			{
-				TSharedPtr<TSet<FWorldPartitionReference>> ActorReferences = MakeShared<TSet<FWorldPartitionReference>>();
+				continue;
+			}
+			const FGuid GridGuid = *GuidPtr;
 
-				auto PostCreation = [](APartitionActor* Actor) { CastChecked<APCGPartitionActor>(Actor)->PostCreation(); };
+			// In case of 2D grid, we are clamping our bounds in Z to be within 0 and GridSize
+			// By doing so, WP will tie all the partition actors to [0, GridSize[ interval, generating a "2D grid" instead of 3D.
+			FBox ModifiedInBounds = InBounds;
+			if (PCGWorldActor->bUse2DGrid)
+			{
+				FVector MinBounds = InBounds.Min;
+				FVector MaxBounds = InBounds.Max;
 
-				const bool bInBoundsSearch = true;
-				const FGuid DefaultGridGuid;
-				const uint32 DefaultGridSize = 0;
+				MinBounds.Z = 0;
+				MaxBounds.Z = FMath::Max<FVector::FReal>(0.0, GridSize - 1.0); // -1 to be just below the GridSize
+				ModifiedInBounds = FBox(MinBounds, MaxBounds);
+			}
 
-				APCGPartitionActor* PCGActor = Cast<APCGPartitionActor>(PartitionSubsystem->GetActor(APCGPartitionActor::StaticClass(), CellCoord, bCreateActor, DefaultGridGuid, DefaultGridSize, bInBoundsSearch, PostCreation));
+			auto CellLambda = [&CellTasks, GraphExecutor, World, bCreateActor, bLoadCell, bSaveActors, ModifiedInBounds, &InOperation, GridSize, &GridGuid, PCGWorldActor](const UActorPartitionSubsystem::FCellCoord& CellCoord, const FBox& CellBounds)
+			{
+				UActorPartitionSubsystem* PartitionSubsystem = UWorld::GetSubsystem<UActorPartitionSubsystem>(World);
+				FBox IntersectedBounds = ModifiedInBounds.Overlap(CellBounds);
 
-				// At this point, if bCreateActor was true, then it exists, but it is not currently loaded; make sure it is loaded
-				// Otherwise, we still need to load it if it exists
-				// TODO: Revisit after API review on the WP side, we shouldn't have to load here or get the actor desc directly
-				if (!PCGActor && bSaveActors)
+				if (IntersectedBounds.IsValid)
 				{
-					const FWorldPartitionActorDesc* PCGActorDesc = nullptr;
-					auto FindFirst = [&CellCoord, &PCGActorDesc](const FWorldPartitionActorDesc* ActorDesc) {
-						FPartitionActorDesc* PartitionActorDesc = (FPartitionActorDesc*)ActorDesc;
-
-						if (PartitionActorDesc &&
-							PartitionActorDesc->GridIndexX == CellCoord.X &&
-							PartitionActorDesc->GridIndexY == CellCoord.Y &&
-							PartitionActorDesc->GridIndexZ == CellCoord.Z)
-						{
-							PCGActorDesc = ActorDesc;
-							return false;
-						}
-						else
+					TSharedPtr<TSet<FWorldPartitionReference>> ActorReferences = MakeShared<TSet<FWorldPartitionReference>>();
+
+					auto PostCreation = [GridSize, &GridGuid](APartitionActor* Actor) { CastChecked<APCGPartitionActor>(Actor)->PostCreation(GridGuid); };
+
+					APCGPartitionActor* PCGActor = Cast<APCGPartitionActor>(PartitionSubsystem->GetActor(
+						APCGPartitionActor::StaticClass(),
+						CellCoord,
+						bCreateActor,
+						GridGuid,
+						GridSize,
+						/*bInBoundsSearch=*/true,
+						PostCreation));
+
+					// At this point, if bCreateActor was true, then it exists, but it is not currently loaded; make sure it is loaded
+					// Otherwise, we still need to load it if it exists
+					// TODO: Revisit after API review on the WP side, we shouldn't have to load here or get the actor desc directly
+					if (!PCGActor && bSaveActors)
+					{
+						const FWorldPartitionActorDesc* PCGActorDesc = nullptr;
+						auto FindFirst = [&CellCoord, &PCGActorDesc](const FWorldPartitionActorDesc* ActorDesc) {
+							FPartitionActorDesc* PartitionActorDesc = (FPartitionActorDesc*)ActorDesc;
+
+							if (PartitionActorDesc &&
+								PartitionActorDesc->GridIndexX == CellCoord.X &&
+								PartitionActorDesc->GridIndexY == CellCoord.Y &&
+								PartitionActorDesc->GridIndexZ == CellCoord.Z)
+							{
+								PCGActorDesc = ActorDesc;
+								return false;
+							}
+							else
+							{
+								return true;
+							}
+						};
+
+						FWorldPartitionHelpers::ForEachIntersectingActorDesc<APCGPartitionActor>(World->GetWorldPartition(), CellBounds, FindFirst);
+
+						check(!bCreateActor || PCGActorDesc);
+						if (PCGActorDesc)
 						{
-							return true;
+							ActorReferences->Add(FWorldPartitionReference(World->GetWorldPartition(), PCGActorDesc->GetGuid()));
+							PCGActor = Cast<APCGPartitionActor>(PCGActorDesc->GetActor());
 						}
-					};
-
-					FWorldPartitionHelpers::ForEachIntersectingActorDesc<APCGPartitionActor>(World->GetWorldPartition(), CellBounds, FindFirst);
+					}
+					else if(PCGActor)
+					{
+						// We still need to keep a reference on the PCG actor - note that newly created PCG actors will not have a reference here, but won't be unloaded
+						ActorReferences->Add(FWorldPartitionReference(World->GetWorldPartition(), PCGActor->GetActorGuid()));
+					}
 
-					check(!bCreateActor || PCGActorDesc);
-					if (PCGActorDesc)
+					if (!PCGActor)
 					{
-						ActorReferences->Add(FWorldPartitionReference(World->GetWorldPartition(), PCGActorDesc->GetGuid()));
-						PCGActor = Cast<APCGPartitionActor>(PCGActorDesc->GetActor());
+						return true;
 					}
-				}
-				else if(PCGActor)
-				{
-					// We still need to keep a reference on the PCG actor - note that newly created PCG actors will not have a reference here, but won't be unloaded
-					ActorReferences->Add(FWorldPartitionReference(World->GetWorldPartition(), PCGActor->GetActorGuid()));
-				}
 
-				if (!PCGActor)
-				{
-					return true;
-				}
+					TArray<FPCGTaskId> PreviousTasks;
 
-				TArray<FPCGTaskId> PreviousTasks;
+					auto SetPreviousTaskIfValid = [&PreviousTasks](FPCGTaskId TaskId){
+						if (TaskId != InvalidPCGTaskId)
+						{
+							PreviousTasks.Reset();
+							PreviousTasks.Add(TaskId);
+						}
+					};
 
-				auto SetPreviousTaskIfValid = [&PreviousTasks](FPCGTaskId TaskId){
-					if (TaskId != InvalidPCGTaskId)
+					// We'll need to make sure actors in the bounds are loaded only if we need them.
+					if (bLoadCell)
 					{
-						PreviousTasks.Reset();
-						PreviousTasks.Add(TaskId);
-					}
-				};
+						auto WorldPartitionLoadActorsInBounds = [World, ActorReferences](const FWorldPartitionActorDesc* ActorDesc) {
+							check(ActorDesc);
+							ActorReferences->Add(FWorldPartitionReference(World->GetWorldPartition(), ActorDesc->GetGuid()));
+							// Load actor if not already loaded
+							ActorDesc->GetActor();
+							return true;
+						};
 
-				// We'll need to make sure actors in the bounds are loaded only if we need them.
-				if (bLoadCell)
-				{
-					auto WorldPartitionLoadActorsInBounds = [World, ActorReferences](const FWorldPartitionActorDesc* ActorDesc) {
-						check(ActorDesc);
-						ActorReferences->Add(FWorldPartitionReference(World->GetWorldPartition(), ActorDesc->GetGuid()));
-						// Load actor if not already loaded
-						ActorDesc->GetActor();
-						return true;
-					};
+						auto LoadActorsTask = [World, IntersectedBounds, WorldPartitionLoadActorsInBounds]() {
+							FWorldPartitionHelpers::ForEachIntersectingActorDesc(World->GetWorldPartition(), IntersectedBounds, WorldPartitionLoadActorsInBounds);
+							return true;
+						};
 
-					auto LoadActorsTask = [World, IntersectedBounds, WorldPartitionLoadActorsInBounds]() {
-						FWorldPartitionHelpers::ForEachIntersectingActorDesc(World->GetWorldPartition(), IntersectedBounds, WorldPartitionLoadActorsInBounds);
-						return true;
-					};
+						FPCGTaskId LoadTaskId = GraphExecutor->ScheduleGeneric(LoadActorsTask, nullptr, {});
+						SetPreviousTaskIfValid(LoadTaskId);
+					}
 
-					FPCGTaskId LoadTaskId = GraphExecutor->ScheduleGeneric(LoadActorsTask, nullptr, {});
-					SetPreviousTaskIfValid(LoadTaskId);
-				}
+					// Execute
+					FPCGTaskId ExecuteTaskId = InOperation(PCGActor, IntersectedBounds, PreviousTasks);
+					SetPreviousTaskIfValid(ExecuteTaskId);
 
-				// Execute
-				FPCGTaskId ExecuteTaskId = InOperation(PCGActor, IntersectedBounds, PreviousTasks);
-				SetPreviousTaskIfValid(ExecuteTaskId);
+					// Save changes; note that there's no need to save if the operation was cancelled
+					if (bSaveActors && ExecuteTaskId != InvalidPCGTaskId)
+					{
+						auto SaveActorTask = [PCGActor, GraphExecutor]() {
+							GraphExecutor->AddToDirtyActors(PCGActor);
+							return true;
+						};
 
-				// Save changes; note that there's no need to save if the operation was cancelled
-				if (bSaveActors && ExecuteTaskId != InvalidPCGTaskId)
-				{
-					auto SaveActorTask = [PCGActor, GraphExecutor]() {
-						GraphExecutor->AddToDirtyActors(PCGActor);
-						return true;
-					};
+						FPCGTaskId SaveTaskId = GraphExecutor->ScheduleGeneric(SaveActorTask, nullptr, PreviousTasks);
+						SetPreviousTaskIfValid(SaveTaskId);
+					}
 
-					FPCGTaskId SaveTaskId = GraphExecutor->ScheduleGeneric(SaveActorTask, nullptr, PreviousTasks);
-					SetPreviousTaskIfValid(SaveTaskId);
-				}
+					if (bLoadCell)
+					{
+						// Unload actors from cell (or the pcg actor refered here)
+						auto UnloadActorsTask = [GraphExecutor, ActorReferences]() {
+							GraphExecutor->AddToUnusedActors(*ActorReferences);
+							return true;
+						};
 
-				if (bLoadCell)
-				{
-					// Unload actors from cell (or the pcg actor refered here)
-					auto UnloadActorsTask = [GraphExecutor, ActorReferences]() {
-						GraphExecutor->AddToUnusedActors(*ActorReferences);
-						return true;
-					};
+						// Schedule after the save (if valid), then the execute so we can queue this after the load.
+						FPCGTaskId UnloadTaskId = GraphExecutor->ScheduleGeneric(UnloadActorsTask, nullptr, PreviousTasks);
+						SetPreviousTaskIfValid(UnloadTaskId);
+					}
 
-					// Schedule after the save (if valid), then the execute so we can queue this after the load.
-					FPCGTaskId UnloadTaskId = GraphExecutor->ScheduleGeneric(UnloadActorsTask, nullptr, PreviousTasks);
-					SetPreviousTaskIfValid(UnloadTaskId);
+					// Finally, mark "last" valid task in the cell tasks.
+					CellTasks.Append(PreviousTasks);
 				}
 
-				// Finally, mark "last" valid task in the cell tasks.
-				CellTasks.Append(PreviousTasks);
-			}
-
-			return true;
-		};
+				return true;
+			};
 
-		// TODO: accumulate last phases of every cell + run a GC at the end?
-		//  or add some mechanism on the graph executor side to run GC every now and then
-		FActorPartitionGridHelper::ForEachIntersectingCell(APCGPartitionActor::StaticClass(), ModifiedInBounds, World->PersistentLevel, CellLambda);
+			// TODO: accumulate last phases of every cell + run a GC at the end?
+			//  or add some mechanism on the graph executor side to run GC every now and then
+			FActorPartitionGridHelper::ForEachIntersectingCell(APCGPartitionActor::StaticClass(), ModifiedInBounds, World->PersistentLevel, CellLambda, GridSize);
+		}
 
 		// Finally, create a dummy generic task to wait on all cells
 		if (!CellTasks.IsEmpty())
@@ -740,7 +771,7 @@ namespace PCGSubsystem
 
 } // end namepsace PCGSubsystem
 
-void UPCGSubsystem::CreatePartitionActorsWithinBounds(const FBox& InBounds)
+void UPCGSubsystem::CreatePartitionActorsWithinBounds(const FBox& InBounds, const PCGHiGenGrid::FSizeArray& InGridSizes)
 {
 	if (!PCGHelpers::IsRuntimeOrPIE())
 	{
@@ -751,7 +782,7 @@ void UPCGSubsystem::CreatePartitionActorsWithinBounds(const FBox& InBounds)
 		UWorld* World = GetWorld();
 		if (ensure(World && !World->bIsRunningConstructionScript))
 		{
-			PCGSubsystem::ForEachIntersectingCell(GraphExecutor, World, InBounds, /*bCreateActor=*/true, /*bLoadCell=*/false, /*bSave=*/false, ScheduleTask);
+			PCGSubsystem::ForEachIntersectingCell(GraphExecutor, World, InBounds, InGridSizes, /*bCreateActor=*/true, /*bLoadCell=*/false, /*bSave=*/false, ScheduleTask);
 		}
 	}
 }
@@ -847,7 +878,9 @@ FPCGTaskId UPCGSubsystem::ProcessGraph(UPCGComponent* Component, const FBox& InP
 
 	if (UnionBounds.IsValid)
 	{
-		ProcessAllCellsTaskId = PCGSubsystem::ForEachIntersectingCell(GraphExecutor, Component->GetWorld(), UnionBounds, bCreateActors, bLoadCell, bSaveActors, ScheduleTask);
+		PCGHiGenGrid::FSizeArray GridSizes;
+		ensure(PCGHelpers::GetGenerationGridSizes(Component->GetGraph(), GetPCGWorldActor(), GridSizes));
+		ProcessAllCellsTaskId = PCGSubsystem::ForEachIntersectingCell(GraphExecutor, Component->GetWorld(), UnionBounds, GridSizes, bCreateActors, bLoadCell, bSaveActors, ScheduleTask);
 	}
 
 	// Finally, call PostProcessGraph if something happened
@@ -865,7 +898,7 @@ FPCGTaskId UPCGSubsystem::ProcessGraph(UPCGComponent* Component, const FBox& InP
 	}
 	else
 	{
-		if(Component)
+		if (Component)
 		{
 			Component->OnProcessGraphAborted();
 		}
@@ -914,7 +947,9 @@ FPCGTaskId UPCGSubsystem::CleanupGraph(UPCGComponent* Component, const FBox& InB
 		return InvalidPCGTaskId;
 	};
 
-	FPCGTaskId ProcessAllCellsTaskId = PCGSubsystem::ForEachIntersectingCell(GraphExecutor, Component->GetWorld(), InBounds, /*bCreateActor=*/false, /*bLoadCell=*/false, bSave, ScheduleTask);
+	PCGHiGenGrid::FSizeArray GridSizes;
+	ensure(PCGHelpers::GetGenerationGridSizes(Component->GetGraph(), GetPCGWorldActor(), GridSizes));
+	FPCGTaskId ProcessAllCellsTaskId = PCGSubsystem::ForEachIntersectingCell(GraphExecutor, Component->GetWorld(), InBounds, GridSizes, /*bCreateActor=*/false, /*bLoadCell=*/false, bSave, ScheduleTask);
 
 	// Finally, call PostCleanupGraph if something happened
 	if (ProcessAllCellsTaskId != InvalidPCGTaskId)
@@ -967,7 +1002,7 @@ void UPCGSubsystem::CleanupPartitionActors(const FBox& InBounds)
 		return GraphExecutor->ScheduleGeneric(CleanupTask, nullptr, TaskDependencies);
 	};
 
-	PCGSubsystem::ForEachIntersectingCell(GraphExecutor, GetWorld(), InBounds, /*bCreateActor=*/false, /*bLoadCell=*/false, /*bSave=*/false, ScheduleTask);
+	PCGSubsystem::ForEachIntersectingCell(GraphExecutor, GetWorld(), InBounds, {}, /*bCreateActor=*/false, /*bLoadCell=*/false, /*bSave=*/false, ScheduleTask);
 }
 
 void UPCGSubsystem::ClearPCGLink(UPCGComponent* InComponent, const FBox& InBounds, AActor* InNewActor)
@@ -990,7 +1025,9 @@ void UPCGSubsystem::ClearPCGLink(UPCGComponent* InComponent, const FBox& InBound
 		return GraphExecutor->ScheduleGeneric(MoveTask, ComponentPtr.Get(), TaskDependencies);
 	};
 
-	FPCGTaskId TaskId = PCGSubsystem::ForEachIntersectingCell(GraphExecutor, GetWorld(), InBounds, /*bCreateActor=*/false, /*bLoadCell=*/false, /*bSave=*/false, ScheduleTask);
+	PCGHiGenGrid::FSizeArray GridSizes;
+	ensure(PCGHelpers::GetGenerationGridSizes(InComponent->GetGraph(), GetPCGWorldActor(), GridSizes));
+	FPCGTaskId TaskId = PCGSubsystem::ForEachIntersectingCell(GraphExecutor, GetWorld(), InBounds, GridSizes, /*bCreateActor=*/false, /*bLoadCell=*/false, /*bSave=*/false, ScheduleTask);
 
 	// Verify if the NewActor has some components attached to its root or attached actors. If not, destroy it.
 	// Return false if the new actor is not valid or destroyed.
@@ -1108,25 +1145,28 @@ void UPCGSubsystem::DeletePartitionActors(bool bOnlyDeleteUnused)
 	if (bOnlyDeleteUnused)
 	{		
 		// Need to copy to avoid deadlock
-		TMap<FIntVector, TObjectPtr<APCGPartitionActor>> PartitionActorsMapCopy;
+		TMap<uint32, TMap<FIntVector, TObjectPtr<APCGPartitionActor>>> PartitionActorsMapCopy;
 		{
 			FReadScopeLock ReadLock(ActorAndComponentMapping.PartitionActorsMapLock);
 			PartitionActorsMapCopy = ActorAndComponentMapping.PartitionActorsMap;
 		}
 
-		for (const TPair<FIntVector, TObjectPtr<APCGPartitionActor>>& Item : PartitionActorsMapCopy)
+		for (const TPair<uint32, TMap<FIntVector, TObjectPtr<APCGPartitionActor>>>& GridItem : PartitionActorsMapCopy)
 		{
-			if (APCGPartitionActor* PartitionActor = Item.Value)
+			for (const TPair<FIntVector, TObjectPtr<APCGPartitionActor>>& Item : GridItem.Value)
 			{
-				bool bIntersectWithOneComponent = false;
+				if (APCGPartitionActor* PartitionActor = Item.Value)
+				{
+					bool bIntersectWithOneComponent = false;
 
-				auto FoundComponents = [&bIntersectWithOneComponent](UPCGComponent*) { bIntersectWithOneComponent = true; };
+					auto FoundComponents = [&bIntersectWithOneComponent](UPCGComponent*) { bIntersectWithOneComponent = true; };
 
-				ActorAndComponentMapping.ForAllIntersectingComponents(PartitionActor->GetFixedBounds(), FoundComponents);
+					ActorAndComponentMapping.ForAllIntersectingComponents(PartitionActor->GetFixedBounds(), FoundComponents);
 
-				if (!bIntersectWithOneComponent)
-				{
-					GatherAndDestroyLoadedActors(PartitionActor);
+					if (!bIntersectWithOneComponent)
+					{
+						GatherAndDestroyLoadedActors(PartitionActor);
+					}
 				}
 			}
 		}
diff --git a/Engine/Plugins/Experimental/PCG/Source/PCG/Private/PCGWorldActor.cpp b/Engine/Plugins/Experimental/PCG/Source/PCG/Private/PCGWorldActor.cpp
--- a/Engine/Plugins/Experimental/PCG/Source/PCG/Private/PCGWorldActor.cpp
+++ b/Engine/Plugins/Experimental/PCG/Source/PCG/Private/PCGWorldActor.cpp
@@ -60,6 +60,65 @@ void APCGWorldActor::EndPlay(const EEndPlayReason::Type EndPlayReason)
 	Super::EndPlay(EndPlayReason);
 }
 
+void APCGWorldActor::CreateGridGuidsIfNecessary(const PCGHiGenGrid::FSizeArray& InGridSizes)
+{
+	// Check if any need adding
+	ensure(!InGridSizes.IsEmpty());
+	PCGHiGenGrid::FSizeArray GridSizesToAdd;
+	{
+		FReadScopeLock ReadLock(GridGuidsLock);
+		for (uint32 GridSize : InGridSizes)
+		{
+			if (!GridGuids.Contains(GridSize))
+			{
+				GridSizesToAdd.Push(GridSize);
+			}
+		}
+	}
+
+	if (GridSizesToAdd.Num() > 0)
+	{
+		FWriteScopeLock WriteLock(GridGuidsLock);
+
+		bool bModified = false;
+		for (uint32 GridSize : GridSizesToAdd)
+		{
+			if (!GridGuids.Contains(GridSize))
+			{
+				GridGuids.Add(GridSize, FGuid::NewGuid());
+				bModified = true;
+			}
+		}
+
+		if (bModified)
+		{
+			// Set dirty flag if we added guids. Unfortunately if the guids are not up to date, this will produce save prompts
+			// to users. However, this was needed to ensure the guids are saved - without this guids were lost and PAs were leaked.
+			// The alternative would be to ensure this never happens automatically, but rather only happens when user clicks Generate
+			// or etc. However we take a very proactive approach to creating PAs in editor because they can't be created at runtime.
+
+			// Schedule dirtying rather than do immediately because dirtying is a no-op during level load
+			if (UPCGSubsystem* PCGSubsystem = UPCGSubsystem::GetInstance(GetWorld()))
+			{
+				PCGSubsystem->ScheduleGeneric([this]()
+				{
+					this->MarkPackageDirty();
+				return true;
+				}, nullptr, {});
+			}
+		}
+	}
+}
+
+void APCGWorldActor::GetGridGuids(PCGHiGenGrid::FSizeToGuidMap& OutSizeToGuidMap) const
+{
+	FReadScopeLock ReadLock(GridGuidsLock);
+	for (const TPair<uint32, FGuid>& SizeGuid : GridGuids)
+	{
+		OutSizeToGuidMap.Add(SizeGuid.Key, SizeGuid.Value);
+	}
+}
+
 #if WITH_EDITOR
 APCGWorldActor* APCGWorldActor::CreatePCGWorldActor(UWorld* InWorld)
 {
@@ -143,6 +202,19 @@ void APCGWorldActor::OnPartitionGridSizeChanged()
 	}
 }
 
+void APCGWorldActor::PostLoad()
+{
+	Super::PostLoad();
+
+	// Deprecation - PAs used to be placed on a grid with guid=0. If no grid guids are registered,
+	// register a 0 guid now for the current grid size, and this will result in already existing PAs
+	// being reused.
+	if (GridGuids.IsEmpty())
+	{
+		GridGuids.Add(PartitionGridSize, FGuid());
+	}
+}
+
 void APCGWorldActor::PostEditChangeProperty(FPropertyChangedEvent& PropertyChangedEvent)
 {
 	if (PropertyChangedEvent.GetPropertyName() == GET_MEMBER_NAME_CHECKED(APCGWorldActor, PartitionGridSize) 
diff --git a/Engine/Plugins/Experimental/PCG/Source/PCG/Private/Utils/PCGGraphExecutionLogging.cpp b/Engine/Plugins/Experimental/PCG/Source/PCG/Private/Utils/PCGGraphExecutionLogging.cpp
new file mode 100644
--- /dev/null
+++ b/Engine/Plugins/Experimental/PCG/Source/PCG/Private/Utils/PCGGraphExecutionLogging.cpp
@@ -0,0 +1,145 @@
+// Copyright Epic Games, Inc. All Rights Reserved.
+
+#include "PCGGraphExecutionLogging.h"
+
+#include "PCGComponent.h"
+#include "PCGContext.h"
+#include "PCGModule.h"
+
+#include "GameFramework/Actor.h"
+#include "HAL/IConsoleManager.h"
+
+#if WITH_EDITOR
+static TAutoConsoleVariable<bool> CVarGraphExecutionLoggingEnable(
+	TEXT("pcg.GraphExecutionLogging"),
+	false,
+	TEXT("Enables fine grained log of graph execution"));
+#endif
+
+namespace PCGGraphExecutionLogging
+{
+	bool LogEnabled()
+	{
+#if WITH_EDITOR
+		return CVarGraphExecutionLoggingEnable.GetValueOnAnyThread();
+#else
+		return false;
+#endif
+	}
+
+	void LogGridLinkageTaskExecuteStore(const FPCGContext* InContext, int32 InFromGridSize, int32 InToGridSize, const FString& InResourcePath)
+	{
+#if WITH_EDITOR
+		if (!LogEnabled())
+		{
+			return;
+		}
+		check(InContext);
+
+		const FString OwnerName = InContext->SourceComponent.Get() ? (InContext->SourceComponent->GetOwner() ? InContext->SourceComponent->GetOwner()->GetActorLabel() : FString()) : FString();
+		UE_LOG(LogPCG, Log, TEXT("[GRIDLINKING] [%s] STORE. GenerationGridSize=%d, FromGridSize=%d, ToGridSize=%d, Path=%s"), *OwnerName, PCGHiGenGrid::GridToGridSize(InContext->GenerationGrid), InFromGridSize, InToGridSize, *InResourcePath);
+#endif
+	}
+
+	void LogGridLinkageTaskExecuteRetrieve(const FPCGContext* InContext, int32 InFromGridSize, int32 InToGridSize, const FString& InResourcePath)
+	{
+#if WITH_EDITOR
+		if (!LogEnabled())
+		{
+			return;
+		}
+		check(InContext);
+
+		const FString OwnerName = InContext->SourceComponent.Get() ? (InContext->SourceComponent->GetOwner() ? InContext->SourceComponent->GetOwner()->GetActorLabel() : FString()) : FString();
+		UE_LOG(LogPCG, Log, TEXT("[GRIDLINKING] [%s] RETRIEVE. GenerationGridSize=%d, FromGridSize=%d, ToGridSize=%d, Path=%s"), *OwnerName, PCGHiGenGrid::GridToGridSize(InContext->GenerationGrid), InFromGridSize, InToGridSize, *InResourcePath);
+#endif
+	}
+
+	void LogGridLinkageTaskExecuteRetrieveSuccess(const FPCGContext* InContext, const UPCGComponent* InComponent, const FString& InResourcePath, int32 InDataItemCount)
+	{
+#if WITH_EDITOR
+		if (!LogEnabled())
+		{
+			return;
+		}
+		check(InContext);
+
+		const FString OwnerName = InContext->SourceComponent.Get() ? (InContext->SourceComponent->GetOwner() ? InContext->SourceComponent->GetOwner()->GetActorLabel() : FString()) : FString();
+		UE_LOG(LogPCG, Log, TEXT("[GRIDLINKING] [%s] RETRIEVE: SUCCESS. Component=%s Path=%s DataItems=%d"), *OwnerName, *InComponent->GetOwner()->GetActorLabel(), *InResourcePath, InDataItemCount);
+#endif
+	}
+
+	void LogGridLinkageTaskExecuteRetrieveScheduleGraph(const FPCGContext* InContext, const UPCGComponent* InScheduledComponent, const FString& InResourcePath)
+	{
+#if WITH_EDITOR
+		if (!LogEnabled())
+		{
+			return;
+		}
+		check(InContext);
+
+		const FString OwnerName = (InContext->SourceComponent.Get() && InContext->SourceComponent->GetOwner()) ? InContext->SourceComponent->GetOwner()->GetActorLabel() : FString();
+		const FString OtherOwnerName = (InScheduledComponent && InScheduledComponent->GetOwner()) ? InScheduledComponent->GetOwner()->GetActorLabel() : FString();
+		UE_LOG(LogPCG, Log, TEXT("[GRIDLINKING] [%s] RETRIEVE: SCHEDULE GRAPH. Component=%s Path=%s"), *OwnerName, *OtherOwnerName, *InResourcePath);
+#endif
+	}
+
+	void LogGridLinkageTaskExecuteRetrieveWaitOnScheduledGraph(const FPCGContext* InContext, const UPCGComponent* InWaitOnComponent, const FString& InResourcePath)
+	{
+#if WITH_EDITOR
+		if (!LogEnabled())
+		{
+			return;
+		}
+		check(InContext);
+
+		const FString OwnerName = (InContext->SourceComponent.Get() && InContext->SourceComponent->GetOwner()) ? InContext->SourceComponent->GetOwner()->GetActorLabel() : FString();
+		const FString OtherOwnerName = (InWaitOnComponent && InWaitOnComponent->GetOwner()) ? InWaitOnComponent->GetOwner()->GetActorLabel() : FString();
+		UE_LOG(LogPCG, Log, TEXT("[GRIDLINKING] [%s] RETRIEVE: WAIT FOR SCHEDULED GRAPH. Component=%s Path=%s"), *OwnerName, *OtherOwnerName, *InResourcePath);
+#endif
+	}
+	
+	void LogGridLinkageTaskExecuteRetrieveWakeUp(const FPCGContext* InContext, const UPCGComponent* InWokenBy)
+	{
+#if WITH_EDITOR
+		if (!LogEnabled())
+		{
+			return;
+		}
+		check(InContext);
+
+		const FString OwnerName = (InContext->SourceComponent.Get() && InContext->SourceComponent->GetOwner()) ? InContext->SourceComponent->GetOwner()->GetActorLabel() : FString();
+		const FString OtherOwnerName = (InWokenBy && InWokenBy->GetOwner()) ? InWokenBy->GetOwner()->GetActorLabel() : FString();
+		UE_LOG(LogPCG, Log, TEXT("[GRIDLINKING] [%s] RETRIEVE: WOKEN BY Component=%s"), *OwnerName, *OtherOwnerName);
+#endif
+	}
+
+	void LogGridLinkageTaskExecuteRetrieveNoLocalComponent(const FPCGContext* InContext, const FString& InResourcePath)
+	{
+#if WITH_EDITOR
+		if (!LogEnabled())
+		{
+			return;
+		}
+		check(InContext);
+
+		const FString OwnerName = (InContext->SourceComponent.Get() && InContext->SourceComponent->GetOwner()) ? InContext->SourceComponent->GetOwner()->GetActorLabel() : FString();
+		UE_LOG(LogPCG, Warning, TEXT("[GRIDLINKING] [%s] RETRIEVE: FAILED: No overlapping local component found. This may be expected. Path=%s"), *OwnerName, *InResourcePath);
+#endif
+	}
+
+	void LogGridLinkageTaskExecuteRetrieveNoData(const FPCGContext* InContext, const UPCGComponent* InComponent, const FString& InResourcePath)
+	{
+#if WITH_EDITOR
+		if (!LogEnabled())
+		{
+			return;
+		}
+		check(InContext);
+
+		const FString OwnerName = (InContext->SourceComponent.Get() && InContext->SourceComponent->GetOwner()) ? InContext->SourceComponent->GetOwner()->GetActorLabel() : FString();
+		const FString OtherOwnerName = (InComponent && InComponent->GetOwner()) ? InComponent->GetOwner()->GetActorLabel() : FString();
+		UE_LOG(LogPCG, Warning, TEXT("[GRIDLINKING] [%s] RETRIEVE: FAILED: No data found on local component. Component=%s, Path=%s"), *OwnerName, *OtherOwnerName, *InResourcePath);
+#endif
+	}
+}
diff --git a/Engine/Plugins/Experimental/PCG/Source/PCG/Private/Utils/PCGGraphExecutionLogging.h b/Engine/Plugins/Experimental/PCG/Source/PCG/Private/Utils/PCGGraphExecutionLogging.h
new file mode 100644
--- /dev/null
+++ b/Engine/Plugins/Experimental/PCG/Source/PCG/Private/Utils/PCGGraphExecutionLogging.h
@@ -0,0 +1,21 @@
+// Copyright Epic Games, Inc. All Rights Reserved.
+
+#pragma once
+
+#include "Containers/UnrealString.h"
+#include "HAL/Platform.h"
+
+struct FPCGContext;
+class UPCGComponent;
+
+namespace PCGGraphExecutionLogging
+{
+	void LogGridLinkageTaskExecuteStore(const FPCGContext* InContext, int32 InFromLOD, int32 InToLOD, const FString& InResourcePath);
+	void LogGridLinkageTaskExecuteRetrieve(const FPCGContext* InContext, int32 InFromLOD, int32 InToLOD, const FString& InResourcePath);
+	void LogGridLinkageTaskExecuteRetrieveSuccess(const FPCGContext* InContext, const UPCGComponent* InComponent, const FString& InResourcePath, int32 InDataItemCount);
+	void LogGridLinkageTaskExecuteRetrieveScheduleGraph(const FPCGContext* InContext, const UPCGComponent* InScheduledComponent, const FString& InResourcePath);
+	void LogGridLinkageTaskExecuteRetrieveWaitOnScheduledGraph(const FPCGContext* InContext, const UPCGComponent* InWaitOnComponent, const FString& InResourcePath);
+	void LogGridLinkageTaskExecuteRetrieveWakeUp(const FPCGContext* InContext, const UPCGComponent* InWokenBy);
+	void LogGridLinkageTaskExecuteRetrieveNoLocalComponent(const FPCGContext* InContext, const FString& InResourcePath);
+	void LogGridLinkageTaskExecuteRetrieveNoData(const FPCGContext* InContext, const UPCGComponent* InComponent, const FString& InResourcePath);
+}
diff --git a/Engine/Plugins/Experimental/PCG/Source/PCG/Public/Elements/PCGHiGenGridSize.h b/Engine/Plugins/Experimental/PCG/Source/PCG/Public/Elements/PCGHiGenGridSize.h
new file mode 100644
--- /dev/null
+++ b/Engine/Plugins/Experimental/PCG/Source/PCG/Public/Elements/PCGHiGenGridSize.h
@@ -0,0 +1,49 @@
+// Copyright Epic Games, Inc. All Rights Reserved.
+#pragma once
+
+#include "PCGElement.h"
+#include "PCGSettings.h"
+
+#include "PCGHiGenGridSize.generated.h"
+
+/**
+ * Set the execution grid size for downstream nodes. Enables executing a single graph across a hierarchy of grids.
+ */
+UCLASS(BlueprintType, ClassGroup = (Procedural))
+class PCG_API UPCGHiGenGridSizeSettings : public UPCGSettings
+{
+	GENERATED_BODY()
+
+public:
+	EPCGHiGenGrid GetGrid() const { return HiGenGridSize; }
+	uint32 GetGridSize() const { return PCGHiGenGrid::GridToGridSize(HiGenGridSize); }
+
+	//~Begin UPCGSettings interface
+#if WITH_EDITOR
+	virtual FName GetDefaultNodeName() const override;
+	virtual FText GetDefaultNodeTitle() const override;
+	virtual FText GetNodeTooltipText() const override;
+	virtual EPCGSettingsType GetType() const override { return EPCGSettingsType::Spatial; }
+	virtual bool HasDynamicPins() const override { return true; }
+#endif
+	virtual EPCGDataType GetCurrentPinTypes(const UPCGPin* InPin) const override;
+
+protected:
+	virtual TArray<FPCGPinProperties> InputPinProperties() const override;
+	virtual TArray<FPCGPinProperties> OutputPinProperties() const override;
+	virtual FPCGElementPtr CreateElement() const override;
+	virtual FName AdditionalTaskName() const override;
+#if WITH_EDITOR
+	virtual bool IsStructuralProperty(const FName& InPropertyName) const override;
+#endif
+	//~End UPCGSettings interface
+
+	UPROPERTY(BlueprintReadWrite, EditAnywhere, Category = Settings)
+	EPCGHiGenGrid HiGenGridSize = EPCGHiGenGrid::Grid256;
+};
+
+class FPCGHiGenGridSizeElement : public FSimplePCGElement
+{
+protected:
+	virtual bool ExecuteInternal(FPCGContext* Context) const override;
+};
diff --git a/Engine/Plugins/Experimental/PCG/Source/PCG/Public/Helpers/PCGHelpers.h b/Engine/Plugins/Experimental/PCG/Source/PCG/Public/Helpers/PCGHelpers.h
--- a/Engine/Plugins/Experimental/PCG/Source/PCG/Public/Helpers/PCGHelpers.h
+++ b/Engine/Plugins/Experimental/PCG/Source/PCG/Public/Helpers/PCGHelpers.h
@@ -2,6 +2,8 @@
 
 #pragma once
 
+#include "PCGCommon.h"
+
 #include "Math/Box.h"
 
 class AActor;
@@ -9,6 +11,7 @@ class APCGWorldActor;
 class ALandscape;
 class ALandscapeProxy;
 class UPCGComponent;
+class UPCGGraph;
 class UWorld;
 
 namespace PCGHelpers
@@ -57,4 +60,7 @@ namespace PCGHelpers
 	* Use the bCheckHierarchy flag to true to go up the object hierarchy if you want to check for this situation.
 	*/
 	PCG_API bool IsNewObjectAndNotDefault(const UObject* InObject, bool bCheckHierarchy = false);
+
+	/** If hierarchical generation is enabled, returns all relevant grid sizes for the graph, otherwise returns partition grid size from world actor. */
+	PCG_API bool GetGenerationGridSizes(const UPCGGraph* InGraph, const APCGWorldActor* InWorldActor, PCGHiGenGrid::FSizeArray& OutGridSizes);
 };
diff --git a/Engine/Plugins/Experimental/PCG/Source/PCG/Public/PCGActorAndComponentMapping.h b/Engine/Plugins/Experimental/PCG/Source/PCG/Public/PCGActorAndComponentMapping.h
--- a/Engine/Plugins/Experimental/PCG/Source/PCG/Public/PCGActorAndComponentMapping.h
+++ b/Engine/Plugins/Experimental/PCG/Source/PCG/Public/PCGActorAndComponentMapping.h
@@ -118,8 +118,8 @@ private:
 	/** Octree tracking all partitioned pcg components */
 	FPCGComponentOctreeAndMap PartitionedOctree;
 
-	/** Grid mapping for partition actors. We can only have 1 partition actor per grid cell. */
-	TMap<FIntVector, TObjectPtr<APCGPartitionActor>> PartitionActorsMap;
+	/** Mapping from grid size and grid coords to partition actor. We can only have 1 partition actor per grid cell. */
+	TMap<uint32, TMap<FIntVector, TObjectPtr<APCGPartitionActor>>> PartitionActorsMap;
 	mutable FRWLock PartitionActorsMapLock;
 
 	/** Mapping between original components and its overlapping partition actors. */
diff --git a/Engine/Plugins/Experimental/PCG/Source/PCG/Public/PCGCommon.h b/Engine/Plugins/Experimental/PCG/Source/PCG/Public/PCGCommon.h
--- a/Engine/Plugins/Experimental/PCG/Source/PCG/Public/PCGCommon.h
+++ b/Engine/Plugins/Experimental/PCG/Source/PCG/Public/PCGCommon.h
@@ -118,4 +118,50 @@ namespace PCGObjectMetadata
 namespace PCGFeatureSwitches
 {
 	extern PCG_API TAutoConsoleVariable<bool> CVarCheckSamplerMemory;
-}
\ No newline at end of file
+}
+
+/** Describes one or more target execution grids. */
+UENUM(meta = (Bitflags))
+enum class EPCGHiGenGrid : uint32
+{
+	Uninitialized = 0 UMETA(Hidden),
+
+	// NOTE: When adding new grids, increment PCGHiGenGrid::NumGridValues below.
+	Grid32 = 32 UMETA(DisplayName = "3200"),
+	Grid64 = 64 UMETA(DisplayName = "6400"),
+	Grid128 = 128 UMETA(DisplayName = "12800"),
+	Grid256 = 256 UMETA(DisplayName = "25600"),
+	Grid512 = 512 UMETA(DisplayName = "51200"),
+	Grid1024 = 1024 UMETA(DisplayName = "102400"),
+	Grid2048 = 2048 UMETA(DisplayName = "204800"),
+	
+	GridMin = Grid32 UMETA(Hidden),
+	GridMax = Grid2048 UMETA(Hidden),
+
+	// Should execute once rather than executing on any grid
+	Unbounded = 2 * GridMax,
+
+	// Flag for grid not known yet. Represents the default grid size (which can be changed at runtime on PCGWorldActor)
+	GenerationDefault = 2 * Unbounded UMETA(Hidden),
+};
+ENUM_CLASS_FLAGS(EPCGHiGenGrid);
+
+namespace PCGHiGenGrid
+{
+	// Number of unique values of EPCGHiGenGrid, const so it can be used for the inline allocator below.
+	constexpr uint32 NumGridValues = 10;
+
+	// Alias for array which is allocated on the stack (we have a strong idea of the max required elements).
+	using FSizeArray = TArray<uint32, TInlineAllocator<PCGHiGenGrid::NumGridValues>>;
+
+	// Alias for grid size to guid map allocated on the stack, which is unlikely to have a large number of elements.
+	using FSizeToGuidMap = TMap<uint32, FGuid, TInlineSetAllocator<32>>;
+
+	PCG_API bool IsValidGridSize(uint32 InGridSize);
+	PCG_API bool IsValidGrid(EPCGHiGenGrid InGrid);
+	PCG_API uint32 GridToGridSize(EPCGHiGenGrid InGrid);
+	PCG_API EPCGHiGenGrid GridSizeToGrid(uint32 InGridSize);
+
+	PCG_API uint32 UninitializedGridSize();
+	PCG_API uint32 UnboundedGridSize();
+}
diff --git a/Engine/Plugins/Experimental/PCG/Source/PCG/Public/PCGComponent.h b/Engine/Plugins/Experimental/PCG/Source/PCG/Public/PCGComponent.h
--- a/Engine/Plugins/Experimental/PCG/Source/PCG/Public/PCGComponent.h
+++ b/Engine/Plugins/Experimental/PCG/Source/PCG/Public/PCGComponent.h
@@ -8,6 +8,8 @@
 #include "PCGSettings.h"
 #include "Utils/PCGExtraCapture.h"
 
+#include "Containers/Map.h"
+
 #include "PCGComponent.generated.h"
 
 namespace EEndPlayReason { enum Type : int; }
@@ -126,6 +128,8 @@ public:
 	UFUNCTION(BlueprintCallable, Category = PCG)
 	void GenerateLocal(bool bForce);
 
+	FPCGTaskId GenerateLocalGetTaskId(bool bForce);
+
 	/** Cleans up the generation from a local (vs. remote) standpoint. Will not be replicated. Will be delayed. */
 	UFUNCTION(BlueprintCallable, Category = PCG)
 	void CleanupLocal(bool bRemoveComponents, bool bSave = false);
@@ -156,6 +160,18 @@ public:
 	UFUNCTION(BlueprintCallable, Category = PCG)
 	AActor* ClearPCGLink(UClass* TemplateActor = nullptr);
 
+	uint32 GetGenerationGridSize() const { return GenerationGridSize; }
+	void SetGenerationGridSize(uint32 InGenerationGridSize) { GenerationGridSize = InGenerationGridSize; }
+
+	/** Store data with a resource key that identifies the pin. */
+	void StoreOutputDataForPin(const FString& InResourceKey, const FPCGDataCollection& InData);
+
+	/** Lookup data using a resource key that identifies the pin. */
+	const FPCGDataCollection* RetrieveOutputDataForPin(const FString& InResourceKey);
+
+	/** Clear any data stored for any pins. */
+	void ClearPerPinGeneratedOutput();
+
 	UPROPERTY(BlueprintReadWrite, EditAnywhere, Category = Settings)
 	EPCGComponentInput InputType = EPCGComponentInput::Actor;
 
@@ -199,7 +215,7 @@ public:
 	FOnPCGGraphCleaned OnPCGGraphCleanedDelegate;
 #endif
 
-	/** Can specify a list of functions from the owner of this component to be called when generation is done, in order. 
+	/** Can specify a list of functions from the owner of this component to be called when generation is done, in order.
 	*   Need to take (and only take) a PCGDataCollection as parameter and with "CallInEditor" flag enabled.
 	*/
 	UPROPERTY(BlueprintReadWrite, EditAnywhere, AdvancedDisplay, Category = Properties)
@@ -268,6 +284,9 @@ protected:
 	UPROPERTY(BlueprintReadOnly, VisibleAnywhere, Category = PCG, Instanced, meta = (NoResetToDefault))
 	TObjectPtr<UPCGGraphInstance> GraphInstance;
 
+	UPROPERTY(Transient, VisibleAnywhere, AdvancedDisplay, Category = Properties)
+	uint32 GenerationGridSize = PCGHiGenGrid::UnboundedGridSize();
+
 #if WITH_EDITORONLY_DATA
 	UPROPERTY()
 	TObjectPtr<UPCGGraph> Graph_DEPRECATED;
@@ -381,6 +400,12 @@ private:
 	UPROPERTY()
 	FPCGDataCollection GeneratedGraphOutput;
 
+	/** If any graph edges cross execution grid sizes, data on the edge is stored / retrieved from this map. */
+	UPROPERTY(Transient, VisibleAnywhere, Category = Properties, AdvancedDisplay)
+	TMap<FString, FPCGDataCollection> PerPinGeneratedOutput;
+
+	mutable FRWLock PerPinGeneratedOutputLock;
+
 	FPCGTaskId CurrentGenerationTask = InvalidPCGTaskId;
 	FPCGTaskId CurrentCleanupTask = InvalidPCGTaskId;
 
diff --git a/Engine/Plugins/Experimental/PCG/Source/PCG/Public/PCGContext.h b/Engine/Plugins/Experimental/PCG/Source/PCG/Public/PCGContext.h
--- a/Engine/Plugins/Experimental/PCG/Source/PCG/Public/PCGContext.h
+++ b/Engine/Plugins/Experimental/PCG/Source/PCG/Public/PCGContext.h
@@ -68,6 +68,9 @@ struct PCG_API FPCGContext
 	EPCGExecutionPhase CurrentPhase = EPCGExecutionPhase::NotExecuted;
 	int32 BypassedOutputCount = 0;
 
+	/** Used by grid size data linkages to store/retrieve data. */
+	EPCGHiGenGrid GenerationGrid = EPCGHiGenGrid::Uninitialized;
+
 	/** The current call stack. */
 	const FPCGStack* Stack = nullptr;
 
diff --git a/Engine/Plugins/Experimental/PCG/Source/PCG/Public/PCGGraph.h b/Engine/Plugins/Experimental/PCG/Source/PCG/Public/PCGGraph.h
--- a/Engine/Plugins/Experimental/PCG/Source/PCG/Public/PCGGraph.h
+++ b/Engine/Plugins/Experimental/PCG/Source/PCG/Public/PCGGraph.h
@@ -2,6 +2,7 @@
 
 #pragma once
 
+#include "PCGCommon.h"
 #include "PCGNode.h"
 #include "PCGSettings.h"
 
@@ -89,7 +90,7 @@ public:
 
 	bool IsInstance() const;
 
-	/** A graph interface is equivalent to another graph interface if they are the same (same ptr), or if they have the same graph. Will be overriden when graph instance supports overrides. */
+	/** A graph interface is equivalent to another graph interface if they are the same (same ptr), or if they have the same graph. Will be overridden when graph instance supports overrides. */
 	virtual bool IsEquivalent(const UPCGGraphInterface* Other) const;
 
 #if WITH_EDITOR
@@ -128,6 +129,11 @@ public:
 	virtual const UPCGGraph* GetGraph() const override { return this; }
 	/** ~End UPCGGraphInterface interface */
 
+	// Default grid size for generation. For hierarchical generation, nodes outside of grid size graph ranges will generate on this grid.
+	EPCGHiGenGrid GetDefaultGrid() const { ensure(IsHierarchicalGenerationEnabled()); return HiGenGridSize; }
+	uint32 GetDefaultGridSize() const { ensure(IsHierarchicalGenerationEnabled()); return PCGHiGenGrid::GridToGridSize(HiGenGridSize); }
+	bool IsHierarchicalGenerationEnabled() const { return bHierarchicalGenerationEnabled; }
+
 #if WITH_EDITORONLY_DATA
 	UPROPERTY(BlueprintReadWrite, EditAnywhere, Category = AssetInfo, AssetRegistrySearchable)
 	bool bExposeToLibrary = false;
@@ -193,6 +199,9 @@ public:
 	bool RemoveInboundEdges(UPCGNode* InNode, const FName& InboundLabel);
 	bool RemoveOutboundEdges(UPCGNode* InNode, const FName& OutboundLabel);
 
+	/** Determine the relevant grid sizes by inspecting all HiGenGridSize nodes. */
+	void GetGridSizes(PCGHiGenGrid::FSizeArray& OutGridSizes) const;
+
 #if WITH_EDITOR
 	void DisableNotificationsForEditor();
 	void EnableNotificationsForEditor();
@@ -214,10 +223,16 @@ public:
 	void GetTrackedTagsToSettings(FPCGTagToSettingsMap& OutTagsToSettings, TArray<TObjectPtr<const UPCGGraph>>& OutVisitedGraphs) const;
 #endif
 
+	/** Size of grid on which this node should be executed. Nodes execute at the minimum of all input grid sizes. */
+	uint32 GetNodeGenerationGridSize(const UPCGNode* InNode, uint32 InDefaultGridSize) const;
+
 protected:
 	void OnNodeAdded(UPCGNode* InNode);
 	void OnNodeRemoved(UPCGNode* InNode);
 
+	/** Calculates node grid size. Not thread safe, called within write lock. */
+	uint32 CalculateNodeGridSizeRecursive_Unsafe(const UPCGNode* InNode, uint32 InDefaultGridSize) const;
+
 	UPROPERTY(BlueprintReadOnly, VisibleAnywhere, Category = Graph)
 	TArray<TObjectPtr<UPCGNode>> Nodes;
 
@@ -238,6 +253,16 @@ protected:
 	UPROPERTY(EditAnywhere, Category = Instance, meta = (DisplayName = "Parameters", NoResetToDefault))
 	FInstancedPropertyBag UserParameters;
 
+	UPROPERTY(EditAnywhere, Category = Settings)
+	bool bHierarchicalGenerationEnabled = false;
+
+	UPROPERTY(EditAnywhere, Category = Settings, meta = (DisplayName = "HiGen Default Grid Size", EditCondition = "bHierarchicalGenerationEnabled"))
+	EPCGHiGenGrid HiGenGridSize = EPCGHiGenGrid::Grid256;
+
+	/** Execution grid size for nodes. */
+	mutable TMap<const UPCGNode*, uint32> NodeToGridSize;
+	mutable FRWLock NodeToGridSizeLock;
+
 public:
 	virtual const FInstancedPropertyBag* GetUserParametersStruct() const override { return &UserParameters; }
 
diff --git a/Engine/Plugins/Experimental/PCG/Source/PCG/Public/PCGManagedResource.h b/Engine/Plugins/Experimental/PCG/Source/PCG/Public/PCGManagedResource.h
--- a/Engine/Plugins/Experimental/PCG/Source/PCG/Public/PCGManagedResource.h
+++ b/Engine/Plugins/Experimental/PCG/Source/PCG/Public/PCGManagedResource.h
@@ -32,6 +32,7 @@ public:
 	virtual bool ReleaseIfUnused(TSet<TSoftObjectPtr<AActor>>& OutActorsToDelete);
 
 	virtual void MarkAsUsed() { bIsMarkedUnused = false; }
+	// Ensure may fire if multiple executions of the graph are happening in parallel/overlapping
 	virtual void MarkAsReused() { ensure(bIsMarkedUnused); bIsMarkedUnused = false; }
 	bool IsMarkedUnused() const { return bIsMarkedUnused; }
 
diff --git a/Engine/Plugins/Experimental/PCG/Source/PCG/Public/PCGSubsystem.h b/Engine/Plugins/Experimental/PCG/Source/PCG/Public/PCGSubsystem.h
--- a/Engine/Plugins/Experimental/PCG/Source/PCG/Public/PCGSubsystem.h
+++ b/Engine/Plugins/Experimental/PCG/Source/PCG/Public/PCGSubsystem.h
@@ -201,7 +201,7 @@ private:
 	};
 
 	FPCGTaskId ProcessGraph(UPCGComponent* Component, const FBox& InPreviousBounds, const FBox& InNewBounds, EOperation InOperation, bool bSave);
-	void CreatePartitionActorsWithinBounds(const FBox& InBounds);
+	void CreatePartitionActorsWithinBounds(const FBox& InBounds, const PCGHiGenGrid::FSizeArray& InGridSizes);
 
 	FPCGNodeVisualLogs NodeVisualLogs;
 #endif // WITH_EDITOR
diff --git a/Engine/Plugins/Experimental/PCG/Source/PCG/Public/PCGWorldActor.h b/Engine/Plugins/Experimental/PCG/Source/PCG/Public/PCGWorldActor.h
--- a/Engine/Plugins/Experimental/PCG/Source/PCG/Public/PCGWorldActor.h
+++ b/Engine/Plugins/Experimental/PCG/Source/PCG/Public/PCGWorldActor.h
@@ -2,8 +2,10 @@
 
 #pragma once
 
-#include "GameFramework/Actor.h"
+#include "PCGCommon.h"
 
+#include "GameFramework/Actor.h"
+#include "Misc/Guid.h"
 
 #include "PCGWorldActor.generated.h"
 
@@ -23,6 +25,10 @@ public:
 	virtual void BeginPlay() override;
 	virtual void EndPlay(const EEndPlayReason::Type EndPlayReason) override;
 
+	/* Returns the grid GUIDs used for the partitioned actors, one per grid size. */
+	void CreateGridGuidsIfNecessary(const PCGHiGenGrid::FSizeArray& InGridSizes);
+	void GetGridGuids(PCGHiGenGrid::FSizeToGuidMap& OutSizeToGuidMap) const;
+
 #if WITH_EDITOR	
 	virtual bool CanChangeIsSpatiallyLoadedFlag() const { return false; }
 	virtual bool IsUserManaged() const override { return false; }
@@ -36,12 +42,13 @@ public:
 
 	//~ Begin UObject Interface.
 #if WITH_EDITOR	
+	virtual void PostLoad() override;
 	virtual void PostEditChangeProperty(FPropertyChangedEvent& PropertyChangedEvent) override;
 #endif
 	//~ End UObject Interface.
 
-	/** Size of the grid for PCG partition actors */
-	UPROPERTY(config, EditAnywhere, Category = WorldPartition)
+	/** Size of the PCG partition actor grid for non-hierarchical-generation graphs. */
+	UPROPERTY(config, EditAnywhere, Category = GenerationSettings)
 	uint32 PartitionGridSize;
 
 	/** Contains all the PCG data required to query the landscape complete. Serialized in cooked builds only */
@@ -49,7 +56,7 @@ public:
 	TObjectPtr<UPCGLandscapeCache> LandscapeCacheObject = nullptr;
 
 	/** Disable creation of Partition Actors on the Z axis. Can improve performances if 3D partitioning is not needed. */
-	UPROPERTY(config, EditAnywhere, Category = WorldPartition)
+	UPROPERTY(config, EditAnywhere, Category = GenerationSettings)
 	bool bUse2DGrid = true;
 
 private:
@@ -59,6 +66,11 @@ private:
 #if WITH_EDITOR	
 	void OnPartitionGridSizeChanged();
 #endif
+
+	/** GUIDs of the partitioned actor grids, one per grid size. */
+	UPROPERTY()
+	TMap<uint32, FGuid> GridGuids;
+	mutable FRWLock GridGuidsLock;
 };
 
 #if UE_ENABLE_INCLUDE_ORDER_DEPRECATED_IN_5_2
diff --git a/Engine/Plugins/Experimental/PCG/Source/PCGEditor/Private/PCGEditor.cpp b/Engine/Plugins/Experimental/PCG/Source/PCGEditor/Private/PCGEditor.cpp
--- a/Engine/Plugins/Experimental/PCG/Source/PCGEditor/Private/PCGEditor.cpp
+++ b/Engine/Plugins/Experimental/PCG/Source/PCGEditor/Private/PCGEditor.cpp
@@ -187,7 +187,13 @@ void FPCGEditor::SetComponentAndStackBeingInspected(UPCGComponent* InPCGComponen
 			if (UPCGEditorGraphNodeBase* PCGNode = Cast<UPCGEditorGraphNodeBase>(Node))
 			{
 				// Update now that component has changed. Will fire OnNodeChanged if necessary.
-				PCGNode->UpdateErrorsAndWarnings();
+				EPCGChangeType ChangeType = PCGNode->UpdateErrorsAndWarnings();
+				ChangeType |= PCGNode->UpdateGridSizeVisualisation(InPCGComponent);
+
+				if (ChangeType != EPCGChangeType::None)
+				{
+					PCGNode->ReconstructNode();
+				}
 			}
 		}
 	}
diff --git a/Engine/Plugins/Experimental/PCG/Source/PCGEditor/Private/PCGEditorGraphNodeBase.cpp b/Engine/Plugins/Experimental/PCG/Source/PCGEditor/Private/PCGEditorGraphNodeBase.cpp
--- a/Engine/Plugins/Experimental/PCG/Source/PCGEditor/Private/PCGEditorGraphNodeBase.cpp
+++ b/Engine/Plugins/Experimental/PCG/Source/PCGEditor/Private/PCGEditorGraphNodeBase.cpp
@@ -2,6 +2,7 @@
 
 #include "PCGEditorGraphNodeBase.h"
 
+#include "PCGComponent.h"
 #include "PCGEditor.h"
 #include "PCGEditorCommands.h"
 #include "PCGEditorCommon.h"
@@ -11,6 +12,8 @@
 #include "PCGGraph.h"
 #include "PCGPin.h"
 #include "PCGSubsystem.h"
+#include "PCGWorldActor.h"
+#include "Helpers/PCGHelpers.h"
 
 #include "GraphEditorActions.h"
 #include "ToolMenu.h"
@@ -279,7 +282,7 @@ void UPCGEditorGraphNodeBase::OnNodeChanged(UPCGNode* InNode, EPCGChangeType Cha
 			}
 		}
 
-		UpdateErrorsAndWarnings();
+		ChangeType |= UpdateErrorsAndWarnings();
 
 		if (!!(ChangeType & (EPCGChangeType::Structural | EPCGChangeType::Node | EPCGChangeType::Edge | EPCGChangeType::Cosmetic)))
 		{
@@ -288,9 +291,9 @@ void UPCGEditorGraphNodeBase::OnNodeChanged(UPCGNode* InNode, EPCGChangeType Cha
 	}
 }
 
-void UPCGEditorGraphNodeBase::UpdateErrorsAndWarnings()
+EPCGChangeType UPCGEditorGraphNodeBase::UpdateErrorsAndWarnings()
 {
-	bool bNodeChanged = false;
+	EPCGChangeType ChangeType = EPCGChangeType::None;
 
 	// Pull current errors/warnings state from PCG subsystem.
 	if (const UPCGSubsystem* Subsystem = UPCGSubsystem::GetActiveEditorInstance())
@@ -321,13 +324,78 @@ void UPCGEditorGraphNodeBase::UpdateErrorsAndWarnings()
 			ErrorType = 0;
 		}
 
-		bNodeChanged = (bHasCompilerMessage != bOldHasCompilerMessage) || (ErrorType != OldErrorType) || (ErrorMsg != OldErrorMsg);
+		if ((bHasCompilerMessage != bOldHasCompilerMessage) || (ErrorType != OldErrorType) || (ErrorMsg != OldErrorMsg))
+		{
+			ChangeType = EPCGChangeType::Cosmetic;
+		}
+	}
+
+	return ChangeType;
+}
+
+EPCGChangeType UPCGEditorGraphNodeBase::UpdateGridSizeVisualisation(UPCGComponent* InComponentBeingDebugged)
+{
+	if (!PCGNode)
+	{
+		return EPCGChangeType::None;
 	}
 
-	if (bNodeChanged)
+	EPCGChangeType ChangeType = EPCGChangeType::None;
+	const uint32 InspectingGridSize = (InComponentBeingDebugged && InComponentBeingDebugged->IsLocalComponent()) ?
+		InComponentBeingDebugged->GetGenerationGridSize() : PCGHiGenGrid::UninitializedGridSize();
+
+	if (InspectingGridSize == PCGHiGenGrid::UninitializedGridSize())
+	{
+		if (IsDisplayAsDisabledForced())
+		{
+			SetForceDisplayAsDisabled(false);
+			ChangeType |= EPCGChangeType::Cosmetic;
+		}
+
+		if (IsHighlighted())
+		{
+			SetIsHighlighted(false);
+			ChangeType |= EPCGChangeType::Cosmetic;
+		}
+	}
+	else if (UPCGGraph* Graph = PCGNode->GetGraph())
 	{
-		OnNodeChangedDelegate.ExecuteIfBound();
+		uint32 DefaultGridSize;
+		UPCGSubsystem* Subsystem = UPCGSubsystem::GetActiveEditorInstance();
+		if (Graph->IsHierarchicalGenerationEnabled())
+		{
+			DefaultGridSize = Graph->GetDefaultGridSize();
+		}
+		else if (Subsystem && Subsystem->GetPCGWorldActor())
+		{
+			DefaultGridSize = Subsystem->GetPCGWorldActor()->PartitionGridSize;
+		}
+		else
+		{
+			// Fallback, should not be hit
+			ensure(false);
+			DefaultGridSize = APCGWorldActor::DefaultPartitionGridSize;
+		}
+		const uint32 NodeGridSize = Graph->GetNodeGenerationGridSize(PCGNode, DefaultGridSize);
+
+		// Disable nodes that are on a smaller grid
+		const bool bForceDisable = NodeGridSize < InspectingGridSize;
+		if (IsDisplayAsDisabledForced() != bForceDisable)
+		{
+			SetForceDisplayAsDisabled(bForceDisable);
+			ChangeType |= EPCGChangeType::Cosmetic;
+		}
+
+		// If node is on larger grid than current, highlight it to indicate that its data is available for use
+		const bool bHighlight = NodeGridSize > InspectingGridSize;
+		if (IsHighlighted() != bHighlight)
+		{
+			SetIsHighlighted(bHighlight);
+			ChangeType |= EPCGChangeType::Cosmetic;
+		}
 	}
+
+	return ChangeType;
 }
 
 bool UPCGEditorGraphNodeBase::ShouldDrawCompact() const
@@ -456,7 +524,7 @@ FLinearColor UPCGEditorGraphNodeBase::GetNodeBodyTintColor() const
 		}
 	}
 
-	return Super::GetNodeBodyTintColor();
+	return IsHighlighted() ? GetDefault<UPCGEditorSettings>()->HighlightedNodeBodyTintColor : Super::GetNodeBodyTintColor();
 }
 
 FEdGraphPinType UPCGEditorGraphNodeBase::GetPinType(const UPCGPin* InPin)
diff --git a/Engine/Plugins/Experimental/PCG/Source/PCGEditor/Private/PCGEditorGraphNodeBase.h b/Engine/Plugins/Experimental/PCG/Source/PCGEditor/Private/PCGEditorGraphNodeBase.h
--- a/Engine/Plugins/Experimental/PCG/Source/PCGEditor/Private/PCGEditorGraphNodeBase.h
+++ b/Engine/Plugins/Experimental/PCG/Source/PCGEditor/Private/PCGEditorGraphNodeBase.h
@@ -56,11 +56,18 @@ public:
 	void DisableDeferredReconstruct();
 
 	/** Pulls current errors/warnings state from PCG subsystem. */
-	void UpdateErrorsAndWarnings();
+	EPCGChangeType UpdateErrorsAndWarnings();
+
+	/** If the currently inspected grid size is smaller than the grid size of this node, display transparent. */
+	EPCGChangeType UpdateGridSizeVisualisation(class UPCGComponent* InComponentBeingDebugged);
 
 	/** Puts node title on node body, reducing overall node size */
 	bool ShouldDrawCompact() const;
 
+	// Highlighted nodes draw with a light tint. Currently used to indicate nodes from hi-gen grids that are available for use.
+	void SetIsHighlighted(bool bInIsHighlighted) { bIsHighlighted = bInIsHighlighted; }
+	bool IsHighlighted() const { return bIsHighlighted; }
+
 	DECLARE_DELEGATE(FOnPCGEditorGraphNodeChanged);
 	FOnPCGEditorGraphNodeChanged OnNodeChangedDelegate;
 
@@ -88,4 +95,5 @@ protected:
 	bool bDeferredReconstruct = false;
 	bool bDisableReconstructFromNode = false;
 	bool bIsInspected = false;
+	bool bIsHighlighted = false;
 };
diff --git a/Engine/Plugins/Experimental/PCG/Source/PCGEditor/Private/PCGEditorGraphSchema.cpp b/Engine/Plugins/Experimental/PCG/Source/PCGEditor/Private/PCGEditorGraphSchema.cpp
--- a/Engine/Plugins/Experimental/PCG/Source/PCGEditor/Private/PCGEditorGraphSchema.cpp
+++ b/Engine/Plugins/Experimental/PCG/Source/PCGEditor/Private/PCGEditorGraphSchema.cpp
@@ -647,24 +647,28 @@ void FPCGEditorConnectionDrawingPolicy::DetermineWiringStyle(UEdGraphPin* Output
 		Params.WireColor = GetDefault<UPCGEditorSettings>()->GetPinColor(OutputPin->PinType);
 	}
 
-	// Desaturate and connection if the node is disabled and the data on this wire won't be used
+	// Desaturate connection if downstream node is disabled or if the data on this wire won't be used
 	if (InputPin && OutputPin)
 	{
 		const UPCGEditorGraphNodeBase* EditorNode = CastChecked<const UPCGEditorGraphNodeBase>(InputPin->GetOwningNode());
 		const UPCGNode* PCGNode = EditorNode ? EditorNode->GetPCGNode() : nullptr;
 		const UPCGPin* PCGPin = PCGNode ? PCGNode->GetInputPin(InputPin->GetFName()) : nullptr;
 		const UPCGEditorGraphNodeBase* UpstreamEditorNode = CastChecked<const UPCGEditorGraphNodeBase>(OutputPin->GetOwningNode());
+		const UPCGEditorGraphNodeBase* DownstreamEditorNode = CastChecked<const UPCGEditorGraphNodeBase>(InputPin->GetOwningNode());
 
-		if (PCGPin && UpstreamEditorNode)
+		if (PCGPin && UpstreamEditorNode && DownstreamEditorNode)
 		{
+			const bool bDownstreamNodeForceDisabled = DownstreamEditorNode->IsDisplayAsDisabledForced() && !DownstreamEditorNode->IsHighlighted();
+
 			// Look for the PCG edge that correlates with passed in (OutputPin, InputPin) edge
 			const TObjectPtr<UPCGEdge>* PCGEdge = PCGPin->Edges.FindByPredicate([UpstreamEditorNode, OutputPin](const UPCGEdge* ConnectedPCGEdge)
 			{
 				return UpstreamEditorNode->GetPCGNode() == ConnectedPCGEdge->InputPin->Node && ConnectedPCGEdge->InputPin->Properties.Label == OutputPin->GetFName();
 			});
+			const bool bDownstreamNodeDoesNotUseData = PCGEdge && !PCGNode->IsEdgeUsedByNodeExecution(*PCGEdge);
 
 			// If edge found and is not used, gray it out
-			if (PCGEdge && !PCGNode->IsEdgeUsedByNodeExecution(*PCGEdge))
+			if (bDownstreamNodeForceDisabled || bDownstreamNodeDoesNotUseData)
 			{
 				Params.WireColor = Params.WireColor.Desaturate(0.7f);
 			}
diff --git a/Engine/Plugins/Experimental/PCG/Source/PCGEditor/Private/PCGEditorSettings.cpp b/Engine/Plugins/Experimental/PCG/Source/PCGEditor/Private/PCGEditorSettings.cpp
--- a/Engine/Plugins/Experimental/PCG/Source/PCGEditor/Private/PCGEditorSettings.cpp
+++ b/Engine/Plugins/Experimental/PCG/Source/PCGEditor/Private/PCGEditorSettings.cpp
@@ -11,6 +11,7 @@ UPCGEditorSettings::UPCGEditorSettings(const FObjectInitializer& ObjectInitializ
 {
 	DefaultNodeColor = FLinearColor(0.4f, 0.62f, 1.0f);
 	InstancedNodeBodyTintColor = FLinearColor(0.5f, 0.5f, 0.5f);
+	HighlightedNodeBodyTintColor = FLinearColor(0.25f, 0.25f, 0.5f, 0.33f);
 	InputOutputNodeColor = FLinearColor(1.0f, 0.0f, 0.0f);
 	SetOperationNodeColor = FLinearColor(1.0f, 0.2f, 1.0f);
 	DensityOperationNodeColor = FLinearColor(0.2f, 0.59f, 0.99f);
diff --git a/Engine/Plugins/Experimental/PCG/Source/PCGEditor/Private/SPCGEditorGraphNode.cpp b/Engine/Plugins/Experimental/PCG/Source/PCGEditor/Private/SPCGEditorGraphNode.cpp
--- a/Engine/Plugins/Experimental/PCG/Source/PCGEditor/Private/SPCGEditorGraphNode.cpp
+++ b/Engine/Plugins/Experimental/PCG/Source/PCGEditor/Private/SPCGEditorGraphNode.cpp
@@ -373,7 +373,9 @@ void SPCGEditorGraphNode::UpdateGraphNode()
 
 const FSlateBrush* SPCGEditorGraphNode::GetNodeBodyBrush() const
 {
-	if (PCGEditorGraphNode && PCGEditorGraphNode->GetPCGNode() && PCGEditorGraphNode->GetPCGNode()->IsInstance())
+	const bool bNeedsTint = PCGEditorGraphNode &&
+		((PCGEditorGraphNode->GetPCGNode() && PCGEditorGraphNode->GetPCGNode()->IsInstance()) || PCGEditorGraphNode->IsHighlighted());
+	if (bNeedsTint)
 	{
 		return FAppStyle::GetBrush("Graph.Node.TintedBody");
 	}
diff --git a/Engine/Plugins/Experimental/PCG/Source/PCGEditor/Public/PCGEditorSettings.h b/Engine/Plugins/Experimental/PCG/Source/PCGEditor/Public/PCGEditorSettings.h
--- a/Engine/Plugins/Experimental/PCG/Source/PCGEditor/Public/PCGEditorSettings.h
+++ b/Engine/Plugins/Experimental/PCG/Source/PCGEditor/Public/PCGEditorSettings.h
@@ -29,6 +29,10 @@ public:
 	UPROPERTY(EditAnywhere, config, Category = Node, meta = (HideAlphaChannel))
 	FLinearColor InstancedNodeBodyTintColor;
 
+	/** Highlighted node body tint color, used to indicate available data sources in hierarchical graphs */
+	UPROPERTY(EditAnywhere, config, Category = Node)
+	FLinearColor HighlightedNodeBodyTintColor;
+
 	/** Color used for input & output nodes */
 	UPROPERTY(EditAnywhere, config, Category = Node, meta = (HideAlphaChannel))
 	FLinearColor InputOutputNodeColor;
