[PCG] Disallow SingleData and MultipleConnection on the same pin

diff --git a/Engine/Plugins/Experimental/PCG/Source/PCG/Private/Elements/PCGExecuteBlueprint.cpp b/Engine/Plugins/Experimental/PCG/Source/PCG/Private/Elements/PCGExecuteBlueprint.cpp
--- a/Engine/Plugins/Experimental/PCG/Source/PCG/Private/Elements/PCGExecuteBlueprint.cpp
+++ b/Engine/Plugins/Experimental/PCG/Source/PCG/Private/Elements/PCGExecuteBlueprint.cpp
@@ -120,6 +120,18 @@ void UPCGBlueprintElement::PostEditChangeProperty(FPropertyChangedEvent& Propert
 {
 	Super::PostEditChangeProperty(PropertyChangedEvent);
 
+	// Disable multiple connections on single data pins
+	if (PropertyChangedEvent.Property && PropertyChangedEvent.Property->GetFName() == GET_MEMBER_NAME_CHECKED(FPCGPinProperties, bAllowMultipleData))
+	{
+		for (FPCGPinProperties& Pin : CustomInputPins)
+		{
+			if (!Pin.bAllowMultipleData)
+			{
+				Pin.bAllowMultipleConnections = false;
+			}
+		}
+	}
+
 	// Since we don't really know what changed, let's just rebuild our data dependencies
 	DataDependencies = PCGBlueprintHelper::GetDataDependencies(this, DependencyParsingDepth);
 
diff --git a/Engine/Plugins/Experimental/PCG/Source/PCG/Private/PCGInputOutputSettings.cpp b/Engine/Plugins/Experimental/PCG/Source/PCG/Private/PCGInputOutputSettings.cpp
--- a/Engine/Plugins/Experimental/PCG/Source/PCG/Private/PCGInputOutputSettings.cpp
+++ b/Engine/Plugins/Experimental/PCG/Source/PCG/Private/PCGInputOutputSettings.cpp
@@ -71,9 +71,8 @@ void UPCGGraphInputOutputSettings::PostLoad()
 TArray<FPCGPinProperties> UPCGGraphInputOutputSettings::GetPinProperties() const
 {
 	TArray<FPCGPinProperties> PinProperties;
-	const bool bIsInputPin = bIsInput;
 	const EPCGDataType DefaultPinDataType = bIsInput ? EPCGDataType::Spatial : EPCGDataType::Any;
-	Algo::Transform(StaticLabels(), PinProperties, [bIsInputPin, DefaultPinDataType](const FLabelAndTooltip& InLabelAndTooltip) {
+	Algo::Transform(StaticLabels(), PinProperties, [DefaultPinDataType](const FLabelAndTooltip& InLabelAndTooltip) {
 		FPCGPinProperties Res = FPCGPinProperties(InLabelAndTooltip.Label, DefaultPinDataType, /*bMultiConnections=*/true, /*bMultiData=*/true, InLabelAndTooltip.Tooltip);
 		Res.bAdvancedPin = true;
 		return Res;
@@ -107,9 +106,8 @@ TArray<FPCGPinProperties> UPCGGraphInputOutputSettings::DefaultInputPinPropertie
 {
 	// It is important for serialization that this is not modified, or it could break existing graphs.
 	TArray<FPCGPinProperties> PinProperties;
-	const bool bIsInputPin = bIsInput;
 	const EPCGDataType DefaultPinDataType = bIsInput ? EPCGDataType::Spatial : EPCGDataType::Any;
-	Algo::Transform(StaticLabels(), PinProperties, [bIsInputPin, DefaultPinDataType](const FLabelAndTooltip& InLabelAndTooltip) {
+	Algo::Transform(StaticLabels(), PinProperties, [DefaultPinDataType](const FLabelAndTooltip& InLabelAndTooltip) {
 		FPCGPinProperties Res = FPCGPinProperties(InLabelAndTooltip.Label, DefaultPinDataType, /*bMultiConnections=*/true, /*bMultiData=*/true, InLabelAndTooltip.Tooltip);
 		Res.bAdvancedPin = true;
 		return Res;
diff --git a/Engine/Plugins/Experimental/PCG/Source/PCG/Public/PCGPin.h b/Engine/Plugins/Experimental/PCG/Source/PCG/Public/PCGPin.h
--- a/Engine/Plugins/Experimental/PCG/Source/PCG/Public/PCGPin.h
+++ b/Engine/Plugins/Experimental/PCG/Source/PCG/Public/PCGPin.h
@@ -30,7 +30,7 @@ struct PCG_API FPCGPinProperties
 	UPROPERTY(BlueprintReadWrite, EditAnywhere, Category = Settings)
 	bool bAllowMultipleData = true;
 
-	UPROPERTY(BlueprintReadWrite, EditAnywhere, Category = Settings)
+	UPROPERTY(BlueprintReadWrite, EditAnywhere, Category = Settings, meta = (EditCondition = "bAllowMultipleData"))
 	bool bAllowMultipleConnections = true;
 
 	/* Advanced pin will be hidden by default in the UI and will be shown only if the user extend the node (in the UI) to see advanced pins. */
