[PCG] Put new Grid Size nodes into a new Hierarchical Generation category

#rnx

diff --git a/Engine/Plugins/Experimental/PCG/Source/PCG/Public/Elements/PCGHiGenGridSize.h b/Engine/Plugins/Experimental/PCG/Source/PCG/Public/Elements/PCGHiGenGridSize.h
--- a/Engine/Plugins/Experimental/PCG/Source/PCG/Public/Elements/PCGHiGenGridSize.h
+++ b/Engine/Plugins/Experimental/PCG/Source/PCG/Public/Elements/PCGHiGenGridSize.h
@@ -23,7 +23,7 @@ public:
 	virtual FName GetDefaultNodeName() const override;
 	virtual FText GetDefaultNodeTitle() const override;
 	virtual FText GetNodeTooltipText() const override;
-	virtual EPCGSettingsType GetType() const override { return EPCGSettingsType::Spatial; }
+	virtual EPCGSettingsType GetType() const override { return EPCGSettingsType::HierarchicalGeneration; }
 	virtual bool HasDynamicPins() const override { return true; }
 #endif
 	virtual EPCGDataType GetCurrentPinTypes(const UPCGPin* InPin) const override;
diff --git a/Engine/Plugins/Experimental/PCG/Source/PCG/Public/PCGSettings.h b/Engine/Plugins/Experimental/PCG/Source/PCG/Public/PCGSettings.h
--- a/Engine/Plugins/Experimental/PCG/Source/PCG/Public/PCGSettings.h
+++ b/Engine/Plugins/Experimental/PCG/Source/PCG/Public/PCGSettings.h
@@ -44,7 +44,8 @@ enum class EPCGSettingsType : uint8
 	Subgraph,
 	Debug,
 	Generic,
-	Param
+	Param,
+	HierarchicalGeneration,
 };
 
 #if WITH_EDITOR
diff --git a/Engine/Plugins/Experimental/PCG/Source/PCGEditor/Private/PCGEditorSettings.cpp b/Engine/Plugins/Experimental/PCG/Source/PCGEditor/Private/PCGEditorSettings.cpp
--- a/Engine/Plugins/Experimental/PCG/Source/PCGEditor/Private/PCGEditorSettings.cpp
+++ b/Engine/Plugins/Experimental/PCG/Source/PCGEditor/Private/PCGEditorSettings.cpp
@@ -23,9 +23,9 @@ UPCGEditorSettings::UPCGEditorSettings(const FObjectInitializer& ObjectInitializ
 	SubgraphNodeColor = FLinearColor(1.0f, 0.05f, 0.05f);
 	ParamDataNodeColor = FLinearColor(1.0f, 0.38f, 0.02f);
 	DebugNodeColor = FLinearColor(1.0f, 0.0f, 1.0f);
+	HierarchicalGenerationNodeColor = FLinearColor(1.0f, 0.4f, 0.0f);
 
 	DefaultPinColor = FLinearColor(0.29f, 0.29f, 0.29f);
-
 	SpatialDataPinColor = FLinearColor(1.0f, 1.0f, 1.0f);
 	ConcreteDataPinColor = FLinearColor(0.45f, 0.38f, 0.96f);
 	PointDataPinColor = FLinearColor(0.05f, 0.25f, 1.0f);
@@ -55,7 +55,7 @@ FLinearColor UPCGEditorSettings::GetColor(UPCGSettings* Settings) const
 		return *Override;
 	}
 	// Otherwise, check against the classes we know
-	else if(Settings->GetType() == EPCGSettingsType::InputOutput)
+	else if (Settings->GetType() == EPCGSettingsType::InputOutput)
 	{
 		return InputOutputNodeColor;
 	}
@@ -99,6 +99,10 @@ FLinearColor UPCGEditorSettings::GetColor(UPCGSettings* Settings) const
 	{
 		return ParamDataNodeColor;
 	}
+	else if (Settings->GetType() == EPCGSettingsType::HierarchicalGeneration)
+	{
+		return HierarchicalGenerationNodeColor;
+	}
 	else
 	{
 		// Finally, we couldn't find any match, so return the default value
diff --git a/Engine/Plugins/Experimental/PCG/Source/PCGEditor/Public/PCGEditorSettings.h b/Engine/Plugins/Experimental/PCG/Source/PCGEditor/Public/PCGEditorSettings.h
--- a/Engine/Plugins/Experimental/PCG/Source/PCGEditor/Public/PCGEditorSettings.h
+++ b/Engine/Plugins/Experimental/PCG/Source/PCGEditor/Public/PCGEditorSettings.h
@@ -77,6 +77,10 @@ public:
 	UPROPERTY(EditAnywhere, config, Category = Node, meta = (HideAlphaChannel))
 	FLinearColor DebugNodeColor;
 
+	/** Color used for hierarchical generation operations */
+	UPROPERTY(EditAnywhere, config, Category = Node, meta = (HideAlphaChannel))
+	FLinearColor HierarchicalGenerationNodeColor;
+
 	/** Default pin color */
 	UPROPERTY(EditAnywhere, config, Category = Node, meta = (HideAlphaChannel))
 	FLinearColor DefaultPinColor;
