Tentative build fix following datatable changes in PCG

#rnx

diff --git a/Engine/Plugins/Experimental/PCG/Source/PCG/Private/Elements/IO/PCGDataTableElement.cpp b/Engine/Plugins/Experimental/PCG/Source/PCG/Private/Elements/IO/PCGDataTableElement.cpp
--- a/Engine/Plugins/Experimental/PCG/Source/PCG/Private/Elements/IO/PCGDataTableElement.cpp
+++ b/Engine/Plugins/Experimental/PCG/Source/PCG/Private/Elements/IO/PCGDataTableElement.cpp
@@ -7,6 +7,8 @@
 #include "Data/PCGPointData.h"
 #include "Elements/IO/PCGExternalDataContext.h"
 
+#include "Engine/DataTable.h"
+
 #define LOCTEXT_NAMESPACE "PCGLoadDataTable"
 
 #if WITH_EDITOR
diff --git a/Engine/Plugins/Experimental/PCG/Source/PCG/Public/Elements/IO/PCGDataTableElement.h b/Engine/Plugins/Experimental/PCG/Source/PCG/Public/Elements/IO/PCGDataTableElement.h
--- a/Engine/Plugins/Experimental/PCG/Source/PCG/Public/Elements/IO/PCGDataTableElement.h
+++ b/Engine/Plugins/Experimental/PCG/Source/PCG/Public/Elements/IO/PCGDataTableElement.h
@@ -4,6 +4,8 @@
 
 #include "PCGExternalData.h"
 
+class UDataTable;
+
 #include "PCGDataTableElement.generated.h"
 
 UCLASS(BlueprintType, ClassGroup = (Procedural))
diff --git a/Engine/Plugins/Experimental/PCG/Source/PCG/Public/Elements/IO/PCGExternalDataContext.h b/Engine/Plugins/Experimental/PCG/Source/PCG/Public/Elements/IO/PCGExternalDataContext.h
--- a/Engine/Plugins/Experimental/PCG/Source/PCG/Public/Elements/IO/PCGExternalDataContext.h
+++ b/Engine/Plugins/Experimental/PCG/Source/PCG/Public/Elements/IO/PCGExternalDataContext.h
@@ -4,6 +4,7 @@
 
 #include "PCGContext.h"
 #include "Data/PCGPointData.h"
+#include "Metadata/PCGAttributePropertySelector.h"
 #include "Metadata/Accessors/IPCGAttributeAccessor.h"
 #include "Metadata/Accessors/PCGAttributeAccessorHelpers.h"
 #include "Metadata/Accessors/PCGAttributeAccessorKeys.h"
