[PCG] Dropdown to select output/input pins in Attributes window

Allows debugging of input pins. Output pins are listed first, followed by *connected* input pins.

diff --git a/Engine/Plugins/Experimental/PCG/Source/PCG/Private/Graph/PCGGraphExecutor.cpp b/Engine/Plugins/Experimental/PCG/Source/PCG/Private/Graph/PCGGraphExecutor.cpp
--- a/Engine/Plugins/Experimental/PCG/Source/PCG/Private/Graph/PCGGraphExecutor.cpp
+++ b/Engine/Plugins/Experimental/PCG/Source/PCG/Private/Graph/PCGGraphExecutor.cpp
@@ -691,7 +691,7 @@ void FPCGGraphExecutor::Execute()
 						if (Task.StackIndex != INDEX_NONE)
 						{
 							const FPCGStack* Stack = Task.StackContext->GetStack(Task.StackIndex);
-							SourceComponent->StoreInspectionData(Stack, Task.Node, CachedOutput);
+							SourceComponent->StoreInspectionData(Stack, Task.Node, TaskInput, CachedOutput);
 						}
 					}
 #endif
@@ -852,7 +852,7 @@ void FPCGGraphExecutor::Execute()
 				if (ActiveTask.StackIndex != INDEX_NONE)
 				{
 					const FPCGStack* Stack = ActiveTask.StackContext->GetStack(ActiveTask.StackIndex);
-					SourceComponent->StoreInspectionData(Stack, ActiveTask.Context->Node, ActiveTask.Context->OutputData);
+					SourceComponent->StoreInspectionData(Stack, ActiveTask.Context->Node, ActiveTask.Context->InputData, ActiveTask.Context->OutputData);
 				}
 			}
 #endif
diff --git a/Engine/Plugins/Experimental/PCG/Source/PCG/Private/PCGComponent.cpp b/Engine/Plugins/Experimental/PCG/Source/PCG/Private/PCGComponent.cpp
--- a/Engine/Plugins/Experimental/PCG/Source/PCG/Private/PCGComponent.cpp
+++ b/Engine/Plugins/Experimental/PCG/Source/PCG/Private/PCGComponent.cpp
@@ -8,6 +8,7 @@
 #include "PCGInputOutputSettings.h"
 #include "PCGManagedResource.h"
 #include "PCGParamData.h"
+#include "PCGPin.h"
 #include "PCGSubsystem.h"
 #include "Data/PCGCollisionShapeData.h"
 #include "Data/PCGDifferenceData.h"
@@ -1730,18 +1731,38 @@ void UPCGComponent::DisableInspection()
 	}
 };
 
-void UPCGComponent::StoreInspectionData(const FPCGStack* InStack, const UPCGNode* InNode, const FPCGDataCollection& InInspectionData)
+void UPCGComponent::StoreInspectionData(const FPCGStack* InStack, const UPCGNode* InNode, const FPCGDataCollection& InInputData, const FPCGDataCollection& InOutputData)
 {
-	if (!IsInspecting())
+	if (!IsInspecting() || !InNode || !ensure(InStack))
 	{
 		return;
 	}
 
-	FString StackFramePath;
-	if (ensure(InStack) && ensure(InStack->CreateStackFramePath(StackFramePath, InNode)))
+	auto StorePinInspectionData = [InStack, InNode](const TArray<TObjectPtr<UPCGPin>>& InPins, const FPCGDataCollection& InData, TMap<FString, FPCGDataCollection>& InOutInspectionCache)
 	{
-		InspectionCache.Add(StackFramePath, InInspectionData);
-	}
+		for (const UPCGPin* Pin : InPins)
+		{
+			FString StackFramePath;
+			if (!ensure(InStack->CreateStackFramePath(StackFramePath, InNode, Pin)))
+			{
+				continue;
+			}
+
+			FPCGDataCollection PinDataCollection;
+			PinDataCollection.TaggedData = InData.GetInputsByPin(Pin->Properties.Label);
+			if (!PinDataCollection.TaggedData.IsEmpty())
+			{
+				InOutInspectionCache.Add(StackFramePath, PinDataCollection);
+			}
+			else
+			{
+				InOutInspectionCache.Remove(StackFramePath);
+			}
+		}
+	};
+
+	StorePinInspectionData(InNode->GetInputPins(), InInputData, InspectionCache);
+	StorePinInspectionData(InNode->GetOutputPins(), InOutputData, InspectionCache);
 }
 
 const FPCGDataCollection* UPCGComponent::GetInspectionData(const FString& InStackPath) const
diff --git a/Engine/Plugins/Experimental/PCG/Source/PCG/Public/PCGComponent.h b/Engine/Plugins/Experimental/PCG/Source/PCG/Public/PCGComponent.h
--- a/Engine/Plugins/Experimental/PCG/Source/PCG/Public/PCGComponent.h
+++ b/Engine/Plugins/Experimental/PCG/Source/PCG/Public/PCGComponent.h
@@ -244,7 +244,7 @@ public:
 	bool IsInspecting() const;
 	void EnableInspection();
 	void DisableInspection();
-	void StoreInspectionData(const FPCGStack* InStack, const UPCGNode* InNode, const FPCGDataCollection& InInspectionData);
+	void StoreInspectionData(const FPCGStack* InStack, const UPCGNode* InNode, const FPCGDataCollection& InInputData, const FPCGDataCollection& InOutputData);
 	const FPCGDataCollection* GetInspectionData(const FString& InStackPath) const;
 
 	/** Used by the tracking system to know if the component need to track actors. Not enabled for now, tracking is still done on the component.*/
diff --git a/Engine/Plugins/Experimental/PCG/Source/PCGEditor/Private/SPCGEditorGraphAttributeListView.cpp b/Engine/Plugins/Experimental/PCG/Source/PCGEditor/Private/SPCGEditorGraphAttributeListView.cpp
--- a/Engine/Plugins/Experimental/PCG/Source/PCGEditor/Private/SPCGEditorGraphAttributeListView.cpp
+++ b/Engine/Plugins/Experimental/PCG/Source/PCGEditor/Private/SPCGEditorGraphAttributeListView.cpp
@@ -4,9 +4,9 @@
 
 #include "PCGComponent.h"
 #include "PCGData.h"
-#include "PCGEditor.h"
-#include "PCGEditorGraphNodeBase.h"
+#include "PCGNode.h"
 #include "PCGParamData.h"
+#include "PCGPin.h"
 #include "PCGSubsystem.h"
 #include "Data/PCGPointData.h"
 #include "Data/PCGSpatialData.h"
@@ -17,6 +17,9 @@
 #include "Metadata/Accessors/PCGAttributeAccessorKeys.h"
 #include "Metadata/Accessors/PCGCustomAccessor.h"
 
+#include "PCGEditor.h"
+#include "PCGEditorGraphNodeBase.h"
+
 #include "Fonts/FontMeasure.h"
 #include "Framework/MultiBox/MultiBoxBuilder.h"
 #include "Widgets/Images/SLayeredImage.h"
@@ -33,6 +36,7 @@ static TAutoConsoleVariable<bool> CVarShowAdvancedAttributesFields(
 namespace PCGEditorGraphAttributeListView
 {
 	const FString LastMetadataId = TEXT("@Last");
+	const FText NoPinAvailableText = LOCTEXT("NoPinAvailableText", "No pins");
 	const FText NoDataAvailableText = LOCTEXT("NoDataAvailableText", "No data available");
 	const FText NoNodeInspectedText = LOCTEXT("NoNodeInspectedText", "No node being inspected");
 	const FText NoNodeInspectedToolTip = LOCTEXT("NoNodeInspectedToolTip", "Inspect a node using the right click menu");
@@ -226,6 +230,15 @@ void SPCGEditorGraphAttributeListView::Construct(const FArguments& InArgs, TShar
 		.Visibility_Lambda(VisibilityTest)
 		.ConsumeMouseWheel(EConsumeMouseWheel::Always);
 
+	SAssignNew(PinComboBox, SComboBox<TSharedPtr<FPinComboBoxItem>>)
+		.OptionsSource(&PinComboBoxItems)
+		.OnGenerateWidget(this, &SPCGEditorGraphAttributeListView::OnGeneratePinWidget)
+		.OnSelectionChanged(this, &SPCGEditorGraphAttributeListView::OnSelectionChangedPin)
+		[
+			SNew(STextBlock)
+			.Text(this, &SPCGEditorGraphAttributeListView::OnGenerateSelectedPinText)
+		];
+
 	SAssignNew(DataComboBox, SComboBox<TSharedPtr<FName>>)
 		.OptionsSource(&DataComboBoxItems)
 		.OnGenerateWidget(this, &SPCGEditorGraphAttributeListView::OnGenerateDataWidget)
@@ -268,6 +281,12 @@ void SPCGEditorGraphAttributeListView::Construct(const FArguments& InArgs, TShar
 			+SHorizontalBox::Slot()
 			.AutoWidth()
 			.Padding(1.0f, 0.0f)
+			[
+				PinComboBox->AsShared()
+			]
+			+SHorizontalBox::Slot()
+			.AutoWidth()
+			.Padding(1.0f, 0.0f)
 			[
 				DataComboBox->AsShared()
 			]
@@ -332,6 +351,7 @@ void SPCGEditorGraphAttributeListView::Tick(const FGeometry& AllottedGeometry, c
 	{
 		bNeedsRefresh = false;
 
+		RefreshPinComboBox();
 		RefreshDataComboBox();
 		RefreshAttributeList();
 	}
@@ -404,22 +424,36 @@ const FPCGDataCollection* SPCGEditorGraphAttributeListView::GetInspectionData()
 		return nullptr;
 	}
 
-	if (!PCGEditorGraphNode.IsValid())
+	const UPCGNode* PCGNode = PCGEditorGraphNode.IsValid() ? PCGEditorGraphNode->GetPCGNode() : nullptr;
+	if (!PCGNode)
 	{
 		return nullptr;
 	}
 
-	const UPCGNode* PCGNode = PCGEditorGraphNode->GetPCGNode();
-	if (!PCGNode)
+	const TSharedPtr<FPCGEditor> PCGEditor = PCGEditorPtr.Pin();
+	const FPCGStack& PCGStack = PCGEditor->GetStackBeingInspected();
+
+	const int32 SelectedIndex = GetSelectedPinIndex();
+	if (SelectedIndex == INDEX_NONE)
 	{
 		return nullptr;
 	}
 
-	const TSharedPtr<FPCGEditor> PCGEditor = PCGEditorPtr.Pin();
-	const FPCGStack& PCGStack = PCGEditor->GetStackBeingInspected();
+	// Selected pin is an output if it's at the beginning of the list, otherwise it's an input
+	const UPCGPin* Pin = nullptr;
+	const int32 InputPinCount = PCGNode->GetInputPins().Num();
+	const int32 OutputPinCount = PCGNode->GetOutputPins().Num();
+	if (SelectedIndex < (InputPinCount + OutputPinCount))
+	{
+		Pin = (SelectedIndex < OutputPinCount) ? PCGNode->GetOutputPins()[SelectedIndex] : PCGNode->GetInputPins()[SelectedIndex - OutputPinCount];
+	}
+	if (!Pin)
+	{
+		return nullptr;
+	}
 
 	FString ResourcePath;
-	if (!PCGStack.CreateStackFramePath(ResourcePath, PCGNode))
+	if (!PCGStack.CreateStackFramePath(ResourcePath, PCGNode, Pin))
 	{
 		return nullptr;
 	}
@@ -520,6 +554,55 @@ void SPCGEditorGraphAttributeListView::RefreshAttributeList()
 	RefreshSorting();
 }
 
+void SPCGEditorGraphAttributeListView::RefreshPinComboBox()
+{
+	PinComboBoxItems.Empty();
+	PinComboBox->ClearSelection();
+	PinComboBox->RefreshOptions();
+
+	const UPCGNode* PCGNode = PCGEditorGraphNode.IsValid() ? PCGEditorGraphNode->GetPCGNode() : nullptr;
+	if (!PCGNode)
+	{
+		return;
+	}
+
+	// Add output and then input pins to list. Optionally output the first connected item - useful for initializing
+	// the selected item to the first connected output pin.
+	auto PopulatePins = [](
+		const TArray<TObjectPtr<UPCGPin>> InPins,
+		const FString& InFormatText,
+		TArray<TSharedPtr<FPinComboBoxItem>>& InOutItems,
+		int32* OutFirstConnectedItemIndex)
+	{
+		for (int32 i = 0; i < InPins.Num(); ++i)
+		{
+			// Pin is included in list if it is connected, or if it is an output pin.
+			if (InPins[i] && (InPins[i]->IsConnected() || InPins[i]->IsOutputPin()))
+			{
+				FString ItemName = FString::Format(*InFormatText, { InPins[i]->Properties.Label.ToString() });
+				InOutItems.Add(MakeShared<FPinComboBoxItem>(FName(ItemName), i));
+
+				// Look for first connected, null pointer once found so only first is taken.
+				if (OutFirstConnectedItemIndex && InPins[i]->IsConnected())
+				{
+					*OutFirstConnectedItemIndex = InOutItems.Num() - 1;
+					OutFirstConnectedItemIndex = nullptr;
+				}
+			}
+		}
+	};
+
+	// Pick first connected output pin by default if there is one, otherwise default to first output pin.
+	int32 FirstConnectedItemIndex = 0;
+	PopulatePins(PCGNode->GetOutputPins(), TEXT("Output: {0}"), PinComboBoxItems, &FirstConnectedItemIndex);
+	PopulatePins(PCGNode->GetInputPins(), TEXT("Input: {0}"), PinComboBoxItems, nullptr);
+
+	if (PinComboBoxItems.Num() > 0 && ensure(PinComboBoxItems.IsValidIndex(FirstConnectedItemIndex)))
+	{
+		PinComboBox->SetSelectedItem(PinComboBoxItems[FirstConnectedItemIndex]);
+	}
+}
+
 void SPCGEditorGraphAttributeListView::RefreshDataComboBox()
 {
 	DataComboBoxItems.Empty();
@@ -535,8 +618,9 @@ void SPCGEditorGraphAttributeListView::RefreshDataComboBox()
 	for(int32 TaggedDataIndex = 0; TaggedDataIndex < InspectionData->TaggedData.Num(); ++TaggedDataIndex)
 	{
 		const FPCGTaggedData& TaggedData = InspectionData->TaggedData[TaggedDataIndex];
-		FString ItemName = FString::Format(TEXT("[{0}] {1} - {2}"),
-			{ FText::AsNumber(TaggedDataIndex).ToString(), TaggedData.Pin.ToString(), (TaggedData.Data ? TaggedData.Data->GetClass()->GetDisplayNameText().ToString(): TEXT("No Data")) });
+		FString ItemName = FString::Format(
+			TEXT("[{0}] {1}"),
+			{ FText::AsNumber(TaggedDataIndex).ToString(), (TaggedData.Data ? TaggedData.Data->GetClass()->GetDisplayNameText().ToString(): TEXT("No Data")) });
 
 		if (!TaggedData.Tags.IsEmpty())
 		{
@@ -605,6 +689,44 @@ TSharedRef<SWidget> SPCGEditorGraphAttributeListView::OnGenerateFilterMenu()
 	return MenuBuilder.MakeWidget();
 }
 
+FText SPCGEditorGraphAttributeListView::OnGenerateSelectedPinText() const
+{
+	if (const TSharedPtr<FPinComboBoxItem> SelectedPin = PinComboBox->GetSelectedItem())
+	{
+		return FText::FromName(SelectedPin->Name);
+	}
+	else
+	{
+		return PCGEditorGraphAttributeListView::NoPinAvailableText;
+	}
+}
+
+int32 SPCGEditorGraphAttributeListView::GetSelectedPinIndex() const
+{
+	int32 Index = INDEX_NONE;
+	if (const TSharedPtr<FPinComboBoxItem> SelectedPin = PinComboBox->GetSelectedItem())
+	{
+		PinComboBoxItems.Find(SelectedPin, Index);
+	}
+
+	return Index;
+}
+
+void SPCGEditorGraphAttributeListView::OnSelectionChangedPin(TSharedPtr<FPinComboBoxItem> InItem, ESelectInfo::Type InSelectInfo)
+{
+	RefreshDataComboBox();
+
+	if (InSelectInfo != ESelectInfo::Direct)
+	{
+		RefreshAttributeList();
+	}
+}
+
+TSharedRef<SWidget> SPCGEditorGraphAttributeListView::OnGeneratePinWidget(TSharedPtr<FPinComboBoxItem> InItem) const
+{
+	return SNew(STextBlock).Text(FText::FromName(InItem.IsValid() ? InItem->Name : NAME_None));
+}
+
 TSharedRef<SWidget> SPCGEditorGraphAttributeListView::OnGenerateDataWidget(TSharedPtr<FName> InItem) const
 {
 	return SNew(STextBlock).Text(FText::FromName(InItem.IsValid() ? *InItem : NAME_None));
diff --git a/Engine/Plugins/Experimental/PCG/Source/PCGEditor/Private/SPCGEditorGraphAttributeListView.h b/Engine/Plugins/Experimental/PCG/Source/PCGEditor/Private/SPCGEditorGraphAttributeListView.h
--- a/Engine/Plugins/Experimental/PCG/Source/PCGEditor/Private/SPCGEditorGraphAttributeListView.h
+++ b/Engine/Plugins/Experimental/PCG/Source/PCGEditor/Private/SPCGEditorGraphAttributeListView.h
@@ -96,8 +96,23 @@ private:
 	const FPCGDataCollection* GetInspectionData();
 	
 	void RefreshAttributeList();
+	void RefreshPinComboBox();
 	void RefreshDataComboBox();
 
+	/** Only connected input pins are added to combo box, so keep track of the node pin index for each item. */
+	struct FPinComboBoxItem
+	{
+		explicit FPinComboBoxItem(FName InName, int32 InPinIndex) : Name(InName), PinIndex(InPinIndex) {}
+
+		FName Name;
+		int32 PinIndex = INDEX_NONE;
+	};
+
+	FText OnGenerateSelectedPinText() const;
+	int32 GetSelectedPinIndex() const;
+	void OnSelectionChangedPin(TSharedPtr<FPinComboBoxItem> InItem, ESelectInfo::Type InSelectInfo);
+	TSharedRef<SWidget> OnGeneratePinWidget(TSharedPtr<FPinComboBoxItem> InItem) const;
+
 	const FSlateBrush* GetFilterBadgeIcon() const;
 	TSharedRef<SWidget> OnGenerateFilterMenu();
 	TSharedRef<SWidget> OnGenerateDataWidget(TSharedPtr<FName> InItem) const;
@@ -136,6 +151,9 @@ private:
 	TSharedPtr<SListView<PCGListviewItemPtr>> ListView;
 	TArray<PCGListviewItemPtr> ListViewItems;
 
+	TSharedPtr<SComboBox<TSharedPtr<FPinComboBoxItem>>> PinComboBox;
+	TArray<TSharedPtr<FPinComboBoxItem>> PinComboBoxItems;
+
 	TSharedPtr<SComboBox<TSharedPtr<FName>>> DataComboBox;
 	TArray<TSharedPtr<FName>> DataComboBoxItems;
 	
