[PCG] Add simple mechanism to get static mesh bounds

diff --git a/Engine/Plugins/Experimental/PCG/Source/PCG/Private/Elements/PCGPointFromMeshElement.cpp b/Engine/Plugins/Experimental/PCG/Source/PCG/Private/Elements/PCGPointFromMeshElement.cpp
new file mode 100644
--- /dev/null
+++ b/Engine/Plugins/Experimental/PCG/Source/PCG/Private/Elements/PCGPointFromMeshElement.cpp
@@ -0,0 +1,68 @@
+// Copyright Epic Games, Inc. All Rights Reserved.
+
+#include "Elements/PCGPointFromMeshElement.h"
+
+#include "PCGComponent.h"
+#include "PCGContext.h"
+#include "PCGEdge.h"
+#include "Data/PCGPointData.h"
+#include "Data/PCGSpatialData.h"
+
+#include "Engine/StaticMesh.h"
+
+#include UE_INLINE_GENERATED_CPP_BY_NAME(PCGPointFromMeshElement)
+
+#define LOCTEXT_NAMESPACE "PCGPointFromMeshElement"
+
+#if WITH_EDITOR
+FText UPCGPointFromMeshSettings::GetDefaultNodeTitle() const
+{
+	return LOCTEXT("NodeTitle", "PointFromMesh");
+}
+
+FText UPCGPointFromMeshSettings::GetNodeTooltipText() const
+{
+	return LOCTEXT("PointFromMeshNodeTooltip", "Creates a single point at the origin with an attribute named MeshPathAttributeName containing a SoftObjectPath to the StaticMesh.");
+}
+#endif
+
+FPCGElementPtr UPCGPointFromMeshSettings::CreateElement() const
+{
+	return MakeShared<FPCGPointFromMeshElement>();
+}
+
+bool FPCGPointFromMeshElement::ExecuteInternal(FPCGContext* Context) const
+{
+	TRACE_CPUPROFILER_EVENT_SCOPE(FPCGPointFromMeshElement::Execute);
+
+	check(Context);
+
+	const UPCGPointFromMeshSettings* Settings = Context->GetInputSettings<UPCGPointFromMeshSettings>();
+	check(Settings);
+
+	if (!Settings->StaticMesh.LoadSynchronous())
+	{
+		PCGE_LOG(Error, GraphAndLog, LOCTEXT("LoadStaticMeshFailed", "Failed to load StaticMesh"));
+		return true;
+	}
+
+	TArray<FPCGTaggedData>& Outputs = Context->OutputData.TaggedData;
+	UPCGPointData* OutPointData = NewObject<UPCGPointData>();
+	Outputs.Emplace_GetRef().Data = OutPointData;
+	
+	TArray<FPCGPoint>& Points = OutPointData->GetMutablePoints();
+	FPCGPoint& Point = Points.Emplace_GetRef();
+
+	// Capture StaticMesh bounds
+	const FBox StaticMeshBounds = Settings->StaticMesh->GetBoundingBox();
+	Point.BoundsMin = StaticMeshBounds.Min;
+	Point.BoundsMax = StaticMeshBounds.Max;
+
+	// Write StaticMesh path to MeshPathAttribute
+	check(OutPointData->Metadata);
+	OutPointData->Metadata->CreateStringAttribute(Settings->MeshPathAttributeName, Settings->StaticMesh.ToString(), /*bAllowsInterpolation=*/false);
+
+	return true;
+}
+
+#undef LOCTEXT_NAMESPACE
diff --git a/Engine/Plugins/Experimental/PCG/Source/PCG/Private/Elements/PCGStaticMeshSpawner.cpp b/Engine/Plugins/Experimental/PCG/Source/PCG/Private/Elements/PCGStaticMeshSpawner.cpp
--- a/Engine/Plugins/Experimental/PCG/Source/PCG/Private/Elements/PCGStaticMeshSpawner.cpp
+++ b/Engine/Plugins/Experimental/PCG/Source/PCG/Private/Elements/PCGStaticMeshSpawner.cpp
@@ -3,6 +3,7 @@
 #include "Elements/PCGStaticMeshSpawner.h"
 
 #include "PCGComponent.h"
+#include "PCGCustomVersion.h"
 #include "PCGManagedResource.h"
 #include "Data/PCGPointData.h"
 #include "Data/PCGSpatialData.h"
@@ -40,6 +41,26 @@ UPCGStaticMeshSpawnerSettings::UPCGStaticMeshSpawnerSettings(const FObjectInitia
 	}
 }
 
+#if WITH_EDITOR
+FText UPCGStaticMeshSpawnerSettings::GetDefaultNodeTitle() const
+{
+	return LOCTEXT("NodeTitle", "Static Mesh Spawner");
+}
+
+void UPCGStaticMeshSpawnerSettings::ApplyDeprecation(UPCGNode* InOutNode)
+{
+	check(InOutNode);
+
+	if (DataVersion < FPCGCustomVersion::StaticMeshSpawnerApplyMeshBoundsToPointsByDefault)
+	{
+		UE_LOG(LogPCG, Log, TEXT("Static Mesh Spawner node migrated from an older version. Disabling 'ApplyMeshBoundsToPoints' by default to match previous behavior."));
+		bApplyMeshBoundsToPoints = false;
+	}
+
+	Super::ApplyDeprecation(InOutNode);
+}
+#endif
+
 FPCGElementPtr UPCGStaticMeshSpawnerSettings::CreateElement() const
 {
 	return MakeShared<FPCGStaticMeshSpawnerElement>();
diff --git a/Engine/Plugins/Experimental/PCG/Source/PCG/Private/MeshSelectors/PCGMeshSelectorByAttribute.cpp b/Engine/Plugins/Experimental/PCG/Source/PCG/Private/MeshSelectors/PCGMeshSelectorByAttribute.cpp
--- a/Engine/Plugins/Experimental/PCG/Source/PCG/Private/MeshSelectors/PCGMeshSelectorByAttribute.cpp
+++ b/Engine/Plugins/Experimental/PCG/Source/PCG/Private/MeshSelectors/PCGMeshSelectorByAttribute.cpp
@@ -8,6 +8,8 @@
 #include "Data/PCGSpatialData.h"
 #include "MeshSelectors/PCGMeshSelectorBase.h"
 
+#include "Engine/StaticMesh.h"
+
 #include UE_INLINE_GENERATED_CPP_BY_NAME(PCGMeshSelectorByAttribute)
 
 #define LOCTEXT_NAMESPACE "PCGMeshSelectorByAttribute"
@@ -157,9 +159,12 @@ bool UPCGMeshSelectorByAttribute::SelectInstances(
 	int32 LastCheckpointIndex = CurrentPointIndex;
 	constexpr int32 TimeSlicingCheckFrequency = 1024;
 	TMap<PCGMetadataValueKey, TSoftObjectPtr<UStaticMesh>>& ValueKeyToMesh = Context.ValueKeyToMesh;
+	TMap<TSoftObjectPtr<UStaticMesh>, FBox>& MeshToBoundingBox = Context.MeshToBoundingBox;
 
 	const TArray<FPCGPoint>& Points = InPointData->GetPoints();
-	while(CurrentPointIndex < Points.Num())
+	TArray<FPCGPoint>* OutPoints = OutPointData ? &OutPointData->GetMutablePoints() : nullptr;
+
+	while (CurrentPointIndex < Points.Num())
 	{
 		const FPCGPoint& Point = Points[CurrentPointIndex++];
 
@@ -209,6 +214,24 @@ bool UPCGMeshSelectorByAttribute::SelectInstances(
 		InstanceList.Instances.Emplace(Point.Transform);
 		InstanceList.InstancesMetadataEntry.Emplace(Point.MetadataEntry);
 
+		// Update point bounds to reflect StaticMesh bounds
+		if (OutPoints && Settings->bApplyMeshBoundsToPoints)
+		{
+			if (!MeshToBoundingBox.Contains(Mesh) && Mesh.LoadSynchronous())
+			{
+				MeshToBoundingBox.Add(Mesh, Mesh->GetBoundingBox());
+			}
+
+			if (MeshToBoundingBox.Contains(Mesh))
+			{
+				const FBox MeshBounds = MeshToBoundingBox[Mesh];
+
+				// CurrentPointIndex - 1, because CurrentPointIndex is incremented at the beginning of the function
+				(*OutPoints)[CurrentPointIndex - 1].BoundsMin = MeshBounds.Min;
+				(*OutPoints)[CurrentPointIndex - 1].BoundsMax = MeshBounds.Max;
+			}
+		}
+
 		// Check if we should stop here and continue in a subsequent call
 		if(CurrentPointIndex - LastCheckpointIndex >= TimeSlicingCheckFrequency)
 		{
diff --git a/Engine/Plugins/Experimental/PCG/Source/PCG/Private/MeshSelectors/PCGMeshSelectorWeighted.cpp b/Engine/Plugins/Experimental/PCG/Source/PCG/Private/MeshSelectors/PCGMeshSelectorWeighted.cpp
--- a/Engine/Plugins/Experimental/PCG/Source/PCG/Private/MeshSelectors/PCGMeshSelectorWeighted.cpp
+++ b/Engine/Plugins/Experimental/PCG/Source/PCG/Private/MeshSelectors/PCGMeshSelectorWeighted.cpp
@@ -3,13 +3,14 @@
 #include "MeshSelectors/PCGMeshSelectorWeighted.h"
 
 #include "Data/PCGPointData.h"
+#include "Data/PCGSpatialData.h"
 #include "Elements/PCGStaticMeshSpawner.h"
 #include "Elements/PCGStaticMeshSpawnerContext.h"
-#include "Data/PCGSpatialData.h"
 #include "Helpers/PCGBlueprintHelpers.h"
+#include "MeshSelectors/PCGMeshSelectorBase.h"
 
+#include "Engine/StaticMesh.h"
 #include "Math/RandomStream.h"
-#include "MeshSelectors/PCGMeshSelectorBase.h"
 
 #include UE_INLINE_GENERATED_CPP_BY_NAME(PCGMeshSelectorWeighted)
 
@@ -204,10 +205,11 @@ bool UPCGMeshSelectorWeighted::SelectInstances(
 		int32 LastCheckpointIndex = CurrentPointIndex;
 		constexpr int32 TimeSlicingCheckFrequency = 1024;
 		TMap<TSoftObjectPtr<UStaticMesh>, PCGMetadataValueKey>& MeshToValueKey = Context.MeshToValueKey;
+		TMap<TSoftObjectPtr<UStaticMesh>, FBox>& MeshToBoundingBox = Context.MeshToBoundingBox;
 		const int32 TotalWeight = CumulativeWeights.Last();
 
 		const TArray<FPCGPoint>& Points = InPointData->GetPoints();
-		while(CurrentPointIndex < Points.Num())
+		while (CurrentPointIndex < Points.Num())
 		{
 			const FPCGPoint& Point = Points[CurrentPointIndex++];
 
@@ -236,16 +238,33 @@ bool UPCGMeshSelectorWeighted::SelectInstances(
 
 				if (OutPointData && OutAttribute)
 				{
+					FPCGPoint& OutPoint = OutPoints->Add_GetRef(Point);
+
+					// Update point bounds to reflect StaticMesh bounds
+					if (Settings->bApplyMeshBoundsToPoints)
+					{
+						if (!MeshToBoundingBox.Contains(Mesh) && Mesh.LoadSynchronous())
+						{
+							MeshToBoundingBox.Add(Mesh, Mesh->GetBoundingBox());
+						}
+
+						if (MeshToBoundingBox.Contains(Mesh))
+						{
+							const FBox MeshBounds = MeshToBoundingBox[Mesh];
+							OutPoint.BoundsMin = MeshBounds.Min;
+							OutPoint.BoundsMax = MeshBounds.Max;
+						}
+					}
+
 					PCGMetadataValueKey* OutValueKey = MeshToValueKey.Find(Mesh);
-					if(!OutValueKey)
+					if (!OutValueKey)
 					{
 						PCGMetadataValueKey ValueKey = OutAttribute->AddValue(Mesh.ToSoftObjectPath().ToString());
 						OutValueKey = &MeshToValueKey.Add(Mesh, ValueKey);
 					}
-					
+
 					check(OutValueKey);
-					
-					FPCGPoint& OutPoint = OutPoints->Add_GetRef(Point);
+
 					OutPointData->Metadata->InitializeOnSet(OutPoint.MetadataEntry);
 					OutAttribute->SetValueFromValueKey(OutPoint.MetadataEntry, *OutValueKey);
 				}
diff --git a/Engine/Plugins/Experimental/PCG/Source/PCG/Private/MeshSelectors/PCGMeshSelectorWeightedByCategory.cpp b/Engine/Plugins/Experimental/PCG/Source/PCG/Private/MeshSelectors/PCGMeshSelectorWeightedByCategory.cpp
--- a/Engine/Plugins/Experimental/PCG/Source/PCG/Private/MeshSelectors/PCGMeshSelectorWeightedByCategory.cpp
+++ b/Engine/Plugins/Experimental/PCG/Source/PCG/Private/MeshSelectors/PCGMeshSelectorWeightedByCategory.cpp
@@ -7,9 +7,10 @@
 #include "Elements/PCGStaticMeshSpawner.h"
 #include "Elements/PCGStaticMeshSpawnerContext.h"
 #include "Helpers/PCGBlueprintHelpers.h"
+#include "MeshSelectors/PCGMeshSelectorWeighted.h"
 
+#include "Engine/StaticMesh.h"
 #include "Math/RandomStream.h"
-#include "MeshSelectors/PCGMeshSelectorWeighted.h"
 
 #include UE_INLINE_GENERATED_CPP_BY_NAME(PCGMeshSelectorWeightedByCategory)
 
@@ -83,6 +84,8 @@ bool UPCGMeshSelectorWeightedByCategory::SelectInstances(
 	// Setup
 	if (Context.CurrentPointIndex == 0)
 	{
+		const FString AttributeDefaultValue = Attribute->GetValue(PCGDefaultValueKey);
+
 		for (const FPCGWeightedByCategoryEntryList& Entry : Entries)
 		{
 			if (Entry.WeightedMeshEntries.Num() == 0)
@@ -93,9 +96,11 @@ bool UPCGMeshSelectorWeightedByCategory::SelectInstances(
 
 			PCGMetadataValueKey ValueKey = Attribute->FindValue<FString>(Entry.CategoryEntry);
 
-			if (ValueKey == PCGDefaultValueKey)
+			// DefaultValueKey is only an invalid ValueKey when the CategoryEntry is not the AttributeDefaultValue
+			if (ValueKey == PCGDefaultValueKey && !Entry.CategoryEntry.Equals(AttributeDefaultValue))
 			{
-				PCGE_LOG_C(Verbose, LogOnly, &Context, FText::Format(LOCTEXT("InvalidCategory", "Invalid category '{0}'"), FText::FromString(Entry.CategoryEntry)));
+				PCGE_LOG_C(Verbose, LogOnly, &Context, FText::Format(LOCTEXT("UnusedCategory", "Unused category '{0}'. Not a valid value for attribute '{1}'."), 
+					FText::FromString(Entry.CategoryEntry), FText::FromName(CategoryAttribute)));
 				continue;
 			}
 
@@ -191,10 +196,11 @@ bool UPCGMeshSelectorWeightedByCategory::SelectInstances(
 	int32 LastCheckpointIndex = CurrentPointIndex;
 	constexpr int32 TimeSlicingCheckFrequency = 1024;
 	TMap<TSoftObjectPtr<UStaticMesh>, PCGMetadataValueKey>& MeshToValueKey = Context.MeshToValueKey;
+	TMap<TSoftObjectPtr<UStaticMesh>, FBox>& MeshToBoundingBox = Context.MeshToBoundingBox;
 
 	// Assign points to entries
 	const TArray<FPCGPoint>& Points = InPointData->GetPoints();
-	while(CurrentPointIndex < Points.Num())
+	while (CurrentPointIndex < Points.Num())
 	{
 		const FPCGPoint& Point = Points[CurrentPointIndex++];
 
@@ -242,16 +248,33 @@ bool UPCGMeshSelectorWeightedByCategory::SelectInstances(
 
 			if (OutPointData && OutAttribute)
 			{
+				FPCGPoint& OutPoint = OutPoints->Add_GetRef(Point);
+
+				// Update point bounds to reflect StaticMesh bounds
+				if (Settings->bApplyMeshBoundsToPoints)
+				{
+					if (!MeshToBoundingBox.Contains(Mesh) && Mesh.LoadSynchronous())
+					{
+						MeshToBoundingBox.Add(Mesh, Mesh->GetBoundingBox());
+					}
+
+					if (MeshToBoundingBox.Contains(Mesh))
+					{
+						const FBox MeshBounds = MeshToBoundingBox[Mesh];
+						OutPoint.BoundsMin = MeshBounds.Min;
+						OutPoint.BoundsMax = MeshBounds.Max;
+					}
+				}
+
 				PCGMetadataValueKey* OutValueKey = MeshToValueKey.Find(Mesh);
-				if(!OutValueKey)
+				if (!OutValueKey)
 				{
 					PCGMetadataValueKey TempValueKey = OutAttribute->AddValue(Mesh.ToSoftObjectPath().ToString());
 					OutValueKey = &MeshToValueKey.Add(Mesh, TempValueKey);
 				}
-				
+
 				check(OutValueKey);
-				
-				FPCGPoint& OutPoint = OutPoints->Add_GetRef(Point);
+
 				OutPointData->Metadata->InitializeOnSet(OutPoint.MetadataEntry);
 				OutAttribute->SetValueFromValueKey(OutPoint.MetadataEntry, *OutValueKey);
 			}
diff --git a/Engine/Plugins/Experimental/PCG/Source/PCG/Public/Elements/PCGPointFromMeshElement.h b/Engine/Plugins/Experimental/PCG/Source/PCG/Public/Elements/PCGPointFromMeshElement.h
new file mode 100644
--- /dev/null
+++ b/Engine/Plugins/Experimental/PCG/Source/PCG/Public/Elements/PCGPointFromMeshElement.h
@@ -0,0 +1,48 @@
+// Copyright Epic Games, Inc. All Rights Reserved.
+
+#pragma once
+
+#include "PCGPin.h"
+#include "PCGSettings.h"
+
+#include "UObject/SoftObjectPtr.h"
+
+#include "PCGPointFromMeshElement.generated.h"
+
+class UStaticMesh;
+
+// PointFromMesh creates a single point at the origin with an attribute containing a SoftObjectPath to the selected UStaticMesh
+UCLASS(BlueprintType, ClassGroup = (Procedural))
+class PCG_API UPCGPointFromMeshSettings : public UPCGSettings
+{
+	GENERATED_BODY()
+
+public:
+	//~Begin UPCGSettings interface
+#if WITH_EDITOR
+	virtual FName GetDefaultNodeName() const override { return FName(TEXT("PointFromMesh")); }
+	virtual FText GetDefaultNodeTitle() const override;
+	virtual FText GetNodeTooltipText() const override;
+	virtual EPCGSettingsType GetType() const override { return EPCGSettingsType::Spatial; }
+#endif
+
+protected:
+	virtual TArray<FPCGPinProperties> InputPinProperties() const override { return TArray<FPCGPinProperties>(); }
+	virtual TArray<FPCGPinProperties> OutputPinProperties() const override { return Super::DefaultPointOutputPinProperties(); }
+	virtual FPCGElementPtr CreateElement() const override;
+	//~End UPCGSettings interface
+
+public:
+	UPROPERTY(BlueprintReadWrite, EditAnywhere, Category = Settings, meta = (PCG_Overridable))
+	TSoftObjectPtr<UStaticMesh> StaticMesh;
+
+	/** Name of the string attribute to be created and hold a SoftObjectPath to the StaticMesh */
+	UPROPERTY(BlueprintReadWrite, EditAnywhere, Category = Settings)
+	FName MeshPathAttributeName = NAME_None;
+};
+
+class FPCGPointFromMeshElement : public FSimplePCGElement
+{
+protected:
+	virtual bool ExecuteInternal(FPCGContext* Context) const;
+};
diff --git a/Engine/Plugins/Experimental/PCG/Source/PCG/Public/Elements/PCGStaticMeshSpawner.h b/Engine/Plugins/Experimental/PCG/Source/PCG/Public/Elements/PCGStaticMeshSpawner.h
--- a/Engine/Plugins/Experimental/PCG/Source/PCG/Public/Elements/PCGStaticMeshSpawner.h
+++ b/Engine/Plugins/Experimental/PCG/Source/PCG/Public/Elements/PCGStaticMeshSpawner.h
@@ -47,8 +47,9 @@ public:
 #if WITH_EDITOR
 	// ~Begin UPCGSettings interface
 	virtual FName GetDefaultNodeName() const override { return FName(TEXT("StaticMeshSpawner")); }
-	virtual FText GetDefaultNodeTitle() const override { return NSLOCTEXT("PCGStaticMeshSpawnerSettings", "NodeTitle", "Static Mesh Spawner"); }
+	virtual FText GetDefaultNodeTitle() const override;
 	virtual EPCGSettingsType GetType() const override { return EPCGSettingsType::Spawner; }
+	virtual void ApplyDeprecation(UPCGNode* InOutNode) override;
 #endif
 
 protected:
@@ -90,6 +91,10 @@ public:
 	UPROPERTY(BlueprintReadWrite, EditAnywhere, Category = Settings)
 	FName OutAttributeName = NAME_None;
 
+	/** Sets the BoundsMin and BoundsMax attributes of each point to reflect the StaticMesh spawned at its location */
+	UPROPERTY(BlueprintReadWrite, EditAnywhere, Category = Settings)
+	bool bApplyMeshBoundsToPoints = true;
+
 	UPROPERTY(meta=(DeprecatedProperty, DeprecationMessage="Use MeshSelectorType and MeshSelectorParameters instead."))
 	TArray<FPCGStaticMeshSpawnerEntry> Meshes_DEPRECATED;
 
diff --git a/Engine/Plugins/Experimental/PCG/Source/PCG/Public/Elements/PCGStaticMeshSpawnerContext.h b/Engine/Plugins/Experimental/PCG/Source/PCG/Public/Elements/PCGStaticMeshSpawnerContext.h
--- a/Engine/Plugins/Experimental/PCG/Source/PCG/Public/Elements/PCGStaticMeshSpawnerContext.h
+++ b/Engine/Plugins/Experimental/PCG/Source/PCG/Public/Elements/PCGStaticMeshSpawnerContext.h
@@ -51,6 +51,9 @@ struct FPCGStaticMeshSpawnerContext : public FPCGContext
 	FPCGMeshMaterialOverrideHelper MaterialOverrideHelper;
 	int32 CurrentPointIndex = 0;
 
+	// Used in all selectors
+	TMap<TSoftObjectPtr<UStaticMesh>, FBox> MeshToBoundingBox;
+
 	// Used in by-attribute selector
 	TMap<PCGMetadataValueKey, TSoftObjectPtr<UStaticMesh>> ValueKeyToMesh;
 
diff --git a/Engine/Plugins/Experimental/PCG/Source/PCG/Public/PCGCustomVersion.h b/Engine/Plugins/Experimental/PCG/Source/PCG/Public/PCGCustomVersion.h
--- a/Engine/Plugins/Experimental/PCG/Source/PCG/Public/PCGCustomVersion.h
+++ b/Engine/Plugins/Experimental/PCG/Source/PCG/Public/PCGCustomVersion.h
@@ -41,6 +41,9 @@ struct PCG_API FPCGCustomVersion
 		// Behavior change for SplineSampler which now defaults to being bounded
 		SplineSamplerBoundedByDefault = 10,
 
+		// StaticMeshSpawner now defaults to modify point bounds based on StaticMesh bounds
+		StaticMeshSpawnerApplyMeshBoundsToPointsByDefault = 11,
+
 		// -----<new versions can be added above this line>-------------------------------------------------
 		VersionPlusOne,
 		LatestVersion = VersionPlusOne - 1
