[PCG] Fixed a minor bug where distance checks between the same source and target would lead to checking points against themselves.

* Distance node will ignore checking a point vs itself

Note: This does not resolve some oddities happening with overlapping or fully encompassed points. This will be resolved at a later time after discussion (UE-188660).

diff --git a/Engine/Plugins/Experimental/PCG/Source/PCG/Private/Elements/PCGDistance.cpp b/Engine/Plugins/Experimental/PCG/Source/PCG/Private/Elements/PCGDistance.cpp
--- a/Engine/Plugins/Experimental/PCG/Source/PCG/Private/Elements/PCGDistance.cpp
+++ b/Engine/Plugins/Experimental/PCG/Source/PCG/Private/Elements/PCGDistance.cpp
@@ -41,12 +41,24 @@ namespace PCGDistance
 		return SourceCenter;
 	}
 }
-	
+
+#if WITH_EDITOR
+FText UPCGDistanceSettings::GetNodeTooltipText() const
+{
+	return LOCTEXT("PCGDistanceTooltip", "Calculates and appends a signed 'Distance' attribute to the source data. For each of the source points, a distance attribute will be calculated between it and the nearest target point.");
+}
+#endif // WITH_EDITOR
+
 TArray<FPCGPinProperties> UPCGDistanceSettings::InputPinProperties() const
 {
 	TArray<FPCGPinProperties> PinProperties;
-	PinProperties.Emplace(PCGDistance::SourceLabel, EPCGDataType::Point);
-	PinProperties.Emplace(PCGDistance::TargetLabel, EPCGDataType::Point);
+	FPCGPinProperties& PinPropertySource = PinProperties.Emplace_GetRef(PCGDistance::SourceLabel, EPCGDataType::Point);
+	FPCGPinProperties& PinPropertyTarget = PinProperties.Emplace_GetRef(PCGDistance::TargetLabel, EPCGDataType::Point);
+
+#if WITH_EDITOR
+	PinPropertySource.Tooltip = LOCTEXT("PCGSourcePinTooltip", "For each of the source points, a distance attribute will be calculated between it and the nearest target point.");
+	PinPropertyTarget.Tooltip = LOCTEXT("PCGTargetPinTooltip", "The target points to conduct a distance check with each source point.");
+#endif // WITH_EDITOR
 
 	return PinProperties;
 }
@@ -54,8 +66,12 @@ TArray<FPCGPinProperties> UPCGDistanceSettings::InputPinProperties() const
 TArray<FPCGPinProperties> UPCGDistanceSettings::OutputPinProperties() const
 {
 	TArray<FPCGPinProperties> PinProperties;
-	PinProperties.Emplace(PCGPinConstants::DefaultOutputLabel, EPCGDataType::Spatial);
+	FPCGPinProperties& PinPropertyOutput = PinProperties.Emplace_GetRef(PCGPinConstants::DefaultOutputLabel, EPCGDataType::Point);
 
+#if WITH_EDITOR
+	PinPropertyOutput.Tooltip = LOCTEXT("PCGOutputPinTooltip", "The source points will be output with the newly added 'Distance' attribute as well as have their density set to [0,1] based on the 'Maximum Distance' if 'Set Density' is enabled.");
+#endif // WITH_EDITOR
+	
 	return PinProperties;
 }
 
@@ -82,9 +98,6 @@ bool FPCGDistanceElement::ExecuteInternal(FPCGContext* Context) const
 	TArray<FPCGTaggedData> Targets = Context->InputData.GetInputsByPin(PCGDistance::TargetLabel);
 	TArray<FPCGTaggedData>& Outputs = Context->OutputData.TaggedData;
 
-	FBox InputBounds(EForceInit::ForceInit);
-	int32 InputPointCount = 0;
-
 	TArray<const UPCGPointData*> TargetPointDatas;
 	TargetPointDatas.Reserve(Targets.Num());
 
@@ -157,10 +170,16 @@ bool FPCGDistanceElement::ExecuteInternal(FPCGContext* Context) const
 				double MinDistanceSquared = MaximumDistance*MaximumDistance;
 				FVector MinDistanceVector = FVector::ZeroVector;
 
+				// Signed distance field for calculating the closest point of source and target
 				auto CalculateSDF = [&MinDistanceSquared, &MinDistanceVector, &SourcePoint, SourceCenter, SourceShape, TargetShape](const FPCGPointRef& TargetPointRef)
 				{
+					// If the source pointer and target pointer are the same, ignore distance to the exact same point
+					if (&SourcePoint == TargetPointRef.Point)
+					{
+						return;
+					}
+					
 					const FPCGPoint& TargetPoint = *TargetPointRef.Point;
-
 					const FVector& TargetCenter = TargetPointRef.Bounds.Origin;
 
 					const FVector SourceShapePos = PCGDistance::CalcPosition(SourceShape, SourcePoint, TargetPoint, SourceCenter, TargetCenter);
diff --git a/Engine/Plugins/Experimental/PCG/Source/PCG/Public/Elements/PCGDistance.h b/Engine/Plugins/Experimental/PCG/Source/PCG/Public/Elements/PCGDistance.h
--- a/Engine/Plugins/Experimental/PCG/Source/PCG/Public/Elements/PCGDistance.h
+++ b/Engine/Plugins/Experimental/PCG/Source/PCG/Public/Elements/PCGDistance.h
@@ -33,6 +33,7 @@ public:
 #if WITH_EDITOR
 	virtual FName GetDefaultNodeName() const override { return FName(TEXT("Distance")); }
 	virtual FText GetDefaultNodeTitle() const override { return NSLOCTEXT("PCGDistanceSettings", "NodeTitle", "Distance"); }
+	virtual FText GetNodeTooltipText() const override;
 	virtual EPCGSettingsType GetType() const override { return EPCGSettingsType::Spatial; }
 #endif
 
