[PCG] Added search filter to attribute list view with logic expression support

diff --git a/Engine/Plugins/Experimental/PCG/Source/PCGEditor/Private/SPCGEditorGraphAttributeListView.cpp b/Engine/Plugins/Experimental/PCG/Source/PCGEditor/Private/SPCGEditorGraphAttributeListView.cpp
--- a/Engine/Plugins/Experimental/PCG/Source/PCGEditor/Private/SPCGEditorGraphAttributeListView.cpp
+++ b/Engine/Plugins/Experimental/PCG/Source/PCGEditor/Private/SPCGEditorGraphAttributeListView.cpp
@@ -24,6 +24,7 @@
 #include "Framework/MultiBox/MultiBoxBuilder.h"
 #include "Widgets/Images/SLayeredImage.h"
 #include "Widgets/Input/SComboBox.h"
+#include "Widgets/Input/SSearchBox.h"
 #include "Widgets/Layout/SScrollBox.h"
 
 #define LOCTEXT_NAMESPACE "SPCGEditorGraphAttributeListView"
@@ -187,6 +188,74 @@ TSharedRef<SWidget> SPCGListViewItemRow::GenerateWidgetForColumn(const FName& Co
 		.Margin(FMargin(2.0f, 0.0f));
 }
 
+FPCGPointFilterExpressionContext::FPCGPointFilterExpressionContext(const FPCGListViewItem* InRowItem, const TMap<FName, FPCGColumnData>* InPCGColumnData)
+	: RowItem(InRowItem)
+	, PCGColumnData(InPCGColumnData)
+{
+	check(InRowItem);
+}
+
+bool FPCGPointFilterExpressionContext::TestBasicStringExpression(const FTextFilterString& InValue, const ETextFilterTextComparisonMode InTextComparisonMode) const
+{
+	// Basic string search is disabled as it would require us to search the entire attribute table at once and it's not very useful.
+	return false;
+}
+
+bool FPCGPointFilterExpressionContext::TestComplexExpression(const FName& InKey, const FTextFilterString& InValue, const ETextFilterComparisonOperation InComparisonOperation, const ETextFilterTextComparisonMode InTextComparisonMode) const
+{
+	const int32 Index = RowItem->Index;
+
+	if (PCGEditorGraphAttributeListView::TEXT_IndexLabel.EqualToCaseIgnored(FText::FromName(InKey)))
+	{
+		const FTextFilterString PointValue(FString::FromInt(Index));
+		return TextFilterUtils::TestComplexExpression(PointValue, InValue, InComparisonOperation, InTextComparisonMode);
+	}
+	else if(const FPCGColumnData* PCGColumnInfo = PCGColumnData->Find(InKey))
+	{
+		if (PCGColumnInfo->DataAccessor.IsValid() && PCGColumnInfo->DataKeys.IsValid())
+		{
+			auto Callback = [&PCGColumnInfo, Index, &InValue, &InComparisonOperation, &InTextComparisonMode] (auto Dummy)
+			{
+				using ValueType = decltype(Dummy);
+				ValueType Value{};
+				if (PCGColumnInfo->DataAccessor->Get<ValueType>(Value, Index, *PCGColumnInfo->DataKeys))
+				{					
+					FText TextValue;
+					if constexpr (PCG::Private::IsOfTypes<ValueType, bool>())
+					{
+						TextValue = FText::FromString(LexToString(Value));
+					}
+					else if constexpr (PCG::Private::IsOfTypes<ValueType, FString>())
+					{
+						TextValue = FText::FromString(Value);
+					}
+					else if constexpr (PCG::Private::IsOfTypes<ValueType, FName>())
+					{
+						TextValue = FText::FromName(Value);
+					}
+					else if constexpr (FTextAsNumberIsValid<ValueType>::value)
+					{
+						TextValue = FText::AsNumber(Value, &FNumberFormattingOptions::DefaultNoGrouping());
+					}
+					else
+					{
+						ensureMsgf(false, TEXT("Unsupported Data Type"));
+						return false;
+					}
+					
+					const FTextFilterString PointValue(TextValue.ToString());
+					return TextFilterUtils::TestComplexExpression(PointValue, InValue, InComparisonOperation, InTextComparisonMode);
+				}
+				return false;
+			};
+
+			return PCGMetadataAttribute::CallbackWithRightType(PCGColumnInfo->DataAccessor->GetUnderlyingType(), Callback);
+		}
+	}
+
+	return true;
+}
+
 SPCGEditorGraphAttributeListView::~SPCGEditorGraphAttributeListView()
 {
 	if (PCGEditorPtr.IsValid())
@@ -205,6 +274,8 @@ void SPCGEditorGraphAttributeListView::Construct(const FArguments& InArgs, TShar
 	PCGEditorPtr.Pin()->OnInspectedStackChangedDelegate.AddSP(this, &SPCGEditorGraphAttributeListView::OnInspectedStackChanged);
 	PCGEditorPtr.Pin()->OnInspectedNodeChangedDelegate.AddSP(this, &SPCGEditorGraphAttributeListView::OnInspectedNodeChanged);
 
+	TextFilter = MakeShareable(new FTextFilterExpressionEvaluator(ETextFilterExpressionEvaluatorMode::Complex));
+
 	ListViewHeader = CreateHeaderRowWidget();
 
 	const TSharedRef<SScrollBar> HorizontalScrollBar = SNew(SScrollBar)
@@ -264,6 +335,14 @@ void SPCGEditorGraphAttributeListView::Construct(const FArguments& InArgs, TShar
 			FilterImage.ToSharedRef()
 		];
 
+	SAssignNew(SearchBoxWidget, SSearchBox)
+		.MinDesiredWidth(300.0f)
+		.InitialText(ActiveFilterText)
+		.OnTextChanged(this, &SPCGEditorGraphAttributeListView::OnFilterTextChanged)
+		.OnTextCommitted(this, &SPCGEditorGraphAttributeListView::OnFilterTextCommitted)
+		.DelayChangeNotificationsWhileTyping(true)
+		.DelayChangeNotificationsWhileTypingSeconds(0.5f);
+
 	this->ChildSlot
 	[
 		SNew(SVerticalBox)
@@ -292,6 +371,12 @@ void SPCGEditorGraphAttributeListView::Construct(const FArguments& InArgs, TShar
 			]
 			+SHorizontalBox::Slot()
 			.AutoWidth()
+			.Padding(1.0f, 0.0f)
+			[
+				SearchBoxWidget->AsShared()
+			]
+			+SHorizontalBox::Slot()
+			.AutoWidth()
 			.VAlign(EVerticalAlignment::VAlign_Center)
 			.Padding(4.0f, 0.0f)
 			[
@@ -550,7 +635,6 @@ void SPCGEditorGraphAttributeListView::RefreshAttributeList()
 		}
 	}
 
-	ListView->SetItemsSource(&ListViewItems);
 	RefreshSorting();
 }
 
@@ -636,6 +720,34 @@ void SPCGEditorGraphAttributeListView::RefreshDataComboBox()
 	}
 }
 
+void SPCGEditorGraphAttributeListView::ApplyRowFilter()
+{
+	TRACE_CPUPROFILER_EVENT_SCOPE(SPCGEditorGraphAttributeListView::ApplyRowFilter);
+
+	FilteredListViewItems.Empty();
+
+	const FString& FilterString = ActiveFilterText.ToString();
+	if (!FilterString.IsEmpty())
+	{
+		for (const PCGListviewItemPtr& ListViewItem : ListViewItems)
+		{
+			const FPCGPointFilterExpressionContext PointFilterContext(ListViewItem.Get(), &PCGColumnData);
+			if (TextFilter->TestTextFilter(PointFilterContext))
+			{
+				FilteredListViewItems.Add(ListViewItem);
+			}
+		}
+		
+		ListView->SetItemsSource(&FilteredListViewItems);
+	}
+	else
+	{
+		ListView->SetItemsSource(&ListViewItems);
+	}
+
+	ListView->RequestListRefresh();
+}
+
 const FSlateBrush* SPCGEditorGraphAttributeListView::GetFilterBadgeIcon() const
 {
 	for (const SHeaderRow::FColumn& Column : ListViewHeader->GetColumns())
@@ -1019,7 +1131,25 @@ void SPCGEditorGraphAttributeListView::RefreshSorting()
 		}
 	}
 
-	ListView->RequestListRefresh();
+	ApplyRowFilter();
+}
+
+void SPCGEditorGraphAttributeListView::OnFilterTextChanged(const FText& InFilterText)
+{
+	ActiveFilterText = InFilterText;
+	TextFilter->SetFilterText(InFilterText);
+
+	ApplyRowFilter();
+	SearchBoxWidget->SetError(TextFilter->GetFilterErrorText());
+}
+
+void SPCGEditorGraphAttributeListView::OnFilterTextCommitted(const FText& NewText, ETextCommit::Type CommitInfo)
+{
+	if (CommitInfo == ETextCommit::OnCleared)
+	{
+		SearchBoxWidget->SetText(FText::GetEmpty());
+		OnFilterTextChanged(FText::GetEmpty());
+	}
 }
 
 void SPCGEditorGraphAttributeListView::AddColumn(const UPCGPointData* InPCGPointData, const FName& InColumnId, const FText& ColumnLabel, EHorizontalAlignment HeaderHAlign, EHorizontalAlignment CellHAlign)
diff --git a/Engine/Plugins/Experimental/PCG/Source/PCGEditor/Private/SPCGEditorGraphAttributeListView.h b/Engine/Plugins/Experimental/PCG/Source/PCGEditor/Private/SPCGEditorGraphAttributeListView.h
--- a/Engine/Plugins/Experimental/PCG/Source/PCGEditor/Private/SPCGEditorGraphAttributeListView.h
+++ b/Engine/Plugins/Experimental/PCG/Source/PCGEditor/Private/SPCGEditorGraphAttributeListView.h
@@ -6,6 +6,7 @@
 #include "Metadata/Accessors/IPCGAttributeAccessor.h"
 #include "Metadata/Accessors/PCGAttributeAccessorKeys.h"
 
+#include "Misc/TextFilterExpressionEvaluator.h"
 #include "Widgets/Views/SListView.h"
 #include "Widgets/Views/STableRow.h"
 
@@ -28,6 +29,7 @@ enum class EPCGMetadataTypes : uint8;
 
 class SComboButton;
 class SHeaderRow;
+class SSearchBox;
 struct FSlateBrush;
 
 struct FPCGListViewItem
@@ -68,6 +70,21 @@ private:
 	PCGListviewItemPtr InternalItem;
 };
 
+class FPCGPointFilterExpressionContext : public ITextFilterExpressionContext
+{
+public:
+	explicit FPCGPointFilterExpressionContext(const FPCGListViewItem* InRowItem, const TMap<FName, FPCGColumnData>* InPCGColumnData);
+
+	// ~Begin ITextFilterExpressionContext interface
+	virtual bool TestBasicStringExpression(const FTextFilterString& InValue, const ETextFilterTextComparisonMode InTextComparisonMode) const override;
+	virtual bool TestComplexExpression(const FName& InKey, const FTextFilterString& InValue, const ETextFilterComparisonOperation InComparisonOperation, const ETextFilterTextComparisonMode InTextComparisonMode) const override;
+	// ~End ITextFilterExpressionContext interface
+
+private:
+	const FPCGListViewItem* RowItem = nullptr;
+	const TMap<FName, FPCGColumnData>* PCGColumnData = nullptr;
+};
+
 class SPCGEditorGraphAttributeListView : public SCompoundWidget
 {
 	friend SPCGListViewItemRow;
@@ -99,6 +116,8 @@ private:
 	void RefreshPinComboBox();
 	void RefreshDataComboBox();
 
+	void ApplyRowFilter();
+
 	/** Only connected input pins are added to combo box, so keep track of the node pin index for each item. */
 	struct FPinComboBoxItem
 	{
@@ -133,6 +152,9 @@ private:
 	EColumnSortMode::Type GetColumnSortMode(const FName InColumnId) const;
 	void RefreshSorting();
 
+	void OnFilterTextChanged(const FText& InFilterText);
+	void OnFilterTextCommitted(const FText& NewText, ETextCommit::Type CommitInfo);
+
 	void AddColumn(const UPCGPointData* InPCGPointData, const FName& InColumnId, const FText& ColumnLabel, EHorizontalAlignment HeaderHAlign = HAlign_Center, EHorizontalAlignment CellHAlign = HAlign_Right);
 	void AddIndexColumn();
 	void AddPointDataColumns(const UPCGPointData* InPCGPointData);
@@ -147,9 +169,13 @@ private:
 	/** Cached PCGGraphNode being viewed */
 	TWeakObjectPtr<UPCGEditorGraphNodeBase> PCGEditorGraphNode;
 
+	TSharedPtr<FTextFilterExpressionEvaluator> TextFilter;
+
+	TSharedPtr<SSearchBox> SearchBoxWidget;	
 	TSharedPtr<SHeaderRow> ListViewHeader;
 	TSharedPtr<SListView<PCGListviewItemPtr>> ListView;
 	TArray<PCGListviewItemPtr> ListViewItems;
+	TArray<PCGListviewItemPtr> FilteredListViewItems;
 
 	TSharedPtr<SComboBox<TSharedPtr<FPinComboBoxItem>>> PinComboBox;
 	TArray<TSharedPtr<FPinComboBoxItem>> PinComboBoxItems;
@@ -165,6 +191,8 @@ private:
 
 	TMap<FName, FPCGColumnData> PCGColumnData;
 
+	FText ActiveFilterText;
+
 	FName SortingColumn = NAME_None;
 	EColumnSortMode::Type SortMode = EColumnSortMode::Type::Ascending;
 
