Fixed clashing loctext key error

#rnx

diff --git a/Engine/Plugins/Experimental/PCG/Source/PCG/Private/Elements/IO/PCGDataTableElement.cpp b/Engine/Plugins/Experimental/PCG/Source/PCG/Private/Elements/IO/PCGDataTableElement.cpp
--- a/Engine/Plugins/Experimental/PCG/Source/PCG/Private/Elements/IO/PCGDataTableElement.cpp
+++ b/Engine/Plugins/Experimental/PCG/Source/PCG/Private/Elements/IO/PCGDataTableElement.cpp
@@ -55,7 +55,7 @@ bool FPCGLoadDataTableElement::PrepareLoad(FPCGExternalDataContext* Context) con
 	// 2.1 Validate that there are rows too
 	if (DataTable->GetRowMap().IsEmpty())
 	{
-		PCGE_LOG(Warning, GraphAndLog, FText::Format(LOCTEXT("DataTableNoRowStruct", "Provided Data Table ('{0}') is empty."), FText::FromString(Settings->DataTable.ToString())));
+		PCGE_LOG(Warning, GraphAndLog, FText::Format(LOCTEXT("DataTableIsEmpty", "Provided Data Table ('{0}') is empty."), FText::FromString(Settings->DataTable.ToString())));
 		return true;
 	}
 
