[PCG] Make use of latest addition in StructUtils to control GraphParameters widget for PCG

* Make a new parameter double by default (instead of bool)
* Remove the possibility to create arrays (not yet supported)
* Filter out not supported types in the type list
* Add a dialog to confirm parameter deletion if they are in use

diff --git a/Engine/Plugins/Experimental/PCG/Source/PCG/Private/PCGGraph.cpp b/Engine/Plugins/Experimental/PCG/Source/PCG/Private/PCGGraph.cpp
--- a/Engine/Plugins/Experimental/PCG/Source/PCG/Private/PCGGraph.cpp
+++ b/Engine/Plugins/Experimental/PCG/Source/PCG/Private/PCGGraph.cpp
@@ -15,8 +15,13 @@
 
 #include UE_INLINE_GENERATED_CPP_BY_NAME(PCGGraph)
 
+#define LOCTEXT_NAMESPACE "PCGGraph"
+
 #if WITH_EDITOR
+#include "CoreGlobals.h"
 #include "Editor.h"
+#include "Dialogs/Dialogs.h"
+#include "EdGraph/EdGraphPin.h"
 #else
 #include "UObject/Package.h"
 #endif
@@ -980,6 +985,57 @@ void UPCGGraph::OnGraphParametersChanged(EPCGGraphParameterEvent InChangeType, F
 	NotifyGraphParametersChanged(InChangeType, InChangedPropertyName);
 }
 
+bool UPCGGraph::UserParametersIsPinTypeAccepted(FEdGraphPinType InPinType)
+{
+	// Text and interface not supported
+	if (InPinType.PinCategory == TEXT("text") || InPinType.PinCategory == TEXT("interface"))
+	{
+		return false;
+	}
+	else if (InPinType.PinCategory == TEXT("struct"))
+	{
+		// Structs other than Vector/Transform/Rotator not supported
+		return InPinType.PinSubCategoryObject == TBaseStructure<FVector>::Get() 
+			|| InPinType.PinSubCategoryObject == TBaseStructure<FTransform>::Get()
+			|| InPinType.PinSubCategoryObject == TBaseStructure<FRotator>::Get();
+	}
+	else
+	{
+		return true;
+	}
+}
+
+bool UPCGGraph::UserParametersCanRemoveProperty(FGuid InPropertyID, FName InPropertyName)
+{
+	// Check if the property has some getters in the graph
+	for (const UPCGNode* Node : Nodes)
+	{
+		if (!Node)
+		{
+			continue;
+		}
+
+		if (const UPCGUserParameterGetSettings* Settings = Cast<UPCGUserParameterGetSettings>(Node->GetSettings()))
+		{
+			if (Settings->PropertyGuid == InPropertyID)
+			{
+				// We found a getter. Ask the user if he is OK with that
+				FText RemoveCheckMessage = FText::Format(LOCTEXT("UserParametersRemoveCheck", "Property {0} is in use in the graph. Are you sure you want to remove it?"), FText::FromName(InPropertyName));
+				FSuppressableWarningDialog::FSetupInfo Info(RemoveCheckMessage, LOCTEXT("UserParametersRemoveCheck_Message", "Remove property"), "UserParametersRemove");
+				Info.ConfirmText = FCoreTexts::Get().Yes;
+				Info.CancelText = FCoreTexts::Get().No;
+				FSuppressableWarningDialog AddLevelWarning(Info);
+				if (AddLevelWarning.ShowModal() == FSuppressableWarningDialog::Cancel)
+				{
+					return false;
+				}
+			}
+		}
+	}
+
+	return true;
+}
+
 #endif // WITH_EDITOR
 
 uint32 UPCGGraph::GetNodeGenerationGridSize(const UPCGNode* InNode, uint32 InDefaultGridSize) const
@@ -1591,3 +1647,5 @@ void FPCGOverrideInstancedPropertyBag::MigrateToNewBagInstance(const FInstancedP
 		}
 	}
 }
+
+#undef LOCTEXT_NAMESPACE
\ No newline at end of file
diff --git a/Engine/Plugins/Experimental/PCG/Source/PCG/Public/PCGGraph.h b/Engine/Plugins/Experimental/PCG/Source/PCG/Public/PCGGraph.h
--- a/Engine/Plugins/Experimental/PCG/Source/PCG/Public/PCGGraph.h
+++ b/Engine/Plugins/Experimental/PCG/Source/PCG/Public/PCGGraph.h
@@ -11,6 +11,10 @@
 
 #include "PCGGraph.generated.h"
 
+#if WITH_EDITOR
+struct FEdGraphPinType;
+#endif // WITH_EDITOR
+
 enum class EPCGGraphParameterEvent
 {
 	GraphChanged,
@@ -250,9 +254,17 @@ protected:
 #endif // WITH_EDITORONLY_DATA
 
 	// Parameters
-	UPROPERTY(EditAnywhere, Category = Instance, meta = (DisplayName = "Parameters", NoResetToDefault))
+	UPROPERTY(EditAnywhere, Category = Instance, meta = (DisplayName = "Parameters", NoResetToDefault, DefaultType = "EPropertyBagPropertyType::Double", AllowArrays = "false", IsPinTypeAccepted = "UserParametersIsPinTypeAccepted", CanRemoveProperty = "UserParametersCanRemoveProperty"))
 	FInstancedPropertyBag UserParameters;
 
+#if WITH_EDITOR
+	UFUNCTION(BlueprintInternalUseOnly)
+	bool UserParametersIsPinTypeAccepted(FEdGraphPinType InPinType);
+	
+	UFUNCTION(BlueprintInternalUseOnly)
+	bool UserParametersCanRemoveProperty(FGuid InPropertyID, FName InPropertyName);
+#endif // WITH_EDITOR
+
 	UPROPERTY(EditAnywhere, Category = Settings)
 	bool bHierarchicalGenerationEnabled = false;
 
