[PCG] Fixed a small bug where when exporting a new graph instance, the dialog browser would begin at an invalid directory.

* It now defaults to the content directory if invalid.

diff --git a/Engine/Plugins/Experimental/PCG/Source/PCGEditor/Private/PCGEditorUtils.cpp b/Engine/Plugins/Experimental/PCG/Source/PCGEditor/Private/PCGEditorUtils.cpp
--- a/Engine/Plugins/Experimental/PCG/Source/PCGEditor/Private/PCGEditorUtils.cpp
+++ b/Engine/Plugins/Experimental/PCG/Source/PCGEditor/Private/PCGEditorUtils.cpp
@@ -35,7 +35,12 @@ void PCGEditorUtils::GetParentPackagePathAndUniqueName(const UObject* OriginalOb
 	FPackageName::SplitLongPackageName(OriginalObject->GetPackage()->GetPathName(), PackageRoot, PackagePath, PackageName);
 
 	OutPackagePath = PackageRoot / PackagePath;
-
+	
+	if (!FPackageName::IsValidObjectPath(OutPackagePath))
+	{
+		OutPackagePath = FPaths::ProjectContentDir();
+	}
+	
 	FString DummyPackageName;
 	AssetTools.CreateUniqueAssetName(OutPackagePath, NewAssetTentativeName, DummyPackageName, OutUniqueName);
 }
