[PCG] HiGen - fix for default grid size not updating from the settings - clear the node to grid size mapping

#rnx

diff --git a/Engine/Plugins/Experimental/PCG/Source/PCG/Private/PCGGraph.cpp b/Engine/Plugins/Experimental/PCG/Source/PCG/Private/PCGGraph.cpp
--- a/Engine/Plugins/Experimental/PCG/Source/PCG/Private/PCGGraph.cpp
+++ b/Engine/Plugins/Experimental/PCG/Source/PCG/Private/PCGGraph.cpp
@@ -827,6 +827,13 @@ void UPCGGraph::NotifyGraphChanged(EPCGChangeType ChangeType)
 		}
 	}
 
+	// Graph settings, nodes, graph structure can all change the higen grid sizes.
+	if (ChangeType != EPCGChangeType::Cosmetic)
+	{
+		FWriteScopeLock ScopedWriteLock(NodeToGridSizeLock);
+		NodeToGridSize.Reset();
+	}
+
 	OnGraphChangedDelegate.Broadcast(this, ChangeType);
 
 	bIsNotifying = false;
@@ -850,6 +857,7 @@ void UPCGGraph::OnNodeChanged(UPCGNode* InNode, EPCGChangeType ChangeType)
 {
 	if (ChangeType != EPCGChangeType::Cosmetic)
 	{
+		FWriteScopeLock ScopedWriteLock(NodeToGridSizeLock);
 		NodeToGridSize.Reset();
 	}
 
