[PCG] Fix possible crash if actor has no root component (like in Tests)

#rnx

diff --git a/Engine/Plugins/Experimental/PCG/Source/PCG/Private/Helpers/PCGActorHelpers.cpp b/Engine/Plugins/Experimental/PCG/Source/PCG/Private/Helpers/PCGActorHelpers.cpp
--- a/Engine/Plugins/Experimental/PCG/Source/PCG/Private/Helpers/PCGActorHelpers.cpp
+++ b/Engine/Plugins/Experimental/PCG/Source/PCG/Private/Helpers/PCGActorHelpers.cpp
@@ -121,7 +121,10 @@ UPCGManagedISMComponent* UPCGActorHelpers::GetOrCreateManagedISMC(AActor* InTarg
 	UPCGManagedISMComponent* Resource = NewObject<UPCGManagedISMComponent>(InSourceComponent);
 	Resource->SetComponent(ISMC);
 	Resource->SetDescriptor(InParams.Descriptor);
-	Resource->SetRootLocation(InTargetActor->GetRootComponent()->GetComponentLocation());
+	if (InTargetActor->GetRootComponent())
+	{
+		Resource->SetRootLocation(InTargetActor->GetRootComponent()->GetComponentLocation());
+	}
 	Resource->SetSettingsUID(SettingsUID);
 	InSourceComponent->AddToManagedResources(Resource);
 
