Refactored usage of TargetActor so it appropriately controls artifact targetting
Small optimization in the metadata partition
Added tag parsing & actor self-reference when getting single point data

diff --git a/Engine/Plugins/Experimental/PCG/Source/PCG/Private/Data/PCGCollisionShapeData.cpp b/Engine/Plugins/Experimental/PCG/Source/PCG/Private/Data/PCGCollisionShapeData.cpp
--- a/Engine/Plugins/Experimental/PCG/Source/PCG/Private/Data/PCGCollisionShapeData.cpp
+++ b/Engine/Plugins/Experimental/PCG/Source/PCG/Private/Data/PCGCollisionShapeData.cpp
@@ -21,7 +21,6 @@ void UPCGCollisionShapeData::Initialize(UShapeComponent* InComponent)
 
 	//Note: Shape is pre-scaled
 	ShapeAdapter = MakeUnique<FPhysicsShapeAdapter>(Transform.GetRotation(), Shape);
-	TargetActor = InComponent->GetOwner();
 
 	CachedBounds = InComponent->Bounds.GetBox();
 	CachedStrictBounds = CachedBounds;
diff --git a/Engine/Plugins/Experimental/PCG/Source/PCG/Private/Data/PCGLandscapeData.cpp b/Engine/Plugins/Experimental/PCG/Source/PCG/Private/Data/PCGLandscapeData.cpp
--- a/Engine/Plugins/Experimental/PCG/Source/PCG/Private/Data/PCGLandscapeData.cpp
+++ b/Engine/Plugins/Experimental/PCG/Source/PCG/Private/Data/PCGLandscapeData.cpp
@@ -50,7 +50,6 @@ void UPCGLandscapeData::Initialize(const TArray<TWeakObjectPtr<ALandscapeProxy>>
 	ALandscapeProxy* FirstLandscape = Landscapes[0].Get();
 	check(FirstLandscape);
 
-	TargetActor = FirstLandscape;
 	Bounds = InBounds;
 	bHeightOnly = bInHeightOnly;
 	bUseMetadata = bInUseMetadata;
@@ -110,7 +109,6 @@ void UPCGLandscapeData::PostLoad()
 		}
 	}
 
-	TargetActor = FirstLandscape;
 	UPCGSubsystem* PCGSubsystem = UPCGSubsystem::GetInstance(FirstLandscape ? FirstLandscape->GetWorld() : nullptr);
 	LandscapeCache = PCGSubsystem ? PCGSubsystem->GetLandscapeCache() : nullptr;
 }
diff --git a/Engine/Plugins/Experimental/PCG/Source/PCG/Private/Data/PCGLandscapeSplineData.cpp b/Engine/Plugins/Experimental/PCG/Source/PCG/Private/Data/PCGLandscapeSplineData.cpp
--- a/Engine/Plugins/Experimental/PCG/Source/PCG/Private/Data/PCGLandscapeSplineData.cpp
+++ b/Engine/Plugins/Experimental/PCG/Source/PCG/Private/Data/PCGLandscapeSplineData.cpp
@@ -41,7 +41,6 @@ void UPCGLandscapeSplineData::Initialize(ULandscapeSplinesComponent* InSplineCom
 {
 	check(InSplineComponent);
 	Spline = InSplineComponent;
-	TargetActor = InSplineComponent->GetOwner();
 }
 
 FTransform UPCGLandscapeSplineData::GetTransform() const
diff --git a/Engine/Plugins/Experimental/PCG/Source/PCG/Private/Data/PCGPointData.cpp b/Engine/Plugins/Experimental/PCG/Source/PCG/Private/Data/PCGPointData.cpp
--- a/Engine/Plugins/Experimental/PCG/Source/PCG/Private/Data/PCGPointData.cpp
+++ b/Engine/Plugins/Experimental/PCG/Source/PCG/Private/Data/PCGPointData.cpp
@@ -396,8 +396,57 @@ void UPCGPointData::InitializeFromActor(AActor* InActor)
 	Points[0].BoundsMin = LocalBounds.Min;
 	Points[0].BoundsMax = LocalBounds.Max;
 
-	TargetActor = InActor;
 	Metadata = NewObject<UPCGMetadata>(this);
+
+	const bool bHasTags = (!InActor->Tags.IsEmpty());
+
+	// Create a metadata entry only when we have tags, since all point data initialized this way will always have the actor reference field, 
+	// it is therefore not needed to be careful about default values
+	if (bHasTags)
+	{
+		Points[0].MetadataEntry = Metadata->AddEntry();
+	}	
+
+	// TODO: Maybe make that a method somewhere else
+	// NOTE: for booleans, it's not sufficient to set the default value to true, it's just not going to be a good idea when merging point data together
+	// Add a self-reference attribute for use in GetActorData with single point
+	Metadata->CreateAttribute<FString>(TEXT("ActorReference"), FSoftObjectPath(InActor).ToString(), /*bAllowsInterpolation=*/false, /*bOverrideParent=*/false);
+
+	// Parse tags as well
+	for (FName Tag : InActor->Tags)
+	{
+		const FString TagString = Tag.ToString();
+		int32 EqualPosition = INDEX_NONE;
+		
+		if(TagString.FindChar(':', EqualPosition))
+		{
+			FString LeftSide = TagString.Left(EqualPosition);
+			FString RightSide = TagString.RightChop(EqualPosition+1);
+
+			if (LeftSide.IsEmpty() || RightSide.IsEmpty())
+			{
+				continue;
+			}
+
+			if(RightSide.IsNumeric())
+			{
+				double RightValue = FCString::Atod(*RightSide);
+
+				FPCGMetadataAttribute<double>* Attribute = Metadata->CreateAttribute<double>(FName(LeftSide), 0.0, /*bAllowsInterpolation=*/false, /*bOverrideParent=*/false);
+				Attribute->SetValue(Points[0].MetadataEntry, RightValue);
+			}
+			else
+			{
+				FPCGMetadataAttribute<FString>* Attribute = Metadata->CreateAttribute<FString>(FName(LeftSide), FString(), false, false);
+				Attribute->SetValue(Points[0].MetadataEntry, RightSide);
+			}
+		}
+		else
+		{
+			FPCGMetadataAttribute<bool>* Attribute = Metadata->CreateAttribute<bool>(Tag, false, /*bAllowsInterpolation=*/false, /*bOverrideParent=*/false);
+			Attribute->SetValue(Points[0].MetadataEntry, true);
+		}
+	}
 }
 
 FPCGPoint UPCGPointData::GetPoint(int32 Index) const
diff --git a/Engine/Plugins/Experimental/PCG/Source/PCG/Private/Data/PCGPrimitiveData.cpp b/Engine/Plugins/Experimental/PCG/Source/PCG/Private/Data/PCGPrimitiveData.cpp
--- a/Engine/Plugins/Experimental/PCG/Source/PCG/Private/Data/PCGPrimitiveData.cpp
+++ b/Engine/Plugins/Experimental/PCG/Source/PCG/Private/Data/PCGPrimitiveData.cpp
@@ -14,7 +14,6 @@ void UPCGPrimitiveData::Initialize(UPrimitiveComponent* InPrimitive)
 {
 	check(InPrimitive);
 	Primitive = InPrimitive;
-	TargetActor = InPrimitive->GetOwner();
 	CachedBounds = Primitive->Bounds.GetBox();
 	// Not obvious to find strict bounds, leave at the default value
 }
diff --git a/Engine/Plugins/Experimental/PCG/Source/PCG/Private/Data/PCGSplineData.cpp b/Engine/Plugins/Experimental/PCG/Source/PCG/Private/Data/PCGSplineData.cpp
--- a/Engine/Plugins/Experimental/PCG/Source/PCG/Private/Data/PCGSplineData.cpp
+++ b/Engine/Plugins/Experimental/PCG/Source/PCG/Private/Data/PCGSplineData.cpp
@@ -16,7 +16,6 @@
 void UPCGSplineData::Initialize(USplineComponent* InSpline)
 {
 	check(InSpline);
-	TargetActor = InSpline->GetOwner();
 
 	SplineStruct.Initialize(InSpline);
 
@@ -34,11 +33,8 @@ void UPCGSplineData::Initialize(USplineComponent* InSpline)
 	CachedBounds = CachedBounds.ExpandBy(SplinePointsRadius, SplinePointsRadius);
 }
 
-void UPCGSplineData::Initialize(const TArray<FSplinePoint>& InSplinePoints, bool bIsClosedLoop, AActor* InTargetActor, const FTransform& InTransform)
+void UPCGSplineData::Initialize(const TArray<FSplinePoint>& InSplinePoints, bool bIsClosedLoop, const FTransform& InTransform)
 {
-	check(InTargetActor);
-	TargetActor = InTargetActor;
-
 	SplineStruct.Initialize(InSplinePoints, bIsClosedLoop, InTransform);
 
 	CachedBounds = SplineStruct.GetBounds();
diff --git a/Engine/Plugins/Experimental/PCG/Source/PCG/Private/Data/PCGVolumeData.cpp b/Engine/Plugins/Experimental/PCG/Source/PCG/Private/Data/PCGVolumeData.cpp
--- a/Engine/Plugins/Experimental/PCG/Source/PCG/Private/Data/PCGVolumeData.cpp
+++ b/Engine/Plugins/Experimental/PCG/Source/PCG/Private/Data/PCGVolumeData.cpp
@@ -10,11 +10,10 @@
 
 #include UE_INLINE_GENERATED_CPP_BY_NAME(PCGVolumeData)
 
-void UPCGVolumeData::Initialize(AVolume* InVolume, AActor* InTargetActor)
+void UPCGVolumeData::Initialize(AVolume* InVolume)
 {
 	check(InVolume);
 	Volume = InVolume;
-	TargetActor = InTargetActor ? InTargetActor : InVolume;
 	
 	FBoxSphereBounds BoxSphereBounds = Volume->GetBounds();
 	Bounds = FBox::BuildAABB(BoxSphereBounds.Origin, BoxSphereBounds.BoxExtent);
@@ -23,11 +22,10 @@ void UPCGVolumeData::Initialize(AVolume* InVolume, AActor* InTargetActor)
 	// Currently, we'll leave the strict bounds empty and fall back to checking against the local box
 }
 
-void UPCGVolumeData::Initialize(const FBox& InBounds, AActor* InTargetActor)
+void UPCGVolumeData::Initialize(const FBox& InBounds)
 {
 	Bounds = InBounds;
 	StrictBounds = InBounds;
-	TargetActor = InTargetActor;
 }
 
 FBox UPCGVolumeData::GetBounds() const
diff --git a/Engine/Plugins/Experimental/PCG/Source/PCG/Private/Data/PCGWorldData.cpp b/Engine/Plugins/Experimental/PCG/Source/PCG/Private/Data/PCGWorldData.cpp
--- a/Engine/Plugins/Experimental/PCG/Source/PCG/Private/Data/PCGWorldData.cpp
+++ b/Engine/Plugins/Experimental/PCG/Source/PCG/Private/Data/PCGWorldData.cpp
@@ -48,8 +48,7 @@ void FPCGWorldRayHitQueryParams::Initialize()
 
 void UPCGWorldVolumetricData::Initialize(UWorld* InWorld, const FBox& InBounds)
 {
-	Super::Initialize(InBounds, nullptr);
-	TargetActor = nullptr;
+	Super::Initialize(InBounds);
 	World = InWorld;
 	check(World.IsValid());
 }
diff --git a/Engine/Plugins/Experimental/PCG/Source/PCG/Private/Elements/Metadata/PCGMetadataPartition.cpp b/Engine/Plugins/Experimental/PCG/Source/PCG/Private/Elements/Metadata/PCGMetadataPartition.cpp
--- a/Engine/Plugins/Experimental/PCG/Source/PCG/Private/Elements/Metadata/PCGMetadataPartition.cpp
+++ b/Engine/Plugins/Experimental/PCG/Source/PCG/Private/Elements/Metadata/PCGMetadataPartition.cpp
@@ -72,6 +72,8 @@ bool FPCGMetadataPartitionElement::ExecuteInternal(FPCGContext* Context) const
 		TArray<PCGMetadataValueKey> ValueKeyMapping;
 		ValueKeyMapping.Reserve(MetadataValueKeyCount);
 
+		TArray<PCGMetadataValueKey> UniqueValueKeys;
+
 		const bool bUsesValueKeys = AttributeBase->UsesValueKeys();
 
 		for (PCGMetadataValueKey ValueKey = 0; ValueKey < MetadataValueKeyCount; ++ValueKey)
@@ -82,12 +84,18 @@ bool FPCGMetadataPartitionElement::ExecuteInternal(FPCGContext* Context) const
 			}
 			else if (!bUsesValueKeys)
 			{
-				PCGMetadataValueKey* MatchingVK = Algo::FindByPredicate(ValueKeyMapping, [ValueKey, AttributeBase](const PCGMetadataValueKey& Key)
+				PCGMetadataValueKey* MatchingVK = Algo::FindByPredicate(UniqueValueKeys, [ValueKey, AttributeBase](const PCGMetadataValueKey& Key)
 				{
 					return AttributeBase->AreValuesEqual(ValueKey, Key);
 				});
 
 				const PCGMetadataValueKey TentativeKey = MatchingVK ? *MatchingVK : ValueKey;
+
+				if (!MatchingVK)
+				{
+					UniqueValueKeys.Add(TentativeKey);
+				}
+
 				ValueKeyMapping.Add(TentativeKey);
 			}
 			else
diff --git a/Engine/Plugins/Experimental/PCG/Source/PCG/Private/Elements/PCGCreateSpline.cpp b/Engine/Plugins/Experimental/PCG/Source/PCG/Private/Elements/PCGCreateSpline.cpp
--- a/Engine/Plugins/Experimental/PCG/Source/PCG/Private/Elements/PCGCreateSpline.cpp
+++ b/Engine/Plugins/Experimental/PCG/Source/PCG/Private/Elements/PCGCreateSpline.cpp
@@ -83,18 +83,18 @@ bool FPCGCreateSplineElement::ExecuteInternal(FPCGContext* Context) const
 			continue;
 		}
 
-		const UPCGPointData* PointData = SpatialData->ToPointData(Context);
-
-		if (!PointData)
+		AActor* TargetActor = Settings->TargetActor.Get() ? Settings->TargetActor.Get() : Context->GetTargetActor(SpatialData);
+		if (!TargetActor)
 		{
-			PCGE_LOG(Error, GraphAndLog, LOCTEXT("UnableToGetPointData", "Unable to get point data from input"));
+			PCGE_LOG(Error, GraphAndLog, LOCTEXT("InvalidTargetActor", "Invalid target actor. Ensure TargetActor member is initialized when creating SpatialData."));
 			continue;
 		}
 
-		AActor* TargetActor = PointData->TargetActor.Get();
-		if (!TargetActor)
+		const UPCGPointData* PointData = SpatialData->ToPointData(Context);
+
+		if (!PointData)
 		{
-			PCGE_LOG(Error, GraphAndLog, LOCTEXT("InvalidTargetActor", "Invalid target actor. Ensure TargetActor member is initialized when creating SpatialData."));
+			PCGE_LOG(Error, GraphAndLog, LOCTEXT("UnableToGetPointData", "Unable to get point data from input"));
 			continue;
 		}
 
@@ -200,7 +200,7 @@ bool FPCGCreateSplineElement::ExecuteInternal(FPCGContext* Context) const
 				PointType);
 		}
 
-		SplineData->Initialize(SplinePoints, Settings->bClosedLoop, TargetActor, FTransform(SplineActorTransform.GetLocation()));
+		SplineData->Initialize(SplinePoints, Settings->bClosedLoop, FTransform(SplineActorTransform.GetLocation()));
 
 		USplineComponent* SplineComponent = nullptr;
 
diff --git a/Engine/Plugins/Experimental/PCG/Source/PCG/Private/Elements/PCGCreateTargetActor.cpp b/Engine/Plugins/Experimental/PCG/Source/PCG/Private/Elements/PCGCreateTargetActor.cpp
new file mode 100644
--- /dev/null
+++ b/Engine/Plugins/Experimental/PCG/Source/PCG/Private/Elements/PCGCreateTargetActor.cpp
@@ -0,0 +1,254 @@
+// Copyright Epic Games, Inc. All Rights Reserved.
+
+#include "Elements/PCGCreateTargetActor.h"
+
+#include "PCGComponent.h"
+#include "PCGContext.h"
+#include "PCGElement.h"
+#include "PCGManagedResource.h"
+#include "PCGModule.h"
+#include "PCGParamData.h"
+#include "Helpers/PCGActorHelpers.h"
+#include "Helpers/PCGHelpers.h"
+
+#include "Engine/Blueprint.h"
+#include "Engine/BlueprintGeneratedClass.h"
+#include "Engine/Engine.h"
+#include "GameFramework/Actor.h"
+#include "UObject/Package.h"
+
+#include UE_INLINE_GENERATED_CPP_BY_NAME(PCGCreateTargetActor)
+
+#define LOCTEXT_NAMESPACE "PCGCreateTargetActor"
+
+#if WITH_EDITOR
+
+FText UPCGCreateTargetActor::GetDefaultNodeTitle() const
+{
+	return LOCTEXT("NodeTitle", "Create Target Actor");
+}
+
+#endif // WITH_EDITOR
+
+FPCGElementPtr UPCGCreateTargetActor::CreateElement() const
+{
+	return MakeShared<FPCGCreateTargetActorElement>();
+}
+
+TArray<FPCGPinProperties> UPCGCreateTargetActor::InputPinProperties() const
+{
+	return TArray<FPCGPinProperties>();
+}
+
+TArray<FPCGPinProperties> UPCGCreateTargetActor::OutputPinProperties() const
+{
+	TArray<FPCGPinProperties> PinProperties;
+	PinProperties.Emplace_GetRef(PCGPinConstants::DefaultOutputLabel, EPCGDataType::Param, false);
+
+	return PinProperties;
+}
+
+void UPCGCreateTargetActor::BeginDestroy()
+{
+#if WITH_EDITOR
+	TeardownBlueprintEvent();
+#endif
+
+	Super::BeginDestroy();
+}
+
+#if WITH_EDITOR
+void UPCGCreateTargetActor::SetupBlueprintEvent()
+{
+	if (UBlueprintGeneratedClass* BlueprintClass = Cast<UBlueprintGeneratedClass>(TemplateActorClass))
+	{
+		if (UBlueprint* Blueprint = Cast<UBlueprint>(BlueprintClass->ClassGeneratedBy))
+		{
+			Blueprint->OnChanged().AddUObject(this, &UPCGCreateTargetActor::OnBlueprintChanged);
+		}
+	}
+}
+
+void UPCGCreateTargetActor::TeardownBlueprintEvent()
+{
+	if (UBlueprintGeneratedClass* BlueprintClass = Cast<UBlueprintGeneratedClass>(TemplateActorClass))
+	{
+		if (UBlueprint* Blueprint = Cast<UBlueprint>(BlueprintClass->ClassGeneratedBy))
+		{
+			Blueprint->OnChanged().RemoveAll(this);
+		}
+	}
+}
+
+void UPCGCreateTargetActor::PreEditChange(FProperty* PropertyAboutToChange)
+{
+	if (PropertyAboutToChange && PropertyAboutToChange->GetFName() == GET_MEMBER_NAME_CHECKED(UPCGCreateTargetActor, TemplateActorClass))
+	{
+		TeardownBlueprintEvent();
+	}
+}
+
+void UPCGCreateTargetActor::PostEditChangeProperty(FPropertyChangedEvent& PropertyChangedEvent)
+{
+	if (PropertyChangedEvent.Property)
+	{
+		const FName& PropertyName = PropertyChangedEvent.Property->GetFName();
+
+		if (PropertyName == GET_MEMBER_NAME_CHECKED(UPCGCreateTargetActor, TemplateActorClass))
+		{
+			SetupBlueprintEvent();
+			RefreshTemplateActor();
+		}
+	}
+
+	Super::PostEditChangeProperty(PropertyChangedEvent);
+}
+
+void UPCGCreateTargetActor::PreEditUndo()
+{
+	TeardownBlueprintEvent();
+
+	Super::PreEditUndo();
+}
+
+void UPCGCreateTargetActor::PostEditUndo()
+{
+	Super::PostEditUndo();
+
+	SetupBlueprintEvent();
+	RefreshTemplateActor();
+}
+
+void UPCGCreateTargetActor::OnBlueprintChanged(UBlueprint* InBlueprint)
+{
+	RefreshTemplateActor();
+	OnSettingsChangedDelegate.Broadcast(this, EPCGChangeType::Settings);
+}
+
+#endif // WITH_EDITOR
+
+void UPCGCreateTargetActor::RefreshTemplateActor()
+{
+	if (TemplateActorClass)
+	{
+		AActor* NewTemplateActor = NewObject<AActor>(this, TemplateActorClass, NAME_None, RF_ArchetypeObject | RF_Transactional | RF_Public);
+
+		if (TemplateActor)
+		{
+			UEngine::FCopyPropertiesForUnrelatedObjectsParams Options;
+			Options.bNotifyObjectReplacement = true;
+			UEngine::CopyPropertiesForUnrelatedObjects(TemplateActor, NewTemplateActor, Options);
+
+			TemplateActor->Rename(nullptr, GetTransientPackage(), REN_DontCreateRedirectors | REN_ForceNoResetLoaders);
+		}
+
+		TemplateActor = NewTemplateActor;
+	}
+	else
+	{
+		TemplateActor = nullptr;
+	}
+}
+
+void UPCGCreateTargetActor::PostLoad()
+{
+	Super::PostLoad();
+
+#if WITH_EDITOR
+	SetupBlueprintEvent();
+
+	if (TemplateActorClass)
+	{
+		if (TemplateActor)
+		{
+			TemplateActor->ConditionalPostLoad();
+		}
+
+		RefreshTemplateActor();
+	}
+#endif // WITH_EDITOR
+}
+
+bool FPCGCreateTargetActorElement::ExecuteInternal(FPCGContext* Context) const
+{
+	TRACE_CPUPROFILER_EVENT_SCOPE(FPCGCreateTargetActorElement::Execute);
+
+	// Early out if the actor isn't going to be consumed by something else
+	if (Context->Node && !Context->Node->IsOutputPinConnected(PCGPinConstants::DefaultOutputLabel))
+	{
+		return true;
+	}
+
+	const UPCGCreateTargetActor* Settings = Context->GetInputSettings<UPCGCreateTargetActor>();
+	check(Settings);
+
+	// Early out if the template actor isn't valid
+	if (!Settings->TemplateActorClass || Settings->TemplateActorClass->HasAnyClassFlags(CLASS_Abstract))
+	{
+		const FText ClassName = Settings->TemplateActorClass ? FText::FromString(Settings->TemplateActorClass->GetFName().ToString()) : FText::FromName(NAME_None);
+		PCGE_LOG(Error, GraphAndLog, FText::Format(LOCTEXT("InvalidTemplateActorClass", "Invalid template actor class '{0}'"), ClassName));
+		return true;
+	}
+
+	if (!ensure(Settings->TemplateActor && Settings->TemplateActor->IsA(Settings->TemplateActorClass)))
+	{
+		return true;
+	}
+
+	AActor* TargetActor = Settings->RootActor.Get() ? Settings->RootActor.Get() : Context->GetTargetActor(nullptr);
+	if (!TargetActor)
+	{
+		PCGE_LOG(Error, GraphAndLog, LOCTEXT("InvalidTargetActor", "Invalid target actor"));
+		return true;
+	}
+
+	const bool bHasAuthority = !Context->SourceComponent.IsValid() || (Context->SourceComponent->GetOwner() && Context->SourceComponent->GetOwner()->HasAuthority());
+	const bool bSpawnedActorRequiresAuthority = Settings->TemplateActor->GetIsReplicated();
+
+	if (!bHasAuthority && bSpawnedActorRequiresAuthority)
+	{
+		return true;
+	}
+
+	// Spawn actor
+	AActor* TemplateActor = Settings->TemplateActor.Get();
+
+	FActorSpawnParameters SpawnParams;
+	SpawnParams.Owner = TargetActor;
+	SpawnParams.Template = TemplateActor;
+	SpawnParams.SpawnCollisionHandlingOverride = ESpawnActorCollisionHandlingMethod::AlwaysSpawn;
+
+	if (PCGHelpers::IsRuntimeOrPIE())
+	{
+		SpawnParams.ObjectFlags |= RF_Transient;
+	}
+
+	FTransform Transform = TargetActor->GetTransform();
+	AActor* GeneratedActor = UPCGActorHelpers::SpawnDefaultActor(TargetActor->GetWorld(), Settings->TemplateActorClass, Transform, SpawnParams, TargetActor);
+
+	if (!GeneratedActor)
+	{
+		PCGE_LOG(Error, GraphAndLog, LOCTEXT("ActorSpawnFailed", "Failed to spawn actor"));
+		return true;
+	}
+
+	GeneratedActor->Tags.Add(PCGHelpers::DefaultPCGActorTag);
+
+	if (UPCGComponent* SourceComponent = Context->SourceComponent.Get())
+	{
+		UPCGManagedActors* ManagedActors = NewObject<UPCGManagedActors>(SourceComponent);
+		ManagedActors->GeneratedActors.Add(GeneratedActor);
+		SourceComponent->AddToManagedResources(ManagedActors);
+	}
+
+	// Create param data output with reference to actor
+	UPCGParamData* ParamData = NewObject<UPCGParamData>();
+	check(ParamData && ParamData->Metadata);
+	ParamData->Metadata->CreateAttribute<FString>(NAME_None, FSoftObjectPath(GeneratedActor).ToString(), /*bAllowsInterpolation=*/false, /*bOverrideParent=*/false);
+
+	// Add param data to output and we're done
+	Context->OutputData.TaggedData.Emplace_GetRef().Data = ParamData;
+	return true;
+}
+
+#undef LOCTEXT_NAMESPACE
\ No newline at end of file
diff --git a/Engine/Plugins/Experimental/PCG/Source/PCG/Private/Elements/PCGDataFromActor.cpp b/Engine/Plugins/Experimental/PCG/Source/PCG/Private/Elements/PCGDataFromActor.cpp
--- a/Engine/Plugins/Experimental/PCG/Source/PCG/Private/Elements/PCGDataFromActor.cpp
+++ b/Engine/Plugins/Experimental/PCG/Source/PCG/Private/Elements/PCGDataFromActor.cpp
@@ -333,17 +333,6 @@ void FPCGDataFromActorElement::ProcessActor(FPCGContext* Context, const UPCGData
 		FPCGDataCollection Collection = UPCGComponent::CreateActorPCGDataCollection(FoundActor, Context->SourceComponent.Get(), DataFilter, bParseActor);
 		Outputs += Collection.TaggedData;
 	}
-
-	if (Context->SourceComponent.IsValid())
-	{
-		for (FPCGTaggedData& Output : Outputs)
-		{
-			if (UPCGSpatialData* SpatialData = Cast<UPCGSpatialData>(Output.Data))
-			{
-				SpatialData->TargetActor = Context->SourceComponent->GetOwner();
-			}
-		}
-	}
 }
 
 #undef LOCTEXT_NAMESPACE
diff --git a/Engine/Plugins/Experimental/PCG/Source/PCG/Private/Elements/PCGDebugElement.cpp b/Engine/Plugins/Experimental/PCG/Source/PCG/Private/Elements/PCGDebugElement.cpp
--- a/Engine/Plugins/Experimental/PCG/Source/PCG/Private/Elements/PCGDebugElement.cpp
+++ b/Engine/Plugins/Experimental/PCG/Source/PCG/Private/Elements/PCGDebugElement.cpp
@@ -22,7 +22,7 @@
 
 namespace PCGDebugElement
 {
-	void ExecuteDebugDisplay(FPCGContext* Context)
+	void ExecuteDebugDisplay(FPCGContext* Context, AActor* InTargetActor)
 	{
 		TRACE_CPUPROFILER_EVENT_SCOPE(PCGDebugElement::ExecuteDebugDisplay);
 #if WITH_EDITOR
@@ -93,7 +93,7 @@ namespace PCGDebugElement
 				continue;
 			}
 
-			AActor* TargetActor = Context->GetTargetActor(SpatialData);
+			AActor* TargetActor = InTargetActor ? InTargetActor : Context->GetTargetActor(SpatialData);
 
 			if (!TargetActor)
 			{
@@ -244,7 +244,7 @@ TArray<FPCGPinProperties> UPCGDebugSettings::OutputPinProperties() const
 bool FPCGDebugElement::ExecuteInternal(FPCGContext* Context) const
 {
 	TRACE_CPUPROFILER_EVENT_SCOPE(FPCGDebugElement::Execute);
-	PCGDebugElement::ExecuteDebugDisplay(Context);
+	PCGDebugElement::ExecuteDebugDisplay(Context, Context->GetInputSettings<UPCGDebugSettings>()->TargetActor.Get());
 	
 	return true;
 }
diff --git a/Engine/Plugins/Experimental/PCG/Source/PCG/Private/Elements/PCGSpawnActor.cpp b/Engine/Plugins/Experimental/PCG/Source/PCG/Private/Elements/PCGSpawnActor.cpp
--- a/Engine/Plugins/Experimental/PCG/Source/PCG/Private/Elements/PCGSpawnActor.cpp
+++ b/Engine/Plugins/Experimental/PCG/Source/PCG/Private/Elements/PCGSpawnActor.cpp
@@ -410,7 +410,7 @@ bool FPCGSpawnActorElement::ExecuteInternal(FPCGContext* Context) const
 			continue;
 		}
 
-		AActor* TargetActor = Context->GetTargetActor(SpatialData);
+		AActor* TargetActor = Settings->RootActor.Get() ? Settings->RootActor.Get() : Context->GetTargetActor(SpatialData);
 
 		if (!TargetActor)
 		{
diff --git a/Engine/Plugins/Experimental/PCG/Source/PCG/Private/Elements/PCGStaticMeshSpawner.cpp b/Engine/Plugins/Experimental/PCG/Source/PCG/Private/Elements/PCGStaticMeshSpawner.cpp
--- a/Engine/Plugins/Experimental/PCG/Source/PCG/Private/Elements/PCGStaticMeshSpawner.cpp
+++ b/Engine/Plugins/Experimental/PCG/Source/PCG/Private/Elements/PCGStaticMeshSpawner.cpp
@@ -163,7 +163,7 @@ bool FPCGStaticMeshSpawnerElement::PrepareDataInternal(FPCGContext* InContext) c
 				continue;
 			}
 
-			AActor* TargetActor = Context->GetTargetActor(PointData);
+			AActor* TargetActor = Settings->TargetActor.Get() ? Settings->TargetActor.Get() : Context->GetTargetActor(PointData);
 			if (!TargetActor)
 			{
 				PCGE_LOG(Error, GraphAndLog, LOCTEXT("InvalidTargetActor", "Invalid target actor. Ensure TargetActor member is initialized when creating SpatialData."));
diff --git a/Engine/Plugins/Experimental/PCG/Source/PCG/Private/Elements/PCGTextureSampler.cpp b/Engine/Plugins/Experimental/PCG/Source/PCG/Private/Elements/PCGTextureSampler.cpp
--- a/Engine/Plugins/Experimental/PCG/Source/PCG/Private/Elements/PCGTextureSampler.cpp
+++ b/Engine/Plugins/Experimental/PCG/Source/PCG/Private/Elements/PCGTextureSampler.cpp
@@ -73,8 +73,6 @@ bool FPCGTextureSamplerElement::ExecuteInternal(FPCGContext* Context) const
 	UPCGTextureData* TextureData = NewObject<UPCGTextureData>();
 	Output.Data = TextureData;
 
-	TextureData->TargetActor = Context->SourceComponent->GetOwner();
-
 	AActor* OriginalActor = UPCGBlueprintHelpers::GetOriginalComponent(*Context)->GetOwner();
 	FTransform FinalTransform = Transform;
 	if (!bUseAbsoluteTransform)
diff --git a/Engine/Plugins/Experimental/PCG/Source/PCG/Private/Elements/PCGTypedGetter.cpp b/Engine/Plugins/Experimental/PCG/Source/PCG/Private/Elements/PCGTypedGetter.cpp
--- a/Engine/Plugins/Experimental/PCG/Source/PCG/Private/Elements/PCGTypedGetter.cpp
+++ b/Engine/Plugins/Experimental/PCG/Source/PCG/Private/Elements/PCGTypedGetter.cpp
@@ -105,12 +105,6 @@ void FPCGGetLandscapeDataElement::ProcessActors(FPCGContext* Context, const UPCG
 		UPCGLandscapeData* LandscapeData = NewObject<UPCGLandscapeData>();
 		LandscapeData->Initialize(Landscapes, LandscapeBounds, Settings->bGetHeightOnly, Settings->bGetLayerWeights);
 		
-		// Override target actor for the one we're running on, to make sure we don't output on the landscape directly
-		if (Context->SourceComponent.IsValid())
-		{
-			LandscapeData->TargetActor = Context->SourceComponent->GetOwner();
-		}
-
 		FPCGTaggedData& TaggedData = Context->OutputData.TaggedData.Emplace_GetRef();
 		TaggedData.Data = LandscapeData;
 		TaggedData.Tags = LandscapeTags;
diff --git a/Engine/Plugins/Experimental/PCG/Source/PCG/Private/Elements/PCGWorldQuery.cpp b/Engine/Plugins/Experimental/PCG/Source/PCG/Private/Elements/PCGWorldQuery.cpp
--- a/Engine/Plugins/Experimental/PCG/Source/PCG/Private/Elements/PCGWorldQuery.cpp
+++ b/Engine/Plugins/Experimental/PCG/Source/PCG/Private/Elements/PCGWorldQuery.cpp
@@ -51,7 +51,6 @@ bool FPCGWorldVolumetricQueryElement::ExecuteInternal(FPCGContext* Context) cons
 	Data->QueryParams = QueryParams;
 	Data->QueryParams.Initialize();
 	Data->OriginatingComponent = Context->SourceComponent;
-	Data->TargetActor = Context->SourceComponent->GetOwner();
 
 	FPCGTaggedData& Output = Context->OutputData.TaggedData.Emplace_GetRef();
 	Output.Data = Data;
@@ -129,7 +128,6 @@ bool FPCGWorldRayHitQueryElement::ExecuteInternal(FPCGContext* Context) const
 	Data->QueryParams = QueryParams;
 	Data->QueryParams.Initialize();
 	Data->OriginatingComponent = Context->SourceComponent;
-	Data->TargetActor = Context->SourceComponent->GetOwner();
 
 	FPCGTaggedData& Output = Context->OutputData.TaggedData.Emplace_GetRef();
 	Output.Data = Data;
diff --git a/Engine/Plugins/Experimental/PCG/Source/PCG/Private/Helpers/PCGActorHelpers.cpp b/Engine/Plugins/Experimental/PCG/Source/PCG/Private/Helpers/PCGActorHelpers.cpp
--- a/Engine/Plugins/Experimental/PCG/Source/PCG/Private/Helpers/PCGActorHelpers.cpp
+++ b/Engine/Plugins/Experimental/PCG/Source/PCG/Private/Helpers/PCGActorHelpers.cpp
@@ -45,7 +45,7 @@ UPCGManagedISMComponent* UPCGActorHelpers::GetOrCreateManagedISMC(AActor* InTarg
 	{
 		TRACE_CPUPROFILER_EVENT_SCOPE(UPCGActorHelpers::GetOrCreateManagedISMC::FindMatchingMISMC);
 		UPCGManagedISMComponent* MatchingResource = nullptr;
-		InSourceComponent->ForEachManagedResource([&MatchingResource, &InParams, SettingsUID](UPCGManagedResource* InResource)
+		InSourceComponent->ForEachManagedResource([&MatchingResource, &InParams, &InTargetActor, SettingsUID](UPCGManagedResource* InResource)
 		{
 			// Early out if already found a match
 			if (MatchingResource)
@@ -63,6 +63,7 @@ UPCGManagedISMComponent* UPCGActorHelpers::GetOrCreateManagedISMC(AActor* InTarg
 				if (UInstancedStaticMeshComponent* ISMC = Resource->GetComponent())
 				{
 					if (IsValid(ISMC) &&
+						ISMC->GetOwner() == InTargetActor &&
 						ISMC->NumCustomDataFloats == InParams.NumCustomDataFloats &&
 						Resource->GetDescriptor() == InParams.Descriptor)
 					{
@@ -125,6 +126,7 @@ UPCGManagedISMComponent* UPCGActorHelpers::GetOrCreateManagedISMC(AActor* InTarg
 	{
 		Resource->SetRootLocation(InTargetActor->GetRootComponent()->GetComponentLocation());
 	}
+	
 	Resource->SetSettingsUID(SettingsUID);
 	InSourceComponent->AddToManagedResources(Resource);
 
@@ -279,6 +281,22 @@ AActor* UPCGActorHelpers::SpawnDefaultActor(UWorld* World, TSubclassOf<AActor> A
 	FActorSpawnParameters SpawnParams;
 	SpawnParams.Name = MakeUniqueObjectName(World->GetCurrentLevel(), ActorClass, BaseName);
 	SpawnParams.Owner = Parent;
+
+	if (PCGHelpers::IsRuntimeOrPIE())
+	{
+		SpawnParams.ObjectFlags |= RF_Transient;
+	}
+
+	return SpawnDefaultActor(World, ActorClass, Transform, SpawnParams, Parent);
+}
+
+AActor* UPCGActorHelpers::SpawnDefaultActor(UWorld* World, TSubclassOf<AActor> ActorClass, const FTransform& Transform, const FActorSpawnParameters& SpawnParams, AActor* Parent)
+{
+	if (!World || !ActorClass)
+	{
+		return nullptr;
+	}
+
 	AActor* NewActor = World->SpawnActor(*ActorClass, &Transform, SpawnParams);
 	
 	if (!NewActor)
@@ -286,8 +304,14 @@ AActor* UPCGActorHelpers::SpawnDefaultActor(UWorld* World, TSubclassOf<AActor> A
 		return nullptr;
 	}
 
+	// HACK: until UE-62747 is fixed, we have to force set the scale after spawning the actor
+	NewActor->SetActorRelativeScale3D(Transform.GetScale3D());
+
 #if WITH_EDITOR
-	NewActor->SetActorLabel(SpawnParams.Name.ToString());
+	if (SpawnParams.Name != NAME_None)
+	{
+		NewActor->SetActorLabel(SpawnParams.Name.ToString());
+	}
 #endif // WITH_EDITOR
 
 	USceneComponent* RootComponent = NewActor->GetRootComponent();
diff --git a/Engine/Plugins/Experimental/PCG/Source/PCG/Private/PCGComponent.cpp b/Engine/Plugins/Experimental/PCG/Source/PCG/Private/PCGComponent.cpp
--- a/Engine/Plugins/Experimental/PCG/Source/PCG/Private/PCGComponent.cpp
+++ b/Engine/Plugins/Experimental/PCG/Source/PCG/Private/PCGComponent.cpp
@@ -708,7 +708,7 @@ bool UPCGComponent::MoveResourcesToNewActor(AActor* InNewActor, bool bCreateChil
 		{
 			if (GeneratedResource)
 			{
-				GeneratedResource->MoveResourceToNewActor(NewActor);
+				GeneratedResource->MoveResourceToNewActor(NewActor, Owner);
 				TSet<TSoftObjectPtr<AActor>> Dummy;
 				GeneratedResource->ReleaseIfUnused(Dummy);
 				bHasMovedResources = true;
@@ -2004,7 +2004,7 @@ FPCGDataCollection UPCGComponent::CreateActorPCGDataCollection(AActor* Actor, co
 		check(!Component || Component->GetOwner() == Actor); // Invalid processing otherwise because of the this usage
 
 		UPCGVolumeData* Data = NewObject<UPCGVolumeData>();
-		Data->Initialize(PartitionActor->GetFixedBounds(), PartitionActor);
+		Data->Initialize(PartitionActor->GetFixedBounds());
 
 		UPCGComponent* OriginalComponent = Component ? PartitionActor->GetOriginalComponent(Component) : nullptr;
 		// Important note: we do NOT call the collection version here, as we want to have a union if that's the case
@@ -2248,8 +2248,6 @@ UPCGData* UPCGComponent::CreateLandscapePCGData(bool bHeightOnly)
 	UPCGLandscapeData* LandscapeData = NewObject<UPCGLandscapeData>();
 	const UPCGGraph* PCGGraph = GetGraph();
 	LandscapeData->Initialize(Landscapes, LandscapeBounds, bHeightOnly, /*bUseMetadata=*/PCGGraph && PCGGraph->bLandscapeUsesMetadata);
-	// Need to override target actor for this one, not the landscape
-	LandscapeData->TargetActor = Actor;
 
 	return LandscapeData;
 }
diff --git a/Engine/Plugins/Experimental/PCG/Source/PCG/Private/PCGContext.cpp b/Engine/Plugins/Experimental/PCG/Source/PCG/Private/PCGContext.cpp
--- a/Engine/Plugins/Experimental/PCG/Source/PCG/Private/PCGContext.cpp
+++ b/Engine/Plugins/Experimental/PCG/Source/PCG/Private/PCGContext.cpp
@@ -79,13 +79,17 @@ FString FPCGContext::GetComponentName() const
 
 AActor* FPCGContext::GetTargetActor(const UPCGSpatialData* InSpatialData) const
 {
-	if (SourceComponent.IsValid() && SourceComponent->GetOwner())
+	if (InSpatialData && InSpatialData->TargetActor.Get())
+	{
+		return InSpatialData->TargetActor.Get();
+	}
+	else if (SourceComponent.IsValid() && SourceComponent->GetOwner())
 	{
 		return SourceComponent->GetOwner();
 	}
 	else
 	{
-		return InSpatialData ? InSpatialData->TargetActor.Get() : nullptr;
+		return nullptr;
 	}
 }
 
diff --git a/Engine/Plugins/Experimental/PCG/Source/PCG/Private/PCGManagedResource.cpp b/Engine/Plugins/Experimental/PCG/Source/PCG/Private/PCGManagedResource.cpp
--- a/Engine/Plugins/Experimental/PCG/Source/PCG/Private/PCGManagedResource.cpp
+++ b/Engine/Plugins/Experimental/PCG/Source/PCG/Private/PCGManagedResource.cpp
@@ -225,7 +225,7 @@ bool UPCGManagedComponent::ReleaseIfUnused(TSet<TSoftObjectPtr<AActor>>& OutActo
 	return Super::ReleaseIfUnused(OutActorsToDelete) || !GeneratedComponent.IsValid();
 }
 
-bool UPCGManagedComponent::MoveResourceToNewActor(AActor* NewActor)
+bool UPCGManagedComponent::MoveResourceToNewActor(AActor* NewActor, const AActor* ExpectedPreviousOwner)
 {
 	check(NewActor);
 
@@ -237,6 +237,12 @@ bool UPCGManagedComponent::MoveResourceToNewActor(AActor* NewActor)
 	TObjectPtr<AActor> OldOwner = GeneratedComponent->GetOwner();
 	check(OldOwner);
 
+	// Prevent moving of components on external (or spawned) actors
+	if (ExpectedPreviousOwner && OldOwner != ExpectedPreviousOwner)
+	{
+		return false;
+	}
+
 	bool bDetached = false;
 	bool bAttached = false;
 
diff --git a/Engine/Plugins/Experimental/PCG/Source/PCG/Private/Tests/PCGTestsCommon.cpp b/Engine/Plugins/Experimental/PCG/Source/PCG/Private/Tests/PCGTestsCommon.cpp
--- a/Engine/Plugins/Experimental/PCG/Source/PCG/Private/Tests/PCGTestsCommon.cpp
+++ b/Engine/Plugins/Experimental/PCG/Source/PCG/Private/Tests/PCGTestsCommon.cpp
@@ -166,7 +166,7 @@ namespace PCGTestsCommon
 	UPCGVolumeData* CreateVolumeData(const FBox& InBounds)
 	{
 		UPCGVolumeData* VolumeData = NewObject<UPCGVolumeData>();
-		VolumeData->Initialize(InBounds, nullptr);
+		VolumeData->Initialize(InBounds);
 		return VolumeData;
 	}
 
diff --git a/Engine/Plugins/Experimental/PCG/Source/PCG/Public/Data/PCGSplineData.h b/Engine/Plugins/Experimental/PCG/Source/PCG/Public/Data/PCGSplineData.h
--- a/Engine/Plugins/Experimental/PCG/Source/PCG/Public/Data/PCGSplineData.h
+++ b/Engine/Plugins/Experimental/PCG/Source/PCG/Public/Data/PCGSplineData.h
@@ -22,7 +22,7 @@ class PCG_API UPCGSplineData : public UPCGPolyLineData
 
 public:
 	void Initialize(USplineComponent* InSpline);
-	void Initialize(const TArray<FSplinePoint>& InSplinePoints, bool bInClosedLoop, AActor* InTargetActor, const FTransform& InTransform);
+	void Initialize(const TArray<FSplinePoint>& InSplinePoints, bool bInClosedLoop, const FTransform& InTransform);
 	void ApplyTo(USplineComponent* InSpline);
 
 	// ~Begin UPCGData interface
diff --git a/Engine/Plugins/Experimental/PCG/Source/PCG/Public/Data/PCGVolumeData.h b/Engine/Plugins/Experimental/PCG/Source/PCG/Public/Data/PCGVolumeData.h
--- a/Engine/Plugins/Experimental/PCG/Source/PCG/Public/Data/PCGVolumeData.h
+++ b/Engine/Plugins/Experimental/PCG/Source/PCG/Public/Data/PCGVolumeData.h
@@ -14,8 +14,8 @@ class PCG_API UPCGVolumeData : public UPCGSpatialDataWithPointCache
 	GENERATED_BODY()
 
 public:
-	void Initialize(AVolume* InVolume, AActor* InTargetActor = nullptr);
-	void Initialize(const FBox& InBounds, AActor* InTargetActor);
+	void Initialize(AVolume* InVolume);
+	void Initialize(const FBox& InBounds);
 
 	// ~Begin UPCGData interface
 	virtual EPCGDataType GetDataType() const override { return EPCGDataType::Volume; }
diff --git a/Engine/Plugins/Experimental/PCG/Source/PCG/Public/Elements/PCGCreateSpline.h b/Engine/Plugins/Experimental/PCG/Source/PCG/Public/Elements/PCGCreateSpline.h
--- a/Engine/Plugins/Experimental/PCG/Source/PCG/Public/Elements/PCGCreateSpline.h
+++ b/Engine/Plugins/Experimental/PCG/Source/PCG/Public/Elements/PCGCreateSpline.h
@@ -61,6 +61,9 @@ public:
 
 	UPROPERTY(BlueprintReadWrite, EditAnywhere, Category = Settings, meta = (EditCondition = "!bLinear && bApplyCustomTangents"))
 	FName LeaveTangentAttribute;
+
+	UPROPERTY(BlueprintReadWrite, Category = Settings, meta = (PCG_Overridable))
+	TSoftObjectPtr<AActor> TargetActor;
 };
 
 class FPCGCreateSplineElement : public FSimplePCGElement
diff --git a/Engine/Plugins/Experimental/PCG/Source/PCG/Public/Elements/PCGCreateTargetActor.h b/Engine/Plugins/Experimental/PCG/Source/PCG/Public/Elements/PCGCreateTargetActor.h
new file mode 100644
--- /dev/null
+++ b/Engine/Plugins/Experimental/PCG/Source/PCG/Public/Elements/PCGCreateTargetActor.h
@@ -0,0 +1,72 @@
+// Copyright Epic Games, Inc. All Rights Reserved.
+
+#pragma once
+
+#include "PCGPin.h"
+#include "PCGSettings.h"
+#include "Metadata/PCGAttributePropertySelector.h"
+
+#include "PCGCreateTargetActor.generated.h"
+
+class AActor;
+
+UCLASS(BlueprintType, ClassGroup = (Procedural))
+class PCG_API UPCGCreateTargetActor : public UPCGSettings
+{
+	GENERATED_BODY()
+
+public:
+	//~Begin UObject interface
+	virtual void PostLoad() override;
+	virtual void BeginDestroy() override;
+#if WITH_EDITOR
+	virtual void PreEditChange(FProperty* PropertyAboutToChange) override;
+	virtual void PostEditChangeProperty(FPropertyChangedEvent& PropertyChangedEvent) override;
+	virtual void PreEditUndo() override;
+	virtual void PostEditUndo() override;
+#endif
+	//~End UObject interface
+
+	//~Begin UPCGSettings interface
+#if WITH_EDITOR
+	virtual FName GetDefaultNodeName() const override { return FName(TEXT("CreateTargetActor")); }
+	virtual FText GetDefaultNodeTitle() const override;
+	virtual EPCGSettingsType GetType() const override { return EPCGSettingsType::Spawner; }
+#endif
+
+protected:
+	virtual TArray<FPCGPinProperties> InputPinProperties() const override;
+	virtual TArray<FPCGPinProperties> OutputPinProperties() const override;
+
+	virtual FPCGElementPtr CreateElement() const override;
+	//~End UPCGSettings interface
+
+#if WITH_EDITOR
+	void OnBlueprintChanged(UBlueprint* Blueprint);
+	void SetupBlueprintEvent();
+	void TeardownBlueprintEvent();
+#endif
+
+private:
+	void RefreshTemplateActor();
+
+public:
+	UPROPERTY(BlueprintReadOnly, EditAnywhere, Category = Settings, meta = (OnlyPlaceable, DisallowCreateNew))
+	TSubclassOf<AActor> TemplateActorClass = nullptr;
+
+	UPROPERTY(EditAnywhere, BlueprintReadWrite, Instanced, Category = Settings, meta = (ShowInnerProperties))
+	TObjectPtr<AActor> TemplateActor;
+
+	UPROPERTY(BlueprintReadWrite, Category = Settings, meta = (PCG_Overridable))
+	TSoftObjectPtr<AActor> RootActor;
+};
+
+class FPCGCreateTargetActorElement : public FSimplePCGElement
+{
+protected:
+	// Since this element creates an actor, it needs to run on the game thread and cannot be cached
+	virtual bool CanExecuteOnlyOnMainThread(FPCGContext* Context) const override { return true; }
+	virtual bool IsCacheable(const UPCGSettings* Settings) const override { return false; }
+	
+	virtual bool ExecuteInternal(FPCGContext* Context) const override;
+};
\ No newline at end of file
diff --git a/Engine/Plugins/Experimental/PCG/Source/PCG/Public/Elements/PCGDebugElement.h b/Engine/Plugins/Experimental/PCG/Source/PCG/Public/Elements/PCGDebugElement.h
--- a/Engine/Plugins/Experimental/PCG/Source/PCG/Public/Elements/PCGDebugElement.h
+++ b/Engine/Plugins/Experimental/PCG/Source/PCG/Public/Elements/PCGDebugElement.h
@@ -4,11 +4,13 @@
 
 #include "PCGSettings.h"
 
+#include "GameFramework/Actor.h"
+
 #include "PCGDebugElement.generated.h"
 
 namespace PCGDebugElement
 {
-	void ExecuteDebugDisplay(FPCGContext* Context);
+	void ExecuteDebugDisplay(FPCGContext* Context, AActor* TargetActor = nullptr);
 }
 
 UCLASS(BlueprintType, ClassGroup = (Procedural))
@@ -30,6 +32,10 @@ protected:
 
 	virtual FPCGElementPtr CreateElement() const override;
 	// ~End UPCGSettings interface
+
+public:
+	UPROPERTY(BlueprintReadWrite, Category = Settings, meta = (PCG_Overridable))
+	TSoftObjectPtr<AActor> TargetActor;
 };
 
 class FPCGDebugElement : public FSimplePCGElement
diff --git a/Engine/Plugins/Experimental/PCG/Source/PCG/Public/Elements/PCGSpawnActor.h b/Engine/Plugins/Experimental/PCG/Source/PCG/Public/Elements/PCGSpawnActor.h
--- a/Engine/Plugins/Experimental/PCG/Source/PCG/Public/Elements/PCGSpawnActor.h
+++ b/Engine/Plugins/Experimental/PCG/Source/PCG/Public/Elements/PCGSpawnActor.h
@@ -75,6 +75,9 @@ public:
 	UPROPERTY(EditAnywhere, BlueprintReadWrite, Category = Settings, meta = (EditCondition = "Option != EPCGSpawnActorOption::CollapseActors", EditConditionHides))
 	TArray<FPCGActorPropertyOverride> ActorOverrides;
 
+	UPROPERTY(BlueprintReadWrite, Category = Settings, meta = (PCG_Overridable))
+	TSoftObjectPtr<AActor> RootActor;
+
 protected:
 #if WITH_EDITORONLY_DATA
 	UPROPERTY()
diff --git a/Engine/Plugins/Experimental/PCG/Source/PCG/Public/Elements/PCGStaticMeshSpawner.h b/Engine/Plugins/Experimental/PCG/Source/PCG/Public/Elements/PCGStaticMeshSpawner.h
--- a/Engine/Plugins/Experimental/PCG/Source/PCG/Public/Elements/PCGStaticMeshSpawner.h
+++ b/Engine/Plugins/Experimental/PCG/Source/PCG/Public/Elements/PCGStaticMeshSpawner.h
@@ -95,6 +95,9 @@ public:
 	UPROPERTY(BlueprintReadWrite, EditAnywhere, Category = Settings)
 	bool bApplyMeshBoundsToPoints = true;
 
+	UPROPERTY(BlueprintReadWrite, Category = Settings, meta = (PCG_Overridable))
+	TSoftObjectPtr<AActor> TargetActor;
+
 	UPROPERTY(meta=(DeprecatedProperty, DeprecationMessage="Use MeshSelectorType and MeshSelectorParameters instead."))
 	TArray<FPCGStaticMeshSpawnerEntry> Meshes_DEPRECATED;
 
diff --git a/Engine/Plugins/Experimental/PCG/Source/PCG/Public/Helpers/PCGActorHelpers.h b/Engine/Plugins/Experimental/PCG/Source/PCG/Public/Helpers/PCGActorHelpers.h
--- a/Engine/Plugins/Experimental/PCG/Source/PCG/Public/Helpers/PCGActorHelpers.h
+++ b/Engine/Plugins/Experimental/PCG/Source/PCG/Public/Helpers/PCGActorHelpers.h
@@ -3,6 +3,7 @@
 #pragma once
 
 #include "Engine/EngineTypes.h"
+#include "Engine/World.h"
 #include "ISMPartition/ISMComponentDescriptor.h"
 #include "Kismet/BlueprintFunctionLibrary.h"
 #include "Templates/SubclassOf.h"
@@ -93,6 +94,16 @@ public:
 	*/
 	static AActor* SpawnDefaultActor(UWorld* World, TSubclassOf<AActor> ActorClass, FName BaseName, const FTransform& Transform, AActor* Parent = nullptr);
 
+	/**
+	* Spawn a new actor and attach it to the parent (if not null)
+	* @param World The world
+	* @param ActorClass Class of the actor to spawn
+	* @param Transform The transform for the new actor
+	* @param SpawnParams The spawn parameters
+	* @param Parent Optional parent to attach to.
+	*/
+	static AActor* SpawnDefaultActor(UWorld* World, TSubclassOf<AActor> ActorClass, const FTransform& Transform, const FActorSpawnParameters& SpawnParams, AActor* Parent = nullptr);
+
 	/**
 	 * Return the grid cell coordinates on the PCG partition grid given a position and the grid size.
 	 */
diff --git a/Engine/Plugins/Experimental/PCG/Source/PCG/Public/PCGManagedResource.h b/Engine/Plugins/Experimental/PCG/Source/PCG/Public/PCGManagedResource.h
--- a/Engine/Plugins/Experimental/PCG/Source/PCG/Public/PCGManagedResource.h
+++ b/Engine/Plugins/Experimental/PCG/Source/PCG/Public/PCGManagedResource.h
@@ -38,6 +38,7 @@ public:
 
 	/** Move the given resource to a new actor. Return true if it has succeeded */
 	virtual bool MoveResourceToNewActor(AActor* NewActor) { return false; };
+	virtual bool MoveResourceToNewActor(AActor* NewActor, const AActor* ExpectedPreviousOwner) { return MoveResourceToNewActor(NewActor); }
 
 	static bool DebugForcePurgeAllResourcesOnGenerate();
 
@@ -88,7 +89,7 @@ public:
 	//~Begin UPCGManagedResource interface
 	virtual bool Release(bool bHardRelease, TSet<TSoftObjectPtr<AActor>>& OutActorsToDelete) override;
 	virtual bool ReleaseIfUnused(TSet<TSoftObjectPtr<AActor>>& OutActorsToDelete) override;
-	virtual bool MoveResourceToNewActor(AActor* NewActor) override;
+	virtual bool MoveResourceToNewActor(AActor* NewActor, const AActor* ExpectedPreviousOwner) override;
 	//~End UPCGManagedResource interface
 
 	virtual void ResetComponent() { check(0); }
diff --git a/Engine/Plugins/Experimental/PCGGeometryScriptInterop/Source/PCGGeometryScriptInterop/Private/Elements/PCGMeshSampler.cpp b/Engine/Plugins/Experimental/PCGGeometryScriptInterop/Source/PCGGeometryScriptInterop/Private/Elements/PCGMeshSampler.cpp
--- a/Engine/Plugins/Experimental/PCGGeometryScriptInterop/Source/PCGGeometryScriptInterop/Private/Elements/PCGMeshSampler.cpp
+++ b/Engine/Plugins/Experimental/PCGGeometryScriptInterop/Source/PCGGeometryScriptInterop/Private/Elements/PCGMeshSampler.cpp
@@ -175,7 +175,6 @@ bool FPCGMeshSamplerElement::PrepareDataInternal(FPCGContext* InContext) const
 
 	TArray<FPCGTaggedData>& Outputs = InContext->OutputData.TaggedData;
 	Context->OutPointData = NewObject<UPCGPointData>();
-	Context->OutPointData->TargetActor = Context->SourceComponent->GetOwner();
 	Outputs.Emplace_GetRef().Data = Context->OutPointData;
 
 	Context->bDataPrepared = true;
