Added substring/matches (incl. wildcards) search in the PointFilter node

diff --git a/Engine/Plugins/Experimental/PCG/Source/PCG/Private/Elements/PCGPointFilter.cpp b/Engine/Plugins/Experimental/PCG/Source/PCG/Private/Elements/PCGPointFilter.cpp
--- a/Engine/Plugins/Experimental/PCG/Source/PCG/Private/Elements/PCGPointFilter.cpp
+++ b/Engine/Plugins/Experimental/PCG/Source/PCG/Private/Elements/PCGPointFilter.cpp
@@ -60,13 +60,24 @@ namespace PCGPointFilterHelpers
 			case EPCGPointFilterOperator::LesserOrEqual:
 				return PCG::Private::MetadataTraits<T>::LessOrEqual(Input1, Input2);
 			default:
-				return false;
+				break;
 			}
 		}
-		else
+
+		if constexpr (PCG::Private::MetadataTraits<T>::CanSearchString)
 		{
-			return false;
+			switch (Operation)
+			{
+			case EPCGPointFilterOperator::Substring:
+				return PCG::Private::MetadataTraits<T>::Substring(Input1, Input2);
+			case EPCGPointFilterOperator::Matches:
+				return PCG::Private::MetadataTraits<T>::Matches(Input1, Input2);
+			default:
+				break;
+			}
 		}
+
+		return false;
 	}
 
 	template <typename T>
@@ -142,7 +153,7 @@ namespace PCGPointFilterHelpers
 		return true;
 	}
 
-	bool PrepareThresholdInfoFromInput(FPCGContext* InContext, const UPCGPointData* InputData, const FPCGPointFilterThresholdSettings& ThresholdSettings, ThresholdInfo& InOutThresholdInfo, int16 TargetType, bool bCheckCompare, const ThresholdInfo* OtherInfo = nullptr)
+	bool PrepareThresholdInfoFromInput(FPCGContext* InContext, const UPCGPointData* InputData, const FPCGPointFilterThresholdSettings& ThresholdSettings, ThresholdInfo& InOutThresholdInfo, int16 TargetType, bool bCheckCompare, bool bCheckStringSearch, const ThresholdInfo* OtherInfo = nullptr)
 	{
 		check(InContext && InputData);
 
@@ -209,6 +220,20 @@ namespace PCGPointFilterHelpers
 			}
 		}
 
+		if (bCheckStringSearch)
+		{
+			bool bCanSearchString = PCGMetadataAttribute::CallbackWithRightType(InOutThresholdInfo.ThresholdAccessor->GetUnderlyingType(), [](auto Dummy) -> bool
+			{
+				return PCG::Private::MetadataTraits<decltype(Dummy)>::CanSearchString;
+			});
+
+			if (!bCanSearchString)
+			{
+				PCGE_LOG_C(Warning, GraphAndLog, InContext, LOCTEXT("TypeStringSearchFailed", "Cannot perform string operations on target type"));
+				return false;
+			}
+		}
+
 		// Check that if we have points as threshold, that the point data has the same number of point that the input data
 		if (InOutThresholdInfo.ThresholdSpatialData != nullptr && !InOutThresholdInfo.bUseSpatialQuery)
 		{
@@ -470,15 +495,16 @@ bool FPCGPointFilterElementBase::DoFiltering(FPCGContext* Context, EPCGPointFilt
 		}
 
 		const int16 TargetType = TargetAccessor->GetUnderlyingType();
-		const bool bCheckCompare = (Operator != EPCGPointFilterOperator::Equal) && (Operator != EPCGPointFilterOperator::NotEqual);
+		const bool bCheckStringSearch = (Operator == EPCGPointFilterOperator::Substring || Operator == EPCGPointFilterOperator::Matches);
+		const bool bCheckCompare = (Operator != EPCGPointFilterOperator::Equal) && (Operator != EPCGPointFilterOperator::NotEqual) && !bCheckStringSearch;
 
-		if (!PCGPointFilterHelpers::PrepareThresholdInfoFromInput(Context, OriginalData, FirstThreshold, FirstThresholdInfo, TargetType, bCheckCompare))
+		if (!PCGPointFilterHelpers::PrepareThresholdInfoFromInput(Context, OriginalData, FirstThreshold, FirstThresholdInfo, TargetType, bCheckCompare, bCheckStringSearch))
 		{
 			ForwardInputToInFilterPin();
 			continue;
 		}
 
-		if (SecondThreshold && !PCGPointFilterHelpers::PrepareThresholdInfoFromInput(Context, OriginalData, *SecondThreshold, SecondThresholdInfo, TargetType, bCheckCompare, &FirstThresholdInfo))
+		if (SecondThreshold && !PCGPointFilterHelpers::PrepareThresholdInfoFromInput(Context, OriginalData, *SecondThreshold, SecondThresholdInfo, TargetType, bCheckCompare, bCheckStringSearch, &FirstThresholdInfo))
 		{
 			ForwardInputToInFilterPin();
 			continue;
diff --git a/Engine/Plugins/Experimental/PCG/Source/PCG/Public/Elements/PCGPointFilter.h b/Engine/Plugins/Experimental/PCG/Source/PCG/Public/Elements/PCGPointFilter.h
--- a/Engine/Plugins/Experimental/PCG/Source/PCG/Public/Elements/PCGPointFilter.h
+++ b/Engine/Plugins/Experimental/PCG/Source/PCG/Public/Elements/PCGPointFilter.h
@@ -44,7 +44,9 @@ enum class EPCGPointFilterOperator : uint8
 	LesserOrEqual UMETA(DisplayName="<="),
 	Equal UMETA(DisplayName="="),
 	NotEqual UMETA(DisplayName="!="),
-	InRange UMETA(Hidden)
+	InRange UMETA(Hidden),
+	Substring,
+	Matches
 };
 
 USTRUCT(BlueprintType)
diff --git a/Engine/Plugins/Experimental/PCG/Source/PCG/Public/Metadata/PCGMetadataAttributeTraits.h b/Engine/Plugins/Experimental/PCG/Source/PCG/Public/Metadata/PCGMetadataAttributeTraits.h
--- a/Engine/Plugins/Experimental/PCG/Source/PCG/Public/Metadata/PCGMetadataAttributeTraits.h
+++ b/Engine/Plugins/Experimental/PCG/Source/PCG/Public/Metadata/PCGMetadataAttributeTraits.h
@@ -357,12 +357,14 @@ namespace PCG
 		struct MetadataTraits : DefaultOperationTraits<T>, DefaultWeightedSumTraits<T>, DefaultMinMaxTraits<T>, DefaultCompareTraits<T>
 		{
 			enum { CompressData = false };
+			enum { CanSearchString = false };
 		};
 
 		template<>
 		struct MetadataTraits<float> : DefaultOperationTraits<float>, DefaultWeightedSumTraits<float>, DefaultMinMaxTraits<float>, DefaultCompareTraits<float>
 		{
 			enum { CompressData = false };
+			enum { CanSearchString = false };
 
 			static bool Equal(const float& A, const float& B)
 			{
@@ -374,6 +376,7 @@ namespace PCG
 		struct MetadataTraits<double> : DefaultOperationTraits<double>, DefaultWeightedSumTraits<double>, DefaultMinMaxTraits<double>, DefaultCompareTraits<double>
 		{
 			enum { CompressData = false };
+			enum { CanSearchString = false };
 
 			static bool Equal(const double& A, const double& B)
 			{
@@ -390,6 +393,7 @@ namespace PCG
 			enum { CanMulDiv = false };
 			enum { CanInterpolate = false };
 			enum { CanCompare = true };
+			enum { CanSearchString = false };
 
 			static bool ZeroValue()
 			{
@@ -450,6 +454,7 @@ namespace PCG
 			enum { CanMinMax = true };
 			enum { CanInterpolate = true };
 			enum { CanCompare = true };
+			enum { CanSearchString = false };
 
 			static T ZeroValue()
 			{
@@ -584,6 +589,7 @@ namespace PCG
 			enum { CanMulDiv = true };
 			enum { CanInterpolate = true };
 			enum { CanCompare = false };
+			enum { CanSearchString = false };
 
 			static bool Equal(const FQuat& A, const FQuat& B)
 			{
@@ -637,6 +643,7 @@ namespace PCG
 			enum { CanMulDiv = true };
 			enum { CanInterpolate = true };
 			enum { CanCompare = false };
+			enum { CanSearchString = false };
 
 			static bool Equal(const FRotator& A, const FRotator& B)
 			{
@@ -685,6 +692,7 @@ namespace PCG
 			enum { CanMulDiv = true };
 			enum { CanInterpolate = true };
 			enum { CanCompare = false };
+			enum { CanSearchString = false };
 
 			static bool Equal(const FTransform& A, const FTransform& B)
 			{
@@ -735,6 +743,7 @@ namespace PCG
 			enum { CanSubAdd = false };
 			enum { CanMulDiv = false };
 			enum { CanInterpolate = false };
+			enum { CanSearchString = true };
 
 			static bool Equal(const FString& A, const FString& B)
 			{
@@ -745,6 +754,16 @@ namespace PCG
 			{
 				return FString();
 			}
+
+			static bool Substring(const FString& A, const FString& B)
+			{
+				return A.Contains(B);
+			}
+
+			static bool Matches(const FString& A, const FString& B)
+			{
+				return A.MatchesWildcard(B);
+			}
 		};
 
 		template<>
@@ -756,6 +775,7 @@ namespace PCG
 			enum { CanMulDiv = false };
 			enum { CanInterpolate = false };
 			enum { CanCompare = true };
+			enum { CanSearchString = true };
 
 			static bool Equal(const FName& A, const FName& B)
 			{
@@ -786,6 +806,16 @@ namespace PCG
 			{
 				return NAME_None;
 			}
+
+			static bool Substring(const FName& A, const FName& B)
+			{
+				return A.ToString().Contains(B.ToString());
+			}
+
+			static bool Matches(const FName& A, const FName& B)
+			{
+				return A.ToString().MatchesWildcard(B.ToString());
+			}
 		};
 	}
 }
