Fixed filter by type not working on attribute set

#rnx

diff --git a/Engine/Plugins/Experimental/PCG/Source/PCG/Private/Elements/PCGFilterByTag.cpp b/Engine/Plugins/Experimental/PCG/Source/PCG/Private/Elements/PCGFilterByTag.cpp
--- a/Engine/Plugins/Experimental/PCG/Source/PCG/Private/Elements/PCGFilterByTag.cpp
+++ b/Engine/Plugins/Experimental/PCG/Source/PCG/Private/Elements/PCGFilterByTag.cpp
@@ -93,7 +93,7 @@ bool FPCGFilterByTagElement::ExecuteInternal(FPCGContext* Context) const
 	const bool bKeepIfTag = (Settings->Operation == EPCGFilterByTagOperation::KeepTagged);
 	const TArray<FString> Tags = PCGHelpers::GetStringArrayFromCommaSeparatedString(Settings->SelectedTags);
 
-	TArray<FPCGTaggedData> Inputs = Context->InputData.GetInputs();
+	TArray<FPCGTaggedData> Inputs = Context->InputData.TaggedData;
 	TArray<FPCGTaggedData>& Outputs = Context->OutputData.TaggedData;
 
 	for (const FPCGTaggedData& Input : Inputs)
diff --git a/Engine/Plugins/Experimental/PCG/Source/PCG/Private/Elements/PCGFilterByType.cpp b/Engine/Plugins/Experimental/PCG/Source/PCG/Private/Elements/PCGFilterByType.cpp
--- a/Engine/Plugins/Experimental/PCG/Source/PCG/Private/Elements/PCGFilterByType.cpp
+++ b/Engine/Plugins/Experimental/PCG/Source/PCG/Private/Elements/PCGFilterByType.cpp
@@ -46,7 +46,7 @@ bool FPCGFilterByTypeElement::ExecuteInternal(FPCGContext* InContext) const
 	const UPCGFilterByTypeSettings* Settings = InContext->GetInputSettings<UPCGFilterByTypeSettings>();
 	check(Settings);
 
-	TArray<FPCGTaggedData> Inputs = InContext->InputData.GetInputs();
+	TArray<FPCGTaggedData> Inputs = InContext->InputData.TaggedData;
 	TArray<FPCGTaggedData>& Outputs = InContext->OutputData.TaggedData;
 
 	for (const FPCGTaggedData& Input : Inputs)
