This test ensures the PCG Component is available in the Unreal Engine's Component Type Registry, thus can be added to Actor Blueprints

1) Access the Component Type Registry of the Unreal Engine.
2) Obtain the Component List from the registry.
3) Search for the PCG Component (UPCGComponent) in the list.
4) Assert that the PCG Component is found in the list.

Expected Result: The PCG Component should be found in the Component List, confirming its availability for adding to Actor Blueprints.

diff --git a/Engine/Plugins/Experimental/PCG/Source/PCGEditor/Private/Tests/Display/PCGComponentFoundInRegistry.cpp b/Engine/Plugins/Experimental/PCG/Source/PCGEditor/Private/Tests/Display/PCGComponentFoundInRegistry.cpp
new file mode 100644
--- /dev/null
+++ b/Engine/Plugins/Experimental/PCG/Source/PCGEditor/Private/Tests/Display/PCGComponentFoundInRegistry.cpp
@@ -0,0 +1,34 @@
+// Copyright Epic Games, Inc. All Rights Reserved.
+
+#include "ComponentTypeRegistry.h"
+#include "SComponentClassCombo.h"
+#include "Tests/PCGTestsCommon.h"
+#include "PCGComponent.h"
+
+#if WITH_AUTOMATION_TESTS
+
+IMPLEMENT_CUSTOM_SIMPLE_AUTOMATION_TEST(PCGComponentFoundInRegistry, FPCGTestBaseClass, "Editor.Plugins.Tools.PCG.PCGComponentFoundInRegistry", PCGTestsCommon::TestFlags)
+
+bool PCGComponentFoundInRegistry::RunTest(const FString& Parameters)
+{
+	//Add component button gets list from ComponentList
+	FComponentTypeRegistry& Registry = FComponentTypeRegistry::Get();
+
+	TArray<FComponentClassComboEntryPtr>* ComponentClassList = nullptr;
+
+	Registry.SubscribeToComponentList(ComponentClassList);
+
+	if (ComponentClassList != nullptr) {
+
+		const bool bFoundPCGComponent = ComponentClassList->ContainsByPredicate([](const FComponentClassComboEntryPtr& Entry) -> bool {
+			
+			return Entry->GetComponentClass() == UPCGComponent::StaticClass();
+			
+			});
+
+		TestTrue(TEXT("UPCGComponent not found in component list"), bFoundPCGComponent);
+	}
+
+	return true;
+}
+#endif
