[PCG] CreateAttribute should have Param output if nothing is plugged

#rnx

diff --git a/Engine/Plugins/Experimental/PCG/Source/PCG/Private/Elements/PCGCreateAttribute.cpp b/Engine/Plugins/Experimental/PCG/Source/PCG/Private/Elements/PCGCreateAttribute.cpp
--- a/Engine/Plugins/Experimental/PCG/Source/PCG/Private/Elements/PCGCreateAttribute.cpp
+++ b/Engine/Plugins/Experimental/PCG/Source/PCG/Private/Elements/PCGCreateAttribute.cpp
@@ -60,7 +60,7 @@ EPCGDataType UPCGCreateAttributeSettings::GetCurrentPinTypes(const UPCGPin* InPi
 
 	// Output pin narrows to union of inputs on first pin
 	const EPCGDataType PrimaryInputType = GetTypeUnionOfIncidentEdges(PCGPinConstants::DefaultInputLabel);
-	return (PrimaryInputType != EPCGDataType::None) ? PrimaryInputType : EPCGDataType::Any;
+	return (PrimaryInputType != EPCGDataType::None) ? PrimaryInputType : EPCGDataType::Param; // No input (None) means param.
 }
 
 FName UPCGCreateAttributeSettings::AdditionalTaskName() const
