Fixed seed issue related to override duplication of settings

#rnx

diff --git a/Engine/Plugins/Experimental/PCG/Source/PCG/Private/PCGContext.cpp b/Engine/Plugins/Experimental/PCG/Source/PCG/Private/PCGContext.cpp
--- a/Engine/Plugins/Experimental/PCG/Source/PCG/Private/PCGContext.cpp
+++ b/Engine/Plugins/Experimental/PCG/Source/PCG/Private/PCGContext.cpp
@@ -132,6 +132,9 @@ void FPCGContext::InitializeSettings()
 			{
 				SettingsWithOverride = Cast<UPCGSettings>(StaticDuplicateObject(NodeSettings, GetTransientPackage(), NAME_None, RF_Transient));
 				SettingsWithOverride->AddToRoot();
+
+				// Force seed copy to prevent issue due to delta serialization vs. Seed being initialized in the constructor only for new nodes
+				SettingsWithOverride->Seed = NodeSettings->Seed;
 			}
 		}
 	}
