[PCG] Add [REDACTED_HANDLE] and [REDACTED_HANDLE] in place of None for Attribute Selector

* Allows for an explicit behavior, and allows "None" to be a valid attribute name.
* Introduce 3 new selectors, to explicitly declare it as input or output
* Default input value is [REDACTED_HANDLE], default for output can be [REDACTED_HANDLE] or None
* Make the selector for fool proof by hiding memebrs and allowing only setter to modify its content
* LOTS of files modified because of all that

diff --git a/Engine/Plugins/Experimental/PCG/Source/PCG/Private/Data/PCGSpatialData.cpp b/Engine/Plugins/Experimental/PCG/Source/PCG/Private/Data/PCGSpatialData.cpp
--- a/Engine/Plugins/Experimental/PCG/Source/PCG/Private/Data/PCGSpatialData.cpp
+++ b/Engine/Plugins/Experimental/PCG/Source/PCG/Private/Data/PCGSpatialData.cpp
@@ -240,5 +240,49 @@ UPCGSpatialData* UPCGSpatialData::DuplicateData(const bool bInitializeMetadata)
 		NewSpatialData->InitializeFromData(this);
 	}
 
+	if (bHasCachedLastSelector)
+	{
+		NewSpatialData->SetLastSelector(CachedLastSelector);
+	}
+
 	return NewSpatialData;
 }
+
+bool UPCGSpatialData::HasCachedLastSelector() const
+{
+	return bHasCachedLastSelector || (Metadata && Metadata->GetAttributeCount() > 0);
+}
+
+FPCGAttributePropertyInputSelector UPCGSpatialData::GetCachedLastSelector() const
+{
+	if (bHasCachedLastSelector)
+	{
+		return CachedLastSelector;
+	}
+
+	FPCGAttributePropertyInputSelector TempSelector{};
+
+	// If we have attribute and no last selector, create a cached last selector on the latest attribute, to catch "CreateAttribute" calls that didn't use accessors.
+	if (Metadata && Metadata->GetAttributeCount() > 0)
+	{
+		TempSelector.SetAttributeName(Metadata->GetLatestAttributeNameOrNone());
+	}
+
+	return TempSelector;
+}
+
+void UPCGSpatialData::SetLastSelector(const FPCGAttributePropertySelector& InSelector)
+{
+	// Check that it is not a Last or Source selector
+	if (InSelector.GetSelection() == EPCGAttributePropertySelection::Attribute &&
+		(InSelector.GetAttributeName() == PCGMetadataAttributeConstants::LastAttributeName
+			|| InSelector.GetAttributeName() == PCGMetadataAttributeConstants::LastCreatedAttributeName
+			|| InSelector.GetAttributeName() == PCGMetadataAttributeConstants::SourceAttributeName
+			|| InSelector.GetAttributeName() == PCGMetadataAttributeConstants::SourceNameAttributeName))
+	{
+		return;
+	}
+
+	bHasCachedLastSelector = true;
+	CachedLastSelector.ImportFromOtherSelector(InSelector);
+}
diff --git a/Engine/Plugins/Experimental/PCG/Source/PCG/Private/Elements/IO/PCGDataTableElement.cpp b/Engine/Plugins/Experimental/PCG/Source/PCG/Private/Elements/IO/PCGDataTableElement.cpp
--- a/Engine/Plugins/Experimental/PCG/Source/PCG/Private/Elements/IO/PCGDataTableElement.cpp
+++ b/Engine/Plugins/Experimental/PCG/Source/PCG/Private/Elements/IO/PCGDataTableElement.cpp
@@ -83,11 +83,11 @@ bool FPCGLoadDataTableElement::PrepareLoad(FPCGExternalDataContext* Context) con
 		}
 
 		// Try to match to a property (if the name maps to a property) or create attribute as needed
-		FPCGAttributePropertySelector PointPropertySelector;
+		FPCGAttributePropertyOutputSelector PointPropertySelector;
 
 		if (const FPCGAttributePropertySelector* MappedField = Settings->AttributeMapping.Find(FieldName))
 		{
-			PointPropertySelector = *MappedField;
+			PointPropertySelector.ImportFromOtherSelector(*MappedField);
 			FieldName = PointPropertySelector.GetDisplayText().ToString();
 		}
 		else
@@ -95,7 +95,7 @@ bool FPCGLoadDataTableElement::PrepareLoad(FPCGExternalDataContext* Context) con
 			PointPropertySelector.Update(FieldName);
 		}
 
-		if (PointPropertySelector.Selection == EPCGAttributePropertySelection::Attribute)
+		if (PointPropertySelector.GetSelection() == EPCGAttributePropertySelection::Attribute)
 		{
 			// Create attribute with the appropriate type
 			PointMetadata->CreateAttributeFromProperty(FName(FieldName), RowStruct, *FieldIt);
diff --git a/Engine/Plugins/Experimental/PCG/Source/PCG/Private/Elements/Metadata/PCGMetadataBitwiseOpElement.cpp b/Engine/Plugins/Experimental/PCG/Source/PCG/Private/Elements/Metadata/PCGMetadataBitwiseOpElement.cpp
--- a/Engine/Plugins/Experimental/PCG/Source/PCG/Private/Elements/Metadata/PCGMetadataBitwiseOpElement.cpp
+++ b/Engine/Plugins/Experimental/PCG/Source/PCG/Private/Elements/Metadata/PCGMetadataBitwiseOpElement.cpp
@@ -73,7 +73,7 @@ bool UPCGMetadataBitwiseSettings::IsSupportedInputType(uint16 TypeId, uint32 Inp
 	return PCG::Private::IsOfTypes<int32, int64>(TypeId);
 }
 
-FPCGAttributePropertySelector UPCGMetadataBitwiseSettings::GetInputSource(uint32 Index) const
+FPCGAttributePropertyInputSelector UPCGMetadataBitwiseSettings::GetInputSource(uint32 Index) const
 {
 	switch (Index)
 	{
@@ -82,7 +82,7 @@ FPCGAttributePropertySelector UPCGMetadataBitwiseSettings::GetInputSource(uint32
 	case 1:
 		return InputSource2;
 	default:
-		return FPCGAttributePropertySelector();
+		return FPCGAttributePropertyInputSelector();
 	}
 }
 
diff --git a/Engine/Plugins/Experimental/PCG/Source/PCG/Private/Elements/Metadata/PCGMetadataBooleanOpElement.cpp b/Engine/Plugins/Experimental/PCG/Source/PCG/Private/Elements/Metadata/PCGMetadataBooleanOpElement.cpp
--- a/Engine/Plugins/Experimental/PCG/Source/PCG/Private/Elements/Metadata/PCGMetadataBooleanOpElement.cpp
+++ b/Engine/Plugins/Experimental/PCG/Source/PCG/Private/Elements/Metadata/PCGMetadataBooleanOpElement.cpp
@@ -73,7 +73,7 @@ bool UPCGMetadataBooleanSettings::IsSupportedInputType(uint16 TypeId, uint32 Inp
 	return PCG::Private::IsOfTypes<bool>(TypeId);
 }
 
-FPCGAttributePropertySelector UPCGMetadataBooleanSettings::GetInputSource(uint32 Index) const
+FPCGAttributePropertyInputSelector UPCGMetadataBooleanSettings::GetInputSource(uint32 Index) const
 {
 	switch (Index)
 	{
@@ -82,7 +82,7 @@ FPCGAttributePropertySelector UPCGMetadataBooleanSettings::GetInputSource(uint32
 	case 1:
 		return InputSource2;
 	default:
-		return FPCGAttributePropertySelector();
+		return FPCGAttributePropertyInputSelector();
 	}
 }
 
diff --git a/Engine/Plugins/Experimental/PCG/Source/PCG/Private/Elements/Metadata/PCGMetadataBreakTransform.cpp b/Engine/Plugins/Experimental/PCG/Source/PCG/Private/Elements/Metadata/PCGMetadataBreakTransform.cpp
--- a/Engine/Plugins/Experimental/PCG/Source/PCG/Private/Elements/Metadata/PCGMetadataBreakTransform.cpp
+++ b/Engine/Plugins/Experimental/PCG/Source/PCG/Private/Elements/Metadata/PCGMetadataBreakTransform.cpp
@@ -37,7 +37,7 @@ void UPCGMetadataBreakTransformSettings::PostLoad()
 #endif // WITH_EDITOR
 }
 
-FPCGAttributePropertySelector UPCGMetadataBreakTransformSettings::GetInputSource(uint32 Index) const
+FPCGAttributePropertyInputSelector UPCGMetadataBreakTransformSettings::GetInputSource(uint32 Index) const
 {
 	return InputSource;
 }
diff --git a/Engine/Plugins/Experimental/PCG/Source/PCG/Private/Elements/Metadata/PCGMetadataBreakVector.cpp b/Engine/Plugins/Experimental/PCG/Source/PCG/Private/Elements/Metadata/PCGMetadataBreakVector.cpp
--- a/Engine/Plugins/Experimental/PCG/Source/PCG/Private/Elements/Metadata/PCGMetadataBreakVector.cpp
+++ b/Engine/Plugins/Experimental/PCG/Source/PCG/Private/Elements/Metadata/PCGMetadataBreakVector.cpp
@@ -87,7 +87,7 @@ void UPCGMetadataBreakVectorSettings::PostLoad()
 #endif // WITH_EDITOR
 }
 
-FPCGAttributePropertySelector UPCGMetadataBreakVectorSettings::GetInputSource(uint32 Index) const
+FPCGAttributePropertyInputSelector UPCGMetadataBreakVectorSettings::GetInputSource(uint32 Index) const
 {
 	return InputSource;
 }
diff --git a/Engine/Plugins/Experimental/PCG/Source/PCG/Private/Elements/Metadata/PCGMetadataCompareOpElement.cpp b/Engine/Plugins/Experimental/PCG/Source/PCG/Private/Elements/Metadata/PCGMetadataCompareOpElement.cpp
--- a/Engine/Plugins/Experimental/PCG/Source/PCG/Private/Elements/Metadata/PCGMetadataCompareOpElement.cpp
+++ b/Engine/Plugins/Experimental/PCG/Source/PCG/Private/Elements/Metadata/PCGMetadataCompareOpElement.cpp
@@ -99,7 +99,7 @@ bool UPCGMetadataCompareSettings::IsSupportedInputType(uint16 TypeId, uint32 Inp
 	}
 }
 
-FPCGAttributePropertySelector UPCGMetadataCompareSettings::GetInputSource(uint32 Index) const
+FPCGAttributePropertyInputSelector UPCGMetadataCompareSettings::GetInputSource(uint32 Index) const
 {
 	switch (Index)
 	{
@@ -108,7 +108,7 @@ FPCGAttributePropertySelector UPCGMetadataCompareSettings::GetInputSource(uint32
 	case 1:
 		return InputSource2;
 	default:
-		return FPCGAttributePropertySelector();
+		return FPCGAttributePropertyInputSelector();
 	}
 }
 
diff --git a/Engine/Plugins/Experimental/PCG/Source/PCG/Private/Elements/Metadata/PCGMetadataElement.cpp b/Engine/Plugins/Experimental/PCG/Source/PCG/Private/Elements/Metadata/PCGMetadataElement.cpp
--- a/Engine/Plugins/Experimental/PCG/Source/PCG/Private/Elements/Metadata/PCGMetadataElement.cpp
+++ b/Engine/Plugins/Experimental/PCG/Source/PCG/Private/Elements/Metadata/PCGMetadataElement.cpp
@@ -3,6 +3,7 @@
 #include "Elements/Metadata/PCGMetadataElement.h"
 
 #include "PCGContext.h"
+#include "PCGCustomVersion.h"
 #include "PCGParamData.h"
 #include "Data/PCGPointData.h"
 #include "Data/PCGSpatialData.h"
@@ -43,6 +44,20 @@ bool UPCGMetadataOperationSettings::IsPinUsedByNodeExecution(const UPCGPin* InPi
 {
 	return !InPin || (InPin->Properties.Label != PCGMetadataOperationSettings::AttributeLabel) || InPin->IsConnected();
 }
+
+void UPCGMetadataOperationSettings::ApplyDeprecation(UPCGNode* InOutNode)
+{
+	if (DataVersion < FPCGCustomVersion::UpdateAttributePropertyInputSelector
+		&& OutputTarget.GetSelection() == EPCGAttributePropertySelection::Attribute
+		&& OutputTarget.GetAttributeName() == PCGMetadataAttributeConstants::SourceAttributeName)
+	{
+		// Previous behavior of the output target for this node was:
+		// None => LastCreated
+		OutputTarget.SetAttributeName(PCGMetadataAttributeConstants::LastCreatedAttributeName);
+	}
+
+	Super::ApplyDeprecation(InOutNode);
+}
 #endif // WITH_EDITOR
 
 TArray<FPCGPinProperties> UPCGMetadataOperationSettings::InputPinProperties() const
@@ -104,7 +119,7 @@ bool FPCGMetadataOperationElement::ExecuteInternal(FPCGContext* Context) const
 
 	const UPCGParamData* SourceAttributeSet = (!SourceAttributeInputs.IsEmpty() ? Cast<UPCGParamData>(SourceAttributeInputs[0].Data) : nullptr);
 
-	if (SourceAttributeSet && Settings->InputSource.Selection != EPCGAttributePropertySelection::Attribute)
+	if (SourceAttributeSet && Settings->InputSource.GetSelection() != EPCGAttributePropertySelection::Attribute)
 	{
 		PCGE_LOG(Error, GraphAndLog, FText::Format(LOCTEXT("InvalidInputSource", "{0} pin is connected but Input Source is not set to get an attribute."), FText::FromName(PCGMetadataOperationSettings::AttributeLabel)));
 		return true;
@@ -132,6 +147,7 @@ bool FPCGMetadataOperationElement::ExecuteInternal(FPCGContext* Context) const
 			continue;
 		}
 
+		const UPCGData* SourceData = SourceAttributeSet ? static_cast<const UPCGData*>(SourceAttributeSet) : static_cast<const UPCGData*>(OriginalData);
 		const UPCGMetadata* SourceMetadata = SourceAttributeSet ? SourceAttributeSet->Metadata : OriginalData->Metadata;
 
 		if (!SourceMetadata)
@@ -140,19 +156,14 @@ bool FPCGMetadataOperationElement::ExecuteInternal(FPCGContext* Context) const
 			continue;
 		}
 
-		FPCGAttributePropertySelector InputSource = Settings->InputSource;
-		FPCGAttributePropertySelector OutputTarget = Settings->OutputTarget;
+		FPCGAttributePropertyInputSelector InputSource = Settings->InputSource.CopyAndFixLast(SourceData);
+		// Deprecation, old behavior was if the Target was None, we took LastCreated in the SourceData
+		FPCGAttributePropertyOutputSelector OutputTarget = Settings->OutputTarget.CopyAndFixSource(&InputSource, SourceData);
 
-		const FName SourceAttribute = InputSource.GetName();
-		const FName DestinationAttribute = OutputTarget.GetName();
-		const FName LocalSourceAttribute = ((SourceAttribute != NAME_None) ? SourceAttribute : SourceMetadata->GetLatestAttributeNameOrNone());
-		const FName LocalDestinationAttribute = ((DestinationAttribute != NAME_None) ? DestinationAttribute : SourceMetadata->GetLatestAttributeNameOrNone());
+		const FName LocalSourceAttribute = InputSource.GetName();
+		const FName LocalDestinationAttribute = OutputTarget.GetName();
 
-		// Make sure we use the right attribute name
-		InputSource.AttributeName = LocalSourceAttribute;
-		OutputTarget.AttributeName = LocalDestinationAttribute;
-
-		if (Settings->InputSource.Selection == EPCGAttributePropertySelection::Attribute && !SourceMetadata->HasAttribute(LocalSourceAttribute))
+		if (Settings->InputSource.GetSelection() == EPCGAttributePropertySelection::Attribute && !SourceMetadata->HasAttribute(LocalSourceAttribute))
 		{
 			PCGE_LOG(Warning, GraphAndLog, FText::Format(LOCTEXT("InputMissingAttribute", "Input does not have the '{0}' attribute"), FText::FromName(LocalSourceAttribute)));
 			continue;
@@ -171,9 +182,9 @@ bool FPCGMetadataOperationElement::ExecuteInternal(FPCGContext* Context) const
 		SampledPoints = Points;
 
 		// If it is attribute to attribute, just copy the attributes, if they exist and are valid
-		if (Settings->InputSource.Selection == EPCGAttributePropertySelection::Attribute && OutputTarget.Selection == EPCGAttributePropertySelection::Attribute)
+		if (Settings->InputSource.GetSelection() == EPCGAttributePropertySelection::Attribute && OutputTarget.GetSelection() == EPCGAttributePropertySelection::Attribute)
 		{
-			if (LocalSourceAttribute == DestinationAttribute)
+			if (LocalSourceAttribute == LocalDestinationAttribute)
 			{
 				// Nothing to do...
 				continue;
@@ -185,16 +196,14 @@ bool FPCGMetadataOperationElement::ExecuteInternal(FPCGContext* Context) const
 			// We only copy entries/values if we copy from the input spatial metadata
 			// If it is from the source attribute set, we don't copy (and all points will have the same default value, value from the attribute set)
 			const bool bCopyEntriesAndValues = (SourceAttributeSet == nullptr);
-			if (!SampledData->Metadata->CopyAttribute(OriginalAttribute, DestinationAttribute, /*bKeepParent=*/ false, /*bCopyEntries=*/ bCopyEntriesAndValues, /*bCopyValues=*/ bCopyEntriesAndValues))
+			if (!SampledData->Metadata->CopyAttribute(OriginalAttribute, LocalDestinationAttribute, /*bKeepParent=*/ false, /*bCopyEntries=*/ bCopyEntriesAndValues, /*bCopyValues=*/ bCopyEntriesAndValues))
 			{
-				PCGE_LOG(Warning, GraphAndLog, FText::Format(LOCTEXT("FailedCopyToNewAttribute", "Failed to copy to new attribute '{0}'"), FText::FromName(DestinationAttribute)));
+				PCGE_LOG(Warning, GraphAndLog, FText::Format(LOCTEXT("FailedCopyToNewAttribute", "Failed to copy to new attribute '{0}'"), FText::FromName(LocalDestinationAttribute)));
 			}
 
 			continue;
 		}
 
-		const UPCGData* SourceData = (SourceAttributeSet ? static_cast<const UPCGData*>(SourceAttributeSet) : static_cast<const UPCGData*>(OriginalData));
-
 		TUniquePtr<const IPCGAttributeAccessor> InputAccessor = PCGAttributeAccessorHelpers::CreateConstAccessor(SourceData, InputSource);
 		TUniquePtr<const IPCGAttributeAccessorKeys> InputKeys = PCGAttributeAccessorHelpers::CreateConstKeys(SourceData, InputSource);
 
@@ -206,7 +215,7 @@ bool FPCGMetadataOperationElement::ExecuteInternal(FPCGContext* Context) const
 
 		// If the target is an attribute, only create a new one if the attribute doesn't already exist.
 		// If it exist, it will try to write to it.
-		if (OutputTarget.Selection == EPCGAttributePropertySelection::Attribute && !SampledData->Metadata->HasAttribute(LocalDestinationAttribute))
+		if (OutputTarget.GetSelection() == EPCGAttributePropertySelection::Attribute && !SampledData->Metadata->HasAttribute(LocalDestinationAttribute))
 		{
 			auto CreateAttribute = [SampledData, LocalDestinationAttribute](auto Dummy)
 			{
diff --git a/Engine/Plugins/Experimental/PCG/Source/PCG/Private/Elements/Metadata/PCGMetadataMakeRotator.cpp b/Engine/Plugins/Experimental/PCG/Source/PCG/Private/Elements/Metadata/PCGMetadataMakeRotator.cpp
--- a/Engine/Plugins/Experimental/PCG/Source/PCG/Private/Elements/Metadata/PCGMetadataMakeRotator.cpp
+++ b/Engine/Plugins/Experimental/PCG/Source/PCG/Private/Elements/Metadata/PCGMetadataMakeRotator.cpp
@@ -74,7 +74,7 @@ bool UPCGMetadataMakeRotatorSettings::IsSupportedInputType(uint16 TypeId, uint32
 	}
 }
 
-FPCGAttributePropertySelector UPCGMetadataMakeRotatorSettings::GetInputSource(uint32 Index) const
+FPCGAttributePropertyInputSelector UPCGMetadataMakeRotatorSettings::GetInputSource(uint32 Index) const
 {
 	switch (Index)
 	{
@@ -85,7 +85,7 @@ FPCGAttributePropertySelector UPCGMetadataMakeRotatorSettings::GetInputSource(ui
 	case 2:
 		return InputSource3;
 	default:
-		return FPCGAttributePropertySelector();
+		return FPCGAttributePropertyInputSelector();
 	}
 }
 
diff --git a/Engine/Plugins/Experimental/PCG/Source/PCG/Private/Elements/Metadata/PCGMetadataMakeTransform.cpp b/Engine/Plugins/Experimental/PCG/Source/PCG/Private/Elements/Metadata/PCGMetadataMakeTransform.cpp
--- a/Engine/Plugins/Experimental/PCG/Source/PCG/Private/Elements/Metadata/PCGMetadataMakeTransform.cpp
+++ b/Engine/Plugins/Experimental/PCG/Source/PCG/Private/Elements/Metadata/PCGMetadataMakeTransform.cpp
@@ -89,7 +89,7 @@ bool UPCGMetadataMakeTransformSettings::IsSupportedInputType(uint16 TypeId, uint
 	}
 }
 
-FPCGAttributePropertySelector UPCGMetadataMakeTransformSettings::GetInputSource(uint32 Index) const
+FPCGAttributePropertyInputSelector UPCGMetadataMakeTransformSettings::GetInputSource(uint32 Index) const
 {
 	switch (Index)
 	{
@@ -100,7 +100,7 @@ FPCGAttributePropertySelector UPCGMetadataMakeTransformSettings::GetInputSource(
 	case 2:
 		return InputSource3;
 	default:
-		return FPCGAttributePropertySelector();
+		return FPCGAttributePropertyInputSelector();
 	}
 }
 
diff --git a/Engine/Plugins/Experimental/PCG/Source/PCG/Private/Elements/Metadata/PCGMetadataMakeVector.cpp b/Engine/Plugins/Experimental/PCG/Source/PCG/Private/Elements/Metadata/PCGMetadataMakeVector.cpp
--- a/Engine/Plugins/Experimental/PCG/Source/PCG/Private/Elements/Metadata/PCGMetadataMakeVector.cpp
+++ b/Engine/Plugins/Experimental/PCG/Source/PCG/Private/Elements/Metadata/PCGMetadataMakeVector.cpp
@@ -162,7 +162,7 @@ bool UPCGMetadataMakeVectorSettings::IsSupportedInputType(uint16 TypeId, uint32
 	}
 }
 
-FPCGAttributePropertySelector UPCGMetadataMakeVectorSettings::GetInputSource(uint32 Index) const
+FPCGAttributePropertyInputSelector UPCGMetadataMakeVectorSettings::GetInputSource(uint32 Index) const
 {
 	switch (Index)
 	{
@@ -175,7 +175,7 @@ FPCGAttributePropertySelector UPCGMetadataMakeVectorSettings::GetInputSource(uin
 	case 3:
 		return InputSource4;
 	default:
-		return FPCGAttributePropertySelector();
+		return FPCGAttributePropertyInputSelector();
 	}
 }
 
diff --git a/Engine/Plugins/Experimental/PCG/Source/PCG/Private/Elements/Metadata/PCGMetadataMathsOpElement.cpp b/Engine/Plugins/Experimental/PCG/Source/PCG/Private/Elements/Metadata/PCGMetadataMathsOpElement.cpp
--- a/Engine/Plugins/Experimental/PCG/Source/PCG/Private/Elements/Metadata/PCGMetadataMathsOpElement.cpp
+++ b/Engine/Plugins/Experimental/PCG/Source/PCG/Private/Elements/Metadata/PCGMetadataMathsOpElement.cpp
@@ -220,7 +220,7 @@ bool UPCGMetadataMathsSettings::IsSupportedInputType(uint16 TypeId, uint32 Input
 	return PCG::Private::IsOfTypes<float, double, int32, int64, FVector2D, FVector, FVector4>(TypeId);
 }
 
-FPCGAttributePropertySelector UPCGMetadataMathsSettings::GetInputSource(uint32 Index) const
+FPCGAttributePropertyInputSelector UPCGMetadataMathsSettings::GetInputSource(uint32 Index) const
 {
 	switch (Index)
 	{
@@ -231,7 +231,7 @@ FPCGAttributePropertySelector UPCGMetadataMathsSettings::GetInputSource(uint32 I
 	case 2:
 		return InputSource3;
 	default:
-		return FPCGAttributePropertySelector();
+		return FPCGAttributePropertyInputSelector();
 	}
 }
 
diff --git a/Engine/Plugins/Experimental/PCG/Source/PCG/Private/Elements/Metadata/PCGMetadataOpElementBase.cpp b/Engine/Plugins/Experimental/PCG/Source/PCG/Private/Elements/Metadata/PCGMetadataOpElementBase.cpp
--- a/Engine/Plugins/Experimental/PCG/Source/PCG/Private/Elements/Metadata/PCGMetadataOpElementBase.cpp
+++ b/Engine/Plugins/Experimental/PCG/Source/PCG/Private/Elements/Metadata/PCGMetadataOpElementBase.cpp
@@ -3,12 +3,13 @@
 #include "Elements/Metadata/PCGMetadataOpElementBase.h"
 
 #include "PCGContext.h"
+#include "PCGCustomVersion.h"
 #include "PCGParamData.h"
+#include "PCGPin.h"
 #include "Data/PCGPointData.h"
 #include "Data/PCGSpatialData.h"
 #include "Elements/Metadata/PCGMetadataElementCommon.h"
 #include "Metadata/Accessors/PCGAttributeAccessorHelpers.h"
-#include "PCGPin.h"
 
 #include UE_INLINE_GENERATED_CPP_BY_NAME(PCGMetadataOpElementBase)
 
@@ -21,8 +22,7 @@ void UPCGMetadataSettingsBase::PostLoad()
 #if WITH_EDITOR
 	if (OutputAttributeName_DEPRECATED != NAME_None)
 	{
-		OutputTarget.Selection = EPCGAttributePropertySelection::Attribute;
-		OutputTarget.AttributeName = OutputAttributeName_DEPRECATED;
+		OutputTarget.SetAttributeName(OutputAttributeName_DEPRECATED);
 		OutputAttributeName_DEPRECATED = NAME_None;
 	}
 #endif // WITH_EDITOR
@@ -164,6 +164,29 @@ FText UPCGMetadataSettingsBase::GetNodeTooltipText() const
 		"Output data will be taken from the first spatial data by default, or first pin if all are attribute sets.\n"
 		"It can be overridden in the settings.");
 }
+
+void UPCGMetadataSettingsBase::ApplyDeprecation(UPCGNode* InOutNode)
+{
+	if (DataVersion < FPCGCustomVersion::UpdateAttributePropertyInputSelector
+		&& OutputTarget.GetSelection() == EPCGAttributePropertySelection::Attribute
+		&& OutputTarget.GetAttributeName() == PCGMetadataAttributeConstants::SourceAttributeName)
+	{
+		// Previous behavior of the output target for this node was:
+		// - If the input to forward is an attribute -> SourceName
+		// - If the input to forward was not an attribute -> None
+		const FPCGAttributePropertyInputSelector& InputSource = GetInputSource(GetInputPinToForward());
+		if (InputSource.GetSelection() == EPCGAttributePropertySelection::Attribute)
+		{
+			OutputTarget.SetAttributeName(PCGMetadataAttributeConstants::SourceNameAttributeName);
+		}
+		else
+		{
+			OutputTarget.SetAttributeName(NAME_None);
+		}
+	}
+
+	Super::ApplyDeprecation(InOutNode);
+}
 #endif // WITH_EDITOR
 
 TArray<FPCGPinProperties> UPCGMetadataSettingsBase::InputPinProperties() const
@@ -310,6 +333,7 @@ bool FPCGMetadataElementBase::ExecuteInternal(FPCGContext* Context) const
 	FOperationData OperationData;
 	OperationData.InputAccessors.SetNum(NumberOfInputs);
 	OperationData.InputKeys.SetNum(NumberOfInputs);
+	OperationData.InputSources.SetNum(NumberOfInputs);
 	TArray<int32> NumberOfElements;
 	NumberOfElements.SetNum(NumberOfInputs);
 
@@ -318,11 +342,7 @@ bool FPCGMetadataElementBase::ExecuteInternal(FPCGContext* Context) const
 
 	bool bNoOperationNeeded = false;
 
-	// Use this name to forward it to the output if needed.
-	// Only set it if the first input is an attribute.
-	FName InputName = NAME_None;
-
-	auto CreateInputAccessorAndValidate = [this, Settings, InputPinToForward, &InputTaggedData, Context, &SourceMetadata, &OperationData, &InputName, &NumberOfElements, &bNoOperationNeeded, &SourceAttribute](uint32 Index) -> bool
+	auto CreateInputAccessorAndValidate = [this, Settings, InputPinToForward, &InputTaggedData, Context, &SourceMetadata, &OperationData, &NumberOfElements, &bNoOperationNeeded, &SourceAttribute](uint32 Index) -> bool
 	{
 		// First we verify if the input data match the first one. If the pin to forward is not connected, behave like a param data
 		const UClass* InputPinToForwardClass = (InputTaggedData[InputPinToForward].Data ? InputTaggedData[InputPinToForward].Data->GetClass() : UPCGParamData::StaticClass());
@@ -335,17 +355,8 @@ bool FPCGMetadataElementBase::ExecuteInternal(FPCGContext* Context) const
 			return false;
 		}
 
-		FPCGAttributePropertySelector InputSource = Settings->GetInputSource(Index);
-
-		if (InputSource.Selection == EPCGAttributePropertySelection::Attribute && InputSource.AttributeName == NAME_None && SourceMetadata[Index])
-		{
-			InputSource.AttributeName = SourceMetadata[Index]->GetLatestAttributeNameOrNone();
-		}
-
-		if (Index == InputPinToForward && InputSource.Selection == EPCGAttributePropertySelection::Attribute)
-		{
-			InputName = InputSource.AttributeName;
-		}
+		OperationData.InputSources[Index] = Settings->GetInputSource(Index).CopyAndFixLast(InputTaggedData[Index].Data);
+		const FPCGAttributePropertyInputSelector& InputSource = OperationData.InputSources[Index];
 
 		OperationData.InputAccessors[Index] = PCGAttributeAccessorHelpers::CreateConstAccessor(InputTaggedData[Index].Data, InputSource);
 		OperationData.InputKeys[Index] = PCGAttributeAccessorHelpers::CreateConstKeys(InputTaggedData[Index].Data, InputSource);
@@ -414,7 +425,7 @@ bool FPCGMetadataElementBase::ExecuteInternal(FPCGContext* Context) const
 			return false;
 		}
 
-		if (InputSource.Selection == EPCGAttributePropertySelection::Attribute)
+		if (InputSource.GetSelection() == EPCGAttributePropertySelection::Attribute)
 		{
 			SourceAttribute[Index] = SourceMetadata[Index]->GetConstAttribute(InputSource.GetName());
 		}
@@ -458,12 +469,7 @@ bool FPCGMetadataElementBase::ExecuteInternal(FPCGContext* Context) const
 	OperationData.OutputAccessors.SetNum(Settings->GetOutputPinNum());
 	OperationData.OutputKeys.SetNum(Settings->GetOutputPinNum());
 
-	FPCGAttributePropertySelector OutputTarget = Settings->OutputTarget;
-
-	if (OutputTarget.Selection == EPCGAttributePropertySelection::Attribute && OutputTarget.AttributeName == NAME_None)
-	{
-		OutputTarget.AttributeName = InputName;
-	}
+	FPCGAttributePropertyOutputSelector OutputTarget = Settings->OutputTarget.CopyAndFixSource(&OperationData.InputSources[InputPinToForward]);
 
 	// Use implicit capture, since we capture a lot
 	auto CreateAttribute = [&](uint32 OutputIndex, auto DummyOutValue) -> bool
@@ -477,7 +483,7 @@ bool FPCGMetadataElementBase::ExecuteInternal(FPCGContext* Context) const
 
 		FName OutputName = OutputTarget.GetName();
 
-		if (OutputTarget.Selection == EPCGAttributePropertySelection::Attribute && OutputTarget.ExtraNames.IsEmpty())
+		if (OutputTarget.GetSelection() == EPCGAttributePropertySelection::Attribute && OutputTarget.GetExtraNames().IsEmpty())
 		{
 			// In case of an attribute, we check if we have extra selectors. If not, we can just delete the attribute
 			// and create a new one of the right type.
diff --git a/Engine/Plugins/Experimental/PCG/Source/PCG/Private/Elements/Metadata/PCGMetadataRotatorOpElement.cpp b/Engine/Plugins/Experimental/PCG/Source/PCG/Private/Elements/Metadata/PCGMetadataRotatorOpElement.cpp
--- a/Engine/Plugins/Experimental/PCG/Source/PCG/Private/Elements/Metadata/PCGMetadataRotatorOpElement.cpp
+++ b/Engine/Plugins/Experimental/PCG/Source/PCG/Private/Elements/Metadata/PCGMetadataRotatorOpElement.cpp
@@ -179,7 +179,7 @@ bool UPCGMetadataRotatorSettings::IsSupportedInputType(uint16 TypeId, uint32 Inp
 	}
 }
 
-FPCGAttributePropertySelector UPCGMetadataRotatorSettings::GetInputSource(uint32 Index) const
+FPCGAttributePropertyInputSelector UPCGMetadataRotatorSettings::GetInputSource(uint32 Index) const
 {
 	switch (Index)
 	{
@@ -190,7 +190,7 @@ FPCGAttributePropertySelector UPCGMetadataRotatorSettings::GetInputSource(uint32
 	case 2:
 		return InputSource3;
 	default:
-		return FPCGAttributePropertySelector();
+		return FPCGAttributePropertyInputSelector();
 	}
 }
 
diff --git a/Engine/Plugins/Experimental/PCG/Source/PCG/Private/Elements/Metadata/PCGMetadataTransformOpElement.cpp b/Engine/Plugins/Experimental/PCG/Source/PCG/Private/Elements/Metadata/PCGMetadataTransformOpElement.cpp
--- a/Engine/Plugins/Experimental/PCG/Source/PCG/Private/Elements/Metadata/PCGMetadataTransformOpElement.cpp
+++ b/Engine/Plugins/Experimental/PCG/Source/PCG/Private/Elements/Metadata/PCGMetadataTransformOpElement.cpp
@@ -142,7 +142,7 @@ bool UPCGMetadataTransformSettings::IsSupportedInputType(uint16 TypeId, uint32 I
 	}
 }
 
-FPCGAttributePropertySelector UPCGMetadataTransformSettings::GetInputSource(uint32 Index) const
+FPCGAttributePropertyInputSelector UPCGMetadataTransformSettings::GetInputSource(uint32 Index) const
 {
 	switch (Index)
 	{
@@ -153,7 +153,7 @@ FPCGAttributePropertySelector UPCGMetadataTransformSettings::GetInputSource(uint
 	case 2:
 		return InputSource3;
 	default:
-		return FPCGAttributePropertySelector();
+		return FPCGAttributePropertyInputSelector();
 	}
 }
 
diff --git a/Engine/Plugins/Experimental/PCG/Source/PCG/Private/Elements/Metadata/PCGMetadataTrigOpElement.cpp b/Engine/Plugins/Experimental/PCG/Source/PCG/Private/Elements/Metadata/PCGMetadataTrigOpElement.cpp
--- a/Engine/Plugins/Experimental/PCG/Source/PCG/Private/Elements/Metadata/PCGMetadataTrigOpElement.cpp
+++ b/Engine/Plugins/Experimental/PCG/Source/PCG/Private/Elements/Metadata/PCGMetadataTrigOpElement.cpp
@@ -85,7 +85,7 @@ bool UPCGMetadataTrigSettings::IsSupportedInputType(uint16 TypeId, uint32 InputI
 	return PCG::Private::IsOfTypes<int32, int64, float, double>(TypeId);
 }
 
-FPCGAttributePropertySelector UPCGMetadataTrigSettings::GetInputSource(uint32 Index) const
+FPCGAttributePropertyInputSelector UPCGMetadataTrigSettings::GetInputSource(uint32 Index) const
 {
 	switch (Index)
 	{
@@ -94,7 +94,7 @@ FPCGAttributePropertySelector UPCGMetadataTrigSettings::GetInputSource(uint32 In
 	case 1:
 		return InputSource2;
 	default:
-		return FPCGAttributePropertySelector();
+		return FPCGAttributePropertyInputSelector();
 	}
 }
 
diff --git a/Engine/Plugins/Experimental/PCG/Source/PCG/Private/Elements/Metadata/PCGMetadataVectorOpElement.cpp b/Engine/Plugins/Experimental/PCG/Source/PCG/Private/Elements/Metadata/PCGMetadataVectorOpElement.cpp
--- a/Engine/Plugins/Experimental/PCG/Source/PCG/Private/Elements/Metadata/PCGMetadataVectorOpElement.cpp
+++ b/Engine/Plugins/Experimental/PCG/Source/PCG/Private/Elements/Metadata/PCGMetadataVectorOpElement.cpp
@@ -304,7 +304,7 @@ bool UPCGMetadataVectorSettings::IsSupportedInputType(uint16 TypeId, uint32 Inpu
 	}
 }
 
-FPCGAttributePropertySelector UPCGMetadataVectorSettings::GetInputSource(uint32 Index) const
+FPCGAttributePropertyInputSelector UPCGMetadataVectorSettings::GetInputSource(uint32 Index) const
 {
 	switch (Index)
 	{
@@ -315,7 +315,7 @@ FPCGAttributePropertySelector UPCGMetadataVectorSettings::GetInputSource(uint32
 	case 2:
 		return InputSource3;
 	default:
-		return FPCGAttributePropertySelector();
+		return FPCGAttributePropertyInputSelector();
 	}
 }
 
diff --git a/Engine/Plugins/Experimental/PCG/Source/PCG/Private/Elements/PCGAttributeGetFromPointIndexElement.cpp b/Engine/Plugins/Experimental/PCG/Source/PCG/Private/Elements/PCGAttributeGetFromPointIndexElement.cpp
--- a/Engine/Plugins/Experimental/PCG/Source/PCG/Private/Elements/PCGAttributeGetFromPointIndexElement.cpp
+++ b/Engine/Plugins/Experimental/PCG/Source/PCG/Private/Elements/PCGAttributeGetFromPointIndexElement.cpp
@@ -3,6 +3,7 @@
 #include "Elements/PCGAttributeGetFromPointIndexElement.h"
 
 #include "PCGContext.h"
+#include "PCGCustomVersion.h"
 #include "PCGPin.h"
 #include "PCGParamData.h"
 #include "Data/PCGPointData.h"
@@ -24,6 +25,19 @@ FText UPCGAttributeGetFromPointIndexSettings::GetDefaultNodeTitle() const
 {
 	return LOCTEXT("NodeTitle", "Get Attribute From Point Index");
 }
+
+void UPCGAttributeGetFromPointIndexSettings::ApplyDeprecation(UPCGNode* InOutNode)
+{
+	if (DataVersion < FPCGCustomVersion::UpdateAttributePropertyInputSelector
+		&& (OutputAttributeName == NAME_None))
+	{
+		// Previous behavior of the output attribute for this node was:
+		// None => SameName
+		OutputAttributeName = PCGMetadataAttributeConstants::SourceNameAttributeName;
+	}
+
+	Super::ApplyDeprecation(InOutNode);
+}
 #endif
 
 void UPCGAttributeGetFromPointIndexSettings::PostLoad()
@@ -111,13 +125,9 @@ bool FPCGAttributeGetFromPointIndexElement::ExecuteInternal(FPCGContext* Context
 		Output.Pin = PCGAttributeGetFromPointIndexConstants::OutputPointLabel;
 	}
 
-	FPCGAttributePropertySelector InputSource = Settings->InputSource;
-	if (InputSource.Selection == EPCGAttributePropertySelection::Attribute && InputSource.AttributeName == NAME_None)
-	{
-		InputSource.SetAttributeName(PointData->Metadata->GetLatestAttributeNameOrNone());
-	}
+	FPCGAttributePropertyInputSelector InputSource = Settings->InputSource.CopyAndFixLast(PointData);
 
-	FName OutputAttributeName = (Settings->OutputAttributeName == NAME_None) ? InputSource.GetName() : Settings->OutputAttributeName;
+	const FName OutputAttributeName = (Settings->OutputAttributeName == PCGMetadataAttributeConstants::SourceNameAttributeName) ? InputSource.GetName() : Settings->OutputAttributeName;
 
 	TUniquePtr<const IPCGAttributeAccessor> Accessor = PCGAttributeAccessorHelpers::CreateConstAccessor(PointData, InputSource);
 	FPCGAttributeAccessorKeysPoints PointKey(Point);
diff --git a/Engine/Plugins/Experimental/PCG/Source/PCG/Private/Elements/PCGAttributeNoise.cpp b/Engine/Plugins/Experimental/PCG/Source/PCG/Private/Elements/PCGAttributeNoise.cpp
--- a/Engine/Plugins/Experimental/PCG/Source/PCG/Private/Elements/PCGAttributeNoise.cpp
+++ b/Engine/Plugins/Experimental/PCG/Source/PCG/Private/Elements/PCGAttributeNoise.cpp
@@ -211,12 +211,14 @@ bool FPCGAttributeNoiseElement::ExecuteInternal(FPCGContext* InContext) const
 
 			InputPoints = &PointData->GetPoints();
 
+			Context->InputSource = Settings->InputSource.CopyAndFixLast(PointData);
+
 			// Create a dummy accessor on the input before allocating the output, to avoid doing useless allocation.
-			TUniquePtr<const IPCGAttributeAccessor> TempInputAccessor = PCGAttributeAccessorHelpers::CreateConstAccessor(PointData, Settings->InputSource);
+			TUniquePtr<const IPCGAttributeAccessor> TempInputAccessor = PCGAttributeAccessorHelpers::CreateConstAccessor(PointData, Context->InputSource);
 			if (!TempInputAccessor)
 			{
 				Outputs.RemoveAt(Outputs.Num() - 1);
-				PCGE_LOG(Error, GraphAndLog, FText::Format(LOCTEXT("CantCreateAccessor", "Could not find Attribute/Property {0}"), FText::FromName(Settings->InputSource.GetName())));
+				PCGE_LOG(Error, GraphAndLog, FText::Format(LOCTEXT("CantCreateAccessor", "Could not find Attribute/Property {0}"), FText::FromName(Context->InputSource.GetName())));
 				Context->CurrentInput++;
 				continue;
 			}
@@ -225,7 +227,7 @@ bool FPCGAttributeNoiseElement::ExecuteInternal(FPCGContext* InContext) const
 			if (!PCG::Private::IsOfTypes<int32, int64, float, double, FVector, FVector2D, FVector4, FRotator, FQuat>(TempInputAccessor->GetUnderlyingType()))
 			{
 				Outputs.RemoveAt(Outputs.Num() - 1);
-				PCGE_LOG(Error, GraphAndLog, FText::Format(LOCTEXT("AttributeIsNotANumericalType", "Attribute/Property {0} is not a numerical type, we can't apply noise to it."), FText::FromName(Settings->InputSource.GetName())));
+				PCGE_LOG(Error, GraphAndLog, FText::Format(LOCTEXT("AttributeIsNotANumericalType", "Attribute/Property {0} is not a numerical type, we can't apply noise to it."), FText::FromName(Context->InputSource.GetName())));
 				Context->CurrentInput++;
 				continue;
 			}
@@ -240,8 +242,8 @@ bool FPCGAttributeNoiseElement::ExecuteInternal(FPCGContext* InContext) const
 			Output.Data = OutputData;
 
 			// Then create the accessor/keys
-			Context->InputAccessor = PCGAttributeAccessorHelpers::CreateAccessor(OutputData, Settings->InputSource);
-			Context->Keys = PCGAttributeAccessorHelpers::CreateKeys(OutputData, Settings->InputSource);
+			Context->InputAccessor = PCGAttributeAccessorHelpers::CreateConstAccessor(OutputData, Context->InputSource);
+			Context->Keys = PCGAttributeAccessorHelpers::CreateKeys(OutputData, Context->InputSource);
 
 			// It won't fail because we validated on the input data and output is initialized from input.
 			check(Context->InputAccessor && Context->Keys);
@@ -255,31 +257,29 @@ bool FPCGAttributeNoiseElement::ExecuteInternal(FPCGContext* InContext) const
 				int32 NumPoints = Context->Keys->GetNum();
 				Context->TempValuesBuffer.SetNumUninitialized(sizeof(AttributeType) * NumPoints);
 
-				// If we have a different output target, create accessor (and also create the attribute if it doesn't exist), and check that we can broadcast to it
-				if (Settings->bOutputTargetDifferentFromInputSource)
+				Context->OutputTarget = Settings->OutputTarget.CopyAndFixSource(&Context->InputSource);
+
+				Context->OutputAccessor = PCGAttributeAccessorHelpers::CreateAccessor(OutputData, Context->OutputTarget);
+				if (!Context->OutputAccessor && Context->OutputTarget.GetSelection() == EPCGAttributePropertySelection::Attribute)
+				{
+					OutputData->Metadata->CreateAttribute<AttributeType>(Context->OutputTarget.GetName(), AttributeType{}, /*bAllowsInterpolation=*/ true, /*bOverrideParent=*/false);
+					Context->OutputAccessor = PCGAttributeAccessorHelpers::CreateAccessor(OutputData, Context->OutputTarget);
+				}
+
+				if (!Context->OutputAccessor)
+				{
+					PCGE_LOG(Error, GraphAndLog, FText::Format(LOCTEXT("OutputTargetInvalid", "Failed to find/create Attribute/Property {0}."), FText::FromName(Context->OutputTarget.GetName())));
+					return false;
+				}
+
+				if (!PCG::Private::IsBroadcastable(Context->InputAccessor->GetUnderlyingType(), Context->OutputAccessor->GetUnderlyingType()))
 				{
-					Context->OptionalOutputAccessor = PCGAttributeAccessorHelpers::CreateAccessor(OutputData, Settings->OutputTarget);
-					if (!Context->OptionalOutputAccessor && Settings->OutputTarget.Selection == EPCGAttributePropertySelection::Attribute)
-					{
-						OutputData->Metadata->CreateAttribute<AttributeType>(Settings->OutputTarget.GetName(), AttributeType{}, /*bAllowsInterpolation=*/ true, /*bOverrideParent=*/false);
-						Context->OptionalOutputAccessor = PCGAttributeAccessorHelpers::CreateAccessor(OutputData, Settings->OutputTarget);
-					}
-
-					if (!Context->OptionalOutputAccessor)
-					{
-						PCGE_LOG(Error, GraphAndLog, FText::Format(LOCTEXT("OutputTargetInvalid", "Failed to find/create Attribute/Property {0}."), FText::FromName(Settings->OutputTarget.GetName())));
-						return false;
-					}
-
-					if (!PCG::Private::IsBroadcastable(Context->InputAccessor->GetUnderlyingType(), Context->OptionalOutputAccessor->GetUnderlyingType()))
-					{
-						PCGE_LOG(Error, GraphAndLog, FText::Format(LOCTEXT("CantBroadcast", "Cannot convert Attribute {0} ({1}) into Attribute {2} ({3})."),
-							FText::FromName(Settings->InputSource.GetName()),
-							FText::FromString(PCG::Private::GetTypeName(Context->InputAccessor->GetUnderlyingType())),
-							FText::FromName(Settings->OutputTarget.GetName()),
-							FText::FromString(PCG::Private::GetTypeName(Context->OptionalOutputAccessor->GetUnderlyingType()))));
-						return false;
-					}
+					PCGE_LOG(Error, GraphAndLog, FText::Format(LOCTEXT("CantBroadcast", "Cannot convert Attribute {0} ({1}) into Attribute {2} ({3})."),
+						FText::FromName(Context->InputSource.GetName()),
+						FText::FromString(PCG::Private::GetTypeName(Context->InputAccessor->GetUnderlyingType())),
+						FText::FromName(Context->OutputTarget.GetName()),
+						FText::FromString(PCG::Private::GetTypeName(Context->OutputAccessor->GetUnderlyingType()))));
+					return false;
 				}
 
 				return true;
@@ -303,8 +303,7 @@ bool FPCGAttributeNoiseElement::ExecuteInternal(FPCGContext* InContext) const
 		check(InputPoints && OutputPoints);
 
 		// Force clamp on Density
-		const FPCGAttributePropertySelector& Selector = Settings->bOutputTargetDifferentFromInputSource ? Settings->OutputTarget : Settings->InputSource;
-		const bool bClampResult = Settings->bClampResult || (Selector.Selection == EPCGAttributePropertySelection::PointProperty && Selector.PointProperty == EPCGPointProperties::Density);
+		const bool bClampResult = Settings->bClampResult || (Context->OutputTarget.GetSelection() == EPCGAttributePropertySelection::PointProperty && Context->OutputTarget.GetPointProperty() == EPCGPointProperties::Density);
 
 		// Dummy Initialize, we already initialized the output points before.
 		auto Initialize = []() {};
@@ -339,14 +338,7 @@ bool FPCGAttributeNoiseElement::ExecuteInternal(FPCGContext* InContext) const
 				using AttributeType = std::decay_t<decltype(Value)>;
 				int32 NumPoints = Context->Keys->GetNum();
 				TArrayView<AttributeType> Values(reinterpret_cast<AttributeType*>(Context->TempValuesBuffer.GetData()), NumPoints);
-				if (Context->OptionalOutputAccessor)
-				{
-					Context->OptionalOutputAccessor->SetRange<AttributeType>(Values, 0, *Context->Keys, EPCGAttributeAccessorFlags::AllowBroadcast);
-				}
-				else
-				{
-					Context->InputAccessor->SetRange<AttributeType>(Values, 0, *Context->Keys);
-				}
+				Context->OutputAccessor->SetRange<AttributeType>(Values, 0, *Context->Keys, EPCGAttributeAccessorFlags::AllowBroadcast);
 			});
 			
 			Context->CurrentInput++;
diff --git a/Engine/Plugins/Experimental/PCG/Source/PCG/Private/Elements/PCGAttributeReduceElement.cpp b/Engine/Plugins/Experimental/PCG/Source/PCG/Private/Elements/PCGAttributeReduceElement.cpp
--- a/Engine/Plugins/Experimental/PCG/Source/PCG/Private/Elements/PCGAttributeReduceElement.cpp
+++ b/Engine/Plugins/Experimental/PCG/Source/PCG/Private/Elements/PCGAttributeReduceElement.cpp
@@ -2,13 +2,14 @@
 
 #include "Elements/PCGAttributeReduceElement.h"
 
-#include "Data/PCGSpatialData.h"
+#include "PCGContext.h"
+#include "PCGCustomVersion.h"
 #include "PCGParamData.h"
+#include "PCGPin.h"
 #include "Data/PCGPointData.h"
+#include "Data/PCGSpatialData.h"
 #include "Elements/Metadata/PCGMetadataElementCommon.h"
 #include "Metadata/Accessors/PCGAttributeAccessorHelpers.h"
-#include "PCGContext.h"
-#include "PCGPin.h"
 
 #include UE_INLINE_GENERATED_CPP_BY_NAME(PCGAttributeReduceElement)
 
@@ -86,6 +87,19 @@ FText UPCGAttributeReduceSettings::GetDefaultNodeTitle() const
 {
 	return LOCTEXT("NodeTitle", "Attribute Reduce");
 }
+
+void UPCGAttributeReduceSettings::ApplyDeprecation(UPCGNode* InOutNode)
+{
+	if (DataVersion < FPCGCustomVersion::UpdateAttributePropertyInputSelector
+		&& (OutputAttributeName == NAME_None))
+	{
+		// Previous behavior of the output attribute for this node was:
+		// None => SameName
+		OutputAttributeName = PCGMetadataAttributeConstants::SourceNameAttributeName;
+	}
+
+	Super::ApplyDeprecation(InOutNode);
+}
 #endif
 
 void UPCGAttributeReduceSettings::PostLoad()
@@ -176,13 +190,9 @@ bool FPCGAttributeReduceElement::ExecuteInternal(FPCGContext* Context) const
 
 	const UPCGPointData* PointData = Cast<UPCGPointData>(SpatialData);
 
-	FPCGAttributePropertySelector InputSource = Settings->InputSource;
-	if (InputSource.Selection == EPCGAttributePropertySelection::Attribute && InputSource.AttributeName == NAME_None && SpatialData->Metadata)
-	{
-		InputSource.SetAttributeName(SpatialData->Metadata->GetLatestAttributeNameOrNone());
-	}
+	FPCGAttributePropertyInputSelector InputSource = Settings->InputSource.CopyAndFixLast(SpatialData);
 
-	const FName OutputAttributeName = (Settings->OutputAttributeName == NAME_None) ? InputSource.GetName() : Settings->OutputAttributeName;
+	const FName OutputAttributeName = (Settings->OutputAttributeName == PCGMetadataAttributeConstants::SourceNameAttributeName) ? InputSource.GetName() : Settings->OutputAttributeName;
 	UPCGParamData* OutputParamData = NewObject<UPCGParamData>();
 
 	TUniquePtr<const IPCGAttributeAccessor> Accessor = PCGAttributeAccessorHelpers::CreateConstAccessor(PointData, InputSource);
diff --git a/Engine/Plugins/Experimental/PCG/Source/PCG/Private/Elements/PCGAttributeSelectElement.cpp b/Engine/Plugins/Experimental/PCG/Source/PCG/Private/Elements/PCGAttributeSelectElement.cpp
--- a/Engine/Plugins/Experimental/PCG/Source/PCG/Private/Elements/PCGAttributeSelectElement.cpp
+++ b/Engine/Plugins/Experimental/PCG/Source/PCG/Private/Elements/PCGAttributeSelectElement.cpp
@@ -2,13 +2,14 @@
 
 #include "Elements/PCGAttributeSelectElement.h"
 
-#include "Data/PCGSpatialData.h"
+#include "PCGContext.h"
+#include "PCGCustomVersion.h"
 #include "PCGParamData.h"
+#include "PCGPin.h"
 #include "Data/PCGPointData.h"
+#include "Data/PCGSpatialData.h"
 #include "Elements/Metadata/PCGMetadataElementCommon.h"
 #include "Metadata/Accessors/PCGAttributeAccessorHelpers.h"
-#include "PCGContext.h"
-#include "PCGPin.h"
 
 #include UE_INLINE_GENERATED_CPP_BY_NAME(PCGAttributeSelectElement)
 
@@ -191,6 +192,19 @@ FText UPCGAttributeSelectSettings::GetDefaultNodeTitle() const
 {
 	return LOCTEXT("NodeTitle", "Attribute Select");
 }
+
+void UPCGAttributeSelectSettings::ApplyDeprecation(UPCGNode* InOutNode)
+{
+	if (DataVersion < FPCGCustomVersion::UpdateAttributePropertyInputSelector
+		&& (OutputAttributeName == NAME_None))
+	{
+		// Previous behavior of the output attribute for this node was:
+		// None => SameName
+		OutputAttributeName = PCGMetadataAttributeConstants::SourceNameAttributeName;
+	}
+
+	Super::ApplyDeprecation(InOutNode);
+}
 #endif
 
 FName UPCGAttributeSelectSettings::AdditionalTaskName() const
@@ -289,13 +303,9 @@ bool FPCGAttributeSelectElement::ExecuteInternal(FPCGContext* Context) const
 		return true;
 	}
 
-	FPCGAttributePropertySelector InputSource = Settings->InputSource;
-	if (InputSource.Selection == EPCGAttributePropertySelection::Attribute && InputSource.AttributeName == NAME_None && SpatialData->Metadata)
-	{
-		InputSource.SetAttributeName(SpatialData->Metadata->GetLatestAttributeNameOrNone());
-	}
+	FPCGAttributePropertyInputSelector InputSource = Settings->InputSource.CopyAndFixLast(SpatialData);
 
-	const FName OutputAttributeName = (Settings->OutputAttributeName == NAME_None) ? InputSource.GetName() : Settings->OutputAttributeName;
+	const FName OutputAttributeName = (Settings->OutputAttributeName == PCGMetadataAttributeConstants::SourceNameAttributeName) ? InputSource.GetName() : Settings->OutputAttributeName;
 	UPCGParamData* OutputParamData = NewObject<UPCGParamData>();
 
 	TUniquePtr<const IPCGAttributeAccessor> Accessor = PCGAttributeAccessorHelpers::CreateConstAccessor(PointData, InputSource);
diff --git a/Engine/Plugins/Experimental/PCG/Source/PCG/Private/Elements/PCGPointFilter.cpp b/Engine/Plugins/Experimental/PCG/Source/PCG/Private/Elements/PCGPointFilter.cpp
--- a/Engine/Plugins/Experimental/PCG/Source/PCG/Private/Elements/PCGPointFilter.cpp
+++ b/Engine/Plugins/Experimental/PCG/Source/PCG/Private/Elements/PCGPointFilter.cpp
@@ -140,8 +140,9 @@ namespace PCGPointFilterHelpers
 
 				if (ThresholdData)
 				{
-					OutThresholdInfo.ThresholdAccessor = PCGAttributeAccessorHelpers::CreateConstAccessor(ThresholdData, ThresholdSettings.ThresholdAttribute);
-					OutThresholdInfo.ThresholdKeys = PCGAttributeAccessorHelpers::CreateConstKeys(ThresholdData, ThresholdSettings.ThresholdAttribute);
+					FPCGAttributePropertyInputSelector ThresholdSelector = ThresholdSettings.ThresholdAttribute.CopyAndFixLast(ThresholdData);
+					OutThresholdInfo.ThresholdAccessor = PCGAttributeAccessorHelpers::CreateConstAccessor(ThresholdData, ThresholdSelector);
+					OutThresholdInfo.ThresholdKeys = PCGAttributeAccessorHelpers::CreateConstKeys(ThresholdData, ThresholdSelector);
 				}
 			}
 			else
@@ -160,8 +161,9 @@ namespace PCGPointFilterHelpers
 		if (InOutThresholdInfo.bUseInputDataForThreshold)
 		{
 			// If we have no threshold accessor, we use the same data as input
-			InOutThresholdInfo.ThresholdAccessor = PCGAttributeAccessorHelpers::CreateConstAccessor(InputData, ThresholdSettings.ThresholdAttribute);
-			InOutThresholdInfo.ThresholdKeys = PCGAttributeAccessorHelpers::CreateConstKeys(InputData, ThresholdSettings.ThresholdAttribute);
+			FPCGAttributePropertyInputSelector ThresholdSelector = ThresholdSettings.ThresholdAttribute.CopyAndFixLast(InputData);
+			InOutThresholdInfo.ThresholdAccessor = PCGAttributeAccessorHelpers::CreateConstAccessor(InputData, ThresholdSelector);
+			InOutThresholdInfo.ThresholdKeys = PCGAttributeAccessorHelpers::CreateConstKeys(InputData, ThresholdSelector);
 		}
 		else if (InOutThresholdInfo.ThresholdSpatialData != nullptr && InOutThresholdInfo.bUseSpatialQuery)
 		{
@@ -181,8 +183,9 @@ namespace PCGPointFilterHelpers
 			}
 
 			// Accessor will be valid, but keys will point to default points. But since it is a view, it will be updated when we sample the points.
-			InOutThresholdInfo.ThresholdAccessor = PCGAttributeAccessorHelpers::CreateConstAccessor(InOutThresholdInfo.ThresholdPointData, ThresholdSettings.ThresholdAttribute);
-			InOutThresholdInfo.ThresholdKeys = PCGAttributeAccessorHelpers::CreateConstKeys(InOutThresholdInfo.ThresholdPointData, ThresholdSettings.ThresholdAttribute);
+			FPCGAttributePropertyInputSelector ThresholdSelector = ThresholdSettings.ThresholdAttribute.CopyAndFixLast(InOutThresholdInfo.ThresholdPointData);
+			InOutThresholdInfo.ThresholdAccessor = PCGAttributeAccessorHelpers::CreateConstAccessor(InOutThresholdInfo.ThresholdPointData, ThresholdSelector);
+			InOutThresholdInfo.ThresholdKeys = PCGAttributeAccessorHelpers::CreateConstKeys(InOutThresholdInfo.ThresholdPointData, ThresholdSelector);
 		}
 
 		if (!InOutThresholdInfo.ThresholdAccessor.IsValid() || !InOutThresholdInfo.ThresholdKeys.IsValid())
@@ -413,7 +416,7 @@ TArray<FPCGPinProperties> UPCGPointFilterRangeSettings::OutputPinProperties() co
 	return PinProperties;
 }
 
-bool FPCGPointFilterElementBase::DoFiltering(FPCGContext* Context, EPCGPointFilterOperator InOperation, const FPCGAttributePropertySelector& TargetAttribute, const FPCGPointFilterThresholdSettings& FirstThreshold, const FPCGPointFilterThresholdSettings* SecondThreshold) const
+bool FPCGPointFilterElementBase::DoFiltering(FPCGContext* Context, EPCGPointFilterOperator InOperation, const FPCGAttributePropertyInputSelector& InTargetAttribute, const FPCGPointFilterThresholdSettings& FirstThreshold, const FPCGPointFilterThresholdSettings* SecondThreshold) const
 {
 	TRACE_CPUPROFILER_EVENT_SCOPE(FPCGPointFilterElementBase::DoFiltering);
 	check(Context);
@@ -484,6 +487,7 @@ bool FPCGPointFilterElementBase::DoFiltering(FPCGContext* Context, EPCGPointFilt
 
 		const TArray<FPCGPoint>& OriginalPoints = OriginalData->GetPoints();
 
+		FPCGAttributePropertyInputSelector TargetAttribute = InTargetAttribute.CopyAndFixLast(OriginalData);
 		TUniquePtr<const IPCGAttributeAccessor> TargetAccessor = PCGAttributeAccessorHelpers::CreateConstAccessor(OriginalData, TargetAttribute);
 		TUniquePtr<const IPCGAttributeAccessorKeys> TargetKeys = PCGAttributeAccessorHelpers::CreateConstKeys(OriginalData, TargetAttribute);
 
diff --git a/Engine/Plugins/Experimental/PCG/Source/PCG/Private/Elements/PCGPointMatchAndSet.cpp b/Engine/Plugins/Experimental/PCG/Source/PCG/Private/Elements/PCGPointMatchAndSet.cpp
--- a/Engine/Plugins/Experimental/PCG/Source/PCG/Private/Elements/PCGPointMatchAndSet.cpp
+++ b/Engine/Plugins/Experimental/PCG/Source/PCG/Private/Elements/PCGPointMatchAndSet.cpp
@@ -3,6 +3,7 @@
 #include "Elements/PCGPointMatchAndSet.h"
 
 #include "PCGContext.h"
+#include "PCGCustomVersion.h"
 #include "PCGPin.h"
 #include "Data/PCGPointData.h"
 #include "Helpers/PCGHelpers.h"
@@ -33,6 +34,20 @@ FText UPCGPointMatchAndSetSettings::GetNodeTooltipText() const
 {
 	return LOCTEXT("PointMatchAndSetNodeTooltip", "For all points, if a match is found (e.g. some attribute is equal to some value), sets a value on the point (e.g. another attribute).");
 }
+
+void UPCGPointMatchAndSetSettings::ApplyDeprecation(UPCGNode* InOutNode)
+{
+	if (DataVersion < FPCGCustomVersion::UpdateAttributePropertyInputSelector
+		&& SetTarget.GetSelection() == EPCGAttributePropertySelection::Attribute
+		&& SetTarget.GetAttributeName() == PCGMetadataAttributeConstants::SourceAttributeName)
+	{
+		// Previous behavior of the output target for this node was:
+		// None => LastCreated
+		SetTarget.SetAttributeName(PCGMetadataAttributeConstants::LastCreatedAttributeName);
+	}
+
+	Super::ApplyDeprecation(InOutNode);
+}
 #endif // WITH_EDITOR
 
 TArray<FPCGPinProperties> UPCGPointMatchAndSetSettings::InputPinProperties() const
@@ -88,7 +103,7 @@ void UPCGPointMatchAndSetSettings::PostEditChangeProperty(FPropertyChangedEvent&
 		}
 		else if (PropertyName == GET_MEMBER_NAME_CHECKED(UPCGPointMatchAndSetSettings, SetTarget))
 		{
-			bSetTargetIsAttribute = (SetTarget.Selection == EPCGAttributePropertySelection::Attribute);
+			bSetTargetIsAttribute = (SetTarget.GetSelection() == EPCGAttributePropertySelection::Attribute);
 
 			// If not targetting an attribute, but a property, assign the new type to the OutputType accordingly
 			if (!bSetTargetIsAttribute)
diff --git a/Engine/Plugins/Experimental/PCG/Source/PCG/Private/Elements/PCGPropertyToParamData.cpp b/Engine/Plugins/Experimental/PCG/Source/PCG/Private/Elements/PCGPropertyToParamData.cpp
--- a/Engine/Plugins/Experimental/PCG/Source/PCG/Private/Elements/PCGPropertyToParamData.cpp
+++ b/Engine/Plugins/Experimental/PCG/Source/PCG/Private/Elements/PCGPropertyToParamData.cpp
@@ -214,7 +214,8 @@ bool FPCGPropertyToParamDataElement::ExecuteInternal(FPCGContext* Context) const
 	}
 	else
 	{
-		ExtractableProperties.Emplace(Settings->OutputAttributeName, ObjectToInspect, Property);
+		const FName AttributeName = (Settings->OutputAttributeName == PCGMetadataAttributeConstants::SourceNameAttributeName) ? Property->GetFName() : Settings->OutputAttributeName;
+		ExtractableProperties.Emplace(AttributeName, ObjectToInspect, Property);
 	}
 
 	if (ExtractableProperties.IsEmpty())
diff --git a/Engine/Plugins/Experimental/PCG/Source/PCG/Private/Elements/PCGSpatialNoise.cpp b/Engine/Plugins/Experimental/PCG/Source/PCG/Private/Elements/PCGSpatialNoise.cpp
--- a/Engine/Plugins/Experimental/PCG/Source/PCG/Private/Elements/PCGSpatialNoise.cpp
+++ b/Engine/Plugins/Experimental/PCG/Source/PCG/Private/Elements/PCGSpatialNoise.cpp
@@ -340,9 +340,16 @@ namespace PCGSpatialNoise
 	template<typename T>
 	void SetAttributeHelper(UPCGSpatialData* Data, const FPCGAttributePropertySelector& PropertySelector, const TArrayView<const T> Values)
 	{
-		if (PropertySelector.Selection == EPCGAttributePropertySelection::Attribute)
+		if (PropertySelector.GetSelection() == EPCGAttributePropertySelection::Attribute)
 		{
-			Data->Metadata->FindOrCreateAttribute<T>(PropertySelector.AttributeName);
+			// Discard None attributes
+			const FName AttributeName = PropertySelector.GetAttributeName();
+			if (AttributeName == NAME_None)
+			{
+				return;
+			}
+
+			Data->Metadata->FindOrCreateAttribute<T>(AttributeName);
 		}
 
 		TUniquePtr<IPCGAttributeAccessor> Accessor = PCGAttributeAccessorHelpers::CreateAccessor(Data, PropertySelector);
diff --git a/Engine/Plugins/Experimental/PCG/Source/PCG/Private/Elements/PCGSpawnActor.cpp b/Engine/Plugins/Experimental/PCG/Source/PCG/Private/Elements/PCGSpawnActor.cpp
--- a/Engine/Plugins/Experimental/PCG/Source/PCG/Private/Elements/PCGSpawnActor.cpp
+++ b/Engine/Plugins/Experimental/PCG/Source/PCG/Private/Elements/PCGSpawnActor.cpp
@@ -104,7 +104,7 @@ namespace PCGSpawnActorHelpers
 
 			for (int32 i = 0; i < Overrides.Num(); ++i)
 			{
-				const FPCGAttributePropertySelector& InputSelector = Overrides[i].InputSource;
+				FPCGAttributePropertyInputSelector InputSelector = Overrides[i].InputSource.CopyAndFixLast(PointData);
 				const FString& OutputProperty = Overrides[i].PropertyTarget;
 
 				ActorSingleOverrides.Emplace(InputSelector, OutputProperty, TemplateActor, PointData);
diff --git a/Engine/Plugins/Experimental/PCG/Source/PCG/Private/Helpers/PCGMetadataHelpers.cpp b/Engine/Plugins/Experimental/PCG/Source/PCG/Private/Helpers/PCGMetadataHelpers.cpp
--- a/Engine/Plugins/Experimental/PCG/Source/PCG/Private/Helpers/PCGMetadataHelpers.cpp
+++ b/Engine/Plugins/Experimental/PCG/Source/PCG/Private/Helpers/PCGMetadataHelpers.cpp
@@ -1,6 +1,9 @@
 // Copyright Epic Games, Inc. All Rights Reserved.
 
 #include "Helpers/PCGMetadataHelpers.h"
+
+#include "PCGParamData.h"
+#include "Data/PCGSpatialData.h"
 #include "Metadata/PCGMetadata.h"
 
 namespace PCGMetadataHelpers
@@ -20,4 +23,36 @@ namespace PCGMetadataHelpers
 		//check(Parent.IsExplicitlyNull() || Parent.IsValid());
 		return Parent.Get();
 	}
+
+	const UPCGMetadata* GetConstMetadata(const UPCGData* InData)
+	{
+		if (const UPCGSpatialData* SpatialData = Cast<UPCGSpatialData>(InData))
+		{
+			return SpatialData->ConstMetadata();
+		}
+		else if (const UPCGParamData* ParamData = Cast<UPCGParamData>(InData))
+		{
+			return ParamData->ConstMetadata();
+		}
+		else
+		{
+			return nullptr;
+		}
+	}
+
+	UPCGMetadata* GetMutableMetadata(UPCGData* InData)
+	{
+		if (UPCGSpatialData* SpatialData = Cast<UPCGSpatialData>(InData))
+		{
+			return SpatialData->MutableMetadata();
+		}
+		else if (UPCGParamData* ParamData = Cast<UPCGParamData>(InData))
+		{
+			return ParamData->MutableMetadata();
+		}
+		else
+		{
+			return nullptr;
+		}
+	}
 }
\ No newline at end of file
diff --git a/Engine/Plugins/Experimental/PCG/Source/PCG/Private/MatchAndSet/PCGMatchAndSetBase.cpp b/Engine/Plugins/Experimental/PCG/Source/PCG/Private/MatchAndSet/PCGMatchAndSetBase.cpp
--- a/Engine/Plugins/Experimental/PCG/Source/PCG/Private/MatchAndSet/PCGMatchAndSetBase.cpp
+++ b/Engine/Plugins/Experimental/PCG/Source/PCG/Private/MatchAndSet/PCGMatchAndSetBase.cpp
@@ -31,13 +31,9 @@ bool UPCGMatchAndSetBase::CreateAttributeIfNeeded(FPCGContext& Context, const FP
 	check(OutPointData && OutPointData->Metadata);
 
 	check(OutPointData->Metadata);
-	if (Selector.Selection == EPCGAttributePropertySelection::Attribute)
+	if (Selector.GetSelection() == EPCGAttributePropertySelection::Attribute)
 	{
 		FName DestinationAttribute = Selector.GetName();
-		if (DestinationAttribute == NAME_None)
-		{
-			DestinationAttribute = OutPointData->Metadata->GetLatestAttributeNameOrNone();
-		}
 
 		if (!OutPointData->Metadata->HasAttribute(DestinationAttribute) ||
 			OutPointData->Metadata->GetConstAttribute(DestinationAttribute)->GetTypeId() != static_cast<uint16>(InSettings->SetTargetType))
diff --git a/Engine/Plugins/Experimental/PCG/Source/PCG/Private/MatchAndSet/PCGMatchAndSetByAttribute.cpp b/Engine/Plugins/Experimental/PCG/Source/PCG/Private/MatchAndSet/PCGMatchAndSetByAttribute.cpp
--- a/Engine/Plugins/Experimental/PCG/Source/PCG/Private/MatchAndSet/PCGMatchAndSetByAttribute.cpp
+++ b/Engine/Plugins/Experimental/PCG/Source/PCG/Private/MatchAndSet/PCGMatchAndSetByAttribute.cpp
@@ -78,7 +78,12 @@ void UPCGMatchAndSetByAttribute::MatchAndSet_Implementation(
 		return;
 	}
 
-	const FPCGAttributePropertySelector& SetTarget = InSettings->SetTarget;
+	FPCGAttributePropertyInputSelector InputSource;
+	InputSource.SetAttributeName(MatchSourceAttribute);
+	InputSource = InputSource.CopyAndFixLast(InPointData);
+
+	// Deprecation old behavior, if SetTarget was None, we took last created in OutPointData
+	FPCGAttributePropertyOutputSelector SetTarget = InSettings->SetTarget.CopyAndFixSource(&InputSource, OutPointData);
 
 	// Create attribute if needed
 	if (!CreateAttributeIfNeeded(Context, SetTarget, Entries[0].Value, OutPointData, InSettings))
@@ -112,10 +117,6 @@ void UPCGMatchAndSetByAttribute::MatchAndSet_Implementation(
 	}
 
 	//TODO: implement async loop?
-	FPCGAttributePropertySelector InputSource;
-	InputSource.Selection = EPCGAttributePropertySelection::Attribute;
-	InputSource.AttributeName = ((MatchSourceAttribute == NAME_None) ? InPointData->Metadata->GetLatestAttributeNameOrNone() : MatchSourceAttribute);
-
 	TUniquePtr<const IPCGAttributeAccessor> InputAccessor = PCGAttributeAccessorHelpers::CreateConstAccessor(InPointData, InputSource);
 	TUniquePtr<const IPCGAttributeAccessorKeys> InputKeys = PCGAttributeAccessorHelpers::CreateConstKeys(InPointData, InputSource);
 
diff --git a/Engine/Plugins/Experimental/PCG/Source/PCG/Private/MatchAndSet/PCGMatchAndSetWeightedByCategory.cpp b/Engine/Plugins/Experimental/PCG/Source/PCG/Private/MatchAndSet/PCGMatchAndSetWeightedByCategory.cpp
--- a/Engine/Plugins/Experimental/PCG/Source/PCG/Private/MatchAndSet/PCGMatchAndSetWeightedByCategory.cpp
+++ b/Engine/Plugins/Experimental/PCG/Source/PCG/Private/MatchAndSet/PCGMatchAndSetWeightedByCategory.cpp
@@ -127,7 +127,12 @@ void UPCGMatchAndSetWeightedByCategory::MatchAndSet_Implementation(
 		return;
 	}
 
-	const FPCGAttributePropertySelector& SetTarget = InSettings->SetTarget;
+	FPCGAttributePropertyInputSelector InputSource;
+	InputSource.SetAttributeName(CategoryAttribute);
+	InputSource = InputSource.CopyAndFixLast(InPointData);
+
+	// Deprecation old behavior, if SetTarget was None, we took last created in OutPointData
+	FPCGAttributePropertyOutputSelector SetTarget = InSettings->SetTarget.CopyAndFixSource(&InputSource, OutPointData);
 
 	// Create attribute if needed
 	if (!CreateAttributeIfNeeded(Context, SetTarget, *FirstValueForTypeExtraction, OutPointData, InSettings))
@@ -167,10 +172,6 @@ void UPCGMatchAndSetWeightedByCategory::MatchAndSet_Implementation(
 	const UPCGComponent* SourceComponent = Context.SourceComponent.Get();
 	const TArray<FPCGPoint>& InPoints = InPointData->GetPoints();
 
-	FPCGAttributePropertySelector InputSource;
-	InputSource.Selection = EPCGAttributePropertySelection::Attribute;
-	InputSource.AttributeName = ((CategoryAttribute == NAME_None) ? InPointData->Metadata->GetLatestAttributeNameOrNone() : CategoryAttribute);
-
 	TUniquePtr<const IPCGAttributeAccessor> InputAccessor = PCGAttributeAccessorHelpers::CreateConstAccessor(InPointData, InputSource);
 	TUniquePtr<const IPCGAttributeAccessorKeys> InputKeys = PCGAttributeAccessorHelpers::CreateConstKeys(InPointData, InputSource);
 
diff --git a/Engine/Plugins/Experimental/PCG/Source/PCG/Private/Metadata/Accessors/PCGAttributeAccessorHelpers.cpp b/Engine/Plugins/Experimental/PCG/Source/PCG/Private/Metadata/Accessors/PCGAttributeAccessorHelpers.cpp
--- a/Engine/Plugins/Experimental/PCG/Source/PCG/Private/Metadata/Accessors/PCGAttributeAccessorHelpers.cpp
+++ b/Engine/Plugins/Experimental/PCG/Source/PCG/Private/Metadata/Accessors/PCGAttributeAccessorHelpers.cpp
@@ -18,7 +18,7 @@
 
 namespace PCGAttributeAccessorHelpers
 {
-	void ExtractMetadataAtribute(UPCGData* InData, FName Name, UPCGMetadata*& OutMetadata, FPCGMetadataAttributeBase*& OutAttribute, bool bForceGetNone = false)
+	void ExtractMetadataAtribute(UPCGData* InData, FName Name, UPCGMetadata*& OutMetadata, FPCGMetadataAttributeBase*& OutAttribute)
 	{
 		OutMetadata = nullptr;
 		OutAttribute = nullptr;
@@ -34,12 +34,6 @@ namespace PCGAttributeAccessorHelpers
 
 		if (OutMetadata)
 		{
-			// If Name is None, try to get the latest attribute
-			if (!bForceGetNone && Name == NAME_None)
-			{
-				Name = OutMetadata->GetLatestAttributeNameOrNone();
-			}
-
 			OutAttribute = OutMetadata->GetMutableAttribute(Name);
 		}
 	}
@@ -65,7 +59,7 @@ namespace PCGAttributeAccessorHelpers
 		const FName Name = InSelector.GetName();
 		TUniquePtr<IPCGAttributeAccessor> Accessor;
 
-		if (InSelector.Selection == EPCGAttributePropertySelection::PointProperty)
+		if (InSelector.GetSelection() == EPCGAttributePropertySelection::PointProperty)
 		{
 			if (!InData || InData->template IsA<UPCGPointData>())
 			{
@@ -79,9 +73,9 @@ namespace PCGAttributeAccessorHelpers
 				}
 			}
 		}
-		else if (InSelector.Selection == EPCGAttributePropertySelection::ExtraProperty)
+		else if (InSelector.GetSelection() == EPCGAttributePropertySelection::ExtraProperty)
 		{
-			Accessor = CreateExtraAccessor(InSelector.ExtraProperty);
+			Accessor = CreateExtraAccessor(InSelector.GetExtraProperty());
 
 			if (!Accessor.IsValid())
 			{
@@ -92,13 +86,13 @@ namespace PCGAttributeAccessorHelpers
 
 		// At this point, it is not a point data or we didn't find a property.
 		// We can't continue if it is a property wanted.
-		if (InSelector.Selection == EPCGAttributePropertySelection::PointProperty && !Accessor.IsValid())
+		if (InSelector.GetSelection() == EPCGAttributePropertySelection::PointProperty && !Accessor.IsValid())
 		{
 			UE_LOG(LogPCG, Error, TEXT("[PCGAttributeAccessorHelpers::CreateAccessor] Expected to select a property but the data doesn't support this property."));
 			return TUniquePtr<IPCGAttributeAccessor>();
 		}
 
-		if (InSelector.Selection == EPCGAttributePropertySelection::Attribute)
+		if (InSelector.GetSelection() == EPCGAttributePropertySelection::Attribute)
 		{
 			UPCGMetadata* Metadata = nullptr;
 			FPCGMetadataAttributeBase* Attribute = nullptr;
@@ -135,7 +129,7 @@ namespace PCGAttributeAccessorHelpers
 		}
 
 		// At this point, check if we have chain accessors
-		for (const FString& ExtraName : InSelector.ExtraNames)
+		for (const FString& ExtraName : InSelector.GetExtraNames())
 		{
 			bool bSuccess = false;
 			Accessor = CreateChainAccessor(std::move(Accessor), FName(ExtraName), bSuccess);
@@ -154,6 +148,9 @@ namespace PCGAttributeAccessorHelpers
 				UE_LOG(LogPCG, Error, TEXT("[PCGAttributeAccessorHelpers::CreateAccessor] Attribute can not be written into, since it is read-only."));
 				return TUniquePtr<AccessorType>();
 			}
+
+			// Set the cached selector
+			InData->SetLastSelector(InSelector);
 		}
 
 		return Accessor;
@@ -437,7 +434,7 @@ TUniquePtr<const IPCGAttributeAccessor> PCGAttributeAccessorHelpers::CreateConst
 			*OutAttributeName = AttributeName;
 		}
 
-		FPCGAttributePropertySelector InputSelector{};
+		FPCGAttributePropertyInputSelector InputSelector{};
 		InputSelector.SetAttributeName(AttributeName);
 		return PCGAttributeAccessorHelpers::CreateConstAccessor(ParamData, InputSelector);
 	}
diff --git a/Engine/Plugins/Experimental/PCG/Source/PCG/Private/Metadata/PCGAttributePropertySelector.cpp b/Engine/Plugins/Experimental/PCG/Source/PCG/Private/Metadata/PCGAttributePropertySelector.cpp
--- a/Engine/Plugins/Experimental/PCG/Source/PCG/Private/Metadata/PCGAttributePropertySelector.cpp
+++ b/Engine/Plugins/Experimental/PCG/Source/PCG/Private/Metadata/PCGAttributePropertySelector.cpp
@@ -1,14 +1,19 @@
 // Copyright Epic Games, Inc. All Rights Reserved.
 
 #include "Metadata/PCGAttributePropertySelector.h"
+
+#include "PCGCustomVersion.h"
+#include "PCGData.h"
+#include "Helpers/PCGMetadataHelpers.h"
+#include "Metadata/PCGMetadata.h"
 #include "Metadata/PCGMetadataAttribute.h"
 
 namespace PCGAttributePropertySelectorConstants
 {
-	const TCHAR* PropertyPrefix = TEXT("$");
-	const TCHAR* ExtraSeparator = TEXT(".");
-	const TCHAR PropertyPrefixChar = PropertyPrefix[0];
-	const TCHAR ExtraSeparatorChar = ExtraSeparator[0];
+	static const TCHAR* PropertyPrefix = TEXT("$");
+	static const TCHAR* ExtraSeparator = TEXT(".");
+	static const TCHAR PropertyPrefixChar = PropertyPrefix[0];
+	static const TCHAR ExtraSeparatorChar = ExtraSeparator[0];
 }
 
 FName FPCGAttributePropertySelector::GetName() const
@@ -38,12 +43,16 @@ FName FPCGAttributePropertySelector::GetName() const
 			return NAME_None;
 		}
 	}
+	case EPCGAttributePropertySelection::Attribute:
+	{
+		return GetAttributeName();
+	}
 	default:
-		return AttributeName;
+		return NAME_None;
 	}
 }
 
-bool FPCGAttributePropertySelector::SetPointProperty(EPCGPointProperties InPointProperty, bool bResetExtraNames)
+bool FPCGAttributePropertySelector::SetAttributeName(FName InAttributeName, bool bResetExtraNames)
 {
 	const bool bHasExtraNames = !ExtraNames.IsEmpty();
 	if (bResetExtraNames)
@@ -51,20 +60,20 @@ bool FPCGAttributePropertySelector::SetPointProperty(EPCGPointProperties InPoint
 		ExtraNames.Empty();
 	}
 
-	if (Selection == EPCGAttributePropertySelection::PointProperty && InPointProperty == PointProperty)
+	if (Selection == EPCGAttributePropertySelection::Attribute && GetAttributeName() == InAttributeName)
 	{
 		// Nothing changed, except perhaps the extra names
 		return (bHasExtraNames && bResetExtraNames);
 	}
 	else
 	{
-		Selection = EPCGAttributePropertySelection::PointProperty;
-		PointProperty = InPointProperty;
+		Selection = EPCGAttributePropertySelection::Attribute;
+		AttributeName = InAttributeName;
 		return true;
 	}
 }
 
-bool FPCGAttributePropertySelector::SetAttributeName(FName InAttributeName, bool bResetExtraNames)
+bool FPCGAttributePropertySelector::SetPointProperty(EPCGPointProperties InPointProperty, bool bResetExtraNames)
 {
 	const bool bHasExtraNames = !ExtraNames.IsEmpty();
 	if (bResetExtraNames)
@@ -72,15 +81,15 @@ bool FPCGAttributePropertySelector::SetAttributeName(FName InAttributeName, bool
 		ExtraNames.Empty();
 	}
 
-	if (Selection == EPCGAttributePropertySelection::Attribute && AttributeName == InAttributeName)
+	if (Selection == EPCGAttributePropertySelection::PointProperty && InPointProperty == PointProperty)
 	{
 		// Nothing changed, except perhaps the extra names
 		return (bHasExtraNames && bResetExtraNames);
 	}
 	else
 	{
-		Selection = EPCGAttributePropertySelection::Attribute;
-		AttributeName = InAttributeName;
+		Selection = EPCGAttributePropertySelection::PointProperty;
+		PointProperty = InPointProperty;
 		return true;
 	}
 }
@@ -132,9 +141,61 @@ FText FPCGAttributePropertySelector::GetDisplayText() const
 	return FText::FromString(Res);
 }
 
+
+bool FPCGAttributePropertySelector::operator==(const FPCGAttributePropertySelector& Other) const
+{
+	return IsSame(Other);
+}
+
+bool FPCGAttributePropertySelector::IsSame(const FPCGAttributePropertySelector& Other, bool bIncludeExtraNames) const
+{
+	if (Selection != Other.Selection || (bIncludeExtraNames && ExtraNames != Other.ExtraNames))
+	{
+		return false;
+	}
+
+	switch (Selection)
+	{
+	case EPCGAttributePropertySelection::Attribute:
+		return GetAttributeName() == Other.GetAttributeName();
+	case EPCGAttributePropertySelection::PointProperty:
+		return PointProperty == Other.PointProperty;
+	case EPCGAttributePropertySelection::ExtraProperty:
+		return ExtraProperty == Other.ExtraProperty;
+	default:
+		return false;
+	}
+}
+
+void FPCGAttributePropertySelector::ImportFromOtherSelector(const FPCGAttributePropertySelector& InOther)
+{
+	Selection = InOther.GetSelection();
+
+	switch (Selection)
+	{
+	case EPCGAttributePropertySelection::Attribute:
+		SetAttributeName(InOther.GetAttributeName());
+		break;
+	case EPCGAttributePropertySelection::PointProperty:
+		SetPointProperty(InOther.GetPointProperty());
+		break;
+	case EPCGAttributePropertySelection::ExtraProperty:
+		SetExtraProperty(InOther.GetExtraProperty());
+		break;
+	}
+
+	ExtraNames = InOther.GetExtraNames();
+}
+
 bool FPCGAttributePropertySelector::IsValid() const
 {
-	return (Selection != EPCGAttributePropertySelection::Attribute) || FPCGMetadataAttributeBase::IsValidName(AttributeName);
+	const FName ThisAttributeName = GetAttributeName();
+	return (Selection != EPCGAttributePropertySelection::Attribute) || 
+		ThisAttributeName == PCGMetadataAttributeConstants::LastAttributeName ||
+		ThisAttributeName == PCGMetadataAttributeConstants::LastCreatedAttributeName ||
+		ThisAttributeName == PCGMetadataAttributeConstants::SourceAttributeName ||
+		ThisAttributeName == PCGMetadataAttributeConstants::SourceNameAttributeName ||
+		FPCGMetadataAttributeBase::IsValidName(ThisAttributeName);
 }
 
 bool FPCGAttributePropertySelector::Update(FString NewValue)
@@ -188,6 +249,103 @@ bool FPCGAttributePropertySelector::Update(FString NewValue)
 
 ///////////////////////////////////////////////////////////////////////
 
+FPCGAttributePropertyInputSelector::FPCGAttributePropertyInputSelector()
+{
+	AttributeName = PCGMetadataAttributeConstants::LastAttributeName;
+}
+
+FPCGAttributePropertyInputSelector FPCGAttributePropertyInputSelector::CopyAndFixLast(const UPCGData* InData) const
+{
+	if (Selection == EPCGAttributePropertySelection::Attribute)
+	{
+		if (AttributeName == PCGMetadataAttributeConstants::LastAttributeName && InData && InData->HasCachedLastSelector())
+		{
+			return InData->GetCachedLastSelector();
+		}
+		else if (AttributeName == PCGMetadataAttributeConstants::LastCreatedAttributeName && InData)
+		{
+			if (const UPCGMetadata* Metadata = PCGMetadataHelpers::GetConstMetadata(InData))
+			{
+				FPCGAttributePropertyInputSelector Selector;
+				Selector.SetAttributeName(Metadata->GetLatestAttributeNameOrNone());
+				return Selector;
+			}
+		}
+	}
+
+	return *this;
+}
+
+bool FPCGAttributePropertyInputSelector::SerializeFromMismatchedTag(const struct FPropertyTag& Tag, FStructuredArchive::FSlot Slot)
+{
+	if (Tag.Type == NAME_StructProperty
+		&& Tag.StructName == FPCGAttributePropertySelector::StaticStruct()->GetFName())
+	{
+		FPCGAttributePropertyInputSelector::StaticStruct()->SerializeItem(Slot, this, nullptr);
+		return true;
+	}
+
+	return false;
+}
+
+void FPCGAttributePropertyInputSelector::ApplyDeprecation(int32 InPCGCustomVersion)
+{
+	if ((InPCGCustomVersion < FPCGCustomVersion::UpdateAttributePropertyInputSelector) && (Selection == EPCGAttributePropertySelection::Attribute) && (AttributeName == PCGMetadataAttributeConstants::LastAttributeName))
+	{
+		AttributeName = PCGMetadataAttributeConstants::LastCreatedAttributeName;
+	}
+}
+
+///////////////////////////////////////////////////////////////////////
+
+FPCGAttributePropertyOutputSelector::FPCGAttributePropertyOutputSelector()
+{
+	AttributeName = PCGMetadataAttributeConstants::SourceAttributeName;
+}
+
+FPCGAttributePropertyOutputSelector FPCGAttributePropertyOutputSelector::CopyAndFixSource(const FPCGAttributePropertyInputSelector* InSourceSelector, const UPCGData* InOptionalData) const
+{
+	if (Selection == EPCGAttributePropertySelection::Attribute)
+	{
+		if (AttributeName == PCGMetadataAttributeConstants::SourceAttributeName && InSourceSelector)
+		{
+			return FPCGAttributePropertySelector::CreateFromOtherSelector<FPCGAttributePropertyOutputSelector>(*InSourceSelector);
+		}
+		else if (AttributeName == PCGMetadataAttributeConstants::SourceNameAttributeName && InSourceSelector)
+		{
+			FPCGAttributePropertyOutputSelector Selector;
+			Selector.SetAttributeName(InSourceSelector->GetName());
+			return Selector;
+		}
+		// Only for deprecation
+		else if (AttributeName == PCGMetadataAttributeConstants::LastCreatedAttributeName && InOptionalData)
+		{
+			if (const UPCGMetadata* Metadata = PCGMetadataHelpers::GetConstMetadata(InOptionalData))
+			{
+				FPCGAttributePropertyOutputSelector Selector;
+				Selector.SetAttributeName(Metadata->GetLatestAttributeNameOrNone());
+				return Selector;
+			}
+		}
+	}
+
+	return *this;
+}
+
+bool FPCGAttributePropertyOutputSelector::SerializeFromMismatchedTag(const struct FPropertyTag& Tag, FStructuredArchive::FSlot Slot)
+{
+	if (Tag.Type == NAME_StructProperty
+		&& Tag.StructName == FPCGAttributePropertySelector::StaticStruct()->GetFName())
+	{
+		FPCGAttributePropertyOutputSelector::StaticStruct()->SerializeItem(Slot, this, nullptr);
+		return true;
+	}
+
+	return false;
+}
+
+///////////////////////////////////////////////////////////////////////
+
 bool UPCGAttributePropertySelectorBlueprintHelpers::SetPointProperty(FPCGAttributePropertySelector& Selector, EPCGPointProperties InPointProperty)
 {
 	return Selector.SetPointProperty(InPointProperty);
@@ -203,7 +361,42 @@ bool UPCGAttributePropertySelectorBlueprintHelpers::SetExtraProperty(FPCGAttribu
 	return Selector.SetExtraProperty(InExtraProperty);
 }
 
+EPCGAttributePropertySelection UPCGAttributePropertySelectorBlueprintHelpers::GetSelection(const FPCGAttributePropertySelector& Selector)
+{
+	return Selector.GetSelection();
+}
+
+EPCGPointProperties UPCGAttributePropertySelectorBlueprintHelpers::GetPointProperty(const FPCGAttributePropertySelector& Selector)
+{
+	return Selector.GetPointProperty();
+}
+
+FName UPCGAttributePropertySelectorBlueprintHelpers::GetAttributeName(const FPCGAttributePropertySelector& Selector)
+{
+	return Selector.GetAttributeName();
+}
+
+EPCGExtraProperties UPCGAttributePropertySelectorBlueprintHelpers::GetExtraProperty(const FPCGAttributePropertySelector& Selector)
+{
+	return Selector.GetExtraProperty();
+}
+
+const TArray<FString>& UPCGAttributePropertySelectorBlueprintHelpers::GetExtraNames(const FPCGAttributePropertySelector& Selector)
+{
+	return Selector.GetExtraNames();
+}
+
 FName UPCGAttributePropertySelectorBlueprintHelpers::GetName(const FPCGAttributePropertySelector& Selector)
 {
 	return Selector.GetName();
+}
+
+FPCGAttributePropertyInputSelector UPCGAttributePropertySelectorBlueprintHelpers::CopyAndFixLast(const FPCGAttributePropertyInputSelector& Selector, const UPCGData* InData)
+{
+	return Selector.CopyAndFixLast(InData);
+}
+
+FPCGAttributePropertyOutputSelector UPCGAttributePropertySelectorBlueprintHelpers::CopyAndFixSource(const FPCGAttributePropertyOutputSelector& Selector, const FPCGAttributePropertyInputSelector& InSelector)
+{
+	return Selector.CopyAndFixSource(&InSelector);
 }
\ No newline at end of file
diff --git a/Engine/Plugins/Experimental/PCG/Source/PCG/Private/Metadata/PCGMetadata.cpp b/Engine/Plugins/Experimental/PCG/Source/PCG/Private/Metadata/PCGMetadata.cpp
--- a/Engine/Plugins/Experimental/PCG/Source/PCG/Private/Metadata/PCGMetadata.cpp
+++ b/Engine/Plugins/Experimental/PCG/Source/PCG/Private/Metadata/PCGMetadata.cpp
@@ -2,8 +2,10 @@
 
 #include "Metadata/PCGMetadata.h"
 
-#include "Helpers/PCGPropertyHelpers.h"
+#include "PCGData.h"
 #include "PCGPoint.h"
+#include "Helpers/PCGPropertyHelpers.h"
+#include "Metadata/PCGAttributePropertySelector.h"
 
 #include UE_INLINE_GENERATED_CPP_BY_NAME(PCGMetadata)
 
@@ -328,6 +330,9 @@ FPCGMetadataAttributeBase* UPCGMetadata::GetMutableAttribute(FName AttributeName
 	if (FPCGMetadataAttributeBase** FoundAttribute = Attributes.Find(AttributeName))
 	{
 		Attribute = *FoundAttribute;
+
+		// Also when accessing an attribute, notify the PCG Data owner that the latest attribute manipulated is this one.
+		SetLastCachedSelectorOnOwner(AttributeName);
 	}
 	AttributeLock.ReadUnlock();
 
@@ -655,6 +660,10 @@ FPCGMetadataAttributeBase* UPCGMetadata::CopyAttribute(const FPCGMetadataAttribu
 		AttributeLock.WriteLock();
 		NewAttribute->AttributeId = NextAttributeId++;
 		AddAttributeInternal(NewAttributeName, NewAttribute);
+
+		// Also when creating an attribute, notify the PCG Data owner that the latest attribute manipulated is this one.
+		SetLastCachedSelectorOnOwner(NewAttributeName);
+
 		AttributeLock.WriteUnlock();
 	}
 
@@ -1146,3 +1155,12 @@ void UPCGMetadata::AccumulatePointWeightedAttributes(const FPCGPoint& InPoint, c
 	AccumulateWeightedAttributes(InPoint.MetadataEntry, InMetadata, Weight, bSetNonInterpolableAttributes, OutPoint.MetadataEntry);
 }
 
+void UPCGMetadata::SetLastCachedSelectorOnOwner(FName AttributeName)
+{
+	if (UPCGData* OwnerData = Cast<UPCGData>(GetOuter()))
+	{
+		FPCGAttributePropertyInputSelector Selector;
+		Selector.SetAttributeName(AttributeName);
+		OwnerData->SetLastSelector(Selector);
+	}
+}
\ No newline at end of file
diff --git a/Engine/Plugins/Experimental/PCG/Source/PCG/Private/PCGParamData.cpp b/Engine/Plugins/Experimental/PCG/Source/PCG/Private/PCGParamData.cpp
--- a/Engine/Plugins/Experimental/PCG/Source/PCG/Private/PCGParamData.cpp
+++ b/Engine/Plugins/Experimental/PCG/Source/PCG/Private/PCGParamData.cpp
@@ -38,7 +38,7 @@ void UPCGParamData::AddToCrc(FArchiveCrc32& Ar, bool bFullDataCrc) const
 	// Add all attribute values
 	for (FName AttributeName : AttributeNames)
 	{
-		FPCGAttributePropertySelector InputSource;
+		FPCGAttributePropertyInputSelector InputSource;
 		InputSource.SetAttributeName(AttributeName);
 
 		TUniquePtr<const IPCGAttributeAccessor> InputAccessor = PCGAttributeAccessorHelpers::CreateConstAccessor(this, InputSource);
@@ -121,3 +121,42 @@ UPCGParamData* UPCGParamData::FilterParamsByKey(int64 InKey) const
 
 	return NewParams;
 }
+
+bool UPCGParamData::HasCachedLastSelector() const
+{
+	return bHasCachedLastSelector || (Metadata && Metadata->GetAttributeCount() > 0);
+}
+
+FPCGAttributePropertyInputSelector UPCGParamData::GetCachedLastSelector() const
+{
+	if (bHasCachedLastSelector)
+	{
+		return CachedLastSelector;
+	}
+
+	FPCGAttributePropertyInputSelector TempSelector{};
+
+	// If we have attribute and no last selector, create a cached last selector on the latest attribute, to catch "CreateAttribute" calls that didn't use accessors.
+	if (Metadata && Metadata->GetAttributeCount() > 0)
+	{
+		TempSelector.SetAttributeName(Metadata->GetLatestAttributeNameOrNone());
+	}
+
+	return TempSelector;
+}
+
+void UPCGParamData::SetLastSelector(const FPCGAttributePropertySelector& InSelector)
+{
+	// Check that it is not a not Attribute selector or Last/Source selector
+	if (InSelector.GetSelection() != EPCGAttributePropertySelection::Attribute 
+		|| InSelector.GetAttributeName() == PCGMetadataAttributeConstants::LastAttributeName
+		|| InSelector.GetAttributeName() == PCGMetadataAttributeConstants::LastCreatedAttributeName
+		|| InSelector.GetAttributeName() == PCGMetadataAttributeConstants::SourceAttributeName
+		|| InSelector.GetAttributeName() == PCGMetadataAttributeConstants::SourceNameAttributeName)
+	{
+		return;
+	}
+
+	bHasCachedLastSelector = true;
+	CachedLastSelector.ImportFromOtherSelector(InSelector);
+}
\ No newline at end of file
diff --git a/Engine/Plugins/Experimental/PCG/Source/PCG/Private/PCGSettings.cpp b/Engine/Plugins/Experimental/PCG/Source/PCG/Private/PCGSettings.cpp
--- a/Engine/Plugins/Experimental/PCG/Source/PCG/Private/PCGSettings.cpp
+++ b/Engine/Plugins/Experimental/PCG/Source/PCG/Private/PCGSettings.cpp
@@ -115,6 +115,41 @@ bool UPCGSettings::operator==(const UPCGSettings& Other) const
 #if WITH_EDITOR
 void UPCGSettings::ApplyDeprecation(UPCGNode* InOutNode)
 {
+	// For versions older than the update of input selector, look for any Input Selector in the properties, and if we find one that is default,
+	// we force to @LastCreated to keep the old behavior.
+	if (DataVersion < FPCGCustomVersion::UpdateAttributePropertyInputSelector)
+	{
+		auto Recurse = [this](UStruct* InStructClass, void* InContainer)
+		{
+			auto RecurseImpl = [this](UStruct* InStructClass, void* InContainer, auto Callback) -> void
+			{
+				for (TFieldIterator<FStructProperty> It(InStructClass, EFieldIterationFlags::IncludeSuper); It; ++It)
+				{
+					const FStructProperty* Property = CastField<FStructProperty>(*It);
+					if (!Property)
+					{
+						continue;
+					}
+
+					// If it is an Input Selector, apply deprecation
+					if (Property->Struct == FPCGAttributePropertyInputSelector::StaticStruct())
+					{
+						Property->ContainerPtrToValuePtr<FPCGAttributePropertyInputSelector>(InContainer)->ApplyDeprecation(DataVersion);
+					}
+					else
+					{
+						// Otherwise, go deeper.
+						Callback(Property->Struct.Get(), Property->ContainerPtrToValuePtr<void>(InContainer), Callback);
+					}
+				}
+			};
+
+			RecurseImpl(InStructClass, InContainer, RecurseImpl);
+		};
+
+		Recurse(GetClass(), this);
+	}
+
 	DataVersion = FPCGCustomVersion::LatestVersion;
 }
 
diff --git a/Engine/Plugins/Experimental/PCG/Source/PCG/Private/Tests/Elements/Metadata/PCGMetadataBooleanOpTest.cpp b/Engine/Plugins/Experimental/PCG/Source/PCG/Private/Tests/Elements/Metadata/PCGMetadataBooleanOpTest.cpp
--- a/Engine/Plugins/Experimental/PCG/Source/PCG/Private/Tests/Elements/Metadata/PCGMetadataBooleanOpTest.cpp
+++ b/Engine/Plugins/Experimental/PCG/Source/PCG/Private/Tests/Elements/Metadata/PCGMetadataBooleanOpTest.cpp
@@ -32,8 +32,7 @@ bool FPCGMetadataBooleanOpTest::RunTest(const FString& Parameters)
 	const FName InvalidAttribute = TEXT("Invalid");
 	const FName OutputAttributeName = TEXT("Output");
 
-	Settings->OutputTarget.Selection = EPCGAttributePropertySelection::Attribute;
-	Settings->OutputTarget.AttributeName = OutputAttributeName;
+	Settings->OutputTarget.SetAttributeName(OutputAttributeName);
 	Settings->ForceOutputConnections[0] = true;
 
 	const bool bAllowInterpolation = false;
@@ -118,16 +117,14 @@ bool FPCGMetadataBooleanOpTest::RunTest(const FString& Parameters)
 	ParamTaggedData1.Pin = PCGPinConstants::DefaultInputLabel;
 
 	Settings->Operation = EPCGMedadataBooleanOperation::Not;
-	Settings->InputSource1.Selection = EPCGAttributePropertySelection::Attribute;
-	Settings->InputSource2.Selection = EPCGAttributePropertySelection::Attribute;
 
 	{
-		Settings->InputSource1.AttributeName = TrueAttribute;
+		Settings->InputSource1.SetAttributeName(TrueAttribute);
 		bTestPassed &= ValidateOp(/*bExpectedResult=*/ false);
 	}
 
 	{
-		Settings->InputSource1.AttributeName = FalseAttribute;
+		Settings->InputSource1.SetAttributeName(FalseAttribute);
 		bTestPassed &= ValidateOp(true);
 	}
 
diff --git a/Engine/Plugins/Experimental/PCG/Source/PCG/Private/Tests/Elements/Metadata/PCGMetadataBreakTransformTest.cpp b/Engine/Plugins/Experimental/PCG/Source/PCG/Private/Tests/Elements/Metadata/PCGMetadataBreakTransformTest.cpp
--- a/Engine/Plugins/Experimental/PCG/Source/PCG/Private/Tests/Elements/Metadata/PCGMetadataBreakTransformTest.cpp
+++ b/Engine/Plugins/Experimental/PCG/Source/PCG/Private/Tests/Elements/Metadata/PCGMetadataBreakTransformTest.cpp
@@ -97,8 +97,7 @@ bool FPCGMetadataBreakTransformTest::RunTest(const FString& Parameters)
 		const FName OutputAttributeName = TEXT("Output");
 
 		UPCGMetadataBreakTransformSettings* Settings = CastChecked<UPCGMetadataBreakTransformSettings>(TestData.Settings);
-		Settings->OutputTarget.Selection = EPCGAttributePropertySelection::Attribute;
-		Settings->OutputTarget.AttributeName = OutputAttributeName;
+		Settings->OutputTarget.SetAttributeName(OutputAttributeName);
 		FPCGElementPtr MetadataBreakTransformElement = TestData.Settings->GetElement();
 
 		TUniquePtr<FPCGContext> Context = TestData.InitializeTestContext();
diff --git a/Engine/Plugins/Experimental/PCG/Source/PCG/Private/Tests/Elements/Metadata/PCGMetadataBreakVectorTest.cpp b/Engine/Plugins/Experimental/PCG/Source/PCG/Private/Tests/Elements/Metadata/PCGMetadataBreakVectorTest.cpp
--- a/Engine/Plugins/Experimental/PCG/Source/PCG/Private/Tests/Elements/Metadata/PCGMetadataBreakVectorTest.cpp
+++ b/Engine/Plugins/Experimental/PCG/Source/PCG/Private/Tests/Elements/Metadata/PCGMetadataBreakVectorTest.cpp
@@ -115,8 +115,7 @@ bool FPCGMetadataBreakVectorTest::RunTest(const FString& Parameters)
 		const FName OutputAttributeName = TEXT("Output");
 
 		UPCGMetadataBreakVectorSettings* Settings = CastChecked<UPCGMetadataBreakVectorSettings>(TestData.Settings);
-		Settings->OutputTarget.Selection = EPCGAttributePropertySelection::Attribute;
-		Settings->OutputTarget.AttributeName = OutputAttributeName;
+		Settings->OutputTarget.SetAttributeName(OutputAttributeName);
 		FPCGElementPtr MetadataBreakVectorElement = TestData.Settings->GetElement();
 
 		TUniquePtr<FPCGContext> Context = TestData.InitializeTestContext();
diff --git a/Engine/Plugins/Experimental/PCG/Source/PCG/Private/Tests/Elements/Metadata/PCGMetadataMakeTransformTest.cpp b/Engine/Plugins/Experimental/PCG/Source/PCG/Private/Tests/Elements/Metadata/PCGMetadataMakeTransformTest.cpp
--- a/Engine/Plugins/Experimental/PCG/Source/PCG/Private/Tests/Elements/Metadata/PCGMetadataMakeTransformTest.cpp
+++ b/Engine/Plugins/Experimental/PCG/Source/PCG/Private/Tests/Elements/Metadata/PCGMetadataMakeTransformTest.cpp
@@ -93,8 +93,7 @@ bool FPCGMetadataMakeTransformTest::RunTest(const FString& Parameters)
 		const FName OutputAttributeName = TEXT("Output");
 
 		UPCGMetadataMakeTransformSettings* Settings = CastChecked<UPCGMetadataMakeTransformSettings>(TestData.Settings);
-		Settings->OutputTarget.Selection = EPCGAttributePropertySelection::Attribute;
-		Settings->OutputTarget.AttributeName = OutputAttributeName;
+		Settings->OutputTarget.SetAttributeName(OutputAttributeName);
 		FPCGElementPtr MetadataBreakTransformElement = TestData.Settings->GetElement();
 
 		TUniquePtr<FPCGContext> Context = TestData.InitializeTestContext();
diff --git a/Engine/Plugins/Experimental/PCG/Source/PCG/Private/Tests/Elements/PCGSpatialNoiseTest.cpp b/Engine/Plugins/Experimental/PCG/Source/PCG/Private/Tests/Elements/PCGSpatialNoiseTest.cpp
--- a/Engine/Plugins/Experimental/PCG/Source/PCG/Private/Tests/Elements/PCGSpatialNoiseTest.cpp
+++ b/Engine/Plugins/Experimental/PCG/Source/PCG/Private/Tests/Elements/PCGSpatialNoiseTest.cpp
@@ -48,7 +48,7 @@ bool FPCGSpatialNoise_Perlin2D::RunTest(const FString& Parameters)
 	UPCGSpatialNoiseSettings* Settings = CastChecked<UPCGSpatialNoiseSettings>(TestData.Settings);
 
 	Settings->Mode = PCGSpatialNoiseMode::Perlin2D;
-	Settings->ValueTarget.AttributeName = TEXT("Noise");
+	Settings->ValueTarget.SetAttributeName(TEXT("Noise"));
 
 	FPCGElementPtr TestElement = TestData.Settings->GetElement();
 
@@ -71,7 +71,7 @@ bool FPCGSpatialNoise_Perlin2D::RunTest(const FString& Parameters)
 	UTEST_NOT_NULL("Output data", OutputData);
 	UTEST_EQUAL("Output point data count", OutputData->GetPoints().Num(), 100);
 
-	const FPCGMetadataAttribute<double>* NoiseAttribute = OutputData->Metadata->GetConstTypedAttribute<double>(Settings->ValueTarget.AttributeName);
+	const FPCGMetadataAttribute<double>* NoiseAttribute = OutputData->Metadata->GetConstTypedAttribute<double>(Settings->ValueTarget.GetAttributeName());
 	UTEST_NOT_NULL("Noise Attribute Created", NoiseAttribute);
 
 	return true;
@@ -86,7 +86,7 @@ bool FPCGSpatialNoise_Caustic2D::RunTest(const FString& Parameters)
 	UPCGSpatialNoiseSettings* Settings = CastChecked<UPCGSpatialNoiseSettings>(TestData.Settings);
 
 	Settings->Mode = PCGSpatialNoiseMode::Caustic2D;
-	Settings->ValueTarget.AttributeName = TEXT("Noise");
+	Settings->ValueTarget.SetAttributeName(TEXT("Noise"));
 
 	FPCGElementPtr TestElement = TestData.Settings->GetElement();
 
@@ -109,7 +109,7 @@ bool FPCGSpatialNoise_Caustic2D::RunTest(const FString& Parameters)
 	UTEST_NOT_NULL("Output data", OutputData);
 	UTEST_EQUAL("Output point data count", OutputData->GetPoints().Num(), 100);
 
-	const FPCGMetadataAttribute<double>* NoiseAttribute = OutputData->Metadata->GetConstTypedAttribute<double>(Settings->ValueTarget.AttributeName);
+	const FPCGMetadataAttribute<double>* NoiseAttribute = OutputData->Metadata->GetConstTypedAttribute<double>(Settings->ValueTarget.GetAttributeName());
 	UTEST_NOT_NULL("Noise Attribute Created", NoiseAttribute);
 
 	return true;
@@ -124,8 +124,8 @@ bool FPCGSpatialNoise_Voronoi2D::RunTest(const FString& Parameters)
 	UPCGSpatialNoiseSettings* Settings = CastChecked<UPCGSpatialNoiseSettings>(TestData.Settings);
 
 	Settings->Mode = PCGSpatialNoiseMode::Voronoi2D;
-	Settings->ValueTarget.AttributeName = TEXT("Distance");
-	Settings->VoronoiCellIDTarget.AttributeName = TEXT("CellID");
+	Settings->ValueTarget.SetAttributeName(TEXT("Distance"));
+	Settings->VoronoiCellIDTarget.SetAttributeName(TEXT("CellID"));
 
 	FPCGElementPtr TestElement = TestData.Settings->GetElement();
 
@@ -148,10 +148,10 @@ bool FPCGSpatialNoise_Voronoi2D::RunTest(const FString& Parameters)
 	UTEST_NOT_NULL("Output data", OutputData);
 	UTEST_EQUAL("Output point data count", OutputData->GetPoints().Num(), 100);
 
-	const FPCGMetadataAttribute<double>* NoiseAttribute = OutputData->Metadata->GetConstTypedAttribute<double>(Settings->ValueTarget.AttributeName);
+	const FPCGMetadataAttribute<double>* NoiseAttribute = OutputData->Metadata->GetConstTypedAttribute<double>(Settings->ValueTarget.GetAttributeName());
 	UTEST_NOT_NULL("Distance Attribute Created", NoiseAttribute);
 
-	const FPCGMetadataAttribute<double>* CellIDAttribute = OutputData->Metadata->GetConstTypedAttribute<double>(Settings->VoronoiCellIDTarget.AttributeName);
+	const FPCGMetadataAttribute<double>* CellIDAttribute = OutputData->Metadata->GetConstTypedAttribute<double>(Settings->VoronoiCellIDTarget.GetAttributeName());
 	UTEST_NOT_NULL("Cell ID Attribute Created", CellIDAttribute);
 
 
@@ -167,7 +167,7 @@ bool FPCGSpatialNoise_FractionalBrownian2D::RunTest(const FString& Parameters)
 	UPCGSpatialNoiseSettings* Settings = CastChecked<UPCGSpatialNoiseSettings>(TestData.Settings);
 
 	Settings->Mode = PCGSpatialNoiseMode::FractionalBrownian2D;
-	Settings->ValueTarget.AttributeName = TEXT("Noise");
+	Settings->ValueTarget.SetAttributeName(TEXT("Noise"));
 
 	FPCGElementPtr TestElement = TestData.Settings->GetElement();
 
@@ -190,7 +190,7 @@ bool FPCGSpatialNoise_FractionalBrownian2D::RunTest(const FString& Parameters)
 	UTEST_NOT_NULL("Output data", OutputData);
 	UTEST_EQUAL("Output point data count", OutputData->GetPoints().Num(), 100);
 
-	const FPCGMetadataAttribute<double>* NoiseAttribute = OutputData->Metadata->GetConstTypedAttribute<double>(Settings->ValueTarget.AttributeName);
+	const FPCGMetadataAttribute<double>* NoiseAttribute = OutputData->Metadata->GetConstTypedAttribute<double>(Settings->ValueTarget.GetAttributeName());
 	UTEST_NOT_NULL("Noise Attribute Created", NoiseAttribute);
 
 	return true;
@@ -205,7 +205,7 @@ bool FPCGSpatialNoise_EdgeMask2D::RunTest(const FString& Parameters)
 	UPCGSpatialNoiseSettings* Settings = CastChecked<UPCGSpatialNoiseSettings>(TestData.Settings);
 
 	Settings->Mode = PCGSpatialNoiseMode::EdgeMask2D;
-	Settings->ValueTarget.AttributeName = TEXT("Noise");
+	Settings->ValueTarget.SetAttributeName(TEXT("Noise"));
 
 	FPCGElementPtr TestElement = TestData.Settings->GetElement();
 
@@ -228,7 +228,7 @@ bool FPCGSpatialNoise_EdgeMask2D::RunTest(const FString& Parameters)
 	UTEST_NOT_NULL("Output data", OutputData);
 	UTEST_EQUAL("Output point data count", OutputData->GetPoints().Num(), 100);
 
-	const FPCGMetadataAttribute<double>* NoiseAttribute = OutputData->Metadata->GetConstTypedAttribute<double>(Settings->ValueTarget.AttributeName);
+	const FPCGMetadataAttribute<double>* NoiseAttribute = OutputData->Metadata->GetConstTypedAttribute<double>(Settings->ValueTarget.GetAttributeName());
 	UTEST_NOT_NULL("Noise Attribute Created", NoiseAttribute);
 
 	return true;
@@ -243,7 +243,7 @@ bool FPCGSpatialNoise_TilingPerlin2D::RunTest(const FString& Parameters)
 
 	Settings->Mode = PCGSpatialNoiseMode::Perlin2D;
 	Settings->bTiling = true;
-	Settings->ValueTarget.AttributeName = TEXT("Noise");
+	Settings->ValueTarget.SetAttributeName(TEXT("Noise"));
 
 	FPCGElementPtr TestElement = TestData.Settings->GetElement();
 
@@ -266,7 +266,7 @@ bool FPCGSpatialNoise_TilingPerlin2D::RunTest(const FString& Parameters)
 	UTEST_NOT_NULL("Output data", OutputData);
 	UTEST_EQUAL("Output point data count", OutputData->GetPoints().Num(), 100);
 
-	const FPCGMetadataAttribute<double>* NoiseAttribute = OutputData->Metadata->GetConstTypedAttribute<double>(Settings->ValueTarget.AttributeName);
+	const FPCGMetadataAttribute<double>* NoiseAttribute = OutputData->Metadata->GetConstTypedAttribute<double>(Settings->ValueTarget.GetAttributeName());
 	UTEST_NOT_NULL("Noise Attribute Created", NoiseAttribute);
 
 	return true;
@@ -281,8 +281,8 @@ bool FPCGSpatialNoise_TilingVoronoi2D::RunTest(const FString& Parameters)
 
 	Settings->Mode = PCGSpatialNoiseMode::Voronoi2D;
 	Settings->bTiling = true;
-	Settings->ValueTarget.AttributeName = TEXT("Distance");
-	Settings->VoronoiCellIDTarget.AttributeName = TEXT("CellID");
+	Settings->ValueTarget.SetAttributeName(TEXT("Distance"));
+	Settings->VoronoiCellIDTarget.SetAttributeName(TEXT("CellID"));
 
 	FPCGElementPtr TestElement = TestData.Settings->GetElement();
 
@@ -305,10 +305,10 @@ bool FPCGSpatialNoise_TilingVoronoi2D::RunTest(const FString& Parameters)
 	UTEST_NOT_NULL("Output data", OutputData);
 	UTEST_EQUAL("Output point data count", OutputData->GetPoints().Num(), 100);
 
-	const FPCGMetadataAttribute<double>* NoiseAttribute = OutputData->Metadata->GetConstTypedAttribute<double>(Settings->ValueTarget.AttributeName);
+	const FPCGMetadataAttribute<double>* NoiseAttribute = OutputData->Metadata->GetConstTypedAttribute<double>(Settings->ValueTarget.GetAttributeName());
 	UTEST_NOT_NULL("Distance Attribute Created", NoiseAttribute);
 
-	const FPCGMetadataAttribute<double>* CellIDAttribute = OutputData->Metadata->GetConstTypedAttribute<double>(Settings->VoronoiCellIDTarget.AttributeName);
+	const FPCGMetadataAttribute<double>* CellIDAttribute = OutputData->Metadata->GetConstTypedAttribute<double>(Settings->VoronoiCellIDTarget.GetAttributeName());
 	UTEST_NOT_NULL("Cell ID Attribute Created", CellIDAttribute);
 
 	return true;
diff --git a/Engine/Plugins/Experimental/PCG/Source/PCG/Public/Data/PCGSpatialData.h b/Engine/Plugins/Experimental/PCG/Source/PCG/Public/Data/PCGSpatialData.h
--- a/Engine/Plugins/Experimental/PCG/Source/PCG/Public/Data/PCGSpatialData.h
+++ b/Engine/Plugins/Experimental/PCG/Source/PCG/Public/Data/PCGSpatialData.h
@@ -5,6 +5,7 @@
 #include "PCGData.h"
 #include "PCGPoint.h" // IWYU pragma: keep
 #include "Elements/PCGProjectionParams.h"
+#include "Metadata/PCGAttributePropertySelector.h"
 #include "Metadata/PCGMetadata.h"
 
 #include "PCGSpatialData.generated.h"
@@ -40,6 +41,10 @@ public:
 
 	// ~Begin UPCGData interface
 	virtual EPCGDataType GetDataType() const override { return EPCGDataType::Spatial; }
+
+	virtual bool HasCachedLastSelector() const override;
+	virtual FPCGAttributePropertyInputSelector GetCachedLastSelector() const override;
+	virtual void SetLastSelector(const FPCGAttributePropertySelector& InSelector) override;
 	// ~End UPCGData interface
 
 	/** Virtual call to allocate a new spacial data object, duplicate this spatial data into
@@ -150,6 +155,14 @@ public:
 
 protected:
 	virtual UPCGSpatialData* CopyInternal() const PURE_VIRTUAL(UPCGSpatialData::CopyInternal, return nullptr;);
+
+private:
+	/** Cache to keep track of the latest attribute manipulated on this data. */
+	UPROPERTY()
+	bool bHasCachedLastSelector = false;
+
+	UPROPERTY()
+	FPCGAttributePropertyInputSelector CachedLastSelector;
 };
 
 UCLASS(Abstract, ClassGroup = (Procedural))
diff --git a/Engine/Plugins/Experimental/PCG/Source/PCG/Public/Elements/IO/PCGExternalData.h b/Engine/Plugins/Experimental/PCG/Source/PCG/Public/Elements/IO/PCGExternalData.h
--- a/Engine/Plugins/Experimental/PCG/Source/PCG/Public/Elements/IO/PCGExternalData.h
+++ b/Engine/Plugins/Experimental/PCG/Source/PCG/Public/Elements/IO/PCGExternalData.h
@@ -31,7 +31,7 @@ protected:
 
 public:
 	UPROPERTY(BlueprintReadWrite, EditAnywhere, Category = "Settings")
-	TMap<FString, FPCGAttributePropertySelector> AttributeMapping;
+	TMap<FString, FPCGAttributePropertyInputSelector> AttributeMapping;
 };
 
 class PCG_API FPCGExternalDataElement : public IPCGElement
diff --git a/Engine/Plugins/Experimental/PCG/Source/PCG/Public/Elements/IO/PCGExternalDataContext.h b/Engine/Plugins/Experimental/PCG/Source/PCG/Public/Elements/IO/PCGExternalDataContext.h
--- a/Engine/Plugins/Experimental/PCG/Source/PCG/Public/Elements/IO/PCGExternalDataContext.h
+++ b/Engine/Plugins/Experimental/PCG/Source/PCG/Public/Elements/IO/PCGExternalDataContext.h
@@ -16,14 +16,14 @@ struct FPCGExternalDataContext : public FPCGContext
 	{
 		TUniquePtr<const IPCGAttributeAccessor> RowAccessor;
 		TUniquePtr<IPCGAttributeAccessor> PointAccessor;
-		FPCGAttributePropertySelector Selector;
+		FPCGAttributePropertyOutputSelector Selector;
 
 		FRowToPointAccessors() = default;
 		FRowToPointAccessors(const FRowToPointAccessors&) = delete;
 		FRowToPointAccessors(FRowToPointAccessors&&) = default;
 		FRowToPointAccessors& operator=(FRowToPointAccessors&&) = default;
 		FRowToPointAccessors& operator=(const FRowToPointAccessors&) = delete;
-		FRowToPointAccessors(TUniquePtr<const IPCGAttributeAccessor> InRow, TUniquePtr<IPCGAttributeAccessor> InPoint, const FPCGAttributePropertySelector& InSelector)
+		FRowToPointAccessors(TUniquePtr<const IPCGAttributeAccessor> InRow, TUniquePtr<IPCGAttributeAccessor> InPoint, const FPCGAttributePropertyOutputSelector& InSelector)
 			: RowAccessor(MoveTemp(InRow)), PointAccessor(MoveTemp(InPoint)), Selector(InSelector) {}
 	};
 
diff --git a/Engine/Plugins/Experimental/PCG/Source/PCG/Public/Elements/Metadata/PCGMetadataBitwiseOpElement.h b/Engine/Plugins/Experimental/PCG/Source/PCG/Public/Elements/Metadata/PCGMetadataBitwiseOpElement.h
--- a/Engine/Plugins/Experimental/PCG/Source/PCG/Public/Elements/Metadata/PCGMetadataBitwiseOpElement.h
+++ b/Engine/Plugins/Experimental/PCG/Source/PCG/Public/Elements/Metadata/PCGMetadataBitwiseOpElement.h
@@ -37,7 +37,7 @@ public:
 	//~End UPCGSettings interface
 
 	//~Begin UPCGMetadataSettingsBase interface
-	FPCGAttributePropertySelector GetInputSource(uint32 Index) const override;
+	FPCGAttributePropertyInputSelector GetInputSource(uint32 Index) const override;
 
 	virtual FName GetInputPinLabel(uint32 Index) const override;
 	virtual uint32 GetInputPinNum() const override;
@@ -56,10 +56,10 @@ public:
 	EPCGMedadataBitwiseOperation Operation = EPCGMedadataBitwiseOperation::And;
 
 	UPROPERTY(BlueprintReadWrite, EditAnywhere, Category = Input)
-	FPCGAttributePropertySelector InputSource1;
+	FPCGAttributePropertyInputSelector InputSource1;
 
 	UPROPERTY(BlueprintReadWrite, EditAnywhere, Category = Input, meta = (EditCondition = "Operation != EPCGMedadataBitwiseOperation::Not", EditConditionHides))
-	FPCGAttributePropertySelector InputSource2;
+	FPCGAttributePropertyInputSelector InputSource2;
 	
 #if WITH_EDITORONLY_DATA
 	UPROPERTY()
diff --git a/Engine/Plugins/Experimental/PCG/Source/PCG/Public/Elements/Metadata/PCGMetadataBooleanOpElement.h b/Engine/Plugins/Experimental/PCG/Source/PCG/Public/Elements/Metadata/PCGMetadataBooleanOpElement.h
--- a/Engine/Plugins/Experimental/PCG/Source/PCG/Public/Elements/Metadata/PCGMetadataBooleanOpElement.h
+++ b/Engine/Plugins/Experimental/PCG/Source/PCG/Public/Elements/Metadata/PCGMetadataBooleanOpElement.h
@@ -37,7 +37,7 @@ public:
 	//~End UPCGSettings interface
 
 	//~Begin UPCGMetadataSettingsBase interface
-	FPCGAttributePropertySelector GetInputSource(uint32 Index) const override;
+	FPCGAttributePropertyInputSelector GetInputSource(uint32 Index) const override;
 
 	virtual FName GetInputPinLabel(uint32 Index) const override;
 	virtual uint32 GetInputPinNum() const override;
@@ -56,10 +56,10 @@ public:
 	EPCGMedadataBooleanOperation Operation = EPCGMedadataBooleanOperation::And;
 
 	UPROPERTY(BlueprintReadWrite, EditAnywhere, Category = Input)
-	FPCGAttributePropertySelector InputSource1;
+	FPCGAttributePropertyInputSelector InputSource1;
 
 	UPROPERTY(BlueprintReadWrite, EditAnywhere, Category = Input, meta = (EditCondition = "Operation != EPCGMedadataBooleanOperation::Not", EditConditionHides))
-	FPCGAttributePropertySelector InputSource2;
+	FPCGAttributePropertyInputSelector InputSource2;
 
 #if WITH_EDITORONLY_DATA
 	UPROPERTY()
diff --git a/Engine/Plugins/Experimental/PCG/Source/PCG/Public/Elements/Metadata/PCGMetadataBreakTransform.h b/Engine/Plugins/Experimental/PCG/Source/PCG/Public/Elements/Metadata/PCGMetadataBreakTransform.h
--- a/Engine/Plugins/Experimental/PCG/Source/PCG/Public/Elements/Metadata/PCGMetadataBreakTransform.h
+++ b/Engine/Plugins/Experimental/PCG/Source/PCG/Public/Elements/Metadata/PCGMetadataBreakTransform.h
@@ -23,7 +23,7 @@ public:
 	virtual FText GetDefaultNodeTitle() const override;
 #endif
 
-	FPCGAttributePropertySelector GetInputSource(uint32 Index) const override;
+	FPCGAttributePropertyInputSelector GetInputSource(uint32 Index) const override;
 
 	virtual FName GetOutputPinLabel(uint32 Index) const override;
 	virtual uint32 GetOutputPinNum() const override;
@@ -40,7 +40,7 @@ protected:
 
 public:
 	UPROPERTY(BlueprintReadWrite, EditAnywhere, Category = Input)
-	FPCGAttributePropertySelector InputSource;
+	FPCGAttributePropertyInputSelector InputSource;
 
 #if WITH_EDITORONLY_DATA
 	UPROPERTY()
diff --git a/Engine/Plugins/Experimental/PCG/Source/PCG/Public/Elements/Metadata/PCGMetadataBreakVector.h b/Engine/Plugins/Experimental/PCG/Source/PCG/Public/Elements/Metadata/PCGMetadataBreakVector.h
--- a/Engine/Plugins/Experimental/PCG/Source/PCG/Public/Elements/Metadata/PCGMetadataBreakVector.h
+++ b/Engine/Plugins/Experimental/PCG/Source/PCG/Public/Elements/Metadata/PCGMetadataBreakVector.h
@@ -31,7 +31,7 @@ public:
 	virtual FText GetDefaultNodeTitle() const override;
 #endif
 
-	FPCGAttributePropertySelector GetInputSource(uint32 Index) const override;
+	FPCGAttributePropertyInputSelector GetInputSource(uint32 Index) const override;
 
 	virtual FName GetOutputPinLabel(uint32 Index) const override;
 	virtual uint32 GetOutputPinNum() const override;
@@ -46,7 +46,7 @@ protected:
 
 public:
 	UPROPERTY(BlueprintReadWrite, EditAnywhere, Category = Input)
-	FPCGAttributePropertySelector InputSource;
+	FPCGAttributePropertyInputSelector InputSource;
 
 #if WITH_EDITORONLY_DATA
 	UPROPERTY()
diff --git a/Engine/Plugins/Experimental/PCG/Source/PCG/Public/Elements/Metadata/PCGMetadataCompareOpElement.h b/Engine/Plugins/Experimental/PCG/Source/PCG/Public/Elements/Metadata/PCGMetadataCompareOpElement.h
--- a/Engine/Plugins/Experimental/PCG/Source/PCG/Public/Elements/Metadata/PCGMetadataCompareOpElement.h
+++ b/Engine/Plugins/Experimental/PCG/Source/PCG/Public/Elements/Metadata/PCGMetadataCompareOpElement.h
@@ -39,7 +39,7 @@ public:
 	//~End UPCGSettings interface
 
 	//~Begin UPCGMetadataSettingsBase interface
-	FPCGAttributePropertySelector GetInputSource(uint32 Index) const override;
+	FPCGAttributePropertyInputSelector GetInputSource(uint32 Index) const override;
 
 	virtual FName GetInputPinLabel(uint32 Index) const override;
 	virtual uint32 GetInputPinNum() const override;
@@ -58,10 +58,10 @@ public:
 	EPCGMedadataCompareOperation Operation = EPCGMedadataCompareOperation::Equal;
 
 	UPROPERTY(BlueprintReadWrite, EditAnywhere, Category = Input)
-	FPCGAttributePropertySelector InputSource1;
+	FPCGAttributePropertyInputSelector InputSource1;
 
 	UPROPERTY(BlueprintReadWrite, EditAnywhere, Category = Input)
-	FPCGAttributePropertySelector InputSource2;
+	FPCGAttributePropertyInputSelector InputSource2;
 
 	UPROPERTY(BlueprintReadWrite, EditAnywhere, Category = Settings, meta = (EditCondition = "Operation == EPCGMedadataCompareOperation::Equal || Operation == EPCGMedadataCompareOperation::NotEqual", EditConditionHides))
 	double Tolerance = UE_DOUBLE_SMALL_NUMBER;
diff --git a/Engine/Plugins/Experimental/PCG/Source/PCG/Public/Elements/Metadata/PCGMetadataElement.h b/Engine/Plugins/Experimental/PCG/Source/PCG/Public/Elements/Metadata/PCGMetadataElement.h
--- a/Engine/Plugins/Experimental/PCG/Source/PCG/Public/Elements/Metadata/PCGMetadataElement.h
+++ b/Engine/Plugins/Experimental/PCG/Source/PCG/Public/Elements/Metadata/PCGMetadataElement.h
@@ -35,6 +35,7 @@ public:
 	virtual FText GetNodeTooltipText() const;
 	virtual EPCGSettingsType GetType() const override { return EPCGSettingsType::Metadata; }
 	virtual bool IsPinUsedByNodeExecution(const UPCGPin* InPin) const override;
+	virtual void ApplyDeprecation(UPCGNode* InOutNode) override;
 #endif
 
 protected:
@@ -45,10 +46,10 @@ protected:
 
 public:
 	UPROPERTY(BlueprintReadWrite, EditAnywhere, Category = Settings)
-	FPCGAttributePropertySelector InputSource;
+	FPCGAttributePropertyInputSelector InputSource;
 
 	UPROPERTY(BlueprintReadWrite, EditAnywhere, Category = Settings)
-	FPCGAttributePropertySelector OutputTarget;
+	FPCGAttributePropertyOutputSelector OutputTarget;
 
 #if WITH_EDITORONLY_DATA
 	UPROPERTY()
diff --git a/Engine/Plugins/Experimental/PCG/Source/PCG/Public/Elements/Metadata/PCGMetadataMakeRotator.h b/Engine/Plugins/Experimental/PCG/Source/PCG/Public/Elements/Metadata/PCGMetadataMakeRotator.h
--- a/Engine/Plugins/Experimental/PCG/Source/PCG/Public/Elements/Metadata/PCGMetadataMakeRotator.h
+++ b/Engine/Plugins/Experimental/PCG/Source/PCG/Public/Elements/Metadata/PCGMetadataMakeRotator.h
@@ -57,7 +57,7 @@ public:
 #endif // WITH_EDITOR
 	//~End UPCGSettings interface
 
-	FPCGAttributePropertySelector GetInputSource(uint32 Index) const override;
+	FPCGAttributePropertyInputSelector GetInputSource(uint32 Index) const override;
 
 	virtual FName GetInputPinLabel(uint32 Index) const override;
 	virtual uint32 GetInputPinNum() const override;
@@ -72,13 +72,13 @@ protected:
 
 public:
 	UPROPERTY(BlueprintReadWrite, EditAnywhere, Category = Input)
-	FPCGAttributePropertySelector InputSource1;
+	FPCGAttributePropertyInputSelector InputSource1;
 
 	UPROPERTY(BlueprintReadWrite, EditAnywhere, Category = Input, meta = (EditCondition = "Operation != EPCGMetadataMakeRotatorOp::MakeRotFromX && Operation != EPCGMetadataMakeRotatorOp::MakeRotFromY && Operation != EPCGMetadataMakeRotatorOp::MakeRotFromZ", EditConditionHides))
-	FPCGAttributePropertySelector InputSource2;
+	FPCGAttributePropertyInputSelector InputSource2;
 
 	UPROPERTY(BlueprintReadWrite, EditAnywhere, Category = Input, meta = (EditCondition = "Operation == EPCGMetadataMakeRotatorOp::MakeRotFromAxes || Operation == EPCGMetadataMakeRotatorOp::MakeRotFromAngles", EditConditionHides))
-	FPCGAttributePropertySelector InputSource3;
+	FPCGAttributePropertyInputSelector InputSource3;
 
 	UPROPERTY(BlueprintReadWrite, EditAnywhere, Category = Settings)
 	EPCGMetadataMakeRotatorOp Operation = EPCGMetadataMakeRotatorOp::MakeRotFromAxes;
diff --git a/Engine/Plugins/Experimental/PCG/Source/PCG/Public/Elements/Metadata/PCGMetadataMakeTransform.h b/Engine/Plugins/Experimental/PCG/Source/PCG/Public/Elements/Metadata/PCGMetadataMakeTransform.h
--- a/Engine/Plugins/Experimental/PCG/Source/PCG/Public/Elements/Metadata/PCGMetadataMakeTransform.h
+++ b/Engine/Plugins/Experimental/PCG/Source/PCG/Public/Elements/Metadata/PCGMetadataMakeTransform.h
@@ -36,7 +36,7 @@ public:
 #endif // WITH_EDITORs
 	//~End UPCGSettings interface
 
-	FPCGAttributePropertySelector GetInputSource(uint32 Index) const override;
+	FPCGAttributePropertyInputSelector GetInputSource(uint32 Index) const override;
 
 	virtual FName GetInputPinLabel(uint32 Index) const override;
 	virtual uint32 GetInputPinNum() const override;
@@ -51,13 +51,13 @@ protected:
 
 public:
 	UPROPERTY(BlueprintReadWrite, EditAnywhere, Category = Input)
-	FPCGAttributePropertySelector InputSource1;
+	FPCGAttributePropertyInputSelector InputSource1;
 
 	UPROPERTY(BlueprintReadWrite, EditAnywhere, Category = Input)
-	FPCGAttributePropertySelector InputSource2;
+	FPCGAttributePropertyInputSelector InputSource2;
 	
 	UPROPERTY(BlueprintReadWrite, EditAnywhere, Category = Input)
-	FPCGAttributePropertySelector InputSource3;
+	FPCGAttributePropertyInputSelector InputSource3;
 
 #if WITH_EDITORONLY_DATA
 	UPROPERTY()
diff --git a/Engine/Plugins/Experimental/PCG/Source/PCG/Public/Elements/Metadata/PCGMetadataMakeVector.h b/Engine/Plugins/Experimental/PCG/Source/PCG/Public/Elements/Metadata/PCGMetadataMakeVector.h
--- a/Engine/Plugins/Experimental/PCG/Source/PCG/Public/Elements/Metadata/PCGMetadataMakeVector.h
+++ b/Engine/Plugins/Experimental/PCG/Source/PCG/Public/Elements/Metadata/PCGMetadataMakeVector.h
@@ -58,7 +58,7 @@ public:
 #endif // WITH_EDITOR
 	//~End UPCGSettings interface
 
-	FPCGAttributePropertySelector GetInputSource(uint32 Index) const override;
+	FPCGAttributePropertyInputSelector GetInputSource(uint32 Index) const override;
 
 	virtual FName GetInputPinLabel(uint32 Index) const override;
 	virtual uint32 GetInputPinNum() const override;
@@ -73,16 +73,16 @@ protected:
 
 public:
 	UPROPERTY(BlueprintReadWrite, EditAnywhere, Category = Input)
-	FPCGAttributePropertySelector InputSource1;
+	FPCGAttributePropertyInputSelector InputSource1;
 
 	UPROPERTY(BlueprintReadWrite, EditAnywhere, Category = Input)
-	FPCGAttributePropertySelector InputSource2;
+	FPCGAttributePropertyInputSelector InputSource2;
 
 	UPROPERTY(BlueprintReadWrite, EditAnywhere, Category = Input)
-	FPCGAttributePropertySelector InputSource3;
+	FPCGAttributePropertyInputSelector InputSource3;
 
 	UPROPERTY(BlueprintReadWrite, EditAnywhere, Category = Input)
-	FPCGAttributePropertySelector InputSource4;
+	FPCGAttributePropertyInputSelector InputSource4;
 
 	UPROPERTY(BlueprintReadWrite, EditAnywhere, Category = Settings, meta = (ValidEnumValues = "Vector2, Vector, Vector4"))
 	EPCGMetadataTypes OutputType = EPCGMetadataTypes::Vector2;
diff --git a/Engine/Plugins/Experimental/PCG/Source/PCG/Public/Elements/Metadata/PCGMetadataMathsOpElement.h b/Engine/Plugins/Experimental/PCG/Source/PCG/Public/Elements/Metadata/PCGMetadataMathsOpElement.h
--- a/Engine/Plugins/Experimental/PCG/Source/PCG/Public/Elements/Metadata/PCGMetadataMathsOpElement.h
+++ b/Engine/Plugins/Experimental/PCG/Source/PCG/Public/Elements/Metadata/PCGMetadataMathsOpElement.h
@@ -63,7 +63,7 @@ public:
 	//~End UPCGSettings interface
 
 	//~Begin UPCGMetadataSettingsBase interface
-	FPCGAttributePropertySelector GetInputSource(uint32 Index) const override;
+	FPCGAttributePropertyInputSelector GetInputSource(uint32 Index) const override;
 
 	virtual FName GetInputPinLabel(uint32 Index) const override;
 	virtual uint32 GetInputPinNum() const override;
@@ -81,13 +81,13 @@ public:
 	EPCGMedadataMathsOperation Operation = EPCGMedadataMathsOperation::Add;
 
 	UPROPERTY(BlueprintReadWrite, EditAnywhere, Category = Input)
-	FPCGAttributePropertySelector InputSource1;
+	FPCGAttributePropertyInputSelector InputSource1;
 
 	UPROPERTY(BlueprintReadWrite, EditAnywhere, Category = Input, meta = (EditCondition = "(Operation & '/Script/PCG.EPCGMedadataMathsOperation::BinaryOp') || (Operation & '/Script/PCG.EPCGMedadataMathsOperation::TernaryOp')", EditConditionHides))
-	FPCGAttributePropertySelector InputSource2;
+	FPCGAttributePropertyInputSelector InputSource2;
 
 	UPROPERTY(BlueprintReadWrite, EditAnywhere, Category = Input, meta = (EditCondition = "Operation & '/Script/PCG.EPCGMedadataMathsOperation::TernaryOp'", EditConditionHides))
-	FPCGAttributePropertySelector InputSource3;
+	FPCGAttributePropertyInputSelector InputSource3;
 
 #if WITH_EDITORONLY_DATA
 	UPROPERTY()
diff --git a/Engine/Plugins/Experimental/PCG/Source/PCG/Public/Elements/Metadata/PCGMetadataOpElementBase.h b/Engine/Plugins/Experimental/PCG/Source/PCG/Public/Elements/Metadata/PCGMetadataOpElementBase.h
--- a/Engine/Plugins/Experimental/PCG/Source/PCG/Public/Elements/Metadata/PCGMetadataOpElementBase.h
+++ b/Engine/Plugins/Experimental/PCG/Source/PCG/Public/Elements/Metadata/PCGMetadataOpElementBase.h
@@ -89,6 +89,7 @@ public:
 	virtual bool HasDynamicPins() const override { return true; }
 	virtual bool GetPinExtraIcon(const UPCGPin* InPin, FName& OutExtraIcon, FText& OutTooltip) const override;
 	virtual FText GetNodeTooltipText() const override;
+	virtual void ApplyDeprecation(UPCGNode* InOutNode) override;
 #endif
 	virtual EPCGDataType GetCurrentPinTypes(const UPCGPin* InPin) const override;
 
@@ -98,7 +99,7 @@ protected:
 	//~End UPCGSettings interface
 
 public:
-	virtual FPCGAttributePropertySelector GetInputSource(uint32 Index) const { return FPCGAttributePropertySelector(); };
+	virtual FPCGAttributePropertyInputSelector GetInputSource(uint32 Index) const { return FPCGAttributePropertyInputSelector(); };
 
 	virtual FName GetInputPinLabel(uint32 Index) const { return PCGPinConstants::DefaultInputLabel; }
 	virtual uint32 GetInputPinNum() const { return 1; };
@@ -124,7 +125,7 @@ public:
 	uint32 GetInputPinToForward() const;
 
 	UPROPERTY(BlueprintReadWrite, EditAnywhere, Category = "Output")
-	FPCGAttributePropertySelector OutputTarget;
+	FPCGAttributePropertyOutputSelector OutputTarget;
 
 	/* By default, output is taken from first non-param pin (aka if the second pin is a point data, the output will be this point data). You can change it to any available input pin. */
 	UPROPERTY(BlueprintReadWrite, EditAnywhere, AdvancedDisplay, Category = "Output", meta = (GetOptions = GetOutputDataFromPinOptions))
@@ -169,6 +170,8 @@ public:
 		uint16 OutputType;
 		const UPCGMetadataSettingsBase* Settings = nullptr;
 
+		TArray<FPCGAttributePropertyInputSelector> InputSources;
+
 		TArray<TUniquePtr<const IPCGAttributeAccessorKeys>> InputKeys;
 		TArray<TUniquePtr<IPCGAttributeAccessorKeys>> OutputKeys;
 
diff --git a/Engine/Plugins/Experimental/PCG/Source/PCG/Public/Elements/Metadata/PCGMetadataRotatorOpElement.h b/Engine/Plugins/Experimental/PCG/Source/PCG/Public/Elements/Metadata/PCGMetadataRotatorOpElement.h
--- a/Engine/Plugins/Experimental/PCG/Source/PCG/Public/Elements/Metadata/PCGMetadataRotatorOpElement.h
+++ b/Engine/Plugins/Experimental/PCG/Source/PCG/Public/Elements/Metadata/PCGMetadataRotatorOpElement.h
@@ -48,7 +48,7 @@ public:
 	//~End UPCGSettings interface
 
 	//~Begin UPCGMetadataSettingsBase interface
-	FPCGAttributePropertySelector GetInputSource(uint32 Index) const override;
+	FPCGAttributePropertyInputSelector GetInputSource(uint32 Index) const override;
 
 	virtual FName GetInputPinLabel(uint32 Index) const override;
 	virtual uint32 GetInputPinNum() const override;
@@ -66,13 +66,13 @@ public:
 	EPCGMedadataRotatorOperation Operation = EPCGMedadataRotatorOperation::Combine;
 
 	UPROPERTY(BlueprintReadWrite, EditAnywhere, Category = Input)
-	FPCGAttributePropertySelector InputSource1;
+	FPCGAttributePropertyInputSelector InputSource1;
 
 	UPROPERTY(BlueprintReadWrite, EditAnywhere, Category = Input, meta = (EditCondition = "Operation != EPCGMedadataRotatorOperation::Invert", EditConditionHides))
-	FPCGAttributePropertySelector InputSource2;
+	FPCGAttributePropertyInputSelector InputSource2;
 
 	UPROPERTY(BlueprintReadWrite, EditAnywhere, Category = Input, meta = (EditCondition = "Operation == EPCGMedadataRotatorOperation::Lerp", EditConditionHides))
-	FPCGAttributePropertySelector InputSource3;
+	FPCGAttributePropertyInputSelector InputSource3;
 
 #if WITH_EDITORONLY_DATA
 	UPROPERTY()
diff --git a/Engine/Plugins/Experimental/PCG/Source/PCG/Public/Elements/Metadata/PCGMetadataTransformOpElement.h b/Engine/Plugins/Experimental/PCG/Source/PCG/Public/Elements/Metadata/PCGMetadataTransformOpElement.h
--- a/Engine/Plugins/Experimental/PCG/Source/PCG/Public/Elements/Metadata/PCGMetadataTransformOpElement.h
+++ b/Engine/Plugins/Experimental/PCG/Source/PCG/Public/Elements/Metadata/PCGMetadataTransformOpElement.h
@@ -50,7 +50,7 @@ public:
 	//~End UPCGSettings interface
 
 	//~Begin UPCGMetadataSettingsBase interface
-	FPCGAttributePropertySelector GetInputSource(uint32 Index) const override;
+	FPCGAttributePropertyInputSelector GetInputSource(uint32 Index) const override;
 
 	virtual FName GetInputPinLabel(uint32 Index) const override;
 	virtual uint32 GetInputPinNum() const override;
@@ -71,13 +71,13 @@ public:
 	EPCGTransformLerpMode TransformLerpMode = EPCGTransformLerpMode::QuatInterp;
 
 	UPROPERTY(BlueprintReadWrite, EditAnywhere, Category = Input)
-	FPCGAttributePropertySelector InputSource1;
+	FPCGAttributePropertyInputSelector InputSource1;
 
 	UPROPERTY(BlueprintReadWrite, EditAnywhere, Category = Input, meta = (EditCondition = "Operation != EPCGMedadataTransformOperation::Invert", EditConditionHides))
-	FPCGAttributePropertySelector InputSource2;
+	FPCGAttributePropertyInputSelector InputSource2;
 
 	UPROPERTY(BlueprintReadWrite, EditAnywhere, Category = Input, meta = (EditCondition = "Operation == EPCGMedadataTransformOperation::Lerp", EditConditionHides))
-	FPCGAttributePropertySelector InputSource3;
+	FPCGAttributePropertyInputSelector InputSource3;
 
 #if WITH_EDITORONLY_DATA
 	UPROPERTY()
diff --git a/Engine/Plugins/Experimental/PCG/Source/PCG/Public/Elements/Metadata/PCGMetadataTrigOpElement.h b/Engine/Plugins/Experimental/PCG/Source/PCG/Public/Elements/Metadata/PCGMetadataTrigOpElement.h
--- a/Engine/Plugins/Experimental/PCG/Source/PCG/Public/Elements/Metadata/PCGMetadataTrigOpElement.h
+++ b/Engine/Plugins/Experimental/PCG/Source/PCG/Public/Elements/Metadata/PCGMetadataTrigOpElement.h
@@ -42,7 +42,7 @@ public:
 	//~End UPCGSettings interface
 
 	//~Begin UPCGMetadataSettingsBase interface
-	FPCGAttributePropertySelector GetInputSource(uint32 Index) const override;
+	FPCGAttributePropertyInputSelector GetInputSource(uint32 Index) const override;
 
 	virtual FName GetInputPinLabel(uint32 Index) const override;
 	virtual uint32 GetInputPinNum() const override;
@@ -61,10 +61,10 @@ public:
 	EPCGMedadataTrigOperation Operation = EPCGMedadataTrigOperation::Acos;
 
 	UPROPERTY(BlueprintReadWrite, EditAnywhere, Category = Input)
-	FPCGAttributePropertySelector InputSource1;
+	FPCGAttributePropertyInputSelector InputSource1;
 
 	UPROPERTY(BlueprintReadWrite, EditAnywhere, Category = Input, meta = (EditCondition = "Operation == EPCGMedadataTrigOperation::Atan2", EditConditionHides))
-	FPCGAttributePropertySelector InputSource2;
+	FPCGAttributePropertyInputSelector InputSource2;
 
 #if WITH_EDITORONLY_DATA
 	UPROPERTY()
diff --git a/Engine/Plugins/Experimental/PCG/Source/PCG/Public/Elements/Metadata/PCGMetadataVectorOpElement.h b/Engine/Plugins/Experimental/PCG/Source/PCG/Public/Elements/Metadata/PCGMetadataVectorOpElement.h
--- a/Engine/Plugins/Experimental/PCG/Source/PCG/Public/Elements/Metadata/PCGMetadataVectorOpElement.h
+++ b/Engine/Plugins/Experimental/PCG/Source/PCG/Public/Elements/Metadata/PCGMetadataVectorOpElement.h
@@ -46,7 +46,7 @@ public:
 	//~End UPCGSettings interface
 
 	//~Begin UPCGMetadataSettingsBase interface
-	FPCGAttributePropertySelector GetInputSource(uint32 Index) const override;
+	FPCGAttributePropertyInputSelector GetInputSource(uint32 Index) const override;
 
 	virtual FName GetInputPinLabel(uint32 Index) const override;
 	virtual uint32 GetInputPinNum() const override;
@@ -65,13 +65,13 @@ public:
 	EPCGMedadataVectorOperation Operation = EPCGMedadataVectorOperation::Cross;
 
 	UPROPERTY(BlueprintReadWrite, EditAnywhere, Category = Input)
-	FPCGAttributePropertySelector InputSource1;
+	FPCGAttributePropertyInputSelector InputSource1;
 
 	UPROPERTY(BlueprintReadWrite, EditAnywhere, Category = Input, meta = (EditCondition = "Operation != EPCGMedadataVectorOperation::Normalize && Operation != EPCGMedadataVectorOperation::Length", EditConditionHides))
-	FPCGAttributePropertySelector InputSource2;
+	FPCGAttributePropertyInputSelector InputSource2;
 
 	UPROPERTY(BlueprintReadWrite, EditAnywhere, Category = Input, meta = (EditCondition = "Operation == EPCGMedadataVectorOperation::RotateAroundAxis", EditConditionHides))
-	FPCGAttributePropertySelector InputSource3;
+	FPCGAttributePropertyInputSelector InputSource3;
 
 #if WITH_EDITORONLY_DATA
 	UPROPERTY()
diff --git a/Engine/Plugins/Experimental/PCG/Source/PCG/Public/Elements/PCGAttributeGetFromPointIndexElement.h b/Engine/Plugins/Experimental/PCG/Source/PCG/Public/Elements/PCGAttributeGetFromPointIndexElement.h
--- a/Engine/Plugins/Experimental/PCG/Source/PCG/Public/Elements/PCGAttributeGetFromPointIndexElement.h
+++ b/Engine/Plugins/Experimental/PCG/Source/PCG/Public/Elements/PCGAttributeGetFromPointIndexElement.h
@@ -4,6 +4,7 @@
 
 #include "PCGSettings.h"
 #include "Metadata/PCGAttributePropertySelector.h"
+#include "Metadata/PCGMetadataAttribute.h"
 
 #include "PCGAttributeGetFromPointIndexElement.generated.h"
 
@@ -35,6 +36,7 @@ public:
 	virtual FName GetDefaultNodeName() const override;
 	virtual FText GetDefaultNodeTitle() const override;
 	virtual EPCGSettingsType GetType() const override { return EPCGSettingsType::Metadata; }
+	virtual void ApplyDeprecation(UPCGNode* InOutNode) override;
 #endif
 
 	virtual TArray<FPCGPinProperties> InputPinProperties() const override;
@@ -42,7 +44,7 @@ public:
 	//~End UPCGSettings interface
 
 	UPROPERTY(BlueprintReadWrite, EditAnywhere, Category = Settings)
-	FPCGAttributePropertySelector InputSource;
+	FPCGAttributePropertyInputSelector InputSource;
 
 	UPROPERTY(BlueprintReadWrite, EditAnywhere, Category = Settings, meta = (PCG_Overridable))
 	int32 Index = 0;
diff --git a/Engine/Plugins/Experimental/PCG/Source/PCG/Public/Elements/PCGAttributeNoise.h b/Engine/Plugins/Experimental/PCG/Source/PCG/Public/Elements/PCGAttributeNoise.h
--- a/Engine/Plugins/Experimental/PCG/Source/PCG/Public/Elements/PCGAttributeNoise.h
+++ b/Engine/Plugins/Experimental/PCG/Source/PCG/Public/Elements/PCGAttributeNoise.h
@@ -56,13 +56,10 @@ protected:
 
 public:
 	UPROPERTY(BlueprintReadWrite, EditAnywhere, Category = Settings)
-	FPCGAttributePropertySelector InputSource;
+	FPCGAttributePropertyInputSelector InputSource;
 
 	UPROPERTY(BlueprintReadWrite, EditAnywhere, Category = Settings)
-	bool bOutputTargetDifferentFromInputSource = false;
-
-	UPROPERTY(BlueprintReadWrite, EditAnywhere, Category = Settings, meta = (EditCondition = "bOutputTargetDifferentFromInputSource", EditConditionHides))
-	FPCGAttributePropertySelector OutputTarget;
+	FPCGAttributePropertyOutputSelector OutputTarget;
 
 	/** Attribute = (Original op Noise), Noise in [NoiseMin, NoiseMax] */
 	UPROPERTY(BlueprintReadWrite, EditAnywhere, Category = Settings, meta = (PCG_Overridable))
@@ -95,6 +92,9 @@ public:
 
 	UPROPERTY()
 	bool bInvertSourceDensity_DEPRECATED = false;
+
+	UPROPERTY()
+	bool bOutputTargetDifferentFromInputSource_DEPRECATED = false;
 #endif // WITH_EDITORDATA_ONLY
 };
 
@@ -102,8 +102,10 @@ struct FPCGAttributeNoiseContext : public FPCGContext
 {
 	int32 CurrentInput = 0;
 	bool bDataPreparedForCurrentInput = false;
-	TUniquePtr<IPCGAttributeAccessor> InputAccessor;
-	TUniquePtr<IPCGAttributeAccessor> OptionalOutputAccessor;
+	FPCGAttributePropertyInputSelector InputSource;
+	FPCGAttributePropertyOutputSelector OutputTarget;
+	TUniquePtr<const IPCGAttributeAccessor> InputAccessor;
+	TUniquePtr<IPCGAttributeAccessor> OutputAccessor;
 	TUniquePtr<IPCGAttributeAccessorKeys> Keys;
 
 	TArray<uint8> TempValuesBuffer;
diff --git a/Engine/Plugins/Experimental/PCG/Source/PCG/Public/Elements/PCGAttributeReduceElement.h b/Engine/Plugins/Experimental/PCG/Source/PCG/Public/Elements/PCGAttributeReduceElement.h
--- a/Engine/Plugins/Experimental/PCG/Source/PCG/Public/Elements/PCGAttributeReduceElement.h
+++ b/Engine/Plugins/Experimental/PCG/Source/PCG/Public/Elements/PCGAttributeReduceElement.h
@@ -4,6 +4,7 @@
 
 #include "PCGSettings.h"
 #include "Metadata/PCGAttributePropertySelector.h"
+#include "Metadata/PCGMetadataAttribute.h"
 
 #include "PCGAttributeReduceElement.generated.h"
 
@@ -36,6 +37,7 @@ public:
 	virtual FName GetDefaultNodeName() const override;
 	virtual FText GetDefaultNodeTitle() const override;
 	virtual EPCGSettingsType GetType() const override { return EPCGSettingsType::Metadata; }
+	virtual void ApplyDeprecation(UPCGNode* InOutNode) override;
 #endif
 
 	virtual FName AdditionalTaskName() const override;
@@ -45,7 +47,7 @@ public:
 	//~End UPCGSettings interface
 
 	UPROPERTY(BlueprintReadWrite, EditAnywhere, Category = Settings)
-	FPCGAttributePropertySelector InputSource;
+	FPCGAttributePropertyInputSelector InputSource;
 
 	UPROPERTY(BlueprintReadWrite, EditAnywhere, Category = Settings)
 	FName OutputAttributeName = NAME_None;
diff --git a/Engine/Plugins/Experimental/PCG/Source/PCG/Public/Elements/PCGAttributeSelectElement.h b/Engine/Plugins/Experimental/PCG/Source/PCG/Public/Elements/PCGAttributeSelectElement.h
--- a/Engine/Plugins/Experimental/PCG/Source/PCG/Public/Elements/PCGAttributeSelectElement.h
+++ b/Engine/Plugins/Experimental/PCG/Source/PCG/Public/Elements/PCGAttributeSelectElement.h
@@ -4,6 +4,7 @@
 
 #include "PCGSettings.h"
 #include "Metadata/PCGAttributePropertySelector.h"
+#include "Metadata/PCGMetadataAttribute.h"
 
 #include "PCGAttributeSelectElement.generated.h"
 
@@ -59,6 +60,7 @@ public:
 	virtual FName GetDefaultNodeName() const override;
 	virtual FText GetDefaultNodeTitle() const override;
 	virtual EPCGSettingsType GetType() const override { return EPCGSettingsType::Metadata; }
+	virtual void ApplyDeprecation(UPCGNode* InOutNode) override;
 #endif
 
 	virtual FName AdditionalTaskName() const override;
@@ -68,7 +70,7 @@ public:
 	//~End UPCGSettings interface
 
 	UPROPERTY(BlueprintReadWrite, EditAnywhere, Category = Settings)
-	FPCGAttributePropertySelector InputSource;
+	FPCGAttributePropertyInputSelector InputSource;
 
 	UPROPERTY(BlueprintReadWrite, EditAnywhere, Category = Settings)
 	FName OutputAttributeName = NAME_None;
diff --git a/Engine/Plugins/Experimental/PCG/Source/PCG/Public/Elements/PCGPointFilter.h b/Engine/Plugins/Experimental/PCG/Source/PCG/Public/Elements/PCGPointFilter.h
--- a/Engine/Plugins/Experimental/PCG/Source/PCG/Public/Elements/PCGPointFilter.h
+++ b/Engine/Plugins/Experimental/PCG/Source/PCG/Public/Elements/PCGPointFilter.h
@@ -62,7 +62,7 @@ struct PCG_API FPCGPointFilterThresholdSettings
 	bool bUseConstantThreshold = false;
 
 	UPROPERTY(BlueprintReadWrite, EditAnywhere, Category = Settings, meta = (EditCondition = "!bUseConstantThreshold", EditConditionHides, PCG_NotOverridable))
-	FPCGAttributePropertySelector ThresholdAttribute;
+	FPCGAttributePropertyInputSelector ThresholdAttribute;
 
 	/** If the threshold data is Point data, it will sample input points in threshold data. Always true with Spatial data.*/
 	UPROPERTY(BlueprintReadWrite, EditAnywhere, Category = Settings, meta = (EditCondition = "!bUseConstantThreshold", EditConditionHides, PCG_Overridable))
@@ -115,14 +115,14 @@ public:
 
 	/** Target property/attribute related properties */
 	UPROPERTY(BlueprintReadWrite, EditAnywhere, Category = Settings)
-	FPCGAttributePropertySelector TargetAttribute;
+	FPCGAttributePropertyInputSelector TargetAttribute;
 
 	/** Threshold property/attribute/constant related properties */
 	UPROPERTY(BlueprintReadWrite, EditAnywhere, Category = Settings, meta = (PCG_Overridable))
 	bool bUseConstantThreshold = false;
 
 	UPROPERTY(BlueprintReadWrite, EditAnywhere, Category = Settings, meta = (EditCondition = "!bUseConstantThreshold", EditConditionHides, PCG_NotOverridable))
-	FPCGAttributePropertySelector ThresholdAttribute;
+	FPCGAttributePropertyInputSelector ThresholdAttribute;
 
 	/** If the threshold data is Point data, it will sample input points in threshold data. Always true with Spatial data.*/
 	UPROPERTY(BlueprintReadWrite, EditAnywhere, Category = Settings, meta = (EditCondition = "!bUseConstantThreshold", EditConditionHides, PCG_Overridable))
@@ -210,7 +210,7 @@ protected:
 public:
 	/** Target property/attribute related properties */
 	UPROPERTY(BlueprintReadWrite, EditAnywhere, Category = Settings)
-	FPCGAttributePropertySelector TargetAttribute;
+	FPCGAttributePropertyInputSelector TargetAttribute;
 
 	/** Threshold property/attribute/constant related properties */
 	UPROPERTY(BlueprintReadWrite, EditAnywhere, Category = Settings, meta = (PCG_Overridable))
@@ -223,7 +223,7 @@ public:
 class FPCGPointFilterElementBase : public FSimplePCGElement
 {
 protected:
-	bool DoFiltering(FPCGContext* Context, EPCGPointFilterOperator InOperation, const FPCGAttributePropertySelector& TargetAttribute, const FPCGPointFilterThresholdSettings& FirstThreshold, const FPCGPointFilterThresholdSettings* SecondThreshold = nullptr) const;
+	bool DoFiltering(FPCGContext* Context, EPCGPointFilterOperator InOperation, const FPCGAttributePropertyInputSelector& TargetAttribute, const FPCGPointFilterThresholdSettings& FirstThreshold, const FPCGPointFilterThresholdSettings* SecondThreshold = nullptr) const;
 };
 
 class FPCGPointFilterElement : public FPCGPointFilterElementBase
diff --git a/Engine/Plugins/Experimental/PCG/Source/PCG/Public/Elements/PCGPointMatchAndSet.h b/Engine/Plugins/Experimental/PCG/Source/PCG/Public/Elements/PCGPointMatchAndSet.h
--- a/Engine/Plugins/Experimental/PCG/Source/PCG/Public/Elements/PCGPointMatchAndSet.h
+++ b/Engine/Plugins/Experimental/PCG/Source/PCG/Public/Elements/PCGPointMatchAndSet.h
@@ -28,6 +28,7 @@ public:
 	virtual FText GetDefaultNodeTitle() const override { return NSLOCTEXT("PCGPointMatchAndSetSettings", "NodeTitle", "Point Match And Set"); }
 	virtual FText GetNodeTooltipText() const override;
 	virtual EPCGSettingsType GetType() const override { return EPCGSettingsType::Metadata; }
+	virtual void ApplyDeprecation(UPCGNode* InOutNode) override;
 #endif // WITH_EDITOR
 
 	virtual TArray<FPCGPinProperties> InputPinProperties() const override;
@@ -63,7 +64,7 @@ public:
 
 	/** "Set" part of the Match & Set - defines what will be changed in the operation */
 	UPROPERTY(BlueprintReadWrite, EditAnywhere, Category = Settings)
-	FPCGAttributePropertySelector SetTarget;
+	FPCGAttributePropertyOutputSelector SetTarget;
 
 	/** If the "Set" part is an attribute, then the type must be provided */
 	UPROPERTY(BlueprintReadWrite, EditAnywhere, Category = Settings, meta = (EditCondition = "bSetTargetIsAttribute", HideEditConditionToggle, EditConditionHides, ValidEnumValues = "Float, Double, Integer32, Integer64, Vector2, Vector, Vector4, Quaternion, Transform, String, Boolean, Rotator, Name"))
diff --git a/Engine/Plugins/Experimental/PCG/Source/PCG/Public/Elements/PCGPropertyToParamData.h b/Engine/Plugins/Experimental/PCG/Source/PCG/Public/Elements/PCGPropertyToParamData.h
--- a/Engine/Plugins/Experimental/PCG/Source/PCG/Public/Elements/PCGPropertyToParamData.h
+++ b/Engine/Plugins/Experimental/PCG/Source/PCG/Public/Elements/PCGPropertyToParamData.h
@@ -5,6 +5,7 @@
 #include "PCGPin.h"
 #include "PCGSettings.h"
 #include "Elements/PCGActorSelector.h"
+#include "Metadata/PCGMetadataAttribute.h"
 
 #include "PCGPropertyToParamData.generated.h"
 
@@ -62,7 +63,7 @@ public:
 	UPROPERTY(BlueprintReadWrite, EditAnywhere, Category = Settings)
 	bool bExtractObjectAndStruct = false;
 
-	/** By default, attribute name will be the property name, but it can be overridden by this name. */
+	/** By default, attribute name will be None, but it can be overridden by this name. */
 	UPROPERTY(BlueprintReadWrite, EditAnywhere, Category = Settings, meta = (EditCondition = "!bExtractObjectAndStruct", EditConditionHides))
 	FName OutputAttributeName = NAME_None;
 
diff --git a/Engine/Plugins/Experimental/PCG/Source/PCG/Public/Elements/PCGSpatialNoise.h b/Engine/Plugins/Experimental/PCG/Source/PCG/Public/Elements/PCGSpatialNoise.h
--- a/Engine/Plugins/Experimental/PCG/Source/PCG/Public/Elements/PCGSpatialNoise.h
+++ b/Engine/Plugins/Experimental/PCG/Source/PCG/Public/Elements/PCGSpatialNoise.h
@@ -93,7 +93,7 @@ public:
 
 	// The output attribute name to write, if not 'None'
 	UPROPERTY(BlueprintReadWrite, EditAnywhere, Category = Settings)
-	FPCGAttributePropertySelector ValueTarget;
+	FPCGAttributePropertyOutputNoSourceSelector ValueTarget;
 
 	// Adds a random amount of offset up to this amount
 	UPROPERTY(BlueprintReadWrite, EditAnywhere, Category = Settings, meta = (PCG_Overridable))
@@ -109,7 +109,7 @@ public:
 
 	// The output attribute name to write the voronoi cell id, if not 'None'
 	UPROPERTY(BlueprintReadWrite, EditAnywhere, Category = Settings, meta = (EditCondition = "Mode == PCGSpatialNoiseMode::Voronoi2D", EditConditionHides, PCG_Overridable))
-	FPCGAttributePropertySelector VoronoiCellIDTarget;
+	FPCGAttributePropertyOutputNoSourceSelector VoronoiCellIDTarget;
 	
 	// If true it will orient the output points to point towards the cell edges, which can be used for effects
 	UPROPERTY(BlueprintReadWrite, EditAnywhere, Category = Settings, meta = (EditCondition = "Mode == PCGSpatialNoiseMode::Voronoi2D", EditConditionHides, PCG_Overridable))
diff --git a/Engine/Plugins/Experimental/PCG/Source/PCG/Public/Elements/PCGSpawnActor.h b/Engine/Plugins/Experimental/PCG/Source/PCG/Public/Elements/PCGSpawnActor.h
--- a/Engine/Plugins/Experimental/PCG/Source/PCG/Public/Elements/PCGSpawnActor.h
+++ b/Engine/Plugins/Experimental/PCG/Source/PCG/Public/Elements/PCGSpawnActor.h
@@ -34,7 +34,7 @@ struct FPCGActorPropertyOverride
 	GENERATED_BODY()
 
 	UPROPERTY(EditAnywhere, Category = Settings)
-	FPCGAttributePropertySelector InputSource;
+	FPCGAttributePropertyInputSelector InputSource;
 
 	UPROPERTY(EditAnywhere, Category = Settings)
 	FString PropertyTarget;
diff --git a/Engine/Plugins/Experimental/PCG/Source/PCG/Public/Helpers/PCGMetadataHelpers.h b/Engine/Plugins/Experimental/PCG/Source/PCG/Public/Helpers/PCGMetadataHelpers.h
--- a/Engine/Plugins/Experimental/PCG/Source/PCG/Public/Helpers/PCGMetadataHelpers.h
+++ b/Engine/Plugins/Experimental/PCG/Source/PCG/Public/Helpers/PCGMetadataHelpers.h
@@ -4,12 +4,17 @@
 
 #include "HAL/Platform.h"
 
+class UPCGData;
 class UPCGMetadata;
 
 namespace PCGMetadataHelpers
 {
 	PCG_API bool HasSameRoot(const UPCGMetadata* Metadata1, const UPCGMetadata* Metadata2);
 	PCG_API const UPCGMetadata* GetParentMetadata(const UPCGMetadata* Metadata);
+
+	// Helpers functions to cast in spatial/param and return metadata. Nullptr if data doesn't have metadata
+	PCG_API const UPCGMetadata* GetConstMetadata(const UPCGData* InData);
+	PCG_API UPCGMetadata* GetMutableMetadata(UPCGData* InData);
 }
 
 #if UE_ENABLE_INCLUDE_ORDER_DEPRECATED_IN_5_2
diff --git a/Engine/Plugins/Experimental/PCG/Source/PCG/Public/Metadata/Accessors/PCGAttributeAccessorHelpers.h b/Engine/Plugins/Experimental/PCG/Source/PCG/Public/Metadata/Accessors/PCGAttributeAccessorHelpers.h
--- a/Engine/Plugins/Experimental/PCG/Source/PCG/Public/Metadata/Accessors/PCGAttributeAccessorHelpers.h
+++ b/Engine/Plugins/Experimental/PCG/Source/PCG/Public/Metadata/Accessors/PCGAttributeAccessorHelpers.h
@@ -40,13 +40,15 @@ namespace PCGAttributeAccessorHelpers
 	* Creates a const accessor to the property or attribute pointed at by the InSelector.
 	* Note that InData must not be null if the selector points to an attribute,
 	* but in the case of properties, it either has to be the appropriate type or null.
+	* Make sure to update your selector before-hand if you want to support "@Last"
 	*/
 	PCG_API TUniquePtr<const IPCGAttributeAccessor> CreateConstAccessor(const UPCGData* InData, const FPCGAttributePropertySelector& InSelector);
 	
-	/** 
-	* Creates a non-const accessor to the property or attribute pointed at by the InSelector.
+	/**
+	* Creates a accessor to the property or attribute pointed at by the InSelector.
 	* Note that InData must not be null if the selector points to an attribute,
 	* but in the case of properties, it either has to be the appropriate type or null.
+	* Make sure to update your selector before-hand if you want to support "@Source". Otherwise the creation will fail.
 	*/
 	PCG_API TUniquePtr<IPCGAttributeAccessor> CreateAccessor(UPCGData* InData, const FPCGAttributePropertySelector& InSelector);
 
diff --git a/Engine/Plugins/Experimental/PCG/Source/PCG/Public/Metadata/PCGAttributePropertySelector.h b/Engine/Plugins/Experimental/PCG/Source/PCG/Public/Metadata/PCGAttributePropertySelector.h
--- a/Engine/Plugins/Experimental/PCG/Source/PCG/Public/Metadata/PCGAttributePropertySelector.h
+++ b/Engine/Plugins/Experimental/PCG/Source/PCG/Public/Metadata/PCGAttributePropertySelector.h
@@ -22,41 +22,43 @@ enum class EPCGExtraProperties : uint8
 	Index
 };
 
+struct FPCGCustomVersion;
+class UPCGData;
+
 /**
 * Blueprint class to allow to select an attribute or a property.
 * It will handle the logic and can only be modified using the blueprint library defined below.
 * Also has a custom detail view in the PCGEditor plugin.
+* 
+* Note: This class should not be used as is, but need to be referenced by either an "InputSelector" or an "OutputSelector" (defined below).
+* The reason for that is to provide 2 different default values for input and output. Input will have the "@Last" default value (meaning last attribute written to)
+* and the Output will have "@Source" default value (meaning, same thing as input).
 */
-USTRUCT(BlueprintType)
+USTRUCT(BlueprintType, meta = (Hidden))
 struct PCG_API FPCGAttributePropertySelector
 {
 	GENERATED_BODY()
 
 public:
-	UPROPERTY(BlueprintReadOnly, VisibleAnywhere, Category = "Settings")
-	EPCGAttributePropertySelection Selection = EPCGAttributePropertySelection::Attribute;
-
-	UPROPERTY(BlueprintReadOnly, VisibleAnywhere, Category = "Settings", meta = (EditCondition = "Selection == EPCGAttributePropertySelection::Attribute", EditConditionHides))
-	FName AttributeName;
-
-	UPROPERTY(BlueprintReadOnly, VisibleAnywhere, Category = "Settings", meta = (EditCondition = "Selection == EPCGAttributePropertySelection::PointProperty", EditConditionHides))
-	EPCGPointProperties PointProperty = EPCGPointProperties::Position;
-
-	UPROPERTY(BlueprintReadOnly, VisibleAnywhere, Category = "Settings", meta = (EditCondition = "Selection == EPCGAttributePropertySelection::ExtraProperty", EditConditionHides))
-	EPCGExtraProperties ExtraProperty = EPCGExtraProperties::Index;
+	virtual ~FPCGAttributePropertySelector() = default;
 
-	UPROPERTY()
-	TArray<FString> ExtraNames;
-
-	// Force get the "None" attribute instead getting the last attribute
-	UPROPERTY()
-	bool bForceGetNone = false;
+	bool operator==(const FPCGAttributePropertySelector& Other) const;
+	bool IsSame(const FPCGAttributePropertySelector& Other, bool bIncludeExtraNames = true) const;
 
 	// Setters, retrurn true if something changed.
 	bool SetPointProperty(EPCGPointProperties InPointProperty, bool bResetExtraNames = true);
-	bool SetAttributeName(FName InAttributeName, bool bResetExtraNames = true);
 	bool SetExtraProperty(EPCGExtraProperties InExtraProperty, bool bResetExtraNames = true);
+	bool SetAttributeName(FName InAttributeName, bool bResetExtraNames = true);
+
+	// Getters
+	EPCGAttributePropertySelection GetSelection() const { return Selection; }
+	const TArray<FString>& GetExtraNames() const { return ExtraNames; }
+	FName GetAttributeName() const { return AttributeName; }
+	EPCGPointProperties GetPointProperty() const { return PointProperty; }
+	EPCGExtraProperties GetExtraProperty() const { return ExtraProperty; }
 
+
+	// Return the name of the selector.
 	FName GetName() const;
 
 	// Returns the text to display in the widget.
@@ -67,6 +69,91 @@ public:
 
 	// Update the selector with an incoming string.
 	bool Update(FString NewValue);
+
+	template <typename T>
+	static T CreateFromOtherSelector(const FPCGAttributePropertySelector& InOther)
+	{
+		T OutSelector;
+		OutSelector.ImportFromOtherSelector(InOther);
+		return OutSelector;
+	}
+
+	void ImportFromOtherSelector(const FPCGAttributePropertySelector& InOther);
+
+protected:
+	UPROPERTY()
+	EPCGAttributePropertySelection Selection = EPCGAttributePropertySelection::Attribute;
+
+	UPROPERTY()
+	FName AttributeName = NAME_None;
+
+	UPROPERTY()
+	EPCGPointProperties PointProperty = EPCGPointProperties::Position;
+
+	UPROPERTY()
+	EPCGExtraProperties ExtraProperty = EPCGExtraProperties::Index;
+
+	UPROPERTY()
+	TArray<FString> ExtraNames;
+};
+
+/** Struct that will default on @Last (or @LastCreated for previously created selectors). */
+USTRUCT(BlueprintType)
+struct PCG_API FPCGAttributePropertyInputSelector : public FPCGAttributePropertySelector
+{
+	GENERATED_BODY()
+
+	FPCGAttributePropertyInputSelector();
+
+	/** Get a copy of the selector, with @Last replaced by the right selector. */
+	FPCGAttributePropertyInputSelector CopyAndFixLast(const UPCGData* InData) const;
+
+	/** To support previously saved nodes, that used FPCGAttributePropertySelector, we need to define this function to de-serialize the new class using the old. And add a trait (see below). */
+	bool SerializeFromMismatchedTag(const struct FPropertyTag& Tag, FStructuredArchive::FSlot Slot);
+
+	/** For older nodes, before the split between Input and Output, force any last attribute to be last created to preserve the old behavior. Will be called by the PCGSetting deprecation function. Not meant to be used otherwise. */
+	void ApplyDeprecation(int32 InPCGCustomVersion);
+};
+
+/** Struct that will default on @Source. */
+USTRUCT(BlueprintType)
+struct PCG_API FPCGAttributePropertyOutputSelector : public FPCGAttributePropertySelector
+{
+	GENERATED_BODY()
+
+	FPCGAttributePropertyOutputSelector();
+
+	/** Get a copy of the selector, with @Source replaced by the right selector. Can add extra data for specific deprecation cases. */
+	FPCGAttributePropertyOutputSelector CopyAndFixSource(const FPCGAttributePropertyInputSelector* InSourceSelector, const UPCGData* InOptionalData = nullptr) const;
+
+	/** To support previously saved nodes, that used FPCGAttributePropertySelector, we need to define this function to de-serialize the new class using the old. And add a trait (see below). */
+	bool SerializeFromMismatchedTag(const struct FPropertyTag& Tag, FStructuredArchive::FSlot Slot);
+};
+
+// Version where it doesn't make sense to have @Source, alias for FPCGAttributePropertySelector
+USTRUCT(BlueprintType)
+struct PCG_API FPCGAttributePropertyOutputNoSourceSelector : public FPCGAttributePropertySelector
+{
+	GENERATED_BODY()
+};
+
+// Extra trait to specify to the deserializer that those two classes can be deserialized using the old class.
+template<>
+struct TStructOpsTypeTraits<FPCGAttributePropertyInputSelector> : public TStructOpsTypeTraitsBase2<FPCGAttributePropertyInputSelector>
+{
+	enum
+	{
+		WithStructuredSerializeFromMismatchedTag = true,
+	};
+};
+
+template<>
+struct TStructOpsTypeTraits<FPCGAttributePropertyOutputSelector> : public TStructOpsTypeTraitsBase2<FPCGAttributePropertyOutputSelector>
+{
+	enum
+	{
+		WithStructuredSerializeFromMismatchedTag = true,
+	};
 };
 
 /**
@@ -87,6 +174,27 @@ public:
 	UFUNCTION(BlueprintCallable, Category = "PCG|PCGAttributePropertySelector", meta = (ScriptMethod))
 	static bool SetExtraProperty(UPARAM(ref) FPCGAttributePropertySelector& Selector, EPCGExtraProperties InExtraProperty);
 
+	UFUNCTION(BlueprintCallable, Category = "PCG|PCGAttributePropertySelector", meta = (ScriptMethod))
+	static EPCGAttributePropertySelection GetSelection(UPARAM(const, ref) const FPCGAttributePropertySelector& Selector);
+
+	UFUNCTION(BlueprintCallable, Category = "PCG|PCGAttributePropertySelector", meta = (ScriptMethod))
+	static EPCGPointProperties GetPointProperty(UPARAM(const, ref) const FPCGAttributePropertySelector& Selector);
+
+	UFUNCTION(BlueprintCallable, Category = "PCG|PCGAttributePropertySelector", meta = (ScriptMethod))
+	static FName GetAttributeName(UPARAM(const, ref) const FPCGAttributePropertySelector& Selector);
+
+	UFUNCTION(BlueprintCallable, Category = "PCG|PCGAttributePropertySelector", meta = (ScriptMethod))
+	static EPCGExtraProperties GetExtraProperty(UPARAM(const, ref) const FPCGAttributePropertySelector& Selector);
+
+	UFUNCTION(BlueprintCallable, Category = "PCG|PCGAttributePropertySelector", meta = (ScriptMethod))
+	static const TArray<FString>& GetExtraNames(UPARAM(const, ref) const FPCGAttributePropertySelector& Selector);
+
 	UFUNCTION(BlueprintCallable, Category = "PCG|PCGAttributePropertySelector", meta = (ScriptMethod))
 	static FName GetName(UPARAM(const, ref) const FPCGAttributePropertySelector& Selector);
+
+	UFUNCTION(BlueprintCallable, Category = "PCG|PCGAttributePropertySelector", meta = (ScriptMethod))
+	static FPCGAttributePropertyInputSelector CopyAndFixLast(UPARAM(const, ref) const FPCGAttributePropertyInputSelector& Selector, const UPCGData* InData);
+
+	UFUNCTION(BlueprintCallable, Category = "PCG|PCGAttributePropertySelector", meta = (ScriptMethod))
+	static FPCGAttributePropertyOutputSelector CopyAndFixSource(UPARAM(const, ref) const FPCGAttributePropertyOutputSelector& Selector, const FPCGAttributePropertyInputSelector& InSelector);
 };
\ No newline at end of file
diff --git a/Engine/Plugins/Experimental/PCG/Source/PCG/Public/Metadata/PCGMetadata.h b/Engine/Plugins/Experimental/PCG/Source/PCG/Public/Metadata/PCGMetadata.h
--- a/Engine/Plugins/Experimental/PCG/Source/PCG/Public/Metadata/PCGMetadata.h
+++ b/Engine/Plugins/Experimental/PCG/Source/PCG/Public/Metadata/PCGMetadata.h
@@ -296,6 +296,8 @@ protected:
 	void AddAttributeInternal(FName AttributeName, FPCGMetadataAttributeBase* Attribute);
 	void RemoveAttributeInternal(FName AttributeName);
 
+	void SetLastCachedSelectorOnOwner(FName AttributeName);
+
 	UPROPERTY()
 	TObjectPtr<const UPCGMetadata> Parent;
 
@@ -340,27 +342,32 @@ FPCGMetadataAttribute<T>* UPCGMetadata::CreateAttribute(FName AttributeName, con
 
 	FPCGMetadataAttribute<T>* NewAttribute = new FPCGMetadataAttribute<T>(this, AttributeName, ParentAttribute, DefaultValue, bAllowsInterpolation);
 
-	AttributeLock.WriteLock();
-	if (FPCGMetadataAttributeBase** ExistingAttribute = Attributes.Find(AttributeName))
 	{
-		delete NewAttribute;
-		if ((*ExistingAttribute)->GetTypeId() != PCG::Private::MetadataTypes<T>::Id)
+		FWriteScopeLock WriteLock(AttributeLock);
+
+		if (FPCGMetadataAttributeBase** ExistingAttribute = Attributes.Find(AttributeName))
 		{
-			UE_LOG(LogPCG, Error, TEXT("Attribute %s already exists but is not the right type. Abort."), *AttributeName.ToString());
-			return nullptr;
+			delete NewAttribute;
+			if ((*ExistingAttribute)->GetTypeId() != PCG::Private::MetadataTypes<T>::Id)
+			{
+				UE_LOG(LogPCG, Error, TEXT("Attribute %s already exists but is not the right type. Abort."), *AttributeName.ToString());
+				return nullptr;
+			}
+			else
+			{
+				UE_LOG(LogPCG, Warning, TEXT("Attribute %s already exists"), *AttributeName.ToString());
+				NewAttribute = static_cast<FPCGMetadataAttribute<T>*>(*ExistingAttribute);
+			}
 		}
 		else
 		{
-			UE_LOG(LogPCG, Warning, TEXT("Attribute %s already exists"), *AttributeName.ToString());
-			NewAttribute = static_cast<FPCGMetadataAttribute<T>*>(*ExistingAttribute);
+			NewAttribute->AttributeId = NextAttributeId++;
+			AddAttributeInternal(AttributeName, NewAttribute);
+
+			// Also when creating an attribute, notify the PCG Data owner that the latest attribute manipulated is this one.
+			SetLastCachedSelectorOnOwner(AttributeName);
 		}
 	}
-	else
-	{
-		NewAttribute->AttributeId = NextAttributeId++;
-		AddAttributeInternal(AttributeName, NewAttribute);
-	}
-	AttributeLock.WriteUnlock();
 
 	return NewAttribute;
 }
diff --git a/Engine/Plugins/Experimental/PCG/Source/PCG/Public/Metadata/PCGMetadataAttribute.h b/Engine/Plugins/Experimental/PCG/Source/PCG/Public/Metadata/PCGMetadataAttribute.h
--- a/Engine/Plugins/Experimental/PCG/Source/PCG/Public/Metadata/PCGMetadataAttribute.h
+++ b/Engine/Plugins/Experimental/PCG/Source/PCG/Public/Metadata/PCGMetadataAttribute.h
@@ -8,6 +8,14 @@
 
 class UPCGMetadata;
 
+namespace PCGMetadataAttributeConstants
+{
+	const FName LastAttributeName = TEXT("@Last");
+	const FName LastCreatedAttributeName = TEXT("@LastCreated");
+	const FName SourceAttributeName = TEXT("@Source");
+	const FName SourceNameAttributeName = TEXT("@SourceName");
+}
+
 class PCG_API FPCGMetadataAttributeBase
 {
 public:
diff --git a/Engine/Plugins/Experimental/PCG/Source/PCG/Public/PCGCustomVersion.h b/Engine/Plugins/Experimental/PCG/Source/PCG/Public/PCGCustomVersion.h
--- a/Engine/Plugins/Experimental/PCG/Source/PCG/Public/PCGCustomVersion.h
+++ b/Engine/Plugins/Experimental/PCG/Source/PCG/Public/PCGCustomVersion.h
@@ -44,6 +44,9 @@ struct PCG_API FPCGCustomVersion
 		// StaticMeshSpawner now defaults to modify point bounds based on StaticMesh bounds
 		StaticMeshSpawnerApplyMeshBoundsToPointsByDefault = 11,
 
+		// Update of Input Selectors. Previous versions should default on @LastCreated
+		UpdateAttributePropertyInputSelector = 12,
+
 		// -----<new versions can be added above this line>-------------------------------------------------
 		VersionPlusOne,
 		LatestVersion = VersionPlusOne - 1
diff --git a/Engine/Plugins/Experimental/PCG/Source/PCG/Public/PCGData.h b/Engine/Plugins/Experimental/PCG/Source/PCG/Public/PCGData.h
--- a/Engine/Plugins/Experimental/PCG/Source/PCG/Public/PCGData.h
+++ b/Engine/Plugins/Experimental/PCG/Source/PCG/Public/PCGData.h
@@ -4,6 +4,7 @@
 
 #include "PCGCommon.h"
 #include "PCGCrc.h"
+#include "Metadata/PCGAttributePropertySelector.h"
 
 #include "Kismet/BlueprintFunctionLibrary.h"
 #include "Templates/SubclassOf.h"
@@ -11,6 +12,7 @@
 #include "PCGData.generated.h"
 
 class FArchiveCrc32;
+class UPCGMetadata;
 class UPCGNode;
 class UPCGParamData;
 class UPCGSettings;
@@ -44,6 +46,10 @@ public:
 	/** CRC for this object instance. */
 	mutable FPCGCrc Crc;
 
+	virtual bool HasCachedLastSelector() const { return false; }
+	virtual FPCGAttributePropertyInputSelector GetCachedLastSelector() const { return FPCGAttributePropertyInputSelector{}; }
+	virtual void SetLastSelector(const FPCGAttributePropertySelector& InSelector) {};
+
 protected:
 	/** Computes Crc for this and any connected data. */
 	virtual FPCGCrc ComputeCrc(bool bFullDataCrc) const;
diff --git a/Engine/Plugins/Experimental/PCG/Source/PCG/Public/PCGParamData.h b/Engine/Plugins/Experimental/PCG/Source/PCG/Public/PCGParamData.h
--- a/Engine/Plugins/Experimental/PCG/Source/PCG/Public/PCGParamData.h
+++ b/Engine/Plugins/Experimental/PCG/Source/PCG/Public/PCGParamData.h
@@ -21,6 +21,10 @@ public:
 	// ~Begin UPCGData interface
 	virtual EPCGDataType GetDataType() const override { return EPCGDataType::Param; }
 	virtual void AddToCrc(FArchiveCrc32& Ar, bool bFullDataCrc) const override;
+
+	virtual bool HasCachedLastSelector() const override;
+	virtual FPCGAttributePropertyInputSelector GetCachedLastSelector() const override;
+	virtual void SetLastSelector(const FPCGAttributePropertySelector& InSelector) override;
 	// ~End UPCGData interface
 
 	UFUNCTION(BlueprintCallable, Category = Metadata)
@@ -51,4 +55,12 @@ public:
 protected:
 	UPROPERTY()
 	TMap<FName, int64> NameMap;
+
+private:
+	/** Cache to keep track of the latest attribute manipulated on this data. */
+	UPROPERTY()
+	bool bHasCachedLastSelector = false;
+
+	UPROPERTY()
+	FPCGAttributePropertyInputSelector CachedLastSelector;
 };
\ No newline at end of file
diff --git a/Engine/Plugins/Experimental/PCG/Source/PCGEditor/Private/Details/PCGAttributePropertySelectorDetails.cpp b/Engine/Plugins/Experimental/PCG/Source/PCGEditor/Private/Details/PCGAttributePropertySelectorDetails.cpp
--- a/Engine/Plugins/Experimental/PCG/Source/PCGEditor/Private/Details/PCGAttributePropertySelectorDetails.cpp
+++ b/Engine/Plugins/Experimental/PCG/Source/PCGEditor/Private/Details/PCGAttributePropertySelectorDetails.cpp
@@ -3,6 +3,7 @@
 #include "Details/PCGAttributePropertySelectorDetails.h"
 
 #include "Metadata/PCGAttributePropertySelector.h"
+#include "Metadata/PCGMetadataAttribute.h"
 
 #include "DetailLayoutBuilder.h"
 #include "DetailWidgetRow.h"
@@ -119,14 +120,30 @@ TSharedRef<SWidget> FPCGAttributePropertySelectorDetails::GenerateExtraMenu()
 
 	MenuBuilder.BeginSection("Attributes", LOCTEXT("AttributesHeader", "Attributes"));
 	{
-		MenuBuilder.AddMenuEntry(
-			LOCTEXT("LastAttributeHeader", "Last Attribute(None)"),
-			FText(),
-			FSlateIcon(),
-			FExecuteAction::CreateSP(this, &FPCGAttributePropertySelectorDetails::SetAttributeName, FName()),
-			NAME_None,
-			EUserInterfaceActionType::Button
-		);
+		FStructProperty* StructProperty = CastFieldChecked<FStructProperty>(PropertyHandle->GetProperty());
+		if (StructProperty->Struct == FPCGAttributePropertyInputSelector::StaticStruct())
+		{
+			MenuBuilder.AddMenuEntry(
+				LOCTEXT("LastAttributeHeader", "Last Attribute"),
+				LOCTEXT("LastAttributeTooltip", "Refer to the last modified attribute. You can check it in the input inspection data."),
+				FSlateIcon(),
+				FExecuteAction::CreateSP(this, &FPCGAttributePropertySelectorDetails::SetAttributeName, PCGMetadataAttributeConstants::LastAttributeName),
+				NAME_None,
+				EUserInterfaceActionType::Button
+			);
+		}
+
+		if (StructProperty->Struct == FPCGAttributePropertyOutputSelector::StaticStruct())
+		{
+			MenuBuilder.AddMenuEntry(
+				LOCTEXT("SourceAttributeHeader", "Source Attribute"),
+				LOCTEXT("SourceAttributeTooltip", "Refer to the same attribute that will be used as input."),
+				FSlateIcon(),
+				FExecuteAction::CreateSP(this, &FPCGAttributePropertySelectorDetails::SetAttributeName, PCGMetadataAttributeConstants::SourceAttributeName),
+				NAME_None,
+				EUserInterfaceActionType::Button
+			);
+		}
 	}
 	MenuBuilder.EndSection();
 
diff --git a/Engine/Plugins/Experimental/PCG/Source/PCGEditor/Private/PCGEditorModule.cpp b/Engine/Plugins/Experimental/PCG/Source/PCGEditor/Private/PCGEditorModule.cpp
--- a/Engine/Plugins/Experimental/PCG/Source/PCGEditor/Private/PCGEditorModule.cpp
+++ b/Engine/Plugins/Experimental/PCG/Source/PCGEditor/Private/PCGEditorModule.cpp
@@ -81,6 +81,9 @@ void FPCGEditorModule::RegisterDetailsCustomizations()
 	PropertyEditor.RegisterCustomClassLayout("PCGVolume", FOnGetDetailCustomizationInstance::CreateStatic(&FPCGVolumeDetails::MakeInstance));
 
 	PropertyEditor.RegisterCustomPropertyTypeLayout("PCGAttributePropertySelector", FOnGetPropertyTypeCustomizationInstance::CreateStatic(&FPCGAttributePropertySelectorDetails::MakeInstance));
+	PropertyEditor.RegisterCustomPropertyTypeLayout("PCGAttributePropertyInputSelector", FOnGetPropertyTypeCustomizationInstance::CreateStatic(&FPCGAttributePropertySelectorDetails::MakeInstance));
+	PropertyEditor.RegisterCustomPropertyTypeLayout("PCGAttributePropertyOutputSelector", FOnGetPropertyTypeCustomizationInstance::CreateStatic(&FPCGAttributePropertySelectorDetails::MakeInstance));
+	PropertyEditor.RegisterCustomPropertyTypeLayout("PCGAttributePropertyOutputNoSourceSelector", FOnGetPropertyTypeCustomizationInstance::CreateStatic(&FPCGAttributePropertySelectorDetails::MakeInstance));
 	PropertyEditor.RegisterCustomPropertyTypeLayout("PCGOverrideInstancedPropertyBag", FOnGetPropertyTypeCustomizationInstance::CreateStatic(&FPCGOverrideInstancedPropertyBagDetails::MakeInstance));
 
 	PropertyEditor.NotifyCustomizationModuleChanged();
@@ -97,6 +100,9 @@ void FPCGEditorModule::UnregisterDetailsCustomizations()
 		PropertyModule->UnregisterCustomClassLayout("PCGVolume");
 
 		PropertyModule->UnregisterCustomPropertyTypeLayout("PCGAttributePropertySelector");
+		PropertyModule->UnregisterCustomPropertyTypeLayout("PCGAttributePropertyInputSelector");
+		PropertyModule->UnregisterCustomPropertyTypeLayout("PCGAttributePropertyOutputSelector");
+		PropertyModule->UnregisterCustomPropertyTypeLayout("PCGAttributePropertyOutputNoSourceSelector");
 		PropertyModule->UnregisterCustomPropertyTypeLayout("PCGOverrideInstancedPropertyBag");
 
 		PropertyModule->NotifyCustomizationModuleChanged();
diff --git a/Engine/Plugins/Experimental/PCG/Source/PCGEditor/Private/SPCGEditorGraphAttributeListView.cpp b/Engine/Plugins/Experimental/PCG/Source/PCGEditor/Private/SPCGEditorGraphAttributeListView.cpp
--- a/Engine/Plugins/Experimental/PCG/Source/PCGEditor/Private/SPCGEditorGraphAttributeListView.cpp
+++ b/Engine/Plugins/Experimental/PCG/Source/PCGEditor/Private/SPCGEditorGraphAttributeListView.cpp
@@ -36,23 +36,24 @@ static TAutoConsoleVariable<bool> CVarShowAdvancedAttributesFields(
 
 namespace PCGEditorGraphAttributeListView
 {
-	const FString LastMetadataId = TEXT("@Last");
+	const FString NoneAttributeId = TEXT("@None");
 	const FText NoPinAvailableText = LOCTEXT("NoPinAvailableText", "No pins");
+	const FText LastLabelFormat = LOCTEXT("LastLabelFormat", "Last attribute: {0} / {1}");
 	const FText NoDataAvailableText = LOCTEXT("NoDataAvailableText", "No data available");
 	const FText NoNodeInspectedText = LOCTEXT("NoNodeInspectedText", "No node being inspected");
 	const FText NoNodeInspectedToolTip = LOCTEXT("NoNodeInspectedToolTip", "Inspect a node using the right click menu");
 	
 	/** Names of the columns in the attribute list */
 	const FName NAME_IndexColumn = FName(TEXT("IndexColumn"));
-	const FName NAME_PointPositionX = FName(TEXT("$Transform.Location.X"));
-	const FName NAME_PointPositionY = FName(TEXT("$Transform.Location.Y"));
-	const FName NAME_PointPositionZ = FName(TEXT("$Transform.Location.Z"));
-	const FName NAME_PointRotationX = FName(TEXT("$Transform.Rotation.X"));
-	const FName NAME_PointRotationY = FName(TEXT("$Transform.Rotation.Y"));
-	const FName NAME_PointRotationZ = FName(TEXT("$Transform.Rotation.Z"));
-	const FName NAME_PointScaleX = FName(TEXT("$Transform.Scale.X"));
-	const FName NAME_PointScaleY = FName(TEXT("$Transform.Scale.Y"));
-	const FName NAME_PointScaleZ = FName(TEXT("$Transform.Scale.Z"));
+	const FName NAME_PointPositionX = FName(TEXT("$Position.X"));
+	const FName NAME_PointPositionY = FName(TEXT("$Position.Y"));
+	const FName NAME_PointPositionZ = FName(TEXT("$Position.Z"));
+	const FName NAME_PointRotationX = FName(TEXT("$Rotation.X"));
+	const FName NAME_PointRotationY = FName(TEXT("$Rotation.Y"));
+	const FName NAME_PointRotationZ = FName(TEXT("$Rotation.Z"));
+	const FName NAME_PointScaleX = FName(TEXT("$Scale.X"));
+	const FName NAME_PointScaleY = FName(TEXT("$Scale.Y"));
+	const FName NAME_PointScaleZ = FName(TEXT("$Scale.Z"));
 	const FName NAME_PointBoundsMinX = FName(TEXT("$BoundsMin.X"));
 	const FName NAME_PointBoundsMinY = FName(TEXT("$BoundsMin.Y"));
 	const FName NAME_PointBoundsMinZ = FName(TEXT("$BoundsMin.Z"));
@@ -635,6 +636,12 @@ void SPCGEditorGraphAttributeListView::RefreshAttributeList()
 		}
 	}
 
+	if (PCGData && PCGData->HasCachedLastSelector())
+	{
+		const FText LastSelector = PCGData->GetCachedLastSelector().GetDisplayText();
+		InfoTextBlock->SetText(FText::Format(PCGEditorGraphAttributeListView::LastLabelFormat, LastSelector, InfoTextBlock->GetText()));
+	}
+
 	RefreshSorting();
 }
 
@@ -1154,14 +1161,11 @@ void SPCGEditorGraphAttributeListView::OnFilterTextCommitted(const FText& NewTex
 
 void SPCGEditorGraphAttributeListView::AddColumn(const UPCGPointData* InPCGPointData, const FName& InColumnId, const FText& ColumnLabel, EHorizontalAlignment HeaderHAlign, EHorizontalAlignment CellHAlign)
 {
-	const float ColumnWidth = PCGEditorGraphAttributeListView::CalculateColumnWidth(ColumnLabel);
-
 	if (InPCGPointData)
 	{
 		const FString ColumnIdString = InColumnId.ToString();
 		
-		FPCGAttributePropertySelector TargetSelector;
-		TargetSelector.bForceGetNone = true;
+		FPCGAttributePropertyInputSelector TargetSelector;
 		TargetSelector.Update(ColumnIdString);
 	
 		FPCGColumnData& ColumnData = PCGColumnData.Add(InColumnId);
@@ -1190,6 +1194,8 @@ void SPCGEditorGraphAttributeListView::AddColumn(const UPCGPointData* InPCGPoint
 			ColumnData.DataKeys = PCGAttributeAccessorHelpers::CreateConstKeys(InPCGPointData, TargetSelector);
 		}
 	}
+
+	const float ColumnWidth = PCGEditorGraphAttributeListView::CalculateColumnWidth(ColumnLabel);
 	
 	SHeaderRow::FColumn::FArguments Arguments;
 	Arguments.ColumnId(InColumnId);
@@ -1256,16 +1262,15 @@ void SPCGEditorGraphAttributeListView::AddMetadataColumn(const UPCGData* InPCGDa
 	FString ColumnIdString(OriginalColumnIdString);
 	if (InColumnId == NAME_None)
 	{
-		ColumnIdString.ReplaceInline(LexToString(NAME_None), *PCGEditorGraphAttributeListView::LastMetadataId);
-		ColumnLabel = FText::Format(LOCTEXT("LastMetadataLabelFormat", "{0}{1}"), UEnum::GetDisplayValueAsText(InMetadataType), FText::FromString(PostFix));
+		ColumnIdString = PCGEditorGraphAttributeListView::NoneAttributeId;
+		ColumnLabel = FText::Format(LOCTEXT("NoneLabelFormat", "{0}{1}"), UEnum::GetDisplayValueAsText(InMetadataType), FText::FromString(PostFix));
 	}
 	
 	const FName ColumnId(ColumnIdString);
 
 	if (InPCGData)
 	{
-		FPCGAttributePropertySelector TargetSelector;
-		TargetSelector.bForceGetNone = true;
+		FPCGAttributePropertyInputSelector TargetSelector;
 		TargetSelector.Update(OriginalColumnIdString);
 	
 		FPCGColumnData& ColumnData = PCGColumnData.Add(ColumnId);
diff --git a/Engine/Plugins/Experimental/PCGExternalDataInterop/Source/PCGExternalDataInteropEditor/Public/PCGLoadAlembic.h b/Engine/Plugins/Experimental/PCGExternalDataInterop/Source/PCGExternalDataInteropEditor/Public/PCGLoadAlembic.h
--- a/Engine/Plugins/Experimental/PCGExternalDataInterop/Source/PCGExternalDataInteropEditor/Public/PCGLoadAlembic.h
+++ b/Engine/Plugins/Experimental/PCGExternalDataInterop/Source/PCGExternalDataInteropEditor/Public/PCGLoadAlembic.h
@@ -21,7 +21,7 @@ struct PCGEXTERNALDATAINTEROPEDITOR_API FPCGLoadAlembicBPData
 	FAbcConversionSettings ConversionSettings;
 
 	UPROPERTY(BlueprintReadWrite, EditAnywhere, Category = "Settings")
-	TMap<FString, FPCGAttributePropertySelector> AttributeMapping;
+	TMap<FString, FPCGAttributePropertyInputSelector> AttributeMapping;
 };
 
 UCLASS()
