[PCG] Tentative fix for break build in external plugin

#rnx

diff --git a/Engine/Plugins/Experimental/PCG/Source/PCG/Public/Elements/IO/PCGExternalDataContext.h b/Engine/Plugins/Experimental/PCG/Source/PCG/Public/Elements/IO/PCGExternalDataContext.h
--- a/Engine/Plugins/Experimental/PCG/Source/PCG/Public/Elements/IO/PCGExternalDataContext.h
+++ b/Engine/Plugins/Experimental/PCG/Source/PCG/Public/Elements/IO/PCGExternalDataContext.h
@@ -23,8 +23,11 @@ struct FPCGExternalDataContext : public FPCGContext
 		FRowToPointAccessors(FRowToPointAccessors&&) = default;
 		FRowToPointAccessors& operator=(FRowToPointAccessors&&) = default;
 		FRowToPointAccessors& operator=(const FRowToPointAccessors&) = delete;
-		FRowToPointAccessors(TUniquePtr<const IPCGAttributeAccessor> InRow, TUniquePtr<IPCGAttributeAccessor> InPoint, const FPCGAttributePropertyOutputSelector& InSelector)
-			: RowAccessor(MoveTemp(InRow)), PointAccessor(MoveTemp(InPoint)), Selector(InSelector) {}
+		FRowToPointAccessors(TUniquePtr<const IPCGAttributeAccessor> InRow, TUniquePtr<IPCGAttributeAccessor> InPoint, const FPCGAttributePropertySelector& InSelector)
+			: RowAccessor(MoveTemp(InRow)), PointAccessor(MoveTemp(InPoint))
+		{
+			Selector.ImportFromOtherSelector(InSelector);
+		}
 	};
 
 	struct FPointDataAccessorsMapping
