Added a more robust way to get the target actor from data from the PCG blueprint API

diff --git a/Engine/Plugins/Experimental/PCG/Source/PCG/Private/Helpers/PCGBlueprintHelpers.cpp b/Engine/Plugins/Experimental/PCG/Source/PCG/Private/Helpers/PCGBlueprintHelpers.cpp
--- a/Engine/Plugins/Experimental/PCG/Source/PCG/Private/Helpers/PCGBlueprintHelpers.cpp
+++ b/Engine/Plugins/Experimental/PCG/Source/PCG/Private/Helpers/PCGBlueprintHelpers.cpp
@@ -82,6 +82,11 @@ UPCGComponent* UPCGBlueprintHelpers::GetOriginalComponent(FPCGContext& Context)
 	}
 }
 
+AActor* UPCGBlueprintHelpers::GetTargetActor(FPCGContext& Context, UPCGSpatialData* SpatialData)
+{
+	return Context.GetTargetActor(SpatialData);
+}
+
 void UPCGBlueprintHelpers::SetExtents(FPCGPoint& InPoint, const FVector& InExtents)
 {
 	InPoint.SetExtents(InExtents);
diff --git a/Engine/Plugins/Experimental/PCG/Source/PCG/Public/Helpers/PCGBlueprintHelpers.h b/Engine/Plugins/Experimental/PCG/Source/PCG/Public/Helpers/PCGBlueprintHelpers.h
--- a/Engine/Plugins/Experimental/PCG/Source/PCG/Public/Helpers/PCGBlueprintHelpers.h
+++ b/Engine/Plugins/Experimental/PCG/Source/PCG/Public/Helpers/PCGBlueprintHelpers.h
@@ -47,6 +47,9 @@ public:
 	UFUNCTION(BlueprintCallable, Category = "PCG|Temporary", meta = (ScriptMethod))
 	static UPCGComponent* GetOriginalComponent(UPARAM(ref) FPCGContext& Context);
 
+	UFUNCTION(BlueprintCallable, Category = "PCG", meta = (ScriptMethod))
+	static AActor* GetTargetActor(UPARAM(ref) FPCGContext& Context, UPCGSpatialData* SpatialData);
+
 	UFUNCTION(BlueprintCallable, Category = "PCG|Points", meta = (ScriptMethod))
 	static void SetExtents(UPARAM(ref) FPCGPoint& InPoint, const FVector& InExtents);
 
