[PCG] Fixed a bug where the PCG Volume did not properly follow surface normals in the Editor.

Issue was related to bUseSurfaceOrientation flag being set to true. This is usually true for things that should face the surface, such as text, lights, decals, etc.

diff --git a/Engine/Plugins/Experimental/PCG/Source/PCGEditor/Private/PCGVolumeFactory.cpp b/Engine/Plugins/Experimental/PCG/Source/PCGEditor/Private/PCGVolumeFactory.cpp
--- a/Engine/Plugins/Experimental/PCG/Source/PCGEditor/Private/PCGVolumeFactory.cpp
+++ b/Engine/Plugins/Experimental/PCG/Source/PCGEditor/Private/PCGVolumeFactory.cpp
@@ -14,7 +14,7 @@ UPCGVolumeFactory::UPCGVolumeFactory(const FObjectInitializer& ObjectInitializer
 {
 	DisplayName = LOCTEXT("PCGVolumeDisplayName", "PCG Volume");
 	NewActorClass = APCGVolume::StaticClass();
-	bUseSurfaceOrientation = true;
+	bUseSurfaceOrientation = false;
 }
 
 bool UPCGVolumeFactory::CanCreateActorFrom(const FAssetData& AssetData, FText& OutErrorMsg)
