[PCG] Fixing BP compilation issue with overridable params

#rnx

diff --git a/Engine/Plugins/Experimental/PCG/Source/PCG/Private/Elements/PCGExecuteBlueprint.cpp b/Engine/Plugins/Experimental/PCG/Source/PCG/Private/Elements/PCGExecuteBlueprint.cpp
--- a/Engine/Plugins/Experimental/PCG/Source/PCG/Private/Elements/PCGExecuteBlueprint.cpp
+++ b/Engine/Plugins/Experimental/PCG/Source/PCG/Private/Elements/PCGExecuteBlueprint.cpp
@@ -623,15 +623,12 @@ TArray<FPCGSettingsOverridableParam> UPCGBlueprintSettings::GatherOverridablePar
 {
 	TArray<FPCGSettingsOverridableParam> OverridableParams = Super::GatherOverridableParams();
 
-	if (BlueprintElementInstance)
+	if (UClass* BPClass = *BlueprintElementType)
 	{
-		if (UClass* BPClass = BlueprintElementInstance->GetClass())
-		{
-			PCGSettingsHelpers::FPCGGetAllOverridableParamsConfig Config;
-			Config.bExcludeSuperProperties = true;
-			Config.ExcludePropertyFlags = CPF_DisableEditOnInstance | CPF_EditConst | CPF_BlueprintReadOnly;
-			OverridableParams.Append(PCGSettingsHelpers::GetAllOverridableParams(BPClass, Config));
-		}
+		PCGSettingsHelpers::FPCGGetAllOverridableParamsConfig Config;
+		Config.bExcludeSuperProperties = true;
+		Config.ExcludePropertyFlags = CPF_DisableEditOnInstance | CPF_EditConst | CPF_BlueprintReadOnly;
+		OverridableParams.Append(PCGSettingsHelpers::GetAllOverridableParams(BPClass, Config));
 	}
 
 	return OverridableParams;
@@ -642,9 +639,9 @@ void UPCGBlueprintSettings::FixingOverridableParamPropertyClass(FPCGSettingsOver
 {
 	bool bFound = false;
 
-	if (BlueprintElementInstance && !Param.PropertiesNames.IsEmpty())
+	if (!Param.PropertiesNames.IsEmpty())
 	{
-		UClass* BPClass = BlueprintElementInstance->GetClass();
+		UClass* BPClass = *BlueprintElementType;
 		if (BPClass && BPClass->FindPropertyByName(Param.PropertiesNames[0]))
 		{
 			Param.PropertyClass = BPClass;
