Fixed GetGridSize that was hiding method from the base class in the PCGPartitionActor

#rnx

diff --git a/Engine/Plugins/Experimental/PCG/Source/PCG/Private/Grid/PCGPartitionActor.cpp b/Engine/Plugins/Experimental/PCG/Source/PCG/Private/Grid/PCGPartitionActor.cpp
--- a/Engine/Plugins/Experimental/PCG/Source/PCG/Private/Grid/PCGPartitionActor.cpp
+++ b/Engine/Plugins/Experimental/PCG/Source/PCG/Private/Grid/PCGPartitionActor.cpp
@@ -421,7 +421,7 @@ AActor* APCGPartitionActor::GetSceneOutlinerParent() const
 
 void APCGPartitionActor::PostCreation(const FGuid& InGridGUID)
 {
-	PCGGridSize = GetGridSize();
+	PCGGridSize = Super::GetGridSize();
 	PCGGuid = InGridGUID;
 
 	// Put in cache if we use the 2D grid or not.
diff --git a/Engine/Plugins/Experimental/PCG/Source/PCG/Private/Grid/PCGPartitionActor.h b/Engine/Plugins/Experimental/PCG/Source/PCG/Private/Grid/PCGPartitionActor.h
--- a/Engine/Plugins/Experimental/PCG/Source/PCG/Private/Grid/PCGPartitionActor.h
+++ b/Engine/Plugins/Experimental/PCG/Source/PCG/Private/Grid/PCGPartitionActor.h
@@ -52,7 +52,7 @@ public:
 	FBox GetFixedBounds() const;
 	FIntVector GetGridCoord() const;
 
-	uint32 GetGridSize() const { return PCGGridSize; }
+	uint32 GetPCGGridSize() const { return PCGGridSize; }
 	bool IsUsing2DGrid() const { return bUse2DGrid; }
 
 	void AddGraphInstance(UPCGComponent* OriginalComponent);
diff --git a/Engine/Plugins/Experimental/PCG/Source/PCG/Private/PCGActorAndComponentMapping.cpp b/Engine/Plugins/Experimental/PCG/Source/PCG/Private/PCGActorAndComponentMapping.cpp
--- a/Engine/Plugins/Experimental/PCG/Source/PCG/Private/PCGActorAndComponentMapping.cpp
+++ b/Engine/Plugins/Experimental/PCG/Source/PCG/Private/PCGActorAndComponentMapping.cpp
@@ -330,7 +330,7 @@ void UPCGActorAndComponentMapping::RegisterPartitionActor(APCGPartitionActor* Ac
 	{
 		FWriteScopeLock WriteLock(PartitionActorsMapLock);
 
-		TMap<FIntVector, TObjectPtr<APCGPartitionActor>>& PartitionActorsMapGrid = PartitionActorsMap.FindOrAdd(Actor->GetGridSize());
+		TMap<FIntVector, TObjectPtr<APCGPartitionActor>>& PartitionActorsMapGrid = PartitionActorsMap.FindOrAdd(Actor->GetPCGGridSize());
 		if (PartitionActorsMapGrid.Contains(GridCoord))
 		{
 			return;
@@ -376,7 +376,7 @@ void UPCGActorAndComponentMapping::UnregisterPartitionActor(APCGPartitionActor*
 
 	FIntVector GridCoord = Actor->GetGridCoord();
 
-	if (TMap<FIntVector, TObjectPtr<APCGPartitionActor>>* PartitionActorsMapGrid = PartitionActorsMap.Find(Actor->GetGridSize()))
+	if (TMap<FIntVector, TObjectPtr<APCGPartitionActor>>* PartitionActorsMapGrid = PartitionActorsMap.Find(Actor->GetPCGGridSize()))
 	{
 		FWriteScopeLock WriteLock(PartitionActorsMapLock);
 		PartitionActorsMapGrid->Remove(GridCoord);
