Fixed behavior of Get Actor Data Self as Single Point for partition actors, now returns point for PA volume only
Fixed check on ActorFilter that shouldn't have been visible.
Added merge mutiple single point data option in a more efficient manner vs using a Merge node
Added "ignore self" option when getting data from world
Remove limitation of no spaces in attribute names (fixes tag parsing error in the point data code)

diff --git a/Engine/Plugins/Experimental/PCG/Source/PCG/Private/Data/PCGPointData.cpp b/Engine/Plugins/Experimental/PCG/Source/PCG/Private/Data/PCGPointData.cpp
--- a/Engine/Plugins/Experimental/PCG/Source/PCG/Private/Data/PCGPointData.cpp
+++ b/Engine/Plugins/Experimental/PCG/Source/PCG/Private/Data/PCGPointData.cpp
@@ -229,6 +229,11 @@ namespace PCGPointHelpers
 	}
 }
 
+namespace PCGPointDataConstants
+{
+	const FName ActorReferenceAttribute = TEXT("ActorReference");
+}
+
 FPCGPointRef::FPCGPointRef(const FPCGPoint& InPoint)
 {
 	Point = &InPoint;
@@ -383,34 +388,34 @@ void UPCGPointData::SetPoints(const TArray<FPCGPoint>& InPoints)
 }
 
 void UPCGPointData::InitializeFromActor(AActor* InActor)
+{
+	check(InActor);
+	check(Metadata && Metadata->GetAttributeCount() == 0);
+
+	AddSinglePointFromActor(InActor);
+}
+
+void UPCGPointData::AddSinglePointFromActor(AActor* InActor)
 {
 	check(InActor);
 
-	Points.SetNum(1);
-	Points[0].Transform = InActor->GetActorTransform();
+	FPCGPoint& Point = GetMutablePoints().Emplace_GetRef();
+	Point.Transform = InActor->GetActorTransform();
 
-	const FVector& Position = Points[0].Transform.GetLocation();
-	Points[0].Seed = PCGHelpers::ComputeSeed((int)Position.X, (int)Position.Y, (int)Position.Z);
+	const FVector& Position = Point.Transform.GetLocation();
+	Point.Seed = PCGHelpers::ComputeSeed((int)Position.X, (int)Position.Y, (int)Position.Z);
 
 	const FBox LocalBounds = PCGHelpers::GetActorLocalBounds(InActor);
-	Points[0].BoundsMin = LocalBounds.Min;
-	Points[0].BoundsMax = LocalBounds.Max;
+	Point.BoundsMin = LocalBounds.Min;
+	Point.BoundsMax = LocalBounds.Max;
 
-	Metadata = NewObject<UPCGMetadata>(this);
+	Point.MetadataEntry = Metadata->AddEntry();
 
-	const bool bHasTags = (!InActor->Tags.IsEmpty());
-
-	// Create a metadata entry only when we have tags, since all point data initialized this way will always have the actor reference field, 
-	// it is therefore not needed to be careful about default values
-	if (bHasTags)
+	FPCGMetadataAttribute<FString>* ActorReferenceAttribute = Metadata->FindOrCreateAttribute(PCGPointDataConstants::ActorReferenceAttribute, FString(), /*bAllowsInterpolation=*/false, /*bOverrideParent=*/false, /*bOverwriteIfTypeMismatch=*/false);
+	if (ActorReferenceAttribute)
 	{
-		Points[0].MetadataEntry = Metadata->AddEntry();
-	}	
-
-	// TODO: Maybe make that a method somewhere else
-	// NOTE: for booleans, it's not sufficient to set the default value to true, it's just not going to be a good idea when merging point data together
-	// Add a self-reference attribute for use in GetActorData with single point
-	Metadata->CreateAttribute<FString>(TEXT("ActorReference"), FSoftObjectPath(InActor).ToString(), /*bAllowsInterpolation=*/false, /*bOverrideParent=*/false);
+		ActorReferenceAttribute->SetValue(Point.MetadataEntry, FSoftObjectPath(InActor).ToString());
+	}
 
 	// Parse tags as well
 	for (FName Tag : InActor->Tags)
@@ -418,6 +423,7 @@ void UPCGPointData::InitializeFromActor(AActor* InActor)
 		const FString TagString = Tag.ToString();
 		int32 EqualPosition = INDEX_NONE;
 		
+		// Tags that contain a colon will be consider a field:value pair; we'll try to read a number first then default to a string
 		if(TagString.FindChar(':', EqualPosition))
 		{
 			FString LeftSide = TagString.Left(EqualPosition);
@@ -430,21 +436,26 @@ void UPCGPointData::InitializeFromActor(AActor* InActor)
 
 			if(RightSide.IsNumeric())
 			{
-				double RightValue = FCString::Atod(*RightSide);
-
-				FPCGMetadataAttribute<double>* Attribute = Metadata->CreateAttribute<double>(FName(LeftSide), 0.0, /*bAllowsInterpolation=*/false, /*bOverrideParent=*/false);
-				Attribute->SetValue(Points[0].MetadataEntry, RightValue);
+				if (FPCGMetadataAttribute<double>* Attribute = Metadata->FindOrCreateAttribute<double>(FName(LeftSide), 0.0, /*bAllowsInterpolation=*/false, /*bOverrideParent=*/false, /*bOverwriteIfTypeMismatch=*/false))
+				{
+					double RightSideValue = FCString::Atod(*RightSide);
+					Attribute->SetValue(Point.MetadataEntry, RightSideValue);
+				}
 			}
 			else
 			{
-				FPCGMetadataAttribute<FString>* Attribute = Metadata->CreateAttribute<FString>(FName(LeftSide), FString(), false, false);
-				Attribute->SetValue(Points[0].MetadataEntry, RightSide);
+				if (FPCGMetadataAttribute<FString>* Attribute = Metadata->FindOrCreateAttribute<FString>(FName(LeftSide), FString(), /*bAllowsInterpolation=*/false, /*bOverrideParent=*/false, /*bOverwriteIfTypeMismatch=*/false))
+				{
+					Attribute->SetValue(Point.MetadataEntry, RightSide);
+				}
 			}
 		}
-		else
+		else // Otherwise, consider that the tag is a boolean value
 		{
-			FPCGMetadataAttribute<bool>* Attribute = Metadata->CreateAttribute<bool>(Tag, false, /*bAllowsInterpolation=*/false, /*bOverrideParent=*/false);
-			Attribute->SetValue(Points[0].MetadataEntry, true);
+			if (FPCGMetadataAttribute<bool>* Attribute = Metadata->FindOrCreateAttribute<bool>(Tag, false, /*bAllowsInterpolation=*/false, /*bOverrideParent=*/false, /*bOverwriteIfTypeMismatch=*/false))
+			{
+				Attribute->SetValue(Point.MetadataEntry, true);
+			}
 		}
 	}
 }
diff --git a/Engine/Plugins/Experimental/PCG/Source/PCG/Private/Elements/PCGActorSelector.cpp b/Engine/Plugins/Experimental/PCG/Source/PCG/Private/Elements/PCGActorSelector.cpp
--- a/Engine/Plugins/Experimental/PCG/Source/PCG/Private/Elements/PCGActorSelector.cpp
+++ b/Engine/Plugins/Experimental/PCG/Source/PCG/Private/Elements/PCGActorSelector.cpp
@@ -20,13 +20,13 @@ namespace PCGActorSelector
 	}
 
 	// Need to pass a pointer of pointer to the found actor. The lambda will capture this pointer and modify its value when an actor is found.
-	TFunction<bool(AActor*)> GetFilteringFunction(const FPCGActorSelectorSettings& InSettings, const TFunction<bool(const AActor*)>& BoundsCheck, TArray<AActor*>& InFoundActors)
+	TFunction<bool(AActor*)> GetFilteringFunction(const FPCGActorSelectorSettings& InSettings, const TFunction<bool(const AActor*)>& BoundsCheck, const TFunction<bool(const AActor*)>& SelfIgnoreCheck, TArray<AActor*>& InFoundActors)
 	{
 		if (!FilterRequired(InSettings))
 		{
-			return [&InFoundActors, &BoundsCheck](AActor* Actor) -> bool
+			return [&InFoundActors, &BoundsCheck, &SelfIgnoreCheck](AActor* Actor) -> bool
 			{
-				if (BoundsCheck(Actor))
+				if (BoundsCheck(Actor) && SelfIgnoreCheck(Actor))
 				{
 					InFoundActors.Add(Actor);
 				}
@@ -39,9 +39,9 @@ namespace PCGActorSelector
 		switch (InSettings.ActorSelection)
 		{
 		case EPCGActorSelection::ByTag:
-			return[ActorSelectionTag = InSettings.ActorSelectionTag, &InFoundActors, bMultiSelect, &BoundsCheck](AActor* Actor) -> bool
+			return[ActorSelectionTag = InSettings.ActorSelectionTag, &InFoundActors, bMultiSelect, &BoundsCheck, &SelfIgnoreCheck](AActor* Actor) -> bool
 			{
-				if (Actor->ActorHasTag(ActorSelectionTag) && BoundsCheck(Actor))
+				if (Actor->ActorHasTag(ActorSelectionTag) && BoundsCheck(Actor) && SelfIgnoreCheck(Actor))
 				{
 					InFoundActors.Add(Actor);
 					return bMultiSelect;
@@ -51,9 +51,9 @@ namespace PCGActorSelector
 			};
 
 		case EPCGActorSelection::ByClass:
-			return[ActorSelectionClass = InSettings.ActorSelectionClass, &InFoundActors, bMultiSelect, &BoundsCheck](AActor* Actor) -> bool
+			return[ActorSelectionClass = InSettings.ActorSelectionClass, &InFoundActors, bMultiSelect, &BoundsCheck, &SelfIgnoreCheck](AActor* Actor) -> bool
 			{
-				if (Actor->IsA(ActorSelectionClass) && BoundsCheck(Actor))
+				if (Actor->IsA(ActorSelectionClass) && BoundsCheck(Actor) && SelfIgnoreCheck(Actor))
 				{
 					InFoundActors.Add(Actor);
 					return bMultiSelect;
@@ -73,7 +73,7 @@ namespace PCGActorSelector
 		return [](AActor* Actor) -> bool { return false; };
 	}
 
-	TArray<AActor*> FindActors(const FPCGActorSelectorSettings& Settings, const UPCGComponent* InComponent, const TFunction<bool(const AActor*)>& BoundsCheck)
+	TArray<AActor*> FindActors(const FPCGActorSelectorSettings& Settings, const UPCGComponent* InComponent, const TFunction<bool(const AActor*)>& BoundsCheck, const TFunction<bool(const AActor*)>& SelfIgnoreCheck)
 	{
 		TRACE_CPUPROFILER_EVENT_SCOPE(PCGActorSelector::FindActor);
 
@@ -97,7 +97,7 @@ namespace PCGActorSelector
 
 		// We pass FoundActor ref, that will be captured by the FilteringFunction
 		// It will modify the FoundActor pointer to the found actor, if found.
-		TFunction<bool(AActor*)> FilteringFunction = PCGActorSelector::GetFilteringFunction(Settings, BoundsCheck, FoundActors);
+		TFunction<bool(AActor*)> FilteringFunction = PCGActorSelector::GetFilteringFunction(Settings, BoundsCheck, SelfIgnoreCheck, FoundActors);
 
 		// In case of iterating over all actors in the world, call our filtering function and get out.
 		if (Settings.ActorFilter == EPCGActorFilter::AllWorldActors)
@@ -192,13 +192,13 @@ namespace PCGActorSelector
 		return FoundActors;
 	}
 
-	AActor* FindActor(const FPCGActorSelectorSettings& InSettings, UPCGComponent* InComponent, const TFunction<bool(const AActor*)>& BoundsCheck)
+	AActor* FindActor(const FPCGActorSelectorSettings& InSettings, UPCGComponent* InComponent, const TFunction<bool(const AActor*)>& BoundsCheck, const TFunction<bool(const AActor*)>& SelfIgnoreCheck)
 	{
 		// In order to make sure we don't try to select multiple, we'll do a copy of the settings here.
 		FPCGActorSelectorSettings Settings = InSettings;
 		Settings.bSelectMultiple = false;
 
-		TArray<AActor*> Actors = FindActors(Settings, InComponent, BoundsCheck);
+		TArray<AActor*> Actors = FindActors(Settings, InComponent, BoundsCheck, SelfIgnoreCheck);
 		return Actors.IsEmpty() ? nullptr : Actors[0];
 	}
 }
diff --git a/Engine/Plugins/Experimental/PCG/Source/PCG/Private/Elements/PCGDataFromActor.cpp b/Engine/Plugins/Experimental/PCG/Source/PCG/Private/Elements/PCGDataFromActor.cpp
--- a/Engine/Plugins/Experimental/PCG/Source/PCG/Private/Elements/PCGDataFromActor.cpp
+++ b/Engine/Plugins/Experimental/PCG/Source/PCG/Private/Elements/PCGDataFromActor.cpp
@@ -4,9 +4,13 @@
 
 #include "PCGComponent.h"
 #include "PCGSubsystem.h"
+#include "Data/PCGPointData.h"
 #include "Data/PCGSpatialData.h"
+#include "Elements/PCGMergeElement.h"
+#include "Grid/PCGPartitionActor.h"
 #include "Helpers/PCGHelpers.h"
 
+#include "Algo/AnyOf.h"
 #include "GameFramework/Actor.h"
 #include "UObject/Package.h"
 
@@ -147,7 +151,40 @@ bool FPCGDataFromActorElement::ExecuteInternal(FPCGContext* InContext) const
 			};
 		}
 
-		Context->FoundActors = PCGActorSelector::FindActors(Settings->ActorSelector, Context->SourceComponent.Get(), BoundsCheck);
+		TFunction<bool(const AActor*)> SelfIgnoreCheck = [](const AActor*) -> bool { return true; };
+		if (Self && Settings->ActorSelector.bIgnoreSelfAndChildren)
+		{
+			SelfIgnoreCheck = [Self](const AActor* OtherActor) -> bool
+			{
+				// Check if OtherActor is a child of self
+				const AActor* CurrentOtherActor = OtherActor;
+				while (CurrentOtherActor)
+				{
+					if (CurrentOtherActor == Self)
+					{
+						return false;
+					}
+
+					CurrentOtherActor = CurrentOtherActor->GetParentActor();
+				}
+
+				// Check if Self is a child of OtherActor
+				const AActor* CurrentSelfActor = Self;
+				while (CurrentSelfActor)
+				{
+					if (CurrentSelfActor == OtherActor)
+					{
+						return false;
+					}
+
+					CurrentSelfActor = CurrentSelfActor->GetParentActor();
+				}
+
+				return true;
+			};
+		}
+
+		Context->FoundActors = PCGActorSelector::FindActors(Settings->ActorSelector, Context->SourceComponent.Get(), BoundsCheck, SelfIgnoreCheck);
 		Context->bPerformedQuery = true;
 
 		if (Context->FoundActors.IsEmpty())
@@ -222,9 +259,84 @@ void FPCGDataFromActorElement::GatherWaitTasks(AActor* FoundActor, FPCGContext*
 
 void FPCGDataFromActorElement::ProcessActors(FPCGContext* Context, const UPCGDataFromActorSettings* Settings, const TArray<AActor*>& FoundActors) const
 {
-	for (AActor* Actor : FoundActors)
+	// Special case:
+	// If we're asking for single point with the merge single point data, we can do a more efficient process
+	if (Settings->Mode == EPCGGetDataFromActorMode::GetSinglePoint && Settings->bMergeSinglePointData && FoundActors.Num() > 1)
+	{
+		MergeActorsIntoPointData(Context, Settings, FoundActors);
+	}
+	else
+	{
+		for (AActor* Actor : FoundActors)
+		{
+			ProcessActor(Context, Settings, Actor);
+		}
+	}
+}
+
+void FPCGDataFromActorElement::MergeActorsIntoPointData(FPCGContext* Context, const UPCGDataFromActorSettings* Settings, const TArray<AActor*>& FoundActors) const
+{
+	check(Context);
+
+	// At this point in time, the partition actors behave slightly differently, so if we are in the case where
+	// we have one or more partition actors, we'll go through the normal process and do post-processing to merge the point data instead.
+	const bool bContainsPartitionActors = Algo::AnyOf(FoundActors, [](const AActor* Actor) { return Cast<APCGPartitionActor>(Actor) != nullptr; });
+
+	if (!bContainsPartitionActors)
 	{
-		ProcessActor(Context, Settings, Actor);
+		UPCGPointData* Data = NewObject<UPCGPointData>();
+		bool bHasData = false;
+
+		for (AActor* Actor : FoundActors)
+		{
+			if (Actor)
+			{
+				Data->AddSinglePointFromActor(Actor);
+				bHasData = true;
+			}
+		}
+
+		if (bHasData)
+		{
+			FPCGTaggedData& TaggedData = Context->OutputData.TaggedData.Emplace_GetRef();
+			TaggedData.Data = Data;
+		}
+	}
+	else // Stripped-down version of the normal code path with bParseActor = false
+	{
+		FPCGDataCollection DataToMerge;
+		auto DataFilter = [](EPCGDataType InDataType) { return true; };
+		const bool bParseActor = false;
+
+		for (AActor* Actor : FoundActors)
+		{
+			if (Actor)
+			{
+				FPCGDataCollection Collection = UPCGComponent::CreateActorPCGDataCollection(Actor, Context->SourceComponent.Get(), DataFilter, bParseActor);
+				DataToMerge.TaggedData += Collection.TaggedData;
+			}
+		}
+
+		// Perform point data-to-point data merge
+		if (DataToMerge.TaggedData.Num() > 1)
+		{
+			UPCGMergeSettings* MergeSettings = NewObject<UPCGMergeSettings>();
+			FPCGMergeElement MergeElement;
+			FPCGContext* MergeContext = MergeElement.Initialize(DataToMerge, Context->SourceComponent, nullptr);
+			MergeContext->AsyncState.NumAvailableTasks = Context->AsyncState.NumAvailableTasks;
+			MergeContext->InputData.TaggedData.Emplace_GetRef().Data = MergeSettings;
+
+			while (!MergeElement.Execute(MergeContext))
+			{
+			}
+
+			Context->OutputData = MergeContext->OutputData;
+			delete MergeContext;
+		}
+		else if (DataToMerge.TaggedData.Num() == 1)
+		{
+			Context->OutputData.TaggedData = DataToMerge.TaggedData;
+		}
 	}
 }
 
diff --git a/Engine/Plugins/Experimental/PCG/Source/PCG/Private/Elements/PCGPropertyToParamData.cpp b/Engine/Plugins/Experimental/PCG/Source/PCG/Private/Elements/PCGPropertyToParamData.cpp
--- a/Engine/Plugins/Experimental/PCG/Source/PCG/Private/Elements/PCGPropertyToParamData.cpp
+++ b/Engine/Plugins/Experimental/PCG/Source/PCG/Private/Elements/PCGPropertyToParamData.cpp
@@ -125,7 +125,9 @@ bool FPCGPropertyToParamDataElement::ExecuteInternal(FPCGContext* Context) const
 	// First find the actor depending on the selection
 	UPCGComponent* OriginalComponent = UPCGBlueprintHelpers::GetOriginalComponent(*Context);
 	auto NoBoundsCheck = [](const AActor*) -> bool { return true; };
-	AActor* FoundActor = PCGActorSelector::FindActor(Settings->ActorSelector, OriginalComponent, NoBoundsCheck);
+	auto NoSelfIgnoreCheck = [](const AActor*) -> bool { return true; };
+
+	AActor* FoundActor = PCGActorSelector::FindActor(Settings->ActorSelector, OriginalComponent, NoBoundsCheck, NoSelfIgnoreCheck);
 
 	if (!FoundActor)
 	{
diff --git a/Engine/Plugins/Experimental/PCG/Source/PCG/Private/PCGComponent.cpp b/Engine/Plugins/Experimental/PCG/Source/PCG/Private/PCGComponent.cpp
--- a/Engine/Plugins/Experimental/PCG/Source/PCG/Private/PCGComponent.cpp
+++ b/Engine/Plugins/Experimental/PCG/Source/PCG/Private/PCGComponent.cpp
@@ -1999,7 +1999,16 @@ FPCGDataCollection UPCGComponent::CreateActorPCGDataCollection(AActor* Actor, co
 	APCGPartitionActor* PartitionActor = InDataFilter(EPCGDataType::Spatial) ? Cast<APCGPartitionActor>(Actor) : nullptr;
 	ALandscapeProxy* LandscapeActor = InDataFilter(EPCGDataType::Landscape) ? Cast<ALandscapeProxy>(Actor) : nullptr;
 	AVolume* VolumeActor = InDataFilter(EPCGDataType::Volume) ? Cast<AVolume>(Actor) : nullptr;
-	if (PartitionActor)
+	if (!bParseActor && InDataFilter(EPCGDataType::Point))
+	{
+		UPCGPointData* Data = NewObject<UPCGPointData>();
+		Data->InitializeFromActor(Actor);
+
+		FPCGTaggedData& TaggedData = Collection.TaggedData.Emplace_GetRef();
+		TaggedData.Data = Data;
+		TaggedData.Tags = ActorTags;
+	}
+	else if (PartitionActor)
 	{
 		check(!Component || Component->GetOwner() == Actor); // Invalid processing otherwise because of the this usage
 
@@ -2026,15 +2035,6 @@ FPCGDataCollection UPCGComponent::CreateActorPCGDataCollection(AActor* Actor, co
 		TaggedData.Data = Data;
 		TaggedData.Tags = ActorTags;
 	}
-	else if (!bParseActor && InDataFilter(EPCGDataType::Point))
-	{
-		UPCGPointData* Data = NewObject<UPCGPointData>();
-		Data->InitializeFromActor(Actor);
-
-		FPCGTaggedData& TaggedData = Collection.TaggedData.Emplace_GetRef();
-		TaggedData.Data = Data;
-		TaggedData.Tags = ActorTags;
-	}
 	else if (VolumeActor)
 	{
 		UPCGVolumeData* Data = NewObject<UPCGVolumeData>();
diff --git a/Engine/Plugins/Experimental/PCG/Source/PCG/Public/Data/PCGPointData.h b/Engine/Plugins/Experimental/PCG/Source/PCG/Public/Data/PCGPointData.h
--- a/Engine/Plugins/Experimental/PCG/Source/PCG/Public/Data/PCGPointData.h
+++ b/Engine/Plugins/Experimental/PCG/Source/PCG/Public/Data/PCGPointData.h
@@ -87,8 +87,12 @@ protected:
 	// ~End UPCGSpatialData interface
 
 public:
+	/** Initializes a single point based on the given actor */
 	void InitializeFromActor(AActor* InActor);
 
+	/** Adds a single point based on the given actor */
+	void AddSinglePointFromActor(AActor* InActor);
+
 	UFUNCTION(BlueprintCallable, Category = SpatialData)
 	const TArray<FPCGPoint>& GetPoints() const { return Points; }
 
diff --git a/Engine/Plugins/Experimental/PCG/Source/PCG/Public/Elements/PCGActorSelector.h b/Engine/Plugins/Experimental/PCG/Source/PCG/Public/Elements/PCGActorSelector.h
--- a/Engine/Plugins/Experimental/PCG/Source/PCG/Public/Elements/PCGActorSelector.h
+++ b/Engine/Plugins/Experimental/PCG/Source/PCG/Public/Elements/PCGActorSelector.h
@@ -40,7 +40,7 @@ struct FPCGActorSelectorSettings
 	GENERATED_BODY()
 
 	/** Which actors to consider. */
-	UPROPERTY(BlueprintReadWrite, EditAnywhere, Category = Settings, meta = (EditCondition = "bShowActorFilter", EditConditionHides))
+	UPROPERTY(BlueprintReadWrite, EditAnywhere, Category = Settings, meta = (EditCondition = "bShowActorFilter", EditConditionHides, HideEditConditionToggle))
 	EPCGActorFilter ActorFilter = EPCGActorFilter::Self;
 
 	UPROPERTY(BlueprintReadWrite, EditAnywhere, Category = Settings, meta = (EditCondition = "ActorFilter==EPCGActorFilter::AllWorldActors", EditConditionHides))
@@ -67,6 +67,10 @@ struct FPCGActorSelectorSettings
 	UPROPERTY(BlueprintReadWrite, EditAnywhere, Category = Settings, meta = (EditCondition = "bShowSelectMultiple && ActorFilter==EPCGActorFilter::AllWorldActors && ActorSelection!=EPCGActorSelection::ByName", EditConditionHides))
 	bool bSelectMultiple = false;
 
+	/** If true, ignores results found from within this actor's hierarchy */
+	UPROPERTY(BlueprintReadWrite, EditAnywhere, Category = Settings, meta = (EditCondition = "ActorFilter==EPCGActorFilter::AllWorldActors", EditConditionHides))
+	bool bIgnoreSelfAndChildren = false;
+
 	// Properties used to hide some fields when used in different contexts
 	UPROPERTY(Transient, meta = (EditCondition = false, EditConditionHides))
 	bool bShowActorFilter = true;
@@ -91,6 +95,6 @@ struct FPCGActorSelectorSettings
 
 namespace PCGActorSelector
 {
-	TArray<AActor*> FindActors(const FPCGActorSelectorSettings& Settings, const UPCGComponent* InComponent, const TFunction<bool(const AActor*)>& BoundsCheck);
-	AActor* FindActor(const FPCGActorSelectorSettings& InSettings, UPCGComponent* InComponent, const TFunction<bool(const AActor*)>& BoundsCheck);
+	TArray<AActor*> FindActors(const FPCGActorSelectorSettings& Settings, const UPCGComponent* InComponent, const TFunction<bool(const AActor*)>& BoundsCheck, const TFunction<bool(const AActor*)>& SelfIgnoreCheck);
+	AActor* FindActor(const FPCGActorSelectorSettings& InSettings, UPCGComponent* InComponent, const TFunction<bool(const AActor*)>& BoundsCheck, const TFunction<bool(const AActor*)>& SelfIgnoreCheck);
 }
diff --git a/Engine/Plugins/Experimental/PCG/Source/PCG/Public/Elements/PCGDataFromActor.h b/Engine/Plugins/Experimental/PCG/Source/PCG/Public/Elements/PCGDataFromActor.h
--- a/Engine/Plugins/Experimental/PCG/Source/PCG/Public/Elements/PCGDataFromActor.h
+++ b/Engine/Plugins/Experimental/PCG/Source/PCG/Public/Elements/PCGDataFromActor.h
@@ -67,6 +67,9 @@ public:
 	UPROPERTY(BlueprintReadWrite, EditAnywhere, Category = Settings, meta = (EditCondition = bDisplayModeSettings, EditConditionHides, HideEditConditionToggle))
 	EPCGGetDataFromActorMode Mode = EPCGGetDataFromActorMode::ParseActorComponents;
 
+	UPROPERTY(BlueprintReadWrite, EditAnywhere, Category = Settings, meta = (EditCondition = "Mode == EPCGGetDataFromActorMode::GetSinglePoint", EditConditionHides))
+	bool bMergeSinglePointData = false;
+
 	// This can be set false by inheriting nodes to hide the 'Mode' property.
 	UPROPERTY(Transient, meta = (EditCondition = false, EditConditionHides))
 	bool bDisplayModeSettings = true;
@@ -103,4 +106,6 @@ protected:
 	void GatherWaitTasks(AActor* FoundActor, FPCGContext* Context, TArray<FPCGTaskId>& OutWaitTasks) const;
 	virtual void ProcessActors(FPCGContext* Context, const UPCGDataFromActorSettings* Settings, const TArray<AActor*>& FoundActors) const;
 	virtual void ProcessActor(FPCGContext* Context, const UPCGDataFromActorSettings* Settings, AActor* FoundActor) const;
+
+	void MergeActorsIntoPointData(FPCGContext* Context, const UPCGDataFromActorSettings* Settings, const TArray<AActor*>& FoundActors) const;
 };
