[PCG] Added new debug object treeview and updated PCG editor layout

diff --git a/Engine/Plugins/Experimental/PCG/Source/PCGEditor/Private/PCGEditor.cpp b/Engine/Plugins/Experimental/PCG/Source/PCGEditor/Private/PCGEditor.cpp
--- a/Engine/Plugins/Experimental/PCG/Source/PCGEditor/Private/PCGEditor.cpp
+++ b/Engine/Plugins/Experimental/PCG/Source/PCGEditor/Private/PCGEditor.cpp
@@ -28,12 +28,13 @@
 #include "PCGEditorSettings.h"
 #include "PCGEditorUtils.h"
 #include "SPCGEditorGraphAttributeListView.h"
+#include "SPCGEditorGraphDebugObjectTree.h"
 #include "SPCGEditorGraphDebugObjectWidget.h"
 #include "SPCGEditorGraphDeterminism.h"
 #include "SPCGEditorGraphFind.h"
+#include "SPCGEditorGraphLogView.h"
 #include "SPCGEditorGraphNodePalette.h"
 #include "SPCGEditorGraphProfilingView.h"
-#include "SPCGEditorGraphLogView.h"
 
 #include "AssetToolsModule.h"
 #include "EdGraphUtilities.h"
@@ -68,6 +69,7 @@ namespace FPCGEditor_private
 	const FName GraphEditorID = FName(TEXT("GraphEditor"));
 	const FName PropertyDetailsID = FName(TEXT("PropertyDetails"));
 	const FName PaletteID = FName(TEXT("Palette"));
+	const FName DebugObjectID = FName(TEXT("DebugObject"));
 	const FName AttributesID = FName(TEXT("Attributes"));
 	const FName FindID = FName(TEXT("Find"));
 	const FName DeterminismID = FName(TEXT("Determinism"));
@@ -99,6 +101,7 @@ void FPCGEditor::Initialize(const EToolkitMode::Type InMode, const TSharedPtr<cl
 
 	GraphEditorWidget = CreateGraphEditorWidget();
 	PaletteWidget = CreatePaletteWidget();
+	DebugObjectWidget = CreateDebugObjectWidget();
 	FindWidget = CreateFindWidget();
 	AttributesWidget = CreateAttributesWidget();
 	DeterminismWidget = CreateDeterminismWidget();
@@ -108,45 +111,57 @@ void FPCGEditor::Initialize(const EToolkitMode::Type InMode, const TSharedPtr<cl
 	BindCommands();
 	RegisterToolbar();
 
-	const TSharedRef<FTabManager::FLayout> StandaloneDefaultLayout = FTabManager::NewLayout("Standalone_PCGGraphEditor_Layout_v0.5")
+	const TSharedRef<FTabManager::FLayout> StandaloneDefaultLayout = FTabManager::NewLayout("Standalone_PCGGraphEditor_Layout_v0.6")
 		->AddArea
 		(
-			FTabManager::NewPrimaryArea()->SetOrientation(Orient_Horizontal)
+			FTabManager::NewPrimaryArea()->SetOrientation(Orient_Vertical)
 			->Split
-			(				
-				FTabManager::NewStack()
-				->SetSizeCoefficient(0.10f)
-				->SetHideTabWell(true)
-				->AddTab(FPCGEditor_private::PaletteID, ETabState::OpenedTab)
+			(
+				FTabManager::NewSplitter()->SetOrientation(Orient_Horizontal)
+				->SetSizeCoefficient(0.72)
+				->Split
+				(
+					FTabManager::NewStack()
+					->SetSizeCoefficient(0.10f)
+					->SetHideTabWell(true)
+					->AddTab(FPCGEditor_private::PaletteID, ETabState::OpenedTab)
+				)
+				->Split
+				(
+					FTabManager::NewStack()
+					->SetSizeCoefficient(0.70)
+					->SetHideTabWell(true)
+					->AddTab(FPCGEditor_private::GraphEditorID, ETabState::OpenedTab)
+				)
+				->Split
+				(
+					FTabManager::NewStack()
+					->SetSizeCoefficient(0.20f)
+					->SetHideTabWell(true)
+					->AddTab(FPCGEditor_private::PropertyDetailsID, ETabState::OpenedTab)
+				)
 			)
 			->Split
 			(
-				FTabManager::NewSplitter()->SetOrientation(Orient_Vertical)
-				->SetSizeCoefficient(0.70f)
+				FTabManager::NewSplitter()->SetOrientation(Orient_Horizontal)
+				->SetSizeCoefficient(0.28)
 				->Split
 				(
 					FTabManager::NewStack()
-					->SetSizeCoefficient(0.72)
+					->SetSizeCoefficient(0.2)
 					->SetHideTabWell(true)
-					->AddTab(FPCGEditor_private::GraphEditorID, ETabState::OpenedTab)
+					->AddTab(FPCGEditor_private::DebugObjectID, ETabState::OpenedTab)
 				)
 				->Split
 				(
 					FTabManager::NewStack()
-					->SetSizeCoefficient(0.28)
+					->SetSizeCoefficient(0.8)
 					->SetHideTabWell(true)
 					->AddTab(FPCGEditor_private::AttributesID, ETabState::OpenedTab)
 					->AddTab(FPCGEditor_private::DeterminismID, ETabState::ClosedTab)
 					->AddTab(FPCGEditor_private::FindID, ETabState::ClosedTab)
 				)
 			)
-			->Split
-			(
-				FTabManager::NewStack()
-				->SetSizeCoefficient(0.20f)
-				->SetHideTabWell(true)
-				->AddTab(FPCGEditor_private::PropertyDetailsID, ETabState::OpenedTab)
-			)
 		);
 
 	const FName PCGGraphEditorAppName = FName(TEXT("PCGEditorApp"));
@@ -232,6 +247,10 @@ void FPCGEditor::RegisterTabSpawners(const TSharedRef<FTabManager>& InTabManager
 	InTabManager->RegisterTabSpawner(FPCGEditor_private::PaletteID, FOnSpawnTab::CreateSP(this, &FPCGEditor::SpawnTab_Palette))
 		.SetDisplayName(LOCTEXT("PaletteTab", "Palette"))
 		.SetGroup(WorkspaceMenuCategoryRef);
+	
+	InTabManager->RegisterTabSpawner(FPCGEditor_private::DebugObjectID, FOnSpawnTab::CreateSP(this, &FPCGEditor::SpawnTab_DebugObject))
+		.SetDisplayName(LOCTEXT("DebugTab", "Debug Object Tree"))
+		.SetGroup(WorkspaceMenuCategoryRef);
 
 	InTabManager->RegisterTabSpawner(FPCGEditor_private::AttributesID, FOnSpawnTab::CreateSP(this, &FPCGEditor::SpawnTab_Attributes))
 		.SetDisplayName(LOCTEXT("AttributesTab", "Attributes"))
@@ -259,6 +278,7 @@ void FPCGEditor::UnregisterTabSpawners(const TSharedRef<class FTabManager>& InTa
 	InTabManager->UnregisterTabSpawner(FPCGEditor_private::GraphEditorID);
 	InTabManager->UnregisterTabSpawner(FPCGEditor_private::PropertyDetailsID);
 	InTabManager->UnregisterTabSpawner(FPCGEditor_private::PaletteID);
+	InTabManager->UnregisterTabSpawner(FPCGEditor_private::DebugObjectID);
 	InTabManager->UnregisterTabSpawner(FPCGEditor_private::AttributesID);
 	InTabManager->UnregisterTabSpawner(FPCGEditor_private::FindID);
 	InTabManager->UnregisterTabSpawner(FPCGEditor_private::DeterminismID);
@@ -2075,6 +2095,11 @@ TSharedRef<SPCGEditorGraphNodePalette> FPCGEditor::CreatePaletteWidget()
 	return SNew(SPCGEditorGraphNodePalette);
 }
 
+TSharedRef<SPCGEditorGraphDebugObjectTree> FPCGEditor::CreateDebugObjectWidget()
+{
+	return SNew(SPCGEditorGraphDebugObjectTree, SharedThis(this));
+}
+
 TSharedRef<SPCGEditorGraphFind> FPCGEditor::CreateFindWidget()
 {
 	return SNew(SPCGEditorGraphFind, SharedThis(this));
@@ -2295,6 +2320,16 @@ TSharedRef<SDockTab> FPCGEditor::SpawnTab_Palette(const FSpawnTabArgs& Args)
 		];
 }
 
+TSharedRef<SDockTab> FPCGEditor::SpawnTab_DebugObject(const FSpawnTabArgs& Args)
+{
+	return SNew(SDockTab)
+		.Label(LOCTEXT("PCGDebugObjectTitle", "Debug Object"))
+		.TabColorScale(GetTabColorScale())
+		[
+			DebugObjectWidget.ToSharedRef()
+		];
+}
+
 TSharedRef<SDockTab> FPCGEditor::SpawnTab_Attributes(const FSpawnTabArgs& Args)
 {
 	return SNew(SDockTab)
diff --git a/Engine/Plugins/Experimental/PCG/Source/PCGEditor/Private/PCGEditor.h b/Engine/Plugins/Experimental/PCG/Source/PCGEditor/Private/PCGEditor.h
--- a/Engine/Plugins/Experimental/PCG/Source/PCGEditor/Private/PCGEditor.h
+++ b/Engine/Plugins/Experimental/PCG/Source/PCGEditor/Private/PCGEditor.h
@@ -17,11 +17,12 @@ class FUICommandList;
 class IDetailsView;
 class SGraphEditor;
 class SPCGEditorGraphAttributeListView;
+class SPCGEditorGraphDebugObjectTree;
 class SPCGEditorGraphDeterminismListView;
 class SPCGEditorGraphFind;
+class SPCGEditorGraphLogView;
 class SPCGEditorGraphNodePalette;
 class SPCGEditorGraphProfilingView;
-class SPCGEditorGraphLogView;
 class UEdGraphNode;
 class UPCGComponent;
 class UPCGEditorGraph;
@@ -220,6 +221,9 @@ private:
 	/** Create new palette widget */
 	TSharedRef<SPCGEditorGraphNodePalette> CreatePaletteWidget();
 
+	/** Create new debug object tree widget */
+	TSharedRef<SPCGEditorGraphDebugObjectTree> CreateDebugObjectWidget();
+
 	/** Create new find widget */
 	TSharedRef<SPCGEditorGraphFind> CreateFindWidget();
 
@@ -270,6 +274,7 @@ private:
 	TSharedRef<SDockTab> SpawnTab_GraphEditor(const FSpawnTabArgs& Args);
 	TSharedRef<SDockTab> SpawnTab_PropertyDetails(const FSpawnTabArgs& Args);
 	TSharedRef<SDockTab> SpawnTab_Palette(const FSpawnTabArgs& Args);
+	TSharedRef<SDockTab> SpawnTab_DebugObject(const FSpawnTabArgs& Args);
 	TSharedRef<SDockTab> SpawnTab_Attributes(const FSpawnTabArgs& Args);
 	TSharedRef<SDockTab> SpawnTab_Find(const FSpawnTabArgs& Args);
 	TSharedRef<SDockTab> SpawnTab_Determinism(const FSpawnTabArgs& Args);
@@ -279,6 +284,7 @@ private:
 	TSharedPtr<SGraphEditor> GraphEditorWidget;
 	TSharedPtr<IDetailsView> PropertyDetailsWidget;
 	TSharedPtr<SPCGEditorGraphNodePalette> PaletteWidget;
+	TSharedPtr<SPCGEditorGraphDebugObjectTree> DebugObjectWidget;
 	TSharedPtr<SPCGEditorGraphFind> FindWidget;
 	TSharedPtr<SPCGEditorGraphAttributeListView> AttributesWidget;
 	TSharedPtr<SPCGEditorGraphDeterminismListView> DeterminismWidget;
diff --git a/Engine/Plugins/Experimental/PCG/Source/PCGEditor/Private/SPCGEditorGraphAttributeListView.cpp b/Engine/Plugins/Experimental/PCG/Source/PCGEditor/Private/SPCGEditorGraphAttributeListView.cpp
--- a/Engine/Plugins/Experimental/PCG/Source/PCGEditor/Private/SPCGEditorGraphAttributeListView.cpp
+++ b/Engine/Plugins/Experimental/PCG/Source/PCGEditor/Private/SPCGEditorGraphAttributeListView.cpp
@@ -462,6 +462,12 @@ void SPCGEditorGraphAttributeListView::OnInspectedComponentChanged(UPCGComponent
 	{
 		PCGComponent->OnPCGGraphGeneratedDelegate.AddSP(this, &SPCGEditorGraphAttributeListView::OnGenerateUpdated);
 		PCGComponent->OnPCGGraphCleanedDelegate.AddSP(this, &SPCGEditorGraphAttributeListView::OnGenerateUpdated);
+
+		// Refresh if PCGComponent is being inspected already since we wont get a refresh after generation
+		if (PCGComponent->IsInspecting())
+		{
+			RequestRefresh();
+		}
 	}
 	else
 	{
diff --git a/Engine/Plugins/Experimental/PCG/Source/PCGEditor/Private/SPCGEditorGraphDebugObjectTree.cpp b/Engine/Plugins/Experimental/PCG/Source/PCGEditor/Private/SPCGEditorGraphDebugObjectTree.cpp
new file mode 100644
--- /dev/null
+++ b/Engine/Plugins/Experimental/PCG/Source/PCGEditor/Private/SPCGEditorGraphDebugObjectTree.cpp
@@ -0,0 +1,561 @@
+// Copyright Epic Games, Inc. All Rights Reserved.
+
+#include "SPCGEditorGraphDebugObjectTree.h"
+
+#include "PCGComponent.h"
+#include "PCGEditor.h"
+#include "PCGEditorGraph.h"
+#include "PCGEditorModule.h"
+#include "PCGGraph.h"
+
+#include "PropertyCustomizationHelpers.h"
+#include "Selection.h"
+#include "UnrealEdGlobals.h"
+#include "Algo/AnyOf.h"
+#include "Editor/UnrealEdEngine.h"
+#include "Styling/SlateIconFinder.h"
+#include "Widgets/Images/SImage.h"
+#include "Widgets/Layout/SScrollBar.h"
+#include "Widgets/Layout/SScrollBox.h"
+
+#define LOCTEXT_NAMESPACE "SPCGEditorGraphDebugObjectTree"
+
+void FPCGEditorGraphDebugObjectItem::AddChild(TSharedRef<FPCGEditorGraphDebugObjectItem> InChild)
+{
+	check(!Children.Contains(InChild));
+	InChild->Parent = AsShared();
+	Children.Add(MoveTemp(InChild));
+}
+
+const TSet<TSharedPtr<FPCGEditorGraphDebugObjectItem>>& FPCGEditorGraphDebugObjectItem::GetChildren() const
+{
+	return Children;
+}
+
+FPCGEditorGraphDebugObjectItemPtr FPCGEditorGraphDebugObjectItem::GetParent() const
+{
+	return Parent.Pin();
+}
+
+FString FPCGEditorGraphDebugObjectItem_Actor::GetLabel() const
+{
+	return Actor.IsValid() ? Actor->GetActorNameOrLabel() : FString();
+}
+
+FString FPCGEditorGraphDebugObjectItem_PCGComponent::GetLabel() const
+{
+	return PCGComponent.IsValid() ? PCGComponent->GetName() : FString();
+}
+
+FString FPCGEditorGraphDebugObjectItem_PCGGraph::GetLabel() const
+{
+	return PCGGraph.IsValid() ? PCGGraph->GetName() : FString();
+}
+FString FPCGEditorGraphDebugObjectItem_PCGSubgraph::GetLabel() const
+{
+	if (PCGNode.IsValid() && PCGGraph.IsValid())
+	{
+		return PCGNode->GetName() + FString(TEXT(" - ") + PCGGraph->GetName());
+	}
+	
+	return FString();
+}
+
+void SPCGEditorGraphDebugObjectItemRow::Construct(const FArguments& InArgs, const TSharedRef<STableViewBase>& InOwnerTableView, FPCGEditorGraphDebugObjectItemPtr InItem)
+{
+	Item = InItem;
+	
+	ChildSlot
+	[
+		SNew(SHorizontalBox)
+		+SHorizontalBox::Slot()
+		.AutoWidth()
+		.Padding(2.0f, 0.0f)
+		[
+			SNew(SImage)
+			.Visibility(EVisibility::HitTestInvisible)
+			.Image_Lambda([this]()
+			{
+				const UObject* ItemObject = Item->GetObject();
+				return FSlateIconFinder::FindIconBrushForClass(IsValid(ItemObject) ? ItemObject->GetClass() : nullptr); // TODO: IsValid check is done here to prevent crashing due to deleted pcg components, we need to refresh the tree items when world changes, then this might not be needed.
+			})
+		]
+		+SHorizontalBox::Slot()
+		.AutoWidth()
+		[
+			SNew(STextBlock)
+			.Text(FText::FromString(Item->GetLabel()))
+		]
+	];
+}
+
+SPCGEditorGraphDebugObjectTree::~SPCGEditorGraphDebugObjectTree()
+{
+	FCoreUObjectDelegates::OnPreObjectPropertyChanged.RemoveAll(this);
+	FCoreUObjectDelegates::OnObjectPropertyChanged.RemoveAll(this);
+	FCoreUObjectDelegates::OnObjectConstructed.RemoveAll(this);
+}
+
+void SPCGEditorGraphDebugObjectTree::Construct(const FArguments& InArgs, TSharedPtr<FPCGEditor> InPCGEditor)
+{
+	check(InPCGEditor);
+	PCGEditor = InPCGEditor;
+
+	UPCGGraph* PCGGraph = GetPCGGraph();
+	check(PCGGraph);
+	
+	FCoreUObjectDelegates::OnPreObjectPropertyChanged.AddSP(this, &SPCGEditorGraphDebugObjectTree::OnPreObjectPropertyChanged);
+	FCoreUObjectDelegates::OnObjectPropertyChanged.AddSP(this, &SPCGEditorGraphDebugObjectTree::OnObjectPropertyChanged);
+	FCoreUObjectDelegates::OnObjectConstructed.AddSP(this, &SPCGEditorGraphDebugObjectTree::OnObjectConstructed);
+	
+	const TSharedRef<SScrollBar> HorizontalScrollBar = SNew(SScrollBar)
+		.Orientation(Orient_Horizontal)
+		.Thickness(FVector2D(12.0f, 12.0f));
+
+	const TSharedRef<SScrollBar> VerticalScrollBar = SNew(SScrollBar)
+		.Orientation(Orient_Vertical)
+		.Thickness(FVector2D(12.0f, 12.0f));
+	
+	DebugObjectTreeView = SNew(STreeView<FPCGEditorGraphDebugObjectItemPtr>)
+				.TreeItemsSource(&RootItems)
+				.OnGenerateRow(this, &SPCGEditorGraphDebugObjectTree::MakeTreeRowWidget)
+				.OnGetChildren(this, &SPCGEditorGraphDebugObjectTree::OnGetChildren)
+				.OnSelectionChanged(this, &SPCGEditorGraphDebugObjectTree::OnSelectionChanged)
+				.ItemHeight(18)
+				.AllowOverscroll(EAllowOverscroll::No)
+				.ExternalScrollbar(VerticalScrollBar)
+				.ConsumeMouseWheel(EConsumeMouseWheel::Always);
+
+	const TSharedRef<SWidget> SetButton = PropertyCustomizationHelpers::MakeUseSelectedButton(
+		FSimpleDelegate::CreateSP(this, &SPCGEditorGraphDebugObjectTree::SetDebugObjectFromSelection_OnClicked),
+		LOCTEXT("SetDebugObject", "Set debug object from Level Editor selection."),
+		TAttribute<bool>::Create(TAttribute<bool>::FGetter::CreateSP(this, &SPCGEditorGraphDebugObjectTree::IsSetDebugObjectFromSelectionButtonEnabled)));
+
+	const TSharedRef<SWidget> BrowseButton = PropertyCustomizationHelpers::MakeBrowseButton(
+		FSimpleDelegate::CreateSP(this, &SPCGEditorGraphDebugObjectTree::SelectedDebugObject_OnClicked),
+		LOCTEXT("DebugSelectActor", "Select and frame the debug actor in the Level Editor."),
+		TAttribute<bool>::Create(TAttribute<bool>::FGetter::CreateSP(this, &SPCGEditorGraphDebugObjectTree::IsSelectDebugObjectButtonEnabled))
+);
+	
+	ChildSlot
+	[
+		SNew(SVerticalBox)
+		+SVerticalBox::Slot()
+		.AutoHeight()
+		.HAlign(HAlign_Right)
+		[
+			SNew(SHorizontalBox)
+			+SHorizontalBox::Slot()
+			.AutoWidth()
+			.VAlign(VAlign_Center)
+			.Padding(4.0f)
+			[
+				SetButton
+			]
+			+SHorizontalBox::Slot()
+			.AutoWidth()
+			.VAlign(VAlign_Center)
+			.Padding(4.0f)
+			[
+				BrowseButton
+			]
+		]
+		+SVerticalBox::Slot()
+		.FillHeight(1.0f)
+		[
+			SNew(SHorizontalBox)
+			+SHorizontalBox::Slot()
+			.FillWidth(1.0f)
+			[
+				SNew(SScrollBox)
+				.Orientation(Orient_Horizontal)
+				.ExternalScrollbar(HorizontalScrollBar)
+				+SScrollBox::Slot()
+				.FillSize(1.0f)
+				[
+					DebugObjectTreeView.ToSharedRef()
+				]
+			]
+			+SHorizontalBox::Slot()
+			.AutoWidth()
+			[
+				VerticalScrollBar
+			]
+		]
+		+SVerticalBox::Slot()
+		.AutoHeight()
+		[
+			SNew(SHorizontalBox)
+			+SHorizontalBox::Slot()
+			[
+				HorizontalScrollBar
+			]
+		]
+	];
+
+	RequestRefresh();
+}
+
+void SPCGEditorGraphDebugObjectTree::Tick(const FGeometry& AllottedGeometry, const double InCurrentTime, const float InDeltaTime)
+{
+	SCompoundWidget::Tick(AllottedGeometry, InCurrentTime, InDeltaTime);
+
+	if (bNeedsRefresh)
+	{
+		bNeedsRefresh = false;
+		RefreshTree();
+	}
+}
+
+void SPCGEditorGraphDebugObjectTree::SelectedDebugObject_OnClicked() const
+{
+	if (UPCGComponent* PCGComponent = PCGEditor.Pin()->GetPCGComponentBeingInspected())
+	{
+		AActor* Actor = PCGComponent->GetOwner();
+		if (Actor && GEditor && GUnrealEd)
+		{
+			GEditor->SelectNone(/*bNoteSelectionChange=*/false, /*bDeselectBSPSurfs=*/true, /*WarnAboutManyActors=*/false);
+			GEditor->SelectActor(Actor, /*bInSelected=*/true, /*bNotify=*/true, /*bSelectEvenIfHidden=*/true);
+			GUnrealEd->Exec(Actor->GetWorld(), TEXT("CAMERA ALIGN ACTIVEVIEWPORTONLY"));
+			GEditor->SelectComponent(PCGComponent, /*bInSelected=*/true, /*bNotify=*/true, /*bSelectEvenIfHidden=*/true);
+		}
+	}
+}
+
+bool SPCGEditorGraphDebugObjectTree::IsSelectDebugObjectButtonEnabled() const
+{
+	return PCGEditor.IsValid() && (PCGEditor.Pin()->GetPCGComponentBeingInspected() != nullptr);
+}
+
+void SPCGEditorGraphDebugObjectTree::SetDebugObjectFromSelection_OnClicked()
+{
+	const UPCGGraph* PCGGraph = GetPCGGraph();
+	if (!PCGGraph)
+	{
+		return;
+	}
+
+	USelection* SelectedActors = GEditor ? GEditor->GetSelectedActors() : nullptr;
+	if (!IsValid(SelectedActors))
+	{
+		return;
+	}
+
+	for (FSelectionIterator It(*SelectedActors); It; ++It)
+	{
+		const AActor* SelectedActor = Cast<AActor>(*It);
+		if (!IsValid(SelectedActor))
+		{
+			continue;
+		}
+		
+		TArray<const UPCGComponent*> PCGComponents;
+		SelectedActor->GetComponents<const UPCGComponent>(PCGComponents, /*bIncludeFromChildActors=*/true);
+
+		for (const UPCGComponent* PCGComponent : PCGComponents)
+		{
+			if (!IsValid(PCGComponent))
+			{
+				continue;
+			}
+
+			FPCGStackContext StackContext = FPCGStackContext::CreateStackContextFromGraph(PCGComponent->GetGraph());
+
+			for (const FPCGStack& Stack : StackContext.GetStacks())
+			{
+				FPCGEditorGraphDebugObjectItemPtr* DebugObjectItem = AllGraphItems.FindByPredicate([&PCGComponent, &Stack](const FPCGEditorGraphDebugObjectItemPtr Item)
+				{
+					return Item->GetPCGComponent() == PCGComponent && *Item->GetPCGStack() == Stack;
+				});
+			
+				if (DebugObjectItem)
+				{
+					DebugObjectTreeView->SetSingleExpandedItem(*DebugObjectItem);
+				
+					FPCGEditorGraphDebugObjectItemPtr Parent = DebugObjectItem->Get()->GetParent();
+					while (Parent)
+					{
+						DebugObjectTreeView->SetItemExpansion(Parent, true);
+						Parent = Parent->GetParent();
+					}
+				
+					DebugObjectTreeView->SetSelection(*DebugObjectItem);
+					DebugObjectTreeView->RequestScrollIntoView(*DebugObjectItem);
+					break;
+				}
+			}		
+		}
+	}
+}
+
+bool SPCGEditorGraphDebugObjectTree::IsSetDebugObjectFromSelectionButtonEnabled() const
+{
+	const UPCGGraph* PCGGraph = GetPCGGraph();
+	if (!PCGGraph)
+	{
+		return false;
+	}
+
+	USelection* SelectedActors = GEditor->GetSelectedActors();
+	if (!IsValid(SelectedActors))
+	{
+		return false;
+	}
+
+	for (FSelectionIterator It(*SelectedActors); It; ++It)
+	{
+		const AActor* SelectedActor = Cast<AActor>(*It);
+		if (!IsValid(SelectedActor))
+		{
+			continue;
+		}
+		
+		TArray<const UPCGComponent*> PCGComponents;
+		SelectedActor->GetComponents<const UPCGComponent>(PCGComponents, /*bIncludeFromChildActors=*/true);
+
+		for (const UPCGComponent* PCGComponent : PCGComponents)
+		{
+			if (!IsValid(PCGComponent))
+			{
+				continue;
+			}
+			
+			FPCGStackContext StackContext = FPCGStackContext::CreateStackContextFromGraph(PCGComponent->GetGraph());
+
+			const bool bGraphFound = Algo::AnyOf(StackContext.GetStacks(), [&PCGGraph](const FPCGStack& InStack)
+			{
+				return Cast<const UPCGGraph>(InStack.GetStackFrames().Top().Object) == PCGGraph;
+			});
+			
+			if(bGraphFound)
+			{
+				return true;
+			}
+		}
+	}
+
+	return false;
+}
+
+void SPCGEditorGraphDebugObjectTree::RefreshTree()
+{
+	RootItems.Empty();
+	AllGraphItems.Empty();
+	DebugObjectTreeView->RequestTreeRefresh();	
+
+	const UPCGGraph* PCGGraph = GetPCGGraph();
+	if (!PCGGraph)
+	{
+		return;
+	}
+
+	TArray<UObject*> PCGComponents;
+	GetObjectsOfClass(UPCGComponent::StaticClass(), PCGComponents, /*bIncludeDerivedClasses=*/ true);
+
+	TMap<AActor*, TSharedPtr<FPCGEditorGraphDebugObjectItem_Actor>> ActorItems;
+	TMap<UPCGComponent*, TSharedPtr<FPCGEditorGraphDebugObjectItem_PCGComponent>> ComponentItems;
+	
+	for (UObject* PCGComponentObject : PCGComponents)
+	{
+		if (!IsValid(PCGComponentObject))
+		{
+			continue;
+		}
+
+		UPCGComponent* PCGComponent = Cast<UPCGComponent>(PCGComponentObject);
+		if (!PCGComponent)
+		{
+			continue;
+		}
+
+		AActor* Actor = PCGComponent->GetOwner();
+		if (!Actor)
+		{
+			continue;
+		}
+
+		const UPCGGraph* PCGComponentGraph = PCGComponent->GetGraph();
+		if (!PCGComponentGraph)
+		{
+			continue;
+		}
+		
+		FPCGStackContext StackContext = FPCGStackContext::CreateStackContextFromGraph(PCGComponentGraph);
+
+		TMap<const UPCGGraph*, FPCGEditorGraphDebugObjectItemPtr> GraphItems;
+		TArray<FPCGEditorGraphDebugObjectItemPtr> NodeItems;
+
+		for (const FPCGStack& Stack : StackContext.GetStacks())
+		{
+			const UPCGGraph* TopStackGraph = Cast<const UPCGGraph>(Stack.GetStackFrames().Top().Object);
+			if (TopStackGraph == PCGGraph)
+			{
+				TSharedPtr<FPCGEditorGraphDebugObjectItem_Actor> ActorItem;
+				if (TSharedPtr<FPCGEditorGraphDebugObjectItem_Actor>* FoundActorItem = ActorItems.Find(Actor))
+				{
+					ActorItem = *FoundActorItem;
+				}
+				else
+				{
+					ActorItem = ActorItems.Emplace(Actor, MakeShared<FPCGEditorGraphDebugObjectItem_Actor>(Actor));
+				}
+				
+				TSharedPtr<FPCGEditorGraphDebugObjectItem_PCGComponent> ComponentItem;
+				if (TSharedPtr<FPCGEditorGraphDebugObjectItem_PCGComponent>* FoundComponentItem = ComponentItems.Find(PCGComponent))
+				{
+					ComponentItem = *FoundComponentItem;
+				}
+				else
+				{
+					ComponentItem = ComponentItems.Emplace(PCGComponent, MakeShared<FPCGEditorGraphDebugObjectItem_PCGComponent>(PCGComponent));
+				}
+				
+				if (!ComponentItem->GetParent())
+				{
+					ActorItem->AddChild(ComponentItem.ToSharedRef());	
+				}
+
+				const TArray<FPCGStackFrame>& StackFrames = Stack.GetStackFrames();
+				
+				FString StackPath;
+				Stack.CreateStackFramePath(StackPath);
+					
+				FPCGEditorGraphDebugObjectItemPtr PreviousItem = ComponentItem;
+				for (int32 StackIndex = 0; StackIndex < StackFrames.Num(); StackIndex++)
+				{	
+					const FPCGStackFrame& StackFrame = StackFrames[StackIndex];
+					FPCGEditorGraphDebugObjectItemPtr CurrentItem;
+					if (const UPCGGraph* StackGraph = Cast<const UPCGGraph>(StackFrame.Object))
+					{
+						if (FPCGEditorGraphDebugObjectItemPtr* GraphItem = GraphItems.Find(StackGraph))
+						{
+							CurrentItem = *GraphItem;
+						}
+						else
+						{
+							CurrentItem = GraphItems.Emplace(StackGraph, MakeShared<FPCGEditorGraphDebugObjectItem_PCGGraph>(StackGraph, StackGraph == PCGGraph ? Stack : FPCGStack()));
+						}
+					}
+					else if (const UPCGNode* StackNode = Cast<const UPCGNode>(StackFrame.Object))
+					{
+						const int32 NextStackIndex = StackIndex + 1;
+						if (StackFrames.IsValidIndex(NextStackIndex))
+						{
+							const FPCGStackFrame& NextStackFrame = StackFrames[NextStackIndex];
+							if (const UPCGGraph* NextStackGraph = Cast<const UPCGGraph>(NextStackFrame.Object))
+							{
+								const FPCGEditorGraphDebugObjectItemPtr* NodeItem = NodeItems.FindByPredicate([&StackNode,&Stack](const FPCGEditorGraphDebugObjectItemPtr& InItem)
+								{
+									return (InItem->GetObject() == StackNode) && (*InItem->GetPCGStack() == Stack);
+								});
+
+								if (NodeItem)
+								{
+									CurrentItem = *NodeItem;
+								}
+								else
+								{
+									CurrentItem = NodeItems.Add_GetRef(MakeShared<FPCGEditorGraphDebugObjectItem_PCGSubgraph>(StackNode, NextStackGraph, NextStackGraph == TopStackGraph ? Stack : FPCGStack()));
+								}
+								
+								StackIndex++;
+							}
+						}
+					}
+					
+					if (!CurrentItem->GetParent())
+					{
+						PreviousItem->AddChild(CurrentItem.ToSharedRef());	
+					}
+
+					AllGraphItems.Add(CurrentItem);
+					PreviousItem = CurrentItem;
+				}
+			}
+		}
+	}
+	
+	for (TPair<AActor*, TSharedPtr<FPCGEditorGraphDebugObjectItem_Actor>>& ActorItem : ActorItems)
+	{	
+		RootItems.Add(MoveTemp(ActorItem.Value));
+	}
+}
+
+void SPCGEditorGraphDebugObjectTree::OnPreObjectPropertyChanged(UObject* InObject, const FEditPropertyChain& /*InPropertyChain*/)
+{
+	if (UPCGGraphInstance* PCGgraphInstance = Cast<UPCGGraphInstance>(InObject))
+	{
+		const UPCGGraph* PCGGraph = PCGgraphInstance->GetGraph();
+		if (PCGGraph && PCGGraph == GetPCGGraph())
+		{
+			RequestRefresh();
+		}
+	}
+}
+
+void SPCGEditorGraphDebugObjectTree::OnObjectPropertyChanged(UObject* InObject, FPropertyChangedEvent& InPropertyChangedEvent)
+{	
+	if (const UPCGGraphInterface* PCGGraphInterface = Cast<UPCGGraphInterface>(InObject))
+	{
+		const UPCGGraph* PCGGraph = PCGGraphInterface->GetGraph();
+		if (PCGGraph && PCGGraph == GetPCGGraph())
+		{
+			RequestRefresh();
+		}
+	}
+}
+
+void SPCGEditorGraphDebugObjectTree::OnObjectConstructed(UObject* InObject)
+{
+	if (const UPCGComponent* PCGComponent = Cast<UPCGComponent>(InObject))
+	{
+		RequestRefresh();
+	}
+}
+
+UPCGGraph* SPCGEditorGraphDebugObjectTree::GetPCGGraph() const
+{
+	const TSharedPtr<FPCGEditor> PCGEditorPtr = PCGEditor.Pin();
+	const UPCGEditorGraph* PCGEditorGraph = PCGEditorPtr.IsValid() ? PCGEditorPtr->GetPCGEditorGraph() : nullptr;
+
+	return PCGEditorGraph ? PCGEditorGraph->GetPCGGraph() : nullptr;
+}
+
+TSharedRef<ITableRow> SPCGEditorGraphDebugObjectTree::MakeTreeRowWidget(FPCGEditorGraphDebugObjectItemPtr InItem, const TSharedRef<STableViewBase>& InOwnerTable) const
+{
+	return SNew( STableRow<TSharedPtr<SPCGEditorGraphDebugObjectItemRow>>, InOwnerTable)
+		[
+			SNew(SPCGEditorGraphDebugObjectItemRow, InOwnerTable, InItem)
+		];
+}
+
+void SPCGEditorGraphDebugObjectTree::OnGetChildren(FPCGEditorGraphDebugObjectItemPtr InItem, TArray<FPCGEditorGraphDebugObjectItemPtr>& OutChildren) const
+{
+	if (InItem.IsValid())
+	{
+		for (TWeakPtr<FPCGEditorGraphDebugObjectItem> ChildItem : InItem->GetChildren())
+		{
+			FPCGEditorGraphDebugObjectItemPtr ChildItemPtr = ChildItem.Pin();
+			OutChildren.Add(ChildItemPtr);
+		}
+	}
+}
+
+void SPCGEditorGraphDebugObjectTree::OnSelectionChanged(FPCGEditorGraphDebugObjectItemPtr InItem, ESelectInfo::Type InSelectInfo) const
+{
+	if (!InItem)
+	{
+		PCGEditor.Pin()->SetComponentAndStackBeingInspected(nullptr, FPCGStack());
+		return;
+	}
+	
+	if (UPCGComponent* PCGComponent = InItem->GetPCGComponent())
+	{
+		if (const FPCGStack* PCGStack = InItem->GetPCGStack())
+		{
+			PCGEditor.Pin()->SetComponentAndStackBeingInspected(PCGComponent, *PCGStack);
+		}
+	}
+}
+
+#undef LOCTEXT_NAMESPACE
diff --git a/Engine/Plugins/Experimental/PCG/Source/PCGEditor/Private/SPCGEditorGraphDebugObjectTree.h b/Engine/Plugins/Experimental/PCG/Source/PCGEditor/Private/SPCGEditorGraphDebugObjectTree.h
new file mode 100644
--- /dev/null
+++ b/Engine/Plugins/Experimental/PCG/Source/PCGEditor/Private/SPCGEditorGraphDebugObjectTree.h
@@ -0,0 +1,160 @@
+// Copyright Epic Games, Inc. All Rights Reserved.
+
+#pragma once
+
+#include "PCGComponent.h"
+#include "PCGGraph.h"
+#include "PCGNode.h"
+#include "Graph/PCGStackContext.h"
+
+#include "ToolMenus.h"
+#include "GameFramework/Actor.h"
+#include "Widgets/SCompoundWidget.h"
+#include "Widgets/Views/STreeView.h"
+
+class FPCGEditorGraphDebugObjectItem;
+class FPCGEditor;
+
+typedef TSharedPtr<FPCGEditorGraphDebugObjectItem> FPCGEditorGraphDebugObjectItemPtr;
+
+class FPCGEditorGraphDebugObjectItem : public TSharedFromThis<FPCGEditorGraphDebugObjectItem>
+{
+public:
+	FPCGEditorGraphDebugObjectItem() = default;
+	virtual ~FPCGEditorGraphDebugObjectItem() = default;
+
+	void AddChild(TSharedRef<FPCGEditorGraphDebugObjectItem> InChild);
+	const TSet<TSharedPtr<FPCGEditorGraphDebugObjectItem>>& GetChildren() const;
+	FPCGEditorGraphDebugObjectItemPtr GetParent() const;
+
+	virtual FString GetLabel() const = 0;	
+	virtual UPCGComponent* GetPCGComponent() const = 0;
+	virtual const FPCGStack* GetPCGStack() const { return nullptr; }
+	virtual const UObject* GetObject() const  = 0;
+	
+protected:
+	TWeakPtr<FPCGEditorGraphDebugObjectItem> Parent;
+	TSet<TSharedPtr<FPCGEditorGraphDebugObjectItem>> Children;
+};
+
+class FPCGEditorGraphDebugObjectItem_Actor : public FPCGEditorGraphDebugObjectItem
+{
+public:
+	FPCGEditorGraphDebugObjectItem_Actor(TWeakObjectPtr<AActor> InActor)
+		: Actor(InActor)
+	{ }
+
+	virtual FString GetLabel() const override;
+	virtual UPCGComponent* GetPCGComponent() const override { return nullptr; }
+	virtual const UObject* GetObject() const override { return Actor.Get(); };
+	
+protected:
+	TWeakObjectPtr<AActor> Actor = nullptr;
+};
+
+class FPCGEditorGraphDebugObjectItem_PCGComponent : public FPCGEditorGraphDebugObjectItem
+{
+public:
+	FPCGEditorGraphDebugObjectItem_PCGComponent(TWeakObjectPtr<UPCGComponent> InPCGComponent)
+		: PCGComponent(InPCGComponent)
+	{}
+
+	virtual FString GetLabel() const override;
+	virtual UPCGComponent* GetPCGComponent() const override { return PCGComponent.Get(); }
+	virtual const UObject* GetObject() const override { return PCGComponent.Get(); };
+	
+protected:
+	TWeakObjectPtr<UPCGComponent> PCGComponent = nullptr;
+};
+
+class FPCGEditorGraphDebugObjectItem_PCGGraph : public FPCGEditorGraphDebugObjectItem
+{
+public:
+	FPCGEditorGraphDebugObjectItem_PCGGraph(TWeakObjectPtr<const UPCGGraph> InPCGGraph, const FPCGStack& InPCGStack)
+		: PCGGraph(InPCGGraph)
+		, PCGStack(InPCGStack)
+	{}
+
+	virtual FString GetLabel() const override;
+	virtual UPCGComponent* GetPCGComponent() const override { return GetParent() ? GetParent()->GetPCGComponent() : nullptr; }
+	virtual const UObject* GetObject() const override { return PCGGraph.Get(); };
+	virtual const FPCGStack* GetPCGStack() const override { return &PCGStack; }
+	
+protected:
+	TWeakObjectPtr<const UPCGGraph> PCGGraph = nullptr;
+	FPCGStack PCGStack;
+};
+
+class FPCGEditorGraphDebugObjectItem_PCGSubgraph : public FPCGEditorGraphDebugObjectItem
+{
+public:
+	FPCGEditorGraphDebugObjectItem_PCGSubgraph(TWeakObjectPtr<const UPCGNode> InPCGNode, TWeakObjectPtr<const UPCGGraph> InPCGGraph, const FPCGStack& InPCGStack)
+		: PCGNode(InPCGNode)
+		, PCGGraph(InPCGGraph)
+		, PCGStack(InPCGStack)
+	{}
+
+	virtual FString GetLabel() const override;
+	virtual UPCGComponent* GetPCGComponent() const override { return GetParent() ? GetParent()->GetPCGComponent() : nullptr; }
+	virtual const UObject* GetObject() const override { return PCGNode.Get(); };
+	virtual const FPCGStack* GetPCGStack() const override { return &PCGStack; }
+	
+protected:
+	TWeakObjectPtr<const UPCGNode> PCGNode = nullptr;
+	TWeakObjectPtr<const UPCGGraph> PCGGraph = nullptr;
+	FPCGStack PCGStack;
+};
+
+class SPCGEditorGraphDebugObjectItemRow : public SCompoundWidget
+{
+public:
+	SLATE_BEGIN_ARGS(SPCGEditorGraphDebugObjectItemRow) {}
+	SLATE_END_ARGS()
+
+	void Construct(const FArguments& InArgs, const TSharedRef<STableViewBase>& InOwnerTableView, FPCGEditorGraphDebugObjectItemPtr InItem);
+	
+private:
+	FPCGEditorGraphDebugObjectItemPtr Item;
+};
+
+class SPCGEditorGraphDebugObjectTree : public SCompoundWidget
+{
+public:
+	SLATE_BEGIN_ARGS(SPCGEditorGraphDebugObjectTree) {}
+	SLATE_END_ARGS()
+
+	virtual ~SPCGEditorGraphDebugObjectTree();
+	
+	void Construct(const FArguments& InArgs, TSharedPtr<FPCGEditor> InPCGEditor);
+
+	virtual void Tick(const FGeometry& AllottedGeometry, const double InCurrentTime, const float InDeltaTime) override;
+	
+	void RequestRefresh() { bNeedsRefresh = true; }
+	
+private:
+	void SelectedDebugObject_OnClicked() const;
+	bool IsSelectDebugObjectButtonEnabled() const;
+	
+	void SetDebugObjectFromSelection_OnClicked();
+	bool IsSetDebugObjectFromSelectionButtonEnabled() const;
+	
+	void RefreshTree();
+	
+	void OnPreObjectPropertyChanged(UObject* InObject, const FEditPropertyChain& InPropertyChain);
+	void OnObjectPropertyChanged(UObject* InObject, FPropertyChangedEvent& InPropertyChangedEvent);
+	void OnObjectConstructed(UObject* InObject);
+
+	UPCGGraph* GetPCGGraph() const;
+	
+	TSharedRef<ITableRow> MakeTreeRowWidget(FPCGEditorGraphDebugObjectItemPtr InItem, const TSharedRef<STableViewBase>& InOwnerTable) const;
+	void OnGetChildren(FPCGEditorGraphDebugObjectItemPtr InItem, TArray<FPCGEditorGraphDebugObjectItemPtr>& OutChildren) const;
+	void OnSelectionChanged(FPCGEditorGraphDebugObjectItemPtr InItem, ESelectInfo::Type InSelectInfo) const;
+	
+	TWeakPtr<FPCGEditor> PCGEditor;
+
+	TSharedPtr<STreeView<FPCGEditorGraphDebugObjectItemPtr>> DebugObjectTreeView;
+	TArray<FPCGEditorGraphDebugObjectItemPtr> RootItems;
+	TArray<FPCGEditorGraphDebugObjectItemPtr> AllGraphItems;
+
+	bool bNeedsRefresh = false;
+};
