Prevented crash when the param override contains no properties (which likely shouldn't happen)

#rnx

diff --git a/Engine/Plugins/Experimental/PCG/Source/PCG/Private/PCGSettings.cpp b/Engine/Plugins/Experimental/PCG/Source/PCG/Private/PCGSettings.cpp
--- a/Engine/Plugins/Experimental/PCG/Source/PCG/Private/PCGSettings.cpp
+++ b/Engine/Plugins/Experimental/PCG/Source/PCG/Private/PCGSettings.cpp
@@ -508,7 +508,7 @@ bool UPCGSettings::IsPropertyOverriddenByPin(const FProperty* InProperty) const
 		{
 			// In OverridableParam, the array of properties is the chain of properties from the Settings class to the wanted param.
 			// Therefore the property we are editing would match the latest property of the array.
-			return ParamToCheck.Properties.Last() == InProperty;
+			return !ParamToCheck.Properties.IsEmpty() && ParamToCheck.Properties.Last() == InProperty;
 		});
 
 		if (Param)
