[PCG] Refactored the Difference Node to follow new standards

* Loops on Source Pin
* Unions on Difference Pin only
* [DEPREC] Injects Union node before source to respect old behavior

diff --git a/Engine/Plugins/Experimental/PCG/Source/PCG/Private/Elements/PCGDifferenceElement.cpp b/Engine/Plugins/Experimental/PCG/Source/PCG/Private/Elements/PCGDifferenceElement.cpp
--- a/Engine/Plugins/Experimental/PCG/Source/PCG/Private/Elements/PCGDifferenceElement.cpp
+++ b/Engine/Plugins/Experimental/PCG/Source/PCG/Private/Elements/PCGDifferenceElement.cpp
@@ -4,18 +4,96 @@
 
 #include "PCGContext.h"
 #include "PCGCustomVersion.h"
+#include "PCGEdge.h"
+#include "PCGGraph.h"
 #include "PCGPin.h"
 #include "Data/PCGPointData.h"
 #include "Data/PCGSpatialData.h"
 #include "Data/PCGUnionData.h"
+#include "Elements/PCGUnionElement.h"
 
 #include UE_INLINE_GENERATED_CPP_BY_NAME(PCGDifferenceElement)
 
+#define LOCTEXT_NAMESPACE "PCGDifferenceElement"
+
+#if WITH_EDITOR
+FText UPCGDifferenceSettings::GetNodeTooltipText() const
+{
+	return LOCTEXT("NodeTooltip", "Spatially subtracts the target difference data from the source data, outputing the difference of the two.");
+}
+
+void UPCGDifferenceSettings::ApplyStructuralDeprecation(UPCGNode* InOutNode)
+{
+	if (DataVersion < FPCGCustomVersion::DifferenceNodeIterateOnSourceAndUnionDifferences)
+	{
+		UPCGPin* SourcePin = InOutNode->GetInputPin(PCGDifferenceConstants::SourceLabel);
+		check(SourcePin);
+
+		// Check if single connection AND single data, to avoid adding a union if not necessary
+		const int32 SourceEdgeCount = SourcePin->EdgeCount();
+		if (SourceEdgeCount > 1 || (SourceEdgeCount == 1 && !SourcePin->Edges.IsEmpty() && SourcePin->Edges[0]->InputPin->AllowMultipleData()))
+		{
+			// To match previous default behavior, we'll add a union node preceding the source inputs
+			UPCGGraph* PCGGraph = InOutNode->GetGraph();
+			check(PCGGraph);
+
+			UPCGSettings* UnionSettings = nullptr;
+			UPCGNode* UnionNode = PCGGraph->AddNodeOfType(UPCGUnionSettings::StaticClass(), UnionSettings);
+			check(UnionNode && !UnionNode->GetInputPins().IsEmpty() && !UnionNode->GetOutputPins().IsEmpty());
+			UnionSettings->bEnabled = InOutNode->GetSettings()->bEnabled;
+
+			UPCGPin* UnionSourcePin = UnionNode->GetInputPins()[0];
+			UPCGPin* UnionOutputPin = UnionNode->GetOutputPins()[0];
+
+			int32 SumSourceNodesPositionX = 0;
+			int32 SumSourceNodesPositionY = 0;
+			while(!SourcePin->Edges.IsEmpty())
+			{
+				UPCGPin* UpstreamPin = SourcePin->Edges[0]->InputPin;
+				check(UpstreamPin);
+
+				// Collect the source node's position for later mean
+				check(UpstreamPin->Node);
+				int32 NodePositionX, NodePositionY;
+				UpstreamPin->Node->GetNodePosition(NodePositionX, NodePositionY);
+				SumSourceNodesPositionX += NodePositionX;
+				SumSourceNodesPositionY += NodePositionY;
+		
+				UpstreamPin->BreakEdgeTo(SourcePin);
+				UpstreamPin->AddEdgeTo(UnionSourcePin);
+			}
+
+			// The injected node should be placed halfway between the difference node and the mean of the sources' position
+			int32 NodePositionX, NodePositionY;
+			InOutNode->GetNodePosition(NodePositionX, NodePositionY);
+			const int32 SourceNodeLocalPositionMeanX = SumSourceNodesPositionX / SourceEdgeCount;
+			const int32 SourceNodeLocalPositionMeanY = SumSourceNodesPositionY / SourceEdgeCount;
+			
+			UnionNode->SetNodePosition((NodePositionX + SourceNodeLocalPositionMeanX) / 2, (NodePositionY + SourceNodeLocalPositionMeanY) / 2);
+			UnionNode->NodeComment = LOCTEXT("DeprecationUnionNodeCreated", "Node added to replicate deprecated behavior of the Difference node.").ToString();
+			UnionNode->bCommentBubbleVisible = 1;
+
+			// Reconnect union back to source
+			UnionOutputPin->AddEdgeTo(SourcePin);		
+
+			UE_LOG(LogPCG, Log, TEXT("Difference node of previous version detected. A 'Union' node was created automatically to replicate the previous version's behavior."));
+		}
+	}
+	
+	Super::ApplyStructuralDeprecation(InOutNode);
+}
+#endif // WITH_EDITOR
+
 TArray<FPCGPinProperties> UPCGDifferenceSettings::InputPinProperties() const
 {
 	TArray<FPCGPinProperties> PinProperties;
-	PinProperties.Emplace(PCGDifferenceConstants::SourceLabel, EPCGDataType::Spatial);
-	PinProperties.Emplace(PCGDifferenceConstants::DifferencesLabel, EPCGDataType::Spatial);
+	FPCGPinProperties& PinPropertiesSource = PinProperties.Emplace_GetRef(PCGDifferenceConstants::SourceLabel, EPCGDataType::Spatial);
+	FPCGPinProperties& PinPropertiesDifferences = PinProperties.Emplace_GetRef(PCGDifferenceConstants::DifferencesLabel, EPCGDataType::Spatial);
+
+#if WITH_EDITOR
+	PinPropertiesSource.Tooltip = LOCTEXT("SourcePropertiesTooltip", "The source input of the operation, from which the difference will be subtracted. Multiple pins or data on this pin will be iterated upon, resulting in an output of the difference operation for every input.");
+	PinPropertiesDifferences.Tooltip = LOCTEXT("DifferencesPropertiesTooltip", "The difference input of the operation. Inputs on this pin will be implicitly unioned together to be used in the difference operation.");
+#endif // WITH_EDITOR
 
 	return PinProperties;
 }
@@ -23,8 +101,12 @@ TArray<FPCGPinProperties> UPCGDifferenceSettings::InputPinProperties() const
 TArray<FPCGPinProperties> UPCGDifferenceSettings::OutputPinProperties() const
 {
 	TArray<FPCGPinProperties> PinProperties;
-	PinProperties.Emplace(PCGPinConstants::DefaultOutputLabel, EPCGDataType::Spatial);
-
+	FPCGPinProperties& PinPropertiesOutput = PinProperties.Emplace_GetRef(PCGPinConstants::DefaultOutputLabel, EPCGDataType::Spatial);
+	
+#if WITH_EDITOR
+	PinPropertiesOutput.Tooltip = LOCTEXT("OutputPropertiesTooltip", "The data after being processed through the difference operation.");
+#endif // WITH_EDITOR
+	
 	return PinProperties;
 }
 
@@ -40,7 +122,6 @@ bool FPCGDifferenceElement::ExecuteInternal(FPCGContext* Context) const
 	// Early-out for previous behavior (without labeled edges)
 	if (Context->Node && !Context->Node->IsInputPinConnected(PCGDifferenceConstants::SourceLabel) && !Context->Node->IsInputPinConnected(PCGDifferenceConstants::DifferencesLabel))
 	{
-		LabellessProcessing(Context);
 		return true;
 	}
 
@@ -51,158 +132,72 @@ bool FPCGDifferenceElement::ExecuteInternal(FPCGContext* Context) const
 	TArray<FPCGTaggedData> Differences = Context->InputData.GetInputsByPin(PCGDifferenceConstants::DifferencesLabel);
 	TArray<FPCGTaggedData>& Outputs = Context->OutputData.TaggedData;
 
-	// Get only spatial data or build an union from the sources
-	const UPCGSpatialData* FirstSpatialData = nullptr;
-	UPCGUnionData* UnionData = nullptr;
-	int32 DifferenceTaggedDataIndex = -1;
-
 	bool bHasPointsInSource = false;
 	bool bHasPointsInDifferences = false;
 
-	// Start by either selecting a source or building a union
+	// Validate that there's usable data on the differences pin
+	bool bDifferenceDataHasValidData = false;
+	for (const FPCGTaggedData& Difference : Differences)
+	{
+		if (Difference.Data && Difference.Data->IsA<UPCGSpatialData>())
+		{
+			bDifferenceDataHasValidData = true;
+			bHasPointsInDifferences |= Difference.Data->IsA<UPCGPointData>();
+			break;
+		}
+	}
+
+	// If no valid difference data, send sources directly to output and return
+	if (!bDifferenceDataHasValidData)
+	{
+		Outputs = Sources;
+		return true;
+	}
+
+	// Iterate over the sources and subtract the differences
 	for (FPCGTaggedData& Source : Sources)
 	{
-		const UPCGSpatialData* SpatialData = Cast<UPCGSpatialData>(Source.Data);
+		const UPCGSpatialData* SourceSpatialData = Cast<UPCGSpatialData>(Source.Data);
 
-		// Non-spatial data, we're not going to touch
-		if (!SpatialData)
+		// Propagate non-spatial data
+		if (!SourceSpatialData)
 		{
 			Outputs.Add(Source);
 			continue;
 		}
 
-		bHasPointsInSource |= (Cast<UPCGPointData>(SpatialData) != nullptr);
+		bHasPointsInSource |= SourceSpatialData->IsA<UPCGPointData>();
 
-		if (!FirstSpatialData)
-		{
-			FirstSpatialData = SpatialData;
-			DifferenceTaggedDataIndex = Outputs.Num();
-			Outputs.Add(Source);
-		}
-		else
+		UPCGDifferenceData* DifferenceData = NewObject<UPCGDifferenceData>();
+		DifferenceData->Initialize(SourceSpatialData);
+		
+		for (FPCGTaggedData& Difference : Differences)
 		{
-			if (!UnionData)
-			{
-				UnionData = FirstSpatialData->UnionWith(SpatialData);
-				// TODO: expose union settings?
-
-				// Replace source by union
-				Outputs[DifferenceTaggedDataIndex].Data = UnionData;
-			}
-			else
+			if (const UPCGSpatialData* DifferenceSpatialData = Cast<const UPCGSpatialData>(Difference.Data))
 			{
-				UnionData->AddData(SpatialData);
+				DifferenceData->AddDifference(DifferenceSpatialData);
 			}
 		}
-	}
-	
-	if (FirstSpatialData || UnionData)
-	{
-		check(DifferenceTaggedDataIndex >= 0);
-		// Then, depending on the presence of differences, we'll create a difference as needed
-		UPCGDifferenceData* DifferenceData = nullptr;
-		for (FPCGTaggedData& Difference : Differences)
-		{
-			const UPCGSpatialData* SpatialData = Cast<UPCGSpatialData>(Difference.Data);
 
-			if (SpatialData)
-			{
-				bHasPointsInDifferences |= (Cast<UPCGPointData>(SpatialData) != nullptr);
-
-				if (!DifferenceData)
-				{
-					DifferenceData = (UnionData ? UnionData : FirstSpatialData)->Subtract(SpatialData);
-					DifferenceData->SetDensityFunction(Settings->DensityFunction);
-					DifferenceData->bDiffMetadata = Settings->bDiffMetadata;
+		DifferenceData->SetDensityFunction(Settings->DensityFunction);
+		DifferenceData->bDiffMetadata = Settings->bDiffMetadata;
 #if WITH_EDITOR
-					DifferenceData->bKeepZeroDensityPoints = Settings->bKeepZeroDensityPoints;
+		DifferenceData->bKeepZeroDensityPoints = Settings->bKeepZeroDensityPoints;
 #endif
 
-					Outputs[DifferenceTaggedDataIndex].Data = DifferenceData;
-				}
-				else
-				{
-					DifferenceData->AddDifference(SpatialData);
-				}
-			}
-			else
-			{
-				// We are not propagating data from the differences if they don't contribute
-			}
-		}
+		FPCGTaggedData& Output = Outputs.Add_GetRef(Source);
+		Output.Data = DifferenceData;
 
 		// Finally, apply any discretization based on the mode
 		if (DifferenceData && 
 			(Settings->Mode == EPCGDifferenceMode::Discrete || 
 			(Settings->Mode == EPCGDifferenceMode::Inferred && bHasPointsInSource && bHasPointsInDifferences)))
 		{
-			Outputs[DifferenceTaggedDataIndex].Data = DifferenceData->ToPointData(Context);
+			Output.Data = DifferenceData->ToPointData(Context);
 		}
 	}
 
 	return true;
 }
 
-void FPCGDifferenceElement::LabellessProcessing(FPCGContext* Context) const
-{
-	const UPCGDifferenceSettings* Settings = Context->GetInputSettings<UPCGDifferenceSettings>();
-	check(Settings);
-
-	TArray<FPCGTaggedData> Inputs = Context->InputData.GetInputs();
-
-	const EPCGDifferenceDensityFunction DensityFunction = Settings->DensityFunction;
-#if WITH_EDITOR
-	const bool bKeepZeroDensityPoints = Settings->bKeepZeroDensityPoints;
-#else
-	const bool bKeepZeroDensityPoints = false;
-#endif
-
-	TArray<FPCGTaggedData>& Outputs = Context->OutputData.TaggedData;
-
-	const UPCGSpatialData* FirstSpatialData = nullptr;
-	UPCGDifferenceData* DifferenceData = nullptr;
-	int32 DifferenceTaggedDataIndex = -1;
-
-	auto AddToDifference = [&FirstSpatialData, &DifferenceData, &DifferenceTaggedDataIndex, DensityFunction, bKeepZeroDensityPoints, &Outputs](const UPCGSpatialData* SpatialData) {
-		if (!DifferenceData)
-		{
-			DifferenceData = FirstSpatialData->Subtract(SpatialData);
-			DifferenceData->SetDensityFunction(DensityFunction);
-#if WITH_EDITOR
-			DifferenceData->bKeepZeroDensityPoints = bKeepZeroDensityPoints;
-#endif
-
-			FPCGTaggedData& DifferenceTaggedData = Outputs[DifferenceTaggedDataIndex];
-			DifferenceTaggedData.Data = DifferenceData;
-		}
-		else
-		{
-			DifferenceData->AddDifference(SpatialData);
-		}
-	};
-
-	// TODO: it might not make sense to perform the difference if the first
-	// data isn't a spatial data, otherwise, what would it really mean?
-	for (FPCGTaggedData& Input : Inputs)
-	{
-		const UPCGSpatialData* SpatialData = Cast<UPCGSpatialData>(Input.Data);
-
-		// Non-spatial data, we're not going to touch
-		if (!SpatialData)
-		{
-			Outputs.Add(Input);
-			continue;
-		}
-
-		if (!FirstSpatialData)
-		{
-			FirstSpatialData = SpatialData;
-			DifferenceTaggedDataIndex = Outputs.Num();
-			Outputs.Add(Input);
-
-			continue;
-		}
-
-		AddToDifference(SpatialData);
-	}
-}
+#undef LOCTEXT_NAMESPACE
\ No newline at end of file
diff --git a/Engine/Plugins/Experimental/PCG/Source/PCG/Private/PCGGraph.cpp b/Engine/Plugins/Experimental/PCG/Source/PCG/Private/PCGGraph.cpp
--- a/Engine/Plugins/Experimental/PCG/Source/PCG/Private/PCGGraph.cpp
+++ b/Engine/Plugins/Experimental/PCG/Source/PCG/Private/PCGGraph.cpp
@@ -267,6 +267,13 @@ void UPCGGraph::PostLoad()
 		ExtraNode->ConditionalPostLoad();
 	}
 
+	// Create a copy to iterate through the nodes while more might be added
+	TArray<UPCGNode*> NodesCopy(Nodes);
+	for (UPCGNode* Node : NodesCopy)
+	{
+		Node->ApplyStructuralDeprecation();
+	}
+
 	// Finally, apply deprecation that changes edges/rebinds
 	ForEachNode([](UPCGNode* InNode) { return InNode->ApplyDeprecationBeforeUpdatePins(); });
 
@@ -380,8 +387,9 @@ UPCGNode* UPCGGraph::AddNode(UPCGSettingsInterface* InSettingsInterface)
 		const FName DefaultNodeName = InSettingsInterface->GetSettings()->GetDefaultNodeName();
 		if (DefaultNodeName != NAME_None)
 		{
-			FName NodeName = MakeUniqueObjectName(this, UPCGNode::StaticClass(), DefaultNodeName);
-			Node->Rename(*NodeName.ToString());
+			const FName NodeName = MakeUniqueObjectName(this, UPCGNode::StaticClass(), DefaultNodeName);
+			// Flags added because default flags favor tick/interactive, not load-time renaming.
+			Node->Rename(*NodeName.ToString(), nullptr, REN_ForceNoResetLoaders | REN_DontCreateRedirectors);
 		}
 #endif
 
diff --git a/Engine/Plugins/Experimental/PCG/Source/PCG/Private/PCGNode.cpp b/Engine/Plugins/Experimental/PCG/Source/PCG/Private/PCGNode.cpp
--- a/Engine/Plugins/Experimental/PCG/Source/PCG/Private/PCGNode.cpp
+++ b/Engine/Plugins/Experimental/PCG/Source/PCG/Private/PCGNode.cpp
@@ -150,6 +150,15 @@ void UPCGNode::ApplyDeprecation()
 		ensure(Settings->DataVersion == FPCGCustomVersion::LatestVersion);
 	}
 }
+
+void UPCGNode::ApplyStructuralDeprecation()
+{
+	if (UPCGSettings* Settings = GetSettings())
+	{
+		Settings->ApplyStructuralDeprecation(this);
+	}
+}
+
 #endif
 
 #if WITH_EDITOR
diff --git a/Engine/Plugins/Experimental/PCG/Source/PCG/Private/PCGPin.cpp b/Engine/Plugins/Experimental/PCG/Source/PCG/Private/PCGPin.cpp
--- a/Engine/Plugins/Experimental/PCG/Source/PCG/Private/PCGPin.cpp
+++ b/Engine/Plugins/Experimental/PCG/Source/PCG/Private/PCGPin.cpp
@@ -320,6 +320,11 @@ bool UPCGPin::AllowMultipleConnections() const
 	return IsOutputPin() || Properties.bAllowMultipleConnections;
 }
 
+bool UPCGPin::AllowMultipleData() const
+{
+	return Properties.bAllowMultipleData;
+}
+
 bool UPCGPin::CanConnect(const UPCGPin* OtherPin) const
 {
 	return OtherPin && (Edges.IsEmpty() || AllowMultipleConnections());
diff --git a/Engine/Plugins/Experimental/PCG/Source/PCG/Public/Elements/PCGDifferenceElement.h b/Engine/Plugins/Experimental/PCG/Source/PCG/Public/Elements/PCGDifferenceElement.h
--- a/Engine/Plugins/Experimental/PCG/Source/PCG/Public/Elements/PCGDifferenceElement.h
+++ b/Engine/Plugins/Experimental/PCG/Source/PCG/Public/Elements/PCGDifferenceElement.h
@@ -23,7 +23,10 @@ public:
 #if WITH_EDITOR
 	virtual FName GetDefaultNodeName() const override { return FName(TEXT("Difference")); }
 	virtual FText GetDefaultNodeTitle() const override { return NSLOCTEXT("PCGDifferenceSettings", "NodeTitle", "Difference"); }
+	virtual FText GetNodeTooltipText() const override;
 	virtual EPCGSettingsType GetType() const override { return EPCGSettingsType::Spatial; }
+
+	virtual void ApplyStructuralDeprecation(UPCGNode* InOutNode) override;
 #endif
 
 	virtual TArray<FPCGPinProperties> InputPinProperties() const override;
@@ -34,9 +37,15 @@ protected:
 	//~End UPCGSettings interface
 
 public:
+	/** The density function to use when recalculating the density after the operation. */
 	UPROPERTY(BlueprintReadWrite, EditAnywhere, Category = Settings, meta = (PCG_Overridable))
 	EPCGDifferenceDensityFunction DensityFunction = EPCGDifferenceDensityFunction::Minimum;
 
+	/** Describes how the difference operation will treat the output data:
+	 * Continuous - Non-destructive data output will be maintained.
+	 * Discrete - Output data will be discrete points, or explicitly converted to points.
+	 * Inferred - Output data will choose from Continuous or Discrete, based on the source and operation.
+	 */
 	UPROPERTY(BlueprintReadWrite, EditAnywhere, Category = Settings)
 	EPCGDifferenceMode Mode = EPCGDifferenceMode::Inferred;
 
@@ -44,6 +53,7 @@ public:
 	bool bDiffMetadata = true;
 
 #if WITH_EDITORONLY_DATA
+	/** If enabled, the output will not automatically filter out points with zero density. */
 	UPROPERTY(Transient, BlueprintReadWrite, EditAnywhere, Category = "Settings|Debug", meta = (PCG_Overridable))
 	bool bKeepZeroDensityPoints = false;
 #endif
@@ -52,9 +62,7 @@ public:
 class FPCGDifferenceElement : public FSimplePCGElement
 {
 protected:
-	virtual bool ExecuteInternal(FPCGContext* Context) const;
-
-	void LabellessProcessing(FPCGContext* Context) const;
+	virtual bool ExecuteInternal(FPCGContext* Context) const override;
 };
 
 #if UE_ENABLE_INCLUDE_ORDER_DEPRECATED_IN_5_2
diff --git a/Engine/Plugins/Experimental/PCG/Source/PCG/Public/Elements/PCGUnionElement.h b/Engine/Plugins/Experimental/PCG/Source/PCG/Public/Elements/PCGUnionElement.h
--- a/Engine/Plugins/Experimental/PCG/Source/PCG/Public/Elements/PCGUnionElement.h
+++ b/Engine/Plugins/Experimental/PCG/Source/PCG/Public/Elements/PCGUnionElement.h
@@ -23,9 +23,6 @@ public:
 	virtual TArray<FPCGPinProperties> InputPinProperties() const override;
 	virtual TArray<FPCGPinProperties> OutputPinProperties() const override;
 
-	// If node disabled, don't union - pass through first edge
-	virtual bool OnlyPassThroughOneEdgeWhenDisabled() const { return true; }
-
 protected:
 	virtual FPCGElementPtr CreateElement() const override;
 	//~End UPCGSettings interface
diff --git a/Engine/Plugins/Experimental/PCG/Source/PCG/Public/PCGCustomVersion.h b/Engine/Plugins/Experimental/PCG/Source/PCG/Public/PCGCustomVersion.h
--- a/Engine/Plugins/Experimental/PCG/Source/PCG/Public/PCGCustomVersion.h
+++ b/Engine/Plugins/Experimental/PCG/Source/PCG/Public/PCGCustomVersion.h
@@ -47,6 +47,9 @@ struct PCG_API FPCGCustomVersion
 		// Update of Input Selectors. Previous versions should default on @LastCreated
 		UpdateAttributePropertyInputSelector = 12,
 
+		// Difference node now iterates on the source pin and unions the differences pin
+		DifferenceNodeIterateOnSourceAndUnionDifferences = 13,
+
 		// -----<new versions can be added above this line>-------------------------------------------------
 		VersionPlusOne,
 		LatestVersion = VersionPlusOne - 1
diff --git a/Engine/Plugins/Experimental/PCG/Source/PCG/Public/PCGNode.h b/Engine/Plugins/Experimental/PCG/Source/PCG/Public/PCGNode.h
--- a/Engine/Plugins/Experimental/PCG/Source/PCG/Public/PCGNode.h
+++ b/Engine/Plugins/Experimental/PCG/Source/PCG/Public/PCGNode.h
@@ -46,6 +46,9 @@ public:
 
 	/** Used to be able to force deprecation when things need to be deprecated at the graph level. */
 	void ApplyDeprecation();
+
+	/** If a node does require structural changes, this will apply them */
+	virtual void ApplyStructuralDeprecation();
 #endif
 
 	/** Returns the owning graph */
diff --git a/Engine/Plugins/Experimental/PCG/Source/PCG/Public/PCGPin.h b/Engine/Plugins/Experimental/PCG/Source/PCG/Public/PCGPin.h
--- a/Engine/Plugins/Experimental/PCG/Source/PCG/Public/PCGPin.h
+++ b/Engine/Plugins/Experimental/PCG/Source/PCG/Public/PCGPin.h
@@ -137,6 +137,7 @@ public:
 	void SetTooltip(const FText& InTooltip);
 
 	bool AllowMultipleConnections() const;
+	bool AllowMultipleData() const;
 	bool IsCompatible(const UPCGPin* OtherPin) const;
 	bool CanConnect(const UPCGPin* OtherPin) const;
 
diff --git a/Engine/Plugins/Experimental/PCG/Source/PCG/Public/PCGSettings.h b/Engine/Plugins/Experimental/PCG/Source/PCG/Public/PCGSettings.h
--- a/Engine/Plugins/Experimental/PCG/Source/PCG/Public/PCGSettings.h
+++ b/Engine/Plugins/Experimental/PCG/Source/PCG/Public/PCGSettings.h
@@ -224,6 +224,9 @@ public:
 	/** Any final migration/recovery that can be done after pins are finalized. This function should also set DataVersion to LatestVersion. */
 	virtual void ApplyDeprecation(UPCGNode* InOutNode);
 
+	/** If settings require structural changes, this will apply them */
+	virtual void ApplyStructuralDeprecation(UPCGNode* InOutNode) {}
+
 	virtual FName GetDefaultNodeName() const { return NAME_None; }
 	virtual FText GetDefaultNodeTitle() const { return FText::FromName(GetDefaultNodeName()); }
 	virtual FText GetNodeTooltipText() const { return FText::GetEmpty(); }
