[PCG] Fix for static analysis warning in debug object tree

#rnx

diff --git a/Engine/Plugins/Experimental/PCG/Source/PCGEditor/Private/SPCGEditorGraphDebugObjectTree.cpp b/Engine/Plugins/Experimental/PCG/Source/PCGEditor/Private/SPCGEditorGraphDebugObjectTree.cpp
--- a/Engine/Plugins/Experimental/PCG/Source/PCGEditor/Private/SPCGEditorGraphDebugObjectTree.cpp
+++ b/Engine/Plugins/Experimental/PCG/Source/PCGEditor/Private/SPCGEditorGraphDebugObjectTree.cpp
@@ -235,7 +235,8 @@ void SPCGEditorGraphDebugObjectTree::SetDebugObjectFromSelection_OnClicked()
 		return;
 	}
 
-	USelection* SelectedActors = GEditor ? GEditor->GetSelectedActors() : nullptr;
+	check(GEditor);
+	USelection* SelectedActors = GEditor->GetSelectedActors();
 	if (!IsValid(SelectedActors))
 	{
 		return;
