[PCG] Fix projected spline-sampled points not taking landscape surface slope and narrow UPCGSplineProjectionData type to 'Spline'

#rnx

diff --git a/Engine/Plugins/Experimental/PCG/Source/PCG/Private/Elements/PCGSplineSampler.cpp b/Engine/Plugins/Experimental/PCG/Source/PCG/Private/Elements/PCGSplineSampler.cpp
--- a/Engine/Plugins/Experimental/PCG/Source/PCG/Private/Elements/PCGSplineSampler.cpp
+++ b/Engine/Plugins/Experimental/PCG/Source/PCG/Private/Elements/PCGSplineSampler.cpp
@@ -1075,7 +1075,20 @@ namespace PCGSplineSampler
 								continue;
 							}
 
-							TransformWS.SetLocation(ProjectedPoint.Transform.GetLocation());
+							if (InProjectionParams.bProjectPositions)
+							{
+								TransformWS.SetLocation(ProjectedPoint.Transform.GetLocation());
+							}
+
+							if (InProjectionParams.bProjectRotations)
+							{
+								TransformWS.SetRotation(ProjectedPoint.Transform.GetRotation());
+							}
+
+							if (InProjectionParams.bProjectScales)
+							{
+								TransformWS.SetScale3D(ProjectedPoint.Transform.GetScale3D());
+							}
 						}
 
 						// Prune points outside of bounds
diff --git a/Engine/Plugins/Experimental/PCG/Source/PCG/Public/Data/PCGSplineData.h b/Engine/Plugins/Experimental/PCG/Source/PCG/Public/Data/PCGSplineData.h
--- a/Engine/Plugins/Experimental/PCG/Source/PCG/Public/Data/PCGSplineData.h
+++ b/Engine/Plugins/Experimental/PCG/Source/PCG/Public/Data/PCGSplineData.h
@@ -68,6 +68,10 @@ class PCG_API UPCGSplineProjectionData : public UPCGProjectionData
 public:
 	void Initialize(const UPCGSplineData* InSourceSpline, const UPCGSpatialData* InTargetSurface, const FPCGProjectionParams& InParams);
 
+	// ~Begin UPCGData interface
+	virtual EPCGDataType GetDataType() const override { return EPCGDataType::Spline; }
+	// ~End UPCGData interface
+
 	const UPCGSplineData* GetSpline() const;
 	const UPCGSpatialData* GetSurface() const;
 
