[PCG] Add Alias support for deprecated overridable parameters when there is a one to one mapping

diff --git a/Engine/Plugins/Experimental/PCG/Source/PCG/Private/Helpers/PCGSettingsHelpers.cpp b/Engine/Plugins/Experimental/PCG/Source/PCG/Private/Helpers/PCGSettingsHelpers.cpp
--- a/Engine/Plugins/Experimental/PCG/Source/PCG/Private/Helpers/PCGSettingsHelpers.cpp
+++ b/Engine/Plugins/Experimental/PCG/Source/PCG/Private/Helpers/PCGSettingsHelpers.cpp
@@ -284,6 +284,26 @@ namespace PCGSettingsHelpers
 
 		check(InClass);
 
+		auto GatherAliases = [](const FProperty* InProperty, FPCGSettingsOverridableParam& InParam) -> void
+		{
+#if WITH_EDITOR
+			if (InProperty->HasMetaData(PCGObjectMetadata::OverrideAliases))
+			{
+				FPCGPropertyAliases Result;
+				TArray<FString> TempStringArray;
+				const FString& AliasesStr = InProperty->GetMetaData(PCGObjectMetadata::OverrideAliases);
+				AliasesStr.ParseIntoArray(TempStringArray, TEXT(","));
+
+				if (!TempStringArray.IsEmpty())
+				{
+					Algo::Transform(TempStringArray, Result.Aliases, [](const FString& In) -> FName { return FName(In); });
+					// We always emplace at level 0, it will be incremented by the recursion if it is deeper than 1 level.
+					InParam.MapOfAliases.Emplace(0, std::move(Result));
+				}
+			}
+#endif // WITH_EDITOR
+		};
+
 		for (TFieldIterator<FProperty> InputIt(InClass, SuperFlag, EFieldIteratorFlags::ExcludeDeprecated); InputIt; ++InputIt)
 		{
 			const FProperty* Property = *InputIt;
@@ -382,6 +402,7 @@ namespace PCGSettingsHelpers
 				Param.bHasNameClash = bHasNameClash;
 #if WITH_EDITOR
 				Param.Label = bHasNameClash ? FName(Param.GetDisplayPropertyPath()) : Label;
+				GatherAliases(Property, Param);
 #endif // WITH_EDITOR
 			}
 			else if (const FStructProperty* StructProperty = CastField<FStructProperty>(Property))
@@ -404,7 +425,7 @@ namespace PCGSettingsHelpers
 					RecurseConfig.MaxStructDepth--;
 				}
 
-				for (const FPCGSettingsOverridableParam& ChildParam : GetAllOverridableParams(StructProperty->Struct, RecurseConfig))
+				for (FPCGSettingsOverridableParam& ChildParam : GetAllOverridableParams(StructProperty->Struct, RecurseConfig))
 				{
 					FName Label = ChildParam.Label;
 					bool bHasNameClash = false;
@@ -439,6 +460,15 @@ namespace PCGSettingsHelpers
 #endif // WITH_EDITOR
 						Param.bHasNameClash = true;
 					}
+
+#if WITH_EDITOR
+					GatherAliases(Property, Param);
+					// For all already found aliases, add them to the current param, with the index incremented by 1.
+					for (TPair<int32, FPCGPropertyAliases>& It : ChildParam.MapOfAliases)
+					{
+						Param.MapOfAliases.Emplace(It.Key + 1, std::move(It.Value));
+					}
+#endif // WITH_EDITOR
 				}
 			}
 		}
diff --git a/Engine/Plugins/Experimental/PCG/Source/PCG/Private/Metadata/Accessors/PCGAttributeAccessorHelpers.cpp b/Engine/Plugins/Experimental/PCG/Source/PCG/Private/Metadata/Accessors/PCGAttributeAccessorHelpers.cpp
--- a/Engine/Plugins/Experimental/PCG/Source/PCG/Private/Metadata/Accessors/PCGAttributeAccessorHelpers.cpp
+++ b/Engine/Plugins/Experimental/PCG/Source/PCG/Private/Metadata/Accessors/PCGAttributeAccessorHelpers.cpp
@@ -410,17 +410,36 @@ TUniquePtr<IPCGAttributeAccessor> PCGAttributeAccessorHelpers::CreateExtraAccess
 }
 
 TUniquePtr<const IPCGAttributeAccessor> PCGAttributeAccessorHelpers::CreateConstAccessorForOverrideParam(const FPCGDataCollection& InInputData, const FPCGSettingsOverridableParam& InParam, FName* OutAttributeName)
+{
+	if (OutAttributeName)
+	{
+		AccessorParamResult Result{};
+		TUniquePtr<const IPCGAttributeAccessor> Accessor = CreateConstAccessorForOverrideParamWithResult(InInputData, InParam, &Result);
+		*OutAttributeName = Result.AttributeName;
+		return Accessor;
+	}
+	else
+	{
+		return CreateConstAccessorForOverrideParamWithResult(InInputData, InParam);
+	}
+}
+
+TUniquePtr<const IPCGAttributeAccessor> PCGAttributeAccessorHelpers::CreateConstAccessorForOverrideParamWithResult(const FPCGDataCollection& InInputData, const FPCGSettingsOverridableParam& InParam, AccessorParamResult* OutResult)
 {
 	bool bFromGlobalParamsPin = false;
 	TArray<FPCGTaggedData> InputParamData = InInputData.GetParamsByPin(InParam.Label);
 	UPCGParamData* ParamData = !InputParamData.IsEmpty() ? CastChecked<UPCGParamData>(InputParamData[0].Data) : nullptr;
 
-	// If it is empty, try with the param pin
+	// If it is empty, try with the Overrides pin (Global Params)
 	if (!ParamData)
 	{
 		bFromGlobalParamsPin = true;
 		ParamData = InInputData.GetFirstParamsOnParamsPin();
 	}
+	else if (OutResult)
+	{
+		OutResult->bPinConnected = true;
+	}
 
 	if (ParamData && ParamData->Metadata && ParamData->Metadata->GetAttributeCount() > 0)
 	{
@@ -429,14 +448,36 @@ TUniquePtr<const IPCGAttributeAccessor> PCGAttributeAccessorHelpers::CreateConst
 			? ParamData->Metadata->GetLatestAttributeNameOrNone() 
 			: (!InParam.bHasNameClash ? InParam.PropertiesNames.Last() : FName(InParam.GetPropertyPath()));
 
-		if (OutAttributeName)
+		if (OutResult)
 		{
-			*OutAttributeName = AttributeName;
+			OutResult->AttributeName = AttributeName;
 		}
 
 		FPCGAttributePropertyInputSelector InputSelector{};
 		InputSelector.SetAttributeName(AttributeName);
-		return PCGAttributeAccessorHelpers::CreateConstAccessor(ParamData, InputSelector);
+		TUniquePtr<const IPCGAttributeAccessor> Result = PCGAttributeAccessorHelpers::CreateConstAccessor(ParamData, InputSelector);
+
+		if (!Result)
+		{
+			// If we didn't find it, try some aliases.
+			for (const FName& Alias : InParam.GenerateAllPossibleAliases())
+			{
+				InputSelector.SetAttributeName(Alias);
+				Result = PCGAttributeAccessorHelpers::CreateConstAccessor(ParamData, InputSelector);
+				if (Result)
+				{
+					if (OutResult)
+					{
+						OutResult->bUsedAliases = true;
+						OutResult->AliasUsed = Alias;
+					}
+
+					break;
+				}
+			}
+		}
+
+		return Result;
 	}
 
 	return TUniquePtr<const IPCGAttributeAccessor>{};
diff --git a/Engine/Plugins/Experimental/PCG/Source/PCG/Private/PCGContext.cpp b/Engine/Plugins/Experimental/PCG/Source/PCG/Private/PCGContext.cpp
--- a/Engine/Plugins/Experimental/PCG/Source/PCG/Private/PCGContext.cpp
+++ b/Engine/Plugins/Experimental/PCG/Source/PCG/Private/PCGContext.cpp
@@ -192,15 +192,29 @@ void FPCGContext::OverrideSettings()
 			continue;
 		}
 
-		FName AttributeName = NAME_None;
-		TUniquePtr<const IPCGAttributeAccessor> AttributeAccessor = PCGAttributeAccessorHelpers::CreateConstAccessorForOverrideParam(InputData, Param, &AttributeName);
+		PCGAttributeAccessorHelpers::AccessorParamResult AccessorResult{};
+		TUniquePtr<const IPCGAttributeAccessor> AttributeAccessor = PCGAttributeAccessorHelpers::CreateConstAccessorForOverrideParamWithResult(InputData, Param, &AccessorResult);
+
+		const FName AttributeName = AccessorResult.AttributeName;
 
 		// Attribute doesn't exist
 		if (!AttributeAccessor)
 		{
+			// Throw a warning if the pin was connected, but accessor failed
+			if (AccessorResult.bPinConnected)
+			{
+				PCGE_LOG_C(Warning, GraphAndLog, this, FText::Format(LOCTEXT("AttributeNotFoundOnConnectedPin", "Override pin '{0}' is connected, but attribute '{1}' was not found."), FText::FromName(Param.Label), FText::FromName(AttributeName)));
+			}
+
 			continue;
 		}
 
+		// If aliases were used, throw a warning to ask the user to update its graph
+		if (AccessorResult.bUsedAliases)
+		{
+			PCGE_LOG_C(Warning, GraphAndLog, this, FText::Format(LOCTEXT("OverrideWithAlias", "Attribute '{0}' was not found, but one of its deprecated aliases ('{1}') was. Please update the name to the new value."), FText::FromName(AttributeName), FText::FromName(AccessorResult.AliasUsed)));
+		}
+
 		TUniquePtr<IPCGAttributeAccessor> PropertyAccessor = PCGAttributeAccessorHelpers::CreatePropertyAccessor(Param.Properties.Last());
 		check(PropertyAccessor.IsValid());
 
diff --git a/Engine/Plugins/Experimental/PCG/Source/PCG/Private/PCGSettings.cpp b/Engine/Plugins/Experimental/PCG/Source/PCG/Private/PCGSettings.cpp
--- a/Engine/Plugins/Experimental/PCG/Source/PCG/Private/PCGSettings.cpp
+++ b/Engine/Plugins/Experimental/PCG/Source/PCG/Private/PCGSettings.cpp
@@ -65,6 +65,74 @@ FString FPCGSettingsOverridableParam::GetPropertyPath() const
 	return FString::JoinBy(Properties, TEXT("/"), [](const FProperty* InProperty) { return InProperty ? InProperty->GetName() : FString(); });
 }
 
+TArray<FName> FPCGSettingsOverridableParam::GenerateAllPossibleAliases() const
+{
+	if (Properties.IsEmpty() || !HasAliases())
+	{
+		return TArray<FName>{};
+	}
+
+	auto ConcatenateNames = [](const FName& Name1, const FName& Name2) -> FName
+	{
+		return FName(FString::Printf(TEXT("%s/%s"), *Name1.ToString(), *Name2.ToString()));
+	};
+
+	// If we have multiple properties, we'll generate all possible aliases, which is a combination of each aliases for each property in the path.
+	TArray<FName> Result;
+	for (int32 i = 0; i < Properties.Num(); ++i)
+	{
+		const FPCGPropertyAliases* It = MapOfAliases.Find(i);
+		const TArray<FName>* CurrentAliases = It ? &It->Aliases : nullptr;
+		const FProperty* CurrentProperty = Properties[i];
+		const FName PropertyName = CurrentProperty ? CurrentProperty->GetFName() : NAME_None;
+
+		// If no alias at this level, just concatenate the property name to all existing aliases
+		if (!CurrentAliases || CurrentAliases->IsEmpty())
+		{
+			if (Result.IsEmpty())
+			{
+				Result.Add(PropertyName);
+			}
+			else
+			{
+				for (FName& Alias : Result)
+				{
+					Alias = ConcatenateNames(Alias, PropertyName);
+				}
+			}
+		}
+		else if (Result.IsEmpty())
+		{
+			if (CurrentAliases)
+			{
+				Result = *CurrentAliases;
+			}
+
+			// Also need to add the property name
+			Result.Add(PropertyName);
+		}
+		else
+		{
+			TArray<FName> Temp;
+			Temp.Reserve(Result.Num() * CurrentAliases->Num());
+			for (const FName& Alias : Result)
+			{
+				for (const FName& OtherAlias : *CurrentAliases)
+				{
+					Temp.Add(ConcatenateNames(Alias, OtherAlias));
+				}
+
+				// Also need to add the property name
+				Temp.Add(ConcatenateNames(Alias, PropertyName));
+			}
+
+			Result = std::move(Temp);
+		}
+	}
+
+	return Result;
+}
+
 #if WITH_EDITOR
 FString FPCGSettingsOverridableParam::GetDisplayPropertyPath() const
 {
diff --git a/Engine/Plugins/Experimental/PCG/Source/PCG/Public/Elements/PCGAttributeNoise.h b/Engine/Plugins/Experimental/PCG/Source/PCG/Public/Elements/PCGAttributeNoise.h
--- a/Engine/Plugins/Experimental/PCG/Source/PCG/Public/Elements/PCGAttributeNoise.h
+++ b/Engine/Plugins/Experimental/PCG/Source/PCG/Public/Elements/PCGAttributeNoise.h
@@ -62,17 +62,17 @@ public:
 	FPCGAttributePropertyOutputSelector OutputTarget;
 
 	/** Attribute = (Original op Noise), Noise in [NoiseMin, NoiseMax] */
-	UPROPERTY(BlueprintReadWrite, EditAnywhere, Category = Settings, meta = (PCG_Overridable))
+	UPROPERTY(BlueprintReadWrite, EditAnywhere, Category = Settings, meta = (PCG_Overridable, PCG_OverrideAliases = "DensityMode"))
 	EPCGAttributeNoiseMode Mode = EPCGAttributeNoiseMode::Set;
 
-	UPROPERTY(BlueprintReadWrite, EditAnywhere, Category = Settings, meta = (PCG_Overridable))
+	UPROPERTY(BlueprintReadWrite, EditAnywhere, Category = Settings, meta = (PCG_Overridable, PCG_OverrideAliases = "DensityNoiseMin"))
 	float NoiseMin = 0.f;
 
-	UPROPERTY(BlueprintReadWrite, EditAnywhere, Category = Settings, meta = (PCG_Overridable))
+	UPROPERTY(BlueprintReadWrite, EditAnywhere, Category = Settings, meta = (PCG_Overridable, PCG_OverrideAliases = "DensityNoiseMax"))
 	float NoiseMax = 1.f;
 
 	/** Attribute = 1 - Attribute before applying the operation */
-	UPROPERTY(BlueprintReadWrite, EditAnywhere, Category = Settings, meta = (PCG_Overridable))
+	UPROPERTY(BlueprintReadWrite, EditAnywhere, Category = Settings, meta = (PCG_Overridable, PCG_OverrideAliases = "bInvertSourceDensity"))
 	bool bInvertSource = false;
 
 	// Clamp the result between 0 and 1. Always applied if we apply noise to the density.
diff --git a/Engine/Plugins/Experimental/PCG/Source/PCG/Public/Helpers/PCGSettingsHelpers.h b/Engine/Plugins/Experimental/PCG/Source/PCG/Public/Helpers/PCGSettingsHelpers.h
--- a/Engine/Plugins/Experimental/PCG/Source/PCG/Public/Helpers/PCGSettingsHelpers.h
+++ b/Engine/Plugins/Experimental/PCG/Source/PCG/Public/Helpers/PCGSettingsHelpers.h
@@ -177,8 +177,10 @@ PRAGMA_ENABLE_DEPRECATION_WARNINGS
 			return false;
 		}
 
-		FName AttributeName = NAME_None;
-		TUniquePtr<const IPCGAttributeAccessor> AttributeAccessor = PCGAttributeAccessorHelpers::CreateConstAccessorForOverrideParam(InInputData, *Param, &AttributeName);
+		PCGAttributeAccessorHelpers::AccessorParamResult AccessorResult{};
+		TUniquePtr<const IPCGAttributeAccessor> AttributeAccessor = PCGAttributeAccessorHelpers::CreateConstAccessorForOverrideParamWithResult(InInputData, *Param, &AccessorResult);
+
+		const FName AttributeName = AccessorResult.AttributeName;
 
 		if (!AttributeAccessor)
 		{
diff --git a/Engine/Plugins/Experimental/PCG/Source/PCG/Public/Metadata/Accessors/PCGAttributeAccessorHelpers.h b/Engine/Plugins/Experimental/PCG/Source/PCG/Public/Metadata/Accessors/PCGAttributeAccessorHelpers.h
--- a/Engine/Plugins/Experimental/PCG/Source/PCG/Public/Metadata/Accessors/PCGAttributeAccessorHelpers.h
+++ b/Engine/Plugins/Experimental/PCG/Source/PCG/Public/Metadata/Accessors/PCGAttributeAccessorHelpers.h
@@ -31,10 +31,21 @@ namespace PCGAttributeAccessorHelpers
 
 	PCG_API TUniquePtr<IPCGAttributeAccessor> CreateChainAccessor(TUniquePtr<IPCGAttributeAccessor> InAccessor, FName Name, bool& bOutSuccess);
 
+	struct AccessorParamResult
+	{
+		FName AttributeName = NAME_None;
+		FName AliasUsed = NAME_None;
+		bool bUsedAliases = false;
+		bool bPinConnected = false;
+	};
+
+	UE_DEPRECATED(5.3, "Use the CreateConstAccessorForOverrideParamWithResult version")
+	PCG_API TUniquePtr<const IPCGAttributeAccessor> CreateConstAccessorForOverrideParam(const FPCGDataCollection& InInputData, const FPCGSettingsOverridableParam& InParam, FName* OutAttributeName = nullptr);
+
 	/**
 	* Create a const accessor depending on an overridable param
 	*/
-	PCG_API TUniquePtr<const IPCGAttributeAccessor> CreateConstAccessorForOverrideParam(const FPCGDataCollection& InInputData, const FPCGSettingsOverridableParam& InParam, FName* OutAttributeName = nullptr);
+	PCG_API TUniquePtr<const IPCGAttributeAccessor> CreateConstAccessorForOverrideParamWithResult(const FPCGDataCollection& InInputData, const FPCGSettingsOverridableParam& InParam, AccessorParamResult* OutResult = nullptr);
 
 	/** 
 	* Creates a const accessor to the property or attribute pointed at by the InSelector.
diff --git a/Engine/Plugins/Experimental/PCG/Source/PCG/Public/PCGCommon.h b/Engine/Plugins/Experimental/PCG/Source/PCG/Public/PCGCommon.h
--- a/Engine/Plugins/Experimental/PCG/Source/PCG/Public/PCGCommon.h
+++ b/Engine/Plugins/Experimental/PCG/Source/PCG/Public/PCGCommon.h
@@ -103,6 +103,7 @@ namespace PCGObjectMetadata
 {
 	const FName Overridable = TEXT("PCG_Overridable");
 	const FName NotOverridable = TEXT("PCG_NotOverridable");
+	const FName OverrideAliases = TEXT("PCG_OverrideAliases");
 
 	// Metadata usable in UPROPERTY for customizing the behavior when displaying the property in a property panel or graph node
 	enum
@@ -111,7 +112,10 @@ namespace PCGObjectMetadata
 		PCG_Overridable,
 
 		/// [PropertyMetadata] Indicates that the property is not-overridable by params. Used in structs to hide some parameters
-		PCG_NotOverridable
+		PCG_NotOverridable,
+
+		/// [PropertyMetadata] Extra names to match for a given property.
+		PCG_OverrideAliases,
 	};
 }
 
diff --git a/Engine/Plugins/Experimental/PCG/Source/PCG/Public/PCGSettings.h b/Engine/Plugins/Experimental/PCG/Source/PCG/Public/PCGSettings.h
--- a/Engine/Plugins/Experimental/PCG/Source/PCG/Public/PCGSettings.h
+++ b/Engine/Plugins/Experimental/PCG/Source/PCG/Public/PCGSettings.h
@@ -52,6 +52,16 @@ enum class EPCGSettingsType : uint8
 	DECLARE_MULTICAST_DELEGATE_TwoParams(FOnPCGSettingsChanged, UPCGSettings*, EPCGChangeType);
 #endif
 
+// Dummy struct to bypass the UHT limitation for array of arrays.
+USTRUCT(meta=(Hidden))
+struct FPCGPropertyAliases
+{
+	GENERATED_BODY()
+
+	UPROPERTY()
+	TArray<FName> Aliases;
+};
+
 USTRUCT()
 struct FPCGSettingsOverridableParam
 {
@@ -66,15 +76,23 @@ struct FPCGSettingsOverridableParam
 	UPROPERTY()
 	TObjectPtr<const UStruct> PropertyClass;
 
+	// Map of all aliases for a given property, using its Index (to avoid name clashes within the same path)
+	UPROPERTY()
+	TMap<int32, FPCGPropertyAliases> MapOfAliases;
+
 	// If this flag is true, Label will be the full property path.
 	UPROPERTY()
 	bool bHasNameClash = false;
 
+	bool HasAliases() const { return !MapOfAliases.IsEmpty(); }
+
 	// Transient
 	TArray<const FProperty*> Properties;
 
 	FString GetPropertyPath() const;
 
+	TArray<FName> GenerateAllPossibleAliases() const;
+
 #if WITH_EDITOR
 	FString GetDisplayPropertyPath() const;
 #endif // WITH_EDITOR
