[PCG] Tentative fix for build break: missing header

#rnx

diff --git a/Engine/Plugins/Experimental/PCG/Source/PCG/Public/Metadata/Accessors/PCGAttributeAccessorHelpers.h b/Engine/Plugins/Experimental/PCG/Source/PCG/Public/Metadata/Accessors/PCGAttributeAccessorHelpers.h
--- a/Engine/Plugins/Experimental/PCG/Source/PCG/Public/Metadata/Accessors/PCGAttributeAccessorHelpers.h
+++ b/Engine/Plugins/Experimental/PCG/Source/PCG/Public/Metadata/Accessors/PCGAttributeAccessorHelpers.h
@@ -3,6 +3,7 @@
 #pragma once
 
 #include "Templates/UniquePtr.h"
+#include "UObject/NameTypes.h"
 
 enum class EPCGExtraProperties : uint8;
 class IPCGAttributeAccessor;
