[PCG] Attempted hotfix for crash in loading deprecated PCG Graphs

#rnx

diff --git a/Engine/Plugins/Experimental/PCG/Source/PCG/Private/PCGGraph.cpp b/Engine/Plugins/Experimental/PCG/Source/PCG/Private/PCGGraph.cpp
--- a/Engine/Plugins/Experimental/PCG/Source/PCG/Private/PCGGraph.cpp
+++ b/Engine/Plugins/Experimental/PCG/Source/PCG/Private/PCGGraph.cpp
@@ -355,7 +355,7 @@ UPCGNode* UPCGGraph::AddNodeOfType(TSubclassOf<class UPCGSettings> InSettingsCla
 
 	if (Node)
 	{
-		Settings->Rename(nullptr, Node);
+		Settings->Rename(nullptr, Node, REN_ForceNoResetLoaders | REN_DontCreateRedirectors);
 	}
 
 	OutDefaultNodeSettings = Settings;
@@ -381,7 +381,7 @@ UPCGNode* UPCGGraph::AddNode(UPCGSettingsInterface* InSettingsInterface)
 		Node->SetSettingsInterface(InSettingsInterface);
 
 		// Reparent node to this graph
-		Node->Rename(nullptr, this);
+		Node->Rename(nullptr, this, REN_ForceNoResetLoaders | REN_DontCreateRedirectors);
 
 #if WITH_EDITOR
 		const FName DefaultNodeName = InSettingsInterface->GetSettings()->GetDefaultNodeName();
