Added missing controls for the Create Target Actor node: Ability to set base name + Ability to set actor pivot
Added metadata string operations & more string support
Added capability in the context to determine if a proper was overriden

diff --git a/Engine/Plugins/Experimental/PCG/Source/PCG/Private/Elements/Metadata/PCGMetadataStringOpElement.cpp b/Engine/Plugins/Experimental/PCG/Source/PCG/Private/Elements/Metadata/PCGMetadataStringOpElement.cpp
new file mode 100644
--- /dev/null
+++ b/Engine/Plugins/Experimental/PCG/Source/PCG/Private/Elements/Metadata/PCGMetadataStringOpElement.cpp
@@ -0,0 +1,86 @@
+// Copyright Epic Games, Inc. All Rights Reserved.
+
+#include "Elements/Metadata/PCGMetadataStringOpElement.h"
+
+#include "Elements/Metadata/PCGMetadataElementCommon.h"
+
+#include UE_INLINE_GENERATED_CPP_BY_NAME(PCGMetadataStringOpElement)
+
+FName UPCGMetadataStringOpSettings::GetInputPinLabel(uint32 Index) const
+{
+	switch (Index)
+	{
+	case 0: return PCGMetadataSettingsBaseConstants::DoubleInputFirstLabel;
+	case 1: return PCGMetadataSettingsBaseConstants::DoubleInputSecondLabel;
+	default: return NAME_None;
+	}
+}
+
+uint32 UPCGMetadataStringOpSettings::GetInputPinNum() const
+{
+	return 2;
+}
+
+bool UPCGMetadataStringOpSettings::IsSupportedInputType(uint16 TypeId, uint32 InputIndex, bool& bHasSpecialRequirement) const
+{
+	bHasSpecialRequirement = false;
+	return true; // all types support ToString()
+}
+
+FPCGAttributePropertyInputSelector UPCGMetadataStringOpSettings::GetInputSource(uint32 Index) const
+{
+	switch (Index)
+	{
+	case 0: return InputSource1;
+	case 1: return InputSource2;
+	default: return FPCGAttributePropertyInputSelector();
+	}
+}
+
+FName UPCGMetadataStringOpSettings::AdditionalTaskName() const
+{
+	if (const UEnum* EnumPtr = StaticEnum<EPCGMetadataStringOperation>())
+	{
+		return FName(FString("String: ") + EnumPtr->GetNameStringByValue(static_cast<int>(Operation)));
+	}
+	else
+	{
+		return NAME_None;
+	}
+}
+ 
+#if WITH_EDITOR
+FName UPCGMetadataStringOpSettings::GetDefaultNodeName() const
+{
+	return TEXT("AttributeStringOp");
+}
+
+FText UPCGMetadataStringOpSettings::GetDefaultNodeTitle() const
+{
+	return NSLOCTEXT("PCGMetadataStringOpSettings", "NodeTitle", "Attribute String Op");
+}
+#endif // WITH_EDITOR
+
+FPCGElementPtr UPCGMetadataStringOpSettings::CreateElement() const
+{
+	return MakeShared<FPCGMetadataStringOpElement>();
+}
+
+uint16 UPCGMetadataStringOpSettings::GetOutputType(uint16 InputTypeId) const
+{
+	return (uint16)EPCGMetadataTypes::String;
+}
+
+bool FPCGMetadataStringOpElement::DoOperation(FOperationData& OperationData) const
+{
+	TRACE_CPUPROFILER_EVENT_SCOPE(FPCGMetadataStringOpElement::Execute);
+
+	const UPCGMetadataStringOpSettings* Settings = CastChecked<UPCGMetadataStringOpSettings>(OperationData.Settings);
+
+	if (Settings->Operation == EPCGMetadataStringOperation::Append)
+	{
+		DoBinaryOp<FString, FString>(OperationData, [](const FString& Value1, const FString& Value2) -> FString { return Value1 + Value2; });
+	}
+
+	return true;
+}
\ No newline at end of file
diff --git a/Engine/Plugins/Experimental/PCG/Source/PCG/Private/Elements/PCGCreateTargetActor.cpp b/Engine/Plugins/Experimental/PCG/Source/PCG/Private/Elements/PCGCreateTargetActor.cpp
--- a/Engine/Plugins/Experimental/PCG/Source/PCG/Private/Elements/PCGCreateTargetActor.cpp
+++ b/Engine/Plugins/Experimental/PCG/Source/PCG/Private/Elements/PCGCreateTargetActor.cpp
@@ -224,6 +224,11 @@ bool FPCGCreateTargetActorElement::ExecuteInternal(FPCGContext* Context) const
 	}
 
 	FTransform Transform = TargetActor->GetTransform();
+	if(Context->IsValueOverriden(GET_MEMBER_NAME_CHECKED(UPCGCreateTargetActor, ActorPivot)))
+	{
+		Transform = Settings->ActorPivot;
+	}
+
 	AActor* GeneratedActor = UPCGActorHelpers::SpawnDefaultActor(TargetActor->GetWorld(), Settings->TemplateActorClass, Transform, SpawnParams, TargetActor);
 
 	if (!GeneratedActor)
@@ -232,6 +237,13 @@ bool FPCGCreateTargetActorElement::ExecuteInternal(FPCGContext* Context) const
 		return true;
 	}
 
+#if WITH_EDITOR
+	if (Settings->ActorLabel != FString())
+	{
+		GeneratedActor->SetActorLabel(Settings->ActorLabel);
+	}
+#endif
+
 	GeneratedActor->Tags.Add(PCGHelpers::DefaultPCGActorTag);
 
 	if (UPCGComponent* SourceComponent = Context->SourceComponent.Get())
diff --git a/Engine/Plugins/Experimental/PCG/Source/PCG/Private/PCGContext.cpp b/Engine/Plugins/Experimental/PCG/Source/PCG/Private/PCGContext.cpp
--- a/Engine/Plugins/Experimental/PCG/Source/PCG/Private/PCGContext.cpp
+++ b/Engine/Plugins/Experimental/PCG/Source/PCG/Private/PCGContext.cpp
@@ -218,7 +218,7 @@ void FPCGContext::OverrideSettings()
 		TUniquePtr<IPCGAttributeAccessor> PropertyAccessor = PCGAttributeAccessorHelpers::CreatePropertyAccessor(Param.Properties.Last());
 		check(PropertyAccessor.IsValid());
 
-		PCGMetadataAttribute::CallbackWithRightType(PropertyAccessor->GetUnderlyingType(), [this, &AttributeAccessor, &PropertyAccessor, &Param, &AttributeName, Container](auto Dummy) -> bool
+		const bool bParamOverridden = PCGMetadataAttribute::CallbackWithRightType(PropertyAccessor->GetUnderlyingType(), [this, &AttributeAccessor, &PropertyAccessor, &Param, &AttributeName, Container](auto Dummy) -> bool
 		{
 			using PropertyType = decltype(Dummy);
 
@@ -244,9 +244,24 @@ void FPCGContext::OverrideSettings()
 
 			return true;
 		});
+
+		if (bParamOverridden)
+		{
+			OverriddenParams.Add(&Param);
+		}
 	}
 }
 
+bool FPCGContext::IsValueOverriden(const FName PropertyName)
+{
+	const FPCGSettingsOverridableParam** OverriddenParam = OverriddenParams.FindByPredicate([PropertyName](const FPCGSettingsOverridableParam* ParamToCheck)
+	{
+		return ParamToCheck && !ParamToCheck->PropertiesNames.IsEmpty() && ParamToCheck->PropertiesNames[0] == PropertyName;
+	});
+
+	return OverriddenParam != nullptr;
+}
+
 #if WITH_EDITOR
 void FPCGContext::LogVisual(ELogVerbosity::Type InVerbosity, const FText& InMessage) const
 {
diff --git a/Engine/Plugins/Experimental/PCG/Source/PCG/Public/Elements/Metadata/PCGMetadataStringOpElement.h b/Engine/Plugins/Experimental/PCG/Source/PCG/Public/Elements/Metadata/PCGMetadataStringOpElement.h
new file mode 100644
--- /dev/null
+++ b/Engine/Plugins/Experimental/PCG/Source/PCG/Public/Elements/Metadata/PCGMetadataStringOpElement.h
@@ -0,0 +1,58 @@
+// Copyright Epic Games, Inc. All Rights Reserved.
+
+#pragma once
+
+#include "Elements/Metadata/PCGMetadataOpElementBase.h"
+
+#include "PCGMetadataStringOpElement.generated.h"
+
+UENUM()
+enum class EPCGMetadataStringOperation : uint16
+{
+	Append
+};
+
+UCLASS(BlueprintType, ClassGroup = (Procedural))
+class PCG_API UPCGMetadataStringOpSettings : public UPCGMetadataSettingsBase
+{
+	GENERATED_BODY()
+
+public:
+	//~Begin UPCGSettings interface
+#if WITH_EDITOR
+	virtual FName GetDefaultNodeName() const override;
+	virtual FText GetDefaultNodeTitle() const override;
+#endif
+	virtual FName AdditionalTaskName() const override;
+
+protected:
+	virtual FPCGElementPtr CreateElement() const override;
+	//~End UPCGSettings interface
+
+public:
+	//~Begin UPCGMetadataSettingsBase interface
+	FPCGAttributePropertyInputSelector GetInputSource(uint32 Index) const override;
+
+	virtual FName GetInputPinLabel(uint32 Index) const override;
+	virtual uint32 GetInputPinNum() const override;
+
+	virtual bool IsSupportedInputType(uint16 TypeId, uint32 InputIndex, bool& bHasSpecialRequirement) const override;
+	virtual uint16 GetOutputType(uint16 InputTypeId) const override;
+	//~End UPCGMetadataSettingsBase interface
+
+public:
+	UPROPERTY(BlueprintReadWrite, EditAnywhere, Category = Settings)
+	EPCGMetadataStringOperation Operation = EPCGMetadataStringOperation::Append;
+
+	UPROPERTY(BlueprintReadWrite, EditAnywhere, Category = Input)
+	FPCGAttributePropertyInputSelector InputSource1;
+
+	UPROPERTY(BlueprintReadWrite, EditAnywhere, Category = Input)
+	FPCGAttributePropertyInputSelector InputSource2;
+};
+
+class FPCGMetadataStringOpElement : public FPCGMetadataElementBase
+{
+protected:
+	virtual bool DoOperation(FOperationData& OperationData) const override;
+};
\ No newline at end of file
diff --git a/Engine/Plugins/Experimental/PCG/Source/PCG/Public/Elements/PCGCreateTargetActor.h b/Engine/Plugins/Experimental/PCG/Source/PCG/Public/Elements/PCGCreateTargetActor.h
--- a/Engine/Plugins/Experimental/PCG/Source/PCG/Public/Elements/PCGCreateTargetActor.h
+++ b/Engine/Plugins/Experimental/PCG/Source/PCG/Public/Elements/PCGCreateTargetActor.h
@@ -59,6 +59,12 @@ public:
 
 	UPROPERTY(BlueprintReadWrite, Category = Settings, meta = (PCG_Overridable))
 	TSoftObjectPtr<AActor> RootActor;
+
+	UPROPERTY(BlueprintReadWrite, Category = Settings, meta = (PCG_Overridable))
+	FString ActorLabel;
+
+	UPROPERTY(BlueprintReadWrite, Category = Settings, meta = (PCG_Overridable))
+	FTransform ActorPivot;
 };
 
 class FPCGCreateTargetActorElement : public FSimplePCGElement
diff --git a/Engine/Plugins/Experimental/PCG/Source/PCG/Public/Metadata/PCGMetadataAttributeTraits.h b/Engine/Plugins/Experimental/PCG/Source/PCG/Public/Metadata/PCGMetadataAttributeTraits.h
--- a/Engine/Plugins/Experimental/PCG/Source/PCG/Public/Metadata/PCGMetadataAttributeTraits.h
+++ b/Engine/Plugins/Experimental/PCG/Source/PCG/Public/Metadata/PCGMetadataAttributeTraits.h
@@ -151,6 +151,20 @@
 				PCGMetadataBroadcastable(Quaternion, Rotator);
 				PCGMetadataBroadcastable(Rotator, Quaternion);
 
+				PCGMetadataBroadcastable(Float, String);
+				PCGMetadataBroadcastable(Double, String);
+				PCGMetadataBroadcastable(Integer32, String);
+				PCGMetadataBroadcastable(Integer64, String);
+				PCGMetadataBroadcastable(Vector2, String);
+				PCGMetadataBroadcastable(Vector, String);
+				PCGMetadataBroadcastable(Vector4, String);
+				PCGMetadataBroadcastable(Quaternion, String);
+				PCGMetadataBroadcastable(Transform, String);
+				PCGMetadataBroadcastable(Boolean, String);
+				PCGMetadataBroadcastable(Rotator, String);
+				PCGMetadataBroadcastable(Name, String);
+				PCGMetadataBroadcastable(String, Name);
+
 #undef PCGMetadataBroadcastable
 			}
 
@@ -197,70 +211,23 @@
 			return IsMoreComplexType(MetadataTypes<FirstType>::Id, MetadataTypes<SecondType>::Id);
 		}
 
-		/**
-		* Generic function to broadcast an InType to an OutType.
-		* Supports only PCG types
-		* @param InValue - The Value to convert
-		* @param OutValue - The converted value
-		* @returns true if the conversion worked, false otherwise.
-		*/
-		template <typename InType, typename OutType>
-		inline bool GetValueWithBroadcast(const InType& InValue, OutType& OutValue)
+		template<typename T>
+		struct DefaultStringTraits
 		{
-			if constexpr (std::is_same_v<OutType, InType>)
+			static FString ToString(const T& A)
 			{
-				OutValue = InValue;
-				return true;
+				return A.ToString();
 			}
-			else
-			{
-				if constexpr (!IsBroadcastable<InType, OutType>())
-				{
-					return false;
-				}
-				else
-				{
-					if constexpr (std::is_same_v<OutType, FVector4>)
-					{
-						// TODO: Should W be 0? 1? Something else? Depending on operation?
-						// For now, we set Z to 0 (for vec 2) and we set W to 1.
-						if constexpr (std::is_same_v<InType, FVector>)
-						{
-							OutValue = FVector4(InValue, 1.0);
-						}
-						else if constexpr (std::is_same_v<InType, FVector2D>)
-						{
-							OutValue = FVector4(InValue.X, InValue.Y, 0.0, 1.0);
-						}
-						else
-						{
-							OutValue = FVector4(InValue, InValue, InValue, InValue);
-						}
-					}
-					else
-					{
-						// Seems like the && condition is not evaluated correctly on Linux, so we cut the condition in two `if constexpr`.
-						if constexpr (std::is_same_v<OutType, FVector>)
-						{
-							if constexpr (std::is_same_v<InType, FVector2D>)
-							{
-								OutValue = FVector(InValue, 0.0);
-							}
-							else
-							{
-								OutValue = OutType(InValue);
-							}
-						}
-						else
-						{
-							OutValue = OutType(InValue);
-						}
-					}
-
-					return true;
-				}
-			}
-		}
+		};
 
 		template<typename T>
+		struct LexToStringTraits
+		{
+			static FString ToString(const T& A)
+			{
+				return LexToString(A);
+			}
+		};
+
+		template<typename T>
 		struct DefaultOperationTraits
@@ -354,14 +321,14 @@
 
 		// Common traits for int32, int64, float, double
 		template<typename T>
-		struct MetadataTraits : DefaultOperationTraits<T>, DefaultWeightedSumTraits<T>, DefaultMinMaxTraits<T>, DefaultCompareTraits<T>
+		struct MetadataTraits : DefaultOperationTraits<T>, DefaultWeightedSumTraits<T>, DefaultMinMaxTraits<T>, DefaultCompareTraits<T>, LexToStringTraits<T>
 		{
 			enum { CompressData = false };
 			enum { CanSearchString = false };
 		};
 
 		template<>
-		struct MetadataTraits<float> : DefaultOperationTraits<float>, DefaultWeightedSumTraits<float>, DefaultMinMaxTraits<float>, DefaultCompareTraits<float>
+		struct MetadataTraits<float> : DefaultOperationTraits<float>, DefaultWeightedSumTraits<float>, DefaultMinMaxTraits<float>, DefaultCompareTraits<float>, LexToStringTraits<float>
 		{
 			enum { CompressData = false };
 			enum { CanSearchString = false };
@@ -373,7 +340,7 @@
 		};
 
 		template<>
-		struct MetadataTraits<double> : DefaultOperationTraits<double>, DefaultWeightedSumTraits<double>, DefaultMinMaxTraits<double>, DefaultCompareTraits<double>
+		struct MetadataTraits<double> : DefaultOperationTraits<double>, DefaultWeightedSumTraits<double>, DefaultMinMaxTraits<double>, DefaultCompareTraits<double>, LexToStringTraits<double>
 		{
 			enum { CompressData = false };
 			enum { CanSearchString = false };
@@ -385,7 +352,7 @@
 		};
 
 		template<>
-		struct MetadataTraits<bool>
+		struct MetadataTraits<bool> : LexToStringTraits<bool>
 		{
 			enum { CompressData = false };
 			enum { CanMinMax = true };
@@ -448,7 +415,7 @@
 
 		// Vector types
 		template<typename T>
-		struct VectorTraits : DefaultOperationTraits<T>, DefaultWeightedSumTraits<T>
+		struct VectorTraits : DefaultOperationTraits<T>, DefaultWeightedSumTraits<T>, DefaultStringTraits<T>
 		{
 			enum { CompressData = false };
 			enum { CanMinMax = true };
@@ -581,7 +548,7 @@
 
 		// Quaternion
 		template<>
-		struct MetadataTraits<FQuat>
+		struct MetadataTraits<FQuat> : DefaultStringTraits<FQuat>
 		{
 			enum { CompressData = false };
 			enum { CanMinMax = false };
@@ -635,7 +602,7 @@
 
 		// Rotator
 		template<>
-		struct MetadataTraits<FRotator>
+		struct MetadataTraits<FRotator> : DefaultStringTraits<FRotator>
 		{
 			enum { CompressData = false };
 			enum { CanMinMax = false };
@@ -684,7 +651,7 @@
 
 		// Transform
 		template<>
-		struct MetadataTraits<FTransform>
+		struct MetadataTraits<FTransform> : DefaultStringTraits<FTransform>
 		{
 			enum { CompressData = false };
 			enum { CanMinMax = false };
@@ -736,7 +703,7 @@
 
 		// Strings
 		template<>
-		struct MetadataTraits<FString> : DefaultCompareTraits<FString>
+		struct MetadataTraits<FString> : DefaultCompareTraits<FString>, DefaultStringTraits<FString>
 		{
 			enum { CompressData = true };
 			enum { CanMinMax = false };
@@ -767,7 +734,7 @@
 		};
 
 		template<>
-		struct MetadataTraits<FName>
+		struct MetadataTraits<FName> : DefaultStringTraits<FName>
 		{
 			enum { CompressData = false };
 			enum { CanMinMax = false };
@@ -817,5 +784,77 @@
 				return A.ToString().MatchesWildcard(B.ToString());
 			}
 		};
+
+		/**
+		* Generic function to broadcast an InType to an OutType.
+		* Supports only PCG types
+		* @param InValue - The Value to convert
+		* @param OutValue - The converted value
+		* @returns true if the conversion worked, false otherwise.
+		*/
+		template <typename InType, typename OutType>
+		inline bool GetValueWithBroadcast(const InType& InValue, OutType& OutValue)
+		{
+			if constexpr (std::is_same_v<OutType, InType>)
+			{
+				OutValue = InValue;
+				return true;
+			}
+			else
+			{
+				if constexpr (!IsBroadcastable<InType, OutType>())
+				{
+					return false;
+				}
+				else
+				{
+					if constexpr (std::is_same_v<OutType, FVector4>)
+					{
+						// TODO: Should W be 0? 1? Something else? Depending on operation?
+						// For now, we set Z to 0 (for vec 2) and we set W to 1.
+						if constexpr (std::is_same_v<InType, FVector>)
+						{
+							OutValue = FVector4(InValue, 1.0);
+						}
+						else if constexpr (std::is_same_v<InType, FVector2D>)
+						{
+							OutValue = FVector4(InValue.X, InValue.Y, 0.0, 1.0);
+						}
+						else
+						{
+							OutValue = FVector4(InValue, InValue, InValue, InValue);
+						}
+					}
+					else
+					{
+						// Seems like the && condition is not evaluated correctly on Linux, so we cut the condition in two `if constexpr`.
+						if constexpr (std::is_same_v<OutType, FVector>)
+						{
+							if constexpr (std::is_same_v<InType, FVector2D>)
+							{
+								OutValue = FVector(InValue, 0.0);
+							}
+							else
+							{
+								OutValue = OutType(InValue);
+							}
+						}
+						else
+						{
+							if constexpr (std::is_same_v<OutType, FString>)
+							{
+								OutValue = MetadataTraits<InType>::ToString(InValue);
+							}
+							else
+							{
+								OutValue = OutType(InValue);
+							}
+						}
+					}
+
+					return true;
+				}
+			}
+		}
 	}
 }
diff --git a/Engine/Plugins/Experimental/PCG/Source/PCG/Public/PCGContext.h b/Engine/Plugins/Experimental/PCG/Source/PCG/Public/PCGContext.h
--- a/Engine/Plugins/Experimental/PCG/Source/PCG/Public/PCGContext.h
+++ b/Engine/Plugins/Experimental/PCG/Source/PCG/Public/PCGContext.h
@@ -84,6 +84,9 @@ struct PCG_API FPCGContext
 	// in the settings copy.
 	void OverrideSettings();
 
+	// Returns true if the given property has been overriden
+	bool IsValueOverriden(const FName PropertyName);
+
 	// Return the seed, possibly overriden by params, and combined with the source component (if any).
 	int GetSeed() const;
 
@@ -122,4 +125,7 @@ private:
 
 	// Copy of the settings that will be used to apply overrides.
 	TObjectPtr<UPCGSettings> SettingsWithOverride = nullptr;
+
+	// List of params that were in effect overriden
+	TArray<const FPCGSettingsOverridableParam*> OverriddenParams;
 };
