[PCG] Added support for recursive expansion/collapse in the debug object tree view (shift click arrow)

#rnx

diff --git a/Engine/Plugins/Experimental/PCG/Source/PCGEditor/Private/SPCGEditorGraphDebugObjectTree.cpp b/Engine/Plugins/Experimental/PCG/Source/PCGEditor/Private/SPCGEditorGraphDebugObjectTree.cpp
--- a/Engine/Plugins/Experimental/PCG/Source/PCGEditor/Private/SPCGEditorGraphDebugObjectTree.cpp
+++ b/Engine/Plugins/Experimental/PCG/Source/PCGEditor/Private/SPCGEditorGraphDebugObjectTree.cpp
@@ -121,6 +121,7 @@ void SPCGEditorGraphDebugObjectTree::Construct(const FArguments& InArgs, TShared
 				.OnGenerateRow(this, &SPCGEditorGraphDebugObjectTree::MakeTreeRowWidget)
 				.OnGetChildren(this, &SPCGEditorGraphDebugObjectTree::OnGetChildren)
 				.OnSelectionChanged(this, &SPCGEditorGraphDebugObjectTree::OnSelectionChanged)
+				.OnSetExpansionRecursive(this, &SPCGEditorGraphDebugObjectTree::OnSetExpansionRecursive)
 				.ItemHeight(18)
 				.AllowOverscroll(EAllowOverscroll::No)
 				.ExternalScrollbar(VerticalScrollBar)
@@ -559,4 +560,21 @@ void SPCGEditorGraphDebugObjectTree::OnSelectionChanged(FPCGEditorGraphDebugObje
 	}
 }
 
+void SPCGEditorGraphDebugObjectTree::OnSetExpansionRecursive(FPCGEditorGraphDebugObjectItemPtr InItem, bool bInExpand)
+{
+	if (!InItem.IsValid() || !DebugObjectTreeView.IsValid())
+	{
+		return;
+	}
+
+	DebugObjectTreeView->SetItemExpansion(InItem, bInExpand);
+	for (const FPCGEditorGraphDebugObjectItemPtr& ChildItem : InItem->GetChildren())
+	{
+		if (ChildItem.IsValid())
+		{
+			OnSetExpansionRecursive(ChildItem, bInExpand);
+		}
+	}
+}
+
 #undef LOCTEXT_NAMESPACE
diff --git a/Engine/Plugins/Experimental/PCG/Source/PCGEditor/Private/SPCGEditorGraphDebugObjectTree.h b/Engine/Plugins/Experimental/PCG/Source/PCGEditor/Private/SPCGEditorGraphDebugObjectTree.h
--- a/Engine/Plugins/Experimental/PCG/Source/PCGEditor/Private/SPCGEditorGraphDebugObjectTree.h
+++ b/Engine/Plugins/Experimental/PCG/Source/PCGEditor/Private/SPCGEditorGraphDebugObjectTree.h
@@ -149,6 +149,7 @@ private:
 	TSharedRef<ITableRow> MakeTreeRowWidget(FPCGEditorGraphDebugObjectItemPtr InItem, const TSharedRef<STableViewBase>& InOwnerTable) const;
 	void OnGetChildren(FPCGEditorGraphDebugObjectItemPtr InItem, TArray<FPCGEditorGraphDebugObjectItemPtr>& OutChildren) const;
 	void OnSelectionChanged(FPCGEditorGraphDebugObjectItemPtr InItem, ESelectInfo::Type InSelectInfo) const;
+	void OnSetExpansionRecursive(FPCGEditorGraphDebugObjectItemPtr InItem, bool bInExpand);
 	
 	TWeakPtr<FPCGEditor> PCGEditor;
 
