[PCG] Added PCGTextureReadback to expand texture format support for PCGTextureData

* Added PCGCompute module to hold compute shaders and mediate dispatch
* Expanded PCGTextureData to support to all texture formats and settings
* Provided asynchronous initialization for PCGTextureData

diff --git a/Engine/Plugins/Experimental/PCG/PCG.uplugin b/Engine/Plugins/Experimental/PCG/PCG.uplugin
--- a/Engine/Plugins/Experimental/PCG/PCG.uplugin
+++ b/Engine/Plugins/Experimental/PCG/PCG.uplugin
@@ -24,8 +24,12 @@
 			"Name": "PCGEditor",
 			"Type": "Editor",
 			"LoadingPhase": "Default"
+		},
+		{
+			"Name": "PCGCompute",
+			"Type": "Runtime",
+			"LoadingPhase": "PostConfigInit"
 		}
-
 	],
 	"Plugins": [
 		{
diff --git a/Engine/Plugins/Experimental/PCG/Shaders/Private/PCGTextureReadback.usf b/Engine/Plugins/Experimental/PCG/Shaders/Private/PCGTextureReadback.usf
new file mode 100644
--- /dev/null
+++ b/Engine/Plugins/Experimental/PCG/Shaders/Private/PCGTextureReadback.usf
@@ -0,0 +1,15 @@
+// Copyright Epic Games, Inc. All Rights Reserved.
+
+#include "/Engine/Public/Platform.ush"
+
+Texture2D SourceTexture;
+SamplerState SourceSampler;
+float2 SourceDimensions;
+RWTexture2D<float4> OutputTexture;
+
+[numthreads(THREADGROUPSIZE_X, THREADGROUPSIZE_Y, THREADGROUPSIZE_Z)]
+void PCGTextureReadback_CS(uint3 DispatchThreadId : SV_DispatchThreadID, uint GroupIndex : SV_GroupIndex)
+{
+	float2 TexCoord = DispatchThreadId.xy + 0.5;
+	OutputTexture[TexCoord] = SourceTexture.SampleLevel(SourceSampler, TexCoord / SourceDimensions, /*lod=*/0);
+}
diff --git a/Engine/Plugins/Experimental/PCG/Source/PCG/PCG.Build.cs b/Engine/Plugins/Experimental/PCG/Source/PCG/PCG.Build.cs
--- a/Engine/Plugins/Experimental/PCG/Source/PCG/PCG.Build.cs
+++ b/Engine/Plugins/Experimental/PCG/Source/PCG/PCG.Build.cs
@@ -26,6 +26,7 @@ namespace UnrealBuildTool.Rules
 			PrivateDependencyModuleNames.AddRange(
 				new string[]
 				{
+					"PCGCompute",
 					"Voronoi",
 				}
 			);
@@ -41,10 +42,6 @@ namespace UnrealBuildTool.Rules
 					}
 				);
 			}
-
-			PrivateIncludePaths.AddRange(
-				new string[] {
-				});
 		}
 	}
 }
diff --git a/Engine/Plugins/Experimental/PCG/Source/PCG/Private/Data/PCGTextureData.cpp b/Engine/Plugins/Experimental/PCG/Source/PCG/Private/Data/PCGTextureData.cpp
--- a/Engine/Plugins/Experimental/PCG/Source/PCG/Private/Data/PCGTextureData.cpp
+++ b/Engine/Plugins/Experimental/PCG/Source/PCG/Private/Data/PCGTextureData.cpp
@@ -2,6 +2,7 @@
 
 #include "Data/PCGTextureData.h"
 
+#include "PCGTextureReadback.h"
 #include "Data/PCGPointData.h"
 #include "Helpers/PCGAsync.h"
 #include "Helpers/PCGHelpers.h"
@@ -267,82 +268,101 @@ void UPCGBaseTextureData::CopyBaseTextureData(UPCGBaseTextureData* NewTextureDat
 	NewTextureData->Width = Width;
 }
 
-void UPCGTextureData::Initialize(UTexture2D* InTexture, const FTransform& InTransform)
+void UPCGTextureData::Initialize(UTexture2D* InTexture, const FTransform& InTransform, const TFunction<void()>& PostInitializeCallback)
 {
+	TRACE_CPUPROFILER_EVENT_SCOPE(UPCGTextureData::Initialize);
+
 	Texture = InTexture;
 	Transform = InTransform;
 	Width = 0;
 	Height = 0;
 
-	if (InTexture)
+	Bounds = FBox(EForceInit::ForceInit);
+	Bounds += FVector(-1.0f, -1.0f, 0.0f);
+	Bounds += FVector(1.0f, 1.0f, 0.0f);
+	Bounds = Bounds.TransformBy(Transform);
+
+	if (!InTexture)
 	{
-		if (IsSupported(InTexture))
+		PostInitializeCallback();
+		return;
+	}
+
+	if (!IsSupported(InTexture))
+	{
+		UE_LOG(LogPCG, Error, TEXT("PCGTextureData does not support one or more settings on '%s'"), *Texture->GetFName().ToString());
+
+		PostInitializeCallback();
+		return;
+	}
+
+#if WITH_EDITOR
+	// Force a wait on any incomplete async texture compilation and caching operations
+	Texture->FinishCachePlatformData();
+#endif
+
+	// TODO: Forcing a resource update is problematic. Clears cooked platform data and may break other usages of this texture
+	// But it is necessary to force the texture memory to be present on the first dispatch
+	Texture->UpdateResource();
+	Texture->WaitForPendingInitOrStreaming();
+
+	FTexturePlatformData* PlatformData = Texture->GetPlatformData();
+	FTextureResource* TextureResource = InTexture->GetResource();
+
+	if (PlatformData && TextureResource && TextureResource->TextureRHI)
+	{
+		FPCGTextureReadbackDispatchParams Params;
+		Params.SourceTexture = TextureResource->TextureRHI;
+		Params.SourceSampler = TextureResource->SamplerStateRHI;
+		Params.SourceDimensions = FIntPoint(PlatformData->SizeX, PlatformData->SizeY);
+
+		FPCGTextureReadbackInterface::Dispatch(Params, [this, PlatformData, PostInitializeCallback](void* OutBuffer, int32 ReadbackWidth, int32 ReadbackHeight)
 		{
-			FTexturePlatformData* PlatformData = Texture->GetPlatformData();
-			if (PlatformData)
+			TRACE_CPUPROFILER_EVENT_SCOPE(UPCGTextureData::Initialize::DispatchCallback);
+
+			const int32 PixelCount = ReadbackWidth * ReadbackHeight;
+
+			if (PlatformData->SizeX * PlatformData->SizeY != PixelCount)
 			{
-				TRACE_CPUPROFILER_EVENT_SCOPE(UPCGTextureData::Initialize::ReadData);
+				UE_LOG(LogPCG, Error,
+					TEXT("PCGTextureData readback has different dimensions than the source texture '%s'. Expected (%d, %d), received (%d, %d)."),
+					*Texture->GetFName().ToString(),
+					PlatformData->SizeX, PlatformData->SizeY,
+					ReadbackWidth, ReadbackHeight);
+			}
 
-				if (const uint8_t* BulkData = reinterpret_cast<const uint8_t*>(PlatformData->Mips[0].BulkData.LockReadOnly()))
-				{
-					Width = PlatformData->SizeX;
-					Height = PlatformData->SizeY;
-					const int32 PixelCount = Width * Height;
-					ColorData.SetNum(PixelCount);
-
-					if (PlatformData->PixelFormat == PF_B8G8R8A8)
-					{
-						// Memory representation of FColor is BGRA
-						const FColor* FormattedImageData = reinterpret_cast<const FColor*>(BulkData);
-						for (int32 D = 0; D < PixelCount; ++D)
-						{
-							ColorData[D] = FormattedImageData[D].ReinterpretAsLinear();
-						}
-					}
-					else if (PlatformData->PixelFormat == PF_R8G8B8A8)
-					{
-						// Since the memory representation is BGRA, we just need to swap B & R to obtain RGBA 
-						const FColor* FormattedImageData = reinterpret_cast<const FColor*>(BulkData);
-						for (int32 D = 0; D < PixelCount; ++D)
-						{
-							ColorData[D] = FormattedImageData[D].ReinterpretAsLinear();
-							Swap(ColorData[D].R, ColorData[D].B);
-						}
-					}
-					else if (PlatformData->PixelFormat == PF_G8)
-					{
-						for (int32 D = 0; D < PixelCount; ++D)
-						{
-							ColorData[D] = FColor(BulkData[D], BulkData[D], BulkData[D]).ReinterpretAsLinear();
-						}
-					}
-				}
-				else
+			if (const FColor* FormattedImageData = reinterpret_cast<const FColor*>(OutBuffer))
+			{
+				Width = ReadbackWidth;
+				Height = ReadbackHeight;
+				ColorData.SetNum(PixelCount);
+
+				for (int32 D = 0; D < PixelCount; ++D)
 				{
-					UE_LOG(LogPCG, Error, TEXT("PCGTextureData unable to get bulk data from '%s'"), *Texture->GetFName().ToString());
+					ColorData[D] = FormattedImageData[D].ReinterpretAsLinear();
 				}
-
-				PlatformData->Mips[0].BulkData.Unlock();
 			}
-		}
-		else
-		{
-			UE_LOG(LogPCG, Error, TEXT("PCGTextureData does not support the format of '%s'"), *Texture->GetFName().ToString());
-		}
+			else
+			{
+				UE_LOG(LogPCG, Error, TEXT("PCGTextureData unable to get readback results from '%s'"), *Texture->GetFName().ToString());
+			}
+
+			PostInitializeCallback();
+		});
 	}
+	else
+	{
+		UE_LOG(LogPCG, Error, TEXT("PCGTextureData failed to acquire texture resource for '%s'"), *Texture->GetFName().ToString());
 
-	Bounds = FBox(EForceInit::ForceInit);
-	Bounds += FVector(-1.0f, -1.0f, 0.0f);
-	Bounds += FVector(1.0f, 1.0f, 0.0f);
-	Bounds = Bounds.TransformBy(Transform);
+		PostInitializeCallback();
+	}
 }
 
 bool UPCGTextureData::IsSupported(UTexture2D* InTexture)
 {
-	const FTexturePlatformData* PlatformData = InTexture ? InTexture->GetPlatformData() : nullptr;
-
-	return PlatformData && PlatformData->Mips.Num() > 0 &&
-		(PlatformData->PixelFormat == PF_B8G8R8A8 || PlatformData->PixelFormat == PF_R8G8B8A8 || PlatformData->PixelFormat == PF_G8);
+	// TODO: GPU readback seems to support all textures, but future implementations (such as UBitmap) may have
+	// some limitations, so we will avoid removing this trivial function from the API for now
+	return true;
 }
 
 UPCGSpatialData* UPCGTextureData::CopyInternal() const
diff --git a/Engine/Plugins/Experimental/PCG/Source/PCG/Private/Elements/PCGTextureSampler.cpp b/Engine/Plugins/Experimental/PCG/Source/PCG/Private/Elements/PCGTextureSampler.cpp
--- a/Engine/Plugins/Experimental/PCG/Source/PCG/Private/Elements/PCGTextureSampler.cpp
+++ b/Engine/Plugins/Experimental/PCG/Source/PCG/Private/Elements/PCGTextureSampler.cpp
@@ -29,10 +29,22 @@ FPCGElementPtr UPCGTextureSamplerSettings::CreateElement() const
 	return MakeShared<FPCGTextureSamplerElement>();
 }
 
-bool FPCGTextureSamplerElement::ExecuteInternal(FPCGContext* Context) const
+bool FPCGTextureSamplerElement::ExecuteInternal(FPCGContext* InContext) const
 {
 	TRACE_CPUPROFILER_EVENT_SCOPE(FPCGTextureSamplerElement::Execute);
 
+	FPCGTextureSamplerContext* Context = static_cast<FPCGTextureSamplerContext*>(InContext);
+
+	if (Context->bIsPaused)
+	{
+		return false;
+	}
+
+	if (Context->bTextureReadbackDone)
+	{
+		return true;
+	}
+
 	const UPCGTextureSamplerSettings* Settings = Context->GetInputSettings<UPCGTextureSamplerSettings>();
 	check(Settings);
 
@@ -42,6 +54,7 @@ bool FPCGTextureSamplerElement::ExecuteInternal(FPCGContext* Context) const
 	}
 
 	UTexture2D* Texture = Settings->Texture.LoadSynchronous();
+
 	if (!Texture)
 	{
 		PCGE_LOG(Error, GraphAndLog, FText::Format(LOCTEXT("CouldNotResolveTexture", "Texture at path '{0}' could not be loaded"), FText::FromString(Settings->Texture.ToString())));
@@ -50,7 +63,7 @@ bool FPCGTextureSamplerElement::ExecuteInternal(FPCGContext* Context) const
 
 	if (!UPCGTextureData::IsSupported(Texture))
 	{
-		PCGE_LOG(Error, GraphAndLog, FText::Format(LOCTEXT("UnsupportedTextureFormat", "Texture '{0}' has unsupported texture format, currently supported formats are B8G8R8A8, R8G8B8A8 and G8. Also ensure mipmap generation is disabled and sRGB is disabled."), FText::FromName(Texture->GetFName())));
+		PCGE_LOG(Error, GraphAndLog, FText::Format(LOCTEXT("UnsupportedTextureFormat", "Texture '{0}' has unsupported settings"), FText::FromName(Texture->GetFName())));
 		return true;
 	}
 
@@ -79,13 +92,27 @@ bool FPCGTextureSamplerElement::ExecuteInternal(FPCGContext* Context) const
 	{
 		FTransform OriginalActorTransform = OriginalActor->GetTransform();
 		FinalTransform = Transform * OriginalActorTransform;
-		
+
 		FBox OriginalActorLocalBounds = PCGHelpers::GetActorLocalBounds(OriginalActor);
 		FinalTransform.SetScale3D(FinalTransform.GetScale3D() * 0.5 * (OriginalActorLocalBounds.Max - OriginalActorLocalBounds.Min));
 	}
 
 	// Initialize & set properties
-	TextureData->Initialize(Texture, FinalTransform);
+	Context->bIsPaused = true;
+
+	auto PostInitializeCallback = [Context, TextureData]()
+	{
+		Context->bIsPaused = false;
+		Context->bTextureReadbackDone = true;
+
+		if (!TextureData->IsValid())
+		{
+			PCGE_LOG_C(Error, GraphAndLog, Context, LOCTEXT("TextureDataInitFailed", "Texture data failed to initialize, check log for more information"));
+		}
+	};
+
+	TextureData->Initialize(Texture, FinalTransform, PostInitializeCallback);
+
 	TextureData->DensityFunction = DensityFunction;
 	TextureData->ColorChannel = ColorChannel;
 	TextureData->TexelSize = TexelSize;
@@ -96,12 +123,17 @@ bool FPCGTextureSamplerElement::ExecuteInternal(FPCGContext* Context) const
 	TextureData->bUseTileBounds = bUseTileBounds;
 	TextureData->TileBounds = FBox2D(TileBoundsMin, TileBoundsMax);
 
-	if (!TextureData->IsValid())
-	{
-		PCGE_LOG(Error, GraphAndLog, LOCTEXT("TextureDataInitFailed", "Texture data failed to initialize, check log for more information"));
-	}
+	return false;
+}
+
+FPCGContext* FPCGTextureSamplerElement::Initialize(const FPCGDataCollection& InputData, TWeakObjectPtr<UPCGComponent> SourceComponent, const UPCGNode* Node)
+{
+	FPCGTextureSamplerContext* Context = new FPCGTextureSamplerContext();
+	Context->InputData = InputData;
+	Context->SourceComponent = SourceComponent;
+	Context->Node = Node;
 
-	return true;
+	return Context;
 }
 
 void FPCGTextureSamplerElement::GetDependenciesCrc(const FPCGDataCollection& InInput, const UPCGSettings* InSettings, UPCGComponent* InComponent, FPCGCrc& OutCrc) const
diff --git a/Engine/Plugins/Experimental/PCG/Source/PCG/Private/Tests/Data/PCGTextureDataTest.cpp b/Engine/Plugins/Experimental/PCG/Source/PCG/Private/Tests/Data/PCGTextureDataTest.cpp
--- a/Engine/Plugins/Experimental/PCG/Source/PCG/Private/Tests/Data/PCGTextureDataTest.cpp
+++ b/Engine/Plugins/Experimental/PCG/Source/PCG/Private/Tests/Data/PCGTextureDataTest.cpp
@@ -35,70 +35,73 @@ bool FPCGTextureDataOffsetTilingRotation::RunTest(const FString& Parameters)
 	Texture2D->UpdateResource();
 
 	UPCGTextureData* TextureData = NewObject<UPCGTextureData>();
-	TextureData->Initialize(Texture2D, FTransform());
-
 	TextureData->bUseAdvancedTiling = true;
 
-	FRandomStream RandomStream;
-	FPCGPoint OutPoint;
-
-	const float TextureSpacePixelX = static_cast<float>(WhitePixelX) / TextureSize;
-	const float ScaledPixelX = (2.0 * TextureSpacePixelX) - 1.0;
-	const float TextureSpacePixelY = static_cast<float>(WhitePixelY) / TextureSize;
-	const float ScaledPixelY = (2.0 * TextureSpacePixelY) - 1;
-
-	{
-		TextureData->Rotation = 0.f;
-	
-		// sampling with no offset
-		TextureData->CenterOffset = FVector2D::ZeroVector;
-		TextureData->SamplePoint(FTransform(), FBox(), OutPoint, nullptr);
-		TestEqual("Valid color sampled for no offset at (0, 0)", OutPoint.Color, static_cast<FVector>(FColor::Black));
-
-		// sampling at position, with no offset
-		TextureData->CenterOffset = FVector2D::ZeroVector;
-		TextureData->SamplePoint(FTransform(FVector(0.5, 0.5, 0.0)), FBox(), OutPoint, nullptr);
-		TestEqual("Valid color sampled for no offset at position", OutPoint.Color, static_cast<FVector>(FColor::Black));
-
-		// sampling at expected position
-		TextureData->CenterOffset = FVector2D::ZeroVector;
-		TextureData->SamplePoint(FTransform(FVector(ScaledPixelX, ScaledPixelY, 0)), FBox(), OutPoint, nullptr);
-		TestEqual("Valid color sampled at expected position", OutPoint.Color, static_cast<FVector>(FColor::White));
-	
-		// sampling with offset
-		TextureData->CenterOffset = FVector2D(0.5 - TextureSpacePixelX, 0.5 -TextureSpacePixelY);
-		TextureData->SamplePoint(FTransform(), FBox(), OutPoint, nullptr);
-		TestEqual("Valid color sampled for offset from (0, 0)", OutPoint.Color, static_cast<FVector>(FColor::White));
-
-		// sampling at position, with offset
-		TextureData->CenterOffset = FVector2D(1.0 - TextureSpacePixelX, 1.0 - TextureSpacePixelY);
-		TextureData->SamplePoint(FTransform(FVector(1.0, 1.0, 0.0)), FBox(), OutPoint, nullptr);
-		TestEqual("Valid color sampled for offset from position", OutPoint.Color, static_cast<FVector>(FColor::White));
-	
-		for (float Rotation = -360.f; Rotation < 360.f; Rotation += 10.f)
+	bool bTestPassed = true;
+
+	TextureData->Initialize(Texture2D, FTransform(), [this, TextureData, &bTestPassed]() -> void
 		{
-			const float Theta = FMath::DegreesToRadians(Rotation);
-			const float CosTheta = FMath::Cos(Theta);
-			const float SinTheta = FMath::Sin(Theta);
-			const FVector::FReal X = (ScaledPixelX * CosTheta) - (ScaledPixelY * SinTheta);
-			const FVector::FReal Y = (ScaledPixelY * CosTheta) + (ScaledPixelX * SinTheta);
-
-			FTransform RotatedTransform(FVector(X, Y, 0.f));
-
-			// sampling with rotation at black position
-			TextureData->Rotation = -Rotation;
-			TextureData->CenterOffset = FVector2D::ZeroVector;
-			TextureData->SamplePoint(FTransform(FVector(0.5, 0.5, 0.0)), FBox(), OutPoint, nullptr);
-			TestEqual("Valid color sampled for rotated off-position", OutPoint.Color, static_cast<FVector>(FColor::Black));
-	
-			// sampling with rotation at white position
-			TextureData->Rotation = -Rotation;
-			TextureData->CenterOffset = FVector2D::ZeroVector;
-			TextureData->SamplePoint(RotatedTransform, FBox(), OutPoint, nullptr);
-			TestEqual("Valid color sampled for rotated position", OutPoint.Color, static_cast<FVector>(FColor::White));
-		}
-	}
-
-	return true;
+			FRandomStream RandomStream;
+			FPCGPoint OutPoint;
+
+			const float TextureSpacePixelX = static_cast<float>(WhitePixelX) / TextureSize;
+			const float ScaledPixelX = (2.0 * TextureSpacePixelX) - 1.0;
+			const float TextureSpacePixelY = static_cast<float>(WhitePixelY) / TextureSize;
+			const float ScaledPixelY = (2.0 * TextureSpacePixelY) - 1;
+
+			{
+				TextureData->Rotation = 0.f;
+
+				// sampling with no offset
+				TextureData->CenterOffset = FVector2D::ZeroVector;
+				TextureData->SamplePoint(FTransform(), FBox(), OutPoint, nullptr);
+				bTestPassed &= TestEqual("Valid color sampled for no offset at (0, 0)", OutPoint.Color, static_cast<FVector>(FColor::Black));
+
+				// sampling at position, with no offset
+				TextureData->CenterOffset = FVector2D::ZeroVector;
+				TextureData->SamplePoint(FTransform(FVector(0.5, 0.5, 0.0)), FBox(), OutPoint, nullptr);
+				bTestPassed &= TestEqual("Valid color sampled for no offset at position", OutPoint.Color, static_cast<FVector>(FColor::Black));
+
+				// sampling at expected position
+				TextureData->CenterOffset = FVector2D::ZeroVector;
+				TextureData->SamplePoint(FTransform(FVector(ScaledPixelX, ScaledPixelY, 0)), FBox(), OutPoint, nullptr);
+				bTestPassed &= TestEqual("Valid color sampled at expected position", OutPoint.Color, static_cast<FVector>(FColor::White));
+
+				// sampling with offset
+				TextureData->CenterOffset = FVector2D(0.5 - TextureSpacePixelX, 0.5 - TextureSpacePixelY);
+				TextureData->SamplePoint(FTransform(), FBox(), OutPoint, nullptr);
+				bTestPassed &= TestEqual("Valid color sampled for offset from (0, 0)", OutPoint.Color, static_cast<FVector>(FColor::White));
+
+				// sampling at position, with offset
+				TextureData->CenterOffset = FVector2D(1.0 - TextureSpacePixelX, 1.0 - TextureSpacePixelY);
+				TextureData->SamplePoint(FTransform(FVector(1.0, 1.0, 0.0)), FBox(), OutPoint, nullptr);
+				bTestPassed &= TestEqual("Valid color sampled for offset from position", OutPoint.Color, static_cast<FVector>(FColor::White));
+
+				for (float Rotation = -360.f; Rotation < 360.f; Rotation += 10.f)
+				{
+					const float Theta = FMath::DegreesToRadians(Rotation);
+					const float CosTheta = FMath::Cos(Theta);
+					const float SinTheta = FMath::Sin(Theta);
+					const FVector::FReal X = (ScaledPixelX * CosTheta) - (ScaledPixelY * SinTheta);
+					const FVector::FReal Y = (ScaledPixelY * CosTheta) + (ScaledPixelX * SinTheta);
+
+					FTransform RotatedTransform(FVector(X, Y, 0.f));
+
+					// sampling with rotation at black position
+					TextureData->Rotation = -Rotation;
+					TextureData->CenterOffset = FVector2D::ZeroVector;
+					TextureData->SamplePoint(FTransform(FVector(0.5, 0.5, 0.0)), FBox(), OutPoint, nullptr);
+					bTestPassed &= TestEqual("Valid color sampled for rotated off-position", OutPoint.Color, static_cast<FVector>(FColor::Black));
+
+					// sampling with rotation at white position
+					TextureData->Rotation = -Rotation;
+					TextureData->CenterOffset = FVector2D::ZeroVector;
+					TextureData->SamplePoint(RotatedTransform, FBox(), OutPoint, nullptr);
+					bTestPassed &= TestEqual("Valid color sampled for rotated position", OutPoint.Color, static_cast<FVector>(FColor::White));
+				}
+			}
+		});
+
+	return bTestPassed;
 }
 #endif
diff --git a/Engine/Plugins/Experimental/PCG/Source/PCG/Public/Data/PCGTextureData.h b/Engine/Plugins/Experimental/PCG/Source/PCG/Public/Data/PCGTextureData.h
--- a/Engine/Plugins/Experimental/PCG/Source/PCG/Public/Data/PCGTextureData.h
+++ b/Engine/Plugins/Experimental/PCG/Source/PCG/Public/Data/PCGTextureData.h
@@ -104,8 +104,7 @@ public:
 	virtual EPCGDataType GetDataType() const override { return EPCGDataType::Texture; }
 	// ~End UPCGData interface
 
-	UFUNCTION(BlueprintCallable, Category = Texture)
-	void Initialize(UTexture2D* InTexture, const FTransform& InTransform);
+	void Initialize(UTexture2D* InTexture, const FTransform& InTransform, const TFunction<void()>& PostInitializeCallback);
 
 	/** Returns true if the format of InTexture is compatible and can be loaded. Will load texture if not already loaded. */
 	static bool IsSupported(UTexture2D* InTexture);
diff --git a/Engine/Plugins/Experimental/PCG/Source/PCG/Public/Elements/PCGTextureSampler.h b/Engine/Plugins/Experimental/PCG/Source/PCG/Public/Elements/PCGTextureSampler.h
--- a/Engine/Plugins/Experimental/PCG/Source/PCG/Public/Elements/PCGTextureSampler.h
+++ b/Engine/Plugins/Experimental/PCG/Source/PCG/Public/Elements/PCGTextureSampler.h
@@ -2,6 +2,7 @@
 
 #pragma once
 
+#include "PCGContext.h"
 #include "PCGPin.h"
 #include "PCGSettings.h"
 #include "Data/PCGTextureData.h"
@@ -74,11 +75,17 @@ public:
 	FVector2D TileBoundsMax = FVector2D(0.5, 0.5);
 };
 
+struct PCG_API FPCGTextureSamplerContext : public FPCGContext
+{
+	bool bTextureReadbackDone = false;
+};
+
 class FPCGTextureSamplerElement : public FSimplePCGElement
 {
 public:
+	virtual FPCGContext* Initialize(const FPCGDataCollection& InputData, TWeakObjectPtr<UPCGComponent> SourceComponent, const UPCGNode* Node) override;
 	virtual void GetDependenciesCrc(const FPCGDataCollection& InInput, const UPCGSettings* InSettings, UPCGComponent* InComponent, FPCGCrc& OutCrc) const override;
 
 protected:
-	virtual bool ExecuteInternal(FPCGContext* Context) const override;
+	virtual bool ExecuteInternal(FPCGContext* InContext) const override;
 };
diff --git a/Engine/Plugins/Experimental/PCG/Source/PCGCompute/PCGCompute.Build.cs b/Engine/Plugins/Experimental/PCG/Source/PCGCompute/PCGCompute.Build.cs
new file mode 100644
--- /dev/null
+++ b/Engine/Plugins/Experimental/PCG/Source/PCGCompute/PCGCompute.Build.cs
@@ -0,0 +1,21 @@
+// Copyright Epic Games, Inc. All Rights Reserved.
+
+namespace UnrealBuildTool.Rules
+{
+	public class PCGCompute : ModuleRules
+	{
+		public PCGCompute(ReadOnlyTargetRules Target) : base(Target)
+		{
+			PCHUsage = PCHUsageMode.UseExplicitOrSharedPCHs;
+
+			PublicDependencyModuleNames.AddRange(
+				new string[]
+				{
+					"Core",
+					"Projects",
+					"RHI",
+					"RenderCore",
+				});
+		}
+	}
+}
diff --git a/Engine/Plugins/Experimental/PCG/Source/PCGCompute/Private/PCGComputeModule.cpp b/Engine/Plugins/Experimental/PCG/Source/PCGCompute/Private/PCGComputeModule.cpp
new file mode 100644
--- /dev/null
+++ b/Engine/Plugins/Experimental/PCG/Source/PCGCompute/Private/PCGComputeModule.cpp
@@ -0,0 +1,40 @@
+// Copyright Epic Games, Inc. All Rights Reserved.
+
+#include "PCGComputeModule.h"
+
+#include "GlobalShader.h"
+#include "Interfaces/IPluginManager.h"
+#include "Modules/ModuleInterface.h"
+#include "Modules/ModuleManager.h"
+
+#define LOCTEXT_NAMESPACE "FPCGComputeModule"
+
+class FPCGComputeModule final : public IModuleInterface
+{
+public:
+	//~ IModuleInterface implementation
+#if WITH_EDITOR
+	virtual void StartupModule() override;
+	virtual void ShutdownModule() override;
+#endif
+	//~ End IModuleInterface implementation
+};
+
+#if WITH_EDITOR
+void FPCGComputeModule::StartupModule()
+{
+	// Maps virtual shader source directory to the plugin's actual shaders directory.
+	const FString PluginShaderDir = FPaths::Combine(IPluginManager::Get().FindPlugin(TEXT("PCG"))->GetBaseDir(), TEXT("Shaders/Private"));
+	AddShaderSourceDirectoryMapping(TEXT("/PCGComputeShaders"), PluginShaderDir);
+}
+
+void FPCGComputeModule::ShutdownModule()
+{
+}
+#endif
+
+IMPLEMENT_MODULE(FPCGComputeModule, PCGCompute);
+
+PCGCOMPUTE_API DEFINE_LOG_CATEGORY(LogPCGCompute);
+
+#undef LOCTEXT_NAMESPACE
diff --git a/Engine/Plugins/Experimental/PCG/Source/PCGCompute/Private/PCGTextureReadback.cpp b/Engine/Plugins/Experimental/PCG/Source/PCGCompute/Private/PCGTextureReadback.cpp
new file mode 100644
--- /dev/null
+++ b/Engine/Plugins/Experimental/PCG/Source/PCGCompute/Private/PCGTextureReadback.cpp
@@ -0,0 +1,131 @@
+// Copyright Epic Games, Inc. All Rights Reserved.
+
+#include "PCGTextureReadback.h"
+
+#include "GlobalShader.h"
+#include "RHIGPUReadback.h"
+#include "RenderCaptureInterface.h"
+#include "RenderGraphUtils.h"
+#include "Async/Async.h"
+
+#define LOCTEXT_NAMESPACE "PCGCompute"
+
+#define NUM_THREADS_PER_GROUP_DIMENSION 32
+
+class PCGCOMPUTE_API FPCGTextureReadbackCS : public FGlobalShader
+{
+public:
+	DECLARE_GLOBAL_SHADER(FPCGTextureReadbackCS);
+	SHADER_USE_PARAMETER_STRUCT(FPCGTextureReadbackCS, FGlobalShader);
+ 
+	BEGIN_SHADER_PARAMETER_STRUCT(FParameters, )
+		SHADER_PARAMETER_TEXTURE(Texture2D<float4>, SourceTexture)
+		SHADER_PARAMETER_SAMPLER(SamplerState, SourceSampler)
+		SHADER_PARAMETER(FVector2f, SourceDimensions)
+		SHADER_PARAMETER_UAV(RWTexture2D<float4>, OutputTexture)
+	END_SHADER_PARAMETER_STRUCT()
+ 
+public:
+	static void ModifyCompilationEnvironment(const FGlobalShaderPermutationParameters& Parameters, FShaderCompilerEnvironment& OutEnvironment)
+	{
+		FGlobalShader::ModifyCompilationEnvironment(Parameters, OutEnvironment);
+
+		OutEnvironment.SetDefine(TEXT("THREADGROUPSIZE_X"), NUM_THREADS_PER_GROUP_DIMENSION);
+		OutEnvironment.SetDefine(TEXT("THREADGROUPSIZE_Y"), NUM_THREADS_PER_GROUP_DIMENSION);
+		OutEnvironment.SetDefine(TEXT("THREADGROUPSIZE_Z"), 1);
+	}
+};
+
+IMPLEMENT_GLOBAL_SHADER(FPCGTextureReadbackCS, "/PCGComputeShaders/PCGTextureReadback.usf", "PCGTextureReadback_CS", SF_Compute);
+
+void FPCGTextureReadbackInterface::Dispatch_RenderThread(FRHICommandListImmediate& RHICmdList, const FPCGTextureReadbackDispatchParams& Params, const TFunction<void(void* OutBuffer, int32 ReadbackWidth, int32 ReadbackHeight)>& AsyncCallback)
+{
+	check(Params.SourceTexture && Params.SourceSampler);
+
+	RenderCaptureInterface::FScopedCapture RenderCapture(true, &RHICmdList, TEXT("PCGTextureDataReadback"));
+
+	FPCGTextureReadbackCS::FParameters PassParameters;
+	PassParameters.SourceTexture = Params.SourceTexture;
+	PassParameters.SourceSampler = Params.SourceSampler;
+	PassParameters.SourceDimensions = { (float)Params.SourceDimensions.X, (float)Params.SourceDimensions.Y };
+
+	FRHITextureCreateDesc TargetTextureDesc =
+		FRHITextureCreateDesc::Create2D(TEXT("PCGTexture Readback Compute Target"), Params.SourceDimensions.X, Params.SourceDimensions.Y, EPixelFormat::PF_B8G8R8A8)
+			.SetClearValue(FClearValueBinding::None)
+			.SetFlags(ETextureCreateFlags::UAV |
+				ETextureCreateFlags::RenderTargetable |
+				ETextureCreateFlags::ShaderResource |
+				ETextureCreateFlags::NoTiling
+			)
+			.SetInitialState(ERHIAccess::UAVCompute)
+			.DetermineInititialState();
+	check(TargetTextureDesc.IsValid());
+	
+	// Create temporary output texture
+	FTextureRHIRef OutputTexture = RHICreateTexture(TargetTextureDesc);
+	PassParameters.OutputTexture = RHICmdList.CreateUnorderedAccessView(OutputTexture,
+		FRHIViewDesc::CreateTextureUAV()
+			.SetDimensionFromTexture(OutputTexture)
+			.SetMipLevel(0)
+			.SetFormat(TargetTextureDesc.Format)
+			.SetArrayRange(0, 0)
+		);
+
+	TShaderMapRef<FPCGTextureReadbackCS> ComputeShader(GetGlobalShaderMap(GMaxRHIFeatureLevel));
+	FComputeShaderUtils::Dispatch(RHICmdList, ComputeShader, PassParameters,
+		FIntVector(FMath::DivideAndRoundUp(Params.SourceDimensions.X, NUM_THREADS_PER_GROUP_DIMENSION),
+			FMath::DivideAndRoundUp(Params.SourceDimensions.Y, NUM_THREADS_PER_GROUP_DIMENSION), 1));
+
+	// Prepare OutputTexture to be copied
+	RHICmdList.Transition(FRHITransitionInfo(OutputTexture, ERHIAccess::UAVCompute, ERHIAccess::CopySrc));
+
+	FRHIGPUTextureReadback* GPUTextureReadback = new FRHIGPUTextureReadback(TEXT("PCGTextureReadbackCopy"));
+	GPUTextureReadback->EnqueueCopy(RHICmdList, PassParameters.OutputTexture->GetTexture());
+
+	auto RunnerFunc = [GPUTextureReadback, AsyncCallback](auto&& RunnerFunc) -> void
+	{
+		if (GPUTextureReadback->IsReady())
+		{
+			int32 ReadbackWidth = 0, ReadbackHeight = 0;
+			void* OutBuffer = GPUTextureReadback->Lock(ReadbackWidth, &ReadbackHeight);
+
+			AsyncCallback(OutBuffer, ReadbackWidth, ReadbackHeight);
+
+			GPUTextureReadback->Unlock();
+			delete GPUTextureReadback;
+		}
+		else
+		{
+			AsyncTask(ENamedThreads::ActualRenderingThread, [RunnerFunc]()
+			{
+				RunnerFunc(RunnerFunc);
+			});
+		}
+	};
+
+	AsyncTask(ENamedThreads::ActualRenderingThread, [RunnerFunc]()
+	{
+		RunnerFunc(RunnerFunc);
+	});
+}
+
+void FPCGTextureReadbackInterface::Dispatch_GameThread(const FPCGTextureReadbackDispatchParams& Params, const TFunction<void(void* OutBuffer, int32 ReadbackWidth, int32 ReadbackHeight)>& AsyncCallback)
+{
+	ENQUEUE_RENDER_COMMAND(SceneDrawCompletion)(
+		[Params, AsyncCallback](FRHICommandListImmediate& RHICmdList)
+		{
+			Dispatch_RenderThread(RHICmdList, Params, AsyncCallback);
+		});
+}
+
+void FPCGTextureReadbackInterface::Dispatch(const FPCGTextureReadbackDispatchParams& Params, const TFunction<void(void* OutBuffer, int32 ReadbackWidth, int32 ReadbackHeight)>& AsyncCallback)
+{
+	if (IsInRenderingThread())
+	{
+		Dispatch_RenderThread(GetImmediateCommandList_ForRenderCommand(), Params, AsyncCallback);
+	}
+	else
+	{
+		Dispatch_GameThread(Params, AsyncCallback);
+	}
+}
diff --git a/Engine/Plugins/Experimental/PCG/Source/PCGCompute/Public/PCGComputeModule.h b/Engine/Plugins/Experimental/PCG/Source/PCGCompute/Public/PCGComputeModule.h
new file mode 100644
--- /dev/null
+++ b/Engine/Plugins/Experimental/PCG/Source/PCGCompute/Public/PCGComputeModule.h
@@ -0,0 +1,14 @@
+// Copyright Epic Games, Inc. All Rights Reserved.
+
+#pragma once
+
+#include "Stats/Stats2.h"
+
+// Logs
+PCGCOMPUTE_API DECLARE_LOG_CATEGORY_EXTERN(LogPCGCompute, Log, All);
+
+// Stats
+DECLARE_STATS_GROUP(TEXT("PCGCompute"), STATGROUP_PCGCompute, STATCAT_Advanced)
+
+// CVars
+
diff --git a/Engine/Plugins/Experimental/PCG/Source/PCGCompute/Public/PCGTextureReadback.h b/Engine/Plugins/Experimental/PCG/Source/PCGCompute/Public/PCGTextureReadback.h
new file mode 100644
--- /dev/null
+++ b/Engine/Plugins/Experimental/PCG/Source/PCGCompute/Public/PCGTextureReadback.h
@@ -0,0 +1,22 @@
+// Copyright Epic Games, Inc. All Rights Reserved.
+
+#pragma once
+
+#include "RHI.h"
+ 
+struct PCGCOMPUTE_API FPCGTextureReadbackDispatchParams
+{
+	FTextureRHIRef SourceTexture;
+	FSamplerStateRHIRef SourceSampler;
+	FIntPoint SourceDimensions;
+};
+ 
+class PCGCOMPUTE_API FPCGTextureReadbackInterface
+{
+public:
+	static void Dispatch_RenderThread(FRHICommandListImmediate& RHICmdList, const FPCGTextureReadbackDispatchParams& Params, const TFunction<void(void* OutBuffer, int32 ReadbackWidth, int32 ReadbackHeight)>& AsyncCallback);
+	static void Dispatch_GameThread(const FPCGTextureReadbackDispatchParams& Params, const TFunction<void(void* OutBuffer, int32 ReadbackWidth, int32 ReadbackHeight)>& AsyncCallback);
+
+	// Dispatches this shader. Can be called from any thread
+	static void Dispatch(const FPCGTextureReadbackDispatchParams& Params, const TFunction<void(void* OutBuffer, int32 ReadbackWidth, int32 ReadbackHeight)>& AsyncCallback);
+};
