[PCG] Removed zero density point culling from static mesh spawner node

diff --git a/Engine/Plugins/Experimental/PCG/Source/PCG/Private/MeshSelectors/PCGMeshSelectorByAttribute.cpp b/Engine/Plugins/Experimental/PCG/Source/PCG/Private/MeshSelectors/PCGMeshSelectorByAttribute.cpp
--- a/Engine/Plugins/Experimental/PCG/Source/PCG/Private/MeshSelectors/PCGMeshSelectorByAttribute.cpp
+++ b/Engine/Plugins/Experimental/PCG/Source/PCG/Private/MeshSelectors/PCGMeshSelectorByAttribute.cpp
@@ -167,12 +167,7 @@ bool UPCGMeshSelectorByAttribute::SelectInstances(
 	while (CurrentPointIndex < Points.Num())
 	{
 		const FPCGPoint& Point = Points[CurrentPointIndex++];
-
-		if (Point.Density <= 0.0f)
-		{
-			continue;
-		}
-
+		
 		PCGMetadataValueKey ValueKey = Attribute->GetValueKey(Point.MetadataEntry);
 		TSoftObjectPtr<UStaticMesh>* NewMesh = ValueKeyToMesh.Find(ValueKey);
 		TSoftObjectPtr<UStaticMesh> Mesh = nullptr;
@@ -184,7 +179,7 @@ bool UPCGMeshSelectorByAttribute::SelectInstances(
 
 			if (!MeshSoftObjectPath.IsEmpty() && MeshSoftObjectPath != TEXT("None"))
 			{
-				FSoftObjectPath MeshPath(MeshSoftObjectPath);
+				const FSoftObjectPath MeshPath(MeshSoftObjectPath);
 				Mesh = TSoftObjectPtr<UStaticMesh>(MeshPath);
 
 				if (Mesh.IsNull())
diff --git a/Engine/Plugins/Experimental/PCG/Source/PCG/Private/MeshSelectors/PCGMeshSelectorWeighted.cpp b/Engine/Plugins/Experimental/PCG/Source/PCG/Private/MeshSelectors/PCGMeshSelectorWeighted.cpp
--- a/Engine/Plugins/Experimental/PCG/Source/PCG/Private/MeshSelectors/PCGMeshSelectorWeighted.cpp
+++ b/Engine/Plugins/Experimental/PCG/Source/PCG/Private/MeshSelectors/PCGMeshSelectorWeighted.cpp
@@ -213,13 +213,8 @@ bool UPCGMeshSelectorWeighted::SelectInstances(
 		{
 			const FPCGPoint& Point = Points[CurrentPointIndex++];
 
-			if (Point.Density <= 0.0f)
-			{
-				continue;
-			}
-
 			FRandomStream RandomSource = UPCGBlueprintHelpers::GetRandomStream(Point, Settings, Context.SourceComponent.Get());
-			int RandomWeightedPick = RandomSource.RandRange(0, TotalWeight - 1);
+			const int RandomWeightedPick = RandomSource.RandRange(0, TotalWeight - 1);
 
 			int RandomPick = 0;
 			while(RandomPick < MeshInstances.Num() && CumulativeWeights[RandomPick] <= RandomWeightedPick)
diff --git a/Engine/Plugins/Experimental/PCG/Source/PCG/Private/MeshSelectors/PCGMeshSelectorWeightedByCategory.cpp b/Engine/Plugins/Experimental/PCG/Source/PCG/Private/MeshSelectors/PCGMeshSelectorWeightedByCategory.cpp
--- a/Engine/Plugins/Experimental/PCG/Source/PCG/Private/MeshSelectors/PCGMeshSelectorWeightedByCategory.cpp
+++ b/Engine/Plugins/Experimental/PCG/Source/PCG/Private/MeshSelectors/PCGMeshSelectorWeightedByCategory.cpp
@@ -204,12 +204,7 @@ bool UPCGMeshSelectorWeightedByCategory::SelectInstances(
 	{
 		const FPCGPoint& Point = Points[CurrentPointIndex++];
 
-		if (Point.Density <= 0.0f)
-		{
-			continue;
-		}
-
-		PCGMetadataValueKey ValueKey = Attribute->GetValueKey(Point.MetadataEntry);
+		const PCGMetadataValueKey ValueKey = Attribute->GetValueKey(Point.MetadataEntry);
 
 		// if no mesh list was processed for this attribute value, fallback to the default mesh list
 		FPCGInstancesAndWeights* InstancesAndWeights = CategoryEntryToInstancesAndWeights.Find(ValueKey);
@@ -229,7 +224,7 @@ bool UPCGMeshSelectorWeightedByCategory::SelectInstances(
 		const int TotalWeight = InstancesAndWeights->CumulativeWeights.Last();
 
 		FRandomStream RandomSource = UPCGBlueprintHelpers::GetRandomStream(Point, Settings, Context.SourceComponent.Get());
-		int RandomWeightedPick = RandomSource.RandRange(0, TotalWeight - 1);
+		const int RandomWeightedPick = RandomSource.RandRange(0, TotalWeight - 1);
 
 		int RandomPick = 0;
 		while(RandomPick < InstancesAndWeights->MeshInstances.Num() && InstancesAndWeights->CumulativeWeights[RandomPick] <= RandomWeightedPick)
