[PCG] Fix infinite recursion crash with PCG subgraphs

#rnx

diff --git a/Engine/Plugins/Experimental/PCG/Source/PCG/Private/Graph/PCGStackContext.cpp b/Engine/Plugins/Experimental/PCG/Source/PCG/Private/Graph/PCGStackContext.cpp
--- a/Engine/Plugins/Experimental/PCG/Source/PCG/Private/Graph/PCGStackContext.cpp
+++ b/Engine/Plugins/Experimental/PCG/Source/PCG/Private/Graph/PCGStackContext.cpp
@@ -205,8 +205,7 @@ FPCGStackContext FPCGStackContext::CreateStackContextFromGraph(const UPCGGraph*
 {
 	FPCGStackContext StackContext;
 
-	TFunction<void(const UPCGGraph*, FPCGStackContext&)> ParseGraphRecursive;
-	ParseGraphRecursive = [&ParseGraphRecursive] (const UPCGGraph* InPCGGraph, FPCGStackContext& InStackContext)
+	auto ParseGraphRecursive = [](const UPCGGraph* InPCGGraph, FPCGStackContext& InStackContext, TArray<const UPCGGraph*>& VisitedGraphStack, auto RecursiveCallback)
 	{
 		if (!InPCGGraph)
 		{
@@ -219,24 +218,35 @@ FPCGStackContext FPCGStackContext::CreateStackContextFromGraph(const UPCGGraph*
 		{
 			if (PCGNode)
 			{
+				// TODO: GetSettings() has no execution context and therefore cannot recurse on dynamically chosen subgraphs
 				if (const UPCGBaseSubgraphSettings* SubgraphSettings = Cast<UPCGBaseSubgraphSettings>(PCGNode->GetSettings()))
 				{
 					if (const UPCGGraph* PCGSubgraph = SubgraphSettings->GetSubgraph())
 					{
-						InStackContext.PushFrame(PCGNode);
-
-						FPCGStackContext SubgraphStackContext;
-						ParseGraphRecursive(PCGSubgraph, SubgraphStackContext);
-						InStackContext.AppendStacks(SubgraphStackContext);
-
-						InStackContext.PopFrame();
+						// Skip graphs we have already visited to avoid cycles
+						// TODO: This prevents recursive subgraphs
+						if (!VisitedGraphStack.Contains(PCGSubgraph))
+						{
+							InStackContext.PushFrame(PCGNode);
+							VisitedGraphStack.Push(PCGSubgraph);
+
+							FPCGStackContext SubgraphStackContext;
+							RecursiveCallback(PCGSubgraph, SubgraphStackContext, VisitedGraphStack, RecursiveCallback);
+							InStackContext.AppendStacks(SubgraphStackContext);
+
+							VisitedGraphStack.Pop();
+							InStackContext.PopFrame();
+						}
 					}
 				}
 			}
 		}
 	};
 
-	ParseGraphRecursive(InPCGGraph, StackContext);
+	TArray<const UPCGGraph*> VisitedGraphStack;
+	VisitedGraphStack.Push(InPCGGraph);
+
+	ParseGraphRecursive(InPCGGraph, StackContext, VisitedGraphStack, ParseGraphRecursive);
 
 	return StackContext;
 }
