Added tooltip to clarify use of the IsCacheable flag in PCG blueprints

diff --git a/Engine/Plugins/Experimental/PCG/Source/PCG/Public/Elements/PCGExecuteBlueprint.h b/Engine/Plugins/Experimental/PCG/Source/PCG/Public/Elements/PCGExecuteBlueprint.h
--- a/Engine/Plugins/Experimental/PCG/Source/PCG/Public/Elements/PCGExecuteBlueprint.h
+++ b/Engine/Plugins/Experimental/PCG/Source/PCG/Public/Elements/PCGExecuteBlueprint.h
@@ -149,7 +149,8 @@ public:
 #endif
 
 	/** Controls whether results can be cached so we can bypass execution if the inputs & settings are the same in a subsequent execution.
-	* Note that if your node relies on data that is not directly tracked by PCG or creates any kind of artifact (adds components, creates actors, etc.) then it is not cacheable.
+	* If you have implemented the IsCacheableOverride function, then this value is ignored.
+	* Note that if your node relies on data that is not directly tracked by PCG or creates any kind of artifact (adds components, creates actors, etc.) then it should not be cacheable.
 	*/
 	UPROPERTY(BlueprintReadWrite, EditDefaultsOnly, AdvancedDisplay, Category = Settings)
 	bool bIsCacheable = false;
