[PCG] Fix non HiGen partitioned graphs getting associated with HiGen partition actors

#rnx

diff --git a/Engine/Plugins/Experimental/PCG/Source/PCG/Private/PCGActorAndComponentMapping.cpp b/Engine/Plugins/Experimental/PCG/Source/PCG/Private/PCGActorAndComponentMapping.cpp
--- a/Engine/Plugins/Experimental/PCG/Source/PCG/Private/PCGActorAndComponentMapping.cpp
+++ b/Engine/Plugins/Experimental/PCG/Source/PCG/Private/PCGActorAndComponentMapping.cpp
@@ -446,6 +446,11 @@ void UPCGActorAndComponentMapping::ForAllIntersectingPartitionActors(const FBox&
 
 void UPCGActorAndComponentMapping::UpdateMappingPCGComponentPartitionActor(UPCGComponent* InComponent)
 {
+	if (!PCGSubsystem->IsInitialized())
+	{
+		return;
+	}
+
 	check(InComponent);
 
 	// Get the bounds
@@ -470,17 +475,27 @@ void UPCGActorAndComponentMapping::UpdateMappingPCGComponentPartitionActor(UPCGC
 			check(PartitionActorsPtr);
 		}
 
-		TSet<TObjectPtr<APCGPartitionActor>> NewMapping;
-		ForAllIntersectingPartitionActors(Bounds, [&NewMapping, InComponent](APCGPartitionActor* Actor)
+		if (const APCGWorldActor* WorldActor = PCGSubsystem->GetPCGWorldActor())
 		{
-			Actor->AddGraphInstance(InComponent);
-			NewMapping.Add(Actor);
-		});
+			const bool bIsHiGenEnabled = InComponent->GetGraph() && InComponent->GetGraph()->IsHierarchicalGenerationEnabled();
 
-		// Find the ones that were removed
-		RemovedActors = PartitionActorsPtr->Difference(NewMapping);
+			TSet<TObjectPtr<APCGPartitionActor>> NewMapping;
+			ForAllIntersectingPartitionActors(Bounds, [&NewMapping, InComponent, WorldActor, bIsHiGenEnabled](APCGPartitionActor* Actor)
+			{
+				// If this graph does not have HiGen enabled, we should only add a graph instance for
+				// the partition actors whose grid size matches the WorldActor's partition grid size
+				if (bIsHiGenEnabled || (Actor && Actor->GetPCGGridSize() == WorldActor->PartitionGridSize))
+				{
+					Actor->AddGraphInstance(InComponent);
+					NewMapping.Add(Actor);
+				}
+			});
 
-		*PartitionActorsPtr = MoveTemp(NewMapping);
+			// Find the ones that were removed
+			RemovedActors = PartitionActorsPtr->Difference(NewMapping);
+
+			*PartitionActorsPtr = MoveTemp(NewMapping);
+		}
 	}
 
 	// No need to be locked to do this.
