[PCG] Fix copy paste sometimes failing to wire up pasted nodes.

Teases apart post-paste steps so they happen more "per-node" rather than performing all steps on each node in turn. More details in jira. This fix stops nodes reconstructing mid-way through the post-paste operations, by waiting until all steps are complete, by disabling bDisableReconstructFromNode as the very last step for all nodes.

#rnx

diff --git a/Engine/Plugins/Experimental/PCG/Source/PCGEditor/Private/PCGEditor.cpp b/Engine/Plugins/Experimental/PCG/Source/PCGEditor/Private/PCGEditor.cpp
--- a/Engine/Plugins/Experimental/PCG/Source/PCGEditor/Private/PCGEditor.cpp
+++ b/Engine/Plugins/Experimental/PCG/Source/PCGEditor/Private/PCGEditor.cpp
@@ -1823,14 +1823,28 @@ void FPCGEditor::PasteNodesHere(const FVector2D& Location)
 
 		PastedNode->CreateNewGuid();
 
-		if (UPCGEditorGraphNodeBase* PastedPCGGraphNode = Cast<UPCGEditorGraphNodeBase>(PastedNode))
+		UPCGEditorGraphNodeBase* PastedPCGGraphNode = Cast<UPCGEditorGraphNodeBase>(PastedNode);
+		if (UPCGNode* PastedPCGNode = PastedPCGGraphNode ? PastedPCGGraphNode->GetPCGNode() : nullptr)
 		{
-			if (UPCGNode* PastedPCGNode = PastedPCGGraphNode->GetPCGNode())
-			{
-				PCGGraphBeingEdited->AddNode(PastedPCGNode);
+			PCGGraphBeingEdited->AddNode(PastedPCGNode);
+		}
+	}
 
-				PastedPCGGraphNode->PostPaste();
-			}
+	for (UEdGraphNode* PastedNode : PastedNodes)
+	{
+		UPCGEditorGraphNodeBase* PastedPCGGraphNode = Cast<UPCGEditorGraphNodeBase>(PastedNode);
+		if (UPCGNode* PastedPCGNode = PastedPCGGraphNode ? PastedPCGGraphNode->GetPCGNode() : nullptr)
+		{
+			PastedPCGGraphNode->RebuildAfterPaste();
+		}
+	}
+
+	for (UEdGraphNode* PastedNode : PastedNodes)
+	{
+		UPCGEditorGraphNodeBase* PastedPCGGraphNode = Cast<UPCGEditorGraphNodeBase>(PastedNode);
+		if (UPCGNode* PastedPCGNode = PastedPCGGraphNode ? PastedPCGGraphNode->GetPCGNode() : nullptr)
+		{
+			PastedPCGGraphNode->PostPaste();
 		}
 	}
 
diff --git a/Engine/Plugins/Experimental/PCG/Source/PCGEditor/Private/PCGEditorGraphNodeBase.cpp b/Engine/Plugins/Experimental/PCG/Source/PCGEditor/Private/PCGEditorGraphNodeBase.cpp
--- a/Engine/Plugins/Experimental/PCG/Source/PCGEditor/Private/PCGEditorGraphNodeBase.cpp
+++ b/Engine/Plugins/Experimental/PCG/Source/PCGEditor/Private/PCGEditorGraphNodeBase.cpp
@@ -192,7 +192,7 @@ void UPCGEditorGraphNodeBase::PostPasteNode()
 	bDisableReconstructFromNode = true;
 }
 
-void UPCGEditorGraphNodeBase::PostPaste()
+void UPCGEditorGraphNodeBase::RebuildAfterPaste()
 {
 	if (PCGNode)
 	{
@@ -211,7 +211,10 @@ void UPCGEditorGraphNodeBase::PostPaste()
 			}
 		}
 	}
+}
 
+void UPCGEditorGraphNodeBase::PostPaste()
+{
 	bDisableReconstructFromNode = false;
 }
 
@@ -341,9 +344,8 @@ EPCGChangeType UPCGEditorGraphNodeBase::UpdateGridSizeVisualisation(UPCGComponen
 	}
 
 	EPCGChangeType ChangeType = EPCGChangeType::None;
-	const uint32 InspectingGridSize = (InComponentBeingDebugged && InComponentBeingDebugged->IsLocalComponent()) ?
-		InComponentBeingDebugged->GetGenerationGridSize() : PCGHiGenGrid::UninitializedGridSize();
 
+	const uint32 InspectingGridSize = InComponentBeingDebugged ? InComponentBeingDebugged->GetGenerationGridSize() : PCGHiGenGrid::UninitializedGridSize();
 	if (InspectingGridSize == PCGHiGenGrid::UninitializedGridSize())
 	{
 		if (IsDisplayAsDisabledForced())
diff --git a/Engine/Plugins/Experimental/PCG/Source/PCGEditor/Private/PCGEditorGraphNodeBase.h b/Engine/Plugins/Experimental/PCG/Source/PCGEditor/Private/PCGEditorGraphNodeBase.h
--- a/Engine/Plugins/Experimental/PCG/Source/PCGEditor/Private/PCGEditorGraphNodeBase.h
+++ b/Engine/Plugins/Experimental/PCG/Source/PCGEditor/Private/PCGEditorGraphNodeBase.h
@@ -46,6 +46,7 @@ public:
 	UPCGNode* GetPCGNode() { return PCGNode; }
 	const UPCGNode* GetPCGNode() const { return PCGNode; }
 	void PostCopy();
+	void RebuildAfterPaste();
 	void PostPaste();
 	void SetInspected(bool InIsInspecting) { bIsInspected = InIsInspecting; }
 	bool GetInspected() const { return bIsInspected; }
@@ -74,6 +75,7 @@ public:
 protected:
 	static FEdGraphPinType GetPinType(const UPCGPin* InPin);
 
+	/** Create PCG-side edges from editor pins/edges. */
 	void RebuildEdgesFromPins();
 
 	void OnNodeChanged(UPCGNode* InNode, EPCGChangeType ChangeType);
