Refactored spawn actor node to allow spawn-by-attribute.
Added basis for a more robust point data partition pattern across the code base.
Fixed a potential GC issue when doing dynamic dispatch (in the loop element).

diff --git a/Engine/Plugins/Experimental/PCG/Source/PCG/Private/Data/PCGPointData.cpp b/Engine/Plugins/Experimental/PCG/Source/PCG/Private/Data/PCGPointData.cpp
--- a/Engine/Plugins/Experimental/PCG/Source/PCG/Private/Data/PCGPointData.cpp
+++ b/Engine/Plugins/Experimental/PCG/Source/PCG/Private/Data/PCGPointData.cpp
@@ -229,11 +229,6 @@ namespace PCGPointHelpers
 	}
 }
 
-namespace PCGPointDataConstants
-{
-	const FName ActorReferenceAttribute = TEXT("ActorReference");
-}
-
 FPCGPointRef::FPCGPointRef(const FPCGPoint& InPoint)
 {
 	Point = &InPoint;
diff --git a/Engine/Plugins/Experimental/PCG/Source/PCG/Private/Elements/PCGSpawnActor.cpp b/Engine/Plugins/Experimental/PCG/Source/PCG/Private/Elements/PCGSpawnActor.cpp
--- a/Engine/Plugins/Experimental/PCG/Source/PCG/Private/Elements/PCGSpawnActor.cpp
+++ b/Engine/Plugins/Experimental/PCG/Source/PCG/Private/Elements/PCGSpawnActor.cpp
@@ -8,9 +8,11 @@
 #include "PCGSubsystem.h"
 #include "Data/PCGPointData.h"
 #include "Data/PCGSpatialData.h"
+#include "Graph/PCGStackContext.h"
 #include "Grid/PCGPartitionActor.h"
 #include "Helpers/PCGActorHelpers.h"
 #include "Helpers/PCGHelpers.h"
+#include "Helpers/PCGPointDataPartition.h"
 #include "Metadata/Accessors/IPCGAttributeAccessor.h"
 #include "Metadata/Accessors/PCGAttributeAccessorKeys.h"
 #include "Metadata/Accessors/PCGAttributeAccessorHelpers.h"
@@ -129,6 +131,83 @@ namespace PCGSpawnActorHelpers
 	};
 }
 
+class FPCGSpawnActorPartitionByAttribute : public FPCGPointDataPartitionBase<FPCGSpawnActorPartitionByAttribute, TSubclassOf<AActor>>
+{
+public:
+	FPCGSpawnActorPartitionByAttribute(FName InSpawnAttribute)
+		: FPCGPointDataPartitionBase<FPCGSpawnActorPartitionByAttribute, TSubclassOf<AActor>>()
+		, SpawnAttribute(InSpawnAttribute)
+	{
+	}
+
+	bool InitializeForPointData(const UPCGPointData* PointData)
+	{
+		if (!PointData || !PointData->ConstMetadata())
+		{
+			return false;
+		}
+
+		FPCGAttributePropertyInputSelector InputSource;
+		InputSource.SetAttributeName(SpawnAttribute);
+		InputSource = InputSource.CopyAndFixLast(PointData);
+		SpawnAttributeAccessor = PCGAttributeAccessorHelpers::CreateConstAccessor(PointData, InputSource);
+		SpawnAttributeKeys = PCGAttributeAccessorHelpers::CreateConstKeys(PointData, InputSource);
+
+		return SpawnAttributeAccessor.IsValid() && SpawnAttributeKeys.IsValid();
+	}
+
+	void AddToPartitionData(FPCGPointDataPartitionBase::Element* SelectedElement, const UPCGPointData* ParentPointData, const FPCGPoint& Point)
+	{
+		check(SelectedElement);
+		if (!SelectedElement->PartitionData)
+		{
+			SelectedElement->PartitionData = NewObject<UPCGPointData>();
+			SelectedElement->PartitionData->InitializeFromData(ParentPointData);
+		}
+
+		check(SelectedElement->PartitionData);
+		SelectedElement->PartitionData->GetMutablePoints().Add(Point);
+	}
+
+	FPCGPointDataPartitionBase::Element* SelectPoint(const FPCGPoint& Point, int32 PointIndex)
+	{
+		FString ActorPathString;
+		if (!SpawnAttributeAccessor->Get<FString>(ActorPathString, PointIndex, *SpawnAttributeKeys))
+		{
+			return nullptr;
+		}
+
+		FSoftObjectPath ActorPath(ActorPathString);
+		TSoftClassPtr<AActor> ActorClassSoftPtr(ActorPath);
+
+		UClass* ActorClass = ActorClassSoftPtr.LoadSynchronous();
+
+		if (!ActorClass)
+		{
+			UBlueprint* Blueprint = Cast<UBlueprint>(ActorPath.TryLoad());
+			if (Blueprint)
+			{
+				ActorClass = Blueprint->GeneratedClass.Get();
+			}
+		}
+		
+		if (ActorClass && ActorClass->IsChildOf<AActor>())
+		{
+			return &ElementMap.FindOrAdd(ActorClass);
+		}
+
+		return nullptr;
+	}
+
+	// Disables time-slicing altogether because the code isn't setup for this yet
+	int32 TimeSlicingCheckFrequency() const { return std::numeric_limits<int>::max(); }
+
+public:
+	TUniquePtr<const IPCGAttributeAccessor> SpawnAttributeAccessor;
+	TUniquePtr<const IPCGAttributeAccessorKeys> SpawnAttributeKeys;
+	FName SpawnAttribute = NAME_None;
+};
+
 UPCGNode* UPCGSpawnActorSettings::CreateNode() const
 {
 	return NewObject<UPCGSpawnActorNode>();
@@ -139,16 +218,16 @@ FPCGElementPtr UPCGSpawnActorSettings::CreateElement() const
 	return MakeShared<FPCGSpawnActorElement>();
 }
 
-UPCGGraphInterface* UPCGSpawnActorSettings::GetSubgraphInterface() const
+UPCGGraphInterface* UPCGSpawnActorSettings::GetGraphInterfaceFromActorSubclass(TSubclassOf<AActor> InTemplateActorClass)
 {
-	if(!TemplateActorClass || TemplateActorClass->HasAnyClassFlags(CLASS_Abstract))
+	if (!InTemplateActorClass || InTemplateActorClass->HasAnyClassFlags(CLASS_Abstract))
 	{
 		return nullptr;
 	}
 
 	UPCGGraphInterface* Result = nullptr;
 
-	AActor::ForEachComponentOfActorClassDefault<UPCGComponent>(TemplateActorClass, [&](const UPCGComponent* PCGComponent)
+	AActor::ForEachComponentOfActorClassDefault<UPCGComponent>(InTemplateActorClass, [&](const UPCGComponent* PCGComponent)
 	{
 		// If there is no graph, there is no graph instance
 		if (PCGComponent->GetGraph() && PCGComponent->bActivated)
@@ -163,6 +242,11 @@ UPCGGraphInterface* UPCGSpawnActorSettings::GetSubgraphInterface() const
 	return Result;
 }
 
+UPCGGraphInterface* UPCGSpawnActorSettings::GetSubgraphInterface() const
+{
+	return GetGraphInterfaceFromActorSubclass(TemplateActorClass);
+}
+
 void UPCGSpawnActorSettings::BeginDestroy()
 {
 #if WITH_EDITOR
@@ -227,6 +311,7 @@ bool UPCGSpawnActorSettings::IsStructuralProperty(const FName& InPropertyName) c
 {
 	return InPropertyName == GET_MEMBER_NAME_CHECKED(UPCGSpawnActorSettings, TemplateActorClass) || 
 		InPropertyName == GET_MEMBER_NAME_CHECKED(UPCGSpawnActorSettings, Option) ||
+		InPropertyName == GET_MEMBER_NAME_CHECKED(UPCGSpawnActorSettings, bSpawnByAttribute) ||
 		Super::IsStructuralProperty(InPropertyName);
 }
 #endif // WITH_EDITOR
@@ -308,26 +393,51 @@ void UPCGSpawnActorSettings::RefreshTemplateActor()
 	}
 }
 
-bool FPCGSpawnActorElement::ExecuteInternal(FPCGContext* Context) const
+bool FPCGSpawnActorElement::ExecuteInternal(FPCGContext* InContext) const
 {
-	TRACE_CPUPROFILER_EVENT_SCOPE(FPCSpawnActorElement::Execute);
+	FPCGSubgraphContext* Context = static_cast<FPCGSubgraphContext*>(InContext);
 
 	const UPCGSpawnActorSettings* Settings = Context->GetInputSettings<UPCGSpawnActorSettings>();
 	check(Settings);
 
-	// Early out
-	if(!Settings->TemplateActorClass || Settings->TemplateActorClass->HasAnyClassFlags(CLASS_Abstract))
+	if (!Context->bScheduledSubgraph)
 	{
-		const FText ClassName = Settings->TemplateActorClass ? FText::FromString(Settings->TemplateActorClass->GetFName().ToString()) : FText::FromName(NAME_None);
-		PCGE_LOG(Error, GraphAndLog, FText::Format(LOCTEXT("InvalidTemplateActorClass", "Invalid template actor class '{0}'"), ClassName));
+		return SpawnAndPrepareSubgraphs(Context, Settings);
+	}
+	else if (Context->bIsPaused)
+	{
+		// Should not happen once we skip it in the graph executor
+		return false;
+	}
+	else
+	{
+		// TODO: Currently, we don't gather results from subgraphs, but we could (in a single pin).
 		return true;
 	}
+}
 
-	if (!ensure(Settings->TemplateActor && Settings->TemplateActor->IsA(Settings->TemplateActorClass)))
+bool FPCGSpawnActorElement::SpawnAndPrepareSubgraphs(FPCGSubgraphContext* Context, const UPCGSpawnActorSettings* Settings) const
+{
+	TRACE_CPUPROFILER_EVENT_SCOPE(FPCSpawnActorElement::Execute);
+
+	// Early out
+	if (!Settings->bSpawnByAttribute)
 	{
-		return true;
+		if (!Settings->TemplateActorClass || Settings->TemplateActorClass->HasAnyClassFlags(CLASS_Abstract))
+		{
+			const FText ClassName = Settings->TemplateActorClass ? FText::FromString(Settings->TemplateActorClass->GetFName().ToString()) : FText::FromName(NAME_None);
+			PCGE_LOG(Error, GraphAndLog, FText::Format(LOCTEXT("InvalidTemplateActorClass", "Invalid template actor class '{0}'"), ClassName));
+			return true;
+		}
+
+		if (!ensure(Settings->TemplateActor && Settings->TemplateActor->IsA(Settings->TemplateActorClass)))
+		{
+			return true;
+		}
 	}
 
+	UPCGSubsystem* Subsystem = Context->SourceComponent.IsValid() ? Context->SourceComponent->GetSubsystem() : nullptr;
+
 	// Check if we can reuse existing resources - note that this is done on a per-settings basis when collapsed,
 	// Otherwise we'll check against merged crc 
 	bool bFullySkippedDueToReuse = false;
@@ -369,36 +479,17 @@ bool FPCGSpawnActorElement::ExecuteInternal(FPCGContext* Context) const
 		}
 	}
 
-	TArray<UFunction*> PostSpawnFunctions;
-	for (FName PostSpawnFunctionName : Settings->PostSpawnFunctionNames)
-	{
-		if (PostSpawnFunctionName == NAME_None)
-		{
-			continue;
-		}
-
-		if (UFunction* PostSpawnFunction = Settings->TemplateActorClass->FindFunctionByName(PostSpawnFunctionName))
-		{
-			if (PostSpawnFunction->NumParms != 0)
-			{
-				PCGE_LOG(Warning, GraphAndLog, FText::Format(LOCTEXT("ParametersMissing", "PostSpawnFunction '{0}' requires parameters. We only support parameter-less functions. Call will be skipped."), FText::FromName(PostSpawnFunctionName)));
-			}
-			else
-			{
-				PostSpawnFunctions.Add(PostSpawnFunction);
-			}
-		}
-		else
-		{
-			PCGE_LOG(Warning, GraphAndLog, FText::Format(LOCTEXT("FunctionNotFound", "PostSpawnFunction '{0}' was not found in class '{1}'"), FText::FromName(PostSpawnFunctionName), FText::FromName(Settings->TemplateActorClass->GetFName())));
-		}
-	}
-
-	const bool bForceDisableActorParsing = (Settings->bForceDisableActorParsing);
-
 	// Pass-through exclusions & settings
 	TArray<FPCGTaggedData>& Outputs = Context->OutputData.TaggedData;
 
+#if WITH_EDITOR
+	const bool bGenerateOutputsWithActorReference = (Settings->Option != EPCGSpawnActorOption::CollapseActors);
+#else
+	const bool bGenerateOutputsWithActorReference = (Settings->Option != EPCGSpawnActorOption::CollapseActors) && Context->Node && Context->Node->IsOutputPinConnected(PCGPinConstants::DefaultOutputLabel);
+#endif
+
+	const bool bHasAuthority = !Context->SourceComponent.IsValid() || (Context->SourceComponent->GetOwner() && Context->SourceComponent->GetOwner()->HasAuthority());
+
 	TArray<FPCGTaggedData> Inputs = Context->InputData.GetInputs();
 	for (const FPCGTaggedData& Input : Inputs)
 	{
@@ -418,9 +509,6 @@ bool FPCGSpawnActorElement::ExecuteInternal(FPCGContext* Context) const
 			continue;
 		}
 
-		const bool bHasAuthority = !Context->SourceComponent.IsValid() || (Context->SourceComponent->GetOwner() && Context->SourceComponent->GetOwner()->HasAuthority());
-		const bool bSpawnedActorsRequireAuthority = Settings->TemplateActor->GetIsReplicated();
-
 		// First, create target instance transforms
 		const UPCGPointData* PointData = SpatialData->ToPointData(Context);
 
@@ -438,242 +526,430 @@ bool FPCGSpawnActorElement::ExecuteInternal(FPCGContext* Context) const
 			continue;
 		}
 
-		// Spawn actors/populate ISM
-		if (!bFullySkippedDueToReuse && Settings->Option == EPCGSpawnActorOption::CollapseActors)
+		auto SpawnOrCollapse = [this, bHasAuthority, &Context, &TargetActor, &Settings](TSubclassOf<AActor> TemplateActorClass, AActor* TemplateActor, FPCGTaggedData& Output, const UPCGPointData* PointData, UPCGPointData* OutPointData)
 		{
-			TRACE_CPUPROFILER_EVENT_SCOPE(FPCGSpawnActorElement::ExecuteInternal::CollapseActors);
-
-			TMap<FPCGISMCBuilderParameters, TArray<FTransform>> MeshDescriptorTransforms;
+			const bool bSpawnedActorsRequireAuthority = (TemplateActor ? TemplateActor->GetIsReplicated() : CastChecked<AActor>(TemplateActorClass->GetDefaultObject())->GetIsReplicated());
 
-			AActor::ForEachComponentOfActorClassDefault<UStaticMeshComponent>(Settings->TemplateActorClass, [&](const UStaticMeshComponent* StaticMeshComponent)
+			if (Settings->Option == EPCGSpawnActorOption::CollapseActors)
+			{
+				CollapseIntoTargetActor(Context, TargetActor, TemplateActorClass, PointData);
+			}
+			else if (bHasAuthority || !bSpawnedActorsRequireAuthority)
 			{
-				FPCGISMCBuilderParameters Params;
-				Params.Descriptor.InitFrom(StaticMeshComponent);
-				// TODO: No custom data float support?
+				SpawnActors(Context, TargetActor, TemplateActorClass, TemplateActor, Output, PointData, OutPointData);
+			}
+		};
+
+		UPCGPointData* OutPointData = nullptr;
+		if (bGenerateOutputsWithActorReference)
+		{
+			OutPointData = NewObject<UPCGPointData>();
+			OutPointData->InitializeFromData(PointData);
+		}
+
+		FPCGTaggedData Output = Input;
+
+		if (Settings->bSpawnByAttribute && (!bFullySkippedDueToReuse || bGenerateOutputsWithActorReference))
+		{
+			FPCGSpawnActorPartitionByAttribute Selector(Settings->SpawnAttribute);
+			int32 CurrentPointIndex = 0;
 
-				TArray<FTransform>& Transforms = MeshDescriptorTransforms.FindOrAdd(Params);
+			// Selection is still needed if are fully skipped in order to write to the OutPointData.
+			Selector.SelectPoints(*Context, PointData, CurrentPointIndex, OutPointData);
 
-				if (const UInstancedStaticMeshComponent* InstancedStaticMeshComponent = Cast<UInstancedStaticMeshComponent>(StaticMeshComponent))
+			if (!bFullySkippedDueToReuse)
+			{
+				for (auto& Element : Selector.ElementMap)
 				{
-					const int32 NumInstances = InstancedStaticMeshComponent->GetInstanceCount();
-					Transforms.Reserve(Transforms.Num() + NumInstances);
-						
-					for (int32 InstanceIndex = 0; InstanceIndex < NumInstances; InstanceIndex++)
+					FPCGTaggedData PartialInput = Input;
+					PartialInput.Data = Element.Value.PartitionData;
+
+					SpawnOrCollapse(Element.Key, nullptr, PartialInput, Element.Value.PartitionData, OutPointData);
+
+					// Exception case here: if we've spawned actors but are merging the PCG inputs,
+					// normally this node is taken as a subgraph node (e.g. no need to do anything more than forwarding the inputs)
+					// However, if we`re in the spawn by attribute case, we need to dispatch it here.
+					if (Settings->Option != EPCGSpawnActorOption::NoMerging && Subsystem)
 					{
-						FTransform InstanceTransform;
-						if (InstancedStaticMeshComponent->GetInstanceTransform(InstanceIndex, InstanceTransform))
+						// TODO: maybe consider a version that would support multi PCG
+						if (UPCGGraphInterface* GraphInterface = UPCGSpawnActorSettings::GetGraphInterfaceFromActorSubclass(Element.Key))
 						{
-							Transforms.Add(InstanceTransform);
+							FPCGDataCollection SubgraphInputData;
+							SubgraphInputData.TaggedData.Add(PartialInput);
+
+							// Prepare the invocation stack - which is the stack up to this node, and then this node, then a loop index
+							FPCGStack InvocationStack = ensure(Context->Stack) ? *Context->Stack : FPCGStack();
+
+							FPCGTaskId SubgraphTaskId = Subsystem->ScheduleGraph(GraphInterface->GetGraph(),
+								Context->SourceComponent.Get(),
+								MakeShared<FPCGTrivialElement>(),// TODO: prepare user parameters like in subgraph/loop
+								MakeShared<FPCGInputForwardingElement>(SubgraphInputData),
+								{},
+								&InvocationStack);
+
+							if (SubgraphTaskId != InvalidPCGTaskId)
+							{
+								Context->SubgraphTaskIds.Add(SubgraphTaskId);
+							}
 						}
 					}
 				}
-				else
-				{
-					Transforms.Add(StaticMeshComponent->GetRelativeTransform());
-				}
+			}
+		}
+		else if(!bFullySkippedDueToReuse)
+		{
+			// Spawn actors/populate ISM
+			FPCGTaggedData InputCopy = Input;
+			SpawnOrCollapse(Settings->TemplateActorClass, Settings->TemplateActor, InputCopy, PointData, OutPointData);
+		}
 
-				return true;
-			});
-				
-			for(const TPair<FPCGISMCBuilderParameters, TArray<FTransform>>& ISMCBuilderTransforms : MeshDescriptorTransforms)
-			{
-				const FPCGISMCBuilderParameters& ISMCParams = ISMCBuilderTransforms.Key;
+		// Update the data in the output to the final data gathered
+		if (OutPointData)
+		{
+			Output.Data = OutPointData;
+		}
 
-				UPCGManagedISMComponent* MISMC = UPCGActorHelpers::GetOrCreateManagedISMC(TargetActor, Context->SourceComponent.Get(), Settings->UID, ISMCParams);
-				if (!MISMC)
-				{
-					continue;
-				}
+		// Finally, pass through the input, in all cases: 
+		// - if it's not merged, will be the input points directly
+		// - if it's merged but there is no subgraph, will be the input points directly
+		// - if it's merged and there is a subgraph, we'd need to pass the data for it to be given to the subgraph
+		Outputs.Add(Output);
+	}
+
+	// If we've dispatched dynamic execution, we should queue a task here to wait for those
+	if (!Context->SubgraphTaskIds.IsEmpty())
+	{
+		Context->bScheduledSubgraph = true;
+		Context->bIsPaused = true;
+
+		Subsystem->ScheduleGeneric([Context]() {
+			// Wake up the current task
+			Context->bIsPaused = false;
+			return true;
+		}, Context->SourceComponent.Get(), Context->SubgraphTaskIds);
+
+		return false;
+	}
+	else
+	{
+		return true;
+	}
+}
+
+void FPCGSpawnActorElement::CollapseIntoTargetActor(FPCGSubgraphContext* Context, AActor* TargetActor, TSubclassOf<AActor> TemplateActorClass, const UPCGPointData* PointData) const
+{
+	TRACE_CPUPROFILER_EVENT_SCOPE(FPCGSpawnActorElement::ExecuteInternal::CollapseActors);
+	check(Context && TargetActor && PointData);
+
+	const TArray<FPCGPoint>& Points = PointData->GetPoints();
+	if (Points.IsEmpty())
+	{
+		return;
+	}
+
+	const UPCGSpawnActorSettings* Settings = Context->GetInputSettings<UPCGSpawnActorSettings>();
+	check(Settings);
+
+	TMap<FPCGISMCBuilderParameters, TArray<FTransform>> MeshDescriptorTransforms;
+
+	AActor::ForEachComponentOfActorClassDefault<UStaticMeshComponent>(TemplateActorClass, [&MeshDescriptorTransforms](const UStaticMeshComponent* StaticMeshComponent)
+	{
+		FPCGISMCBuilderParameters Params;
+		Params.Descriptor.InitFrom(StaticMeshComponent);
+		// TODO: No custom data float support?
+
+		TArray<FTransform>& Transforms = MeshDescriptorTransforms.FindOrAdd(Params);
 
-				MISMC->SetCrc(Context->DependenciesCrc);
+		if (const UInstancedStaticMeshComponent* InstancedStaticMeshComponent = Cast<UInstancedStaticMeshComponent>(StaticMeshComponent))
+		{
+			const int32 NumInstances = InstancedStaticMeshComponent->GetInstanceCount();
+			Transforms.Reserve(Transforms.Num() + NumInstances);
 
-				UInstancedStaticMeshComponent* ISMC = MISMC->GetComponent();
-				check(ISMC);
-				
-				const TArray<FTransform>& ISMCTransforms = ISMCBuilderTransforms.Value;
-				
-				TArray<FTransform> Transforms;
-				Transforms.Reserve(Points.Num() * ISMCTransforms.Num());
-				for (int32 PointIndex = 0; PointIndex < Points.Num(); PointIndex++)
+			for (int32 InstanceIndex = 0; InstanceIndex < NumInstances; InstanceIndex++)
+			{
+				FTransform InstanceTransform;
+				if (InstancedStaticMeshComponent->GetInstanceTransform(InstanceIndex, InstanceTransform))
 				{
-					const FPCGPoint& Point = Points[PointIndex];
-					for (int32 TransformIndex = 0; TransformIndex < ISMCTransforms.Num(); TransformIndex++)
-					{
-						const FTransform& Transform = ISMCTransforms[TransformIndex];
-						Transforms.Add(Transform * Point.Transform);
-					}
+					Transforms.Add(InstanceTransform);
 				}
-				
-				// Fill in custom data (?)
-				ISMC->AddInstances(Transforms, false, true);
-				ISMC->UpdateBounds();
-				
-				PCGE_LOG(Verbose, LogOnly, FText::Format(LOCTEXT("InstanceCreationInfo", "Added {0} instances of mesh '{1}' to ISMC '{2}' on actor '{3}'"),
-					Transforms.Num(), FText::FromString(ISMC->GetStaticMesh().GetName()), FText::FromString(ISMC->GetName()), FText::FromString(TargetActor->GetActorNameOrLabel())));
 			}
 		}
-		else if (!bFullySkippedDueToReuse && Settings->Option != EPCGSpawnActorOption::CollapseActors && (bHasAuthority || !bSpawnedActorsRequireAuthority))
+		else
 		{
-			TRACE_CPUPROFILER_EVENT_SCOPE(FPCGSpawnActorElement::ExecuteInternal::SpawnActors);
-			AActor* TemplateActor = Settings->ActorOverrides.IsEmpty() ? Settings->TemplateActor.Get() : DuplicateObject(Settings->TemplateActor, GetTransientPackage());
+			Transforms.Add(StaticMeshComponent->GetRelativeTransform());
+		}
+
+		return true;
+	});
+
+	for (const TPair<FPCGISMCBuilderParameters, TArray<FTransform>>& ISMCBuilderTransforms : MeshDescriptorTransforms)
+	{
+		const FPCGISMCBuilderParameters& ISMCParams = ISMCBuilderTransforms.Key;
+
+		UPCGManagedISMComponent* MISMC = UPCGActorHelpers::GetOrCreateManagedISMC(TargetActor, Context->SourceComponent.Get(), Settings->UID, ISMCParams);
+		if (!MISMC)
+		{
+			continue;
+		}
+
+		MISMC->SetCrc(Context->DependenciesCrc);
+
+		UInstancedStaticMeshComponent* ISMC = MISMC->GetComponent();
+		check(ISMC);
+
+		const TArray<FTransform>& ISMCTransforms = ISMCBuilderTransforms.Value;
+
+		TArray<FTransform> Transforms;
+		Transforms.Reserve(Points.Num() * ISMCTransforms.Num());
+		for (int32 PointIndex = 0; PointIndex < Points.Num(); PointIndex++)
+		{
+			const FPCGPoint& Point = Points[PointIndex];
+			for (int32 TransformIndex = 0; TransformIndex < ISMCTransforms.Num(); TransformIndex++)
+			{
+				const FTransform& Transform = ISMCTransforms[TransformIndex];
+				Transforms.Add(Transform * Point.Transform);
+			}
+		}
+
+		// Fill in custom data (?)
+		ISMC->AddInstances(Transforms, false, true);
+		ISMC->UpdateBounds();
+
+		PCGE_LOG(Verbose, LogOnly, FText::Format(LOCTEXT("InstanceCreationInfo", "Added {0} instances of mesh '{1}' to ISMC '{2}' on actor '{3}'"),
+			Transforms.Num(), FText::FromString(ISMC->GetStaticMesh().GetName()), FText::FromString(ISMC->GetName()), FText::FromString(TargetActor->GetActorNameOrLabel())));
+	}
+}
+
+void FPCGSpawnActorElement::SpawnActors(FPCGSubgraphContext* Context, AActor* TargetActor, TSubclassOf<AActor> InTemplateActorClass, AActor* InTemplateActor, FPCGTaggedData& Output, const UPCGPointData* PointData, UPCGPointData* OutPointData) const
+{
+	TRACE_CPUPROFILER_EVENT_SCOPE(FPCGSpawnActorElement::ExecuteInternal::SpawnActors);
+	check(Context && TargetActor && PointData);
+
+	const TArray<FPCGPoint>& Points = PointData->GetPoints();
+	if (Points.IsEmpty())
+	{
+		return;
+	}
+
+	int32 OutPointOffset = 0;
+	FPCGMetadataAttribute<FString>* ActorReferenceAttribute = nullptr;
+
+	if (OutPointData)
+	{
+		OutPointOffset = OutPointData->GetMutablePoints().Num();
+		OutPointData->GetMutablePoints().Append(Points);
+		ActorReferenceAttribute = OutPointData->MutableMetadata()->FindOrCreateAttribute<FString>(PCGPointDataConstants::ActorReferenceAttribute, FString(), /*bAllowsInterpolation=*/false, /*bOverrideParent=*/false, /*bOverwriteIfTypeMismatch=*/false);
+	}
 
-			PCGSpawnActorHelpers::FActorOverrides ActorOverrides(Settings->ActorOverrides, TemplateActor, PointData);
+	const UPCGSpawnActorSettings* Settings = Context->GetInputSettings<UPCGSpawnActorSettings>();
+	check(Settings && Settings->Option != EPCGSpawnActorOption::CollapseActors);
 
-			FActorSpawnParameters SpawnParams;
-			SpawnParams.Owner = TargetActor;
-			SpawnParams.Template = TemplateActor;
-			SpawnParams.SpawnCollisionHandlingOverride = ESpawnActorCollisionHandlingMethod::AlwaysSpawn;
+	TArray<UFunction*> PostSpawnFunctions;
+	for (FName PostSpawnFunctionName : Settings->PostSpawnFunctionNames)
+	{
+		if (PostSpawnFunctionName == NAME_None)
+		{
+			continue;
+		}
 
-			if (PCGHelpers::IsRuntimeOrPIE())
+		if (UFunction* PostSpawnFunction = InTemplateActorClass->FindFunctionByName(PostSpawnFunctionName))
+		{
+			if (PostSpawnFunction->NumParms != 0)
 			{
-				SpawnParams.ObjectFlags |= RF_Transient;
+				PCGE_LOG(Warning, GraphAndLog, FText::Format(LOCTEXT("ParametersMissing", "PostSpawnFunction '{0}' requires parameters. We only support parameter-less functions. Call will be skipped."), FText::FromName(PostSpawnFunctionName)));
 			}
+			else
+			{
+				PostSpawnFunctions.Add(PostSpawnFunction);
+			}
+		}
+		else
+		{
+			PCGE_LOG(Warning, GraphAndLog, FText::Format(LOCTEXT("FunctionNotFound", "PostSpawnFunction '{0}' was not found in class '{1}'"), FText::FromName(PostSpawnFunctionName), FText::FromName(InTemplateActorClass->GetFName())));
+		}
+	}
+
+	const bool bForceDisableActorParsing = (Settings->bForceDisableActorParsing);
+
+	AActor* TemplateActor = nullptr;
+	if (InTemplateActor)
+	{
+		if (Settings->ActorOverrides.IsEmpty())
+		{
+			TemplateActor = InTemplateActor;
+		}
+		else
+		{
+			TemplateActor = DuplicateObject(InTemplateActor, GetTransientPackage());
+		}
+	}
+	else
+	{
+		TemplateActor = NewObject<AActor>(GetTransientPackage(), InTemplateActorClass, NAME_None, RF_ArchetypeObject);
+	}
+
+	check(TemplateActor);
+
+	PCGSpawnActorHelpers::FActorOverrides ActorOverrides(Settings->ActorOverrides, TemplateActor, PointData);
+
+	FActorSpawnParameters SpawnParams;
+	SpawnParams.Owner = TargetActor;
+	SpawnParams.Template = TemplateActor;
+	SpawnParams.SpawnCollisionHandlingOverride = ESpawnActorCollisionHandlingMethod::AlwaysSpawn;
 
-			const bool bForceCallGenerate = (Settings->GenerationTrigger == EPCGSpawnActorGenerationTrigger::ForceGenerate);
+	if (PCGHelpers::IsRuntimeOrPIE())
+	{
+		SpawnParams.ObjectFlags |= RF_Transient;
+	}
+
+	const bool bForceCallGenerate = (Settings->GenerationTrigger == EPCGSpawnActorGenerationTrigger::ForceGenerate);
 #if WITH_EDITOR
-			const bool bOnLoadCallGenerate = (Settings->GenerationTrigger == EPCGSpawnActorGenerationTrigger::Default);
+	const bool bOnLoadCallGenerate = (Settings->GenerationTrigger == EPCGSpawnActorGenerationTrigger::Default);
 #else
-			const bool bOnLoadCallGenerate = (Settings->GenerationTrigger == EPCGSpawnActorGenerationTrigger::Default ||
-				Settings->GenerationTrigger == EPCGSpawnActorGenerationTrigger::DoNotGenerateInEditor);
+	const bool bOnLoadCallGenerate = (Settings->GenerationTrigger == EPCGSpawnActorGenerationTrigger::Default ||
+		Settings->GenerationTrigger == EPCGSpawnActorGenerationTrigger::DoNotGenerateInEditor);
 #endif
-			UPCGSubsystem* Subsystem = (Context->SourceComponent.Get() ? Context->SourceComponent->GetSubsystem() : nullptr);
+	UPCGSubsystem* Subsystem = (Context->SourceComponent.Get() ? Context->SourceComponent->GetSubsystem() : nullptr);
 
-			// Try to reuse actors if the are preexisting
-			TArray<UPCGManagedActors*> ReusedManagedActorsResources;
-			FPCGCrc InputDependenciesCrc;
-			if (CVarAllowActorReuse.GetValueOnAnyThread())
-			{
-				FPCGDataCollection SingleInputCollection;
-				SingleInputCollection.TaggedData.Add(Input);
-				SingleInputCollection.Crc = SingleInputCollection.ComputeCrc(/*bFullDataCrc=*/false);
+	// Try to reuse actors if the are preexisting
+	TArray<UPCGManagedActors*> ReusedManagedActorsResources;
+	FPCGCrc InputDependenciesCrc;
+	if (CVarAllowActorReuse.GetValueOnAnyThread())
+	{
+		FPCGDataCollection SingleInputCollection;
+		SingleInputCollection.TaggedData.Add(Output);
+		// TODO: review this, it might make more sense to do a full crc here
+		SingleInputCollection.Crc = SingleInputCollection.ComputeCrc(/*bFullDataCrc=*/false);
 
-				GetDependenciesCrc(SingleInputCollection, Settings, Context->SourceComponent.Get(), InputDependenciesCrc);
+		GetDependenciesCrc(SingleInputCollection, Settings, Context->SourceComponent.Get(), InputDependenciesCrc);
 
-				if(InputDependenciesCrc.IsValid())
+		if (InputDependenciesCrc.IsValid())
+		{
+			Context->SourceComponent->ForEachManagedResource([&ReusedManagedActorsResources, &InputDependenciesCrc, &Context](UPCGManagedResource* InResource)
+			{
+				if (UPCGManagedActors* Resource = Cast<UPCGManagedActors>(InResource))
 				{
-					Context->SourceComponent->ForEachManagedResource([&ReusedManagedActorsResources, &InputDependenciesCrc , &Context](UPCGManagedResource* InResource)
+					if (Resource->GetCrc().IsValid() && Resource->GetCrc() == InputDependenciesCrc)
 					{
-						if (UPCGManagedActors* Resource = Cast<UPCGManagedActors>(InResource))
-						{
-							if (Resource->GetCrc().IsValid() && Resource->GetCrc() == InputDependenciesCrc)
-							{
-								ReusedManagedActorsResources.Add(Resource);
-							}
-						}
-					});
+						ReusedManagedActorsResources.Add(Resource);
+					}
 				}
-			}
+			});
+		}
+	}
 
-			TArray<AActor*> ProcessedActors;
-			const bool bActorsHavePCGComponents = (Settings->GetSubgraphInterface() != nullptr);
+	TArray<AActor*> ProcessedActors;
+	const bool bActorsHavePCGComponents = (UPCGSpawnActorSettings::GetGraphInterfaceFromActorSubclass(InTemplateActorClass) != nullptr);
 
-			if (!ReusedManagedActorsResources.IsEmpty())
+	if (!ReusedManagedActorsResources.IsEmpty())
+	{
+		// If the actors are fully independent, we might need to make sure to call Generate if the underlying graph has changed - e.g. if the actor is dirty
+		for (UPCGManagedActors* ManagedActors : ReusedManagedActorsResources)
+		{
+			check(ManagedActors);
+			ManagedActors->MarkAsReused();
+
+			// There's no setup to be done, just generation if we're in the no-merge case, so keep track of these actors only in this case
+			if (bActorsHavePCGComponents && Settings->Option == EPCGSpawnActorOption::NoMerging)
 			{
-				// If the actors are fully independent, we might need to make sure to call Generate if the underlying graph has changed - e.g. if the actor is dirty
-				for (UPCGManagedActors* ManagedActors : ReusedManagedActorsResources)
+				for (TSoftObjectPtr<AActor>& ManagedActorPtr : ManagedActors->GeneratedActors)
 				{
-					check(ManagedActors);
-					ManagedActors->MarkAsReused();
-
-					// There's no setup to be done, just generation if we're in the no-merge case, so keep track of these actors only in this case
-					if (bActorsHavePCGComponents && Settings->Option == EPCGSpawnActorOption::NoMerging)
+					if (AActor* ManagedActor = ManagedActorPtr.Get())
 					{
-						for (TSoftObjectPtr<AActor>& ManagedActorPtr : ManagedActors->GeneratedActors)
-						{
-							if (AActor* ManagedActor = ManagedActorPtr.Get())
-							{
-								ProcessedActors.Add(ManagedActor);
-							}
-						}
+						ProcessedActors.Add(ManagedActor);
 					}
 				}
 			}
-			else
-			{
-				TArray<FName> NewActorTags = GetNewActorTags(Context, TargetActor, Settings->bInheritActorTags, Settings->TagsToAddOnActors);
-
-				// Create managed resource for actor tracking
-				UPCGManagedActors* ManagedActors = NewObject<UPCGManagedActors>(Context->SourceComponent.Get());
-				ManagedActors->SetCrc(InputDependenciesCrc);
+		}
+	}
+	else
+	{
+		TArray<FName> NewActorTags = GetNewActorTags(Context, TargetActor, Settings->bInheritActorTags, Settings->TagsToAddOnActors);
 
-				for (int32 i = 0; i < Points.Num(); ++i)
-				{
-					const FPCGPoint& Point = Points[i];
+		// Create managed resource for actor tracking
+		UPCGManagedActors* ManagedActors = NewObject<UPCGManagedActors>(Context->SourceComponent.Get());
+		ManagedActors->SetCrc(InputDependenciesCrc);
 
-					ActorOverrides.Apply(i);
+		for (int32 i = 0; i < Points.Num(); ++i)
+		{
+			const FPCGPoint& Point = Points[i];
 
-					AActor* GeneratedActor = TargetActor->GetWorld()->SpawnActor(Settings->TemplateActorClass, &Point.Transform, SpawnParams);
+			ActorOverrides.Apply(i);
 
-					if (!GeneratedActor)
-					{
-						PCGE_LOG(Error, GraphAndLog, FText::Format(LOCTEXT("ActorSpawnFailed", "Failed to spawn actor on point with index {0}"), i));
-						continue;
-					}
+			AActor* GeneratedActor = TargetActor->GetWorld()->SpawnActor(InTemplateActorClass, &Point.Transform, SpawnParams);
 
-					// HACK: until UE-62747 is fixed, we have to force set the scale after spawning the actor
-					GeneratedActor->SetActorRelativeScale3D(Point.Transform.GetScale3D());
-					GeneratedActor->Tags.Append(NewActorTags);
-					GeneratedActor->AttachToActor(TargetActor, FAttachmentTransformRules::KeepWorldTransform);
+			if (!GeneratedActor)
+			{
+				PCGE_LOG(Error, GraphAndLog, FText::Format(LOCTEXT("ActorSpawnFailed", "Failed to spawn actor on point with index {0}"), i));
+				continue;
+			}
 
-					for (UFunction* PostSpawnFunction : PostSpawnFunctions)
-					{
-						GeneratedActor->ProcessEvent(PostSpawnFunction, nullptr);
-					}
+			// HACK: until UE-62747 is fixed, we have to force set the scale after spawning the actor
+			GeneratedActor->SetActorRelativeScale3D(Point.Transform.GetScale3D());
+			GeneratedActor->Tags.Append(NewActorTags);
+			GeneratedActor->AttachToActor(TargetActor, FAttachmentTransformRules::KeepWorldTransform);
 
-					ManagedActors->GeneratedActors.Add(GeneratedActor);
+			for (UFunction* PostSpawnFunction : PostSpawnFunctions)
+			{
+				GeneratedActor->ProcessEvent(PostSpawnFunction, nullptr);
+			}
 
-					if (bActorsHavePCGComponents)
-					{
-						ProcessedActors.Add(GeneratedActor);
-					}
-				}
+			ManagedActors->GeneratedActors.Add(GeneratedActor);
 
-				Context->SourceComponent->AddToManagedResources(ManagedActors);
+			if (bActorsHavePCGComponents)
+			{
+				ProcessedActors.Add(GeneratedActor);
+			}
 
-				PCGE_LOG(Verbose, LogOnly, FText::Format(LOCTEXT("GenerationInfo", "Generated {0} actors"), Points.Num()));
+			// Write to out data the actor reference
+			if (OutPointData && ActorReferenceAttribute)
+			{
+				FPCGPoint& OutPoint = OutPointData->GetMutablePoints()[i + OutPointOffset];
+				ActorReferenceAttribute->SetValue(OutPoint.MetadataEntry, FSoftObjectPath(GeneratedActor).ToString());
 			}
+		}
+
+		Context->SourceComponent->AddToManagedResources(ManagedActors);
+
+		PCGE_LOG(Verbose, LogOnly, FText::Format(LOCTEXT("GenerationInfo", "Generated {0} actors"), Points.Num()));
+	}
 
-			// Setup & Generate on PCG components if needed
-			for (AActor* Actor : ProcessedActors)
+	// Setup & Generate on PCG components if needed
+	for (AActor* Actor : ProcessedActors)
+	{
+		TInlineComponentArray<UPCGComponent*, 1> PCGComponents;
+		Actor->GetComponents(PCGComponents);
+
+		for (UPCGComponent* PCGComponent : PCGComponents)
+		{
+			if (Settings->Option == EPCGSpawnActorOption::NoMerging)
 			{
-				TInlineComponentArray<UPCGComponent*, 1> PCGComponents;
-				Actor->GetComponents(PCGComponents);
+				if (bForceDisableActorParsing)
+				{
+					PCGComponent->bParseActorComponents = false;
+				}
 
-				for (UPCGComponent* PCGComponent : PCGComponents)
+				if (bForceCallGenerate || (bOnLoadCallGenerate && PCGComponent->GenerationTrigger == EPCGComponentGenerationTrigger::GenerateOnLoad))
 				{
-					if(Settings->Option == EPCGSpawnActorOption::NoMerging)
+					if (Subsystem)
 					{
-						if (bForceDisableActorParsing)
-						{
-							PCGComponent->bParseActorComponents = false;
-						}
-
-						if (bForceCallGenerate || (bOnLoadCallGenerate && PCGComponent->GenerationTrigger == EPCGComponentGenerationTrigger::GenerateOnLoad))
-						{
-							if (Subsystem)
-							{
-								Subsystem->RegisterOrUpdatePCGComponent(PCGComponent);
-							}
-
-							PCGComponent->Generate();
-						}
+						Subsystem->RegisterOrUpdatePCGComponent(PCGComponent);
 					}
-					else
+
+					// TODO: use ScheduleGraph if we want to pass user parameters
+					FPCGTaskId SubgraphTaskId = PCGComponent->GenerateLocalGetTaskId(/*bForce=*/true);
+					if (SubgraphTaskId != InvalidPCGTaskId)
 					{
-						PCGComponent->bActivated = false;
+						Context->SubgraphTaskIds.Add(SubgraphTaskId);
 					}
 				}
 			}
+			else // otherwise, they will be taken care of as-if a subgraph (either dynamically or statically)
+			{
+				PCGComponent->bActivated = false;
+			}
 		}
-
-		// Finally, pass through the input, in all cases: 
-		// - if it's not merged, will be the input points directly
-		// - if it's merged but there is no subgraph, will be the input points directly
-		// - if it's merged and there is a subgraph, we'd need to pass the data for it to be given to the subgraph
-		Outputs.Add(Input);
 	}
-
-	return true;
 }
 
 TArray<FName> FPCGSpawnActorElement::GetNewActorTags(FPCGContext* Context, AActor* TargetActor, bool bInheritActorTags, const TArray<FName>& AdditionalTags) const
diff --git a/Engine/Plugins/Experimental/PCG/Source/PCG/Private/PCGSubgraph.cpp b/Engine/Plugins/Experimental/PCG/Source/PCG/Private/PCGSubgraph.cpp
--- a/Engine/Plugins/Experimental/PCG/Source/PCG/Private/PCGSubgraph.cpp
+++ b/Engine/Plugins/Experimental/PCG/Source/PCG/Private/PCGSubgraph.cpp
@@ -321,7 +321,7 @@ FPCGContext* FPCGSubgraphElement::Initialize(const FPCGDataCollection& InputData
 	Context->Node = Node;
 
 	// Only duplicate the UserParameters if we have overriable params and we have at least one param pin connected.
-	const UPCGSubgraphSettings* Settings = Context->GetInputSettings<UPCGSubgraphSettings>();
+	const UPCGBaseSubgraphSettings* Settings = Context->GetInputSettings<UPCGBaseSubgraphSettings>();
 	check(Settings);
 	const TArray<FPCGSettingsOverridableParam>& OverridableParams = Settings->OverridableParams();
 
@@ -445,7 +445,7 @@ bool FPCGSubgraphElement::ExecuteInternal(FPCGContext* InContext) const
 					Subsystem->ScheduleGeneric([Context]() {
 						// Wake up the current task
 						Context->bIsPaused = false;
-					return true;
+						return true;
 					}, Context->SourceComponent.Get(), Context->SubgraphTaskIds);
 
 					return false;
@@ -506,11 +506,29 @@ bool FPCGSubgraphElement::ExecuteInternal(FPCGContext* InContext) const
 FPCGInputForwardingElement::FPCGInputForwardingElement(const FPCGDataCollection& InputToForward)
 	: Input(InputToForward)
 {
-
+	// Root any previously unrooted data, so we don't need to make sure the caller is still around until this is executed
+	for (const FPCGTaggedData& TaggedData : Input.TaggedData)
+	{
+		if (TaggedData.Data && !TaggedData.Data->IsRooted())
+		{
+			UPCGData* DataToRoot = const_cast<UPCGData*>(TaggedData.Data.Get());
+			DataToRoot->AddToRoot();
+			RootedData.Add(DataToRoot);
+		}
+	}
 }
 
 bool FPCGInputForwardingElement::ExecuteInternal(FPCGContext* Context) const
 {
+	// Remove from rootset during the execution; data will be re-rooted by the normal process in the graph executor
+	for (UPCGData* DataToUnroot : RootedData)
+	{
+		ensure(DataToUnroot->IsRooted());
+		DataToUnroot->RemoveFromRoot();
+	}
+
+	const_cast<FPCGInputForwardingElement*>(this)->RootedData.Reset();
+
 	Context->OutputData = Input;
 	return true;
 }
diff --git a/Engine/Plugins/Experimental/PCG/Source/PCG/Public/Data/PCGPointData.h b/Engine/Plugins/Experimental/PCG/Source/PCG/Public/Data/PCGPointData.h
--- a/Engine/Plugins/Experimental/PCG/Source/PCG/Public/Data/PCGPointData.h
+++ b/Engine/Plugins/Experimental/PCG/Source/PCG/Public/Data/PCGPointData.h
@@ -20,6 +20,11 @@ namespace PCGPointHelpers
 	void BilerpWithSnapping(const FPCGPoint& X0Y0, const FPCGPoint& X1Y0, const FPCGPoint& X0Y1, const FPCGPoint& X1Y1, const UPCGMetadata* SourceMetadata, FPCGPoint& OutPoint, UPCGMetadata* OutMetadata, float XFactor, float YFactor);
 }
 
+namespace PCGPointDataConstants
+{
+	const FName ActorReferenceAttribute = TEXT("ActorReference");
+}
+
 struct PCG_API FPCGPointRef
 {
 	FPCGPointRef(const FPCGPoint& InPoint);
diff --git a/Engine/Plugins/Experimental/PCG/Source/PCG/Public/Elements/PCGSpawnActor.h b/Engine/Plugins/Experimental/PCG/Source/PCG/Public/Elements/PCGSpawnActor.h
--- a/Engine/Plugins/Experimental/PCG/Source/PCG/Public/Elements/PCGSpawnActor.h
+++ b/Engine/Plugins/Experimental/PCG/Source/PCG/Public/Elements/PCGSpawnActor.h
@@ -10,6 +10,7 @@
 
 class AActor;
 class UPCGGraphInterface;
+class UPCGPointData;
 
 UENUM()
 enum class EPCGSpawnActorOption : uint8
@@ -40,6 +41,13 @@ struct FPCGActorPropertyOverride
 	FString PropertyTarget;
 };
 
+/*
+* PCG settings class that allows spawning actors with some options to perform the work more efficiently.
+* Note that depending on the options, any PCG components on the spawned actors can be also generated,
+* which is why this class derives from UPCGBaseSubgraphSettings - it has similar inner-workings to the subgraph node
+* as far as data passing and dispatch go.
+* Note that at this point in time, results from the underlying graphs being generated is not propagated back as results of this node.
+*/
 UCLASS(BlueprintType, ClassGroup = (Procedural))
 class PCG_API UPCGSpawnActorSettings : public UPCGBaseSubgraphSettings
 {
@@ -78,6 +86,12 @@ public:
 	UPROPERTY(BlueprintReadWrite, Category = Settings, meta = (PCG_Overridable))
 	TSoftObjectPtr<AActor> RootActor;
 
+	UPROPERTY(BlueprintReadWrite, EditAnywhere, Category = Settings, meta = (InlineEditConditionToggle))
+	bool bSpawnByAttribute = false;
+
+	UPROPERTY(BlueprintReadWrite, EditAnywhere, Category = Settings, meta = (EditCondition = "bSpawnByAttribute"))
+	FName SpawnAttribute = NAME_None;
+
 protected:
 #if WITH_EDITORONLY_DATA
 	UPROPERTY()
@@ -113,8 +127,12 @@ protected:
 
 public:
 	//~Begin UPCGBaseSubgraphSettings interface
+	// When using spawn by attribute, the potential execution of subgraphs will be done in a dynamic manner
+	virtual bool IsDynamicGraph() const { return bSpawnByAttribute; }
 	virtual UPCGGraphInterface* GetSubgraphInterface() const override;
 
+	static UPCGGraphInterface* GetGraphInterfaceFromActorSubclass(TSubclassOf<AActor> InTemplateActorClass);
+
 protected:
 #if WITH_EDITOR
 	virtual bool IsStructuralProperty(const FName& InPropertyName) const override;
@@ -138,18 +156,23 @@ public:
 	/** ~End UPCGBaseSubgraphNode interface */
 };
 
-class FPCGSpawnActorElement : public FSimplePCGElement
+class FPCGSpawnActorElement : public FPCGSubgraphElement
 {
 public:
 	virtual bool CanExecuteOnlyOnMainThread(FPCGContext* Context) const override { return true; }
 	virtual bool IsCacheable(const UPCGSettings* InSettings) const override { return false; }
-	virtual bool IsPassthrough(const UPCGSettings* InSettings) const override { return !InSettings || InSettings->bEnabled; }
+	virtual bool IsPassthrough(const UPCGSettings* InSettings) const override { return !InSettings; }
 
 protected:
 	virtual bool ExecuteInternal(FPCGContext* Context) const override;
 
 private:
+	bool SpawnAndPrepareSubgraphs(FPCGSubgraphContext* Context, const UPCGSpawnActorSettings* Settings) const;
+
 	TArray<FName> GetNewActorTags(FPCGContext* Context, AActor* TargetActor, bool bInheritActorTags, const TArray<FName>& AdditionalTags) const;
+
+	void CollapseIntoTargetActor(FPCGSubgraphContext* Context, AActor* TargetActor, TSubclassOf<AActor> TemplateActorClass, const UPCGPointData* PointData) const;
+	void SpawnActors(FPCGSubgraphContext* Context, AActor* TargetActor, TSubclassOf<AActor> TemplateActorClass, AActor* TemplateActor, FPCGTaggedData& Output, const UPCGPointData* PointData, UPCGPointData* OutPointData) const;
 };
 
 #if UE_ENABLE_INCLUDE_ORDER_DEPRECATED_IN_5_2
diff --git a/Engine/Plugins/Experimental/PCG/Source/PCG/Public/Helpers/PCGPointDataPartition.h b/Engine/Plugins/Experimental/PCG/Source/PCG/Public/Helpers/PCGPointDataPartition.h
new file mode 100644
--- /dev/null
+++ b/Engine/Plugins/Experimental/PCG/Source/PCG/Public/Helpers/PCGPointDataPartition.h
@@ -0,0 +1,115 @@
+// Copyright Epic Games, Inc. All Rights Reserved.
+
+#pragma once
+
+#include "PCGContext.h"
+#include "PCGElement.h"
+#include "PCGModule.h"
+#include "Data/PCGPointData.h"
+
+/** 
+* Convenience class to do point data partitioning, i.e. building a TMap<Value, TArray<point/point indices>>
+* Uses CRTP to remove the need for virtual calls.
+* Has some behavior options depending on the derived class.
+* Currently supports time-slicing (caveat for rooting data, not done yet)
+* Can easily write to a common output (for match & set cases)
+* Can write to separate data (for partitioning)
+* Supports multiple data (indices won't be useful though)
+* 
+* The intended implementation is that derived classes will implement the same functions (e.g. name hiding) and those will
+* be appropriately called through using the derived class type.
+* Methods that are not re-implemented/hidden will retain their default implementation, which is mostly trivial.
+*/
+template<typename Derived, typename KeyType>
+class FPCGPointDataPartitionBase
+{
+public:
+	struct Element
+	{
+		TArray<int32> PointIndices;
+		UPCGPointData* PartitionData = nullptr;
+	};
+
+	FPCGPointDataPartitionBase() = default;
+	FPCGPointDataPartitionBase(const TArrayView<KeyType>& InKeys)
+	{
+		for (const KeyType& Key : InKeys)
+		{
+			ElementMap.Emplace(Key, Element());
+		}
+	}
+
+	Derived* This() { return static_cast<Derived*>(this); }
+	const Derived* This() const { return static_cast<const Derived*>(this); }
+
+	/** Overridable behavior methods */
+	bool Initialize() { return true; }
+	bool InitializeForPointData(const UPCGPointData* PointData) { return true; }
+	void AddToPartitionData(Element* SelectedElement, const UPCGPointData* ParentPointData, const FPCGPoint& Point) {}
+	void WriteToOutputData(UPCGPointData* OutPointData, Element* SelectedElement, const FPCGPoint& Point) {}
+	Element* SelectPoint(const FPCGPoint& Point, int32 PointIndex) { return nullptr; }
+	int32 TimeSlicingCheckFrequency() const { return 1024; }
+
+	/** API */
+	bool SelectPoints(FPCGContext& Context, const UPCGPointData* PointData, int32& InCurrentPointIndex, UPCGPointData* OutPointData);
+	void Reset() { ElementMap.Reset(); }
+
+	// Data access
+	TMap<KeyType, Element> ElementMap;
+};
+
+template<typename Derived, typename KeyType>
+bool FPCGPointDataPartitionBase<Derived, KeyType>::SelectPoints(FPCGContext& Context, const UPCGPointData* PointData, int32& InCurrentPointIndex, UPCGPointData* OutPointData)
+{
+	if (!PointData)
+	{
+		PCGE_LOG_C(Error, GraphAndLog, &Context, NSLOCTEXT("PCGPointDataPartition", "InputMissingData", "Missing input data"));
+		return true;
+	}
+
+	if (InCurrentPointIndex == 0)
+	{
+		if (!This()->InitializeForPointData(PointData))
+		{
+			return true;
+		}
+	}
+
+	int32 CurrentPointIndex = InCurrentPointIndex;
+	int32 LastCheckpointIndex = InCurrentPointIndex;
+	const int32 TimeSlicingCheckFrequency = This()->TimeSlicingCheckFrequency();
+
+	const TArray<FPCGPoint>& Points = PointData->GetPoints();
+
+	while (CurrentPointIndex < Points.Num())
+	{
+		const FPCGPoint& Point = Points[CurrentPointIndex];
+
+		Element* SelectedElement = This()->SelectPoint(Point, CurrentPointIndex);
+		if (SelectedElement)
+		{
+			SelectedElement->PointIndices.Add(CurrentPointIndex);
+			This()->AddToPartitionData(SelectedElement, PointData, Point);
+		}
+
+		This()->WriteToOutputData(OutPointData, SelectedElement, Point);
+
+		++CurrentPointIndex;
+
+		if (CurrentPointIndex - LastCheckpointIndex >= TimeSlicingCheckFrequency)
+		{
+			if (Context.ShouldStop())
+			{
+				break;
+			}
+			else
+			{
+				LastCheckpointIndex = CurrentPointIndex;
+			}
+		}
+	}
+
+	InCurrentPointIndex = CurrentPointIndex;
+
+	return (CurrentPointIndex == Points.Num());
+}
\ No newline at end of file
diff --git a/Engine/Plugins/Experimental/PCG/Source/PCG/Public/PCGSubgraph.h b/Engine/Plugins/Experimental/PCG/Source/PCG/Public/PCGSubgraph.h
--- a/Engine/Plugins/Experimental/PCG/Source/PCG/Public/PCGSubgraph.h
+++ b/Engine/Plugins/Experimental/PCG/Source/PCG/Public/PCGSubgraph.h
@@ -165,6 +165,8 @@ protected:
 	virtual bool ExecuteInternal(FPCGContext* Context) const override;
 	virtual bool IsPassthrough(const UPCGSettings* InSettings) const override { return true; }
 	FPCGDataCollection Input;
+
+	TArray<UPCGData*> RootedData;
 };
 
 #if UE_ENABLE_INCLUDE_ORDER_DEPRECATED_IN_5_2
