[PCG] Rename graph setting "Use Hierarchical Generation", and revert change to inspection that leaked its way in to my last review.

#rnx

diff --git a/Engine/Plugins/Experimental/PCG/Source/PCG/Private/PCGGraph.cpp b/Engine/Plugins/Experimental/PCG/Source/PCG/Private/PCGGraph.cpp
--- a/Engine/Plugins/Experimental/PCG/Source/PCG/Private/PCGGraph.cpp
+++ b/Engine/Plugins/Experimental/PCG/Source/PCG/Private/PCGGraph.cpp
@@ -932,7 +932,7 @@ void UPCGGraph::PostEditChangeProperty(FPropertyChangedEvent& PropertyChangedEve
 		OnGraphParametersChanged(EPCGGraphParameterEvent::ValueModifiedLocally, MemberPropertyName);
 	}
 	else if (PropertyName == GET_MEMBER_NAME_CHECKED(UPCGGraph, HiGenGridSize)
-		|| PropertyName == GET_MEMBER_NAME_CHECKED(UPCGGraph, bHierarchicalGenerationEnabled))
+		|| PropertyName == GET_MEMBER_NAME_CHECKED(UPCGGraph, bUseHierarchicalGeneration))
 	{
 		// The higen settings change the structure of the graph (presence or absence of links between grid levels).
 		NotifyGraphChanged(EPCGChangeType::Structural);
diff --git a/Engine/Plugins/Experimental/PCG/Source/PCG/Public/PCGGraph.h b/Engine/Plugins/Experimental/PCG/Source/PCG/Public/PCGGraph.h
--- a/Engine/Plugins/Experimental/PCG/Source/PCG/Public/PCGGraph.h
+++ b/Engine/Plugins/Experimental/PCG/Source/PCG/Public/PCGGraph.h
@@ -136,7 +136,7 @@ public:
 	// Default grid size for generation. For hierarchical generation, nodes outside of grid size graph ranges will generate on this grid.
 	EPCGHiGenGrid GetDefaultGrid() const { ensure(IsHierarchicalGenerationEnabled()); return HiGenGridSize; }
 	uint32 GetDefaultGridSize() const { ensure(IsHierarchicalGenerationEnabled()); return PCGHiGenGrid::GridToGridSize(HiGenGridSize); }
-	bool IsHierarchicalGenerationEnabled() const { return bHierarchicalGenerationEnabled; }
+	bool IsHierarchicalGenerationEnabled() const { return bUseHierarchicalGeneration; }
 
 #if WITH_EDITORONLY_DATA
 	UPROPERTY(BlueprintReadWrite, EditAnywhere, Category = AssetInfo, AssetRegistrySearchable)
@@ -266,9 +266,9 @@ protected:
 #endif // WITH_EDITOR
 
 	UPROPERTY(EditAnywhere, Category = Settings)
-	bool bHierarchicalGenerationEnabled = false;
+	bool bUseHierarchicalGeneration = false;
 
-	UPROPERTY(EditAnywhere, Category = Settings, meta = (DisplayName = "HiGen Default Grid Size", EditCondition = "bHierarchicalGenerationEnabled"))
+	UPROPERTY(EditAnywhere, Category = Settings, meta = (DisplayName = "HiGen Default Grid Size", EditCondition = "bUseHierarchicalGeneration"))
 	EPCGHiGenGrid HiGenGridSize = EPCGHiGenGrid::Grid256;
 
 	/** Execution grid size for nodes. */
diff --git a/Engine/Plugins/Experimental/PCG/Source/PCGEditor/Private/PCGEditorGraphNodeBase.cpp b/Engine/Plugins/Experimental/PCG/Source/PCGEditor/Private/PCGEditorGraphNodeBase.cpp
--- a/Engine/Plugins/Experimental/PCG/Source/PCGEditor/Private/PCGEditorGraphNodeBase.cpp
+++ b/Engine/Plugins/Experimental/PCG/Source/PCGEditor/Private/PCGEditorGraphNodeBase.cpp
@@ -344,8 +344,9 @@ EPCGChangeType UPCGEditorGraphNodeBase::UpdateGridSizeVisualisation(UPCGComponen
 	}
 
 	EPCGChangeType ChangeType = EPCGChangeType::None;
+	const uint32 InspectingGridSize = (InComponentBeingDebugged && InComponentBeingDebugged->IsLocalComponent()) ?
+		InComponentBeingDebugged->GetGenerationGridSize() : PCGHiGenGrid::UninitializedGridSize();
 
-	const uint32 InspectingGridSize = InComponentBeingDebugged ? InComponentBeingDebugged->GetGenerationGridSize() : PCGHiGenGrid::UninitializedGridSize();
 	if (InspectingGridSize == PCGHiGenGrid::UninitializedGridSize())
 	{
 		if (IsDisplayAsDisabledForced())
