Submitted on behalf of Dan Smith:
Fixed missing metadata entry initialization that was causing asserts in some instances (i.e. point data with no previously existing attributes)

#rnx

diff --git a/Engine/Plugins/Experimental/PCG/Source/PCG/Private/Elements/PCGSpawnActor.cpp b/Engine/Plugins/Experimental/PCG/Source/PCG/Private/Elements/PCGSpawnActor.cpp
--- a/Engine/Plugins/Experimental/PCG/Source/PCG/Private/Elements/PCGSpawnActor.cpp
+++ b/Engine/Plugins/Experimental/PCG/Source/PCG/Private/Elements/PCGSpawnActor.cpp
@@ -905,6 +905,7 @@ void FPCGSpawnActorElement::SpawnActors(FPCGSubgraphContext* Context, AActor* Ta
 			if (OutPointData && ActorReferenceAttribute)
 			{
 				FPCGPoint& OutPoint = OutPointData->GetMutablePoints()[i + OutPointOffset];
+				OutPointData->Metadata->InitializeOnSet(OutPoint.MetadataEntry);
 				ActorReferenceAttribute->SetValue(OutPoint.MetadataEntry, FSoftObjectPath(GeneratedActor).ToString());
 			}
 		}
