PCG static mesh spawner now respects the ISM component type decision and applies a change only in HISM+Nanite cases
Fixed an issue with the ISM component descriptor where the component class would not be valid when initialized from a static mesh component.

diff --git a/Engine/Plugins/Experimental/PCG/Source/PCG/Private/Helpers/PCGActorHelpers.cpp b/Engine/Plugins/Experimental/PCG/Source/PCG/Private/Helpers/PCGActorHelpers.cpp
--- a/Engine/Plugins/Experimental/PCG/Source/PCG/Private/Helpers/PCGActorHelpers.cpp
+++ b/Engine/Plugins/Experimental/PCG/Source/PCG/Private/Helpers/PCGActorHelpers.cpp
@@ -93,19 +93,34 @@ UPCGManagedISMComponent* UPCGActorHelpers::GetOrCreateManagedISMC(AActor* InTarg
 #endif
 
 	FString ComponentName;
-	TSubclassOf<UInstancedStaticMeshComponent> ComponentClass;
-	if (bMeshHasNaniteData)
+	TSubclassOf<UInstancedStaticMeshComponent> ComponentClass = InParams.Descriptor.ComponentClass;
+
+	// If the component class in invalid, default to HISM.
+	if (!ComponentClass)
 	{
-		ComponentClass = UInstancedStaticMeshComponent::StaticClass();
-		ComponentName += TEXT("ISM");
+		ComponentClass = UHierarchicalInstancedStaticMeshComponent::StaticClass();
 	}
-	else
+
+	// It's potentially less efficient to put nanite meshes inside of HISMs so decay those to ISM in this case.
+	// Note the equality here, not a IsA because we do not want to change derived types either
+	if (ComponentClass == UHierarchicalInstancedStaticMeshComponent::StaticClass())
 	{
-		ComponentClass = UHierarchicalInstancedStaticMeshComponent::StaticClass();
-		ComponentName += TEXT("HISM");
+		if (bMeshHasNaniteData)
+		{
+			ComponentClass = UInstancedStaticMeshComponent::StaticClass();
+		}
+		else
+		{
+			ComponentName = TEXT("HISM_");
+		}
+	}
+
+	if (ComponentClass == UInstancedStaticMeshComponent::StaticClass())
+	{
+		ComponentName = TEXT("ISM_");
 	}
 
-	ComponentName += TEXT("_") + StaticMesh->GetName();
+	ComponentName += StaticMesh->GetName();
 
 	UInstancedStaticMeshComponent* ISMC = NewObject<UInstancedStaticMeshComponent>(InTargetActor, ComponentClass, MakeUniqueObjectName(InTargetActor, ComponentClass, FName(ComponentName)));
 	InParams.Descriptor.InitComponent(ISMC);
diff --git a/Engine/Source/Runtime/Engine/Private/ISMPartition/ISMComponentDescriptor.cpp b/Engine/Source/Runtime/Engine/Private/ISMPartition/ISMComponentDescriptor.cpp
--- a/Engine/Source/Runtime/Engine/Private/ISMPartition/ISMComponentDescriptor.cpp
+++ b/Engine/Source/Runtime/Engine/Private/ISMPartition/ISMComponentDescriptor.cpp
@@ -67,8 +67,15 @@ FISMComponentDescriptor FISMComponentDescriptor::CreateFrom(const TSubclassOf<US
 
 void FISMComponentDescriptorBase::InitFrom(const UStaticMeshComponent* Template, bool bInitBodyInstance)
 {
+	check(Template);
 	bEnableDiscardOnLoad = false;
-	ComponentClass = Template->GetClass();
+
+	// Disregard the template class if it does not stem from an instanced mesh component
+	if (Template->IsA(UInstancedStaticMeshComponent::StaticClass()))
+	{
+		ComponentClass = Template->GetClass();
+	}
+
 	Mobility = Template->Mobility;
 	VirtualTextureRenderPassType = Template->VirtualTextureRenderPassType;
 	LightmapType = Template->LightmapType;
