[PCG] HiGen - support Unbounded execution domain.

Unbounded execution happens on the original component, and any resulting artifacts are stored under the original PCG actor. Data from the Unbounded execution is available for use by all grid levels.

Previously only one of either the original component or the local components were scheduled. With the Unbounded level both can execute. This change will schedule the original component in this scenario to store any data needed by local components first, and then schedule the local components.

A few other places around the code changed to support this, for example Cleanup and Clear PCG Link.

#rnx

diff --git a/Engine/Plugins/Experimental/PCG/Source/PCG/Private/Graph/PCGGraphExecutor.cpp b/Engine/Plugins/Experimental/PCG/Source/PCG/Private/Graph/PCGGraphExecutor.cpp
--- a/Engine/Plugins/Experimental/PCG/Source/PCG/Private/Graph/PCGGraphExecutor.cpp
+++ b/Engine/Plugins/Experimental/PCG/Source/PCG/Private/Graph/PCGGraphExecutor.cpp
@@ -140,11 +140,9 @@ FPCGTaskId FPCGGraphExecutor::Schedule(UPCGGraph* Graph, UPCGComponent* SourceCo
 	// component, knock out the grid sizes.
 	if (SourceComponent && SourceComponent->GetGraph() && SourceComponent->GetGraph()->IsHierarchicalGenerationEnabled())
 	{
-		const int32 ComponentGridSize = SourceComponent->GetGenerationGridSize();
-		const int32 DefaultGridSize = SourceComponent->GetGraph()->GetDefaultGridSize();
-		const EPCGHiGenGrid DefaultGrid = PCGHiGenGrid::GridSizeToGrid(DefaultGridSize);
-		const EPCGHiGenGrid Grid = PCGHiGenGrid::GridSizeToGrid(ensure(PCGHiGenGrid::IsValidGridSize(ComponentGridSize)) ? ComponentGridSize : DefaultGridSize);
 		const bool bNonPartitionedComponent = !SourceComponent->IsLocalComponent() && !SourceComponent->IsPartitioned();
+		const EPCGHiGenGrid Grid = SourceComponent->GetGenerationGrid();
+		const EPCGHiGenGrid DefaultGrid = PCGHiGenGrid::GridSizeToGrid(SourceComponent->GetGraph()->GetDefaultGridSize());
 
 		for (FPCGGraphTask& Task : CompiledTasks)
 		{
@@ -639,12 +637,9 @@ void FPCGGraphExecutor::Execute()
 				if (Task.GraphGenerationGrid != EPCGHiGenGrid::Uninitialized)
 				{
 					// Graph generation sizes should be resolved by now
-					ensure(PCGHiGenGrid::IsValidGrid(Task.GraphGenerationGrid) || Task.GraphGenerationGrid == EPCGHiGenGrid::Unbounded);
+					ensure(PCGHiGenGrid::IsValidGrid(Task.GraphGenerationGrid & ~EPCGHiGenGrid::Unbounded) || Task.GraphGenerationGrid == EPCGHiGenGrid::Unbounded);
 
-					const bool bGraphHasNoGridSizes = Task.GraphGenerationGrid == EPCGHiGenGrid::Unbounded;
-					const bool bAtFromOrToGridSize = !!(Task.GenerationGrid & Task.GraphGenerationGrid);
-					const bool bInRange = !bGraphHasNoGridSizes && PCGHiGenGrid::IsValidGrid(Task.GenerationGrid) && bAtFromOrToGridSize;
-					bNeedsToCreateActiveTask = bGraphHasNoGridSizes || bInRange;
+					bNeedsToCreateActiveTask = !!(Task.GenerationGrid & Task.GraphGenerationGrid);
 
 					if (bGraphCacheDebuggingEnabled && !bNeedsToCreateActiveTask && Task.SourceComponent.Get() && Task.SourceComponent->GetOwner() && Task.Node)
 					{
@@ -1495,8 +1490,8 @@ namespace PCGGraphExecutor
 {
 	bool ExecuteGridLinkage(EPCGHiGenGrid InFromGrid, EPCGHiGenGrid InToGrid, const FString& InResourceKey, const FName& InOutputPinLabel, const UPCGNode* InDownstreamNode, FPCGGridLinkageContext* InContext)
 	{
-		// Non partitioned graph - no linkage required - data should just pass through.
-		if (!InContext->SourceComponent->IsLocalComponent() || !InContext->SourceComponent->GetGraph()->IsHierarchicalGenerationEnabled())
+		// Non-hierarchical generation - no linkage required - data should just pass through.
+		if (!InContext->SourceComponent->GetGraph()->IsHierarchicalGenerationEnabled())
 		{
 			InContext->OutputData = InContext->InputData;
 			return true;
@@ -1509,10 +1504,8 @@ namespace PCGGraphExecutor
 			FromGrid = InContext->SourceComponent->GetGraph()->GetDefaultGrid();
 		}
 
-		ensure(PCGHiGenGrid::IsValidGrid(FromGrid));
-		const int32 FromGridSize = PCGHiGenGrid::GridToGridSize(FromGrid);
-		ensure(PCGHiGenGrid::IsValidGrid(InToGrid));
-		const int32 ToGridSize = PCGHiGenGrid::GridToGridSize(InToGrid);
+		const uint32 FromGridSize = PCGHiGenGrid::IsValidGrid(FromGrid) ? PCGHiGenGrid::GridToGridSize(FromGrid) : PCGHiGenGrid::UnboundedGridSize();
+		const uint32 ToGridSize = PCGHiGenGrid::IsValidGrid(InToGrid) ? PCGHiGenGrid::GridToGridSize(InToGrid) : PCGHiGenGrid::UnboundedGridSize();
 
 		// Never allow a large grid to read data from small grid - this violates hierarchy.
 		if (FromGridSize < ToGridSize)
@@ -1522,12 +1515,28 @@ namespace PCGGraphExecutor
 			{
 				// Using the low level logging call because we have only a node pointer for the downstream node. Note that InContext
 				// is the context for the linkage element/task, which is not represented on the graph and cannot receive graph warnings/errors.
-				Subsystem->GetNodeVisualLogsMutable().Log(
-					InDownstreamNode, InContext->SourceComponent, ELogVerbosity::Error,
-					FText::Format(NSLOCTEXT("PCGGraphCompiler", "InvalidLinkage", "Could not read data across grid levels - origin grid size {0} must be greater than destination grid size {1}. Graph default grid size may need increasing."),
-						FromGridSize, ToGridSize
-					));
-
+				if (ToGridSize == PCGHiGenGrid::UnboundedGridSize())
+				{
+					Subsystem->GetNodeVisualLogsMutable().Log(
+						InDownstreamNode,
+						InContext->SourceComponent,
+						ELogVerbosity::Error,
+						FText::Format(
+							NSLOCTEXT("PCGGraphCompiler", "InvalidLinkageToUnbounded", "Could not read data across grid levels - cannot read from grid size {0} to Unbounded domain."),
+							FromGridSize,
+							ToGridSize));
+				}
+				else
+				{
+					Subsystem->GetNodeVisualLogsMutable().Log(
+						InDownstreamNode,
+						InContext->SourceComponent,
+						ELogVerbosity::Error,
+						FText::Format(
+							NSLOCTEXT("PCGGraphCompiler", "InvalidLinkageInvalidGridSizes", "Could not read data across grid levels - origin grid size {0} must be greater than destination grid size {1}. Graph default grid size may need increasing."),
+							FromGridSize,
+							ToGridSize));
+				}
 			}
 #endif
 
@@ -1552,10 +1561,18 @@ namespace PCGGraphExecutor
 				return false;
 			}
 
-			UPCGComponent* ComponentWithData = nullptr;
 			APCGPartitionActor* PartitionActor = Cast<APCGPartitionActor>(InContext->SourceComponent->GetOwner());
 			UPCGComponent* OriginalComponent = PartitionActor ? PartitionActor->GetOriginalComponent(InContext->SourceComponent.Get()) : nullptr;
-			if (OriginalComponent)
+
+			UPCGComponent* ComponentWithData = nullptr;
+			bool bComponentWithDataIsOriginalComponent = false;
+			if (FromGridSize == PCGHiGenGrid::UnboundedGridSize())
+			{
+				// Unbounded means not generated on grid, instead use the original component volume
+				ComponentWithData = OriginalComponent;
+				bComponentWithDataIsOriginalComponent = true;
+			}
+			else if (OriginalComponent)
 			{
 				// Now iterate through local components and find a component from the correct from-grid, and overlapping this component.
 				// Would be more efficient if we had a spatial index. The current component octree seems to only store original components.
@@ -1601,7 +1618,11 @@ namespace PCGGraphExecutor
 				return true;
 			};
 
-			if (ComponentWithData->IsGenerating())
+			// If we need data from a local component but the local component is still generating, then we'll wait for it. On the other hand
+			// if we need data from the original component we assume the generation has already happened because it is always scheduled before
+			// the local components.
+			const bool bWaitForGeneration = bComponentWithDataIsOriginalComponent ? false : ComponentWithData->IsGenerating();
+			if (bWaitForGeneration)
 			{
 				PCGGraphExecutionLogging::LogGridLinkageTaskExecuteRetrieveWaitOnScheduledGraph(InContext, ComponentWithData, InResourceKey);
 
@@ -1620,7 +1641,8 @@ namespace PCGGraphExecutor
 			}
 
 			// Graph is not currently generating. If we have not already tried generating, try it once now.
-			if (!InContext->bScheduledGraph)
+			// But don't do this for the original component as this will always be scheduled before the local components.
+			if (!InContext->bScheduledGraph && !bComponentWithDataIsOriginalComponent)
 			{
 				PCGGraphExecutionLogging::LogGridLinkageTaskExecuteRetrieveScheduleGraph(InContext, ComponentWithData, InResourceKey);
 
diff --git a/Engine/Plugins/Experimental/PCG/Source/PCG/Private/Helpers/PCGHelpers.cpp b/Engine/Plugins/Experimental/PCG/Source/PCG/Private/Helpers/PCGHelpers.cpp
--- a/Engine/Plugins/Experimental/PCG/Source/PCG/Private/Helpers/PCGHelpers.cpp
+++ b/Engine/Plugins/Experimental/PCG/Source/PCG/Private/Helpers/PCGHelpers.cpp
@@ -408,11 +408,13 @@ namespace PCGHelpers
 		return CurrentInspectedObject && !CurrentInspectedObject->HasAnyFlags(RF_ClassDefaultObject | RF_NeedLoad | RF_NeedPostLoad);
 	}
 
-	bool GetGenerationGridSizes(const UPCGGraph* InGraph, const APCGWorldActor* InWorldActor, PCGHiGenGrid::FSizeArray& OutGridSizes)
+	bool GetGenerationGridSizes(const UPCGGraph* InGraph, const APCGWorldActor* InWorldActor, PCGHiGenGrid::FSizeArray& OutGridSizes, bool& bOutHasUnbounded)
 	{
+		bOutHasUnbounded = false;
+
 		if (InGraph && InGraph->IsHierarchicalGenerationEnabled())
 		{
-			InGraph->GetGridSizes(OutGridSizes);
+			InGraph->GetGridSizes(OutGridSizes, bOutHasUnbounded);
 			return true;
 		}
 		else if (InWorldActor)
diff --git a/Engine/Plugins/Experimental/PCG/Source/PCG/Private/PCGActorAndComponentMapping.cpp b/Engine/Plugins/Experimental/PCG/Source/PCG/Private/PCGActorAndComponentMapping.cpp
--- a/Engine/Plugins/Experimental/PCG/Source/PCG/Private/PCGActorAndComponentMapping.cpp
+++ b/Engine/Plugins/Experimental/PCG/Source/PCG/Private/PCGActorAndComponentMapping.cpp
@@ -140,8 +140,9 @@ bool UPCGActorAndComponentMapping::RegisterOrUpdatePartitionedPCGComponent(UPCGC
 	// sure to update the runtime flow, for them to also create PA if they need to.
 	if (bComponentHasChanged || bComponentWasAdded)
 	{
+		bool bHasUnbounded = false;
 		PCGHiGenGrid::FSizeArray GridSizes;
-		ensure(PCGHelpers::GetGenerationGridSizes(InComponent ? InComponent->GetGraph() : nullptr, PCGSubsystem->GetPCGWorldActor(), GridSizes));
+		ensure(PCGHelpers::GetGenerationGridSizes(InComponent ? InComponent->GetGraph() : nullptr, PCGSubsystem->GetPCGWorldActor(), GridSizes, bHasUnbounded));
 		PCGSubsystem->CreatePartitionActorsWithinBounds(Bounds, GridSizes);
 	}
 #endif // WITH_EDITOR
diff --git a/Engine/Plugins/Experimental/PCG/Source/PCG/Private/PCGCommon.cpp b/Engine/Plugins/Experimental/PCG/Source/PCG/Private/PCGCommon.cpp
--- a/Engine/Plugins/Experimental/PCG/Source/PCG/Private/PCGCommon.cpp
+++ b/Engine/Plugins/Experimental/PCG/Source/PCG/Private/PCGCommon.cpp
@@ -40,10 +40,22 @@ namespace PCGHiGenGrid
 
 	EPCGHiGenGrid GridSizeToGrid(uint32 InGridSize)
 	{
+		if (InGridSize == UnboundedGridSize())
+		{
+			return EPCGHiGenGrid::Unbounded;
+		}
 		// TODO: support other units
 		return ensure(IsValidGridSize(InGridSize)) ? static_cast<EPCGHiGenGrid>(InGridSize / 100) : EPCGHiGenGrid::Uninitialized;
 	}
 
-	uint32 UninitializedGridSize() { return static_cast<uint32>(EPCGHiGenGrid::Uninitialized); }
-	uint32 UnboundedGridSize() { return static_cast<uint32>(EPCGHiGenGrid::Unbounded); }
+	uint32 UninitializedGridSize()
+	{
+		return static_cast<uint32>(EPCGHiGenGrid::Uninitialized);
+	}
+
+	uint32 UnboundedGridSize()
+	{
+		// TODO: support other units
+		return 100 * static_cast<uint32>(EPCGHiGenGrid::Unbounded);
+	}
 }
diff --git a/Engine/Plugins/Experimental/PCG/Source/PCG/Private/PCGComponent.cpp b/Engine/Plugins/Experimental/PCG/Source/PCG/Private/PCGComponent.cpp
--- a/Engine/Plugins/Experimental/PCG/Source/PCG/Private/PCGComponent.cpp
+++ b/Engine/Plugins/Experimental/PCG/Source/PCG/Private/PCGComponent.cpp
@@ -653,6 +653,23 @@ AActor* UPCGComponent::ClearPCGLink(UClass* TemplateActor)
 	return NewActor;
 }
 
+EPCGHiGenGrid UPCGComponent::GetGenerationGrid() const
+{
+	const uint32 GridSize = GetGenerationGridSize();
+	if (PCGHiGenGrid::IsValidGridSize(GridSize))
+	{
+		return PCGHiGenGrid::GridSizeToGrid(GridSize);
+	}
+	else if (GridSize == PCGHiGenGrid::UnboundedGridSize())
+	{
+		return EPCGHiGenGrid::Unbounded;
+	}
+	else
+	{
+		return PCGHiGenGrid::GridSizeToGrid(GetGraph()->GetDefaultGridSize());
+	}
+}
+
 void UPCGComponent::StoreOutputDataForPin(const FString& InResourceKey, const FPCGDataCollection& InData)
 {
 	FReadScopeLock ScopedWriteLock(PerPinGeneratedOutputLock);
@@ -774,7 +791,12 @@ void UPCGComponent::CleanupLocalImmediate(bool bRemoveComponents)
 
 FPCGTaskId UPCGComponent::CreateCleanupTask(bool bRemoveComponents, const TArray<FPCGTaskId>& Dependencies)
 {
-	if ((!bGenerated && GeneratedResources.IsEmpty() && !IsGenerating()) || IsPartitioned() || IsCleaningUp())
+	if (GetSubsystem() && GetSubsystem()->IsGraphCacheDebuggingEnabled())
+	{
+		UE_LOG(LogPCG, Log, TEXT("[%s] --- CLEANUP COMPONENT ---"), GetOwner() ? *GetOwner()->GetName() : TEXT("MissingComponent"));
+	}
+
+	if ((!bGenerated && GeneratedResources.IsEmpty() && !IsGenerating()) || IsCleaningUp())
 	{
 		return InvalidPCGTaskId;
 	}
@@ -1841,9 +1863,9 @@ void UPCGComponent::OnRefresh()
 	}
 	else
 	{
-		// If we just cleaned up resources, call back generate
-		// Also, for BPs, we ask if we should generate, to support generate on added to world.
-		if ((bWasGenerated || ShouldGenerateBPPCGAddedToWorld()) && (!bGenerated || bRegenerateInEditor))
+		// If we just cleaned up resources, call back generate. Only do this for original component, which will then trigger
+		// generation of local components. Also, for BPs, we ask if we should generate, to support generate on added to world.
+		if ((bWasGenerated || ShouldGenerateBPPCGAddedToWorld()) && !IsLocalComponent() && (!bGenerated || bRegenerateInEditor))
 		{
 			GenerateLocal(/*bForce=*/false);
 		}
diff --git a/Engine/Plugins/Experimental/PCG/Source/PCG/Private/PCGGraph.cpp b/Engine/Plugins/Experimental/PCG/Source/PCG/Private/PCGGraph.cpp
--- a/Engine/Plugins/Experimental/PCG/Source/PCG/Private/PCGGraph.cpp
+++ b/Engine/Plugins/Experimental/PCG/Source/PCG/Private/PCGGraph.cpp
@@ -342,6 +342,12 @@ void UPCGGraph::BeginDestroy()
 	Super::BeginDestroy();
 }
 
+uint32 UPCGGraph::GetDefaultGridSize() const
+{
+	ensure(IsHierarchicalGenerationEnabled());
+	return PCGHiGenGrid::IsValidGrid(HiGenGridSize) ? PCGHiGenGrid::GridToGridSize(HiGenGridSize) : PCGHiGenGrid::UnboundedGridSize();
+}
+
 UPCGNode* UPCGGraph::AddNodeOfType(TSubclassOf<class UPCGSettings> InSettingsClass, UPCGSettings*& OutDefaultNodeSettings)
 {
 	UPCGSettings* Settings = NewObject<UPCGSettings>(GetTransientPackage(), InSettingsClass, NAME_None, RF_Transactional);
@@ -690,12 +696,17 @@ void UPCGGraph::PostNodeUndo(UPCGNode* InPCGNode)
 }
 #endif
 
-void UPCGGraph::GetGridSizes(PCGHiGenGrid::FSizeArray& OutGridSizes) const
+void UPCGGraph::GetGridSizes(PCGHiGenGrid::FSizeArray& OutGridSizes, bool& bOutHasUnbounded) const
 {
+	bOutHasUnbounded = HiGenGridSize == EPCGHiGenGrid::Unbounded;
+
 	const uint32 GraphDefaultGridSize = GetDefaultGridSize();
 	if (!IsHierarchicalGenerationEnabled())
 	{
-		OutGridSizes.Add(GraphDefaultGridSize);
+		if (PCGHiGenGrid::IsValidGridSize(GetDefaultGridSize()))
+		{
+			OutGridSizes.Add(GraphDefaultGridSize);
+		}
 		return;
 	}
 
@@ -710,6 +721,10 @@ void UPCGGraph::GetGridSizes(PCGHiGenGrid::FSizeArray& OutGridSizes) const
 				OutGridSizes.Add(GridSize);
 			}
 		}
+		else if (GridSize == PCGHiGenGrid::UnboundedGridSize())
+		{
+			bOutHasUnbounded = true;
+		}
 		else if (GridSize == PCGHiGenGrid::UninitializedGridSize())
 		{
 			// Outside nodes will not have a concrete grid set
diff --git a/Engine/Plugins/Experimental/PCG/Source/PCG/Private/PCGSubsystem.cpp b/Engine/Plugins/Experimental/PCG/Source/PCG/Private/PCGSubsystem.cpp
--- a/Engine/Plugins/Experimental/PCG/Source/PCG/Private/PCGSubsystem.cpp
+++ b/Engine/Plugins/Experimental/PCG/Source/PCG/Private/PCGSubsystem.cpp
@@ -263,7 +263,7 @@ UPCGLandscapeCache* UPCGSubsystem::GetLandscapeCache()
 	return LandscapeCacheOwner ? LandscapeCacheOwner->LandscapeCacheObject.Get() : nullptr;
 }
 
-FPCGTaskId UPCGSubsystem::ScheduleComponent(UPCGComponent* PCGComponent, bool bSave, const TArray<FPCGTaskId>& Dependencies)
+FPCGTaskId UPCGSubsystem::ScheduleComponent(UPCGComponent* PCGComponent, bool bSave, const TArray<FPCGTaskId>& InDependencies)
 {
 	check(GraphExecutor);
 
@@ -272,16 +272,21 @@ FPCGTaskId UPCGSubsystem::ScheduleComponent(UPCGComponent* PCGComponent, bool bS
 		return InvalidPCGTaskId;
 	}
 
+	bool bHasUnbounded = false;
+	PCGHiGenGrid::FSizeArray GridSizes;
+	ensure(PCGHelpers::GetGenerationGridSizes(PCGComponent->GetGraph(), GetPCGWorldActor(), GridSizes, bHasUnbounded));
+
 #if WITH_EDITOR
 	if (PCGComponent->IsPartitioned() && !PCGHelpers::IsRuntimeOrPIE())
 	{
-		// In this case create the PA and update the mapping
-		// Note: This is an immediate operation, as we need the PA for the generation.
-		auto ScheduleTask = [](APCGPartitionActor* PCGActor, const FBox& InIntersectedBounds, const TArray<FPCGTaskId>& TaskDependencies) { return InvalidPCGTaskId; };
+		if (!GridSizes.IsEmpty())
+		{
+			// In this case create the PA and update the mapping
+			// Note: This is an immediate operation, as we need the PA for the generation.
+			auto ScheduleTask = [](APCGPartitionActor* PCGActor, const FBox& InIntersectedBounds, const TArray<FPCGTaskId>& TaskDependencies) { return InvalidPCGTaskId; };
 
-		PCGHiGenGrid::FSizeArray GridSizes;
-		ensure(PCGHelpers::GetGenerationGridSizes(PCGComponent->GetGraph(), GetPCGWorldActor(), GridSizes));
-		PCGSubsystem::ForEachIntersectingCell(GraphExecutor, PCGComponent->GetWorld(), PCGComponent->GetGridBounds(), GridSizes, /*bCreateActor=*/true, /*bLoadCell=*/false, /*bSave=*/false, ScheduleTask);
+			PCGSubsystem::ForEachIntersectingCell(GraphExecutor, PCGComponent->GetWorld(), PCGComponent->GetGridBounds(), GridSizes, /*bCreateActor=*/true, /*bLoadCell=*/false, /*bSave=*/false, ScheduleTask);
+		}
 
 		ActorAndComponentMapping.UpdateMappingPCGComponentPartitionActor(PCGComponent);
 	}
@@ -289,9 +294,27 @@ FPCGTaskId UPCGSubsystem::ScheduleComponent(UPCGComponent* PCGComponent, bool bS
 
 	TArray<FPCGTaskId> AllTasks;
 
+	// Schedule generation of original component if is is non-partitioned, or if it has nodes that will execute at the Unbounded level.
+	FPCGTaskId OriginalComponentTask = InvalidPCGTaskId;
+	if (!PCGComponent->IsPartitioned() || bHasUnbounded)
+	{
+		OriginalComponentTask = PCGComponent->CreateGenerateTask(/*bForce=*/bSave, InDependencies);
+		if (OriginalComponentTask != InvalidPCGTaskId)
+		{
+			AllTasks.Add(OriginalComponentTask);
+		}
+	}
+
 	// If the component is partitioned, we will forward the calls to its registered PCG Partition actors
 	if (PCGComponent->IsPartitioned())
 	{
+		// Local components depend on the original component (to ensure any data is available).
+		TArray<FPCGTaskId> Dependencies = InDependencies;
+		if (OriginalComponentTask != InvalidPCGTaskId)
+		{
+			Dependencies.Add(OriginalComponentTask);
+		}
+
 		auto LocalGenerateTask = [OriginalComponent = PCGComponent, &Dependencies](UPCGComponent* LocalComponent)
 		{
 			// If the local component is currently generating, it's probably because it was requested by a refresh.
@@ -314,15 +337,7 @@ FPCGTaskId UPCGSubsystem::ScheduleComponent(UPCGComponent* PCGComponent, bool bS
 			return LocalComponent->GenerateInternal(/*bForce=*/ false, EPCGComponentGenerationTrigger::GenerateOnDemand, Dependencies);
 		};
 
-		AllTasks = ActorAndComponentMapping.DispatchToRegisteredLocalComponents(PCGComponent, LocalGenerateTask);
-	}
-	else
-	{
-		FPCGTaskId TaskId = PCGComponent->CreateGenerateTask(/*bForce=*/bSave, Dependencies);
-		if (TaskId != InvalidPCGTaskId)
-		{
-			AllTasks.Add(TaskId);
-		}
+		AllTasks.Append(ActorAndComponentMapping.DispatchToRegisteredLocalComponents(PCGComponent, LocalGenerateTask));
 	}
 
 	if (!AllTasks.IsEmpty())
@@ -363,8 +378,22 @@ FPCGTaskId UPCGSubsystem::ScheduleCleanup(UPCGComponent* PCGComponent, bool bRem
 		return InvalidPCGTaskId;
 	}
 
+	bool bHasUnbounded = false;
+	PCGHiGenGrid::FSizeArray GridSizes;
+	ensure(PCGHelpers::GetGenerationGridSizes(PCGComponent->GetGraph(), GetPCGWorldActor(), GridSizes, bHasUnbounded));
+
 	TArray<FPCGTaskId> AllTasks;
 
+	// Schedule cleanup of original component if is is non-partitioned, or if it has nodes that will execute at the Unbounded level.
+	if (!PCGComponent->IsPartitioned() || bHasUnbounded)
+	{
+		FPCGTaskId TaskId = PCGComponent->CreateCleanupTask(bRemoveComponents, Dependencies);
+		if (TaskId != InvalidPCGTaskId)
+		{
+			AllTasks.Add(TaskId);
+		}
+	}
+
 	// If the component is partitioned, we will forward the calls to its registered PCG Partition actors
 	if (PCGComponent->IsPartitioned())
 	{
@@ -380,15 +409,7 @@ FPCGTaskId UPCGSubsystem::ScheduleCleanup(UPCGComponent* PCGComponent, bool bRem
 			return LocalComponent->CleanupInternal(bRemoveComponents, /*bSave=*/ false, Dependencies);
 		};
 
-		AllTasks = ActorAndComponentMapping.DispatchToRegisteredLocalComponents(PCGComponent, LocalCleanupTask);
-	}
-	else
-	{
-		FPCGTaskId TaskId = PCGComponent->CreateCleanupTask(bRemoveComponents, Dependencies);
-		if (TaskId != InvalidPCGTaskId)
-		{
-			AllTasks.Add(TaskId);
-		}
+		AllTasks.Append(ActorAndComponentMapping.DispatchToRegisteredLocalComponents(PCGComponent, LocalCleanupTask));
 	}
 
 	TWeakObjectPtr<UPCGComponent> ComponentPtr(PCGComponent);
@@ -878,8 +899,9 @@ FPCGTaskId UPCGSubsystem::ProcessGraph(UPCGComponent* Component, const FBox& InP
 
 	if (UnionBounds.IsValid)
 	{
+		bool bHasUnbounded = false;
 		PCGHiGenGrid::FSizeArray GridSizes;
-		ensure(PCGHelpers::GetGenerationGridSizes(Component->GetGraph(), GetPCGWorldActor(), GridSizes));
+		ensure(PCGHelpers::GetGenerationGridSizes(Component->GetGraph(), GetPCGWorldActor(), GridSizes, bHasUnbounded));
 		ProcessAllCellsTaskId = PCGSubsystem::ForEachIntersectingCell(GraphExecutor, Component->GetWorld(), UnionBounds, GridSizes, bCreateActors, bLoadCell, bSaveActors, ScheduleTask);
 	}
 
@@ -947,9 +969,23 @@ FPCGTaskId UPCGSubsystem::CleanupGraph(UPCGComponent* Component, const FBox& InB
 		return InvalidPCGTaskId;
 	};
 
+	bool bHasUnbounded = false;
 	PCGHiGenGrid::FSizeArray GridSizes;
-	ensure(PCGHelpers::GetGenerationGridSizes(Component->GetGraph(), GetPCGWorldActor(), GridSizes));
-	FPCGTaskId ProcessAllCellsTaskId = PCGSubsystem::ForEachIntersectingCell(GraphExecutor, Component->GetWorld(), InBounds, GridSizes, /*bCreateActor=*/false, /*bLoadCell=*/false, bSave, ScheduleTask);
+	ensure(PCGHelpers::GetGenerationGridSizes(Component->GetGraph(), GetPCGWorldActor(), GridSizes, bHasUnbounded));
+
+	FPCGTaskId ProcessAllCellsTaskId = InvalidPCGTaskId;
+	if (!GridSizes.IsEmpty())
+	{
+		ProcessAllCellsTaskId = PCGSubsystem::ForEachIntersectingCell(
+			GraphExecutor,
+			Component->GetWorld(),
+			InBounds,
+			GridSizes,
+			/*bCreateActor=*/false,
+			/*bLoadCell=*/false,
+			bSave,
+			ScheduleTask);
+	}
 
 	// Finally, call PostCleanupGraph if something happened
 	if (ProcessAllCellsTaskId != InvalidPCGTaskId)
@@ -1025,9 +1061,15 @@ void UPCGSubsystem::ClearPCGLink(UPCGComponent* InComponent, const FBox& InBound
 		return GraphExecutor->ScheduleGeneric(MoveTask, ComponentPtr.Get(), TaskDependencies);
 	};
 
+	bool bHasUnbounded = false;
 	PCGHiGenGrid::FSizeArray GridSizes;
-	ensure(PCGHelpers::GetGenerationGridSizes(InComponent->GetGraph(), GetPCGWorldActor(), GridSizes));
-	FPCGTaskId TaskId = PCGSubsystem::ForEachIntersectingCell(GraphExecutor, GetWorld(), InBounds, GridSizes, /*bCreateActor=*/false, /*bLoadCell=*/false, /*bSave=*/false, ScheduleTask);
+	ensure(PCGHelpers::GetGenerationGridSizes(InComponent->GetGraph(), GetPCGWorldActor(), GridSizes, bHasUnbounded));
+
+	FPCGTaskId TaskId = InvalidPCGTaskId;
+	if (!GridSizes.IsEmpty())
+	{
+		TaskId = PCGSubsystem::ForEachIntersectingCell(GraphExecutor, GetWorld(), InBounds, GridSizes, /*bCreateActor=*/false, /*bLoadCell=*/false, /*bSave=*/false, ScheduleTask);
+	}
 
 	// Verify if the NewActor has some components attached to its root or attached actors. If not, destroy it.
 	// Return false if the new actor is not valid or destroyed.
diff --git a/Engine/Plugins/Experimental/PCG/Source/PCG/Public/Elements/PCGHiGenGridSize.h b/Engine/Plugins/Experimental/PCG/Source/PCG/Public/Elements/PCGHiGenGridSize.h
--- a/Engine/Plugins/Experimental/PCG/Source/PCG/Public/Elements/PCGHiGenGridSize.h
+++ b/Engine/Plugins/Experimental/PCG/Source/PCG/Public/Elements/PCGHiGenGridSize.h
@@ -16,7 +16,7 @@ class PCG_API UPCGHiGenGridSizeSettings : public UPCGSettings
 
 public:
 	EPCGHiGenGrid GetGrid() const { return HiGenGridSize; }
-	uint32 GetGridSize() const { return PCGHiGenGrid::GridToGridSize(HiGenGridSize); }
+	uint32 GetGridSize() const { return (HiGenGridSize == EPCGHiGenGrid::Unbounded) ? PCGHiGenGrid::UnboundedGridSize() : PCGHiGenGrid::GridToGridSize(HiGenGridSize); }
 
 	//~Begin UPCGSettings interface
 #if WITH_EDITOR
diff --git a/Engine/Plugins/Experimental/PCG/Source/PCG/Public/Helpers/PCGHelpers.h b/Engine/Plugins/Experimental/PCG/Source/PCG/Public/Helpers/PCGHelpers.h
--- a/Engine/Plugins/Experimental/PCG/Source/PCG/Public/Helpers/PCGHelpers.h
+++ b/Engine/Plugins/Experimental/PCG/Source/PCG/Public/Helpers/PCGHelpers.h
@@ -62,5 +62,5 @@ namespace PCGHelpers
 	PCG_API bool IsNewObjectAndNotDefault(const UObject* InObject, bool bCheckHierarchy = false);
 
 	/** If hierarchical generation is enabled, returns all relevant grid sizes for the graph, otherwise returns partition grid size from world actor. */
-	PCG_API bool GetGenerationGridSizes(const UPCGGraph* InGraph, const APCGWorldActor* InWorldActor, PCGHiGenGrid::FSizeArray& OutGridSizes);
+	PCG_API bool GetGenerationGridSizes(const UPCGGraph* InGraph, const APCGWorldActor* InWorldActor, PCGHiGenGrid::FSizeArray& OutGridSizes, bool& bOutHasUnbounded);
 };
diff --git a/Engine/Plugins/Experimental/PCG/Source/PCG/Public/PCGComponent.h b/Engine/Plugins/Experimental/PCG/Source/PCG/Public/PCGComponent.h
--- a/Engine/Plugins/Experimental/PCG/Source/PCG/Public/PCGComponent.h
+++ b/Engine/Plugins/Experimental/PCG/Source/PCG/Public/PCGComponent.h
@@ -162,6 +162,7 @@ public:
 
 	uint32 GetGenerationGridSize() const { return GenerationGridSize; }
 	void SetGenerationGridSize(uint32 InGenerationGridSize) { GenerationGridSize = InGenerationGridSize; }
+	EPCGHiGenGrid GetGenerationGrid() const;
 
 	/** Store data with a resource key that identifies the pin. */
 	void StoreOutputDataForPin(const FString& InResourceKey, const FPCGDataCollection& InData);
diff --git a/Engine/Plugins/Experimental/PCG/Source/PCG/Public/PCGGraph.h b/Engine/Plugins/Experimental/PCG/Source/PCG/Public/PCGGraph.h
--- a/Engine/Plugins/Experimental/PCG/Source/PCG/Public/PCGGraph.h
+++ b/Engine/Plugins/Experimental/PCG/Source/PCG/Public/PCGGraph.h
@@ -135,7 +135,7 @@ public:
 
 	// Default grid size for generation. For hierarchical generation, nodes outside of grid size graph ranges will generate on this grid.
 	EPCGHiGenGrid GetDefaultGrid() const { ensure(IsHierarchicalGenerationEnabled()); return HiGenGridSize; }
-	uint32 GetDefaultGridSize() const { ensure(IsHierarchicalGenerationEnabled()); return PCGHiGenGrid::GridToGridSize(HiGenGridSize); }
+	uint32 GetDefaultGridSize() const;
 	bool IsHierarchicalGenerationEnabled() const { return bUseHierarchicalGeneration; }
 
 #if WITH_EDITORONLY_DATA
@@ -204,7 +204,7 @@ public:
 	bool RemoveOutboundEdges(UPCGNode* InNode, const FName& OutboundLabel);
 
 	/** Determine the relevant grid sizes by inspecting all HiGenGridSize nodes. */
-	void GetGridSizes(PCGHiGenGrid::FSizeArray& OutGridSizes) const;
+	void GetGridSizes(PCGHiGenGrid::FSizeArray& OutGridSizes, bool& bOutHasUnbounded) const;
 
 #if WITH_EDITOR
 	void DisableNotificationsForEditor();
diff --git a/Engine/Plugins/Experimental/PCG/Source/PCG/Public/PCGSubsystem.h b/Engine/Plugins/Experimental/PCG/Source/PCG/Public/PCGSubsystem.h
--- a/Engine/Plugins/Experimental/PCG/Source/PCG/Public/PCGSubsystem.h
+++ b/Engine/Plugins/Experimental/PCG/Source/PCG/Public/PCGSubsystem.h
@@ -82,7 +82,7 @@ public:
 	UPCGLandscapeCache* GetLandscapeCache();
 
 	// Schedule graph (owner -> graph)
-	FPCGTaskId ScheduleComponent(UPCGComponent* PCGComponent, bool bSave, const TArray<FPCGTaskId>& Dependencies);
+	FPCGTaskId ScheduleComponent(UPCGComponent* PCGComponent, bool bSave, const TArray<FPCGTaskId>& InDependencies);
 
 	/** Schedule cleanup(owner->graph). Note that in non-partitioned mode, cleanup is immediate. */
 	FPCGTaskId ScheduleCleanup(UPCGComponent* PCGComponent, bool bRemoveComponents, bool bSave, const TArray<FPCGTaskId>& Dependencies);
diff --git a/Engine/Plugins/Experimental/PCG/Source/PCGEditor/Private/PCGEditorGraphNodeBase.cpp b/Engine/Plugins/Experimental/PCG/Source/PCGEditor/Private/PCGEditorGraphNodeBase.cpp
--- a/Engine/Plugins/Experimental/PCG/Source/PCGEditor/Private/PCGEditorGraphNodeBase.cpp
+++ b/Engine/Plugins/Experimental/PCG/Source/PCGEditor/Private/PCGEditorGraphNodeBase.cpp
@@ -338,16 +338,17 @@ EPCGChangeType UPCGEditorGraphNodeBase::UpdateErrorsAndWarnings()
 
 EPCGChangeType UPCGEditorGraphNodeBase::UpdateGridSizeVisualisation(UPCGComponent* InComponentBeingDebugged)
 {
-	if (!PCGNode)
+	const UPCGGraph* Graph = PCGNode ? PCGNode->GetGraph() : nullptr;
+	if (!Graph)
 	{
 		return EPCGChangeType::None;
 	}
 
 	EPCGChangeType ChangeType = EPCGChangeType::None;
-	const uint32 InspectingGridSize = (InComponentBeingDebugged && InComponentBeingDebugged->IsLocalComponent()) ?
-		InComponentBeingDebugged->GetGenerationGridSize() : PCGHiGenGrid::UninitializedGridSize();
 
-	if (InspectingGridSize == PCGHiGenGrid::UninitializedGridSize())
+	const bool HiGenEnabled = Graph->IsHierarchicalGenerationEnabled();
+	const uint32 InspectingGridSize = InComponentBeingDebugged ? InComponentBeingDebugged->GetGenerationGridSize() : PCGHiGenGrid::UninitializedGridSize();
+	if (!HiGenEnabled || (InspectingGridSize == PCGHiGenGrid::UninitializedGridSize()))
 	{
 		if (IsDisplayAsDisabledForced())
 		{
@@ -361,7 +362,7 @@ EPCGChangeType UPCGEditorGraphNodeBase::UpdateGridSizeVisualisation(UPCGComponen
 			ChangeType |= EPCGChangeType::Cosmetic;
 		}
 	}
-	else if (UPCGGraph* Graph = PCGNode->GetGraph())
+	else
 	{
 		uint32 DefaultGridSize;
 		UPCGSubsystem* Subsystem = UPCGSubsystem::GetActiveEditorInstance();
