[PCG] Move tracking system into the Subsystem

* Now can track Tags and Classes
* Subsystem is responsible to reacting to actor changes
* Component is no more responsible to track actors.

diff --git a/Engine/Plugins/Experimental/PCG/Source/PCG/Private/Elements/PCGActorSelector.cpp b/Engine/Plugins/Experimental/PCG/Source/PCG/Private/Elements/PCGActorSelector.cpp
--- a/Engine/Plugins/Experimental/PCG/Source/PCG/Private/Elements/PCGActorSelector.cpp
+++ b/Engine/Plugins/Experimental/PCG/Source/PCG/Private/Elements/PCGActorSelector.cpp
@@ -41,7 +41,7 @@ namespace PCGActorSelector
 		case EPCGActorSelection::ByTag:
 			return[ActorSelectionTag = InSettings.ActorSelectionTag, &InFoundActors, bMultiSelect, &BoundsCheck, &SelfIgnoreCheck](AActor* Actor) -> bool
 			{
-				if (Actor->ActorHasTag(ActorSelectionTag) && BoundsCheck(Actor) && SelfIgnoreCheck(Actor))
+				if (Actor && Actor->ActorHasTag(ActorSelectionTag) && BoundsCheck(Actor) && SelfIgnoreCheck(Actor))
 				{
 					InFoundActors.Add(Actor);
 					return bMultiSelect;
@@ -53,7 +53,7 @@ namespace PCGActorSelector
 		case EPCGActorSelection::ByClass:
 			return[ActorSelectionClass = InSettings.ActorSelectionClass, &InFoundActors, bMultiSelect, &BoundsCheck, &SelfIgnoreCheck](AActor* Actor) -> bool
 			{
-				if (Actor->IsA(ActorSelectionClass) && BoundsCheck(Actor) && SelfIgnoreCheck(Actor))
+				if (Actor && Actor->IsA(ActorSelectionClass) && BoundsCheck(Actor) && SelfIgnoreCheck(Actor))
 				{
 					InFoundActors.Add(Actor);
 					return bMultiSelect;
@@ -88,9 +88,10 @@ namespace PCGActorSelector
 		}
 
 		// Early out if we have not the information necessary.
-		if (FilterRequired(Settings) &&
-			((Settings.ActorSelection == EPCGActorSelection::ByTag && Settings.ActorSelectionTag == NAME_None) ||
-			(Settings.ActorSelection == EPCGActorSelection::ByClass && !Settings.ActorSelectionClass)))
+		const bool bNoTagInfo = Settings.ActorSelection == EPCGActorSelection::ByTag && Settings.ActorSelectionTag == NAME_None;
+		const bool bNoClassInfo = Settings.ActorSelection == EPCGActorSelection::ByClass && !Settings.ActorSelectionClass;
+
+		if (FilterRequired(Settings) && (bNoTagInfo || bNoClassInfo))
 		{
 			return FoundActors;
 		}
@@ -203,6 +204,78 @@ namespace PCGActorSelector
 	}
 }
 
+FPCGActorSelectionKey::FPCGActorSelectionKey(FName InTag)
+{
+	Selection = EPCGActorSelection::ByTag;
+	Tag = InTag;
+}
+
+FPCGActorSelectionKey::FPCGActorSelectionKey(TSubclassOf<AActor> InSelectionClass)
+{
+	Selection = EPCGActorSelection::ByClass;
+	ActorSelectionClass = InSelectionClass;
+}
+
+bool FPCGActorSelectionKey::IsMatching(const AActor* InActor) const
+{
+	if (!InActor)
+	{
+		return false;
+	}
+
+	switch (Selection)
+	{
+	case EPCGActorSelection::ByTag:
+		return InActor->ActorHasTag(Tag);
+	case EPCGActorSelection::ByClass:
+		return InActor->IsA(ActorSelectionClass);
+	default:
+		return false;
+	}
+}
+
+TArray<FPCGActorSelectionKey> FPCGActorSelectionKey::GenerateAllKeysForActor(const AActor* InActor)
+{
+	TArray<FPCGActorSelectionKey> Result;
+
+	if (!InActor)
+	{
+		return Result;
+	}
+
+	Result.Reserve(InActor->Tags.Num() + 1);
+	Result.Emplace(InActor->GetClass());
+	for (FName Tag : InActor->Tags)
+	{
+		Result.Emplace(Tag);
+	}
+
+	return Result;
+}
+
+bool FPCGActorSelectionKey::operator==(const FPCGActorSelectionKey& InOther) const
+{
+	if (Selection != InOther.Selection)
+	{
+		return false;
+	}
+
+	switch (Selection)
+	{
+	case EPCGActorSelection::ByTag:
+		return Tag == InOther.Tag;
+	case EPCGActorSelection::ByClass:
+		return ActorSelectionClass == InOther.ActorSelectionClass;
+	default:
+		return false;
+	}
+}
+
+uint32 GetTypeHash(const FPCGActorSelectionKey& In)
+{
+	return HashCombine(HashCombine(GetTypeHash(In.Selection), GetTypeHash(In.Tag)), GetTypeHash(In.ActorSelectionClass));
+}
+
 #if WITH_EDITOR
 FText FPCGActorSelectorSettings::GetTaskNameSuffix() const
 {
@@ -229,4 +302,34 @@ FName FPCGActorSelectorSettings::GetTaskName(const FText& Prefix) const
 {
 	return FName(FText::Format(NSLOCTEXT("PCGActorSelectorSettings", "NodeTitleFormat", "{0} ({1})"), Prefix, GetTaskNameSuffix()).ToString());
 }
-#endif // WITH_EDITOR
\ No newline at end of file
+#endif // WITH_EDITOR
+
+FPCGActorSelectionKey FPCGActorSelectorSettings::GetAssociatedKey() const
+{
+	// If we don't look for AllWorldActors, it means we track the pcg component, which should be already picked up by the tracking system.
+	if (ActorFilter != EPCGActorFilter::AllWorldActors)
+	{
+		return FPCGActorSelectionKey();
+	}
+
+	switch (ActorSelection)
+	{
+	case EPCGActorSelection::ByTag:
+		return FPCGActorSelectionKey(ActorSelectionTag);
+	case EPCGActorSelection::ByClass:
+		return FPCGActorSelectionKey(ActorSelectionClass);
+	default:
+		return FPCGActorSelectionKey();
+	}
+}
+
+FPCGActorSelectorSettings FPCGActorSelectorSettings::ReconstructFromKey(const FPCGActorSelectionKey& InKey)
+{
+	FPCGActorSelectorSettings Result{};
+	Result.ActorFilter = EPCGActorFilter::AllWorldActors;
+	Result.ActorSelection = InKey.Selection;
+	Result.ActorSelectionTag = InKey.Tag;
+	Result.ActorSelectionClass = InKey.ActorSelectionClass;
+
+	return Result;
+}
\ No newline at end of file
diff --git a/Engine/Plugins/Experimental/PCG/Source/PCG/Private/Elements/PCGDataFromActor.cpp b/Engine/Plugins/Experimental/PCG/Source/PCG/Private/Elements/PCGDataFromActor.cpp
--- a/Engine/Plugins/Experimental/PCG/Source/PCG/Private/Elements/PCGDataFromActor.cpp
+++ b/Engine/Plugins/Experimental/PCG/Source/PCG/Private/Elements/PCGDataFromActor.cpp
@@ -19,12 +19,11 @@
 #define LOCTEXT_NAMESPACE "PCGDataFromActorElement"
 
 #if WITH_EDITOR
-void UPCGDataFromActorSettings::GetTrackedActorTags(FPCGTagToSettingsMap& OutTagToSettings, TArray<TObjectPtr<const UPCGGraph>>& OutVisitedGraphs) const
+void UPCGDataFromActorSettings::GetTrackedActorKeys(FPCGActorSelectionKeyToSettingsMap& OutKeysToSettings, TArray<TObjectPtr<const UPCGGraph>>& OutVisitedGraphs) const
 {
-	if (ActorSelector.ActorSelection == EPCGActorSelection::ByTag &&
-		ActorSelector.ActorFilter == EPCGActorFilter::AllWorldActors)
+	if (ActorSelector.ActorFilter == EPCGActorFilter::AllWorldActors)
 	{
-		OutTagToSettings.FindOrAdd(ActorSelector.ActorSelectionTag).Emplace({ this, bTrackActorsOnlyWithinBounds });
+		OutKeysToSettings.FindOrAdd(ActorSelector.GetAssociatedKey()).Emplace(this, bTrackActorsOnlyWithinBounds);
 	}
 }
 
diff --git a/Engine/Plugins/Experimental/PCG/Source/PCG/Private/Elements/PCGExecuteBlueprint.cpp b/Engine/Plugins/Experimental/PCG/Source/PCG/Private/Elements/PCGExecuteBlueprint.cpp
--- a/Engine/Plugins/Experimental/PCG/Source/PCG/Private/Elements/PCGExecuteBlueprint.cpp
+++ b/Engine/Plugins/Experimental/PCG/Source/PCG/Private/Elements/PCGExecuteBlueprint.cpp
@@ -501,11 +501,11 @@ EPCGSettingsType UPCGBlueprintSettings::GetType() const
 	}
 }
 
-void UPCGBlueprintSettings::GetTrackedActorTags(FPCGTagToSettingsMap& OutTagToSettings, TArray<TObjectPtr<const UPCGGraph>>& OutVisitedGraphs) const
+void UPCGBlueprintSettings::GetTrackedActorKeys(FPCGActorSelectionKeyToSettingsMap& OutKeysToSettings, TArray<TObjectPtr<const UPCGGraph>>& OutVisitedGraphs) const
 {
 	for (const FName& Tag : TrackedActorTags)
 	{
-		OutTagToSettings.FindOrAdd(Tag).Emplace({ this, bTrackActorsOnlyWithinBounds });
+		OutKeysToSettings.FindOrAdd(FPCGActorSelectionKey(Tag)).Emplace(this, bTrackActorsOnlyWithinBounds);
 	}
 }
 
diff --git a/Engine/Plugins/Experimental/PCG/Source/PCG/Private/Elements/PCGPropertyToParamData.cpp b/Engine/Plugins/Experimental/PCG/Source/PCG/Private/Elements/PCGPropertyToParamData.cpp
--- a/Engine/Plugins/Experimental/PCG/Source/PCG/Private/Elements/PCGPropertyToParamData.cpp
+++ b/Engine/Plugins/Experimental/PCG/Source/PCG/Private/Elements/PCGPropertyToParamData.cpp
@@ -15,12 +15,11 @@
 #define LOCTEXT_NAMESPACE "PCGPropertyToParamDataElement"
 
 #if WITH_EDITOR
-void UPCGPropertyToParamDataSettings::GetTrackedActorTags(FPCGTagToSettingsMap& OutTagToSettings, TArray<TObjectPtr<const UPCGGraph>>& OutVisitedGraphs) const
+void UPCGPropertyToParamDataSettings::GetTrackedActorKeys(FPCGActorSelectionKeyToSettingsMap& OutKeysToSettings, TArray<TObjectPtr<const UPCGGraph>>& OutVisitedGraphs) const
 {
-	if (ActorSelector.ActorSelection == EPCGActorSelection::ByTag &&
-		ActorSelector.ActorFilter == EPCGActorFilter::AllWorldActors)
+	if (ActorSelector.ActorFilter == EPCGActorFilter::AllWorldActors)
 	{
-		OutTagToSettings.FindOrAdd(ActorSelector.ActorSelectionTag).Emplace({ this, bTrackActorsOnlyWithinBounds });
+		OutKeysToSettings.FindOrAdd(ActorSelector.GetAssociatedKey()).Emplace(this, bTrackActorsOnlyWithinBounds);
 	}
 }
 
diff --git a/Engine/Plugins/Experimental/PCG/Source/PCG/Private/Grid/PCGLandscapeCache.cpp b/Engine/Plugins/Experimental/PCG/Source/PCG/Private/Grid/PCGLandscapeCache.cpp
--- a/Engine/Plugins/Experimental/PCG/Source/PCG/Private/Grid/PCGLandscapeCache.cpp
+++ b/Engine/Plugins/Experimental/PCG/Source/PCG/Private/Grid/PCGLandscapeCache.cpp
@@ -10,7 +10,9 @@
 #include "LandscapeDataAccess.h"
 #include "LandscapeInfoMap.h"
 #include "LandscapeProxy.h"
+
 #include "Async/ParallelFor.h"
+#include "Engine/Engine.h"
 #include "Engine/World.h"
 #include "Kismet/GameplayStatics.h"
 #include "Serialization/BufferWriter.h"
@@ -754,6 +756,12 @@ void UPCGLandscapeCache::SetupLandscapeCallbacks()
 		Landscapes.Add(Landscape);
 		Landscape->OnComponentDataChanged.AddUObject(this, &UPCGLandscapeCache::OnLandscapeChanged);
 	}
+
+	// Also track when the landscape is moved
+	if (GEngine)
+	{
+		GEngine->OnActorMoved().AddUObject(this, &UPCGLandscapeCache::OnLandscapeMoved);
+	}
 }
 
 void UPCGLandscapeCache::TeardownLandscapeCallbacks()
@@ -765,11 +773,48 @@ void UPCGLandscapeCache::TeardownLandscapeCallbacks()
 			Landscape->OnComponentDataChanged.RemoveAll(this);
 		}
 	}
+
+	if (GEngine)
+	{
+		GEngine->OnActorMoved().RemoveAll(this);
+	}
+}
+
+void UPCGLandscapeCache::OnLandscapeChanged(ALandscapeProxy* InLandscape, const FLandscapeProxyComponentDataChangedParams& InChangeParams)
+{
+	if (!Landscapes.Contains(InLandscape))
+	{
+		return;
+	}
+
+	CacheLock.WriteLock();
+
+	// Just remove these from the cache, they'll be added back on demand
+	InChangeParams.ForEachComponent([this, InLandscape](const ULandscapeComponent* LandscapeComponent)
+	{
+		if (LandscapeComponent)
+		{
+			FIntPoint Coordinate(LandscapeComponent->SectionBaseX / LandscapeComponent->ComponentSizeQuads, LandscapeComponent->SectionBaseY / LandscapeComponent->ComponentSizeQuads);
+			TPair<FGuid, FIntPoint> ComponentKey(InLandscape->GetOriginalLandscapeGuid(), Coordinate);
+
+			FPCGLandscapeCacheEntry* EntryToDelete = nullptr;
+
+			if (CachedData.RemoveAndCopyValue(ComponentKey, EntryToDelete))
+			{
+				delete EntryToDelete;
+			}
+		}
+	});
+
+	CacheLayerNames(InLandscape);
+
+	CacheLock.WriteUnlock();
 }
 
-void UPCGLandscapeCache::OnLandscapeChanged(ALandscapeProxy* Landscape, const FLandscapeProxyComponentDataChangedParams& ChangeParams)
+void UPCGLandscapeCache::OnLandscapeMoved(AActor* InActor)
 {
-	if (!Landscapes.Contains(Landscape))
+	ALandscapeProxy* Landscape = Cast<ALandscapeProxy>(InActor);
+	if (!Landscape || !Landscapes.Contains(Landscape))
 	{
 		return;
 	}
@@ -777,7 +822,7 @@ void UPCGLandscapeCache::OnLandscapeChanged(ALandscapeProxy* Landscape, const FL
 	CacheLock.WriteLock();
 
 	// Just remove these from the cache, they'll be added back on demand
-	ChangeParams.ForEachComponent([this, Landscape](const ULandscapeComponent* LandscapeComponent)
+	Landscape->ForEachComponent<ULandscapeComponent>(/*bIncludeFromChildActors=*/false, [this, Landscape](const ULandscapeComponent* LandscapeComponent)
 	{
 		if (LandscapeComponent)
 		{
@@ -786,7 +831,7 @@ void UPCGLandscapeCache::OnLandscapeChanged(ALandscapeProxy* Landscape, const FL
 
 			if (FPCGLandscapeCacheEntry** FoundEntry = CachedData.Find(ComponentKey))
 			{
-				delete *FoundEntry;
+				delete* FoundEntry;
 				CachedData.Remove(ComponentKey);
 			}
 		}
@@ -810,11 +855,11 @@ void UPCGLandscapeCache::CacheLayerNames()
 	}
 }
 
-void UPCGLandscapeCache::CacheLayerNames(ALandscapeProxy* Landscape)
+void UPCGLandscapeCache::CacheLayerNames(ALandscapeProxy* InLandscape)
 {
-	check(Landscape);
+	check(InLandscape);
 
-	if (ULandscapeInfo* LandscapeInfo = Landscape->GetLandscapeInfo())
+	if (ULandscapeInfo* LandscapeInfo = InLandscape->GetLandscapeInfo())
 	{
 		for (const FLandscapeInfoLayerSettings& Layer : LandscapeInfo->Layers)
 		{
diff --git a/Engine/Plugins/Experimental/PCG/Source/PCG/Private/PCGActorAndComponentMapping.cpp b/Engine/Plugins/Experimental/PCG/Source/PCG/Private/PCGActorAndComponentMapping.cpp
--- a/Engine/Plugins/Experimental/PCG/Source/PCG/Private/PCGActorAndComponentMapping.cpp
+++ b/Engine/Plugins/Experimental/PCG/Source/PCG/Private/PCGActorAndComponentMapping.cpp
@@ -11,8 +11,25 @@
 #include "Helpers/PCGActorHelpers.h"
 #include "Helpers/PCGHelpers.h"
 
+#include "LandscapeProxy.h"
+#include "Kismet/GameplayStatics.h"
 #include "Engine/Engine.h"
 
+namespace PCGActorAndComponentMapping
+{
+	FBox GetActorBounds(const AActor* InActor)
+	{
+		FBox ActorBounds = InActor->GetComponentsBoundingBox();
+		if (!ActorBounds.IsValid && InActor->GetRootComponent() != nullptr)
+		{
+			// Try on the RootComponent
+			ActorBounds = InActor->GetRootComponent()->Bounds.GetBox();
+		}
+
+		return ActorBounds;
+	}
+}
+
 UPCGActorAndComponentMapping::UPCGActorAndComponentMapping(UPCGSubsystem* InPCGSubsystem)
 	: PCGSubsystem(InPCGSubsystem)
 {
@@ -86,7 +103,7 @@ bool UPCGActorAndComponentMapping::RegisterOrUpdatePCGComponent(UPCGComponent* I
 	check(InComponent);
 
 	// Discard BP templates, local components and invalid component
-	if (!InComponent->GetOwner() || InComponent->GetOwner()->IsA<APCGPartitionActor>() || !IsValid(InComponent))
+	if (!IsValid(InComponent) || !InComponent->GetOwner() || InComponent->GetOwner()->IsA<APCGPartitionActor>())
 	{
 		return false;
 	}
@@ -98,6 +115,8 @@ bool UPCGActorAndComponentMapping::RegisterOrUpdatePCGComponent(UPCGComponent* I
 		return false;
 	}
 
+	const bool bWasAlreadyRegistered = NonPartitionedOctree.Contains(InComponent) || PartitionedOctree.Contains(InComponent);
+
 	// First check if the component has changed its partitioned flag.
 	const bool bIsPartitioned = InComponent->IsPartitioned();
 	if (bIsPartitioned && NonPartitionedOctree.Contains(InComponent))
@@ -120,8 +139,13 @@ bool UPCGActorAndComponentMapping::RegisterOrUpdatePCGComponent(UPCGComponent* I
 		bHasChanged = RegisterOrUpdateNonPartitionedPCGComponent(InComponent);
 	}
 
-	// And finally handle the tracking
-	RegisterOrUpdateTracking(InComponent);
+	// And finally handle the tracking. Only do it when the component is registered for the first time.
+#if WITH_EDITOR
+	if (!bWasAlreadyRegistered && bHasChanged)
+	{
+		RegisterOrUpdateTracking(InComponent, /*bInShouldDirtyActors=*/ false);
+	}
+#endif // WITH_EDITOR
 
 	return bHasChanged;
 }
@@ -171,21 +195,6 @@ bool UPCGActorAndComponentMapping::RegisterOrUpdateNonPartitionedPCGComponent(UP
 {
 	// Tracking is only done in Editor for now
 #if WITH_EDITOR
-	const bool bIsTracking = InComponent->ShouldTrackActors();
-
-	if (!bIsTracking && NonPartitionedOctree.Contains(InComponent))
-	{
-		// Was tracking, but no more, remove it
-		UnregisterNonPartitionedPCGComponent(InComponent);
-		return true;
-	}
-
-	if (!bIsTracking)
-	{
-		// Nothing to do
-		return false;
-	}
-
 	FBox Bounds(EForceInit::ForceInit);
 	bool bComponentHasChanged = false;
 	bool bComponentWasAdded = false;
@@ -250,7 +259,9 @@ bool UPCGActorAndComponentMapping::RemapPCGComponent(const UPCGComponent* OldCom
 		UpdateMappingPCGComponentPartitionActor(NewComponent);
 	}
 
+#if WITH_EDITOR
 	RemapTracking(OldComponent, NewComponent);
+#endif // WITH_EDITOR
 
 	return true;
 }
@@ -280,7 +291,9 @@ void UPCGActorAndComponentMapping::UnregisterPCGComponent(UPCGComponent* InCompo
 	UnregisterPartitionedPCGComponent(InComponent);
 	UnregisterNonPartitionedPCGComponent(InComponent);
 
+#if WITH_EDITOR
 	UnregisterTracking(InComponent);
+#endif // WITH_EDITOR
 }
 
 void UPCGActorAndComponentMapping::UnregisterPartitionedPCGComponent(UPCGComponent* InComponent)
@@ -534,29 +547,164 @@ TSet<TObjectPtr<UPCGComponent>> UPCGActorAndComponentMapping::GetAllRegisteredPa
 	return PartitionedOctree.GetAllComponents();
 }
 
-bool UPCGActorAndComponentMapping::IsComponentTracked(const UPCGComponent* InComponent) const
+TSet<TObjectPtr<UPCGComponent>> UPCGActorAndComponentMapping::GetAllRegisteredNonPartitionedComponents() const
 {
-	FReadScopeLock ReadLock(TrackedComponentsLock);
-	return TrackedComponentsToActorsMap.Contains(InComponent);
+	return NonPartitionedOctree.GetAllComponents();
 }
 
-// TODO: In a following pass
-void UPCGActorAndComponentMapping::RegisterOrUpdateTracking(UPCGComponent* InComponent)
+TSet<TObjectPtr<UPCGComponent>> UPCGActorAndComponentMapping::GetAllRegisteredComponents() const
 {
+	TSet<TObjectPtr<UPCGComponent>> Res = GetAllRegisteredPartitionedComponents();
+	Res.Append(GetAllRegisteredNonPartitionedComponents());
+	return Res;
+}
+
+#if WITH_EDITOR
+void UPCGActorAndComponentMapping::RegisterOrUpdateTracking(UPCGComponent* InComponent, bool bInShouldDirtyActors)
+{
+	// Discard BP templates, local components and invalid component
+	if (!IsValid(InComponent) || !InComponent->GetOwner() || InComponent->GetOwner()->IsA<APCGPartitionActor>())
+	{
+		return;
+	}
+
+	AActor* ComponentOwner = InComponent->GetOwner();
+
+	// If we have no owner, we might be in a BP so don't track
+	if (!ComponentOwner)
+	{
+		return;
+	}
+
+	// Components owner needs to be always tracked
+	RegisterActor(ComponentOwner);
+	TSet<TObjectPtr<UPCGComponent>>& AllComponents = AlwaysTrackedActorsToComponentsMap.FindOrAdd(ComponentOwner);
+	AllComponents.Add(InComponent);
+
+	UWorld* World = PCGSubsystem->GetWorld();
+	APCGWorldActor* PCGWorldActor = PCGSubsystem ? PCGSubsystem->GetPCGWorldActor() : nullptr;
+
+	if (!World || !PCGWorldActor)
+	{
+		return;
+	}
+
+	// And we also need to find all actors that should be tracked
+	if (UPCGGraph* PCGGraph = InComponent->GetGraph())
+	{
+		auto FindActorsAndTrack = [this, InComponent, PCGWorldActor, bInShouldDirtyActors](const FPCGActorSelectionKey& InKey, const TArray<FPCGSettingsAndCulling>& InSettingsAndCulling)
+		{
+			// InKey provide the info for selecting a given actor.
+			// We reconstruct the selector settings from this key, and we also force it to SelectMultiple, since
+			// we want to gather all the actors that matches this given key.
+			FPCGActorSelectorSettings SelectorSettings = FPCGActorSelectorSettings::ReconstructFromKey(InKey);
+			SelectorSettings.bSelectMultiple = true;
+
+			TArray<AActor*> AllActors = PCGActorSelector::FindActors(SelectorSettings, InComponent, [](const AActor*) { return true; }, [](const AActor*) { return true; });
+
+			bool bShouldCull = true;
+			for (const FPCGSettingsAndCulling& SettingsAndCulling : InSettingsAndCulling)
+			{
+				if (!SettingsAndCulling.Value)
+				{
+					bShouldCull = false;
+					break;
+				}
+			}
+
+			for (AActor* Actor : AllActors)
+			{
+				if (bShouldCull)
+				{
+					CulledTrackedActorsToComponentsMap.FindOrAdd(Actor).Add(InComponent);
+				}
+				else
+				{
+					AlwaysTrackedActorsToComponentsMap.FindOrAdd(Actor).Add(InComponent);
+				}
+
+				RegisterActor(Actor);
+
+				if (bInShouldDirtyActors)
+				{
+					// If we need to force dirty, disregard culling (always intersect).
+					InComponent->DirtyTrackedActor(Actor, /*bIntersect=*/ true, {});
+				}
+			}
+		};
+
+		for (TPair<FPCGActorSelectionKey, TArray<FPCGSettingsAndCulling>>& It : PCGGraph->GetTrackedActorKeysToSettings())
+		{
+			if (!KeysToComponentsMap.Contains(It.Key))
+			{
+				FindActorsAndTrack(It.Key, It.Value);
+			}
+
+			KeysToComponentsMap.FindOrAdd(It.Key).Add(InComponent);
+		}
+
+		// Also while we support landscape pins on input node, we need to track landscape if we uses it, or the input is landscape.
+		if (InComponent->ShouldTrackLandscape())
+		{
+			// Landscape doesn't have an associated setting and is always culled.
+			FPCGActorSelectionKey LandscapeKey = FPCGActorSelectionKey(ALandscapeProxy::StaticClass());
+			if (!KeysToComponentsMap.Contains(LandscapeKey))
+			{
+				FindActorsAndTrack(LandscapeKey, { {nullptr, true} });
+			}
 
+			KeysToComponentsMap.FindOrAdd(LandscapeKey).Add(InComponent);
+		}
+	}
 }
 
 void UPCGActorAndComponentMapping::RemapTracking(const UPCGComponent* InOldComponent, UPCGComponent* InNewComponent)
 {
-
+	// Nothing to do?
 }
 
 void UPCGActorAndComponentMapping::UnregisterTracking(UPCGComponent* InComponent)
 {
+	TSet<TWeakObjectPtr<AActor>> CandidatesForUntrack;
+	TSet<FPCGActorSelectionKey> KeysToRemove;
+
+	auto RemoveFromMap = [InComponent](auto& InMap, auto& InCandidateToRemove)
+	{
+		for (auto& It : InMap)
+		{
+			It.Value.Remove(InComponent);
+			if (It.Value.IsEmpty())
+			{
+				InCandidateToRemove.Add(It.Key);
+			}
+		}
+	};
+
+	RemoveFromMap(CulledTrackedActorsToComponentsMap, CandidatesForUntrack);
+	RemoveFromMap(AlwaysTrackedActorsToComponentsMap, CandidatesForUntrack);
+	RemoveFromMap(KeysToComponentsMap, KeysToRemove);
 
+	for (const FPCGActorSelectionKey& Key : KeysToRemove)
+	{
+		KeysToComponentsMap.Remove(Key);
+	}
+
+	// We also need to untrack actors that doesn't have any component that tracks them.
+	auto ShouldBeRemoved = [](const TWeakObjectPtr<AActor>& InActor, TMap<TWeakObjectPtr<AActor>, TSet<TObjectPtr<UPCGComponent>>>& InMap)
+	{
+		TSet<TObjectPtr<UPCGComponent>>* RegisteredComponents = InMap.Find(InActor);
+		return !RegisteredComponents || RegisteredComponents->IsEmpty();
+	};
+
+	for (TWeakObjectPtr<AActor> Candidate : CandidatesForUntrack)
+	{
+		if (ShouldBeRemoved(Candidate, CulledTrackedActorsToComponentsMap) && ShouldBeRemoved(Candidate, AlwaysTrackedActorsToComponentsMap))
+		{
+			UnregisterActor(Candidate.Get());
+		}
+	}
 }
 
-#if WITH_EDITOR
 void UPCGActorAndComponentMapping::ResetPartitionActorsMap()
 {
 	PartitionActorsMapLock.WriteLock();
@@ -570,6 +718,7 @@ void UPCGActorAndComponentMapping::RegisterTrackingCallbacks()
 	GEngine->OnLevelActorAdded().AddRaw(this, &UPCGActorAndComponentMapping::OnActorAdded);
 	GEngine->OnLevelActorDeleted().AddRaw(this, &UPCGActorAndComponentMapping::OnActorDeleted);
 	FCoreUObjectDelegates::OnObjectPropertyChanged.AddRaw(this, &UPCGActorAndComponentMapping::OnObjectPropertyChanged);
+	FCoreUObjectDelegates::OnPreObjectPropertyChanged.AddRaw(this, &UPCGActorAndComponentMapping::OnPreObjectPropertyChanged);
 }
 
 void UPCGActorAndComponentMapping::TeardownTrackingCallbacks()
@@ -578,27 +727,425 @@ void UPCGActorAndComponentMapping::TeardownTrackingCallbacks()
 	GEngine->OnLevelActorAdded().RemoveAll(this);
 	GEngine->OnLevelActorDeleted().RemoveAll(this);
 	FCoreUObjectDelegates::OnObjectPropertyChanged.RemoveAll(this);
+	FCoreUObjectDelegates::OnPreObjectPropertyChanged.RemoveAll(this);
+}
+
+void UPCGActorAndComponentMapping::AddActorsPostInit()
+{
+	// Safeguard, we can't add delayed actors if the subsystem is not initialized
+	if (!PCGSubsystem || !PCGSubsystem->IsInitialized())
+	{
+		return;
+	}
+
+	for (TWeakObjectPtr<AActor>& ActorPtr : DelayedAddedActors)
+	{
+		OnActorAdded(ActorPtr.Get());
+	}
+
+	DelayedAddedActors.Empty();
+
+	// Also add the one tracked by the World Actor
+	if (APCGWorldActor* PCGWorldActor = PCGSubsystem->GetPCGWorldActor())
+	{
+		// Making a copy, since CachedTrackedActors can be modified (if an actor is no longer tracked)
+		TSet<TWeakObjectPtr<AActor>> CachedTrackedActorsCopy = PCGWorldActor->CachedTrackedActors;
+		for (TWeakObjectPtr<AActor>& ActorPtr : CachedTrackedActorsCopy)
+		{
+			AddOrUpdateTrackedActor(ActorPtr.Get());
+		}
+	}
 }
 
-// TODO: In a following pass
 void UPCGActorAndComponentMapping::OnActorAdded(AActor* InActor)
 {
+	TRACE_CPUPROFILER_EVENT_SCOPE(UPCGActorAndComponentMapping::OnActorAdded);
+
+	// We have to make sure to not create a infinite loop
+	if (!InActor || InActor->IsA<APCGWorldActor>() || !PCGSubsystem)
+	{
+		return;
+	}
+
+	// If the subsystem is not initialized, wait for it to be, and store all the actors to check
+	if (!PCGSubsystem->IsInitialized())
+	{
+		DelayedAddedActors.Add(InActor);
+		return;
+	}
+
+	if (AddOrUpdateTrackedActor(InActor))
+	{
+		// Finally notify them all
+		OnActorChanged(InActor, /*bInHasMoved=*/ false);
+	}
+}
+
+bool UPCGActorAndComponentMapping::AddOrUpdateTrackedActor(AActor* InActor)
+{
+	// We have to make sure to not create a infinite loop
+	if (!InActor || InActor->IsA<APCGWorldActor>() || !PCGSubsystem || !PCGSubsystem->GetPCGWorldActor())
+	{
+		return false;
+	}
+
+	// Gather all components, and check if they want to track this one
+	TSet<TObjectPtr<UPCGComponent>> AllComponents = GetAllRegisteredComponents();
+
+	TSet<TObjectPtr<UPCGComponent>>* CulledTrackedComponents = nullptr;
+	TSet<TObjectPtr<UPCGComponent>>* AlwaysTrackedComponents = nullptr;
+	
+	for (UPCGComponent* PCGComponent : AllComponents)
+	{
+		bool bTrackingIsCulled = false;
+		if (PCGComponent && PCGComponent->IsActorTracked(InActor, bTrackingIsCulled))
+		{
+			if (bTrackingIsCulled)
+			{
+				if (!CulledTrackedComponents)
+				{
+					CulledTrackedComponents = &CulledTrackedActorsToComponentsMap.FindOrAdd(InActor);
+				}
+
+				check(CulledTrackedComponents);
+				CulledTrackedComponents->Add(PCGComponent);
+			}
+			else
+			{
+				if (!AlwaysTrackedComponents)
+				{
+					AlwaysTrackedComponents = &AlwaysTrackedActorsToComponentsMap.FindOrAdd(InActor);
+				}
 
+				check(AlwaysTrackedComponents);
+				AlwaysTrackedComponents->Add(PCGComponent);
+			}
+		}
+	}
+
+	if (CulledTrackedComponents || AlwaysTrackedComponents)
+	{
+		RegisterActor(InActor);
+		return true;
+	}
+	else
+	{
+		// Do some cleanup. We will force the refresh here, so return false to make sure we don't refresh it twice.
+		OnActorDeleted(InActor);
+		return false;
+	}
+}
+
+void UPCGActorAndComponentMapping::RegisterActor(AActor* InActor)
+{
+	APCGWorldActor* PCGWorldActor = PCGSubsystem ? PCGSubsystem->GetPCGWorldActor() : nullptr;
+	if (!PCGWorldActor || !InActor)
+	{
+		return;
+	}
+
+	if (!PCGWorldActor->CachedTrackedActors.Contains(InActor))
+	{
+		PCGWorldActor->Modify();
+		PCGWorldActor->CachedTrackedActors.Add(InActor);
+	}
+
+	if (ALandscapeProxy* LandscapeProxy = Cast<ALandscapeProxy>(InActor))
+	{
+		// Only add it once.
+		if (!TrackedActorToPositionMap.Contains(InActor))
+		{
+			LandscapeProxy->OnComponentDataChanged.AddRaw(this, &UPCGActorAndComponentMapping::OnLandscapeChanged);
+		}
+	}
+
+	TrackedActorToPositionMap.FindOrAdd(InActor) = PCGActorAndComponentMapping::GetActorBounds(InActor);
+
+	// Also gather dependencies
+	PCGHelpers::GatherDependencies(InActor, TrackedActorsToDependenciesMap.FindOrAdd(InActor), 1);
+}
+
+bool UPCGActorAndComponentMapping::UnregisterActor(AActor* InActor)
+{
+	APCGWorldActor* PCGWorldActor = PCGSubsystem ? PCGSubsystem->GetPCGWorldActor() : nullptr;
+	if (!PCGWorldActor || !InActor)
+	{
+		return false;
+	}
+
+	if (PCGWorldActor->CachedTrackedActors.Contains(InActor))
+	{
+		PCGWorldActor->Modify();
+		PCGWorldActor->CachedTrackedActors.Remove(InActor);
+		TrackedActorToPositionMap.Remove(InActor);
+		CulledTrackedActorsToComponentsMap.Remove(InActor);
+		AlwaysTrackedActorsToComponentsMap.Remove(InActor);
+		TrackedActorsToDependenciesMap.Remove(InActor);
+
+		if (ALandscapeProxy* LandscapeProxy = Cast<ALandscapeProxy>(InActor))
+		{
+			LandscapeProxy->OnComponentDataChanged.RemoveAll(this);
+		}
+
+		return true;
+	}
+	else
+	{
+		return false;
+	}
 }
 
 void UPCGActorAndComponentMapping::OnActorDeleted(AActor* InActor)
 {
+	TRACE_CPUPROFILER_EVENT_SCOPE(UPCGActorAndComponentMapping::OnActorDeleted);
+
+	APCGWorldActor* PCGWorldActor = PCGSubsystem ? PCGSubsystem->GetPCGWorldActor() : nullptr;
+	if (!PCGWorldActor || !InActor || !PCGWorldActor->CachedTrackedActors.Contains(InActor))
+	{
+		return;
+	}
+
+	// Notify all components that the actor has changed (was removed), but the Refresh will only happen AFTER the actor was actually removed from the world (because of delayed refresh).
+	OnActorChanged(InActor, /*bInHasMoved=*/ false);
 
+	// And then delete everything
+	UnregisterActor(InActor);
 }
 
 void UPCGActorAndComponentMapping::OnActorMoved(AActor* InActor)
 {
+	TRACE_CPUPROFILER_EVENT_SCOPE(UPCGActorAndComponentMapping::OnActorMoved);
 
+	APCGWorldActor* PCGWorldActor = PCGSubsystem ? PCGSubsystem->GetPCGWorldActor() : nullptr;
+	if (!PCGWorldActor || !InActor || !PCGWorldActor->CachedTrackedActors.Contains(InActor))
+	{
+		return;
+	}
+
+	// Notify all components
+	OnActorChanged(InActor, /*bInHasMoved=*/ true);
+
+	// Update Actor position
+	if (FBox* ActorBounds = TrackedActorToPositionMap.Find(InActor))
+	{
+		*ActorBounds = PCGActorAndComponentMapping::GetActorBounds(InActor);
+	}
+}
+
+void UPCGActorAndComponentMapping::OnPreObjectPropertyChanged(UObject* InObject, const FEditPropertyChain& InEditPropertyChain)
+{
+	// We want to track tags, to see if a tag was removed
+	TempTrackedActorTags.Empty();
+	FProperty* MemberProperty = InEditPropertyChain.GetActiveMemberNode() ? InEditPropertyChain.GetActiveMemberNode()->GetValue() : nullptr;
+	AActor* Actor = Cast<AActor>(InObject);
+
+	if (!Actor || !MemberProperty || MemberProperty->GetFName() != GET_MEMBER_NAME_CHECKED(AActor, Tags))
+	{
+		return;
+	}
+
+	TempTrackedActorTags = TSet<FName>(Actor->Tags);
 }
 
 void UPCGActorAndComponentMapping::OnObjectPropertyChanged(UObject* InObject, FPropertyChangedEvent& InEvent)
 {
+	TRACE_CPUPROFILER_EVENT_SCOPE(UPCGActorAndComponentMapping::OnObjectPropertyChanged);
+
+	bool bValueNotInteractive = (InEvent.ChangeType != EPropertyChangeType::Interactive);
+	// Special exception for actor tags, as we can't track otherwise an actor "losing" a tag
+	bool bActorTagChange = (InEvent.Property && InEvent.Property->GetFName() == GET_MEMBER_NAME_CHECKED(AActor, Tags));
+	if (!bValueNotInteractive && !bActorTagChange)
+	{
+		return;
+	}
 
+	AActor* Actor = Cast<AActor>(InObject);
+	APCGWorldActor* PCGWorldActor = PCGSubsystem ? PCGSubsystem->GetPCGWorldActor() : nullptr;
+
+	if (!PCGWorldActor)
+	{
+		return;
+	}
+
+	// If we don't find any actor, try to see if it is a dependency
+	if (!Actor)
+	{
+		for (const TPair<TWeakObjectPtr<AActor>, TSet<TObjectPtr<UObject>>>& TrackedActor : TrackedActorsToDependenciesMap)
+		{
+			if (TrackedActor.Value.Contains(InObject))
+			{
+				OnActorChanged(TrackedActor.Key.Get(), /*bInHasMoved=*/ false);
+			}
+		}
+
+		return;
+	}
+
+	// Check if we are not tracking it or is a tag change.
+	bool bShouldChange = true;
+	if (!PCGWorldActor->CachedTrackedActors.Contains(Actor) || bActorTagChange)
+	{
+		bShouldChange = AddOrUpdateTrackedActor(Actor);
+	}
+
+	if (bShouldChange)
+	{
+		OnActorChanged(Actor, /*bInHasMoved=*/ false);
+	}
+}
+
+void UPCGActorAndComponentMapping::OnActorChanged(AActor* InActor, bool bInHasMoved)
+{
+	TRACE_CPUPROFILER_EVENT_SCOPE(UPCGActorAndComponentMapping::OnActorChanged);
+
+	check(InActor);
+	TSet<UPCGComponent*> DirtyComponents;
+
+	EPCGComponentDirtyFlag DirtyFlag = EPCGComponentDirtyFlag::Actor;
+	if (InActor->IsA<ALandscapeProxy>())
+	{
+		DirtyFlag = DirtyFlag | EPCGComponentDirtyFlag::Landscape;
+	}
+
+	// Check if we have a change of tag too
+	TSet<FName> RemovedTags = TempTrackedActorTags.Difference(TSet<FName>(InActor->Tags));
+
+	if (TSet<TObjectPtr<UPCGComponent>>* CulledTrackedComponents = CulledTrackedActorsToComponentsMap.Find(InActor))
+	{
+		// Not const, since it will be updated with old actor bounds
+		FBox ActorBounds = PCGActorAndComponentMapping::GetActorBounds(InActor);
+
+		// Then do an octree find to get all components that intersect with this actor.
+		// If the actor has moved, we also need to find components that intersected with it before
+		// We first do it for non-partitioned, then we do it for partitioned
+		auto UpdateNonPartitioned = [&DirtyComponents, InActor, CulledTrackedComponents, &RemovedTags, DirtyFlag](const FPCGComponentRef& ComponentRef) -> void
+		{
+			TRACE_CPUPROFILER_EVENT_SCOPE(UPCGActorAndComponentMapping::OnActorChanged::UpdateNonPartitioned);
+
+			if (DirtyComponents.Contains(ComponentRef.Component) || !CulledTrackedComponents->Contains(ComponentRef.Component))
+			{
+				return;
+			}
+
+			if (ComponentRef.Component->DirtyTrackedActor(InActor, /*bIntersect=*/true, RemovedTags))
+			{
+				ComponentRef.Component->DirtyGenerated(DirtyFlag);
+				DirtyComponents.Add(ComponentRef.Component);
+			}
+		};
+
+		NonPartitionedOctree.FindElementsWithBoundsTest(ActorBounds, UpdateNonPartitioned);
+
+		// For partitioned, we first need to find all components that intersect with our actor and then forward the dirty call to all local components that intersect.
+		auto UpdatePartitioned = [this, &DirtyComponents, InActor, CulledTrackedComponents, &ActorBounds, &RemovedTags, DirtyFlag](const FPCGComponentRef& ComponentRef)  -> void
+		{
+			TRACE_CPUPROFILER_EVENT_SCOPE(UPCGActorAndComponentMapping::OnActorChanged::UpdatePartitioned);
+
+			if (!CulledTrackedComponents->Contains(ComponentRef.Component))
+			{
+				return;
+			}
+
+			const FBox Overlap = ActorBounds.Overlap(ComponentRef.Bounds.GetBox());
+			bool bWasDirtied = false;
+
+			ForAllIntersectingPartitionActors(Overlap, [InActor, Component = ComponentRef.Component, &RemovedTags, &bWasDirtied, DirtyFlag](APCGPartitionActor* InPartitionActor) -> void
+			{
+				if (UPCGComponent* LocalComponent = InPartitionActor->GetLocalComponent(Component))
+				{
+					if (LocalComponent->DirtyTrackedActor(InActor, /*bIntersect=*/true, RemovedTags))
+					{
+						bWasDirtied = true;
+						LocalComponent->DirtyGenerated(DirtyFlag);
+					}
+				}
+			});
+
+			if (bWasDirtied)
+			{
+				// Don't dispatch
+				ComponentRef.Component->DirtyGenerated(DirtyFlag, /*bDispatchToLocalComponents=*/false);
+				DirtyComponents.Add(ComponentRef.Component);
+			}
+		};
+
+		PartitionedOctree.FindElementsWithBoundsTest(ActorBounds, UpdatePartitioned);
+
+		// If it has moved, redo it with the old bounds.
+		if (bInHasMoved)
+		{
+			TRACE_CPUPROFILER_EVENT_SCOPE(UPCGActorAndComponentMapping::OnActorChanged::SecondUpdateHasMoved);
+
+			if (FBox* OldActorBounds = TrackedActorToPositionMap.Find(InActor))
+			{
+				if (!OldActorBounds->Equals(ActorBounds))
+				{
+					// Set the actor bounds with the old one, to have the right Overlap in the Partition case.
+					ActorBounds = *OldActorBounds;
+					NonPartitionedOctree.FindElementsWithBoundsTest(*OldActorBounds, UpdateNonPartitioned);
+					PartitionedOctree.FindElementsWithBoundsTest(*OldActorBounds, UpdatePartitioned);
+				}
+			}
+		}
+	}
+
+	// Finally, dirty all components that always track this actor that are not yet notified.
+	if (TSet<TObjectPtr<UPCGComponent>>* AlwaysTrackedComponents = AlwaysTrackedActorsToComponentsMap.Find(InActor))
+	{
+		TRACE_CPUPROFILER_EVENT_SCOPE(UPCGActorAndComponentMapping::OnActorChanged::AlwaysTrackedUpdate);
+
+		for (UPCGComponent* PCGComponent : *AlwaysTrackedComponents)
+		{
+			if (!PCGComponent)
+			{
+				continue;
+			}
+
+			const bool bOwnerChanged = PCGComponent->GetOwner() == InActor;
+			bool bWasDirtied = false;
+
+			if (!DirtyComponents.Contains(PCGComponent) && !bOwnerChanged)
+			{
+				if (PCGComponent->IsPartitioned())
+				{
+					DispatchToRegisteredLocalComponents(PCGComponent, [InActor, &RemovedTags, &bWasDirtied, DirtyFlag](UPCGComponent* InLocalComponent) -> FPCGTaskId
+					{
+						if (InLocalComponent->DirtyTrackedActor(InActor, /*bIntersect=*/false, RemovedTags))
+						{
+							bWasDirtied = true;
+							InLocalComponent->DirtyGenerated(DirtyFlag);
+						}
+						return InvalidPCGTaskId;
+					});
+				}
+				else
+				{
+					bWasDirtied = PCGComponent->DirtyTrackedActor(InActor, /*bIntersect=*/false, RemovedTags);
+				}
+			}
+
+			if (bWasDirtied || bOwnerChanged)
+			{
+				PCGComponent->DirtyGenerated(DirtyFlag, /*bDispatchToLocalComponents=*/bOwnerChanged);
+				DirtyComponents.Add(PCGComponent);
+			}
+		}
+	}
+
+	// And refresh all dirtied components
+	for (UPCGComponent* Component : DirtyComponents)
+	{
+		if (Component)
+		{
+			Component->Refresh();
+		}
+	}
+}
+
+void UPCGActorAndComponentMapping::OnLandscapeChanged(ALandscapeProxy* InLandscape, const FLandscapeProxyComponentDataChangedParams& InChangeParams)
+{
+	// We don't know if the landscape moved, only that it has changed. Since `bHasMoved` is doing a bit more, always assume that the landscape has moved.
+	OnActorChanged(InLandscape, /*bHasMoved=*/true);
 }
 
 #endif // WITH_EDITOR
\ No newline at end of file
diff --git a/Engine/Plugins/Experimental/PCG/Source/PCG/Private/PCGComponent.cpp b/Engine/Plugins/Experimental/PCG/Source/PCG/Private/PCGComponent.cpp
--- a/Engine/Plugins/Experimental/PCG/Source/PCG/Private/PCGComponent.cpp
+++ b/Engine/Plugins/Experimental/PCG/Source/PCG/Private/PCGComponent.cpp
@@ -30,6 +30,7 @@
 
 #include "LandscapeComponent.h"
 #include "LandscapeProxy.h"
+#include "Algo/AnyOf.h"
 #include "Algo/Transform.h"
 #include "Components/BillboardComponent.h"
 #include "Components/InstancedStaticMeshComponent.h"
@@ -201,11 +202,12 @@ void UPCGComponent::SetPropertiesFromOriginal(const UPCGComponent* Original)
 #if WITH_EDITOR
 	const bool bHasDirtyInput = InputType != NewInputType;
 
-	TSet<FName> TrackedTags;
-	TSet<FName> OriginalTrackedTags;
-	CachedTrackedTagsToSettings.GetKeys(TrackedTags);
-	Original->CachedTrackedTagsToSettings.GetKeys(OriginalTrackedTags);
-	const bool bHasDirtyTracking = !(TrackedTags.Num() == OriginalTrackedTags.Num() && TrackedTags.Includes(OriginalTrackedTags));
+	TSet<FPCGActorSelectionKey> TrackedKeys;
+	TSet<FPCGActorSelectionKey> OriginalTrackedKeys;
+	CachedTrackedKeysToSettings.GetKeys(TrackedKeys);
+	Original->CachedTrackedKeysToSettings.GetKeys(OriginalTrackedKeys);
+
+	const bool bHasDirtyTracking = !(TrackedKeys.Num() == OriginalTrackedKeys.Num() && TrackedKeys.Includes(OriginalTrackedKeys));
 
 	const bool bIsDirty = bHasDirtyInput || bHasDirtyTracking || bGraphInstanceIsDifferent;
 #endif // WITH_EDITOR
@@ -228,13 +230,9 @@ void UPCGComponent::SetPropertiesFromOriginal(const UPCGComponent* Original)
 #if WITH_EDITOR
 	if (bHasDirtyTracking)
 	{
-		TeardownTrackingCallbacks();
-		SetupTrackingCallbacks();
-		RefreshTrackingData();
+		UpdateTrackingCache();
 	}
 
-	UpdateTrackedLandscape();
-
 	// Note that while we dirty here, we won't trigger a refresh since we don't have the required context
 	if (bIsDirty)
 	{
@@ -328,48 +326,6 @@ FPCGTaskId UPCGComponent::CreateGenerateTask(bool bForce, const TArray<FPCGTaskI
 	return GetSubsystem()->ScheduleGraph(this, *AllDependencies);
 }
 
-bool UPCGComponent::GetActorsFromTags(const TMap<FName, bool>& InTagsAndCulling, TSet<TWeakObjectPtr<AActor>>& OutActors)
-{
-	TRACE_CPUPROFILER_EVENT_SCOPE(UPCGComponent::GetActorsFromTags::Tracked);
-	UWorld* World = GetWorld();
-
-	if (!World)
-	{
-		return false;
-	}
-
-	const FBox LocalBounds = GetGridBounds();
-
-	TArray<AActor*> PerTagActors;
-
-	OutActors.Reset();
-
-	bool bHasValidTag = false;
-	for (const TPair<FName, bool>& TagAndCulling : InTagsAndCulling)
-	{
-		const FName& Tag = TagAndCulling.Key;
-		const bool bCullAgainstLocalBounds = TagAndCulling.Value;
-
-		if (Tag != NAME_None)
-		{
-			bHasValidTag = true;
-			UGameplayStatics::GetAllActorsWithTag(World, Tag, PerTagActors);
-
-			for (AActor* Actor : PerTagActors)
-			{
-				if (!bCullAgainstLocalBounds || LocalBounds.Intersect(GetGridBounds(Actor)))
-				{
-					OutActors.Emplace(Actor);
-				}
-			}
-
-			PerTagActors.Reset();
-		}
-	}
-
-	return bHasValidTag;
-}
-
 void UPCGComponent::PostProcessGraph(const FBox& InNewBounds, bool bInGenerated, FPCGContext* Context)
 {
 	PCGGeneratedResourcesLogging::LogPostProcessGraph();
@@ -975,15 +931,6 @@ void UPCGComponent::EndPlay(const EEndPlayReason::Type EndPlayReason)
 	Super::EndPlay(EndPlayReason);
 }
 
-void UPCGComponent::OnComponentCreated()
-{
-	Super::OnComponentCreated();
-
-#if WITH_EDITOR
-	SetupActorCallbacks();
-#endif
-}
-
 void UPCGComponent::OnComponentDestroyed(bool bDestroyingHierarchy)
 {
 #if WITH_EDITOR
@@ -1066,6 +1013,18 @@ void UPCGComponent::PostLoad()
 		Graph_DEPRECATED = nullptr;
 	}
 
+	if (!CachedTrackedActors_DEPRECATED.IsEmpty() || !TrackedLandscapes_DEPRECATED.IsEmpty())
+	{
+		if (UPCGSubsystem* Subsystem = GetSubsystem())
+		{
+			Subsystem->AddSerializedTrackedActorsForDeprecation(CachedTrackedActors_DEPRECATED);
+			Subsystem->AddSerializedTrackedActorsForDeprecation(TrackedLandscapes_DEPRECATED);
+		}
+
+		CachedTrackedActors_DEPRECATED.Empty();
+		TrackedLandscapes_DEPRECATED.Empty();
+	}
+
 	SetupCallbacksOnCreation();
 #endif
 }
@@ -1073,17 +1032,7 @@ void UPCGComponent::PostLoad()
 #if WITH_EDITOR
 void UPCGComponent::SetupCallbacksOnCreation()
 {
-	SetupActorCallbacks();
-	SetupTrackingCallbacks();
-
-	if(!TrackedLandscapes.IsEmpty())
-	{
-		SetupLandscapeTracking();
-	}
-	else
-	{
-		UpdateTrackedLandscape(/*bBoundsCheck=*/false);
-	}
+	UpdateTrackingCache();
 
 	if (GraphInstance)
 	{
@@ -1108,13 +1057,6 @@ void UPCGComponent::BeginDestroy()
 	{
 		GraphInstance->OnGraphChangedDelegate.RemoveAll(this);
 	}
-
-	if (!IsEngineExitRequested())
-	{
-		TeardownLandscapeTracking();
-		TeardownTrackingCallbacks();
-		TeardownActorCallbacks();
-	}
 #endif
 
 	Super::BeginDestroy();
@@ -1193,14 +1135,10 @@ void UPCGComponent::RefreshAfterGraphChanged(UPCGGraphInterface* InGraph, bool b
 	// In editor, since we've changed the graph, we might have changed the tracked actor tags as well
 	if (!PCGHelpers::IsRuntimeOrPIE())
 	{
-		TeardownTrackingCallbacks();
-		SetupTrackingCallbacks();
-		RefreshTrackingData();
-		DirtyCacheForAllTrackedTags();
-
-		if (bIsStructural)
+		UpdateTrackingCache();
+		if (UPCGSubsystem* Subsystem = GetSubsystem())
 		{
-			UpdateTrackedLandscape();
+			Subsystem->UpdateComponentTracking(this, /*bInShouldDirtyActors=*/ true);
 		}
 
 		DirtyGenerated(bDirtyInputs ? (EPCGComponentDirtyFlag::Actor | EPCGComponentDirtyFlag::Landscape) : EPCGComponentDirtyFlag::None);
@@ -1271,7 +1209,11 @@ void UPCGComponent::PostEditChangeProperty(FPropertyChangedEvent& PropertyChange
 	}
 	else if (PropName == GET_MEMBER_NAME_CHECKED(UPCGComponent, InputType))
 	{
-		UpdateTrackedLandscape();
+		if (UPCGSubsystem* Subsystem = GetSubsystem())
+		{
+			Subsystem->UpdateComponentTracking(this);
+		}
+
 		DirtyGenerated(EPCGComponentDirtyFlag::Input);
 		Refresh();
 	}
@@ -1306,46 +1248,28 @@ void UPCGComponent::PreEditUndo()
 		CleanupLocalImmediate(/*bRemoveComponents=*/true);
 		// Put back generated flag to its original value so it is captured properly
 		bGenerated = true;
-	}	
-	
-	TeardownTrackingCallbacks();
+	}
 }
 
 void UPCGComponent::PostEditUndo()
 {
 	LastGeneratedBounds = LastGeneratedBoundsPriorToUndo;
 
-	SetupTrackingCallbacks();
-	RefreshTrackingData();
-	UpdateTrackedLandscape();
+	UpdateTrackingCache();
 	DirtyGenerated(EPCGComponentDirtyFlag::All);
-	DirtyCacheForAllTrackedTags();
 
-	if (bGenerated)
+	if (UPCGSubsystem* Subsystem = GetSubsystem())
 	{
-		Refresh();
+		Subsystem->UpdateComponentTracking(this, /*bInShouldDirtyActors=*/ true);
 	}
-}
 
-void UPCGComponent::SetupActorCallbacks()
-{
-	// Without an owner, it probably means we are in a BP template, so no need to setup callbacks
-	if (!GetOwner())
+	if (bGenerated)
 	{
-		return;
+		Refresh();
 	}
-
-	GEngine->OnActorMoved().AddUObject(this, &UPCGComponent::OnActorMoved);
-	FCoreUObjectDelegates::OnObjectPropertyChanged.AddUObject(this, &UPCGComponent::OnObjectPropertyChanged);
-}
-
-void UPCGComponent::TeardownActorCallbacks()
-{
-	FCoreUObjectDelegates::OnObjectPropertyChanged.RemoveAll(this);
-	GEngine->OnActorMoved().RemoveAll(this);
 }
 
-void UPCGComponent::SetupTrackingCallbacks()
+void UPCGComponent::UpdateTrackingCache()
 {
 	// Without an owner, it probably means we are in a BP template, so no need to setup callbacks
 	if (!GetOwner())
@@ -1353,19 +1277,20 @@ void UPCGComponent::SetupTrackingCallbacks()
 		return;
 	}
 
-	CachedTrackedTagsToSettings.Reset();
-	CachedTrackedTagsToCulling.Reset();
+	CachedTrackedKeysToSettings.Reset();
+	CachedTrackedKeysToCulling.Reset();
+
 	if (UPCGGraph* PCGGraph = GetGraph())
 	{
-		CachedTrackedTagsToSettings = PCGGraph->GetTrackedTagsToSettings();
+		CachedTrackedKeysToSettings = PCGGraph->GetTrackedActorKeysToSettings();
 
 		// A tag should be culled, if only all the settings that track this tag should cull.
 		// Note that is only impact the fact that we track (or not) this tag.
 		// If a setting is marked as "should cull", it will only be dirtied (at least by default), if the actor with the
 		// given tag intersect with the component.
-		for (const TPair<FName, TSet<FPCGSettingsAndCulling>>& It : CachedTrackedTagsToSettings)
+		for (const TPair<FPCGActorSelectionKey, TArray<FPCGSettingsAndCulling>>& It : CachedTrackedKeysToSettings)
 		{
-			const FName& Tag = It.Key;
+			const FPCGActorSelectionKey& Key = It.Key;
 
 			bool bShouldCull = true;
 			for (const FPCGSettingsAndCulling& SettingsAndCulling : It.Value)
@@ -1377,286 +1302,7 @@ void UPCGComponent::SetupTrackingCallbacks()
 				}
 			}
 
-			CachedTrackedTagsToCulling.Emplace(Tag, bShouldCull);
-		}
-	}
-
-	if(!CachedTrackedTagsToSettings.IsEmpty())
-	{
-		GEngine->OnLevelActorAdded().AddUObject(this, &UPCGComponent::OnActorAdded);
-		GEngine->OnLevelActorDeleted().AddUObject(this, &UPCGComponent::OnActorDeleted);
-	}
-}
-
-void UPCGComponent::RefreshTrackingData()
-{
-	TRACE_CPUPROFILER_EVENT_SCOPE(UPCGComponent::RefreshTrackingData);
-
-	// Without an owner, it probably means we are in a BP template, so no need to setup callbacks
-	if (!GetOwner())
-	{
-		return;
-	}
-
-	GetActorsFromTags(CachedTrackedTagsToCulling, CachedTrackedActors);
-	PopulateTrackedActorToTagsMap(/*bForce=*/true);
-}
-
-void UPCGComponent::TeardownTrackingCallbacks()
-{
-	GEngine->OnLevelActorAdded().RemoveAll(this);
-	GEngine->OnLevelActorDeleted().RemoveAll(this);
-}
-
-bool UPCGComponent::ActorIsTracked(AActor* InActor) const
-{
-	if (!InActor || !GetGraph())
-	{
-		return false;
-	}
-
-	bool bIsTracked = false;
-	for (const FName& Tag : InActor->Tags)
-	{
-		if (CachedTrackedTagsToSettings.Contains(Tag))
-		{
-			bIsTracked = true;
-			break;
-		}
-	}
-
-	return bIsTracked;
-}
-
-void UPCGComponent::OnActorAdded(AActor* InActor)
-{
-	TRACE_CPUPROFILER_EVENT_SCOPE(UPCGComponent::OnActorAdded);
-	if (!InActor || InActor->bIsEditorPreviewActor)
-	{
-		return;
-	}
-
-	const bool bIsTracked = AddTrackedActor(InActor);
-	if (bIsTracked)
-	{
-		DirtyGenerated(EPCGComponentDirtyFlag::None, /*bDispatchToLocalComponents=*/ false);
-		Refresh();
-	}
-}
-
-void UPCGComponent::OnActorDeleted(AActor* InActor)
-{
-	TRACE_CPUPROFILER_EVENT_SCOPE(UPCGComponent::OnActorDeleted);
-	if (!InActor || InActor->bIsEditorPreviewActor)
-	{
-		return;
-	}
-
-	const bool bWasTracked = RemoveTrackedActor(InActor);
-	if (bWasTracked)
-	{
-		DirtyGenerated(EPCGComponentDirtyFlag::None, /*bDispatchToLocalComponents=*/ false);
-		Refresh();
-	}
-}
-
-void UPCGComponent::OnActorMoved(AActor* InActor)
-{
-	TRACE_CPUPROFILER_EVENT_SCOPE(UPCGComponent::OnActorMoved);
-	if (!InActor || InActor->bIsEditorPreviewActor)
-	{
-		return;
-	}
-
-	const bool bOwnerMoved = (InActor == GetOwner());
-	const bool bLandscapeMoved = (InActor && TrackedLandscapes.Contains(InActor));
-
-	if (bOwnerMoved || bLandscapeMoved)
-	{
-		// TODO: find better metrics to dirty the inputs. 
-		// TODO: this should dirty only the actor pcg data.
-		{
-			UpdateTrackedLandscape();
-			DirtyGenerated((bOwnerMoved ? EPCGComponentDirtyFlag::Actor : EPCGComponentDirtyFlag::None) | (bLandscapeMoved ? EPCGComponentDirtyFlag::Landscape : EPCGComponentDirtyFlag::None));
-			Refresh();
-		}
-	}
-	else
-	{
-		if (DirtyTrackedActor(InActor))
-		{
-			DirtyGenerated(EPCGComponentDirtyFlag::None, /*bDispatchToLocalComponents=*/false);
-			Refresh();
-		}
-	}
-}
-
-void UPCGComponent::UpdateTrackedLandscape(bool bBoundsCheck)
-{
-	TeardownLandscapeTracking();
-	TrackedLandscapes.Reset();
-
-	// Without an owner, it probably means we are in a BP template, so no need to setup callbacks
-	if (!GetOwner())
-	{
-		return;
-	}
-
-	if (ALandscapeProxy* Landscape = Cast<ALandscapeProxy>(GetOwner()))
-	{
-		TrackedLandscapes.Add(Landscape);
-	}
-	else if (InputType == EPCGComponentInput::Landscape || GraphUsesLandscapePin())
-	{
-		if (UWorld* World = GetOwner() ? GetOwner()->GetWorld() : nullptr)
-		{
-			if (bBoundsCheck)
-			{
-				UPCGData* ActorData = GetActorPCGData();
-				if (const UPCGSpatialData* ActorSpatialData = Cast<const UPCGSpatialData>(ActorData))
-				{
-					Algo::Transform(PCGHelpers::GetLandscapeProxies(World, ActorSpatialData->GetBounds()), TrackedLandscapes, [](TWeakObjectPtr<ALandscapeProxy> Landscape) { return TSoftObjectPtr<ALandscapeProxy>(Landscape.Get()); });
-				}
-			}
-			else
-			{
-				Algo::Transform(PCGHelpers::GetAllLandscapeProxies(World), TrackedLandscapes, [](TWeakObjectPtr<ALandscapeProxy> Landscape) { return TSoftObjectPtr<ALandscapeProxy>(Landscape.Get()); });
-			}
-		}
-	}
-
-	SetupLandscapeTracking();
-}
-
-void UPCGComponent::SetupLandscapeTracking()
-{
-	for (TSoftObjectPtr<ALandscapeProxy> LandscapeProxy : TrackedLandscapes)
-	{
-		if (LandscapeProxy.IsValid())
-		{
-			LandscapeProxy->OnComponentDataChanged.AddUObject(this, &UPCGComponent::OnLandscapeChanged);
-		}
-	}
-}
-
-void UPCGComponent::TeardownLandscapeTracking()
-{
-	for (TSoftObjectPtr<ALandscapeProxy> LandscapeProxy : TrackedLandscapes)
-	{
-		if (LandscapeProxy.IsValid())
-		{
-			LandscapeProxy->OnComponentDataChanged.RemoveAll(this);
-		}
-	}
-}
-
-void UPCGComponent::OnLandscapeChanged(ALandscapeProxy* Landscape, const FLandscapeProxyComponentDataChangedParams& ChangeParams)
-{
-	if(TrackedLandscapes.Contains(Landscape))
-	{
-		// Check if there is an overlap in the changed components vs. the current actor data
-		EPCGComponentDirtyFlag DirtyFlag = EPCGComponentDirtyFlag::None;
-
-		if (GetOwner() == Landscape)
-		{
-			DirtyFlag = EPCGComponentDirtyFlag::Actor;
-		}
-		// Note: this means that graphs that are interacting with the landscape outside their bounds might not be updated properly
-		else if (InputType == EPCGComponentInput::Landscape || GraphUsesLandscapePin())
-		{
-			UPCGData* ActorData = GetActorPCGData();
-			if (const UPCGSpatialData* ActorSpatialData = Cast<const UPCGSpatialData>(ActorData))
-			{
-				const FBox ActorBounds = ActorSpatialData->GetBounds();
-				bool bDirtyLandscape = false;
-
-				ChangeParams.ForEachComponent([&bDirtyLandscape, &ActorBounds](const ULandscapeComponent* LandscapeComponent)
-				{
-					if(LandscapeComponent && ActorBounds.Intersect(LandscapeComponent->Bounds.GetBox()))
-					{
-						bDirtyLandscape = true;
-					}
-				});
-
-				if (bDirtyLandscape)
-				{
-					DirtyFlag = EPCGComponentDirtyFlag::Landscape;
-				}
-			}
-		}
-
-		if (DirtyFlag != EPCGComponentDirtyFlag::None)
-		{
-			DirtyGenerated(DirtyFlag, /*bDispatchToLocalComponents=*/ false);
-			Refresh();
-		}
-	}
-}
-
-void UPCGComponent::OnObjectPropertyChanged(UObject* InObject, FPropertyChangedEvent& InEvent)
-{
-	TRACE_CPUPROFILER_EVENT_SCOPE(UPCGComponent::OnObjectPropertyChanged);
-	bool bValueNotInteractive = (InEvent.ChangeType != EPropertyChangeType::Interactive);
-	// Special exception for actor tags, as we can't track otherwise an actor "losing" a tag
-	bool bActorTagChange = (InEvent.Property && InEvent.Property->GetFName() == GET_MEMBER_NAME_CHECKED(AActor, Tags));
-	if (!bValueNotInteractive && !bActorTagChange)
-	{
-		return;
-	}
-
-	// If the object changed is this PCGComponent, dirty ourselves and exit. It will be picked up by PostEditChangeProperty
-	if (InObject == this)
-	{
-		DirtyGenerated();
-		return;
-	}
-
-	// First, check if it's an actor
-	AActor* Actor = Cast<AActor>(InObject);
-
-	// Otherwise, if it's an actor component, track it as well
-	if (!Actor)
-	{
-		if (UActorComponent* ActorComponent = Cast<UActorComponent>(InObject))
-		{
-			Actor = ActorComponent->GetOwner();
-		}
-	}
-
-	// Finally, if it's neither an actor or an actor component, it might be a dependency of a tracked actor
-	if (!Actor)
-	{
-		for(const TPair<TWeakObjectPtr<AActor>, TSet<TObjectPtr<UObject>>>& TrackedActor : CachedTrackedActorToDependencies)
-		{
-			if (TrackedActor.Value.Contains(InObject))
-			{
-				OnActorChanged(TrackedActor.Key.Get(), InObject, bActorTagChange);
-			}
-		}
-	}
-	else
-	{
-		OnActorChanged(Actor, InObject, bActorTagChange);
-	}
-}
-
-void UPCGComponent::OnActorChanged(AActor* Actor, UObject* InObject, bool bActorTagChange)
-{
-	TRACE_CPUPROFILER_EVENT_SCOPE(UPCGComponent::OnActorChanged);
-	if (Actor == GetOwner())
-	{
-		// Something has changed on the owner (including properties of this component)
-		// In the case of splines, this is where we'd get notified if some component properties (incl. spline vertices) have changed
-		// TODO: this should dirty only the actor pcg data.
-		DirtyGenerated(EPCGComponentDirtyFlag::Actor);
-		Refresh();
-	}
-	else if(Actor && !Actor->bIsEditorPreviewActor)
-	{
-		if ((bActorTagChange && Actor == InObject && UpdateTrackedActor(Actor)) || DirtyTrackedActor(Actor))
-		{
-			DirtyGenerated(EPCGComponentDirtyFlag::None, /*bDispatchToLocalComponents=*/ false);
-			Refresh();
+			CachedTrackedKeysToCulling.Emplace(Key, bShouldCull);
 		}
 	}
 }
@@ -1830,6 +1476,37 @@ bool UPCGComponent::ShouldGenerateBPPCGAddedToWorld() const
 	}
 }
 
+bool UPCGComponent::IsActorTracked(AActor* InActor, bool& bOutIsCulled) const
+{
+	check(InActor);
+
+	if (!GetOwner())
+	{
+		return false;
+	}
+
+	bool bFound = false;
+
+	for (const FPCGActorSelectionKey& Key : FPCGActorSelectionKey::GenerateAllKeysForActor(InActor))
+	{
+		if (const bool* ShouldCullPtr = CachedTrackedKeysToCulling.Find(Key))
+		{
+			bOutIsCulled = *ShouldCullPtr;
+			bFound = true;
+
+			// Check for other tags that might not be culled
+			if (bOutIsCulled)
+			{
+				continue;
+			}
+
+			return true;
+		}
+	}
+
+	return bFound;
+}
+
 void UPCGComponent::OnRefresh()
 {
 	// Mark the refresh task invalid to allow re-triggering refreshes
@@ -2359,251 +2036,55 @@ UPCGSubsystem* UPCGComponent::GetSubsystem() const
 }
 
 #if WITH_EDITOR
-bool UPCGComponent::PopulateTrackedActorToTagsMap(bool bForce)
-{
-	TRACE_CPUPROFILER_EVENT_SCOPE(UPCGComponent::PopulateTrackedActorToTagsMap);
-	if (bActorToTagsMapPopulated && !bForce)
-	{
-		return false;
-	}
-
-	CachedTrackedActorToTags.Reset();
-	CachedTrackedActorToDependencies.Reset();
-	for (TWeakObjectPtr<AActor> Actor : CachedTrackedActors)
-	{
-		if (Actor.IsValid())
-		{
-			AddTrackedActor(Actor.Get(), /*bForce=*/true);
-		}
-	}
-
-	bActorToTagsMapPopulated = true;
-	return true;
-}
-
-bool UPCGComponent::AddTrackedActor(AActor* InActor, bool bForce)
+bool UPCGComponent::DirtyTrackedActor(AActor* InActor, bool bIntersect, const TSet<FName>& InRemovedTags)
 {
-	TRACE_CPUPROFILER_EVENT_SCOPE(UPCGComponent::AddTrackedActor);
-	if (!bForce)
-	{
-		PopulateTrackedActorToTagsMap();
-	}
+	TRACE_CPUPROFILER_EVENT_SCOPE(UPCGComponent::DirtyTrackedActor);
 
-	if (!GetOwner())
+	if (!InActor)
 	{
-		// Without an owner, we can't get our bounds, so we won't be able to track.
-		// If we don't have an owner, we are probably in a BP, not instanciated, so no need to track anything.
 		return false;
 	}
 
-	check(InActor);
-	bool bAppliedChange = false;
-
-	bool bIntersectionComputed = false;
-	bool bIntersect = false;
-
-	for (const FName& Tag : InActor->Tags)
-	{
-		const bool* ShouldCullPtr = CachedTrackedTagsToCulling.Find(Tag);
-
-		if (!ShouldCullPtr)
-		{
-			continue;
-		}
-
-		if (*ShouldCullPtr)
-		{
-			if (!bIntersectionComputed)
-			{
-				bIntersect = GetGridBounds().Intersect(GetGridBounds(InActor));
-				bIntersectionComputed = true;
-			}
-
-			if (!bIntersect)
-			{
-				continue;
-			}
-		}
-
-		bAppliedChange = true;
-		CachedTrackedActorToTags.FindOrAdd(InActor).Add(Tag);
-		PCGHelpers::GatherDependencies(InActor, CachedTrackedActorToDependencies.FindOrAdd(InActor), 1);
-
-		if (!bForce)
-		{
-			// Since we arrived here, we already verified culling, so we can ignore it.
-			DirtyCacheFromTag(Tag, InActor, /*bIgnoreCull=*/true);
-		}
-	}
-
-	return bAppliedChange;
-}
-
-bool UPCGComponent::RemoveTrackedActor(AActor* InActor)
-{
-	TRACE_CPUPROFILER_EVENT_SCOPE(UPCGComponent::RemoveTrackedActor);
-	PopulateTrackedActorToTagsMap();
+	bool bWasDirtied = false;
 
-	check(InActor);
-	bool bAppliedChange = false;
+	TArray<FPCGActorSelectionKey> AllKeys = FPCGActorSelectionKey::GenerateAllKeysForActor(InActor);
 
-	if(CachedTrackedActorToTags.Contains(InActor))
+	// We also need to force a refresh on tags that were removed.
+	for (FName RemovedTag : InRemovedTags)
 	{
-		for (const FName& Tag : CachedTrackedActorToTags[InActor])
-		{
-			bAppliedChange |= DirtyCacheFromTag(Tag, InActor);
-		}
-
-		CachedTrackedActorToTags.Remove(InActor);
-		CachedTrackedActorToDependencies.Remove(InActor);
+		AllKeys.Emplace(RemovedTag);
 	}
 
-	return bAppliedChange;
-}
-
-bool UPCGComponent::UpdateTrackedActor(AActor* InActor)
-{
-	TRACE_CPUPROFILER_EVENT_SCOPE(UPCGComponent::UpdateTrackedActor);
-	check(InActor);
-	// If the tracked data wasn't initialized before, then it is not possible to know if we need to update or not - take no chances
-	bool bAppliedChange = PopulateTrackedActorToTagsMap();
-
-	// Update the contents of the tracked actor vs. its current tags, and dirty accordingly
-	if (CachedTrackedActorToTags.Contains(InActor))
+	for (FPCGActorSelectionKey Key : AllKeys)
 	{
-		// Any tags that aren't on the actor and were in the cached actor to tags -> remove & dirty
-		TSet<FName> CachedTags = CachedTrackedActorToTags[InActor];
-		for (const FName& CachedTag : CachedTags)
-		{
-			if (!InActor->Tags.Contains(CachedTag))
-			{
-				bAppliedChange |= DirtyCacheFromTag(CachedTag, InActor);
-				CachedTrackedActorToTags[InActor].Remove(CachedTag);
-			}
-		}
-	}
+		const TArray<FPCGSettingsAndCulling>* SettingsAndCullingPtr = nullptr;
 		
-	// Any tags that are new on the actor and not in the cached actor to tags -> add & dirty
-	for (const FName& Tag : InActor->Tags)
-	{
-		if (!CachedTrackedTagsToSettings.Contains(Tag))
-		{
-			continue;
-		}
-
-		if (!CachedTrackedActorToTags.FindOrAdd(InActor).Find(Tag))
-		{
-			CachedTrackedActorToTags[InActor].Add(Tag);
-			PCGHelpers::GatherDependencies(InActor, CachedTrackedActorToDependencies.FindOrAdd(InActor), 1);
-			bAppliedChange |= DirtyCacheFromTag(Tag, InActor);
-		}
-	}
-
-	// Finally, if the current has no tag anymore, we can remove it from the map
-	if (CachedTrackedActorToTags.Contains(InActor) && CachedTrackedActorToTags[InActor].IsEmpty())
-	{
-		CachedTrackedActorToTags.Remove(InActor);
-		CachedTrackedActorToDependencies.Remove(InActor);
-	}
-
-	return bAppliedChange;
-}
-
-bool UPCGComponent::DirtyTrackedActor(AActor* InActor)
-{
-	TRACE_CPUPROFILER_EVENT_SCOPE(UPCGComponent::DirtyTrackedActor);
-	PopulateTrackedActorToTagsMap();
-
-	check(InActor);
-	bool bAppliedChange = false;
-
-	bool bIntersectionComputed = false;
-	bool bIntersect = false;
-
-	if (CachedTrackedActorToTags.Contains(InActor))
-	{
-		TSet<FName> CachedTags = CachedTrackedActorToTags[InActor];
-		for (const FName& CachedTag : CachedTags)
+		while (!SettingsAndCullingPtr)
 		{
-			const bool* ShouldCullPtr = CachedTrackedTagsToCulling.Find(CachedTag);
+			SettingsAndCullingPtr = CachedTrackedKeysToSettings.Find(Key);
 
-			if (!ShouldCullPtr)
+			if (!SettingsAndCullingPtr && Key.Selection == EPCGActorSelection::ByClass && Key.ActorSelectionClass != nullptr)
 			{
-				continue;
+				Key.ActorSelectionClass = (Key.ActorSelectionClass == AActor::StaticClass()) ? nullptr : Key.ActorSelectionClass->GetSuperClass();
 			}
-
-			if (*ShouldCullPtr)
+			else
 			{
-				if (!bIntersectionComputed)
-				{
-					bIntersect = GetGridBounds().Intersect(GetGridBounds(InActor));
-					bIntersectionComputed = true;
-				}
-
-				if (!bIntersect)
-				{
-					// We were tracking the tag, but it is now culled. We'll need to dirty the settings
-					// even if it is culled (that's why we ignore cull here).
-					bAppliedChange |= DirtyCacheFromTag(CachedTag, InActor, /*bIgnoreCull=*/ true);
-					CachedTrackedActorToTags[InActor].Remove(CachedTag);
-					continue;
-				}
+				break;
 			}
+		} 
 
-			bAppliedChange |= DirtyCacheFromTag(CachedTag, InActor);
-		}
-	}
-	else if (AddTrackedActor(InActor))
-	{
-		bAppliedChange = true;
-	}
-
-	// Since we might have remove some tags, do some cleaning there
-	if (CachedTrackedActorToTags.Contains(InActor) && CachedTrackedActorToTags[InActor].IsEmpty())
-	{
-		CachedTrackedActorToTags.Remove(InActor);
-		CachedTrackedActorToDependencies.Remove(InActor);
-	}
-
-	return bAppliedChange;
-}
-
-bool UPCGComponent::DirtyCacheFromTag(const FName& InTag, const AActor* InActor, bool bIgnoreCull)
-{
-	TRACE_CPUPROFILER_EVENT_SCOPE(UPCGComponent::DirtyCacheFromTag);
-	bool bWasDirtied = false;
-
-	const TSet<FPCGSettingsAndCulling>* SettingsAndCullingPtr = CachedTrackedTagsToSettings.Find(InTag);
-
-	if (SettingsAndCullingPtr)
-	{
-		// Don't compute the bounds if we never cull.
-		bool bIntersectionComputed = false;
-		bool bIntersect = false;
-
-		for (const FPCGSettingsAndCulling& SettingsAndCulling : *SettingsAndCullingPtr)
+		if (SettingsAndCullingPtr)
 		{
-			const TWeakObjectPtr<const UPCGSettings>& Settings = SettingsAndCulling.Key;
-			const bool bShouldCull = SettingsAndCulling.Value && !bIgnoreCull;
-
-			if (Settings.IsValid() && GetSubsystem())
+			for (const FPCGSettingsAndCulling& SettingsAndCulling : *SettingsAndCullingPtr)
 			{
-				// If we cull and no intersection, continue
-				if (bShouldCull)
+				if (SettingsAndCulling.Value && !bIntersect)
 				{
-					if (!bIntersectionComputed)
-					{
-						bIntersect = GetGridBounds().Intersect(GetGridBounds(InActor));
-						bIntersectionComputed = true;
-					}
-
-					if (!bIntersect)
-					{
-						continue;
-					}
+					continue;
 				}
 
+				const TWeakObjectPtr<const UPCGSettings>& Settings = SettingsAndCulling.Key;
 				GetSubsystem()->CleanFromCache(Settings->GetElement().Get(), Settings.Get());
+
 				bWasDirtied = true;
 			}
 		}
@@ -2612,28 +2093,20 @@ bool UPCGComponent::DirtyCacheFromTag(const FName& InTag, const AActor* InActor,
 	return bWasDirtied;
 }
 
-void UPCGComponent::DirtyCacheForAllTrackedTags()
-{
-	TRACE_CPUPROFILER_EVENT_SCOPE(UPCGComponent::DirtyCacheForAllTrackedTags);
-	for (const auto& TagToSettings : CachedTrackedTagsToSettings)
-	{
-		for (const FPCGSettingsAndCulling& SettingsAndCulling : TagToSettings.Value)
-		{
-			const TWeakObjectPtr<const UPCGSettings>& Settings = SettingsAndCulling.Key;
-			if (Settings.IsValid() && GetSubsystem())
-			{
-				GetSubsystem()->CleanFromCache(Settings->GetElement().Get(), Settings.Get());
-			}
-		}
-	}
-}
-
-bool UPCGComponent::GraphUsesLandscapePin() const
+bool UPCGComponent::ShouldTrackLandscape() const
 {
 	const UPCGGraph* PCGGraph = GetGraph();
-	return PCGGraph &&
+	
+	// We should track the landscape if the landscape pins are connected, or if the input type is Landscape and we are using the Input pin.
+	const bool bUseLandscapePin = PCGGraph &&
 		(PCGGraph->GetInputNode()->IsOutputPinConnected(PCGInputOutputConstants::DefaultLandscapeLabel) ||
 		PCGGraph->GetInputNode()->IsOutputPinConnected(PCGInputOutputConstants::DefaultLandscapeHeightLabel));
+
+
+	const bool bHasLandscapeHasInput = PCGGraph && InputType == EPCGComponentInput::Landscape 
+		&& Algo::AnyOf(PCGGraph->GetInputNode()->GetOutputPins(), [](const UPCGPin* InPin) { return InPin && InPin->IsConnected(); });
+
+	return bUseLandscapePin || bHasLandscapeHasInput;
 }
 
 #endif // WITH_EDITOR
diff --git a/Engine/Plugins/Experimental/PCG/Source/PCG/Private/PCGGraph.cpp b/Engine/Plugins/Experimental/PCG/Source/PCG/Private/PCGGraph.cpp
--- a/Engine/Plugins/Experimental/PCG/Source/PCG/Private/PCGGraph.cpp
+++ b/Engine/Plugins/Experimental/PCG/Source/PCG/Private/PCGGraph.cpp
@@ -788,16 +788,16 @@ void UPCGGraph::SetExtraEditorNodes(const TArray<TObjectPtr<const UObject>>& InN
 	}
 }
 
-FPCGTagToSettingsMap UPCGGraph::GetTrackedTagsToSettings() const
+FPCGActorSelectionKeyToSettingsMap UPCGGraph::GetTrackedActorKeysToSettings() const
 {
-	FPCGTagToSettingsMap TagsToSettings;
+	FPCGActorSelectionKeyToSettingsMap TagsToSettings;
 	TArray<TObjectPtr<const UPCGGraph>> VisitedGraphs;
 
-	GetTrackedTagsToSettings(TagsToSettings, VisitedGraphs);
+	GetTrackedActorKeysToSettings(TagsToSettings, VisitedGraphs);
 	return TagsToSettings;
 }
 
-void UPCGGraph::GetTrackedTagsToSettings(FPCGTagToSettingsMap& OutTagsToSettings, TArray<TObjectPtr<const UPCGGraph>>& OutVisitedGraphs) const
+void UPCGGraph::GetTrackedActorKeysToSettings(FPCGActorSelectionKeyToSettingsMap& OutTagsToSettings, TArray<TObjectPtr<const UPCGGraph>>& OutVisitedGraphs) const
 {
 	if (OutVisitedGraphs.Contains(this))
 	{
@@ -810,7 +810,7 @@ void UPCGGraph::GetTrackedTagsToSettings(FPCGTagToSettingsMap& OutTagsToSettings
 	{
 		if (Node && Node->GetSettings())
 		{
-			Node->GetSettings()->GetTrackedActorTags(OutTagsToSettings, OutVisitedGraphs);
+			Node->GetSettings()->GetTrackedActorKeys(OutTagsToSettings, OutVisitedGraphs);
 		}
 	}
 }
diff --git a/Engine/Plugins/Experimental/PCG/Source/PCG/Private/PCGSubgraph.cpp b/Engine/Plugins/Experimental/PCG/Source/PCG/Private/PCGSubgraph.cpp
--- a/Engine/Plugins/Experimental/PCG/Source/PCG/Private/PCGSubgraph.cpp
+++ b/Engine/Plugins/Experimental/PCG/Source/PCG/Private/PCGSubgraph.cpp
@@ -122,11 +122,11 @@ void UPCGBaseSubgraphSettings::PostEditChangeProperty(struct FPropertyChangedEve
 	}
 }
 
-void UPCGBaseSubgraphSettings::GetTrackedActorTags(FPCGTagToSettingsMap& OutTagToSettings, TArray<TObjectPtr<const UPCGGraph>>& VisitedGraphs) const
+void UPCGBaseSubgraphSettings::GetTrackedActorKeys(FPCGActorSelectionKeyToSettingsMap& OutKeysToSettings, TArray<TObjectPtr<const UPCGGraph>>& VisitedGraphs) const
 {
 	if (UPCGGraph* Subgraph = GetSubgraph())
 	{
-		Subgraph->GetTrackedTagsToSettings(OutTagToSettings, VisitedGraphs);
+		Subgraph->GetTrackedActorKeysToSettings(OutKeysToSettings, VisitedGraphs);
 	}
 }
 
diff --git a/Engine/Plugins/Experimental/PCG/Source/PCG/Private/PCGSubsystem.cpp b/Engine/Plugins/Experimental/PCG/Source/PCG/Private/PCGSubsystem.cpp
--- a/Engine/Plugins/Experimental/PCG/Source/PCG/Private/PCGSubsystem.cpp
+++ b/Engine/Plugins/Experimental/PCG/Source/PCG/Private/PCGSubsystem.cpp
@@ -111,8 +111,7 @@ UPCGSubsystem::UPCGSubsystem()
 #if WITH_EDITOR
 	if (!HasAnyFlags(RF_ClassDefaultObject))
 	{
-		// TODO: Uncomment when tracking is functionnal
-		// ActorAndComponentMapping.RegisterTrackingCallbacks();
+		ActorAndComponentMapping.RegisterTrackingCallbacks();
 	}
 #endif // WITH_EDITOR
 }
@@ -128,8 +127,7 @@ void UPCGSubsystem::Deinitialize()
 	bHasTickedOnce = false;
 
 #if WITH_EDITOR
-	// TODO: Uncomment when tracking is functionnal
-	// ActorAndComponentMapping.TeardownTrackingCallbacks();
+	ActorAndComponentMapping.TeardownTrackingCallbacks();
 #endif // WITH_EDITOR
 
 	Super::Deinitialize();
@@ -155,6 +153,11 @@ void UPCGSubsystem::PostInitialize()
 			});
 		}
 	}
+
+#if WITH_EDITOR
+	// When the PCGWorldActor is loaded, we can add delayed actors.
+	ActorAndComponentMapping.AddActorsPostInit();
+#endif // WITH_EDITOR
 }
 
 UPCGSubsystem* UPCGSubsystem::GetInstance(UWorld* World)
diff --git a/Engine/Plugins/Experimental/PCG/Source/PCG/Public/Elements/PCGActorSelector.h b/Engine/Plugins/Experimental/PCG/Source/PCG/Public/Elements/PCGActorSelector.h
--- a/Engine/Plugins/Experimental/PCG/Source/PCG/Public/Elements/PCGActorSelector.h
+++ b/Engine/Plugins/Experimental/PCG/Source/PCG/Public/Elements/PCGActorSelector.h
@@ -3,6 +3,7 @@
 #pragma once
 
 #include "Templates/SubclassOf.h"
+#include "UObject/NameTypes.h"
 
 #include "PCGActorSelector.generated.h"
 
@@ -19,6 +20,25 @@ enum class EPCGActorSelection : uint8
 	ByClass
 };
 
+struct FPCGActorSelectionKey
+{
+	FPCGActorSelectionKey() = default;
+
+	explicit FPCGActorSelectionKey(FName InTag);
+	explicit FPCGActorSelectionKey(TSubclassOf<AActor> InSelectionClass);
+
+	bool operator==(const FPCGActorSelectionKey& InOther) const;
+
+	friend uint32 GetTypeHash(const FPCGActorSelectionKey& In);
+	bool IsMatching(const AActor* InActor) const;
+
+	static TArray<FPCGActorSelectionKey> GenerateAllKeysForActor(const AActor* InActor);
+
+	EPCGActorSelection Selection = EPCGActorSelection::ByTag;
+	FName Tag = NAME_None;
+	TSubclassOf<AActor> ActorSelectionClass = nullptr;
+};
+
 UENUM()
 enum class EPCGActorFilter : uint8
 {
@@ -91,6 +111,9 @@ struct FPCGActorSelectorSettings
 	FText GetTaskNameSuffix() const;
 	FName GetTaskName(const FText& Prefix) const;
 #endif
+
+	FPCGActorSelectionKey GetAssociatedKey() const;
+	static FPCGActorSelectorSettings ReconstructFromKey(const FPCGActorSelectionKey& InKey);
 };
 
 namespace PCGActorSelector
diff --git a/Engine/Plugins/Experimental/PCG/Source/PCG/Public/Elements/PCGDataFromActor.h b/Engine/Plugins/Experimental/PCG/Source/PCG/Public/Elements/PCGDataFromActor.h
--- a/Engine/Plugins/Experimental/PCG/Source/PCG/Public/Elements/PCGDataFromActor.h
+++ b/Engine/Plugins/Experimental/PCG/Source/PCG/Public/Elements/PCGDataFromActor.h
@@ -32,7 +32,7 @@ public:
 	virtual FText GetDefaultNodeTitle() const override { return NSLOCTEXT("PCGDataFromActorSettings", "NodeTitle", "Get Actor Data"); }
 	virtual FText GetNodeTooltipText() const override;
 	virtual EPCGSettingsType GetType() const override { return EPCGSettingsType::Spatial; }
-	virtual void GetTrackedActorTags(FPCGTagToSettingsMap& OutTagToSettings, TArray<TObjectPtr<const UPCGGraph>>& OutVisitedGraphs) const override;
+	virtual void GetTrackedActorKeys(FPCGActorSelectionKeyToSettingsMap& OutKeysToSettings, TArray<TObjectPtr<const UPCGGraph>>& OutVisitedGraphs) const override;
 	virtual bool HasDynamicPins() const override { return true; }
 #endif
 	virtual EPCGDataType GetCurrentPinTypes(const UPCGPin* InPin) const override;
diff --git a/Engine/Plugins/Experimental/PCG/Source/PCG/Public/Elements/PCGExecuteBlueprint.h b/Engine/Plugins/Experimental/PCG/Source/PCG/Public/Elements/PCGExecuteBlueprint.h
--- a/Engine/Plugins/Experimental/PCG/Source/PCG/Public/Elements/PCGExecuteBlueprint.h
+++ b/Engine/Plugins/Experimental/PCG/Source/PCG/Public/Elements/PCGExecuteBlueprint.h
@@ -238,7 +238,7 @@ public:
 	virtual FText GetDefaultNodeTitle() const override { return NSLOCTEXT("PCGBlueprintSettings", "NodeTitle", "Execute Blueprint"); }
 	virtual FLinearColor GetNodeTitleColor() const override;
 	virtual EPCGSettingsType GetType() const override;
-	virtual void GetTrackedActorTags(FPCGTagToSettingsMap& OutTagToSettings, TArray<TObjectPtr<const UPCGGraph>>& OutVisitedGraphs) const override;
+	virtual void GetTrackedActorKeys(FPCGActorSelectionKeyToSettingsMap& OutKeysToSettings, TArray<TObjectPtr<const UPCGGraph>>& OutVisitedGraphs) const override;
 	virtual UObject* GetJumpTargetForDoubleClick() const override;
 	virtual void ApplyDeprecationBeforeUpdatePins(UPCGNode* InOutNode, TArray<TObjectPtr<UPCGPin>>& InputPins, TArray<TObjectPtr<UPCGPin>>& OutputPins) override;
 	virtual TArray<FPCGPreConfiguredSettingsInfo> GetPreconfiguredInfo() const override;
diff --git a/Engine/Plugins/Experimental/PCG/Source/PCG/Public/Elements/PCGPropertyToParamData.h b/Engine/Plugins/Experimental/PCG/Source/PCG/Public/Elements/PCGPropertyToParamData.h
--- a/Engine/Plugins/Experimental/PCG/Source/PCG/Public/Elements/PCGPropertyToParamData.h
+++ b/Engine/Plugins/Experimental/PCG/Source/PCG/Public/Elements/PCGPropertyToParamData.h
@@ -32,7 +32,7 @@ public:
 	virtual FName GetDefaultNodeName() const override { return FName(TEXT("GetActorProperty")); }
 	virtual FText GetDefaultNodeTitle() const override { return NSLOCTEXT("PCGPropertyToParamDataSettings", "NodeTitle", "Get Actor Property"); }
 	virtual EPCGSettingsType GetType() const override { return EPCGSettingsType::Param; }
-	virtual void GetTrackedActorTags(FPCGTagToSettingsMap& OutTagToSettings, TArray<TObjectPtr<const UPCGGraph>>& OutVisitedGraphs) const override;
+	virtual void GetTrackedActorKeys(FPCGActorSelectionKeyToSettingsMap& OutKeysToSettings, TArray<TObjectPtr<const UPCGGraph>>& OutVisitedGraphs) const override;
 #endif
 
 	virtual FName AdditionalTaskName() const override;
diff --git a/Engine/Plugins/Experimental/PCG/Source/PCG/Public/Grid/PCGLandscapeCache.h b/Engine/Plugins/Experimental/PCG/Source/PCG/Public/Grid/PCGLandscapeCache.h
--- a/Engine/Plugins/Experimental/PCG/Source/PCG/Public/Grid/PCGLandscapeCache.h
+++ b/Engine/Plugins/Experimental/PCG/Source/PCG/Public/Grid/PCGLandscapeCache.h
@@ -129,8 +129,9 @@ private:
 #if WITH_EDITOR
 	void SetupLandscapeCallbacks();
 	void TeardownLandscapeCallbacks();
-	void OnLandscapeChanged(ALandscapeProxy* Landscape, const FLandscapeProxyComponentDataChangedParams& ChangeParams);
-	void CacheLayerNames(ALandscapeProxy* Landscape);
+	void OnLandscapeChanged(ALandscapeProxy* InLandscape, const FLandscapeProxyComponentDataChangedParams& InChangeParams);
+	void OnLandscapeMoved(AActor* InActor);
+	void CacheLayerNames(ALandscapeProxy* InLandscape);
 	void CacheLayerNames();
 #endif
 
diff --git a/Engine/Plugins/Experimental/PCG/Source/PCG/Public/PCGActorAndComponentMapping.h b/Engine/Plugins/Experimental/PCG/Source/PCG/Public/PCGActorAndComponentMapping.h
--- a/Engine/Plugins/Experimental/PCG/Source/PCG/Public/PCGActorAndComponentMapping.h
+++ b/Engine/Plugins/Experimental/PCG/Source/PCG/Public/PCGActorAndComponentMapping.h
@@ -3,6 +3,7 @@
 #pragma once
 
 #include "PCGCommon.h"
+#include "Elements/PCGActorSelector.h"
 #include "Grid/PCGComponentOctree.h"
 
 #include "Containers/Array.h"
@@ -10,10 +11,13 @@
 #include "Containers/Set.h"
 #include "Templates/Function.h"
 #include "UObject/ObjectPtr.h"
+#include "UObject/SoftObjectPtr.h"
 #include "UObject/WeakObjectPtr.h"
 
 class AActor;
+class ALandscapeProxy;
 class APCGPartitionActor;
+class FLandscapeProxyComponentDataChangedParams;
 class UObject;
 class UPCGComponent;
 class UPCGSubsystem;
@@ -24,7 +28,7 @@ class UPCGSubsystem;
 * For now it is used for:
 * - Mapping between PCG Components and Partition actors: Dispatch tasks from original components to local components
 * - Tracking non partitioned PCG Components that has tracking actors needs.
-* - TODO: Tracking actors: Be able to react to actors changes, and potentially dirty/refresh affected components
+* - Tracking actors: Be able to react to actors changes, and potentially dirty/refresh affected components
 */
 class UPCGActorAndComponentMapping
 {
@@ -43,10 +47,10 @@ public:
 	void RegisterTrackingCallbacks();
 	void TeardownTrackingCallbacks();
 
-	void OnActorAdded(AActor* InActor);
-	void OnActorDeleted(AActor* InActor);
-	void OnActorMoved(AActor* InActor);
-	void OnObjectPropertyChanged(UObject* InObject, FPropertyChangedEvent& InEvent);
+	void AddActorsPostInit();
+
+	/** Will register/update tracking if a component was registered/updated. */
+	void RegisterOrUpdateTracking(UPCGComponent* InComponent, bool bInShouldDirtyActors);
 #endif // WITH_EDITOR
 
 	/** Register a new PCG Component or update it. Returns true if it was added/updated. Thread safe */
@@ -67,8 +71,11 @@ public:
 	/** Return a copy of all the registered partitioned components. Thread safe */
 	TSet<TObjectPtr<UPCGComponent>> GetAllRegisteredPartitionedComponents() const;
 
-	/** Check if a component is currently tracked. Thread safe */
-	bool IsComponentTracked(const UPCGComponent* InComponent) const;
+	/** Return a copy of all the registered non-partitioned components. Thread safe */
+	TSet<TObjectPtr<UPCGComponent>> GetAllRegisteredNonPartitionedComponents() const;
+
+	/** Return a copy of all the registered components. Thread safe */
+	TSet<TObjectPtr<UPCGComponent>> GetAllRegisteredComponents() const;
 
 private:
 	// This class is only meant to be used as part of the PCG Subsytem and owned by it.
@@ -101,8 +108,25 @@ private:
 	/** Delete the current mapping between a PCG component and its PCG Partition actors */
 	void DeleteMappingPCGComponentPartitionActor(UPCGComponent* InComponent);
 
-	/** Will register/update tracking if a component was registered/updated. */
-	void RegisterOrUpdateTracking(UPCGComponent* InComponent);
+#if WITH_EDITOR
+	/* Return true if something is still tracked or was just untracked. */
+	bool AddOrUpdateTrackedActor(AActor* InActor);
+
+	void RegisterActor(AActor* InActor);
+
+	/* Return true if the actor was tracked. */
+	bool UnregisterActor(AActor* InActor);
+
+	void OnActorAdded(AActor* InActor);
+	void OnActorDeleted(AActor* InActor);
+	void OnActorMoved(AActor* InActor);
+	void OnLandscapeChanged(ALandscapeProxy* InLandscape, const FLandscapeProxyComponentDataChangedParams& InChangeParams);
+	void OnPreObjectPropertyChanged(UObject* InObject, const FEditPropertyChain& InEditPropertyChain);
+	void OnObjectPropertyChanged(UObject* InObject, FPropertyChangedEvent& InEvent);
+
+	/** For deprecation only, register tracked actors from serialized components */
+	template <typename ActorContainer>
+	void AddSerializedTrackedActorsForDeprecation(const ActorContainer& InActors);
 
 	/** Remap the tracking in case of BP components. */
 	void RemapTracking(const UPCGComponent* InOldComponent, UPCGComponent* InNewComponent);
@@ -110,6 +134,10 @@ private:
 	/** Unregister tracking when a component is removed. */
 	void UnregisterTracking(UPCGComponent* InComponent);
 
+	/** Trigger an update when the actor changed. Can also provide a list of tags that has changed. */
+	void OnActorChanged(AActor* InActor, bool bInHasMoved);
+#endif // WITH_EDITOR
+
 private:
 	// Cached subsystem
 	UPCGSubsystem* PCGSubsystem = nullptr;
@@ -134,13 +162,43 @@ private:
 	/** Will hold all the components that are not partitioned (and not local) and are tracking something. Will be use to dispatch actor tracking updates. */
 	FPCGComponentOctreeAndMap NonPartitionedOctree;
 
-	/** Keep a mapping between tracked actors and the components that track them. */
-	TMap<TWeakObjectPtr<AActor>, TSet<TObjectPtr<UPCGComponent>>> TrackedActorsToComponentsMap;
+	/** Keep a mapping between tracked actors and the components that track them, and the tracking needs to be culled.*/
+	TMap<TWeakObjectPtr<AActor>, TSet<TObjectPtr<UPCGComponent>>> CulledTrackedActorsToComponentsMap;
+
+	/** Same mapping but for always tracked actors */
+	TMap<TWeakObjectPtr<AActor>, TSet<TObjectPtr<UPCGComponent>>> AlwaysTrackedActorsToComponentsMap;
+
+	/** Keep the list of the keys already tracked, to avoid requerying the actors everytime */
+	TMap<FPCGActorSelectionKey, TSet<TObjectPtr<UPCGComponent>>> KeysToComponentsMap;
+
+	/** Finally keep a mapping between actors and their position to know if an actor move in/out of a component tracking bounds. Only kept for actors that need to be culled. */
+	TMap<TWeakObjectPtr<AActor>, FBox> TrackedActorToPositionMap;
 
-	/** Also keep the reverse mapping. */
-	TMap<TObjectPtr<UPCGComponent>, TSet<TWeakObjectPtr<AActor>>> TrackedComponentsToActorsMap;
 	mutable FRWLock TrackedComponentsLock;
 
+	TSet<TWeakObjectPtr<AActor>> DelayedAddedActors;
+
 	/** Keep a mapping between tracked actors and their dependencies. */
 	TMap<TWeakObjectPtr<AActor>, TSet<TObjectPtr<UObject>>> TrackedActorsToDependenciesMap;
-};
\ No newline at end of file
+
+	/** Transient list of tags, kept when there is a tag change on a tracked Actor. */
+	TSet<FName> TempTrackedActorTags;
+};
+
+#if WITH_EDITOR
+template <typename ActorContainer>
+void UPCGActorAndComponentMapping::AddSerializedTrackedActorsForDeprecation(const ActorContainer& InActors)
+{
+	using PointerType = typename ActorContainer::ElementType;
+	using ElementType = typename PointerType::ElementType;
+	static_assert(std::is_base_of_v<AActor, ElementType>);
+
+	for (auto ActorPtr : InActors)
+	{
+		if (AActor* Actor = ActorPtr.Get())
+		{
+			OnActorAdded(Actor);
+		}
+	}
+}
+#endif // WITH_EDITOR
\ No newline at end of file
diff --git a/Engine/Plugins/Experimental/PCG/Source/PCG/Public/PCGComponent.h b/Engine/Plugins/Experimental/PCG/Source/PCG/Public/PCGComponent.h
--- a/Engine/Plugins/Experimental/PCG/Source/PCG/Public/PCGComponent.h
+++ b/Engine/Plugins/Experimental/PCG/Source/PCG/Public/PCGComponent.h
@@ -17,6 +17,7 @@ namespace EEndPlayReason { enum Type : int; }
 class APCGPartitionActor;
 struct FPCGContext;
 class FPCGStack;
+class UPCGActorAndComponentMapping;
 class UPCGComponent;
 class UPCGGraph;
 class UPCGGraphInterface;
@@ -71,6 +72,7 @@ class PCG_API UPCGComponent : public UActorComponent
 
 	friend class UPCGManagedActors;
 	friend class UPCGSubsystem;
+	friend class UPCGActorAndComponentMapping;
 
 public:
 	/** ~Begin UObject interface */
@@ -87,7 +89,6 @@ public:
 	//~Begin UActorComponent Interface
 	virtual void BeginPlay() override;
 	virtual void EndPlay(const EEndPlayReason::Type EndPlayReason) override;
-	virtual void OnComponentCreated() override;
 	virtual void OnComponentDestroyed(bool bDestroyingHierarchy) override;
 	virtual void OnRegister() override;
 
@@ -248,8 +249,7 @@ public:
 	void StoreInspectionData(const FPCGStack* InStack, const UPCGNode* InNode, const FPCGDataCollection& InInputData, const FPCGDataCollection& InOutputData);
 	const FPCGDataCollection* GetInspectionData(const FString& InStackPath) const;
 
-	/** Used by the tracking system to know if the component need to track actors. Not enabled for now, tracking is still done on the component.*/
-	bool ShouldTrackActors() const { return false; }
+	bool IsActorTracked(AActor* InActor, bool& bOutIsCulled) const;
 
 	/** Know if we need to force a generation, in case of BP added to the world in editor */
 	bool ShouldGenerateBPPCGAddedToWorld() const;
@@ -329,8 +329,6 @@ private:
 	void GetManagedResources(TArray<TObjectPtr<UPCGManagedResource>>& Resources) const;
 	void SetManagedResources(const TArray<TObjectPtr<UPCGManagedResource>>& Resources);
 
-	bool GetActorsFromTags(const TMap<FName, bool>& InTagsAndCulling, TSet<TWeakObjectPtr<AActor>>& OutActors);
-
 	void RefreshAfterGraphChanged(UPCGGraphInterface* InGraph, bool bIsStructural, bool bDirtyInputs);
 	void OnGraphChanged(UPCGGraphInterface* InGraph, EPCGChangeType ChangeType);
 
@@ -342,29 +340,13 @@ private:
 	/** Sets up actor, tracking, landscape and graph callbacks */
 	void SetupCallbacksOnCreation();
 
-	void SetupActorCallbacks();
-	void TeardownActorCallbacks();
-	void SetupTrackingCallbacks();
-	void TeardownTrackingCallbacks();
-	void RefreshTrackingData();
+	void UpdateTrackingCache();
 
-	void OnActorAdded(AActor* InActor);
-	void OnActorDeleted(AActor* InActor);
-	void OnActorMoved(AActor* InActor);
-	void OnObjectPropertyChanged(UObject* InObject, FPropertyChangedEvent& InEvent);
 	bool ActorIsTracked(AActor* InActor) const;
+	
+	bool DirtyTrackedActor(AActor* InActor, bool bIntersect, const TSet<FName>& InRemovedTags);
 
-	void OnActorChanged(AActor* InActor, UObject* InSourceObject, bool bActorTagChange);
-
-	bool PopulateTrackedActorToTagsMap(bool bForce = false);
-	bool AddTrackedActor(AActor* InActor, bool bForce = false);
-	bool RemoveTrackedActor(AActor* InActor);
-	bool UpdateTrackedActor(AActor* InActor);
-	bool DirtyTrackedActor(AActor* InActor);
-	bool DirtyCacheFromTag(const FName& InTag, const AActor* InActor, bool bIgnoreCull = false);
-	void DirtyCacheForAllTrackedTags();
-
-	bool GraphUsesLandscapePin() const;
+	bool ShouldTrackLandscape() const;
 #endif
 
 	FBox GetGridBounds(const AActor* InActor) const;
@@ -420,29 +402,20 @@ private:
 #if WITH_EDITOR
 	int32 InspectionCounter = 0;
 	FBox LastGeneratedBoundsPriorToUndo = FBox(EForceInit::ForceInit);
-	FPCGTagToSettingsMap CachedTrackedTagsToSettings;
+	FPCGActorSelectionKeyToSettingsMap CachedTrackedKeysToSettings;
 
-	TMap<FName, bool> CachedTrackedTagsToCulling;
-
-	void SetupLandscapeTracking();
-	void TeardownLandscapeTracking();
-	void UpdateTrackedLandscape(bool bBoundsCheck = true);
-	void OnLandscapeChanged(ALandscapeProxy* Landscape, const FLandscapeProxyComponentDataChangedParams& ChangeParams);
+	TMap<FPCGActorSelectionKey, bool> CachedTrackedKeysToCulling;
 #endif
 
 #if WITH_EDITORONLY_DATA
 	UPROPERTY()
-	TArray<TSoftObjectPtr<ALandscapeProxy>> TrackedLandscapes;
+	TArray<TSoftObjectPtr<ALandscapeProxy>> TrackedLandscapes_DEPRECATED;
 #endif
 
 #if WITH_EDITORONLY_DATA
 	// Cached tracked actors list is serialized because we can't get it at postload time
 	UPROPERTY()
-	TSet<TWeakObjectPtr<AActor>> CachedTrackedActors;
-
-	TMap<TWeakObjectPtr<AActor>, TSet<FName>> CachedTrackedActorToTags;
-	TMap<TWeakObjectPtr<AActor>, TSet<TObjectPtr<UObject>>> CachedTrackedActorToDependencies;
-	bool bActorToTagsMapPopulated = false;
+	TSet<TWeakObjectPtr<AActor>> CachedTrackedActors_DEPRECATED;
 
 	UPROPERTY(Transient)
 	TMap<FString, FPCGDataCollection> InspectionCache;
diff --git a/Engine/Plugins/Experimental/PCG/Source/PCG/Public/PCGGraph.h b/Engine/Plugins/Experimental/PCG/Source/PCG/Public/PCGGraph.h
--- a/Engine/Plugins/Experimental/PCG/Source/PCG/Public/PCGGraph.h
+++ b/Engine/Plugins/Experimental/PCG/Source/PCG/Public/PCGGraph.h
@@ -223,8 +223,8 @@ public:
 #endif
 
 #if WITH_EDITOR
-	FPCGTagToSettingsMap GetTrackedTagsToSettings() const;
-	void GetTrackedTagsToSettings(FPCGTagToSettingsMap& OutTagsToSettings, TArray<TObjectPtr<const UPCGGraph>>& OutVisitedGraphs) const;
+	FPCGActorSelectionKeyToSettingsMap GetTrackedActorKeysToSettings() const;
+	void GetTrackedActorKeysToSettings(FPCGActorSelectionKeyToSettingsMap& OutKeysToSettings, TArray<TObjectPtr<const UPCGGraph>>& OutVisitedGraphs) const;
 #endif
 
 	/** Size of grid on which this node should be executed. Nodes execute at the minimum of all input grid sizes. */
diff --git a/Engine/Plugins/Experimental/PCG/Source/PCG/Public/PCGSettings.h b/Engine/Plugins/Experimental/PCG/Source/PCG/Public/PCGSettings.h
--- a/Engine/Plugins/Experimental/PCG/Source/PCG/Public/PCGSettings.h
+++ b/Engine/Plugins/Experimental/PCG/Source/PCG/Public/PCGSettings.h
@@ -5,6 +5,7 @@
 #include "PCGData.h"
 #include "PCGDebug.h"
 #include "PCGElement.h"
+#include "Elements/PCGActorSelector.h"
 #include "Tests/Determinism/PCGDeterminismSettings.h"
 
 #include "PCGSettings.generated.h"
@@ -19,7 +20,7 @@ class UPCGNode;
 class UPCGSettings;
 
 using FPCGSettingsAndCulling = TPair<TWeakObjectPtr<const UPCGSettings>, bool>;
-using FPCGTagToSettingsMap = TMap<FName, TSet<FPCGSettingsAndCulling>>;
+using FPCGActorSelectionKeyToSettingsMap = TMap<FPCGActorSelectionKey, TArray<FPCGSettingsAndCulling>>;
 
 UENUM()
 enum class EPCGSettingsExecutionMode : uint8
@@ -258,7 +259,8 @@ public:
 	virtual bool GetPinExtraIcon(const UPCGPin* InPin, FName& OutExtraIcon, FText& OutTooltip) const { return false; }
 
 	/** Derived classes must implement this to communicate dependencies on external actors */
-	virtual void GetTrackedActorTags(FPCGTagToSettingsMap& OutTagToSettings, TArray<TObjectPtr<const UPCGGraph>>& OutVisitedGraphs) const {}
+	virtual void GetTrackedActorKeys(FPCGActorSelectionKeyToSettingsMap& OutKeysToSettings, TArray<TObjectPtr<const UPCGGraph>>& OutVisitedGraphs) const {}
+
 	/** Override this class to provide an UObject to jump to in case of double click on node
 	 *  ie. returning a blueprint instance will open the given blueprint in its editor.
 	 *  By default, it will return the underlying class, to try to jump to its header in code
diff --git a/Engine/Plugins/Experimental/PCG/Source/PCG/Public/PCGSubgraph.h b/Engine/Plugins/Experimental/PCG/Source/PCG/Public/PCGSubgraph.h
--- a/Engine/Plugins/Experimental/PCG/Source/PCG/Public/PCGSubgraph.h
+++ b/Engine/Plugins/Experimental/PCG/Source/PCG/Public/PCGSubgraph.h
@@ -40,7 +40,7 @@ protected:
 	//~End UObject interface implementation
 
 	//~Begin UPCGSettings interface
-	virtual void GetTrackedActorTags(FPCGTagToSettingsMap& OutTagToSettings, TArray<TObjectPtr<const UPCGGraph>>& OutVisitedGraphs) const override;
+	virtual void GetTrackedActorKeys(FPCGActorSelectionKeyToSettingsMap& OutKeysToSettings, TArray<TObjectPtr<const UPCGGraph>>& OutVisitedGraphs) const override;
 	virtual bool IsStructuralProperty(const FName& InPropertyName) const override;
 #endif
 
diff --git a/Engine/Plugins/Experimental/PCG/Source/PCG/Public/PCGSubsystem.h b/Engine/Plugins/Experimental/PCG/Source/PCG/Public/PCGSubsystem.h
--- a/Engine/Plugins/Experimental/PCG/Source/PCG/Public/PCGSubsystem.h
+++ b/Engine/Plugins/Experimental/PCG/Source/PCG/Public/PCGSubsystem.h
@@ -167,6 +167,15 @@ public:
 	/** Propagate to the graph compiler graph changes */
 	void NotifyGraphChanged(UPCGGraph* InGraph);
 
+	/** Update the tracking on a given component. */
+	void UpdateComponentTracking(UPCGComponent* InComponent, bool bInShouldDirtyActors = false) { ActorAndComponentMapping.RegisterOrUpdateTracking(InComponent, bInShouldDirtyActors); }
+
+#if WITH_EDITOR
+	/** For deprecation only, register tracked actors from serialized components */
+	template <typename ActorContainer>
+	void AddSerializedTrackedActorsForDeprecation(const ActorContainer& InActors) { ActorAndComponentMapping.AddSerializedTrackedActorsForDeprecation(InActors); }
+#endif // WITH_EDITOR
+
 	/** Cleans up the graph cache on an element basis. InSettings is used for debugging and is optional. */
 	void CleanFromCache(const IPCGElement* InElement, const UPCGSettings* InSettings = nullptr);
 
diff --git a/Engine/Plugins/Experimental/PCG/Source/PCG/Public/PCGWorldActor.h b/Engine/Plugins/Experimental/PCG/Source/PCG/Public/PCGWorldActor.h
--- a/Engine/Plugins/Experimental/PCG/Source/PCG/Public/PCGWorldActor.h
+++ b/Engine/Plugins/Experimental/PCG/Source/PCG/Public/PCGWorldActor.h
@@ -59,6 +59,12 @@ public:
 	UPROPERTY(config, EditAnywhere, Category = GenerationSettings)
 	bool bUse2DGrid = true;
 
+#if WITH_EDITORONLY_DATA
+	// Cached tracked actors list is serialized because we can't get it at postload time
+	UPROPERTY()
+	TSet<TWeakObjectPtr<AActor>> CachedTrackedActors;
+#endif // WITH_EDITORONLY_DATA
+
 private:
 	void RegisterToSubsystem();
 	void UnregisterFromSubsystem();
