Fixed missing component replacement code

diff --git a/Engine/Plugins/Experimental/PCG/Source/PCG/Private/PCGActorAndComponentMapping.cpp b/Engine/Plugins/Experimental/PCG/Source/PCG/Private/PCGActorAndComponentMapping.cpp
--- a/Engine/Plugins/Experimental/PCG/Source/PCG/Private/PCGActorAndComponentMapping.cpp
+++ b/Engine/Plugins/Experimental/PCG/Source/PCG/Private/PCGActorAndComponentMapping.cpp
@@ -660,7 +660,19 @@ void UPCGActorAndComponentMapping::RegisterOrUpdateTracking(UPCGComponent* InCom
 
 void UPCGActorAndComponentMapping::RemapTracking(const UPCGComponent* InOldComponent, UPCGComponent* InNewComponent)
 {
-	// Nothing to do?
+	auto ReplaceInMap = [InOldComponent, InNewComponent](auto& InMap)
+	{
+		for (auto& It : InMap)
+		{
+			if (It.Value.Remove(InOldComponent) > 0)
+			{
+				It.Value.Add(InNewComponent);
+			}
+		}
+	};
+
+	ReplaceInMap(CulledTrackedActorsToComponentsMap);
+	ReplaceInMap(AlwaysTrackedActorsToComponentsMap);
 }
 
 void UPCGActorAndComponentMapping::UnregisterTracking(UPCGComponent* InComponent)
