[PCG] New Intersection Node

* Old "Intersection" node renamed "Inner Intersection"
* New "Intersection" node now supports dynamic source inputs
* Primary inputs will be iterated on
* Secondary inputs are implicitly merged
* User can optionally toggle whether to ignore empty secondary inputs

[DEPREC] Previous versions should continue to utilize the previous version

diff --git a/Engine/Plugins/Experimental/PCG/Config/BasePCG.ini b/Engine/Plugins/Experimental/PCG/Config/BasePCG.ini
--- a/Engine/Plugins/Experimental/PCG/Config/BasePCG.ini
+++ b/Engine/Plugins/Experimental/PCG/Config/BasePCG.ini
@@ -31,5 +31,6 @@
 +ClassRedirects=(OldName="PCGInstancePackerBase",NewName="/Script/PCG.PCGInstanceDataPackerBase")
 +ClassRedirects=(OldName="PCGInstancePackerByAttribute",NewName="/Script/PCG.PCGInstanceDataPackerByAttribute")
 +ClassRedirects=(OldName="PCGInstancePackerByRegex",NewName="/Script/PCG.PCGInstanceDataPackerByRegex")
++ClassRedirects=(OldName="PCGIntersectionSettings",NewName="/Script/PCG.PCGInnerIntersectionSettings")
 
 +EnumRedirects=(OldName="EPCGDensityNoiseMode",NewName="EPCGAttributeNoiseMode")
diff --git a/Engine/Plugins/Experimental/PCG/Source/PCG/Private/Elements/PCGInnerIntersectionElement.cpp b/Engine/Plugins/Experimental/PCG/Source/PCG/Private/Elements/PCGInnerIntersectionElement.cpp
new file mode 100644
--- /dev/null
+++ b/Engine/Plugins/Experimental/PCG/Source/PCG/Private/Elements/PCGInnerIntersectionElement.cpp
@@ -0,0 +1,106 @@
+// Copyright Epic Games, Inc. All Rights Reserved.
+
+#include "Elements/PCGInnerIntersectionElement.h"
+#include "Data/PCGSpatialData.h"
+#include "PCGContext.h"
+#include "PCGPin.h"
+
+#include UE_INLINE_GENERATED_CPP_BY_NAME(PCGInnerIntersectionElement)
+
+#define LOCTEXT_NAMESPACE "PCGInnerIntersectionElement"
+
+#if WITH_EDITOR
+FText UPCGInnerIntersectionSettings::GetNodeTooltipText() const
+{
+	return LOCTEXT("NodeTooltipText", "Spatial data will be generated as the result of intersecting with the other source inputs sequentially or no output if such an intersection does not exist. \nSee also: Intersection Node");
+}
+#endif // WITH_EDITOR
+
+TArray<FPCGPinProperties> UPCGInnerIntersectionSettings::InputPinProperties() const
+{
+	TArray<FPCGPinProperties> PinProperties;
+	FPCGPinProperties& SourcePinProperty = PinProperties.Emplace_GetRef(PCGPinConstants::DefaultInputLabel, EPCGDataType::Spatial);
+
+#if WITH_EDITOR
+	SourcePinProperty.Tooltip = LOCTEXT("SourcePinTooltip", "Source spatial data from which to conduct the intersection. Empty spatial data will be ignored.");
+#endif // WITH_EDITOR
+	
+	return PinProperties;
+}
+
+TArray<FPCGPinProperties> UPCGInnerIntersectionSettings::OutputPinProperties() const
+{
+	TArray<FPCGPinProperties> PinProperties;
+	FPCGPinProperties& OutputPinProperty = PinProperties.Emplace_GetRef(PCGPinConstants::DefaultOutputLabel, EPCGDataType::Spatial);
+
+#if WITH_EDITOR
+	OutputPinProperty.Tooltip = LOCTEXT("OutputPinTooltip", "The intersection created from all the source input data.");
+#endif // WITH_EDITOR
+	
+	return PinProperties;
+}
+
+FPCGElementPtr UPCGInnerIntersectionSettings::CreateElement() const
+{
+	return MakeShared<FPCGInnerIntersectionElement>();
+}
+
+bool FPCGInnerIntersectionElement::ExecuteInternal(FPCGContext* Context) const
+{
+	TRACE_CPUPROFILER_EVENT_SCOPE(FPCGInnerIntersectionElement::Execute);
+
+	const UPCGInnerIntersectionSettings* Settings = Context->GetInputSettings<UPCGInnerIntersectionSettings>();
+	check(Settings);
+
+	TArray<FPCGTaggedData> Inputs = Context->InputData.GetInputs();
+
+	const EPCGIntersectionDensityFunction DensityFunction = Settings->DensityFunction;
+#if WITH_EDITOR
+	const bool bKeepZeroDensityPoints = Settings->bKeepZeroDensityPoints;
+#else
+	const bool bKeepZeroDensityPoints = false;
+#endif
+
+	TArray<FPCGTaggedData>& Outputs = Context->OutputData.TaggedData;
+
+	const UPCGSpatialData* FirstSpatialData = nullptr;
+	UPCGIntersectionData* IntersectionData = nullptr;
+	int32 IntersectionTaggedDataIndex = -1;
+
+	for (FPCGTaggedData& Input : Inputs)
+	{
+		const UPCGSpatialData* SpatialData = Cast<UPCGSpatialData>(Input.Data);
+
+		// Non-spatial data we're not going to touch
+		if (!SpatialData)
+		{
+			Outputs.Add(Input);
+			continue;
+		}
+
+		if (!FirstSpatialData)
+		{
+			FirstSpatialData = SpatialData;
+			IntersectionTaggedDataIndex = Outputs.Num();
+			Outputs.Add(Input);
+			continue;
+		}
+
+		// Create a new intersection
+		IntersectionData = (IntersectionData ? IntersectionData : FirstSpatialData)->IntersectWith(SpatialData);
+		// Propagate settings
+		IntersectionData->DensityFunction = DensityFunction;
+#if WITH_EDITOR
+		IntersectionData->bKeepZeroDensityPoints = bKeepZeroDensityPoints;
+#endif
+
+		// Update tagged data
+		FPCGTaggedData& IntersectionTaggedData = Outputs[IntersectionTaggedDataIndex];
+		IntersectionTaggedData.Data = IntersectionData;
+		IntersectionTaggedData.Tags.Append(Input.Tags);
+	}
+
+	return true;
+}
+
+#undef LOCTEXT_NAMESPACE
\ No newline at end of file
diff --git a/Engine/Plugins/Experimental/PCG/Source/PCG/Private/Elements/PCGIntersectionElement.cpp b/Engine/Plugins/Experimental/PCG/Source/PCG/Private/Elements/PCGIntersectionElement.cpp
deleted file mode 100644
diff --git a/Engine/Plugins/Experimental/PCG/Source/PCG/Private/Elements/PCGOuterIntersectionElement.cpp b/Engine/Plugins/Experimental/PCG/Source/PCG/Private/Elements/PCGOuterIntersectionElement.cpp
new file mode 100644
--- /dev/null
+++ b/Engine/Plugins/Experimental/PCG/Source/PCG/Private/Elements/PCGOuterIntersectionElement.cpp
@@ -0,0 +1,171 @@
+// Copyright Epic Games, Inc. All Rights Reserved.
+
+#include "Elements/PCGOuterIntersectionElement.h"
+
+#include "PCGContext.h"
+#include "PCGPin.h"
+#include "PCGSettingsWithDynamicInputs.h"
+#include "Data/PCGSpatialData.h"
+
+#define LOCTEXT_NAMESPACE "PCGOuterIntersectionElement"
+
+namespace PCGIntersectionConstants
+{
+	const FName PrimaryLabel = TEXT("Primary Source");
+	const FName SecondaryLabel = TEXT("Source");
+	const FText SecondaryTooltip = LOCTEXT("SecondaryPinLabel", "Secondary pin inputs will be implicitly unioned together before being compared to the primary pin's input for calculation of the intersection operation. Empty data passed along on one of the secondary pins will result in an empty intersection output, unless the 'Ignore Empty Secondary Input' flag is enabled.");
+}
+
+TArray<FPCGPinProperties> UPCGOuterIntersectionSettings::StaticInputPinProperties() const
+{
+	TArray<FPCGPinProperties> StaticPinProperties;
+	FPCGPinProperties& PrimaryPinProperties = StaticPinProperties.Emplace_GetRef(PCGIntersectionConstants::PrimaryLabel, EPCGDataType::Spatial);
+	FPCGPinProperties& SecondaryPinProperties = StaticPinProperties.Emplace_GetRef(FName(PCGIntersectionConstants::SecondaryLabel.ToString() + FString(TEXT("1"))), EPCGDataType::Spatial);
+
+#if WITH_EDITOR
+	PrimaryPinProperties.Tooltip = LOCTEXT("PrimaryPinTooltip", "Each input on the primary pin will be used to calculate the output of the intersection operation. Assuming input data is provided, an output result will be generated for each input on this pin.");
+	SecondaryPinProperties.Tooltip = PCGIntersectionConstants::SecondaryTooltip;
+#endif // WITH_EDITOR
+
+	return StaticPinProperties;
+}
+
+#if WITH_EDITOR
+void UPCGOuterIntersectionSettings::AddDefaultDynamicInputPin()
+{
+	// Skip "Source 0", and also "Source 1" is always active, thus the label will be count+2
+	FPCGPinProperties SecondaryPinProperties(FName(GetDynamicInputPinsBaseLabel().ToString() + FString::FromInt(DynamicInputPinProperties.Num() + 2)), EPCGDataType::Spatial);
+	
+#if WITH_EDITOR
+	SecondaryPinProperties.Tooltip = PCGIntersectionConstants::SecondaryTooltip;
+#endif // WITH_EDITOR
+
+	AddDynamicInputPin(std::move(SecondaryPinProperties));
+}
+#endif // WITH_EDITOR
+
+TArray<FPCGPinProperties> UPCGOuterIntersectionSettings::OutputPinProperties() const
+{
+	TArray<FPCGPinProperties> PinProperties;
+	FPCGPinProperties& OutputPinProperty = PinProperties.Emplace_GetRef(PCGPinConstants::DefaultOutputLabel, EPCGDataType::Spatial);
+
+#if WITH_EDITOR
+	//Source spatial data from which to construct the intersection. Incoming edges providing no data will be ignored
+	OutputPinProperty.Tooltip = LOCTEXT("OutputPinTooltip", "The intersection created from all the source input data. One output will be generated for every input on the primary input pin.");
+#endif // WITH_EDITOR
+	
+	return PinProperties;
+}
+
+#if WITH_EDITOR
+FText UPCGOuterIntersectionSettings::GetNodeTooltipText() const
+{
+	return LOCTEXT("NodeTooltipText", "For each given source input on the primary pin, spatial data will be generated as the result of sequentially intersecting with the other source inputs (implicitly unioned), should an intersection exist. Additional pins maybe dynamically added and for each of these, all of the inputs into the same pin will be 'unioned' together automatically. Source pins receiving no or empty data will logically return an empty output, unless the 'Ignore Empty Secondary Input' flag has been set to 'true'.\nSee also: Inner Intersection Node, Union Node");
+}
+#endif // WITH_EDITOR
+
+FPCGElementPtr UPCGOuterIntersectionSettings::CreateElement() const
+{
+	return MakeShared<FPCGOuterIntersectionElement>();
+}
+
+FName UPCGOuterIntersectionSettings::GetDynamicInputPinsBaseLabel() const
+{
+	return PCGIntersectionConstants::SecondaryLabel;
+}
+
+bool FPCGOuterIntersectionElement::ExecuteInternal(FPCGContext* Context) const
+{
+	TRACE_CPUPROFILER_EVENT_SCOPE(FPCGOuterIntersectionElement::Execute);
+
+	const UPCGOuterIntersectionSettings* Settings = Context->GetInputSettings<UPCGOuterIntersectionSettings>();
+	check(Settings);
+	
+	const TArray<FPCGTaggedData> Inputs = Context->InputData.GetInputs();
+	const TArray<TObjectPtr<UPCGPin>>& InputPins = Context->Node->GetInputPins();
+	TArray<FPCGTaggedData>& Outputs = Context->OutputData.TaggedData;
+
+	// Early out for empty primary spatial inputs
+	if (Context->InputData.GetSpatialInputCountByPin(PCGIntersectionConstants::PrimaryLabel) == 0)
+	{
+		return true;
+	}
+
+	// If the user toggle is off, check for spatial pins that have no spatial data, for early out
+	if (!Settings->bIgnorePinsWithNoInput)
+	{
+		for (const UPCGPin* InputPin : InputPins)
+		{
+			check(InputPin);
+			// Only checking secondary source pins expecting spatial data
+			if (InputPin->Properties.Label != PCGIntersectionConstants::PrimaryLabel &&
+				!!(InputPin->Properties.AllowedTypes & EPCGDataType::Spatial) &&
+				Context->InputData.GetSpatialInputCountByPin(InputPin->Properties.Label) == 0)
+			{
+				UE_LOG(LogPCG, Verbose, TEXT("Intersection resulted in empty output due to no spatial data on a dynamic input pin."));
+				return true;
+			}
+		}
+	}
+
+	// Create unions for all the secondary source inputs with multiple connections/data
+	TArray<const UPCGSpatialData*> SecondarySourceUnionArray;
+	for (const UPCGPin* InputPin : InputPins)
+	{
+		const bool bIsPrimaryPin = InputPin->Properties.Label == PCGIntersectionConstants::PrimaryLabel;
+		
+		// Only checking pins expecting spatial data, not params, settings, etc
+		if (bIsPrimaryPin || !(InputPin->Properties.AllowedTypes & EPCGDataType::Spatial))
+		{
+			continue;
+		}
+
+		bool bDummy;
+		if (const UPCGSpatialData* UnionData = Context->InputData.GetSpatialUnionOfInputsByPin(InputPin->Properties.Label, bDummy))
+		{
+			SecondarySourceUnionArray.Push(UnionData);
+		}
+	}
+
+	if (SecondarySourceUnionArray.IsEmpty())
+	{
+		if (Settings->bIgnorePinsWithNoInput)
+		{
+			Outputs = Context->InputData.GetInputsByPin(PCGIntersectionConstants::PrimaryLabel);
+		}
+		
+		return true;
+	}
+	
+	// Iterate through primary inputs, and create intersection data from the unioned
+	for (const FPCGTaggedData& Input : Context->InputData.GetInputsByPin(PCGIntersectionConstants::PrimaryLabel))
+	{
+		UPCGSpatialData* PrimarySpatialData = Cast<UPCGSpatialData>(Input.Data);
+
+		// If this data isn't spatial, passthrough to the output
+		if (!PrimarySpatialData)
+		{
+			Outputs.Add(Input);
+			continue;
+		}
+		
+		UPCGSpatialData* IntersectionData = PrimarySpatialData;
+		for (const UPCGSpatialData* SecondarySourceUnionData : SecondarySourceUnionArray)
+		{
+			UPCGIntersectionData* TempIntersectionData = IntersectionData->IntersectWith(SecondarySourceUnionData);
+			// Propagate settings
+			TempIntersectionData->DensityFunction = Settings->DensityFunction;
+#if WITH_EDITOR
+			TempIntersectionData->bKeepZeroDensityPoints = Settings->bKeepZeroDensityPoints;
+#endif // WITH_EDITOR
+			IntersectionData = TempIntersectionData;
+		}
+
+		FPCGTaggedData& IntersectionTaggedData = Outputs.Emplace_GetRef(Input);
+		IntersectionTaggedData.Data = IntersectionData;
+	}
+
+	return true;
+}
+
+#undef LOCTEXT_NAMESPACE
\ No newline at end of file
diff --git a/Engine/Plugins/Experimental/PCG/Source/PCG/Private/PCGData.cpp b/Engine/Plugins/Experimental/PCG/Source/PCG/Private/PCGData.cpp
--- a/Engine/Plugins/Experimental/PCG/Source/PCG/Private/PCGData.cpp
+++ b/Engine/Plugins/Experimental/PCG/Source/PCG/Private/PCGData.cpp
@@ -158,6 +158,18 @@ TArray<FPCGTaggedData> FPCGDataCollection::GetInputsByPin(const FName& InPinLabe
 		});
 }
 
+TArray<FPCGTaggedData> FPCGDataCollection::GetSpatialInputsByPin(const FName& InPinLabel) const
+{
+	return TaggedData.FilterByPredicate([&InPinLabel](const FPCGTaggedData& Data) {
+		if (!ensure(Data.Data))
+		{
+			return false;
+		}
+
+		return Data.Pin == InPinLabel && Data.Data.IsA<UPCGSpatialData>();
+		});
+}
+
 int32 FPCGDataCollection::GetInputCountByPin(const FName& InPinLabel) const
 {
 	int32 Count = 0;
@@ -171,6 +183,19 @@ int32 FPCGDataCollection::GetInputCountByPin(const FName& InPinLabel) const
 	return Count;
 }
 
+int32 FPCGDataCollection::GetSpatialInputCountByPin(const FName& InPinLabel) const
+{
+	int32 Count = 0;
+	for (const FPCGTaggedData& Data : TaggedData)
+	{
+		if (Data.Pin == InPinLabel && Data.Data && Data.Data.IsA<UPCGSpatialData>())
+		{
+			++Count;
+		}
+	}
+	return Count;
+}
+
 const UPCGSpatialData* FPCGDataCollection::GetSpatialUnionOfInputsByPin(const FName& InPinLabel, bool& bOutUnionDataCreated) const
 {
 	TArray<FPCGTaggedData> SpatialDataOnPin = TaggedData.FilterByPredicate([&InPinLabel](const FPCGTaggedData& Data) {
diff --git a/Engine/Plugins/Experimental/PCG/Source/PCG/Private/PCGNode.cpp b/Engine/Plugins/Experimental/PCG/Source/PCG/Private/PCGNode.cpp
--- a/Engine/Plugins/Experimental/PCG/Source/PCG/Private/PCGNode.cpp
+++ b/Engine/Plugins/Experimental/PCG/Source/PCG/Private/PCGNode.cpp
@@ -377,28 +377,34 @@ bool UPCGNode::IsOutputPinConnected(const FName& Label) const
 	}
 }
 
-void UPCGNode::RenameInputPin(const FName& OldLabel, const FName& NewLabel)
+void UPCGNode::RenameInputPin(const FName& InOldLabel, const FName& InNewLabel, bool bInBroadcastUpdate)
 {
-	if (UPCGPin* Pin = GetInputPin(OldLabel))
+	if (UPCGPin* Pin = GetInputPin(InOldLabel))
 	{
 		Pin->Modify();
-		Pin->Properties.Label = NewLabel;
+		Pin->Properties.Label = InNewLabel;
 
 #if WITH_EDITOR
-		OnNodeChangedDelegate.Broadcast(this, EPCGChangeType::Node);
+		if (bInBroadcastUpdate)
+		{
+			OnNodeChangedDelegate.Broadcast(this, EPCGChangeType::Node);
+		}
 #endif // WITH_EDITOR
 	}
 }
 
-void UPCGNode::RenameOutputPin(const FName& OldLabel, const FName& NewLabel)
+void UPCGNode::RenameOutputPin(const FName& InOldLabel, const FName& InNewLabel, bool bInBroadcastUpdate)
 {
-	if (UPCGPin* Pin = GetOutputPin(OldLabel))
+	if (UPCGPin* Pin = GetOutputPin(InOldLabel))
 	{
 		Pin->Modify();
-		Pin->Properties.Label = NewLabel;
+		Pin->Properties.Label = InNewLabel;
 
 #if WITH_EDITOR
-		OnNodeChangedDelegate.Broadcast(this, EPCGChangeType::Node);
+		if (bInBroadcastUpdate)
+		{
+			OnNodeChangedDelegate.Broadcast(this, EPCGChangeType::Node);
+		}
 #endif // WITH_EDITOR
 	}
 }
@@ -419,6 +425,19 @@ bool UPCGNode::HasInboundEdges() const
 	return false;
 }
 
+int32 UPCGNode::GetInboundEdgesNum() const
+{
+	int32 NumInboundEdges = 0;
+
+	for (const UPCGPin* InputPin : InputPins)
+	{
+		check(InputPin);
+		NumInboundEdges += InputPin->EdgeCount();
+	}
+	
+	return NumInboundEdges;
+}
+
 const UPCGPin* UPCGNode::GetPassThroughInputPin() const
 {
 	// No outputs means nothing will be passed through
diff --git a/Engine/Plugins/Experimental/PCG/Source/PCG/Private/PCGSettingsWithDynamicInputs.cpp b/Engine/Plugins/Experimental/PCG/Source/PCG/Private/PCGSettingsWithDynamicInputs.cpp
new file mode 100644
--- /dev/null
+++ b/Engine/Plugins/Experimental/PCG/Source/PCG/Private/PCGSettingsWithDynamicInputs.cpp
@@ -0,0 +1,75 @@
+// Copyright Epic Games, Inc. All Rights Reserved.
+
+#include "PCGSettingsWithDynamicInputs.h"
+
+#include "PCGComponent.h"
+#include "PCGContext.h"
+#include "PCGGraph.h"
+#include "PCGPin.h"
+
+TArray<FPCGPinProperties> UPCGSettingsWithDynamicInputs::InputPinProperties() const
+{
+	TArray<FPCGPinProperties> ConcatenatedProperties(StaticInputPinProperties());
+	ConcatenatedProperties.Append(DynamicInputPinProperties);
+	return ConcatenatedProperties;
+}
+
+#if WITH_EDITOR
+  void UPCGSettingsWithDynamicInputs::OnUserAddDynamicInputPin()
+{
+	AddDefaultDynamicInputPin();
+}
+
+bool UPCGSettingsWithDynamicInputs::CanUserRemoveDynamicInputPin(int32 PinIndex)
+{
+	const int32 Index = PinIndex - GetStaticInputPinNum(); // Convert to the index of the dynamic pins
+	return (Index >= 0 && Index < DynamicInputPinProperties.Num());
+}
+
+void UPCGSettingsWithDynamicInputs::OnUserRemoveDynamicInputPin(UPCGNode* InOutNode, int32 PinIndex)
+{
+	const int32 Index = PinIndex - GetStaticInputPinNum(); // Convert to the index of the dynamic pins
+	check(InOutNode && Index >= 0 && Index < DynamicInputPinProperties.Num());
+
+	// Prepare the pin for removal by breaking edges and giving it a dummy placeholder label
+	UPCGPin* RemovedPin = InOutNode->GetInputPin(DynamicInputPinProperties[Index].Label);
+	check(RemovedPin);
+	const bool bEdgeRemoved = RemovedPin->BreakAllEdges();
+	// The dummy placeholder label prevents the pin label from conflicting with the subsequent pin that will replace it
+	// during the node's reconstruction within UpdatePins
+	InOutNode->RenameInputPin(RemovedPin->Properties.Label, NAME_Error, /*bBroadcastUpdate=*/false);
+	
+	// Update the pin labels on every input pin after the one being removed
+	for (int I = Index + 1; I < DynamicInputPinProperties.Num(); ++I)
+	{
+		const FName NewLabel = FName(GetDynamicInputPinsBaseLabel().ToString() + FString::FromInt(I + 1));
+		InOutNode->RenameInputPin(DynamicInputPinProperties[I].Label, NewLabel, /*bBroadcastUpdate=*/false);
+		DynamicInputPinProperties[I].Label = NewLabel;
+	}
+
+	DynamicInputPinProperties.RemoveAt(Index);
+	// This broadcast will force reconstruction of the node using the now updated settings
+	const EPCGChangeType ChangeType = EPCGChangeType::Node | EPCGChangeType::Settings | (bEdgeRemoved ? EPCGChangeType::Edge : EPCGChangeType::None);
+	OnSettingsChangedDelegate.Broadcast(this, ChangeType);
+}
+
+int32 UPCGSettingsWithDynamicInputs::GetStaticInputPinNum() const
+{
+	return StaticInputPinProperties().Num();
+}
+
+
+void UPCGSettingsWithDynamicInputs::AddDynamicInputPin(FPCGPinProperties&& CustomProperties)
+{
+	if (CustomPropertiesAreValid(CustomProperties))
+	{
+		DynamicInputPinProperties.Emplace(std::forward<FPCGPinProperties>(CustomProperties));
+		OnSettingsChangedDelegate.Broadcast(this, EPCGChangeType::Node | EPCGChangeType::Settings);
+	}
+}
+#endif // WITH_EDITOR
+
+TArray<FPCGPinProperties> UPCGSettingsWithDynamicInputs::StaticInputPinProperties() const
+{
+	return Super::InputPinProperties();
+}
\ No newline at end of file
diff --git a/Engine/Plugins/Experimental/PCG/Source/PCG/Private/Tests/Determinism/PCGIntersectionDeterminismTest.cpp b/Engine/Plugins/Experimental/PCG/Source/PCG/Private/Tests/Determinism/PCGIntersectionDeterminismTest.cpp
--- a/Engine/Plugins/Experimental/PCG/Source/PCG/Private/Tests/Determinism/PCGIntersectionDeterminismTest.cpp
+++ b/Engine/Plugins/Experimental/PCG/Source/PCG/Private/Tests/Determinism/PCGIntersectionDeterminismTest.cpp
@@ -2,7 +2,7 @@
 
 #if WITH_EDITOR
 
-#include "Elements/PCGIntersectionElement.h"
+#include "Elements/PCGInnerIntersectionElement.h"
 #include "Tests/Determinism/PCGDeterminismTestsCommon.h"
 
 IMPLEMENT_CUSTOM_SIMPLE_AUTOMATION_TEST(FPCGIntersectionDeterminismSingleSameDataTest, FPCGTestBaseClass, "pcg.tests.Intersection.Determinism.SingleSameData", PCGTestsCommon::TestFlags)
@@ -15,7 +15,7 @@ namespace
 {
 	void IntersectionTestBase(PCGTestsCommon::FTestData& TestData)
 	{
-		PCGDeterminismTests::GenerateSettings<UPCGIntersectionSettings>(TestData);
+		PCGDeterminismTests::GenerateSettings<UPCGInnerIntersectionSettings>(TestData);
 		// Source Volumes
 		PCGDeterminismTests::AddVolumeInputData(TestData.InputData, PCGDeterminismTests::Defaults::SmallVector, PCGDeterminismTests::Defaults::MediumVector, PCGDeterminismTests::Defaults::MediumVector);
 		PCGDeterminismTests::AddVolumeInputData(TestData.InputData, PCGDeterminismTests::Defaults::SmallVector * -1.f, PCGDeterminismTests::Defaults::MediumVector, PCGDeterminismTests::Defaults::MediumVector);
diff --git a/Engine/Plugins/Experimental/PCG/Source/PCG/Public/Data/PCGIntersectionData.h b/Engine/Plugins/Experimental/PCG/Source/PCG/Public/Data/PCGIntersectionData.h
--- a/Engine/Plugins/Experimental/PCG/Source/PCG/Public/Data/PCGIntersectionData.h
+++ b/Engine/Plugins/Experimental/PCG/Source/PCG/Public/Data/PCGIntersectionData.h
@@ -9,8 +9,8 @@
 UENUM()
 enum class EPCGIntersectionDensityFunction : uint8
 {
-	Multiply,
-	Minimum
+	Multiply UMETA(ToolTip="Multiplies the density values and results in the product."),
+	Minimum UMETA(ToolTip="Chooses the minimum of the density values.")
 };
 
 /**
diff --git a/Engine/Plugins/Experimental/PCG/Source/PCG/Public/Elements/PCGInnerIntersectionElement.h b/Engine/Plugins/Experimental/PCG/Source/PCG/Public/Elements/PCGInnerIntersectionElement.h
new file mode 100644
--- /dev/null
+++ b/Engine/Plugins/Experimental/PCG/Source/PCG/Public/Elements/PCGInnerIntersectionElement.h
@@ -0,0 +1,53 @@
+// Copyright Epic Games, Inc. All Rights Reserved.
+
+#pragma once
+
+#include "PCGSettings.h"
+#include "Data/PCGIntersectionData.h"
+
+#include "PCGInnerIntersectionElement.generated.h"
+
+UCLASS(BlueprintType, ClassGroup = (Procedural))
+class PCG_API UPCGInnerIntersectionSettings : public UPCGSettings
+{
+	GENERATED_BODY()
+
+public:
+	//~Begin UPCGSettings interface
+#if WITH_EDITOR
+	virtual FName GetDefaultNodeName() const override { return FName(TEXT("Inner Intersection")); }
+	virtual FText GetDefaultNodeTitle() const override { return NSLOCTEXT("PCGInnerIntersectionSettings", "NodeTitle", "Inner Intersection"); }
+	virtual FText GetNodeTooltipText() const override;
+	virtual EPCGSettingsType GetType() const override { return EPCGSettingsType::Spatial; }
+#endif
+
+	virtual TArray<FPCGPinProperties> InputPinProperties() const override;
+	virtual TArray<FPCGPinProperties> OutputPinProperties() const override;
+
+	// If node disabled, don't intersect - pass through first edge
+	virtual bool OnlyPassThroughOneEdgeWhenDisabled() const { return true; }
+
+protected:
+	virtual FPCGElementPtr CreateElement() const override;
+	//~End UPCGSettings interface
+
+public:
+	UPROPERTY(BlueprintReadWrite, EditAnywhere, Category = Settings, meta = (PCG_Overridable))
+	EPCGIntersectionDensityFunction DensityFunction = EPCGIntersectionDensityFunction::Multiply;
+
+#if WITH_EDITORONLY_DATA
+	/** [EDITOR ONLY] If enabled, output points with a density value of 0 will NOT be automatically filtered out. */
+	UPROPERTY(Transient, BlueprintReadWrite, EditAnywhere, Category = "Settings|Debug", meta = (PCG_Overridable))
+	bool bKeepZeroDensityPoints = false;
+#endif
+};
+
+class FPCGInnerIntersectionElement : public FSimplePCGElement
+{
+protected:
+	virtual bool ExecuteInternal(FPCGContext* Context) const;
+};
+
+#if UE_ENABLE_INCLUDE_ORDER_DEPRECATED_IN_5_2
+#include "CoreMinimal.h"
+#endif
diff --git a/Engine/Plugins/Experimental/PCG/Source/PCG/Public/Elements/PCGIntersectionElement.h b/Engine/Plugins/Experimental/PCG/Source/PCG/Public/Elements/PCGIntersectionElement.h
deleted file mode 100644
diff --git a/Engine/Plugins/Experimental/PCG/Source/PCG/Public/Elements/PCGOuterIntersectionElement.h b/Engine/Plugins/Experimental/PCG/Source/PCG/Public/Elements/PCGOuterIntersectionElement.h
new file mode 100644
--- /dev/null
+++ b/Engine/Plugins/Experimental/PCG/Source/PCG/Public/Elements/PCGOuterIntersectionElement.h
@@ -0,0 +1,63 @@
+// Copyright Epic Games, Inc. All Rights Reserved.
+
+#pragma once
+
+#include "PCGSettingsWithDynamicInputs.h"
+#include "Data/PCGIntersectionData.h"
+
+#include "PCGOuterIntersectionElement.generated.h"
+
+UCLASS(BlueprintType, ClassGroup = (Procedural))
+class PCG_API UPCGOuterIntersectionSettings : public UPCGSettingsWithDynamicInputs
+{
+	GENERATED_BODY()
+
+public:
+#if WITH_EDITOR
+	//~Begin UPCGSettings interface
+	virtual FName GetDefaultNodeName() const override { return FName(TEXT("Intersection")); }
+	virtual FText GetDefaultNodeTitle() const override { return NSLOCTEXT("PCGOuterIntersectionElement", "NodeTitle", "Intersection"); }
+	virtual FText GetNodeTooltipText() const override;
+	virtual EPCGSettingsType GetType() const override { return EPCGSettingsType::Spatial; }
+	virtual bool HasDynamicPins() const override { return true; }
+#endif // WITH_EDITOR
+
+	// If node disabled, don't intersect - pass through all primary edges
+	virtual bool OnlyPassThroughOneEdgeWhenDisabled() const override { return false; }
+protected:
+	virtual TArray<FPCGPinProperties> OutputPinProperties() const override;
+	virtual FPCGElementPtr CreateElement() const override;
+	//~End UPCGSettings interface
+
+	//~Begin UPCGDynamicSettings interface
+	virtual FName GetDynamicInputPinsBaseLabel() const override;
+	/** The input pin properties that are statically defined */
+	virtual TArray<FPCGPinProperties> StaticInputPinProperties() const override;
+
+#if WITH_EDITOR
+public:
+	virtual void AddDefaultDynamicInputPin() override;
+#endif // WITH_EDITOR
+	//~End UPCGDynamicSettings interface
+
+public:
+	UPROPERTY(BlueprintReadWrite, EditAnywhere, Category = Settings, meta = (PCG_Overridable))
+	EPCGIntersectionDensityFunction DensityFunction = EPCGIntersectionDensityFunction::Multiply;
+
+	/** If enabled, dynamic input pins that have no incoming data will be bypassed during the intersection operation
+	 *  calculation, which would otherwise result in an empty result. */
+	UPROPERTY(BlueprintReadWrite, EditAnywhere, Category = Settings, meta = (PCG_Overridable))
+	bool bIgnorePinsWithNoInput = false;
+
+#if WITH_EDITORONLY_DATA
+	/** [EDITOR ONLY] If enabled, output points with a density value of 0 will NOT be automatically filtered out. */
+	UPROPERTY(Transient, BlueprintReadWrite, EditAnywhere, Category = "Settings|Debug", meta = (PCG_Overridable))
+	bool bKeepZeroDensityPoints = false;
+#endif // WITH_EDITORONLY_DATA
+};
+
+class FPCGOuterIntersectionElement : public FSimplePCGElement
+{
+protected:
+	virtual bool ExecuteInternal(FPCGContext* Context) const override;
+};
\ No newline at end of file
diff --git a/Engine/Plugins/Experimental/PCG/Source/PCG/Public/PCGData.h b/Engine/Plugins/Experimental/PCG/Source/PCG/Public/PCGData.h
--- a/Engine/Plugins/Experimental/PCG/Source/PCG/Public/PCGData.h
+++ b/Engine/Plugins/Experimental/PCG/Source/PCG/Public/PCGData.h
@@ -109,8 +109,13 @@ struct PCG_API FPCGDataCollection
 	TArray<FPCGTaggedData> GetInputs() const;
 	/** Returns all data on a given pin */
 	TArray<FPCGTaggedData> GetInputsByPin(const FName& InPinLabel) const;
+	/** Returns all spatial data on a given pin */
+	TArray<FPCGTaggedData> GetSpatialInputsByPin(const FName& InPinLabel) const;
+	
 	/** Gets number of data items on a given pin */
 	int32 GetInputCountByPin(const FName& InPinLabel) const;
+	/** Gets number of spatial data items on a given pin */
+	int32 GetSpatialInputCountByPin(const FName& InPinLabel) const;
 	/** Returns spatial union of all data on a given pin, returns null if no such data exists. bOutUnionDataCreated indicates if new data created that may need rooting. */
 	const UPCGSpatialData* GetSpatialUnionOfInputsByPin(const FName& InPinLabel, bool& bOutUnionDataCreated) const;
 	/** Returns all spatial data in the collection with the given tag */
diff --git a/Engine/Plugins/Experimental/PCG/Source/PCG/Public/PCGNode.h b/Engine/Plugins/Experimental/PCG/Source/PCG/Public/PCGNode.h
--- a/Engine/Plugins/Experimental/PCG/Source/PCG/Public/PCGNode.h
+++ b/Engine/Plugins/Experimental/PCG/Source/PCG/Public/PCGNode.h
@@ -104,10 +104,11 @@ public:
 	UPCGPin* GetOutputPin(const FName& Label);
 	const UPCGPin* GetOutputPin(const FName& Label) const;
 	bool HasInboundEdges() const;
+	int32 GetInboundEdgesNum() const;
 
 	/** Allow to change the name of a pin, to keep edges connected. You need to make sure that the underlying settings are also updated, otherwise, it will be overwritten next time the settings are updated */
-	void RenameInputPin(const FName& OldLabel, const FName& NewLabel);
-	void RenameOutputPin(const FName& OldLabel, const FName& NewLabel);
+	void RenameInputPin(const FName& InOldLabel, const FName& InNewLabel, bool bInBroadcastUpdate = true);
+	void RenameOutputPin(const FName& InOldLabel, const FName& InNewLabel, bool bInBroadcastUpdate = true);
 
 	/** Pin from which data is passed through when this node is disabled. */
 	virtual const UPCGPin* GetPassThroughInputPin() const;
diff --git a/Engine/Plugins/Experimental/PCG/Source/PCG/Public/PCGSettingsWithDynamicInputs.h b/Engine/Plugins/Experimental/PCG/Source/PCG/Public/PCGSettingsWithDynamicInputs.h
new file mode 100644
--- /dev/null
+++ b/Engine/Plugins/Experimental/PCG/Source/PCG/Public/PCGSettingsWithDynamicInputs.h
@@ -0,0 +1,56 @@
+// Copyright Epic Games, Inc. All Rights Reserved.
+
+#pragma once
+
+#include "PCGSettings.h"
+
+#include "PCGSettingsWithDynamicInputs.generated.h"
+
+/**
+* UPCGSettings subclass with functionality to dynamically add/remove input pins
+*/
+UCLASS(Abstract, BlueprintType, ClassGroup = (Procedural))
+class PCG_API UPCGSettingsWithDynamicInputs: public UPCGSettings
+{
+	GENERATED_BODY()
+
+public:
+	virtual FName GetDynamicInputPinsBaseLabel() const { return NAME_None; }
+
+#if WITH_EDITOR
+	/** Validate custom pin properties */
+	virtual bool CustomPropertiesAreValid(const FPCGPinProperties& CustomProperties) { return true; }
+	
+	/** User driven event to add a dynamic source pin */
+	virtual void OnUserAddDynamicInputPin();
+	/** Overridden logic to add a default source pin */
+	virtual void AddDefaultDynamicInputPin() PURE_VIRTUAL(PCGDynamicSettings::AddDefaultSourcePin, );
+	/** Check if the pin to remove is dynamic */
+	virtual bool CanUserRemoveDynamicInputPin(int32 PinIndex);
+	/** User driven event to remove a dynamic source pin */
+	virtual void OnUserRemoveDynamicInputPin(UPCGNode* InOutNode, int32 PinIndex);
+
+	/** Get the number of static input pins. */
+	int32 GetStaticInputPinNum() const;
+	/** Get the number of dynamic input pins. */
+	int32 GetDynamicInputPinNum() const { return DynamicInputPinProperties.Num(); }
+#endif // WITH_EDITOR
+
+protected:
+#if WITH_EDITOR
+	/** Add a new dynamic source pin with the specified properties */
+	void AddDynamicInputPin(FPCGPinProperties&& CustomProperties);
+#endif // WITH_EDITOR
+
+	/** The input pin properties that are statically defined by the client class */
+	virtual TArray<FPCGPinProperties> StaticInputPinProperties() const;
+
+	/** Dynamic pin properties that the user can add or remove from */
+	UPROPERTY()
+	TArray<FPCGPinProperties> DynamicInputPinProperties;
+
+	//~Begin UPCGSettings interface
+	/** A concatenation of the static and dynamic input pin properties */
+	virtual TArray<FPCGPinProperties> InputPinProperties() const override final;
+	//~End UPCGSettings interface
+};
\ No newline at end of file
diff --git a/Engine/Plugins/Experimental/PCG/Source/PCGEditor/Private/PCGEditor.cpp b/Engine/Plugins/Experimental/PCG/Source/PCGEditor/Private/PCGEditor.cpp
--- a/Engine/Plugins/Experimental/PCG/Source/PCGEditor/Private/PCGEditor.cpp
+++ b/Engine/Plugins/Experimental/PCG/Source/PCGEditor/Private/PCGEditor.cpp
@@ -502,6 +502,11 @@ void FPCGEditor::BindCommands()
 	GraphEditorCommands->MapAction(
 		PCGEditorCommands.DisableDebugOnAllNodes,
 		FExecuteAction::CreateSP(this, &FPCGEditor::OnDisableDebugOnAllNodes));
+
+	GraphEditorCommands->MapAction(
+		PCGEditorCommands.AddSourcePin,
+		FExecuteAction::CreateSP(this, &FPCGEditor::OnAddDynamicInputPin),
+		FCanExecuteAction::CreateSP(this, &FPCGEditor::CanAddDynamicInputPin));
 }
 
 void FPCGEditor::OnFind()
@@ -761,6 +766,36 @@ bool FPCGEditor::CanCollapseNodesInSubgraph() const
 	return false;
 }
 
+void FPCGEditor::OnAddDynamicInputPin()
+{
+	check(GraphEditorWidget.IsValid());
+
+	const FGraphPanelSelectionSet SelectedNodes = GraphEditorWidget->GetSelectedNodes();
+	
+	if (!ensure(SelectedNodes.Num() == 1))
+	{
+		UE_LOG(LogPCGEditor, Warning, TEXT("Attempting to add new input pin to multiple nodes."));
+		return;
+	}
+
+	UPCGEditorGraphNodeBase* Node = CastChecked<UPCGEditorGraphNodeBase>(*SelectedNodes.CreateConstIterator());
+	Node->OnUserAddDynamicInputPin();
+}
+
+bool FPCGEditor::CanAddDynamicInputPin() const
+{
+	check(GraphEditorWidget.IsValid());
+
+	const FGraphPanelSelectionSet SelectedNodes = GraphEditorWidget->GetSelectedNodes();
+	if (SelectedNodes.Num() != 1)
+	{
+		return false;
+	}
+
+	const UPCGEditorGraphNodeBase* Node = Cast<const UPCGEditorGraphNodeBase>(*SelectedNodes.CreateConstIterator());
+	return (Node && Node->CanUserAddRemoveDynamicInputPins());
+}
+
 void FPCGEditor::OnCollapseNodesInSubgraph()
 {
 	if (!GraphEditorWidget.IsValid() || PCGEditorGraph == nullptr)
diff --git a/Engine/Plugins/Experimental/PCG/Source/PCGEditor/Private/PCGEditor.h b/Engine/Plugins/Experimental/PCG/Source/PCGEditor/Private/PCGEditor.h
--- a/Engine/Plugins/Experimental/PCG/Source/PCGEditor/Private/PCGEditor.h
+++ b/Engine/Plugins/Experimental/PCG/Source/PCGEditor/Private/PCGEditor.h
@@ -204,6 +204,11 @@ private:
 	/** Whether we can collapse nodes in a subgraph */
 	bool CanCollapseNodesInSubgraph() const;
 
+	/** User is attempting to add a dynamic source pin to a node */
+	void OnAddDynamicInputPin();
+	/** Whether the user can add a dynamic source pin to a node */
+	bool CanAddDynamicInputPin() const;
+
 	void OnAlignTop();
 	void OnAlignMiddle();
 	void OnAlignBottom();
diff --git a/Engine/Plugins/Experimental/PCG/Source/PCGEditor/Private/PCGEditorCommands.cpp b/Engine/Plugins/Experimental/PCG/Source/PCGEditor/Private/PCGEditorCommands.cpp
--- a/Engine/Plugins/Experimental/PCG/Source/PCGEditor/Private/PCGEditorCommands.cpp
+++ b/Engine/Plugins/Experimental/PCG/Source/PCGEditor/Private/PCGEditorCommands.cpp
@@ -32,6 +32,7 @@ void FPCGEditorCommands::RegisterCommands()
 	UI_COMMAND(DebugOnlySelected, "Debug Only Selected", "Enable node debug state for selected nodes and disable debug state for the others", EUserInterfaceActionType::Button, FInputChord(EModifierKey::Control | EModifierKey::Alt, EKeys::D));
 	UI_COMMAND(DisableDebugOnAllNodes, "Disable Debug on all nodes", "Disable debug state for all nodes", EUserInterfaceActionType::Button, FInputChord(EModifierKey::Alt, EKeys::D));
 	UI_COMMAND(ToggleInspect, "Toggle Inspection", "Toggle node inspection for selected node", EUserInterfaceActionType::ToggleButton, FInputChord(EKeys::A));
+	UI_COMMAND(AddSourcePin, "Add Source Pin", "Add new source pin to the current node", EUserInterfaceActionType::Button, FInputChord());
 }
 
 #undef LOCTEXT_NAMESPACE
diff --git a/Engine/Plugins/Experimental/PCG/Source/PCGEditor/Private/PCGEditorCommands.h b/Engine/Plugins/Experimental/PCG/Source/PCGEditor/Private/PCGEditorCommands.h
--- a/Engine/Plugins/Experimental/PCG/Source/PCGEditor/Private/PCGEditorCommands.h
+++ b/Engine/Plugins/Experimental/PCG/Source/PCGEditor/Private/PCGEditorCommands.h
@@ -28,4 +28,5 @@ public:
 	TSharedPtr<FUICommandInfo> DebugOnlySelected;
 	TSharedPtr<FUICommandInfo> DisableDebugOnAllNodes;
 	TSharedPtr<FUICommandInfo> ToggleInspect;
+	TSharedPtr<FUICommandInfo> AddSourcePin;
 };
diff --git a/Engine/Plugins/Experimental/PCG/Source/PCGEditor/Private/PCGEditorGraphNodeBase.cpp b/Engine/Plugins/Experimental/PCG/Source/PCGEditor/Private/PCGEditorGraphNodeBase.cpp
--- a/Engine/Plugins/Experimental/PCG/Source/PCGEditor/Private/PCGEditorGraphNodeBase.cpp
+++ b/Engine/Plugins/Experimental/PCG/Source/PCGEditor/Private/PCGEditorGraphNodeBase.cpp
@@ -11,13 +11,14 @@
 #include "PCGEditorSettings.h"
 #include "PCGGraph.h"
 #include "PCGPin.h"
+#include "PCGSettingsWithDynamicInputs.h"
 #include "PCGSubsystem.h"
 #include "PCGWorldActor.h"
-#include "Helpers/PCGHelpers.h"
 
 #include "GraphEditorActions.h"
 #include "ToolMenu.h"
 #include "ToolMenuSection.h"
+#include "ScopedTransaction.h"
 #include "Logging/TokenizedMessage.h"
 #include "Misc/TransactionObjectEvent.h"
 #include "Widgets/Colors/SColorPicker.h"
@@ -40,6 +41,8 @@ void UPCGEditorGraphNodeBase::Construct(UPCGNode* InPCGNode)
 	{
 		const ENodeEnabledState NewEnabledState = !PCGSettingsInterface->bEnabled ? ENodeEnabledState::Disabled : ENodeEnabledState::Enabled;
 		SetEnabledState(NewEnabledState);
+
+		bCanUserAddRemoveSourcePins = (InPCGNode->GetSettings() && InPCGNode->GetSettings()->IsA<UPCGSettingsWithDynamicInputs>());
 	}
 }
 
@@ -77,6 +80,27 @@ void UPCGEditorGraphNodeBase::GetNodeContextMenuActions(UToolMenu* Menu, class U
 		return;
 	}
 
+	{
+		FToolMenuSection& Section = Menu->AddSection("EdGraphSchemaPinActions", LOCTEXT("PinActionsMenuHeader", "Pin Actions"));
+		if (Context->Pin && bCanUserAddRemoveSourcePins)
+		{
+			Section.AddMenuEntry(FPCGEditorCommands::Get().AddSourcePin);
+			Section.AddMenuEntry("RemovePin",
+				LOCTEXT("RemovePin", "Remove Source Pin"),
+				LOCTEXT("RemovePinTooltip", "Remove this source pin from the current node"),
+				FSlateIcon(),
+				FUIAction(
+					FExecuteAction::CreateLambda([Pin = Context->Pin, this]
+					{
+						const_cast<UPCGEditorGraphNodeBase*>(this)->OnUserRemoveDynamicInputPin(const_cast<UEdGraphPin*>(Pin));
+					}),
+					FCanExecuteAction::CreateLambda([Pin = Context->Pin, this]
+					{
+						return const_cast<UPCGEditorGraphNodeBase*>(this)->CanUserRemoveDynamicInputPin(const_cast<UEdGraphPin*>(Pin));
+					})));
+		}
+	}
+
 	{
 		FToolMenuSection& Section = Menu->AddSection("EdGraphSchemaNodeActions", LOCTEXT("NodeActionsHeader", "Node Actions"));
 		Section.AddMenuEntry(FPCGEditorCommands::Get().ToggleEnabled, LOCTEXT("ToggleEnabledLabel", "Enable"));
@@ -725,6 +749,46 @@ void UPCGEditorGraphNodeBase::OnCommentBubbleToggled(bool bInCommentBubbleVisibl
 	}
 }
 
+void UPCGEditorGraphNodeBase::OnUserAddDynamicInputPin()
+{
+	check(PCGNode);
+	
+	if (UPCGSettingsWithDynamicInputs* DynamicNodeSettings = Cast<UPCGSettingsWithDynamicInputs>(PCGNode->GetSettings()))
+	{
+		const FScopedTransaction Transaction(*FPCGEditorCommon::ContextIdentifier, LOCTEXT("PCGEditorUserAddDynamicInputPin", "Add Source Pin"), DynamicNodeSettings);
+		DynamicNodeSettings->Modify();
+		DynamicNodeSettings->OnUserAddDynamicInputPin();
+	}
+}
+
+bool UPCGEditorGraphNodeBase::CanUserRemoveDynamicInputPin(UEdGraphPin* InPinToRemove)
+{
+	check(PCGNode && InPinToRemove);
+
+	if (UPCGSettingsWithDynamicInputs* DynamicNodeSettings = Cast<UPCGSettingsWithDynamicInputs>(PCGNode->GetSettings()))
+	{
+		const UPCGEditorGraphNodeBase* PCGGraphNode = CastChecked<const UPCGEditorGraphNodeBase>(InPinToRemove->GetOwningNode());
+		return DynamicNodeSettings->CanUserRemoveDynamicInputPin(PCGGraphNode->GetPinIndex(InPinToRemove));
+	}
+	
+	return false;
+}
+
+void UPCGEditorGraphNodeBase::OnUserRemoveDynamicInputPin(UEdGraphPin* InRemovedPin)
+{
+	check(PCGNode && InRemovedPin);
+
+	if (UPCGSettingsWithDynamicInputs* DynamicNodeSettings = Cast<UPCGSettingsWithDynamicInputs>(PCGNode->GetSettings()))
+	{
+		if (UPCGEditorGraphNodeBase* PCGGraphNode = Cast<UPCGEditorGraphNodeBase>(InRemovedPin->GetOwningNode()))
+		{
+			const FScopedTransaction Transaction(*FPCGEditorCommon::ContextIdentifier, LOCTEXT("PCGEditorUserRemoveDynamicInputPin", "Remove Source Pin"), DynamicNodeSettings);
+			DynamicNodeSettings->Modify();
+			DynamicNodeSettings->OnUserRemoveDynamicInputPin(PCGGraphNode->GetPCGNode(), PCGGraphNode->GetPinIndex(InRemovedPin));
+		}
+	}
+}
+
 void UPCGEditorGraphNodeBase::UpdateCommentBubblePinned()
 {
 	if (PCGNode)
diff --git a/Engine/Plugins/Experimental/PCG/Source/PCGEditor/Private/PCGEditorGraphNodeBase.h b/Engine/Plugins/Experimental/PCG/Source/PCGEditor/Private/PCGEditorGraphNodeBase.h
--- a/Engine/Plugins/Experimental/PCG/Source/PCGEditor/Private/PCGEditorGraphNodeBase.h
+++ b/Engine/Plugins/Experimental/PCG/Source/PCGEditor/Private/PCGEditorGraphNodeBase.h
@@ -43,6 +43,10 @@ public:
 	virtual void OnCommentBubbleToggled(bool bInCommentBubbleVisible) override;
 	// ~End UEdGraphNode interface
 
+	void OnUserAddDynamicInputPin();
+	bool CanUserRemoveDynamicInputPin(UEdGraphPin* InPinToRemove);
+	void OnUserRemoveDynamicInputPin(UEdGraphPin* InRemovedPin);
+	
 	UPCGNode* GetPCGNode() { return PCGNode; }
 	const UPCGNode* GetPCGNode() const { return PCGNode; }
 	void PostCopy();
@@ -69,6 +73,8 @@ public:
 	void SetIsHighlighted(bool bInIsHighlighted) { bIsHighlighted = bInIsHighlighted; }
 	bool IsHighlighted() const { return bIsHighlighted; }
 
+	bool CanUserAddRemoveDynamicInputPins() const { return bCanUserAddRemoveSourcePins; }
+
 	DECLARE_DELEGATE(FOnPCGEditorGraphNodeChanged);
 	FOnPCGEditorGraphNodeChanged OnNodeChangedDelegate;
 
@@ -98,4 +104,5 @@ protected:
 	bool bDisableReconstructFromNode = false;
 	bool bIsInspected = false;
 	bool bIsHighlighted = false;
+	bool bCanUserAddRemoveSourcePins = false;
 };
diff --git a/Engine/Plugins/Experimental/PCG/Source/PCGEditor/Private/SPCGEditorGraphNode.cpp b/Engine/Plugins/Experimental/PCG/Source/PCGEditor/Private/SPCGEditorGraphNode.cpp
--- a/Engine/Plugins/Experimental/PCG/Source/PCGEditor/Private/SPCGEditorGraphNode.cpp
+++ b/Engine/Plugins/Experimental/PCG/Source/PCGEditor/Private/SPCGEditorGraphNode.cpp
@@ -7,6 +7,7 @@
 #include "PCGNode.h"
 #include "PCGPin.h"
 #include "PCGSettings.h"
+#include "PCGSettingsWithDynamicInputs.h"
 
 #include "GraphEditorSettings.h"
 #include "SCommentBubble.h"
@@ -359,6 +360,27 @@ void SPCGEditorGraphNode::Construct(const FArguments& InArgs, UPCGEditorGraphNod
 	UpdateGraphNode();
 }
 
+void SPCGEditorGraphNode::CreateAddPinButtonWidget()
+{
+	// Add Pin Button (+) Source Reference: Engine\Source\Editor\GraphEditor\Private\KismetNodes\SGraphNodeK2Sequence.cpp
+	const TSharedPtr<SWidget> AddPinButton = AddPinButtonContent(LOCTEXT("AddSourcePin", "Add Pin"), LOCTEXT("AddSourcePinTooltip", "Add a dynamic source input pin"));
+
+	FMargin AddPinPadding = Settings->GetInputPinPadding();
+	AddPinPadding.Top += 6.0f;
+
+	check(PCGEditorGraphNode->GetPCGNode());
+	const UPCGSettingsWithDynamicInputs* NodeSettings = CastChecked<UPCGSettingsWithDynamicInputs>(PCGEditorGraphNode->GetPCGNode()->GetSettings());
+
+	const int32 Index = NodeSettings->GetStaticInputPinNum() + NodeSettings->GetDynamicInputPinNum();
+	LeftNodeBox->InsertSlot(Index)
+		.AutoHeight()
+		.VAlign(VAlign_Bottom)
+		.Padding(AddPinPadding)
+		[
+			AddPinButton.ToSharedRef()
+		];
+}
+
 void SPCGEditorGraphNode::UpdateGraphNode()
 {
 	if (PCGEditorGraphNode && PCGEditorGraphNode->ShouldDrawCompact())
@@ -368,6 +390,11 @@ void SPCGEditorGraphNode::UpdateGraphNode()
 	else
 	{
 		SGraphNode::UpdateGraphNode();
+
+		if (PCGEditorGraphNode->CanUserAddRemoveDynamicInputPins())
+		{
+			CreateAddPinButtonWidget();
+		}
 	}
 }
 
@@ -407,6 +434,25 @@ TSharedPtr<SGraphPin> SPCGEditorGraphNode::CreatePinWidget(UEdGraphPin* Pin) con
 	return SNew(SPCGEditorGraphNodePin, Pin);
 }
 
+EVisibility SPCGEditorGraphNode::IsAddPinButtonVisible() const
+{
+	if (PCGEditorGraphNode && PCGEditorGraphNode->IsNodeEnabled() && PCGEditorGraphNode->CanUserAddRemoveDynamicInputPins() && SGraphNode::IsAddPinButtonVisible() == EVisibility::Visible)
+	{
+		return EVisibility::Visible;
+	}
+	
+	return EVisibility::Hidden;
+}
+
+FReply SPCGEditorGraphNode::OnAddPin()
+{
+	check(PCGEditorGraphNode);
+
+	PCGEditorGraphNode->OnUserAddDynamicInputPin();
+	
+	return FReply::Handled();
+}
+
 void SPCGEditorGraphNode::AddPin(const TSharedRef<SGraphPin>& PinToAdd)
 {
 	check(PCGEditorGraphNode);
diff --git a/Engine/Plugins/Experimental/PCG/Source/PCGEditor/Private/SPCGEditorGraphNode.h b/Engine/Plugins/Experimental/PCG/Source/PCGEditor/Private/SPCGEditorGraphNode.h
--- a/Engine/Plugins/Experimental/PCG/Source/PCGEditor/Private/SPCGEditorGraphNode.h
+++ b/Engine/Plugins/Experimental/PCG/Source/PCGEditor/Private/SPCGEditorGraphNode.h
@@ -15,6 +15,7 @@ public:
 	SLATE_END_ARGS()
 
 	void Construct(const FArguments& InArgs, UPCGEditorGraphNodeBase* InNode);
+	void CreateAddPinButtonWidget();
 
 	//~ Begin SGraphNode Interface
 	virtual void UpdateGraphNode() override;
@@ -22,7 +23,9 @@ public:
 	virtual void RequestRenameOnSpawn() override { /* Empty to avoid the default behavior to rename on node spawn */ }
 	virtual void AddPin(const TSharedRef<SGraphPin>& PinToAdd) override;
 	virtual TSharedRef<SWidget> CreateTitleWidget(TSharedPtr<SNodeTitle> InNodeTitle) override;
-	virtual TSharedPtr<SGraphPin> CreatePinWidget(UEdGraphPin* Pin) const;
+	virtual TSharedPtr<SGraphPin> CreatePinWidget(UEdGraphPin* Pin) const override;
+	virtual EVisibility IsAddPinButtonVisible() const override;
+	virtual FReply OnAddPin() override;
 	//~ End SGraphNode Interface
 
 	//~ Begin SNodePanel::SNode Interface
