[PCG] Get Bounds node to compute bounds of input spatial data

diff --git a/Engine/Plugins/Experimental/PCG/Source/PCG/Private/Elements/PCGGetBounds.cpp b/Engine/Plugins/Experimental/PCG/Source/PCG/Private/Elements/PCGGetBounds.cpp
new file mode 100644
--- /dev/null
+++ b/Engine/Plugins/Experimental/PCG/Source/PCG/Private/Elements/PCGGetBounds.cpp
@@ -0,0 +1,91 @@
+// Copyright Epic Games, Inc. All Rights Reserved.
+
+#include "Elements/PCGGetBounds.h"
+
+#include "PCGComponent.h"
+#include "PCGContext.h"
+#include "PCGParamData.h"
+#include "PCGPin.h"
+#include "Data/PCGSpatialData.h"
+
+#define LOCTEXT_NAMESPACE "PCGGetBoundsElement"
+
+#if WITH_EDITOR
+FText UPCGGetBoundsSettings::GetDefaultNodeTitle() const
+{
+	return LOCTEXT("NodeTitle", "Get Bounds");
+}
+
+FText UPCGGetBoundsSettings::GetNodeTooltipText() const
+{
+	return LOCTEXT("NodeTooltip", "Computes the bounds of the inputs as attributes.");
+}
+#endif
+
+TArray<FPCGPinProperties> UPCGGetBoundsSettings::InputPinProperties() const
+{
+	TArray<FPCGPinProperties> PinProperties;
+	PinProperties.Emplace(PCGPinConstants::DefaultInputLabel, EPCGDataType::Spatial);
+
+	return PinProperties;
+}
+
+TArray<FPCGPinProperties> UPCGGetBoundsSettings::OutputPinProperties() const
+{
+	TArray<FPCGPinProperties> PinProperties;
+	PinProperties.Emplace(PCGPinConstants::DefaultOutputLabel, EPCGDataType::Param, /*bInAllowMultipleConnections=*/true, /*bAllowMultipleData=*/true);
+
+	return PinProperties;
+}
+
+FPCGElementPtr UPCGGetBoundsSettings::CreateElement() const
+{
+	return MakeShared<FPCGGetBoundsElement>();
+}
+
+bool FPCGGetBoundsElement::ExecuteInternal(FPCGContext* Context) const
+{
+	const TArray<FPCGTaggedData>& Inputs = Context->InputData.TaggedData;
+	bool bEmittedUnboundedDataWarning = false;
+
+	for (const FPCGTaggedData& Input : Inputs)
+	{
+		UPCGSpatialData* InputSpatialData = Cast<UPCGSpatialData>(Input.Data);
+		if (!InputSpatialData)
+		{
+			continue;
+		}
+
+		if (!InputSpatialData->IsBounded())
+		{
+			if (!bEmittedUnboundedDataWarning)
+			{
+				bEmittedUnboundedDataWarning = true;
+				PCGE_LOG(Warning, GraphAndLog, LOCTEXT("UnboundedInput", "Skipped unbounded spatial data."));
+			}
+
+			continue;
+		}
+
+		const FBox InputBounds = InputSpatialData->GetBounds();
+		if (!ensure(InputBounds.IsValid))
+		{
+			PCGE_LOG(Warning, GraphAndLog, LOCTEXT("InvalidBoundsOnInput", "Skipped spatial data that had invalid bounds."));
+			continue;
+		}
+
+		UPCGParamData* ParamData = NewObject<UPCGParamData>();
+		check(ParamData && ParamData->Metadata);
+		
+		ParamData->Metadata->CreateVectorAttribute("BoundsMin", InputBounds.Min, /*bAllowsInterpolation=*/true, /*bOverrideParent=*/false);
+		ParamData->Metadata->CreateVectorAttribute("BoundsMax", InputBounds.Max, /*bAllowsInterpolation=*/true, /*bOverrideParent=*/false);
+		ParamData->Metadata->AddEntry();
+
+		FPCGTaggedData& Output = Context->OutputData.TaggedData.Emplace_GetRef(Input);
+		Output.Data = ParamData;
+	}
+
+	return true;
+}
+
+#undef LOCTEXT_NAMESPACE
diff --git a/Engine/Plugins/Experimental/PCG/Source/PCG/Public/Elements/PCGGetBounds.h b/Engine/Plugins/Experimental/PCG/Source/PCG/Public/Elements/PCGGetBounds.h
new file mode 100644
--- /dev/null
+++ b/Engine/Plugins/Experimental/PCG/Source/PCG/Public/Elements/PCGGetBounds.h
@@ -0,0 +1,37 @@
+// Copyright Epic Games, Inc. All Rights Reserved.
+#pragma once
+
+#include "PCGElement.h"
+#include "PCGSettings.h"
+
+#include "PCGGetBounds.generated.h"
+
+/**
+ * Computes the bounds of the inputs as attributes.
+ */
+UCLASS(BlueprintType, ClassGroup = (Procedural))
+class PCG_API UPCGGetBoundsSettings : public UPCGSettings
+{
+	GENERATED_BODY()
+
+public:
+	//~Begin UPCGSettings interface
+#if WITH_EDITOR
+	virtual FName GetDefaultNodeName() const override { return TEXT("GetBounds"); }
+	virtual FText GetDefaultNodeTitle() const override;
+	virtual FText GetNodeTooltipText() const override;
+	virtual EPCGSettingsType GetType() const override { return EPCGSettingsType::Spatial; }
+#endif
+
+protected:
+	virtual TArray<FPCGPinProperties> InputPinProperties() const override;
+	virtual TArray<FPCGPinProperties> OutputPinProperties() const override;
+	virtual FPCGElementPtr CreateElement() const override;
+	//~End UPCGSettings interface
+};
+
+class FPCGGetBoundsElement : public FSimplePCGElement
+{
+protected:
+	virtual bool ExecuteInternal(FPCGContext* Context) const override;
+};
