[PCG] The merge node will now properly output all inputs, when disabled

#rnx

diff --git a/Engine/Plugins/Experimental/PCG/Source/PCG/Public/Elements/PCGMergeElement.h b/Engine/Plugins/Experimental/PCG/Source/PCG/Public/Elements/PCGMergeElement.h
--- a/Engine/Plugins/Experimental/PCG/Source/PCG/Public/Elements/PCGMergeElement.h
+++ b/Engine/Plugins/Experimental/PCG/Source/PCG/Public/Elements/PCGMergeElement.h
@@ -24,9 +24,6 @@ public:
 	virtual TArray<FPCGPinProperties> InputPinProperties() const override;
 	virtual TArray<FPCGPinProperties> OutputPinProperties() const override;
 
-	// If node disabled, don't merge - pass through first edge
-	virtual bool OnlyPassThroughOneEdgeWhenDisabled() const { return true; }
-
 protected:
 	virtual FPCGElementPtr CreateElement() const override;
 	//~End UPCGSettings interface
