Tentative fix for the PCG world actor being added without any user input

#rnx

diff --git a/Engine/Plugins/Experimental/PCG/Source/PCG/Private/PCGActorAndComponentMapping.cpp b/Engine/Plugins/Experimental/PCG/Source/PCG/Private/PCGActorAndComponentMapping.cpp
--- a/Engine/Plugins/Experimental/PCG/Source/PCG/Private/PCGActorAndComponentMapping.cpp
+++ b/Engine/Plugins/Experimental/PCG/Source/PCG/Private/PCGActorAndComponentMapping.cpp
@@ -758,7 +758,7 @@ void UPCGActorAndComponentMapping::AddActorsPostInit()
 	DelayedAddedActors.Empty();
 
 	// Also add the one tracked by the World Actor
-	if (APCGWorldActor* PCGWorldActor = PCGSubsystem->GetPCGWorldActor())
+	if (APCGWorldActor* PCGWorldActor = PCGSubsystem->FindPCGWorldActor())
 	{
 		// Making a copy, since CachedTrackedActors can be modified (if an actor is no longer tracked)
 		TSet<TWeakObjectPtr<AActor>> CachedTrackedActorsCopy = PCGWorldActor->CachedTrackedActors;
@@ -796,7 +796,7 @@ void UPCGActorAndComponentMapping::OnActorAdded(AActor* InActor)
 bool UPCGActorAndComponentMapping::AddOrUpdateTrackedActor(AActor* InActor)
 {
 	// We have to make sure to not create a infinite loop
-	if (!InActor || InActor->IsA<APCGWorldActor>() || !PCGSubsystem || !PCGSubsystem->GetPCGWorldActor())
+	if (!InActor || InActor->IsA<APCGWorldActor>() || !PCGSubsystem || !PCGSubsystem->FindPCGWorldActor())
 	{
 		return false;
 	}
@@ -850,7 +850,7 @@ bool UPCGActorAndComponentMapping::AddOrUpdateTrackedActor(AActor* InActor)
 
 void UPCGActorAndComponentMapping::RegisterActor(AActor* InActor)
 {
-	APCGWorldActor* PCGWorldActor = PCGSubsystem ? PCGSubsystem->GetPCGWorldActor() : nullptr;
+	APCGWorldActor* PCGWorldActor = PCGSubsystem ? PCGSubsystem->FindPCGWorldActor() : nullptr;
 	if (!PCGWorldActor || !InActor)
 	{
 		return;
@@ -879,7 +879,7 @@ void UPCGActorAndComponentMapping::RegisterActor(AActor* InActor)
 
 bool UPCGActorAndComponentMapping::UnregisterActor(AActor* InActor)
 {
-	APCGWorldActor* PCGWorldActor = PCGSubsystem ? PCGSubsystem->GetPCGWorldActor() : nullptr;
+	APCGWorldActor* PCGWorldActor = PCGSubsystem ? PCGSubsystem->FindPCGWorldActor() : nullptr;
 	if (!PCGWorldActor || !InActor)
 	{
 		return false;
@@ -911,7 +911,7 @@ void UPCGActorAndComponentMapping::OnActorDeleted(AActor* InActor)
 {
 	TRACE_CPUPROFILER_EVENT_SCOPE(UPCGActorAndComponentMapping::OnActorDeleted);
 
-	APCGWorldActor* PCGWorldActor = PCGSubsystem ? PCGSubsystem->GetPCGWorldActor() : nullptr;
+	APCGWorldActor* PCGWorldActor = PCGSubsystem ? PCGSubsystem->FindPCGWorldActor() : nullptr;
 	if (!PCGWorldActor || !InActor || !PCGWorldActor->CachedTrackedActors.Contains(InActor))
 	{
 		return;
@@ -928,7 +928,7 @@ void UPCGActorAndComponentMapping::OnActorMoved(AActor* InActor)
 {
 	TRACE_CPUPROFILER_EVENT_SCOPE(UPCGActorAndComponentMapping::OnActorMoved);
 
-	APCGWorldActor* PCGWorldActor = PCGSubsystem ? PCGSubsystem->GetPCGWorldActor() : nullptr;
+	APCGWorldActor* PCGWorldActor = PCGSubsystem ? PCGSubsystem->FindPCGWorldActor() : nullptr;
 	if (!PCGWorldActor || !InActor || !PCGWorldActor->CachedTrackedActors.Contains(InActor))
 	{
 		return;
@@ -972,7 +972,7 @@ void UPCGActorAndComponentMapping::OnObjectPropertyChanged(UObject* InObject, FP
 	}
 
 	AActor* Actor = Cast<AActor>(InObject);
-	APCGWorldActor* PCGWorldActor = PCGSubsystem ? PCGSubsystem->GetPCGWorldActor() : nullptr;
+	APCGWorldActor* PCGWorldActor = PCGSubsystem ? PCGSubsystem->FindPCGWorldActor() : nullptr;
 
 	if (!PCGWorldActor)
 	{
diff --git a/Engine/Plugins/Experimental/PCG/Source/PCG/Private/PCGSubsystem.cpp b/Engine/Plugins/Experimental/PCG/Source/PCG/Private/PCGSubsystem.cpp
--- a/Engine/Plugins/Experimental/PCG/Source/PCG/Private/PCGSubsystem.cpp
+++ b/Engine/Plugins/Experimental/PCG/Source/PCG/Private/PCGSubsystem.cpp
@@ -230,6 +230,11 @@ APCGWorldActor* UPCGSubsystem::GetPCGWorldActor()
 	return PCGWorldActor;
 }
 
+APCGWorldActor* UPCGSubsystem::FindPCGWorldActor()
+{
+	return PCGWorldActor;
+}
+
 #if WITH_EDITOR
 void UPCGSubsystem::DestroyPCGWorldActor()
 {
diff --git a/Engine/Plugins/Experimental/PCG/Source/PCG/Public/PCGSubsystem.h b/Engine/Plugins/Experimental/PCG/Source/PCG/Public/PCGSubsystem.h
--- a/Engine/Plugins/Experimental/PCG/Source/PCG/Public/PCGSubsystem.h
+++ b/Engine/Plugins/Experimental/PCG/Source/PCG/Public/PCGSubsystem.h
@@ -73,6 +73,7 @@ public:
 	bool IsInitialized() const { return GraphExecutor != nullptr; }
 
 	APCGWorldActor* GetPCGWorldActor();
+	APCGWorldActor* FindPCGWorldActor();
 #if WITH_EDITOR
 	void DestroyPCGWorldActor();
 #endif
