Removed landscape cache init in the PCG world actor post init properties, was causing issues with some landscape proxies

#rnx

diff --git a/Engine/Plugins/Experimental/PCG/Source/PCG/Private/PCGWorldActor.cpp b/Engine/Plugins/Experimental/PCG/Source/PCG/Private/PCGWorldActor.cpp
--- a/Engine/Plugins/Experimental/PCG/Source/PCG/Private/PCGWorldActor.cpp
+++ b/Engine/Plugins/Experimental/PCG/Source/PCG/Private/PCGWorldActor.cpp
@@ -40,11 +40,12 @@ void APCGWorldActor::PostInitProperties()
 
 	if (!HasAnyFlags(RF_ClassDefaultObject))
 	{
-		if (LandscapeCacheObject.Get())
+		// Commented because it was causing issues with landscape proxies guids not being initialized.
+		/*if (LandscapeCacheObject.Get())
 		{
 			// Make sure landscape cache is ready to provide data immediately.
 			LandscapeCacheObject->Initialize();
-		}
+		}*/
 	}
 }
 
