[PCG] Defensive potential fix for a crash when dragging a BP actor into a level

#rnx

diff --git a/Engine/Plugins/Experimental/PCG/Source/PCG/Private/PCGComponent.cpp b/Engine/Plugins/Experimental/PCG/Source/PCG/Private/PCGComponent.cpp
--- a/Engine/Plugins/Experimental/PCG/Source/PCG/Private/PCGComponent.cpp
+++ b/Engine/Plugins/Experimental/PCG/Source/PCG/Private/PCGComponent.cpp
@@ -2083,7 +2083,10 @@ bool UPCGComponent::DirtyTrackedActor(AActor* InActor, bool bIntersect, const TS
 				}
 
 				const TWeakObjectPtr<const UPCGSettings>& Settings = SettingsAndCulling.Key;
-				GetSubsystem()->CleanFromCache(Settings->GetElement().Get(), Settings.Get());
+				if (ensure(Settings.IsValid()))
+				{
+					GetSubsystem()->CleanFromCache(Settings->GetElement().Get(), Settings.Get());
+				}
 
 				bWasDirtied = true;
 			}
