Tentative fix related to stat group in PCG

#rnx

diff --git a/Engine/Plugins/Experimental/PCG/Source/PCG/Public/PCGModule.h b/Engine/Plugins/Experimental/PCG/Source/PCG/Public/PCGModule.h
--- a/Engine/Plugins/Experimental/PCG/Source/PCG/Public/PCGModule.h
+++ b/Engine/Plugins/Experimental/PCG/Source/PCG/Public/PCGModule.h
@@ -8,7 +8,7 @@
 PCG_API DECLARE_LOG_CATEGORY_EXTERN(LogPCG, Log, All);
 
 // Stats
-DECLARE_STATS_GROUP(TEXT("PCG"), STATGROUP_PCG, STATCAT_Advanced)
+DECLARE_STATS_GROUP(TEXT("PCG"), STATGROUP_PCG, STATCAT_Advanced);
 
 // CVars
 
