Tentative fix for module issue in test/shipping builds

#rnx

diff --git a/Engine/Plugins/Experimental/PCG/Source/PCG/Public/PCGModule.h b/Engine/Plugins/Experimental/PCG/Source/PCG/Public/PCGModule.h
--- a/Engine/Plugins/Experimental/PCG/Source/PCG/Public/PCGModule.h
+++ b/Engine/Plugins/Experimental/PCG/Source/PCG/Public/PCGModule.h
@@ -2,7 +2,7 @@
 
 #pragma once
 
-#include "Stats/Stats2.h"
+#include "Stats/Stats.h"
 
 // Logs
 PCG_API DECLARE_LOG_CATEGORY_EXTERN(LogPCG, Log, All);
