[PCG] Fix Get Volume Data (Self) failing on PAs, due to returning the Spatial data gathered from CreateActorPCGDataCollection which is not compatible with the output pin type (Volume).

#rnx

diff --git a/Engine/Plugins/Experimental/PCG/Source/PCG/Private/Elements/PCGDataFromActor.cpp b/Engine/Plugins/Experimental/PCG/Source/PCG/Private/Elements/PCGDataFromActor.cpp
--- a/Engine/Plugins/Experimental/PCG/Source/PCG/Private/Elements/PCGDataFromActor.cpp
+++ b/Engine/Plugins/Experimental/PCG/Source/PCG/Private/Elements/PCGDataFromActor.cpp
@@ -6,6 +6,7 @@
 #include "PCGSubsystem.h"
 #include "Data/PCGPointData.h"
 #include "Data/PCGSpatialData.h"
+#include "Data/PCGVolumeData.h"
 #include "Elements/PCGMergeElement.h"
 #include "Grid/PCGPartitionActor.h"
 #include "Helpers/PCGHelpers.h"
@@ -304,14 +305,13 @@ void FPCGDataFromActorElement::MergeActorsIntoPointData(FPCGContext* Context, co
 	else // Stripped-down version of the normal code path with bParseActor = false
 	{
 		FPCGDataCollection DataToMerge;
-		auto DataFilter = [](EPCGDataType InDataType) { return true; };
 		const bool bParseActor = false;
 
 		for (AActor* Actor : FoundActors)
 		{
 			if (Actor)
 			{
-				FPCGDataCollection Collection = UPCGComponent::CreateActorPCGDataCollection(Actor, Context->SourceComponent.Get(), DataFilter, bParseActor);
+				FPCGDataCollection Collection = UPCGComponent::CreateActorPCGDataCollection(Actor, Context->SourceComponent.Get(), EPCGDataType::Any, bParseActor);
 				DataToMerge.TaggedData += Collection.TaggedData;
 			}
 		}
@@ -440,8 +440,7 @@ void FPCGDataFromActorElement::ProcessActor(FPCGContext* Context, const UPCGData
 	else
 	{
 		const bool bParseActor = (Settings->Mode != EPCGGetDataFromActorMode::GetSinglePoint);
-		auto DataFilter = [Settings](EPCGDataType InDataType) { return Settings->DataFilter(InDataType); };
-		FPCGDataCollection Collection = UPCGComponent::CreateActorPCGDataCollection(FoundActor, Context->SourceComponent.Get(), DataFilter, bParseActor);
+		FPCGDataCollection Collection = UPCGComponent::CreateActorPCGDataCollection(FoundActor, Context->SourceComponent.Get(), Settings->GetDataFilter(), bParseActor);
 		Outputs += Collection.TaggedData;
 	}
 }
diff --git a/Engine/Plugins/Experimental/PCG/Source/PCG/Private/PCGComponent.cpp b/Engine/Plugins/Experimental/PCG/Source/PCG/Private/PCGComponent.cpp
--- a/Engine/Plugins/Experimental/PCG/Source/PCG/Private/PCGComponent.cpp
+++ b/Engine/Plugins/Experimental/PCG/Source/PCG/Private/PCGComponent.cpp
@@ -1655,8 +1655,7 @@ UPCGData* UPCGComponent::CreateActorPCGData(AActor* Actor, bool bParseActor)
 
 UPCGData* UPCGComponent::CreateActorPCGData(AActor* Actor, const UPCGComponent* Component, bool bParseActor)
 {
-	auto AcceptAllData = [](EPCGDataType InDataType) { return true; };
-	FPCGDataCollection Collection = CreateActorPCGDataCollection(Actor, Component, AcceptAllData, bParseActor);
+	FPCGDataCollection Collection = CreateActorPCGDataCollection(Actor, Component, EPCGDataType::Any, bParseActor);
 	if (Collection.TaggedData.Num() > 1)
 	{
 		UPCGUnionData* Union = NewObject<UPCGUnionData>();
@@ -1677,7 +1676,7 @@ UPCGData* UPCGComponent::CreateActorPCGData(AActor* Actor, const UPCGComponent*
 	}
 }
 
-FPCGDataCollection UPCGComponent::CreateActorPCGDataCollection(AActor* Actor, const UPCGComponent* Component, const TFunction<bool(EPCGDataType)>& InDataFilter, bool bParseActor)
+FPCGDataCollection UPCGComponent::CreateActorPCGDataCollection(AActor* Actor, const UPCGComponent* Component, EPCGDataType InDataFilter, bool bParseActor)
 {
 	TRACE_CPUPROFILER_EVENT_SCOPE(UPCGComponent::CreateActorPCGData);
 	FPCGDataCollection Collection;
@@ -1695,10 +1694,10 @@ FPCGDataCollection UPCGComponent::CreateActorPCGDataCollection(AActor* Actor, co
 	// Some actor types we will forego full parsing to build strictly on the actor existence, such as partition actors, volumes and landscape
 	// TODO: add factory for extensibility
 	// TODO: review the !bParseActor cases - it might make sense to have just a point for a partition actor, even if we preintersect it.
-	APCGPartitionActor* PartitionActor = InDataFilter(EPCGDataType::Spatial) ? Cast<APCGPartitionActor>(Actor) : nullptr;
-	ALandscapeProxy* LandscapeActor = InDataFilter(EPCGDataType::Landscape) ? Cast<ALandscapeProxy>(Actor) : nullptr;
-	AVolume* VolumeActor = InDataFilter(EPCGDataType::Volume) ? Cast<AVolume>(Actor) : nullptr;
-	if (!bParseActor && InDataFilter(EPCGDataType::Point))
+	APCGPartitionActor* PartitionActor = (!!(InDataFilter & EPCGDataType::Spatial) || !!(InDataFilter & EPCGDataType::Volume)) ? Cast<APCGPartitionActor>(Actor) : nullptr;
+	ALandscapeProxy* LandscapeActor = !!(InDataFilter & EPCGDataType::Landscape) ? Cast<ALandscapeProxy>(Actor) : nullptr;
+	AVolume* VolumeActor = !!(InDataFilter & EPCGDataType::Volume) ? Cast<AVolume>(Actor) : nullptr;
+	if (!bParseActor && !!(InDataFilter & EPCGDataType::Point))
 	{
 		UPCGPointData* Data = NewObject<UPCGPointData>();
 		Data->InitializeFromActor(Actor);
@@ -1711,15 +1710,28 @@ FPCGDataCollection UPCGComponent::CreateActorPCGDataCollection(AActor* Actor, co
 	{
 		check(!Component || Component->GetOwner() == Actor); // Invalid processing otherwise because of the this usage
 
-		UPCGVolumeData* Data = NewObject<UPCGVolumeData>();
-		Data->Initialize(PartitionActor->GetFixedBounds());
+		UPCGVolumeData* VolumeData = NewObject<UPCGVolumeData>();
+		UPCGSpatialData* Result = VolumeData;
+		if (InDataFilter == EPCGDataType::Volume)
+		{
+			VolumeData->Initialize(PCGHelpers::GetGridBounds(Actor, Component));
+		}
+		else
+		{
+			VolumeData->Initialize(PartitionActor->GetFixedBounds());
 
-		UPCGComponent* OriginalComponent = Component ? PartitionActor->GetOriginalComponent(Component) : nullptr;
-		// Important note: we do NOT call the collection version here, as we want to have a union if that's the case
-		const UPCGSpatialData* OriginalComponentSpatialData = OriginalComponent ? Cast<const UPCGSpatialData>(OriginalComponent->GetActorPCGData()) : nullptr;
+			UPCGComponent* OriginalComponent = Component ? PartitionActor->GetOriginalComponent(Component) : nullptr;
+			// Important note: we do NOT call the collection version here, as we want to have a union if that's the case
+			const UPCGSpatialData* OriginalComponentSpatialData = OriginalComponent ? Cast<const UPCGSpatialData>(OriginalComponent->GetActorPCGData()) : nullptr;
+
+			if (OriginalComponentSpatialData)
+			{
+				Result = Result->IntersectWith(OriginalComponentSpatialData);
+			}
+		}
 
 		FPCGTaggedData& TaggedData = Collection.TaggedData.Emplace_GetRef();
-		TaggedData.Data = OriginalComponentSpatialData ? Cast<UPCGData>(Data->IntersectWith(OriginalComponentSpatialData)) : Cast<UPCGData>(Data);
+		TaggedData.Data = Result;
 		// No need to keep partition actor tags, though we might want to push PCG grid GUID at some point
 	}
 	else if (LandscapeActor)
@@ -1802,7 +1814,7 @@ FPCGDataCollection UPCGComponent::CreateActorPCGDataCollection(AActor* Actor, co
 		RemovePCGGeneratedEntries(Shapes);
 		RemoveDuplicatesFromPrimitives(Shapes);
 
-		if (InDataFilter(EPCGDataType::Spline))
+		if (!!(InDataFilter & EPCGDataType::Spline))
 		{
 			for (ULandscapeSplinesComponent* SplineComponent : LandscapeSplines)
 			{
@@ -1827,7 +1839,7 @@ FPCGDataCollection UPCGComponent::CreateActorPCGDataCollection(AActor* Actor, co
 			}
 		}
 
-		if (InDataFilter(EPCGDataType::Primitive))
+		if (!!(InDataFilter & EPCGDataType::Primitive))
 		{
 			for (UShapeComponent* ShapeComponent : Shapes)
 			{
@@ -1873,7 +1885,7 @@ FPCGDataCollection UPCGComponent::CreateActorPCGDataCollection(AActor* Actor, co
 	}
 
 	// Finally, if it's not a special actor and there are not parsed components, then return a single point at the actor position
-	if (Collection.TaggedData.IsEmpty() && InDataFilter(EPCGDataType::Point))
+	if (Collection.TaggedData.IsEmpty() && !!(InDataFilter & EPCGDataType::Point))
 	{
 		UPCGPointData* Data = NewObject<UPCGPointData>();
 		Data->InitializeFromActor(Actor);
diff --git a/Engine/Plugins/Experimental/PCG/Source/PCG/Public/Elements/PCGDataFromActor.h b/Engine/Plugins/Experimental/PCG/Source/PCG/Public/Elements/PCGDataFromActor.h
--- a/Engine/Plugins/Experimental/PCG/Source/PCG/Public/Elements/PCGDataFromActor.h
+++ b/Engine/Plugins/Experimental/PCG/Source/PCG/Public/Elements/PCGDataFromActor.h
@@ -56,7 +56,7 @@ public:
 
 public:
 	/** Override this to filter what kinds of data should be retrieved from the actor(s). */
-	virtual bool DataFilter(EPCGDataType InDataType) const { return true; }
+	virtual EPCGDataType GetDataFilter() const { return EPCGDataType::Any; }
 
 	/** Override this to change the default value the selector will revert to when changing the actor selection type */
 	virtual TSubclassOf<AActor> GetDefaultActorSelectorClass() const;
diff --git a/Engine/Plugins/Experimental/PCG/Source/PCG/Public/Elements/PCGTypedGetter.h b/Engine/Plugins/Experimental/PCG/Source/PCG/Public/Elements/PCGTypedGetter.h
--- a/Engine/Plugins/Experimental/PCG/Source/PCG/Public/Elements/PCGTypedGetter.h
+++ b/Engine/Plugins/Experimental/PCG/Source/PCG/Public/Elements/PCGTypedGetter.h
@@ -32,7 +32,7 @@ protected:
 
 public:
 	//~Begin UPCGDataFromActorSettings interface
-	virtual bool DataFilter(EPCGDataType InDataType) const override { return !!(InDataType & EPCGDataType::Landscape); }
+	virtual EPCGDataType GetDataFilter() const override { return EPCGDataType::Landscape; }
 	virtual TSubclassOf<AActor> GetDefaultActorSelectorClass() const override;
 	//~End UPCGDataFromActorSettings
 
@@ -72,7 +72,7 @@ protected:
 
 public:
 	//~Begin UPCGDataFromActorSettings interface
-	virtual bool DataFilter(EPCGDataType InDataType) const override { return !!(InDataType & EPCGDataType::PolyLine); }
+	virtual EPCGDataType GetDataFilter() const override { return EPCGDataType::PolyLine; }
 	//~End UPCGDataFromActorSettings
 };
 
@@ -98,7 +98,7 @@ protected:
 
 public:
 	//~Begin UPCGDataFromActorSettings interface
-	virtual bool DataFilter(EPCGDataType InDataType) const override { return !!(InDataType & EPCGDataType::Volume); }
+	virtual EPCGDataType GetDataFilter() const override { return EPCGDataType::Volume; }
 	//~End UPCGDataFromActorSettings
 };
 
@@ -124,6 +124,6 @@ protected:
 
 public:
 	//~Begin UPCGDataFromActorSettings interface
-	virtual bool DataFilter(EPCGDataType InDataType) const override { return !!(InDataType & EPCGDataType::Primitive); }
+	virtual EPCGDataType GetDataFilter() const override { return EPCGDataType::Primitive; }
 	//~End UPCGDataFromActorSettings
 };
diff --git a/Engine/Plugins/Experimental/PCG/Source/PCG/Public/PCGComponent.h b/Engine/Plugins/Experimental/PCG/Source/PCG/Public/PCGComponent.h
--- a/Engine/Plugins/Experimental/PCG/Source/PCG/Public/PCGComponent.h
+++ b/Engine/Plugins/Experimental/PCG/Source/PCG/Public/PCGComponent.h
@@ -276,7 +276,7 @@ public:
 	FBox GetLastGeneratedBounds() const { return LastGeneratedBounds; }
 
 	/** Builds the PCG data from a given actor and its PCG component, and places it in a data collection with appropriate tags */
-	static FPCGDataCollection CreateActorPCGDataCollection(AActor* Actor, const UPCGComponent* Component, const TFunction<bool(EPCGDataType)>& InDataFilter, bool bParseActor = true);
+	static FPCGDataCollection CreateActorPCGDataCollection(AActor* Actor, const UPCGComponent* Component, EPCGDataType InDataFilter, bool bParseActor = true);
 
 	/** Builds the canonical PCG data from a given actor and its PCG component if any. */
 	static UPCGData* CreateActorPCGData(AActor* Actor, const UPCGComponent* Component, bool bParseActor = true);
