[PCG] Provide a pin for dynamic graph overrides on Subgraphs

diff --git a/Engine/Plugins/Experimental/PCG/Source/PCG/Private/Elements/PCGUserParameterGet.cpp b/Engine/Plugins/Experimental/PCG/Source/PCG/Private/Elements/PCGUserParameterGet.cpp
--- a/Engine/Plugins/Experimental/PCG/Source/PCG/Private/Elements/PCGUserParameterGet.cpp
+++ b/Engine/Plugins/Experimental/PCG/Source/PCG/Private/Elements/PCGUserParameterGet.cpp
@@ -97,7 +97,7 @@ namespace PCGUserParameterGetSettings
 
 		// Finally, we go down the graph instance chain to find the first valid layout.
 		// In most cases, we should never enter this loop, as layouts should match. But it is a safeguard.
-		while (GraphInstance && GraphParameters.GetScriptStruct() != WantedScriptStruct)
+		while (GraphInstance && !WantedScriptStruct->CompareScriptStruct(WantedScriptStruct, GraphParameters.GetScriptStruct(), /*PortFlags=*/0))
 		{
 			GraphInstance = Cast<UPCGGraphInstance>(GraphInstance->Graph);
 			if (!GraphInstance)
diff --git a/Engine/Plugins/Experimental/PCG/Source/PCG/Private/Metadata/Accessors/PCGAttributeAccessorHelpers.cpp b/Engine/Plugins/Experimental/PCG/Source/PCG/Private/Metadata/Accessors/PCGAttributeAccessorHelpers.cpp
--- a/Engine/Plugins/Experimental/PCG/Source/PCG/Private/Metadata/Accessors/PCGAttributeAccessorHelpers.cpp
+++ b/Engine/Plugins/Experimental/PCG/Source/PCG/Private/Metadata/Accessors/PCGAttributeAccessorHelpers.cpp
@@ -292,6 +292,10 @@ namespace PCGAttributeAccessorHelpers
 		{
 			return Functor(Signature<FPCGPropertySoftPtrAccessor>{}, SoftObjectProperty);
 		}
+		else if (const FObjectProperty* ObjectProperty = CastField<FObjectProperty>(InProperty))
+		{
+			return Functor(Signature<FPCGPropertyObjectPtrAccessor>{}, ObjectProperty);
+		}
 		else if (const FStructProperty* StructProperty = CastField<FStructProperty>(InProperty))
 		{
 			if (StructProperty->Struct == TBaseStructure<FVector>::Get())
diff --git a/Engine/Plugins/Experimental/PCG/Source/PCG/Private/PCGContext.cpp b/Engine/Plugins/Experimental/PCG/Source/PCG/Private/PCGContext.cpp
--- a/Engine/Plugins/Experimental/PCG/Source/PCG/Private/PCGContext.cpp
+++ b/Engine/Plugins/Experimental/PCG/Source/PCG/Private/PCGContext.cpp
@@ -135,6 +135,7 @@ void FPCGContext::InitializeSettings()
 
 				// Force seed copy to prevent issue due to delta serialization vs. Seed being initialized in the constructor only for new nodes
 				SettingsWithOverride->Seed = NodeSettings->Seed;
+				SettingsWithOverride->OriginalSettings = NodeSettings;
 			}
 		}
 	}
@@ -292,4 +293,4 @@ bool FPCGContext::HasVisualLogs() const
 }
 #endif // WITH_EDITOR
 
-#undef LOCTEXT_NAMESPACE
+#undef LOCTEXT_NAMESPACE
\ No newline at end of file
diff --git a/Engine/Plugins/Experimental/PCG/Source/PCG/Private/PCGSubgraph.cpp b/Engine/Plugins/Experimental/PCG/Source/PCG/Private/PCGSubgraph.cpp
--- a/Engine/Plugins/Experimental/PCG/Source/PCG/Private/PCGSubgraph.cpp
+++ b/Engine/Plugins/Experimental/PCG/Source/PCG/Private/PCGSubgraph.cpp
@@ -13,6 +13,8 @@
 
 #include UE_INLINE_GENERATED_CPP_BY_NAME(PCGSubgraph)
 
+#define LOCTEXT_NAMESPACE "PCGSubgraphElement"
+
 namespace PCGSubgraphSettings
 {
 	void RemoveAdvancedModeOnConnectedPins(const UPCGGraph* Subgraph, TArray<FPCGPinProperties>& InOutPinProperties, const bool bIsInput)
@@ -264,6 +266,19 @@ void UPCGSubgraphSettings::PostLoad()
 	Super::PostLoad();
 }
 
+#if WITH_EDITOR
+void UPCGSubgraphSettings::PostEditChangeProperty(struct FPropertyChangedEvent& PropertyChangedEvent)
+{
+	if (PropertyChangedEvent.Property && PropertyChangedEvent.Property->GetFName() == GET_MEMBER_NAME_CHECKED(UPCGSubgraphSettings, SubgraphInstance))
+	{
+		// Also rebuild the overrides
+		InitializeCachedOverridableParams(/*bReset=*/true);
+	}
+
+	Super::PostEditChangeProperty(PropertyChangedEvent);
+}
+#endif
+
 #if WITH_EDITOR
 UObject* UPCGSubgraphSettings::GetJumpTargetForDoubleClick() const
 {
@@ -284,6 +299,39 @@ FPCGElementPtr UPCGSubgraphSettings::CreateElement() const
 	return MakeShared<FPCGSubgraphElement>();
 }
 
+bool UPCGSubgraphSettings::IsDynamicGraph() const
+{
+	UPCGNode* Node = Cast<UPCGNode>(GetOuter());
+
+	if (!Node && OriginalSettings)
+	{
+		Node = Cast<UPCGNode>(OriginalSettings->GetOuter());
+
+	}
+
+	if (!Node)
+	{
+		return false;
+	}
+
+	const FName PropertyName = GET_MEMBER_NAME_CHECKED(UPCGSubgraphSettings, SubgraphOverride);
+
+	const FPCGSettingsOverridableParam* Param = CachedOverridableParams.FindByPredicate([PropertyName](const FPCGSettingsOverridableParam& ParamToCheck)
+	{
+		return !ParamToCheck.PropertiesNames.IsEmpty() && ParamToCheck.PropertiesNames[0] == PropertyName;
+	});
+
+	if (Param)
+	{
+		if (const UPCGPin* Pin = Node->GetInputPin(Param->Label))
+		{
+			return Pin->IsConnected();
+		}
+	}
+
+	return false;
+}
+
 TObjectPtr<UPCGGraph> UPCGBaseSubgraphNode::GetSubgraph() const
 {
 	UPCGGraphInterface* SubgraphInterface = GetSubgraphInterface();
@@ -410,12 +458,26 @@ bool FPCGSubgraphElement::ExecuteInternal(FPCGContext* InContext) const
 	const UPCGSubgraphSettings* Settings = Context->GetInputSettings<UPCGSubgraphSettings>();
 	check(Settings);
 
+	// TODO: only prevents A->A recursion, not A->B->A or cycles in general
+	if (Context->SourceComponent->GetGraph() == Settings->GetSubgraph())
+	{
+		PCGE_LOG(Error, GraphAndLog, LOCTEXT("FailedRecursiveSubgraph", "PCGGraph cannot include itself as a subgraph, subgraph will not be executed."));
+		return true;
+	}
+
 	if (Settings->IsDynamicGraph())
 	{
 		TRACE_CPUPROFILER_EVENT_SCOPE(FPCGSubgraphElement::Execute);
 
 		if (!Context->bScheduledSubgraph)
 		{
+			if (Settings->SubgraphInstance)
+			{
+				// If OriginalSettings is null, then we ARE the original settings, and writing over the existing graph is incorrect (and potentially a race condition)
+				check(Settings->OriginalSettings);
+				Settings->SubgraphInstance->SetGraph(Settings->SubgraphOverride);
+			}
+
 			UPCGGraph* Subgraph = Settings->GetSubgraph();
 			UPCGSubsystem* Subsystem = Context->SourceComponent.IsValid() ? Context->SourceComponent->GetSubsystem() : nullptr;
 
@@ -532,3 +594,5 @@ bool FPCGInputForwardingElement::ExecuteInternal(FPCGContext* Context) const
 	Context->OutputData = Input;
 	return true;
 }
+
+#undef LOCTEXT_NAMESPACE
\ No newline at end of file
diff --git a/Engine/Plugins/Experimental/PCG/Source/PCG/Public/Metadata/Accessors/PCGPropertyAccessor.h b/Engine/Plugins/Experimental/PCG/Source/PCG/Public/Metadata/Accessors/PCGPropertyAccessor.h
--- a/Engine/Plugins/Experimental/PCG/Source/PCG/Public/Metadata/Accessors/PCGPropertyAccessor.h
+++ b/Engine/Plugins/Experimental/PCG/Source/PCG/Public/Metadata/Accessors/PCGPropertyAccessor.h
@@ -294,4 +294,43 @@ public:
 
 private:
 	const FSoftObjectProperty* Property = nullptr;
+};
+
+/**
+* Templated accessor class for object/class ptr properties.
+* Do not instantiate it manually, use PCGAttributeAccessorHelpers::CreatePropertyAccessor.
+* Will always convert to FString for PCG
+* Key supported: Generic object
+*/
+class FPCGPropertyObjectPtrAccessor : public IPCGAttributeAccessorT<FPCGPropertyObjectPtrAccessor>
+{
+public:
+	using Type = FString;
+	using Super = IPCGAttributeAccessorT<FPCGPropertyObjectPtrAccessor>;
+
+	FPCGPropertyObjectPtrAccessor(const FObjectProperty* InProperty)
+		: Super(/*bInReadOnly=*/ false)
+		, Property(InProperty)
+	{
+		check(Property);
+	}
+
+	bool GetRangeImpl(TArrayView<Type> OutValues, int32 Index, const IPCGAttributeAccessorKeys& Keys) const
+	{
+		return PCGPropertyAccessor::IterateGet(Property, OutValues, Index, Keys, [this](const void* PropertyAddressData) -> Type
+		{
+			return Property->GetPropertyValue(PropertyAddressData).GetPath();
+		});
+	}
+
+	bool SetRangeImpl(TArrayView<const Type> InValues, int32 Index, IPCGAttributeAccessorKeys& Keys, EPCGAttributeAccessorFlags)
+	{
+		return PCGPropertyAccessor::IterateSet(Property, InValues, Index, Keys, [this](void* PropertyAddressData, const Type& Value) -> void
+		{
+			Property->SetPropertyValue(PropertyAddressData, FSoftObjectPath(Value).TryLoad());
+		});
+	}
+
+private:
+	const FObjectProperty* Property = nullptr;
 };
\ No newline at end of file
diff --git a/Engine/Plugins/Experimental/PCG/Source/PCG/Public/PCGSettings.h b/Engine/Plugins/Experimental/PCG/Source/PCG/Public/PCGSettings.h
--- a/Engine/Plugins/Experimental/PCG/Source/PCG/Public/PCGSettings.h
+++ b/Engine/Plugins/Experimental/PCG/Source/PCG/Public/PCGSettings.h
@@ -323,6 +323,9 @@ public:
 	FText Description;
 #endif
 
+	// Holds the original settings used to duplicate this object if it was overridden
+	const UPCGSettings* OriginalSettings = nullptr;
+
 protected:
 	// Returns an array of all the input pin properties. You should not add manually a "params" pin, it is handled automatically by FillOverridableParamsPins
 	virtual TArray<FPCGPinProperties> InputPinProperties() const;
@@ -468,4 +471,4 @@ public:
 protected:
 	virtual bool ExecuteInternal(FPCGContext* Context) const override;
 	virtual bool IsPassthrough(const UPCGSettings* InSettings) const override { return true; }
-};
+};
\ No newline at end of file
diff --git a/Engine/Plugins/Experimental/PCG/Source/PCG/Public/PCGSubgraph.h b/Engine/Plugins/Experimental/PCG/Source/PCG/Public/PCGSubgraph.h
--- a/Engine/Plugins/Experimental/PCG/Source/PCG/Public/PCGSubgraph.h
+++ b/Engine/Plugins/Experimental/PCG/Source/PCG/Public/PCGSubgraph.h
@@ -73,6 +73,9 @@ public:
 protected:
 	//~Begin UObject interface implementation
 	virtual void PostLoad() override;
+#if WITH_EDITOR
+	virtual void PostEditChangeProperty(struct FPropertyChangedEvent& PropertyChangedEvent) override;
+#endif
 	//~End UObject interface implementation
 
 public:
@@ -95,6 +98,7 @@ protected:
 	//~Begin UPCGBaseSubgraphSettings interface
 public:
 	virtual UPCGGraphInterface* GetSubgraphInterface() const override { return SubgraphInstance.Get(); }
+	virtual bool IsDynamicGraph() const override;
 protected:
 	virtual void SetSubgraphInternal(UPCGGraphInterface* InGraph) override;
 #if WITH_EDITOR
@@ -106,6 +110,9 @@ public:
 	UPROPERTY(BlueprintReadOnly, VisibleAnywhere, Category = Properties, Instanced, meta = (NoResetToDefault))
 	TObjectPtr<UPCGGraphInstance> SubgraphInstance;
 
+	UPROPERTY(BlueprintReadOnly, Category = Properties, meta = (PCG_Overridable))
+	TObjectPtr<UPCGGraphInterface> SubgraphOverride;
+
 #if WITH_EDITORONLY_DATA
 	UPROPERTY()
 	TObjectPtr<UPCGGraph> Subgraph_DEPRECATED;
