[PCG] Fix pass through disabled nodes which fails for nodes that have multiple outputs

Building input data for nodes downstream of the disabled node does a GetInputsByPin(Label) which fails because the labels are not setup properly during pass-through.

#rnx

diff --git a/Engine/Plugins/Experimental/PCG/Source/PCG/Private/PCGElement.cpp b/Engine/Plugins/Experimental/PCG/Source/PCG/Private/PCGElement.cpp
--- a/Engine/Plugins/Experimental/PCG/Source/PCG/Private/PCGElement.cpp
+++ b/Engine/Plugins/Experimental/PCG/Source/PCG/Private/PCGElement.cpp
@@ -254,17 +254,22 @@ void IPCGElement::DisabledPassThroughData(FPCGContext* Context) const
 		return;
 	}
 
-	const UPCGPin* PassThroughPin = Context->Node->GetPassThroughInputPin();
-	if (PassThroughPin == nullptr)
+	const UPCGPin* PassThroughInputPin = Context->Node->GetPassThroughInputPin();
+	const UPCGPin* PassThroughOutputPin = Context->Node->GetPassThroughOutputPin();
+	if (!PassThroughInputPin || !PassThroughOutputPin)
 	{
-		// No pin to grab pass through data from
+		// No pin to grab pass through data from or to pass data to.
 		return;
 	}
 
-	// Grab data from pass-through pin
-	Context->OutputData.TaggedData = Context->InputData.GetInputsByPin(PassThroughPin->Properties.Label);
+	const EPCGDataType OutputType = PassThroughOutputPin->GetCurrentTypes();
 
-	const EPCGDataType OutputType = Context->Node->GetOutputPins()[0]->GetCurrentTypes();
+	// Grab data from pass-through pin, push it all to output pin
+	Context->OutputData.TaggedData = Context->InputData.GetInputsByPin(PassThroughInputPin->Properties.Label);
+	for (FPCGTaggedData& Data : Context->OutputData.TaggedData)
+	{
+		Data.Pin = PassThroughOutputPin->Properties.Label;
+	}
 
 	// Pass through input data if it is not params, and if the output type supports it (e.g. if we have a incoming
 	// surface connected to an input pin of type Any, do not pass the surface through to an output pin of type Point).
@@ -286,7 +291,7 @@ void IPCGElement::DisabledPassThroughData(FPCGContext* Context) const
 	if (Settings->OnlyPassThroughOneEdgeWhenDisabled())
 	{
 		// Find first incoming non-params data that is coming through the pass through pin
-		TArray<FPCGTaggedData> InputsOnFirstPin = Context->InputData.GetInputsByPin(PassThroughPin->Properties.Label);
+		TArray<FPCGTaggedData> InputsOnFirstPin = Context->InputData.GetInputsByPin(PassThroughInputPin->Properties.Label);
 		const int FirstNonParamsDataIndex = InputsOnFirstPin.IndexOfByPredicate(InputDataShouldPassThrough);
 
 		if (FirstNonParamsDataIndex != INDEX_NONE)
@@ -309,7 +314,7 @@ void IPCGElement::DisabledPassThroughData(FPCGContext* Context) const
 	else
 	{
 		// Remove any incoming non-params data that is coming through the pass through pin
-		TArray<FPCGTaggedData> InputsOnFirstPin = Context->InputData.GetInputsByPin(PassThroughPin->Properties.Label);
+		TArray<FPCGTaggedData> InputsOnFirstPin = Context->InputData.GetInputsByPin(PassThroughInputPin->Properties.Label);
 		for (int Index = InputsOnFirstPin.Num() - 1; Index >= 0; --Index)
 		{
 			const FPCGTaggedData& Data = InputsOnFirstPin[Index];
diff --git a/Engine/Plugins/Experimental/PCG/Source/PCG/Private/PCGNode.cpp b/Engine/Plugins/Experimental/PCG/Source/PCG/Private/PCGNode.cpp
--- a/Engine/Plugins/Experimental/PCG/Source/PCG/Private/PCGNode.cpp
+++ b/Engine/Plugins/Experimental/PCG/Source/PCG/Private/PCGNode.cpp
@@ -440,14 +440,13 @@ int32 UPCGNode::GetInboundEdgesNum() const
 
 const UPCGPin* UPCGNode::GetPassThroughInputPin() const
 {
-	// No outputs means nothing will be passed through
-	if (GetOutputPins().Num() == 0)
+	const UPCGPin* PassThroughOutput = GetPassThroughOutputPin();
+	if (!PassThroughOutput)
 	{
 		return nullptr;
 	}
 
-	const EPCGDataType OutputType = GetOutputPins()[0]->GetCurrentTypes();
-
+	const EPCGDataType OutputType = PassThroughOutput->GetCurrentTypes();
 	// We assume a node whose primary output is params is a param processing node.
 	const bool bNodeOutputsParams = (OutputType == EPCGDataType::Param);
 
@@ -488,6 +487,12 @@ const UPCGPin* UPCGNode::GetPassThroughInputPin() const
 	}
 }
 
+const UPCGPin* UPCGNode::GetPassThroughOutputPin() const
+{
+	const int32 PrimaryOutputPinIndex = 0;
+	return GetOutputPins().Num() > PrimaryOutputPinIndex ? GetOutputPins()[PrimaryOutputPinIndex] : nullptr;
+}
+
 bool UPCGNode::IsPinUsedByNodeExecution(const UPCGPin* InPin) const
 {
 	check(InPin);
diff --git a/Engine/Plugins/Experimental/PCG/Source/PCG/Public/PCGNode.h b/Engine/Plugins/Experimental/PCG/Source/PCG/Public/PCGNode.h
--- a/Engine/Plugins/Experimental/PCG/Source/PCG/Public/PCGNode.h
+++ b/Engine/Plugins/Experimental/PCG/Source/PCG/Public/PCGNode.h
@@ -113,6 +113,9 @@ public:
 	/** Pin from which data is passed through when this node is disabled. */
 	virtual const UPCGPin* GetPassThroughInputPin() const;
 
+	/** Pin to which data is passed through when this node is disabled. */
+	virtual const UPCGPin* GetPassThroughOutputPin() const;
+
 	/** True if the pin is being used by the node. UI will gray out unused pins. */
 	virtual bool IsPinUsedByNodeExecution(const UPCGPin* InPin) const;
 
