Copy metadata node now works properly with accessors

diff --git a/Engine/Plugins/Experimental/PCG/Source/PCG/Private/Elements/Metadata/PCGMetadataElement.cpp b/Engine/Plugins/Experimental/PCG/Source/PCG/Private/Elements/Metadata/PCGMetadataElement.cpp
--- a/Engine/Plugins/Experimental/PCG/Source/PCG/Private/Elements/Metadata/PCGMetadataElement.cpp
+++ b/Engine/Plugins/Experimental/PCG/Source/PCG/Private/Elements/Metadata/PCGMetadataElement.cpp
@@ -163,7 +163,7 @@ bool FPCGMetadataOperationElement::ExecuteInternal(FPCGContext* Context) const
 		const FName LocalSourceAttribute = InputSource.GetName();
 		const FName LocalDestinationAttribute = OutputTarget.GetName();
 
-		if (Settings->InputSource.GetSelection() == EPCGAttributePropertySelection::Attribute && !SourceMetadata->HasAttribute(LocalSourceAttribute))
+		if (InputSource.GetSelection() == EPCGAttributePropertySelection::Attribute && !SourceMetadata->HasAttribute(LocalSourceAttribute))
 		{
 			PCGE_LOG(Warning, GraphAndLog, FText::Format(LOCTEXT("InputMissingAttribute", "Input does not have the '{0}' attribute"), FText::FromName(LocalSourceAttribute)));
 			continue;
@@ -182,7 +182,8 @@ bool FPCGMetadataOperationElement::ExecuteInternal(FPCGContext* Context) const
 		SampledPoints = Points;
 
 		// If it is attribute to attribute, just copy the attributes, if they exist and are valid
-		if (Settings->InputSource.GetSelection() == EPCGAttributePropertySelection::Attribute && OutputTarget.GetSelection() == EPCGAttributePropertySelection::Attribute)
+		if (InputSource.GetSelection() == EPCGAttributePropertySelection::Attribute && InputSource.GetExtraNames().IsEmpty() &&
+			OutputTarget.GetSelection() == EPCGAttributePropertySelection::Attribute && OutputTarget.GetExtraNames().IsEmpty())
 		{
 			if (LocalSourceAttribute == LocalDestinationAttribute)
 			{
@@ -230,8 +231,8 @@ bool FPCGMetadataOperationElement::ExecuteInternal(FPCGContext* Context) const
 			}
 		}
 
-		TUniquePtr<IPCGAttributeAccessor> OutputAccessor = PCGAttributeAccessorHelpers::CreateAccessor(SampledData, Settings->OutputTarget);
-		TUniquePtr<IPCGAttributeAccessorKeys> OutputKeys = PCGAttributeAccessorHelpers::CreateKeys(SampledData, Settings->OutputTarget);
+		TUniquePtr<IPCGAttributeAccessor> OutputAccessor = PCGAttributeAccessorHelpers::CreateAccessor(SampledData, OutputTarget);
+		TUniquePtr<IPCGAttributeAccessorKeys> OutputKeys = PCGAttributeAccessorHelpers::CreateKeys(SampledData, OutputTarget);
 
 		if (!OutputAccessor.IsValid() || !OutputKeys.IsValid())
 		{
