[PCG] Material overrides that don't exist will now trigger a warning in the graph and continue to execute, rather than an error and aborting.

#rnx

diff --git a/Engine/Plugins/Experimental/PCG/Source/PCG/Private/MeshSelectors/PCGMeshSelectorBase.cpp b/Engine/Plugins/Experimental/PCG/Source/PCG/Private/MeshSelectors/PCGMeshSelectorBase.cpp
--- a/Engine/Plugins/Experimental/PCG/Source/PCG/Private/MeshSelectors/PCGMeshSelectorBase.cpp
+++ b/Engine/Plugins/Experimental/PCG/Source/PCG/Private/MeshSelectors/PCGMeshSelectorBase.cpp
@@ -62,8 +62,8 @@ void FPCGMeshMaterialOverrideHelper::Initialize(FPCGContext& InContext)
 		{
 			if (!Metadata->HasAttribute(MaterialOverrideAttributeName))
 			{
-				PCGE_LOG_C(Error, GraphAndLog, &InContext, FText::Format(LOCTEXT("AttributeMissing", "Attribute '{0}' for material overrides is not present in the metadata"), FText::FromName(MaterialOverrideAttributeName)));
-				return;
+				PCGE_LOG_C(Warning, GraphAndLog, &InContext, FText::Format(LOCTEXT("AttributeMissing", "Attribute '{0}' for material overrides is not present in the metadata"), FText::FromName(MaterialOverrideAttributeName)));
+				continue;
 			}
 
 			const FPCGMetadataAttributeBase* MaterialAttributeBase = Metadata->GetConstAttribute(MaterialOverrideAttributeName);
