[PCG] Add grid cell volume output to grid size node, and add Cull Points Outside Actor Bounds node

diff --git a/Engine/Plugins/Experimental/PCG/Source/PCG/Private/Elements/PCGCullPointsOutsideActorBounds.cpp b/Engine/Plugins/Experimental/PCG/Source/PCG/Private/Elements/PCGCullPointsOutsideActorBounds.cpp
new file mode 100644
--- /dev/null
+++ b/Engine/Plugins/Experimental/PCG/Source/PCG/Private/Elements/PCGCullPointsOutsideActorBounds.cpp
@@ -0,0 +1,99 @@
+// Copyright Epic Games, Inc. All Rights Reserved.
+
+#include "Elements/PCGCullPointsOutsideActorBounds.h"
+
+#include "PCGComponent.h"
+#include "PCGContext.h"
+#include "PCGPin.h"
+#include "Data/PCGIntersectionData.h"
+#include "Data/PCGPointData.h"
+#include "Data/PCGVolumeData.h"
+
+#define LOCTEXT_NAMESPACE "PCGCullPointsOutsideActorBoundsElement"
+
+#if WITH_EDITOR
+FText UPCGCullPointsOutsideActorBoundsSettings::GetDefaultNodeTitle() const
+{
+	return LOCTEXT("NodeTitle", "Cull Points Outside Actor Bounds");
+}
+
+FText UPCGCullPointsOutsideActorBoundsSettings::GetNodeTooltipText() const
+{
+	return LOCTEXT("NodeTooltip", "Culls points that lie outside the current actor bounds.");
+}
+#endif
+
+TArray<FPCGPinProperties> UPCGCullPointsOutsideActorBoundsSettings::InputPinProperties() const
+{
+	return Super::DefaultPointInputPinProperties();
+}
+
+TArray<FPCGPinProperties> UPCGCullPointsOutsideActorBoundsSettings::OutputPinProperties() const
+{
+	return Super::DefaultPointOutputPinProperties();
+}
+
+FPCGElementPtr UPCGCullPointsOutsideActorBoundsSettings::CreateElement() const
+{
+	return MakeShared<FPCGCullPointsOutsideActorBoundsElement>();
+}
+
+bool FPCGCullPointsOutsideActorBoundsElement::ExecuteInternal(FPCGContext* Context) const
+{
+	if (!Context->SourceComponent.IsValid())
+	{
+		return true;
+	}
+
+	UPCGVolumeData* VolumeData = NewObject<UPCGVolumeData>();
+	check(VolumeData);
+
+	// Initialize directly. Could also have gone through PCGComponent::CreateActorPCGDataCollection.
+	VolumeData->Initialize(Context->SourceComponent->GetGridBounds());
+
+	const TArray<FPCGTaggedData>& Inputs = Context->InputData.TaggedData;
+	for (const FPCGTaggedData& Input : Inputs)
+	{
+		const UPCGPointData* InputPointData = Cast<UPCGPointData>(Input.Data);
+
+		// Skip non-points or empty point data
+		if (!InputPointData || InputPointData->GetPoints().IsEmpty())
+		{
+			continue;
+		}
+
+		UPCGIntersectionData* Intersection = InputPointData->IntersectWith(VolumeData);
+		check(Intersection);
+
+		const UPCGPointData* IntersectedPoints = Intersection->ToPointData(Context);
+		check(IntersectedPoints);
+
+		// Skip empty result
+		if (IntersectedPoints->GetPoints().IsEmpty())
+		{
+			continue;
+		}
+
+		FPCGTaggedData& Output = Context->OutputData.TaggedData.Emplace_GetRef(Input);
+		Output.Data = IntersectedPoints;
+	}
+
+	return true;
+}
+
+void FPCGCullPointsOutsideActorBoundsElement::GetDependenciesCrc(const FPCGDataCollection& InInput, const UPCGSettings* InSettings, UPCGComponent* InComponent, FPCGCrc& OutCrc) const
+{
+	FPCGCrc Crc;
+	FSimplePCGElement::GetDependenciesCrc(InInput, InSettings, InComponent, Crc);
+
+	// The culling volume depends on the component transform.
+	const UPCGData* ActorData = InComponent ? InComponent->GetActorPCGData() : nullptr;
+	if (ActorData)
+	{
+		Crc.Combine(ActorData->GetOrComputeCrc(/*bFullDataCrc=*/false));
+	}
+
+	OutCrc = Crc;
+}
+
+#undef LOCTEXT_NAMESPACE
diff --git a/Engine/Plugins/Experimental/PCG/Source/PCG/Private/Elements/PCGHiGenGridSize.cpp b/Engine/Plugins/Experimental/PCG/Source/PCG/Private/Elements/PCGHiGenGridSize.cpp
--- a/Engine/Plugins/Experimental/PCG/Source/PCG/Private/Elements/PCGHiGenGridSize.cpp
+++ b/Engine/Plugins/Experimental/PCG/Source/PCG/Private/Elements/PCGHiGenGridSize.cpp
@@ -7,6 +7,7 @@
 #include "PCGGraph.h"
 #include "PCGModule.h"
 #include "PCGPin.h"
+#include "Data/PCGVolumeData.h"
 
 #define LOCTEXT_NAMESPACE "PCGHiGenGridSizeElement"
 
@@ -14,6 +15,7 @@ namespace PCGHiGenGridSizeConstants
 {
 	const FName NodeName = TEXT("HiGenGridSize");
 	const FText NodeTitle = LOCTEXT("NodeTitle", "Grid Size");
+	const FName CellVolumeOutputLabel = TEXT("Grid Cell Volume");
 }
 
 #if WITH_EDITOR
@@ -36,7 +38,8 @@ FText UPCGHiGenGridSizeSettings::GetNodeTooltipText() const
 EPCGDataType UPCGHiGenGridSizeSettings::GetCurrentPinTypes(const UPCGPin* InPin) const
 {
 	check(InPin);
-	if (!InPin->IsOutputPin())
+	// Non-dynamically-typed pins
+	if (!InPin->IsOutputPin() || InPin->Properties.Label == PCGHiGenGridSizeConstants::CellVolumeOutputLabel)
 	{
 		return Super::GetCurrentPinTypes(InPin);
 	}
@@ -58,6 +61,12 @@ TArray<FPCGPinProperties> UPCGHiGenGridSizeSettings::OutputPinProperties() const
 {
 	TArray<FPCGPinProperties> PinProperties;
 	PinProperties.Emplace(PCGPinConstants::DefaultOutputLabel, EPCGDataType::Any);
+	PinProperties.Emplace(
+		PCGHiGenGridSizeConstants::CellVolumeOutputLabel,
+		EPCGDataType::Spatial,
+		/*bAllowMultipleConnections=*/true,
+		/*bAllowMultipleData=*/true,
+		LOCTEXT("VolumeOutputPinTooltip", "The volume of the current grid cell."));
 
 	return PinProperties;
 }
@@ -97,9 +106,7 @@ bool UPCGHiGenGridSizeSettings::IsStructuralProperty(const FName& InPropertyName
 
 bool FPCGHiGenGridSizeElement::ExecuteInternal(FPCGContext* Context) const
 {
-	// Trivial pass through. Will only execute on the prescribed grid.
-	Context->OutputData = Context->InputData;
-
+	// Validation
 	if (Context->Node && Context->Node->GetGraph() && !Context->Node->GetGraph()->IsHierarchicalGenerationEnabled())
 	{
 		PCGE_LOG(Warning, GraphAndLog, LOCTEXT("GridSizeUsedInNonHiGenGraph", "Grid Size node used in a non-hierarchical graph. Enable hierarchical generation in the graph settings or remove this node."));
@@ -111,7 +118,38 @@ bool FPCGHiGenGridSizeElement::ExecuteInternal(FPCGContext* Context) const
 		PCGE_LOG(Warning, GraphAndLog, LOCTEXT("NonPartitionedComponent", "Grid Size node used on a non-partitioned component and will have no effect. Is Partitioned must be enabled on the component."));
 	}
 
+	// Trivial pass through. Will only execute on the prescribed grid.
+	Context->OutputData = Context->InputData;
+	for (FPCGTaggedData& Data : Context->OutputData.TaggedData)
+	{
+		Data.Pin = PCGPinConstants::DefaultOutputLabel;
+	}
+
+	UPCGVolumeData* VolumeData = NewObject<UPCGVolumeData>();
+	check(VolumeData);
+	VolumeData->Initialize(Context->SourceComponent->GetGridBounds());
+
+	FPCGTaggedData& VolumeOutput = Context->OutputData.TaggedData.Emplace_GetRef();
+	VolumeOutput.Data = VolumeData;
+	VolumeOutput.Pin = PCGHiGenGridSizeConstants::CellVolumeOutputLabel;
+
 	return true;
 }
 
+void FPCGHiGenGridSizeElement::GetDependenciesCrc(const FPCGDataCollection& InInput, const UPCGSettings* InSettings, UPCGComponent* InComponent, FPCGCrc& OutCrc) const
+{
+	FPCGCrc Crc;
+	FSimplePCGElement::GetDependenciesCrc(InInput, InSettings, InComponent, Crc);
+
+	// The grid cell volume output depends on the component transform.
+	// NOTE: It might be interesting to only incorporate the transform if the cell volume output pin is connected (and if we have a node obviously).
+	const UPCGData* ActorData = InComponent ? InComponent->GetActorPCGData() : nullptr;
+	if (ActorData)
+	{
+		Crc.Combine(ActorData->GetOrComputeCrc(/*bFullDataCrc=*/false));
+	}
+
+	OutCrc = Crc;
+}
+
 #undef LOCTEXT_NAMESPACE
diff --git a/Engine/Plugins/Experimental/PCG/Source/PCG/Public/Elements/PCGCullPointsOutsideActorBounds.h b/Engine/Plugins/Experimental/PCG/Source/PCG/Public/Elements/PCGCullPointsOutsideActorBounds.h
new file mode 100644
--- /dev/null
+++ b/Engine/Plugins/Experimental/PCG/Source/PCG/Public/Elements/PCGCullPointsOutsideActorBounds.h
@@ -0,0 +1,40 @@
+// Copyright Epic Games, Inc. All Rights Reserved.
+#pragma once
+
+#include "PCGElement.h"
+#include "PCGSettings.h"
+
+#include "PCGCullPointsOutsideActorBounds.generated.h"
+
+/**
+ * Removes points that lie outside the current actor bounds.
+ */
+UCLASS(BlueprintType, ClassGroup = (Procedural))
+class PCG_API UPCGCullPointsOutsideActorBoundsSettings : public UPCGSettings
+{
+	GENERATED_BODY()
+
+public:
+	//~Begin UPCGSettings interface
+#if WITH_EDITOR
+	virtual FName GetDefaultNodeName() const override { return TEXT("CullPointsOutsideActorBounds"); }
+	virtual FText GetDefaultNodeTitle() const override;
+	virtual FText GetNodeTooltipText() const override;
+	virtual EPCGSettingsType GetType() const override { return EPCGSettingsType::Spatial; }
+#endif
+
+protected:
+	virtual TArray<FPCGPinProperties> InputPinProperties() const override;
+	virtual TArray<FPCGPinProperties> OutputPinProperties() const override;
+	virtual FPCGElementPtr CreateElement() const override;
+	//~End UPCGSettings interface
+};
+
+class FPCGCullPointsOutsideActorBoundsElement : public FSimplePCGElement
+{
+public:
+	virtual void GetDependenciesCrc(const FPCGDataCollection& InInput, const UPCGSettings* InSettings, UPCGComponent* InComponent, FPCGCrc& OutCrc) const override;
+
+protected:
+	virtual bool ExecuteInternal(FPCGContext* Context) const override;
+};
diff --git a/Engine/Plugins/Experimental/PCG/Source/PCG/Public/Elements/PCGHiGenGridSize.h b/Engine/Plugins/Experimental/PCG/Source/PCG/Public/Elements/PCGHiGenGridSize.h
--- a/Engine/Plugins/Experimental/PCG/Source/PCG/Public/Elements/PCGHiGenGridSize.h
+++ b/Engine/Plugins/Experimental/PCG/Source/PCG/Public/Elements/PCGHiGenGridSize.h
@@ -44,6 +44,9 @@ protected:
 
 class FPCGHiGenGridSizeElement : public FSimplePCGElement
 {
+public:
+	virtual void GetDependenciesCrc(const FPCGDataCollection& InInput, const UPCGSettings* InSettings, UPCGComponent* InComponent, FPCGCrc& OutCrc) const override;
+
 protected:
 	virtual bool ExecuteInternal(FPCGContext* Context) const override;
 };
