[PCG] Grid size editor visualization works with disabled Grid Size nodes

Fix broken memoization in the calculate grid size function

Fix graph grid size visualization not updating when you modify grid size nodes by adding an event

#rnx

diff --git a/Engine/Plugins/Experimental/PCG/Source/PCG/Private/PCGGraph.cpp b/Engine/Plugins/Experimental/PCG/Source/PCG/Private/PCGGraph.cpp
--- a/Engine/Plugins/Experimental/PCG/Source/PCG/Private/PCGGraph.cpp
+++ b/Engine/Plugins/Experimental/PCG/Source/PCG/Private/PCGGraph.cpp
@@ -878,10 +878,16 @@ void UPCGGraph::NotifyGraphParametersChanged(EPCGGraphParameterEvent InChangeTyp
 
 void UPCGGraph::OnNodeChanged(UPCGNode* InNode, EPCGChangeType ChangeType)
 {
-	if (ChangeType != EPCGChangeType::Cosmetic)
+	if (!!(ChangeType & EPCGChangeType::Structural) && Cast<UPCGHiGenGridSizeSettings>(InNode->GetSettings()))
 	{
-		FWriteScopeLock ScopedWriteLock(NodeToGridSizeLock);
-		NodeToGridSize.Reset();
+		// Update node to grid size map for grid size changes.
+		{
+			FWriteScopeLock ScopedWriteLock(NodeToGridSizeLock);
+			NodeToGridSize.Reset();
+		}
+
+		// Broadcast so that grid size visualization can be updated editor-side.
+		OnGraphGridSizesChangedDelegate.Broadcast(this);
 	}
 
 	if ((ChangeType & ~EPCGChangeType::Cosmetic) != EPCGChangeType::None)
@@ -1087,8 +1093,15 @@ uint32 UPCGGraph::GetNodeGenerationGridSize(const UPCGNode* InNode, uint32 InDef
 
 uint32 UPCGGraph::CalculateNodeGridSizeRecursive_Unsafe(const UPCGNode* InNode, uint32 InDefaultGridSize) const
 {
+	if (const uint32* CachedGridSize = NodeToGridSize.Find(InNode))
+	{
+		return *CachedGridSize;
+	}
+
 	uint32 GridSize = PCGHiGenGrid::UninitializedGridSize();
-	if (const UPCGHiGenGridSizeSettings* GridSizeSettings = Cast<UPCGHiGenGridSizeSettings>(InNode->GetSettings()))
+
+	const UPCGHiGenGridSizeSettings* GridSizeSettings = Cast<UPCGHiGenGridSizeSettings>(InNode->GetSettings());
+	if (GridSizeSettings && GridSizeSettings->bEnabled)
 	{
 		GridSize = GridSizeSettings->GetGridSize();
 	}
@@ -1145,6 +1158,7 @@ void UPCGGraphInstance::PostLoad()
 	if (Graph)
 	{
 		Graph->OnGraphChangedDelegate.AddUObject(this, &UPCGGraphInstance::OnGraphChanged);
+		Graph->OnGraphGridSizesChangedDelegate.AddUObject(this, &UPCGGraphInstance::OnGraphGridSizesChanged);
 		Graph->OnGraphParametersChangedDelegate.AddUObject(this, &UPCGGraphInstance::OnGraphParametersChanged);
 	}
 #endif // WITH_EDITOR
@@ -1156,6 +1170,7 @@ void UPCGGraphInstance::BeginDestroy()
 	if (Graph)
 	{
 		Graph->OnGraphChangedDelegate.RemoveAll(this);
+		Graph->OnGraphGridSizesChangedDelegate.RemoveAll(this);
 		Graph->OnGraphParametersChangedDelegate.RemoveAll(this);
 	}
 #endif // WITH_EDITOR
@@ -1176,6 +1191,7 @@ void UPCGGraphInstance::PreEditChange(FProperty* InProperty)
 	if (InProperty->GetFName() == GET_MEMBER_NAME_CHECKED(UPCGGraphInstance, Graph) && Graph)
 	{
 		Graph->OnGraphChangedDelegate.RemoveAll(this);
+		Graph->OnGraphGridSizesChangedDelegate.RemoveAll(this);
 		Graph->OnGraphParametersChangedDelegate.RemoveAll(this);
 	}
 }
@@ -1192,6 +1208,7 @@ void UPCGGraphInstance::PostEditChangeProperty(FPropertyChangedEvent& PropertyCh
 	if (PropertyChangedEvent.GetPropertyName() == GET_MEMBER_NAME_CHECKED(UPCGGraphInstance, Graph) && Graph)
 	{
 		Graph->OnGraphChangedDelegate.AddUObject(this, &UPCGGraphInstance::OnGraphChanged);
+		Graph->OnGraphGridSizesChangedDelegate.AddUObject(this, &UPCGGraphInstance::OnGraphGridSizesChanged);
 		Graph->OnGraphParametersChangedDelegate.AddUObject(this, &UPCGGraphInstance::OnGraphParametersChanged);
 
 		RefreshParameters(EPCGGraphParameterEvent::GraphChanged);
@@ -1209,6 +1226,7 @@ void UPCGGraphInstance::PreEditUndo()
 	if (Graph)
 	{
 		Graph->OnGraphChangedDelegate.RemoveAll(this);
+		Graph->OnGraphGridSizesChangedDelegate.RemoveAll(this);
 		Graph->OnGraphParametersChangedDelegate.RemoveAll(this);
 	}
 
@@ -1222,6 +1240,7 @@ void UPCGGraphInstance::PostEditUndo()
 	if (Graph)
 	{
 		Graph->OnGraphChangedDelegate.AddUObject(this, &UPCGGraphInstance::OnGraphChanged);
+		Graph->OnGraphGridSizesChangedDelegate.AddUObject(this, &UPCGGraphInstance::OnGraphGridSizesChanged);
 		Graph->OnGraphParametersChangedDelegate.AddUObject(this, &UPCGGraphInstance::OnGraphParametersChanged);
 	}
 }
@@ -1234,6 +1253,14 @@ void UPCGGraphInstance::OnGraphChanged(UPCGGraphInterface* InGraph, EPCGChangeTy
 	}
 }
 
+void UPCGGraphInstance::OnGraphGridSizesChanged(UPCGGraphInterface* InGraph)
+{
+	if (InGraph == Graph)
+	{
+		OnGraphGridSizesChangedDelegate.Broadcast(this);
+	}
+}
+
 bool UPCGGraphInstance::CanEditChange(const FProperty* InProperty) const
 {
 	// Graph can only be changed if it is in a PCGComponent (not local) or a PCGSubgraphSettings
@@ -1259,6 +1286,7 @@ void UPCGGraphInstance::TeardownCallbacks()
 	if (Graph)
 	{
 		Graph->OnGraphChangedDelegate.RemoveAll(this);
+		Graph->OnGraphGridSizesChangedDelegate.RemoveAll(this);
 		Graph->OnGraphParametersChangedDelegate.RemoveAll(this);
 	}
 }
@@ -1271,6 +1299,7 @@ void UPCGGraphInstance::FixCallbacks()
 	if (Graph)
 	{
 		Graph->OnGraphChangedDelegate.AddUObject(this, &UPCGGraphInstance::OnGraphChanged);
+		Graph->OnGraphGridSizesChangedDelegate.AddUObject(this, &UPCGGraphInstance::OnGraphGridSizesChanged);
 		Graph->OnGraphParametersChangedDelegate.AddUObject(this, &UPCGGraphInstance::OnGraphParametersChanged);
 	}
 }
@@ -1294,6 +1323,7 @@ void UPCGGraphInstance::SetGraph(UPCGGraphInterface* InGraph)
 	if (Graph)
 	{
 		Graph->OnGraphChangedDelegate.RemoveAll(this);
+		Graph->OnGraphGridSizesChangedDelegate.RemoveAll(this);
 		Graph->OnGraphParametersChangedDelegate.RemoveAll(this);
 	}
 #endif // WITH_EDITOR
@@ -1304,6 +1334,7 @@ void UPCGGraphInstance::SetGraph(UPCGGraphInterface* InGraph)
 	if (Graph)
 	{
 		Graph->OnGraphChangedDelegate.AddUObject(this, &UPCGGraphInstance::OnGraphChanged);
+		Graph->OnGraphGridSizesChangedDelegate.AddUObject(this, &UPCGGraphInstance::OnGraphGridSizesChanged);
 		Graph->OnGraphParametersChangedDelegate.AddUObject(this, &UPCGGraphInstance::OnGraphParametersChanged);
 	}
 #endif // WITH_EDITOR
diff --git a/Engine/Plugins/Experimental/PCG/Source/PCG/Public/PCGGraph.h b/Engine/Plugins/Experimental/PCG/Source/PCG/Public/PCGGraph.h
--- a/Engine/Plugins/Experimental/PCG/Source/PCG/Public/PCGGraph.h
+++ b/Engine/Plugins/Experimental/PCG/Source/PCG/Public/PCGGraph.h
@@ -27,8 +27,9 @@ enum class EPCGGraphParameterEvent
 };
 
 #if WITH_EDITOR
-	DECLARE_MULTICAST_DELEGATE_TwoParams(FOnPCGGraphChanged, UPCGGraphInterface* /*Graph*/, EPCGChangeType /*ChangeType*/);
-	DECLARE_MULTICAST_DELEGATE_ThreeParams(FOnPCGGraphParametersChanged, UPCGGraphInterface* /*Graph*/, EPCGGraphParameterEvent /*ChangeType*/, FName /*ChangedPropertyName*/);
+DECLARE_MULTICAST_DELEGATE_TwoParams(FOnPCGGraphChanged, UPCGGraphInterface* /*Graph*/, EPCGChangeType /*ChangeType*/);
+DECLARE_MULTICAST_DELEGATE_OneParam(FOnPCGGraphGridSizesChanged, UPCGGraphInterface* /*Graph*/);
+DECLARE_MULTICAST_DELEGATE_ThreeParams(FOnPCGGraphParametersChanged, UPCGGraphInterface* /*Graph*/, EPCGGraphParameterEvent /*ChangeType*/, FName /*ChangedPropertyName*/);
 #endif // WITH_EDITOR
 
 /**
@@ -99,6 +100,7 @@ public:
 
 #if WITH_EDITOR
 	FOnPCGGraphChanged OnGraphChangedDelegate;
+	FOnPCGGraphGridSizesChanged OnGraphGridSizesChangedDelegate;
 	FOnPCGGraphParametersChanged OnGraphParametersChangedDelegate;
 #endif // WITH_EDITOR
 };
@@ -326,6 +328,7 @@ public:
 protected:
 #if WITH_EDITOR
 	void OnGraphChanged(UPCGGraphInterface* InGraph, EPCGChangeType ChangeType);
+	void OnGraphGridSizesChanged(UPCGGraphInterface* InGraph);
 	void NotifyGraphParametersChanged(EPCGGraphParameterEvent InChangeType, FName InChangedPropertyName);
 	void OnGraphParametersChanged(UPCGGraphInterface* InGraph, EPCGGraphParameterEvent InChangeType, FName InChangedPropertyName);
 #endif
diff --git a/Engine/Plugins/Experimental/PCG/Source/PCGEditor/Private/PCGEditor.cpp b/Engine/Plugins/Experimental/PCG/Source/PCGEditor/Private/PCGEditor.cpp
--- a/Engine/Plugins/Experimental/PCG/Source/PCGEditor/Private/PCGEditor.cpp
+++ b/Engine/Plugins/Experimental/PCG/Source/PCGEditor/Private/PCGEditor.cpp
@@ -81,6 +81,11 @@ void FPCGEditor::Initialize(const EToolkitMode::Type InMode, const TSharedPtr<cl
 {
 	PCGGraphBeingEdited = InPCGGraph;
 
+	if (PCGGraphBeingEdited)
+	{
+		PCGGraphBeingEdited->OnGraphGridSizesChangedDelegate.AddRaw(this, &FPCGEditor::OnGraphGridSizesChanged);
+	}
+
 	PCGEditorGraph = NewObject<UPCGEditorGraph>(PCGGraphBeingEdited, UPCGEditorGraph::StaticClass(), NAME_None, RF_Transactional | RF_Transient);
 	PCGEditorGraph->Schema = UPCGEditorGraphSchema::StaticClass();
 	PCGEditorGraph->InitFromNodeGraph(InPCGGraph);
@@ -197,13 +202,14 @@ void FPCGEditor::SetComponentAndStackBeingInspected(UPCGComponent* InPCGComponen
 			PCGComponentBeingInspected->EnableInspection();
 		}
 
+		check(PCGEditorGraph);
 		for (UEdGraphNode* Node : PCGEditorGraph->Nodes)
 		{
 			if (UPCGEditorGraphNodeBase* PCGNode = Cast<UPCGEditorGraphNodeBase>(Node))
 			{
 				// Update now that component has changed. Will fire OnNodeChanged if necessary.
 				EPCGChangeType ChangeType = PCGNode->UpdateErrorsAndWarnings();
-				ChangeType |= PCGNode->UpdateGridSizeVisualisation(InPCGComponent);
+				ChangeType |= PCGNode->UpdateGridSizeVisualization(InPCGComponent);
 
 				if (ChangeType != EPCGChangeType::None)
 				{
@@ -2124,9 +2130,14 @@ void FPCGEditor::OnClose()
 
 	FAssetEditorToolkit::OnClose();
 
-	if (PCGGraphBeingEdited && PCGGraphBeingEdited->NotificationsForEditorArePausedByUser())
+	if (PCGGraphBeingEdited)
 	{
-		PCGGraphBeingEdited->ToggleUserPausedNotificationsForEditor();
+		if (PCGGraphBeingEdited->NotificationsForEditorArePausedByUser())
+		{
+			PCGGraphBeingEdited->ToggleUserPausedNotificationsForEditor();
+		}
+
+		PCGGraphBeingEdited->OnGraphGridSizesChangedDelegate.RemoveAll(this);
 	}
 }
 
@@ -2333,6 +2344,15 @@ bool FPCGEditor::IsVisibleProperty(const FPropertyAndParent& InPropertyAndParent
 	return true;
 }
 
+void FPCGEditor::OnGraphGridSizesChanged(UPCGGraphInterface* InGraph)
+{
+	check(PCGEditorGraph);
+	if (UPCGComponent* PCGComponent = GetPCGComponentBeingInspected())
+	{
+		PCGEditorGraph->UpdateGridSizeVisualization(PCGComponent);
+	}
+}
+
 UPCGSubsystem* FPCGEditor::GetSubsystem()
 {
 	UWorld* World = (GEditor ? (GEditor->PlayWorld ? GEditor->PlayWorld.Get() : GEditor->GetEditorWorldContext().World()) : nullptr);
diff --git a/Engine/Plugins/Experimental/PCG/Source/PCGEditor/Private/PCGEditor.h b/Engine/Plugins/Experimental/PCG/Source/PCGEditor/Private/PCGEditor.h
--- a/Engine/Plugins/Experimental/PCG/Source/PCGEditor/Private/PCGEditor.h
+++ b/Engine/Plugins/Experimental/PCG/Source/PCGEditor/Private/PCGEditor.h
@@ -273,6 +273,8 @@ private:
 	/** Returns whether a property should be visible (used for instance vs. settings properties) */
 	bool IsVisibleProperty(const FPropertyAndParent& InPropertyAndParent, IDetailsView* InDetailsView) const;
 
+	void OnGraphGridSizesChanged(UPCGGraphInterface* InGraph);
+
 	/** Helper to get to the subsystem. */
 	static class UPCGSubsystem* GetSubsystem();
 
diff --git a/Engine/Plugins/Experimental/PCG/Source/PCGEditor/Private/PCGEditorGraph.cpp b/Engine/Plugins/Experimental/PCG/Source/PCGEditor/Private/PCGEditorGraph.cpp
--- a/Engine/Plugins/Experimental/PCG/Source/PCGEditor/Private/PCGEditorGraph.cpp
+++ b/Engine/Plugins/Experimental/PCG/Source/PCGEditor/Private/PCGEditorGraph.cpp
@@ -115,6 +115,18 @@ void UPCGEditorGraph::CreateLinks(UPCGEditorGraphNodeBase* GraphNode, bool bCrea
 	CreateLinks(GraphNode, bCreateInbound, bCreateOutbound, GraphNodeToPCGNodeMap);
 }
 
+void UPCGEditorGraph::UpdateGridSizeVisualization(UPCGComponent* InPCGComponentBeingInspected)
+{
+	for (UEdGraphNode* EditorNode : Nodes)
+	{
+		UPCGEditorGraphNodeBase* PCGEditorNode = Cast<UPCGEditorGraphNodeBase>(EditorNode);
+		if (PCGEditorNode && PCGEditorNode->UpdateGridSizeVisualization(InPCGComponentBeingInspected) != EPCGChangeType::None)
+		{
+			PCGEditorNode->ReconstructNode();
+		}
+	}
+}
+
 void UPCGEditorGraph::CreateLinks(UPCGEditorGraphNodeBase* GraphNode, bool bCreateInbound, bool bCreateOutbound, const TMap<UPCGNode*, UPCGEditorGraphNodeBase*>& GraphNodeToPCGNodeMap)
 {
 	check(GraphNode);
diff --git a/Engine/Plugins/Experimental/PCG/Source/PCGEditor/Private/PCGEditorGraph.h b/Engine/Plugins/Experimental/PCG/Source/PCGEditor/Private/PCGEditorGraph.h
--- a/Engine/Plugins/Experimental/PCG/Source/PCGEditor/Private/PCGEditorGraph.h
+++ b/Engine/Plugins/Experimental/PCG/Source/PCGEditor/Private/PCGEditorGraph.h
@@ -37,6 +37,9 @@ public:
 	void SetEditor(TWeakPtr<const FPCGEditor> InEditor) { PCGEditor = InEditor; }
 	TWeakPtr<const FPCGEditor> GetEditor() const { return PCGEditor; }
 
+	/** Updates the grid size visualization in the editor. */
+	void UpdateGridSizeVisualization(UPCGComponent* InPCGComponentBeingInspected);
+
 protected:
 	void CreateLinks(UPCGEditorGraphNodeBase* InGraphNode, bool bCreateInbound, bool bCreateOutbound, const TMap<UPCGNode*, UPCGEditorGraphNodeBase*>& InGraphNodeToPCGNodeMap);
 
diff --git a/Engine/Plugins/Experimental/PCG/Source/PCGEditor/Private/PCGEditorGraphNodeBase.cpp b/Engine/Plugins/Experimental/PCG/Source/PCGEditor/Private/PCGEditorGraphNodeBase.cpp
--- a/Engine/Plugins/Experimental/PCG/Source/PCGEditor/Private/PCGEditorGraphNodeBase.cpp
+++ b/Engine/Plugins/Experimental/PCG/Source/PCGEditor/Private/PCGEditorGraphNodeBase.cpp
@@ -360,7 +360,7 @@ EPCGChangeType UPCGEditorGraphNodeBase::UpdateErrorsAndWarnings()
 	return ChangeType;
 }
 
-EPCGChangeType UPCGEditorGraphNodeBase::UpdateGridSizeVisualisation(UPCGComponent* InComponentBeingDebugged)
+EPCGChangeType UPCGEditorGraphNodeBase::UpdateGridSizeVisualization(UPCGComponent* InComponentBeingDebugged)
 {
 	const UPCGGraph* Graph = PCGNode ? PCGNode->GetGraph() : nullptr;
 	if (!Graph)
diff --git a/Engine/Plugins/Experimental/PCG/Source/PCGEditor/Private/PCGEditorGraphNodeBase.h b/Engine/Plugins/Experimental/PCG/Source/PCGEditor/Private/PCGEditorGraphNodeBase.h
--- a/Engine/Plugins/Experimental/PCG/Source/PCGEditor/Private/PCGEditorGraphNodeBase.h
+++ b/Engine/Plugins/Experimental/PCG/Source/PCGEditor/Private/PCGEditorGraphNodeBase.h
@@ -64,7 +64,7 @@ public:
 	EPCGChangeType UpdateErrorsAndWarnings();
 
 	/** If the currently inspected grid size is smaller than the grid size of this node, display transparent. */
-	EPCGChangeType UpdateGridSizeVisualisation(class UPCGComponent* InComponentBeingDebugged);
+	EPCGChangeType UpdateGridSizeVisualization(class UPCGComponent* InComponentBeingDebugged);
 
 	/** Puts node title on node body, reducing overall node size */
 	bool ShouldDrawCompact() const;
