[PCG] Remove ability to copy/paste PCG World Actor

#rnx

diff --git a/Engine/Plugins/Experimental/PCG/Source/PCG/Public/PCGWorldActor.h b/Engine/Plugins/Experimental/PCG/Source/PCG/Public/PCGWorldActor.h
--- a/Engine/Plugins/Experimental/PCG/Source/PCG/Public/PCGWorldActor.h
+++ b/Engine/Plugins/Experimental/PCG/Source/PCG/Public/PCGWorldActor.h
@@ -25,13 +25,15 @@ public:
 	virtual void BeginPlay() override;
 	virtual void EndPlay(const EEndPlayReason::Type EndPlayReason) override;
 
-	/* Returns the grid GUIDs used for the partitioned actors, one per grid size. */
+	/** Returns the grid GUIDs used for the partitioned actors, one per grid size. */
 	void CreateGridGuidsIfNecessary(const PCGHiGenGrid::FSizeArray& InGridSizes);
 	void GetGridGuids(PCGHiGenGrid::FSizeToGuidMap& OutSizeToGuidMap) const;
 
 #if WITH_EDITOR	
 	virtual bool CanChangeIsSpatiallyLoadedFlag() const { return false; }
 	virtual bool IsUserManaged() const override { return false; }
+	virtual bool ShouldExport() override { return false; }
+	virtual bool ShouldImport(FStringView ActorPropString, bool IsMovingLevel) override { return false; }
 	virtual void BeginCacheForCookedPlatformData(const ITargetPlatform* TargetPlatform);
 	//~End AActor Interface
 
