[PCG] Static analysis - protect against context being nullptr

#rnx

diff --git a/Engine/Plugins/Experimental/PCG/Source/PCG/Private/Elements/PCGHiGenGridSize.cpp b/Engine/Plugins/Experimental/PCG/Source/PCG/Private/Elements/PCGHiGenGridSize.cpp
--- a/Engine/Plugins/Experimental/PCG/Source/PCG/Private/Elements/PCGHiGenGridSize.cpp
+++ b/Engine/Plugins/Experimental/PCG/Source/PCG/Private/Elements/PCGHiGenGridSize.cpp
@@ -107,6 +107,7 @@ bool UPCGHiGenGridSizeSettings::IsStructuralProperty(const FName& InPropertyName
 bool FPCGHiGenGridSizeElement::ExecuteInternal(FPCGContext* Context) const
 {
 	// Validation
+	check(Context);
 	if (Context->Node && Context->Node->GetGraph() && !Context->Node->GetGraph()->IsHierarchicalGenerationEnabled())
 	{
 		PCGE_LOG(Warning, GraphAndLog, LOCTEXT("GridSizeUsedInNonHiGenGraph", "Grid Size node used in a non-hierarchical graph. Enable hierarchical generation in the graph settings or remove this node."));
