[PCG] Fix outstanding issues in the latest tracking changes

* Landscape changes now correctly dirty the component, fixing issues where the component was not refreshing after a landscape move or sculpting
* Actor Dependencies are now ignoreing tracking on PCGComponent and PCGGraphInterface, as they are already tracked by the PCGComponent itself
* Fix an issue where actors already tracked were not registering new components.
* Add detection of StaticMesh/Texture compilation result to avoid dirtying, especially on load. Might need a revisit.
* Add "PanicButton" command to disable object tracking if it becomes unstable.
* Tracking is ignoring "by default" any changes to a PCG Component and PCG Graph, since they are already tracked by the PCG Component
* Tracking will also register to generate/cleanup events from the component.

#rnx

diff --git a/Engine/Plugins/Experimental/PCG/Source/PCG/Private/Helpers/PCGHelpers.cpp b/Engine/Plugins/Experimental/PCG/Source/PCG/Private/Helpers/PCGHelpers.cpp
--- a/Engine/Plugins/Experimental/PCG/Source/PCG/Private/Helpers/PCGHelpers.cpp
+++ b/Engine/Plugins/Experimental/PCG/Source/PCG/Private/Helpers/PCGHelpers.cpp
@@ -10,6 +10,7 @@
 #include "Grid/PCGPartitionActor.h"
 
 #include "Landscape.h"
+#include "Algo/AnyOf.h"
 #include "UObject/UObjectIterator.h"
 
 #if WITH_EDITOR
@@ -285,7 +286,7 @@ namespace PCGHelpers
 		return true;
 	}
 
-	void GatherDependencies(UObject* Object, TSet<TObjectPtr<UObject>>& OutDependencies, int32 MaxDepth)
+	void GatherDependencies(UObject* Object, TSet<TObjectPtr<UObject>>& OutDependencies, int32 MaxDepth, const TArray<UClass*>& InExcludedClasses)
 	{
 		UClass* ObjectClass = Object ? Object->GetClass() : nullptr;
 
@@ -303,6 +304,15 @@ namespace PCGHelpers
 				continue;
 			}
 
+			// For Object properties, if the class is in the excluded classes, continue
+			if (FObjectProperty* ObjectProperty = CastField<FObjectProperty>(Property))
+			{
+				if (ObjectProperty->PropertyClass && Algo::AnyOf(InExcludedClasses, [InClass = ObjectProperty->PropertyClass](const UClass* ExcludedClass){ return InClass->IsChildOf(ExcludedClass); }))
+				{
+					continue;
+				}
+			}
+
 			GatherDependencies(Property, Object, OutDependencies, MaxDepth);
 		}
 	}
diff --git a/Engine/Plugins/Experimental/PCG/Source/PCG/Private/PCGActorAndComponentMapping.cpp b/Engine/Plugins/Experimental/PCG/Source/PCG/Private/PCGActorAndComponentMapping.cpp
--- a/Engine/Plugins/Experimental/PCG/Source/PCG/Private/PCGActorAndComponentMapping.cpp
+++ b/Engine/Plugins/Experimental/PCG/Source/PCG/Private/PCGActorAndComponentMapping.cpp
@@ -12,8 +12,11 @@
 #include "Helpers/PCGHelpers.h"
 
 #include "LandscapeProxy.h"
-#include "Kismet/GameplayStatics.h"
+#include "StaticMeshCompiler.h"
+#include "TextureCompiler.h"
 #include "Engine/Engine.h"
+#include "Engine/StaticMesh.h"
+#include "Kismet/GameplayStatics.h"
 
 namespace PCGActorAndComponentMapping
 {
@@ -28,6 +31,13 @@ namespace PCGActorAndComponentMapping
 
 		return ActorBounds;
 	}
+
+#if WITH_EDITOR
+	static TAutoConsoleVariable<bool> CVarDisableObjectDependenciesTracking(
+		TEXT("pcg.DisableObjectDependenciesTracking"),
+		false,
+		TEXT("If depencencies are being unstable, disable the tracking, allowing people to continue working while we investigate."));
+#endif // WITH_EDITOR
 }
 
 UPCGActorAndComponentMapping::UPCGActorAndComponentMapping(UPCGSubsystem* InPCGSubsystem)
@@ -600,8 +610,6 @@ void UPCGActorAndComponentMapping::RegisterOrUpdateTracking(UPCGComponent* InCom
 			FPCGActorSelectorSettings SelectorSettings = FPCGActorSelectorSettings::ReconstructFromKey(InKey);
 			SelectorSettings.bSelectMultiple = true;
 
-			TArray<AActor*> AllActors = PCGActorSelector::FindActors(SelectorSettings, InComponent, [](const AActor*) { return true; }, [](const AActor*) { return true; });
-
 			bool bShouldCull = true;
 			for (const FPCGSettingsAndCulling& SettingsAndCulling : InSettingsAndCulling)
 			{
@@ -612,6 +620,8 @@ void UPCGActorAndComponentMapping::RegisterOrUpdateTracking(UPCGComponent* InCom
 				}
 			}
 
+			TArray<AActor*> AllActors = PCGActorSelector::FindActors(SelectorSettings, InComponent, [](const AActor*) { return true; }, [](const AActor*) { return true; });
+
 			for (AActor* Actor : AllActors)
 			{
 				if (bShouldCull)
@@ -635,11 +645,7 @@ void UPCGActorAndComponentMapping::RegisterOrUpdateTracking(UPCGComponent* InCom
 
 		for (TPair<FPCGActorSelectionKey, TArray<FPCGSettingsAndCulling>>& It : PCGGraph->GetTrackedActorKeysToSettings())
 		{
-			if (!KeysToComponentsMap.Contains(It.Key))
-			{
-				FindActorsAndTrack(It.Key, It.Value);
-			}
-
+			FindActorsAndTrack(It.Key, It.Value);
 			KeysToComponentsMap.FindOrAdd(It.Key).Add(InComponent);
 		}
 
@@ -648,14 +654,18 @@ void UPCGActorAndComponentMapping::RegisterOrUpdateTracking(UPCGComponent* InCom
 		{
 			// Landscape doesn't have an associated setting and is always culled.
 			FPCGActorSelectionKey LandscapeKey = FPCGActorSelectionKey(ALandscapeProxy::StaticClass());
-			if (!KeysToComponentsMap.Contains(LandscapeKey))
-			{
-				FindActorsAndTrack(LandscapeKey, { {nullptr, true} });
-			}
-
+			FindActorsAndTrack(LandscapeKey, { {nullptr, true} });
 			KeysToComponentsMap.FindOrAdd(LandscapeKey).Add(InComponent);
 		}
 	}
+
+	// Add tracking for when the graph was generated/cleaned, only once
+	if (!InComponent->OnPCGGraphGeneratedDelegate.IsBoundToObject(this))
+	{
+		InComponent->OnPCGGraphGeneratedDelegate.AddRaw(this, &UPCGActorAndComponentMapping::OnPCGGraphGeneratedOrCleaned);
+		InComponent->OnPCGGraphCleanedDelegate.AddRaw(this, &UPCGActorAndComponentMapping::OnPCGGraphGeneratedOrCleaned);
+	}
+
 }
 
 void UPCGActorAndComponentMapping::RemapTracking(const UPCGComponent* InOldComponent, UPCGComponent* InNewComponent)
@@ -673,10 +683,30 @@ void UPCGActorAndComponentMapping::RemapTracking(const UPCGComponent* InOldCompo
 
 	ReplaceInMap(CulledTrackedActorsToComponentsMap);
 	ReplaceInMap(AlwaysTrackedActorsToComponentsMap);
+	ReplaceInMap(KeysToComponentsMap);
+
+	// Old component will probably die, but we'll force removing the delegates even if it is const.
+	if (UPCGComponent* MutableOldComponent = const_cast<UPCGComponent*>(InOldComponent))
+	{
+		MutableOldComponent->OnPCGGraphGeneratedDelegate.RemoveAll(this);
+		MutableOldComponent->OnPCGGraphCleanedDelegate.RemoveAll(this);
+	}
+
+	// And just making sure we are not registering multiple times
+	if (!InNewComponent->OnPCGGraphGeneratedDelegate.IsBoundToObject(this))
+	{
+		InNewComponent->OnPCGGraphGeneratedDelegate.AddRaw(this, &UPCGActorAndComponentMapping::OnPCGGraphGeneratedOrCleaned);
+		InNewComponent->OnPCGGraphCleanedDelegate.AddRaw(this, &UPCGActorAndComponentMapping::OnPCGGraphGeneratedOrCleaned);
+	}
 }
 
 void UPCGActorAndComponentMapping::UnregisterTracking(UPCGComponent* InComponent)
 {
+	if (!InComponent)
+	{
+		return;
+	}
+
 	TSet<TWeakObjectPtr<AActor>> CandidatesForUntrack;
 	TSet<FPCGActorSelectionKey> KeysToRemove;
 
@@ -715,6 +745,9 @@ void UPCGActorAndComponentMapping::UnregisterTracking(UPCGComponent* InComponent
 			UnregisterActor(Candidate.Get());
 		}
 	}
+
+	InComponent->OnPCGGraphGeneratedDelegate.RemoveAll(this);
+	InComponent->OnPCGGraphCleanedDelegate.RemoveAll(this);
 }
 
 void UPCGActorAndComponentMapping::ResetPartitionActorsMap()
@@ -874,7 +907,7 @@ void UPCGActorAndComponentMapping::RegisterActor(AActor* InActor)
 	TrackedActorToPositionMap.FindOrAdd(InActor) = PCGActorAndComponentMapping::GetActorBounds(InActor);
 
 	// Also gather dependencies
-	PCGHelpers::GatherDependencies(InActor, TrackedActorsToDependenciesMap.FindOrAdd(InActor), 1);
+	UpdateActorDependencies(InActor);
 }
 
 bool UPCGActorAndComponentMapping::UnregisterActor(AActor* InActor)
@@ -963,10 +996,19 @@ void UPCGActorAndComponentMapping::OnObjectPropertyChanged(UObject* InObject, FP
 {
 	TRACE_CPUPROFILER_EVENT_SCOPE(UPCGActorAndComponentMapping::OnObjectPropertyChanged);
 
-	bool bValueNotInteractive = (InEvent.ChangeType != EPropertyChangeType::Interactive);
+	const bool bValueNotInteractive = (InEvent.ChangeType != EPropertyChangeType::Interactive);
 	// Special exception for actor tags, as we can't track otherwise an actor "losing" a tag
-	bool bActorTagChange = (InEvent.Property && InEvent.Property->GetFName() == GET_MEMBER_NAME_CHECKED(AActor, Tags));
-	if (!bValueNotInteractive && !bActorTagChange)
+	const bool bActorTagChange = (InEvent.Property && InEvent.Property->GetFName() == GET_MEMBER_NAME_CHECKED(AActor, Tags));
+	// Another special exception for texture/mesh compilation. If the InEvent is empty and the object is a texture/mesh, we ignore it.
+	const bool bEventIsEmpty = (InEvent.Property == nullptr) && (InEvent.ChangeType == EPropertyChangeType::Unspecified);
+	const bool bIsTextureCompilationResult = bEventIsEmpty && InObject && InObject->IsA<UTexture>()
+		&& FTextureCompilingManager::Get().IsCompilingTexture(Cast<UTexture>(InObject));
+	// There is no equivalent for StaticMesh to know if we are in PostCompilation, so we assume there are still some meshes to compile (including this one).
+	// Might be an over-optimistic approach, might need a revisit.
+	const bool bIsStaticMeshCompilationResult = bEventIsEmpty && InObject && InObject->IsA<UStaticMesh>()
+		&& FStaticMeshCompilingManager::Get().GetNumRemainingMeshes() > 0;
+
+	if ((!bValueNotInteractive && !bActorTagChange) || bIsTextureCompilationResult || bIsStaticMeshCompilationResult)
 	{
 		return;
 	}
@@ -986,7 +1028,9 @@ void UPCGActorAndComponentMapping::OnObjectPropertyChanged(UObject* InObject, FP
 		{
 			if (TrackedActor.Value.Contains(InObject))
 			{
-				OnActorChanged(TrackedActor.Key.Get(), /*bInHasMoved=*/ false);
+				AActor* ActorToChange = TrackedActor.Key.Get();
+				OnActorChanged(ActorToChange, /*bInHasMoved=*/ false);
+				UpdateActorDependencies(ActorToChange);
 			}
 		}
 
@@ -1004,9 +1048,14 @@ void UPCGActorAndComponentMapping::OnObjectPropertyChanged(UObject* InObject, FP
 	{
 		OnActorChanged(Actor, /*bInHasMoved=*/ false);
 	}
+	else
+	{
+		// Otherwise we are already tracking the actor, so update its dependencies
+		UpdateActorDependencies(Actor);
+	}
 }
 
-void UPCGActorAndComponentMapping::OnActorChanged(AActor* InActor, bool bInHasMoved)
+void UPCGActorAndComponentMapping::OnActorChanged(AActor* InActor, bool bInHasMoved, const UPCGComponent* InComponentChanged)
 {
 	TRACE_CPUPROFILER_EVENT_SCOPE(UPCGActorAndComponentMapping::OnActorChanged);
 
@@ -1030,11 +1079,14 @@ void UPCGActorAndComponentMapping::OnActorChanged(AActor* InActor, bool bInHasMo
 		// Then do an octree find to get all components that intersect with this actor.
 		// If the actor has moved, we also need to find components that intersected with it before
 		// We first do it for non-partitioned, then we do it for partitioned
-		auto UpdateNonPartitioned = [&DirtyComponents, InActor, CulledTrackedComponents, &RemovedTags, DirtyFlag](const FPCGComponentRef& ComponentRef) -> void
+		auto UpdateNonPartitioned = [&DirtyComponents, InActor, CulledTrackedComponents, &RemovedTags, DirtyFlag, InComponentChanged](const FPCGComponentRef& ComponentRef) -> void
 		{
 			TRACE_CPUPROFILER_EVENT_SCOPE(UPCGActorAndComponentMapping::OnActorChanged::UpdateNonPartitioned);
 
-			if (DirtyComponents.Contains(ComponentRef.Component) || !CulledTrackedComponents->Contains(ComponentRef.Component))
+			// Don't dirty if the component was already dirtied, not tracked, or the origin of the change.
+			if (DirtyComponents.Contains(ComponentRef.Component) || 
+				!CulledTrackedComponents->Contains(ComponentRef.Component) ||
+				InComponentChanged == ComponentRef.Component)
 			{
 				return;
 			}
@@ -1049,11 +1101,15 @@ void UPCGActorAndComponentMapping::OnActorChanged(AActor* InActor, bool bInHasMo
 		NonPartitionedOctree.FindElementsWithBoundsTest(ActorBounds, UpdateNonPartitioned);
 
 		// For partitioned, we first need to find all components that intersect with our actor and then forward the dirty call to all local components that intersect.
-		auto UpdatePartitioned = [this, &DirtyComponents, InActor, CulledTrackedComponents, &ActorBounds, &RemovedTags, DirtyFlag](const FPCGComponentRef& ComponentRef)  -> void
+		auto UpdatePartitioned = [this, &DirtyComponents, InActor, CulledTrackedComponents, &ActorBounds, &RemovedTags, DirtyFlag, InComponentChanged](const FPCGComponentRef& ComponentRef)  -> void
 		{
 			TRACE_CPUPROFILER_EVENT_SCOPE(UPCGActorAndComponentMapping::OnActorChanged::UpdatePartitioned);
 
-			if (!CulledTrackedComponents->Contains(ComponentRef.Component))
+			// Don't dirty if the component is not tracked, or the origin of the change.
+			// We can "re-dirty" it because changes can impact different local components, from the same
+			// original component.
+			if (!CulledTrackedComponents->Contains(ComponentRef.Component) ||
+				InComponentChanged == ComponentRef.Component)
 			{
 				return;
 			}
@@ -1108,7 +1164,7 @@ void UPCGActorAndComponentMapping::OnActorChanged(AActor* InActor, bool bInHasMo
 
 		for (UPCGComponent* PCGComponent : *AlwaysTrackedComponents)
 		{
-			if (!PCGComponent)
+			if (!PCGComponent || PCGComponent == InComponentChanged)
 			{
 				continue;
 			}
@@ -1158,6 +1214,42 @@ void UPCGActorAndComponentMapping::OnLandscapeChanged(ALandscapeProxy* InLandsca
 {
 	// We don't know if the landscape moved, only that it has changed. Since `bHasMoved` is doing a bit more, always assume that the landscape has moved.
 	OnActorChanged(InLandscape, /*bHasMoved=*/true);
+
+	// Also update its dependencies
+	UpdateActorDependencies(InLandscape);
+}
+
+void UPCGActorAndComponentMapping::UpdateActorDependencies(AActor* InActor)
+{
+	if (!InActor)
+	{
+		return;
+	}
+
+	// Don't track what the PCG Component is already tracking (itself and graphs)
+	static const TArray<UClass*> ExcludedClasses =
+	{
+		UPCGComponent::StaticClass(),
+		UPCGGraphInterface::StaticClass()
+	};
+
+	TSet<TObjectPtr<UObject>>& DependenciesSet = TrackedActorsToDependenciesMap.FindOrAdd(InActor);
+	DependenciesSet.Empty();
+
+	if (!PCGActorAndComponentMapping::CVarDisableObjectDependenciesTracking.GetValueOnAnyThread())
+	{
+		PCGHelpers::GatherDependencies(InActor, DependenciesSet, 1, ExcludedClasses);
+	}
+}
+
+void UPCGActorAndComponentMapping::OnPCGGraphGeneratedOrCleaned(UPCGComponent* InComponent)
+{
+	if (!InComponent)
+	{
+		return;
+	}
+
+	OnActorChanged(InComponent->GetOwner(), /*bInHasMoved=*/false, InComponent);
 }
 
 #endif // WITH_EDITOR
\ No newline at end of file
diff --git a/Engine/Plugins/Experimental/PCG/Source/PCG/Private/PCGComponent.cpp b/Engine/Plugins/Experimental/PCG/Source/PCG/Private/PCGComponent.cpp
--- a/Engine/Plugins/Experimental/PCG/Source/PCG/Private/PCGComponent.cpp
+++ b/Engine/Plugins/Experimental/PCG/Source/PCG/Private/PCGComponent.cpp
@@ -2105,6 +2105,12 @@ bool UPCGComponent::DirtyTrackedActor(AActor* InActor, bool bIntersect, const TS
 		}
 	}
 
+	// Special case for landscape, we should dirty.
+	if (ShouldTrackLandscape() && InActor->IsA<ALandscapeProxy>())
+	{
+		bWasDirtied = true;
+	}
+
 	return bWasDirtied;
 }
 
diff --git a/Engine/Plugins/Experimental/PCG/Source/PCG/Public/Helpers/PCGHelpers.h b/Engine/Plugins/Experimental/PCG/Source/PCG/Public/Helpers/PCGHelpers.h
--- a/Engine/Plugins/Experimental/PCG/Source/PCG/Public/Helpers/PCGHelpers.h
+++ b/Engine/Plugins/Experimental/PCG/Source/PCG/Public/Helpers/PCGHelpers.h
@@ -49,7 +49,7 @@ namespace PCGHelpers
 	PCG_API TArray<FString> GetStringArrayFromCommaSeparatedString(const FString& InCommaSeparatedString);
 
 #if WITH_EDITOR
-	PCG_API void GatherDependencies(UObject* Object, TSet<TObjectPtr<UObject>>& OutDependencies, int32 MaxDepth = -1);
+	PCG_API void GatherDependencies(UObject* Object, TSet<TObjectPtr<UObject>>& OutDependencies, int32 MaxDepth = -1, const TArray<UClass*>& InExcludedClasses = TArray<UClass*>());
 	PCG_API void GatherDependencies(FProperty* Property, const void* InContainer, TSet<TObjectPtr<UObject>>& OutDependencies, int32 MaxDepth);
 #endif
 
diff --git a/Engine/Plugins/Experimental/PCG/Source/PCG/Public/PCGActorAndComponentMapping.h b/Engine/Plugins/Experimental/PCG/Source/PCG/Public/PCGActorAndComponentMapping.h
--- a/Engine/Plugins/Experimental/PCG/Source/PCG/Public/PCGActorAndComponentMapping.h
+++ b/Engine/Plugins/Experimental/PCG/Source/PCG/Public/PCGActorAndComponentMapping.h
@@ -123,6 +123,7 @@ private:
 	void OnLandscapeChanged(ALandscapeProxy* InLandscape, const FLandscapeProxyComponentDataChangedParams& InChangeParams);
 	void OnPreObjectPropertyChanged(UObject* InObject, const FEditPropertyChain& InEditPropertyChain);
 	void OnObjectPropertyChanged(UObject* InObject, FPropertyChangedEvent& InEvent);
+	void OnPCGGraphGeneratedOrCleaned(UPCGComponent* InComponent);
 
 	/** For deprecation only, register tracked actors from serialized components */
 	template <typename ActorContainer>
@@ -134,8 +135,14 @@ private:
 	/** Unregister tracking when a component is removed. */
 	void UnregisterTracking(UPCGComponent* InComponent);
 
-	/** Trigger an update when the actor changed. Can also provide a list of tags that has changed. */
-	void OnActorChanged(AActor* InActor, bool bInHasMoved);
+	/** Trigger an update when the actor changed. 
+	* Can specify if the actor has moved to also update components that were at its previous position.
+	* Can also specify an optional component, originating the change, to avoid re-dirtying it.
+	*/
+	void OnActorChanged(AActor* InActor, bool bInHasMoved, const UPCGComponent* InComponentChanged = nullptr);
+
+	/** Update dependencies for a given tracked actor. */
+	void UpdateActorDependencies(AActor* InActor);
 #endif // WITH_EDITOR
 
 private:
