Added toggle to prevent an unnecessary copy of the template actor as that can lead to some serialization issues for very complex blueprints, but is also cognitively hard to understand for users.
Improved deprecation a bit so we don't do work we don't need.

diff --git a/Engine/Plugins/Experimental/PCG/Source/PCG/Private/Elements/PCGSpawnActor.cpp b/Engine/Plugins/Experimental/PCG/Source/PCG/Private/Elements/PCGSpawnActor.cpp
--- a/Engine/Plugins/Experimental/PCG/Source/PCG/Private/Elements/PCGSpawnActor.cpp
+++ b/Engine/Plugins/Experimental/PCG/Source/PCG/Private/Elements/PCGSpawnActor.cpp
@@ -292,6 +292,12 @@ void UPCGSpawnActorSettings::PostLoad()
 		bGenerationTrigger_DEPRECATED = EPCGSpawnActorGenerationTrigger::Default;
 	}
 
+	// Since the template actor editing is set to false by default, this needs to be corrected on post-load for proper deprecation
+	if (TemplateActor)
+	{
+		bAllowTemplateActorEditing = true;
+	}
+
 	SetupBlueprintEvent();
 
 	if (TemplateActorClass)
@@ -300,9 +306,9 @@ void UPCGSpawnActorSettings::PostLoad()
 		{
 			TemplateActor->ConditionalPostLoad();
 		}
-
-		RefreshTemplateActor();
 	}
+
+	RefreshTemplateActor();
 #endif // WITH_EDITOR
 }
 
@@ -342,6 +348,10 @@ void UPCGSpawnActorSettings::PostEditChangeProperty(FPropertyChangedEvent& Prope
 			SetupBlueprintEvent();
 			RefreshTemplateActor();
 		}
+		else if (PropertyName == GET_MEMBER_NAME_CHECKED(UPCGSpawnActorSettings, bAllowTemplateActorEditing))
+		{
+			RefreshTemplateActor();
+		}
 	}
 
 	Super::PostEditChangeProperty(PropertyChangedEvent);
@@ -372,23 +382,46 @@ void UPCGSpawnActorSettings::OnBlueprintChanged(UBlueprint* InBlueprint)
 
 void UPCGSpawnActorSettings::RefreshTemplateActor()
 {
-	if (TemplateActorClass)
+	// Implementation note: this is similar to the child actor component implementation
+	if (TemplateActorClass && bAllowTemplateActorEditing)
 	{
-		AActor* NewTemplateActor = NewObject<AActor>(this, TemplateActorClass, NAME_None, RF_ArchetypeObject | RF_Transactional | RF_Public);
+		const bool bCreateNewTemplateActor = (!TemplateActor || TemplateActor->GetClass() != TemplateActorClass);
 
-		if (TemplateActor)
+		if (bCreateNewTemplateActor)
 		{
-			UEngine::FCopyPropertiesForUnrelatedObjectsParams Options;
-			Options.bNotifyObjectReplacement = true;
-			UEngine::CopyPropertiesForUnrelatedObjects(TemplateActor, NewTemplateActor, Options);
+			AActor* NewTemplateActor = NewObject<AActor>(GetTransientPackage(), TemplateActorClass, NAME_None, RF_ArchetypeObject | RF_Transactional | RF_Public);
 
-			TemplateActor->Rename(nullptr, GetTransientPackage(), REN_DontCreateRedirectors | REN_ForceNoResetLoaders);
-		}
+			if (TemplateActor)
+			{
+				UEngine::FCopyPropertiesForUnrelatedObjectsParams Options;
+				Options.bNotifyObjectReplacement = true;
+				UEngine::CopyPropertiesForUnrelatedObjects(TemplateActor, NewTemplateActor, Options);
+
+				TemplateActor->Rename(nullptr, GetTransientPackage(), REN_DontCreateRedirectors);
+
+				TMap<UObject*, UObject*> OldToNew;
+				OldToNew.Emplace(TemplateActor, NewTemplateActor);
+				GEngine->NotifyToolsOfObjectReplacement(OldToNew);
 
-		TemplateActor = NewTemplateActor;
+				TemplateActor->MarkAsGarbage();
+			}
+
+			TemplateActor = NewTemplateActor;
+
+			// Record initial object state in case we're in a transaction context.
+			TemplateActor->Modify();
+
+			// Outer to this object
+			TemplateActor->Rename(nullptr, this, REN_DoNotDirty | REN_DontCreateRedirectors | REN_ForceNoResetLoaders);
+		}
 	}
 	else
 	{
+		if (TemplateActor)
+		{
+			TemplateActor->MarkAsGarbage();
+		}
+
 		TemplateActor = nullptr;
 	}
 }
@@ -430,7 +463,7 @@ bool FPCGSpawnActorElement::SpawnAndPrepareSubgraphs(FPCGSubgraphContext* Contex
 			return true;
 		}
 
-		if (!ensure(Settings->TemplateActor && Settings->TemplateActor->IsA(Settings->TemplateActorClass)))
+		if (!ensure(!Settings->TemplateActor || Settings->TemplateActor->IsA(Settings->TemplateActorClass)))
 		{
 			return true;
 		}
@@ -785,7 +818,14 @@ void FPCGSpawnActorElement::SpawnActors(FPCGSubgraphContext* Context, AActor* Ta
 	}
 	else
 	{
-		TemplateActor = NewObject<AActor>(GetTransientPackage(), InTemplateActorClass, NAME_None, RF_ArchetypeObject);
+		if (Settings->ActorOverrides.IsEmpty())
+		{
+			TemplateActor = Cast<AActor>(InTemplateActorClass->GetDefaultObject());
+		}
+		else
+		{
+			TemplateActor = NewObject<AActor>(GetTransientPackage(), InTemplateActorClass, NAME_None, RF_ArchetypeObject);
+		}
 	}
 
 	check(TemplateActor);
diff --git a/Engine/Plugins/Experimental/PCG/Source/PCG/Public/Elements/PCGSpawnActor.h b/Engine/Plugins/Experimental/PCG/Source/PCG/Public/Elements/PCGSpawnActor.h
--- a/Engine/Plugins/Experimental/PCG/Source/PCG/Public/Elements/PCGSpawnActor.h
+++ b/Engine/Plugins/Experimental/PCG/Source/PCG/Public/Elements/PCGSpawnActor.h
@@ -77,7 +77,10 @@ public:
 	UPROPERTY(BlueprintReadWrite, EditAnywhere, Category = Settings, meta = (EditCondition = "Option!=EPCGSpawnActorOption::CollapseActors", EditConditionHides))
 	TArray<FName> TagsToAddOnActors;
 
-	UPROPERTY(EditAnywhere, BlueprintReadWrite, Instanced, Category = Settings, meta = (ShowInnerProperties, EditCondition = "Option != EPCGSpawnActorOption::CollapseActors", EditConditionHides))
+	UPROPERTY(BlueprintReadWrite, EditAnywhere, Category = Settings, meta = (EditCondition = "Option != EPCGSpawnActorOption::CollapseActors"))
+	bool bAllowTemplateActorEditing = false;
+
+	UPROPERTY(EditAnywhere, BlueprintReadWrite, Instanced, Category = Settings, meta = (ShowInnerProperties, EditCondition = "bAllowTemplateActorEditing && Option != EPCGSpawnActorOption::CollapseActors", EditConditionHides))
 	TObjectPtr<AActor> TemplateActor;
 
 	UPROPERTY(EditAnywhere, BlueprintReadWrite, Category = Settings, meta = (EditCondition = "Option != EPCGSpawnActorOption::CollapseActors", EditConditionHides))
