[PCG] Fixing UI/UX for AddAttribute/CopyAttribute

* AddAttribute now supports Selectors (for [REDACTED_HANDLE] and [REDACTED_HANDLE])
* AddAttribute InputSource is not editable if the attribute pin is not connected.
* CopyAttribute now works if we try to copy an attribute from source into an attribute of the same name.

diff --git a/Engine/Plugins/Experimental/PCG/Config/BasePCG.ini b/Engine/Plugins/Experimental/PCG/Config/BasePCG.ini
--- a/Engine/Plugins/Experimental/PCG/Config/BasePCG.ini
+++ b/Engine/Plugins/Experimental/PCG/Config/BasePCG.ini
@@ -32,5 +32,6 @@
 +ClassRedirects=(OldName="PCGInstancePackerByAttribute",NewName="/Script/PCG.PCGInstanceDataPackerByAttribute")
 +ClassRedirects=(OldName="PCGInstancePackerByRegex",NewName="/Script/PCG.PCGInstanceDataPackerByRegex")
 +ClassRedirects=(OldName="PCGIntersectionSettings",NewName="/Script/PCG.PCGInnerIntersectionSettings")
++ClassRedirects=(OldName="PCGCreateAttributeSettings",NewName="/Script/PCG.PCGAddAttributeSettings")
 
 +EnumRedirects=(OldName="EPCGDensityNoiseMode",NewName="EPCGAttributeNoiseMode")
diff --git a/Engine/Plugins/Experimental/PCG/Source/PCG/Private/Elements/Metadata/PCGMetadataElement.cpp b/Engine/Plugins/Experimental/PCG/Source/PCG/Private/Elements/Metadata/PCGMetadataElement.cpp
--- a/Engine/Plugins/Experimental/PCG/Source/PCG/Private/Elements/Metadata/PCGMetadataElement.cpp
+++ b/Engine/Plugins/Experimental/PCG/Source/PCG/Private/Elements/Metadata/PCGMetadataElement.cpp
@@ -182,12 +182,14 @@ bool FPCGMetadataOperationElement::ExecuteInternal(FPCGContext* Context) const
 		SampledPoints = Points;
 
 		// If it is attribute to attribute, just copy the attributes, if they exist and are valid
-		if (InputSource.GetSelection() == EPCGAttributePropertySelection::Attribute && InputSource.GetExtraNames().IsEmpty() &&
-			OutputTarget.GetSelection() == EPCGAttributePropertySelection::Attribute && OutputTarget.GetExtraNames().IsEmpty())
+		// Only do that if it is really attribute to attribute, without any extra accessor. Any extra accessor will behave as a property.
+		const bool bInputHasAnyExtra = !InputSource.GetExtraNames().IsEmpty();
+		const bool bOutputHasAnyExtra = !OutputTarget.GetExtraNames().IsEmpty();
+		if (!bInputHasAnyExtra && !bOutputHasAnyExtra && Settings->InputSource.GetSelection() == EPCGAttributePropertySelection::Attribute && OutputTarget.GetSelection() == EPCGAttributePropertySelection::Attribute)
 		{
-			if (LocalSourceAttribute == LocalDestinationAttribute)
+			if (!SourceAttributeSet && LocalSourceAttribute == LocalDestinationAttribute)
 			{
-				// Nothing to do...
+				// Nothing to do if we try to copy an attribute into itself in the original point data.
 				continue;
 			}
 			
@@ -214,9 +216,9 @@ bool FPCGMetadataOperationElement::ExecuteInternal(FPCGContext* Context) const
 			continue;
 		}
 
-		// If the target is an attribute, only create a new one if the attribute doesn't already exist.
-		// If it exist, it will try to write to it.
-		if (OutputTarget.GetSelection() == EPCGAttributePropertySelection::Attribute && !SampledData->Metadata->HasAttribute(LocalDestinationAttribute))
+		// If the target is an attribute, only create a new one if the attribute doesn't already exist or we have any extra.
+		// If it exist or have any extra, it will try to write to it.
+		if (!bOutputHasAnyExtra && OutputTarget.GetSelection() == EPCGAttributePropertySelection::Attribute && !SampledData->Metadata->HasAttribute(LocalDestinationAttribute))
 		{
 			auto CreateAttribute = [SampledData, LocalDestinationAttribute](auto Dummy)
 			{
diff --git a/Engine/Plugins/Experimental/PCG/Source/PCG/Private/Elements/PCGCreateAttribute.cpp b/Engine/Plugins/Experimental/PCG/Source/PCG/Private/Elements/PCGCreateAttribute.cpp
--- a/Engine/Plugins/Experimental/PCG/Source/PCG/Private/Elements/PCGCreateAttribute.cpp
+++ b/Engine/Plugins/Experimental/PCG/Source/PCG/Private/Elements/PCGCreateAttribute.cpp
@@ -3,6 +3,7 @@
 #include "Elements/PCGCreateAttribute.h"
 
 #include "PCGContext.h"
+#include "PCGCustomVersion.h"
 #include "PCGParamData.h"
 #include "PCGPin.h"
 #include "Data/PCGSpatialData.h"
@@ -22,7 +23,7 @@ namespace PCGCreateAttributeConstants
 	const FText AttributesTooltip = LOCTEXT("AttributesTooltip", "Optional Attribute Set to create from. Not used if not connected.");
 }
 
-void UPCGCreateAttributeSettings::PostLoad()
+void UPCGCreateAttributeBaseSettings::PostLoad()
 {
 	Super::PostLoad();
 
@@ -47,10 +48,16 @@ void UPCGCreateAttributeSettings::PostLoad()
 		Type_DEPRECATED = EPCGMetadataTypes::Double;
 		DoubleValue_DEPRECATED = 0.0;
 	}
+
+	if (SourceParamAttributeName_DEPRECATED != NAME_None)
+	{
+		InputSource.SetAttributeName(SourceParamAttributeName_DEPRECATED);
+		SourceParamAttributeName_DEPRECATED = NAME_None;
+	}
 #endif // WITH_EDITOR
 }
 
-EPCGDataType UPCGCreateAttributeSettings::GetCurrentPinTypes(const UPCGPin* InPin) const
+EPCGDataType UPCGCreateAttributeBaseSettings::GetCurrentPinTypes(const UPCGPin* InPin) const
 {
 	check(InPin);
 	if (!HasDynamicPins() || !InPin->IsOutputPin())
@@ -63,52 +70,41 @@ EPCGDataType UPCGCreateAttributeSettings::GetCurrentPinTypes(const UPCGPin* InPi
 	return (PrimaryInputType != EPCGDataType::None) ? PrimaryInputType : EPCGDataType::Param; // No input (None) means param.
 }
 
-FName UPCGCreateAttributeSettings::AdditionalTaskName() const
-{
-	return AdditionalTaskNameInternal(PCGCreateAttributeConstants::NodeNameAddAttribute);
-}
-
 #if WITH_EDITOR
-FName UPCGCreateAttributeSettings::GetDefaultNodeName() const
-{
-	return PCGCreateAttributeConstants::NodeNameAddAttribute;
-}
-
-FText UPCGCreateAttributeSettings::GetDefaultNodeTitle() const
-{
-	return PCGCreateAttributeConstants::NodeTitleAddAttribute;
-}
-
-bool UPCGCreateAttributeSettings::IsPinUsedByNodeExecution(const UPCGPin* InPin) const
+bool UPCGCreateAttributeBaseSettings::IsPinUsedByNodeExecution(const UPCGPin* InPin) const
 {
 	return !InPin || (InPin->Properties.Label != PCGCreateAttributeConstants::AttributesLabel) || InPin->IsConnected();
 }
-#endif // WITH_EDITOR
 
-TArray<FPCGPinProperties> UPCGCreateAttributeSettings::InputPinProperties() const
+bool UPCGCreateAttributeBaseSettings::CanEditChange(const FProperty* InProperty) const
 {
-	TArray<FPCGPinProperties> PinProperties;
-	PinProperties.Emplace(PCGPinConstants::DefaultInputLabel, EPCGDataType::Any);
-	PinProperties.Emplace(PCGCreateAttributeConstants::AttributesLabel, EPCGDataType::Param, /*bInAllowMultipleConnections=*/true, /*bAllowMultipleData=*/true, PCGCreateAttributeConstants::AttributesTooltip);
+	if (!InProperty || !Super::CanEditChange(InProperty))
+	{
+		return false;
+	}
 
-	return PinProperties;
-}
+	const UPCGNode* Node = Cast<UPCGNode>(GetOuter());
+	const bool AttributesPinIsConnected = Node ? Node->IsInputPinConnected(PCGCreateAttributeConstants::AttributesLabel) : false;
 
-TArray<FPCGPinProperties> UPCGCreateAttributeSettings::OutputPinProperties() const
-{
-	FPCGPinProperties PinProperties;
-	PinProperties.Label = PCGPinConstants::DefaultOutputLabel;
-	PinProperties.AllowedTypes = EPCGDataType::Any;
+	if (InProperty->GetFName() == GET_MEMBER_NAME_CHECKED(UPCGCreateAttributeBaseSettings, InputSource))
+	{
+		return AttributesPinIsConnected;
+	}
+	else if (InProperty->GetOwnerStruct() == FPCGMetadataTypesConstantStruct::StaticStruct())
+	{
+		return !AttributesPinIsConnected;
+	}
 
-	return { PinProperties };
+	return true;
 }
+#endif // WITH_EDITOR
 
-FPCGElementPtr UPCGCreateAttributeSettings::CreateElement() const
+FPCGElementPtr UPCGCreateAttributeBaseSettings::CreateElement() const
 {
 	return MakeShared<FPCGCreateAttributeElement>();
 }
 
-FName UPCGCreateAttributeSettings::AdditionalTaskNameInternal(FName NodeName) const
+FName UPCGCreateAttributeBaseSettings::AdditionalTaskNameInternal(FName NodeName) const
 {
 	if (HasAnyFlags(RF_ClassDefaultObject))
 	{
@@ -118,6 +114,9 @@ FName UPCGCreateAttributeSettings::AdditionalTaskNameInternal(FName NodeName) co
 	const UPCGNode* Node = Cast<UPCGNode>(GetOuter());
 	const bool AttributesPinIsConnected = Node ? Node->IsInputPinConnected(PCGCreateAttributeConstants::AttributesLabel) : false;
 
+	const FName OutputAttributeName = GetOutputAttributeName(AttributesPinIsConnected ? &InputSource : nullptr, nullptr);
+	const FName SourceParamAttributeName = InputSource.GetName();
+
 	if (ShouldAddAttributesPin() && AttributesPinIsConnected)
 	{
 		if ((OutputAttributeName == NAME_None) && (SourceParamAttributeName == NAME_None))
@@ -136,12 +135,95 @@ FName UPCGCreateAttributeSettings::AdditionalTaskNameInternal(FName NodeName) co
 	}
 }
 
+UPCGAddAttributeSettings::UPCGAddAttributeSettings()
+{
+	OutputTarget.SetAttributeName(NAME_None);
+}
+
+FName UPCGAddAttributeSettings::AdditionalTaskName() const
+{
+	return AdditionalTaskNameInternal(PCGCreateAttributeConstants::NodeNameAddAttribute);
+}
+
+void UPCGAddAttributeSettings::PostLoad()
+{
+	Super::PostLoad();
+
+#if WITH_EDITOR
+	if (OutputAttributeName_DEPRECATED != NAME_None)
+	{
+		OutputTarget.SetAttributeName(OutputAttributeName_DEPRECATED);
+		OutputAttributeName_DEPRECATED = NAME_None;
+	}
+#endif // WITH_EDITOR
+}
+
+#if WITH_EDITOR
+FName UPCGAddAttributeSettings::GetDefaultNodeName() const
+{
+	return PCGCreateAttributeConstants::NodeNameAddAttribute;
+}
+
+FText UPCGAddAttributeSettings::GetDefaultNodeTitle() const
+{
+	return PCGCreateAttributeConstants::NodeTitleAddAttribute;
+}
+
+void UPCGAddAttributeSettings::ApplyDeprecation(UPCGNode* InOutNode)
+{
+	const bool AttributesPinIsConnected = InOutNode ? InOutNode->IsInputPinConnected(PCGCreateAttributeConstants::AttributesLabel) : false;
+
+	if (DataVersion < FPCGCustomVersion::UpdateAddAttributeWithSelectors
+		&& OutputTarget.GetSelection() == EPCGAttributePropertySelection::Attribute
+		&& OutputTarget.GetAttributeName() == NAME_None
+		&& AttributesPinIsConnected)
+	{
+		// Previous behavior of the output target for this node was:
+		// None => Source if Attributes pin is connected
+		OutputTarget.SetAttributeName(PCGMetadataAttributeConstants::SourceAttributeName);
+	}
+
+	Super::ApplyDeprecation(InOutNode);
+}
+#endif // WITH_EDITOR
+
+TArray<FPCGPinProperties> UPCGAddAttributeSettings::InputPinProperties() const
+{
+	TArray<FPCGPinProperties> PinProperties;
+	PinProperties.Emplace(PCGPinConstants::DefaultInputLabel, EPCGDataType::Any);
+	PinProperties.Emplace(PCGCreateAttributeConstants::AttributesLabel, EPCGDataType::Param, /*bInAllowMultipleConnections=*/true, /*bAllowMultipleData=*/true, PCGCreateAttributeConstants::AttributesTooltip);
+
+	return PinProperties;
+}
+
+TArray<FPCGPinProperties> UPCGAddAttributeSettings::OutputPinProperties() const
+{
+	FPCGPinProperties PinProperties;
+	PinProperties.Label = PCGPinConstants::DefaultOutputLabel;
+	PinProperties.AllowedTypes = EPCGDataType::Any;
+
+	return { PinProperties };
+}
+
 UPCGCreateAttributeSetSettings::UPCGCreateAttributeSetSettings()
 {
 	// No input pin to grab source param from
 	bDisplayFromSourceParamSetting = false;
 }
 
+void UPCGCreateAttributeSetSettings::PostLoad()
+{
+	Super::PostLoad();
+
+#if WITH_EDITOR
+	if (OutputAttributeName_DEPRECATED != NAME_None)
+	{
+		OutputTarget.SetAttributeName(OutputAttributeName_DEPRECATED);
+		OutputAttributeName_DEPRECATED = NAME_None;
+	}
+#endif // WITH_EDITOR
+}
+
 #if WITH_EDITOR
 FName UPCGCreateAttributeSetSettings::GetDefaultNodeName() const
 {
@@ -178,12 +260,13 @@ bool FPCGCreateAttributeElement::ExecuteInternal(FPCGContext* Context) const
 
 	check(Context);
 
-	const UPCGCreateAttributeSettings* Settings = Context->GetInputSettings<UPCGCreateAttributeSettings>();
+	const UPCGCreateAttributeBaseSettings* Settings = Context->GetInputSettings<UPCGCreateAttributeBaseSettings>();
 	check(Settings);
 
 	TArray<FPCGTaggedData> SourceParams = Context->InputData.GetInputsByPin(PCGCreateAttributeConstants::AttributesLabel);
 	UPCGParamData* SourceParamData = nullptr;
 	FName SourceParamAttributeName = NAME_None;
+	FName OutputAttributeName = NAME_None;
 
 	if (!SourceParams.IsEmpty())
 	{
@@ -195,7 +278,10 @@ bool FPCGCreateAttributeElement::ExecuteInternal(FPCGContext* Context) const
 			return true;
 		}
 
-		SourceParamAttributeName = (Settings->SourceParamAttributeName == NAME_None) ? SourceParamData->Metadata->GetLatestAttributeNameOrNone() : Settings->SourceParamAttributeName;
+		FPCGAttributePropertyInputSelector InputSource = Settings->InputSource.CopyAndFixLast(SourceParamData);
+
+		SourceParamAttributeName = InputSource.GetName();
+		OutputAttributeName = Settings->GetOutputAttributeName(&InputSource, SourceParamData);
 
 		if (!SourceParamData->Metadata->HasAttribute(SourceParamAttributeName))
 		{
@@ -203,6 +289,10 @@ bool FPCGCreateAttributeElement::ExecuteInternal(FPCGContext* Context) const
 			return true;
 		}
 	}
+	else
+	{
+		OutputAttributeName = Settings->GetOutputAttributeName(nullptr, nullptr);
+	}
 
 	TArray<FPCGTaggedData> Inputs = Context->InputData.GetInputsByPin(PCGPinConstants::DefaultInputLabel);
 
@@ -258,8 +348,6 @@ bool FPCGCreateAttributeElement::ExecuteInternal(FPCGContext* Context) const
 			continue;
 		}
 
-		const FName OutputAttributeName = (SourceParamData && Settings->OutputAttributeName == NAME_None) ? SourceParamAttributeName : Settings->OutputAttributeName;
-
 		FPCGMetadataAttributeBase* Attribute = nullptr;
 
 		if (SourceParamData)
@@ -269,7 +357,7 @@ bool FPCGCreateAttributeElement::ExecuteInternal(FPCGContext* Context) const
 		}
 		else
 		{
-			Attribute = ClearOrCreateAttribute(Settings, Metadata, &OutputAttributeName);
+			Attribute = ClearOrCreateAttribute(Settings, Metadata, OutputAttributeName);
 		}
 
 		if (!Attribute)
@@ -295,19 +383,19 @@ bool FPCGCreateAttributeElement::ExecuteInternal(FPCGContext* Context) const
 	return true;
 }
 
-FPCGMetadataAttributeBase* FPCGCreateAttributeElement::ClearOrCreateAttribute(const UPCGCreateAttributeSettings* Settings, UPCGMetadata* Metadata, const FName* OutputAttributeNameOverride) const
+FPCGMetadataAttributeBase* FPCGCreateAttributeElement::ClearOrCreateAttribute(const UPCGCreateAttributeBaseSettings* Settings, UPCGMetadata* Metadata, const FName OutputAttributeName) const
 {
 	check(Metadata);
 
-	auto CreateAttribute = [Settings, Metadata, OutputAttributeNameOverride](auto&& Value) -> FPCGMetadataAttributeBase*
+	auto CreateAttribute = [Settings, Metadata, OutputAttributeName](auto&& Value) -> FPCGMetadataAttributeBase*
 	{
-		return PCGMetadataElementCommon::ClearOrCreateAttribute(Metadata, OutputAttributeNameOverride ? *OutputAttributeNameOverride : Settings->OutputAttributeName, std::forward<decltype(Value)>(Value));
+		return PCGMetadataElementCommon::ClearOrCreateAttribute(Metadata, OutputAttributeName, std::forward<decltype(Value)>(Value));
 	};
 
 	return Settings->AttributeTypes.Dispatcher(CreateAttribute);
 }
 
-PCGMetadataEntryKey FPCGCreateAttributeElement::SetAttribute(const UPCGCreateAttributeSettings* Settings, FPCGMetadataAttributeBase* Attribute, UPCGMetadata* Metadata, PCGMetadataEntryKey EntryKey) const
+PCGMetadataEntryKey FPCGCreateAttributeElement::SetAttribute(const UPCGCreateAttributeBaseSettings* Settings, FPCGMetadataAttributeBase* Attribute, UPCGMetadata* Metadata, PCGMetadataEntryKey EntryKey) const
 {
 	check(Attribute && Metadata);
 
diff --git a/Engine/Plugins/Experimental/PCG/Source/PCG/Public/Elements/PCGCreateAttribute.h b/Engine/Plugins/Experimental/PCG/Source/PCG/Public/Elements/PCGCreateAttribute.h
--- a/Engine/Plugins/Experimental/PCG/Source/PCG/Public/Elements/PCGCreateAttribute.h
+++ b/Engine/Plugins/Experimental/PCG/Source/PCG/Public/Elements/PCGCreateAttribute.h
@@ -3,6 +3,7 @@
 #pragma once
 
 #include "PCGSettings.h"
+#include "Metadata/PCGAttributePropertySelector.h"
 #include "Metadata/PCGMetadataTypesConstantStruct.h"
 #include "Metadata/PCGMetadataCommon.h"
 
@@ -17,8 +18,8 @@ class UPCGMetadata;
 * 
 * Note: This need to be updated if we ever add new types.
 */
-UCLASS(BlueprintType, ClassGroup = (Procedural))
-class PCG_API UPCGCreateAttributeSettings : public UPCGSettings
+UCLASS(Abstract, BlueprintType, ClassGroup = (Procedural))
+class PCG_API UPCGCreateAttributeBaseSettings : public UPCGSettings
 {
 	GENERATED_BODY()
 public:
@@ -28,26 +29,21 @@ public:
 
 	//~Begin UPCGSettings interface
 #if WITH_EDITOR
-	virtual FName GetDefaultNodeName() const override;
-	virtual FText GetDefaultNodeTitle() const override;
 	virtual EPCGSettingsType GetType() const override { return EPCGSettingsType::Metadata; }
-	virtual bool HasDynamicPins() const override { return true; }
 	virtual bool IsPinUsedByNodeExecution(const UPCGPin* InPin) const override;
+	virtual bool CanEditChange(const FProperty* InProperty) const override;
 #endif // WITH_EDITOR
 	virtual EPCGDataType GetCurrentPinTypes(const UPCGPin* InPin) const override;
-
-	virtual FName AdditionalTaskName() const override;
 	//~End UPCGSettings interface
 
-	UPROPERTY(BlueprintReadWrite, EditAnywhere, Category = Settings)
-	FName OutputAttributeName = NAME_None;
+	virtual FName GetOutputAttributeName(const FPCGAttributePropertyInputSelector* InSource, const UPCGData* InSourceData) const PURE_VIRTUAL(UPCGCreateAttributeBaseSettings::GetOutputAttributeName, return NAME_None;);
 
 	// This can be set false by inheriting nodes to hide the 'From Source Param' property.
 	UPROPERTY(Transient, meta = (EditCondition = false, EditConditionHides))
 	bool bDisplayFromSourceParamSetting = true;
 
-	UPROPERTY(BlueprintReadWrite, EditAnywhere, Category = Settings, meta = (EditCondition = "bDisplayFromSourceParamSetting", EditConditionHides, HideEditConditionToggle))
-	FName SourceParamAttributeName = NAME_None;
+	UPROPERTY(BlueprintReadWrite, EditAnywhere, Category = Settings, meta = (EditCondition = "bDisplayFromSourceParamSetting", EditConditionHides, HideEditConditionToggle, PCG_DiscardPropertySelection, PCG_DiscardExtraSelection))
+	FPCGAttributePropertyInputSelector InputSource;
 
 	UPROPERTY(BlueprintReadWrite, EditAnywhere, Category = Settings, meta = (ShowOnlyInnerProperties))
 	FPCGMetadataTypesConstantStruct AttributeTypes;
@@ -97,27 +93,72 @@ public:
 
 	UPROPERTY()
 	bool bKeepExistingAttributes_DEPRECATED = false;
+
+	UPROPERTY()
+	FName SourceParamAttributeName_DEPRECATED = NAME_None;
 #endif // WITH_EDITORONLY_DATA
 
 protected:
-	virtual TArray<FPCGPinProperties> InputPinProperties() const override;
-	virtual TArray<FPCGPinProperties> OutputPinProperties() const override;
-
 	FName AdditionalTaskNameInternal(FName NodeName) const;
 
 	virtual FPCGElementPtr CreateElement() const override;
 	virtual bool ShouldAddAttributesPin() const { return true; }
 };
 
+/* Add a new attribute. */
+UCLASS(BlueprintType, ClassGroup = (Procedural))
+class PCG_API UPCGAddAttributeSettings : public UPCGCreateAttributeBaseSettings
+{
+	GENERATED_BODY()
+
+	UPCGAddAttributeSettings();
+
+public:
+	//~Begin UObject interface
+	virtual void PostLoad() override;
+	//~End UObject interface
+
+	//~Begin UPCGSettings interface
+#if WITH_EDITOR
+	virtual FName GetDefaultNodeName() const override;
+	virtual FText GetDefaultNodeTitle() const override;
+	virtual bool HasDynamicPins() const override { return true; }
+	virtual void ApplyDeprecation(UPCGNode* InOutNode) override;
+#endif // WITH_EDITOR
+
+	virtual FName AdditionalTaskName() const override;
+	//~End UPCGSettings interface
+
+	//~Begin UPCGCreateAttributeBaseSettings interface
+	virtual FName GetOutputAttributeName(const FPCGAttributePropertyInputSelector* InSource, const UPCGData* InSourceData) const { return OutputTarget.CopyAndFixSource(InSource, InSourceData).GetName(); }
+	//~End UPCGCreateAttributeBaseSettings interface
+
+	UPROPERTY(BlueprintReadWrite, EditAnywhere, Category = Settings, meta = (PCG_DiscardPropertySelection, PCG_DiscardExtraSelection))
+	FPCGAttributePropertyOutputSelector OutputTarget;
+
+#if WITH_EDITORONLY_DATA
+	UPROPERTY()
+	FName OutputAttributeName_DEPRECATED = NAME_None;
+#endif // WITH_EDITORONLY_DATA
+
+protected:
+	virtual TArray<FPCGPinProperties> InputPinProperties() const override;
+	virtual TArray<FPCGPinProperties> OutputPinProperties() const override;
+};
+
 /* Creates a new Attribute Set. */
 UCLASS(BlueprintType, ClassGroup = (Procedural))
-class PCG_API UPCGCreateAttributeSetSettings : public UPCGCreateAttributeSettings
+class PCG_API UPCGCreateAttributeSetSettings : public UPCGCreateAttributeBaseSettings
 {
 	GENERATED_BODY()
 
 public:
 	UPCGCreateAttributeSetSettings();
 
+	//~Begin UObject interface
+	virtual void PostLoad() override;
+	//~End UObject interface
+
 	//~Begin UPCGSettings interface
 #if WITH_EDITOR
 	virtual FName GetDefaultNodeName() const override;
@@ -125,11 +166,23 @@ public:
 	virtual bool HasDynamicPins() const override { return false; }
 #endif
 	virtual FName AdditionalTaskName() const override;
+	//~End UPCGSettings interface
+
+	UPROPERTY(BlueprintReadWrite, EditAnywhere, Category = Settings, meta = (PCG_DiscardPropertySelection, PCG_DiscardExtraSelection))
+	FPCGAttributePropertyOutputNoSourceSelector OutputTarget;
+
+	//~Begin UPCGCreateAttributeBaseSettings interface
+	virtual FName GetOutputAttributeName(const FPCGAttributePropertyInputSelector*, const UPCGData*) const { return OutputTarget.GetName(); }
+	//~End UPCGCreateAttributeBaseSettings interface
+
+#if WITH_EDITORONLY_DATA
+	UPROPERTY()
+	FName OutputAttributeName_DEPRECATED = NAME_None;
+#endif // WITH_EDITORONLY_DATA
 
 protected:
 	virtual TArray<FPCGPinProperties> InputPinProperties() const override;
 	virtual TArray<FPCGPinProperties> OutputPinProperties() const override;
-	//~End UPCGSettings interface
 
 	virtual bool ShouldAddAttributesPin() const override { return false; }
 };
@@ -141,10 +194,10 @@ protected:
 
 private:
 	/* Create (or clear) an attribute named by OutputAttributeName and depending on the selected type. Value can be overridden by params. Default value will be set to the specified value. */
-	FPCGMetadataAttributeBase* ClearOrCreateAttribute(const UPCGCreateAttributeSettings* Settings, UPCGMetadata* Metadata, const FName* OutputAttributeNameOverride = nullptr) const;
+	FPCGMetadataAttributeBase* ClearOrCreateAttribute(const UPCGCreateAttributeBaseSettings* Settings, UPCGMetadata* Metadata, const FName OutputAttributeName) const;
 
 	/* Set an entry defined by EntryKey (or create one if it is invalid) in the Attribute depending on the selected type. Value can be overridden by params. */
-	PCGMetadataEntryKey SetAttribute(const UPCGCreateAttributeSettings* Settings, FPCGMetadataAttributeBase* Attribute, UPCGMetadata* Metadata, PCGMetadataEntryKey EntryKey) const;
+	PCGMetadataEntryKey SetAttribute(const UPCGCreateAttributeBaseSettings* Settings, FPCGMetadataAttributeBase* Attribute, UPCGMetadata* Metadata, PCGMetadataEntryKey EntryKey) const;
 };
 
 #if UE_ENABLE_INCLUDE_ORDER_DEPRECATED_IN_5_2
diff --git a/Engine/Plugins/Experimental/PCG/Source/PCG/Public/PCGCommon.h b/Engine/Plugins/Experimental/PCG/Source/PCG/Public/PCGCommon.h
--- a/Engine/Plugins/Experimental/PCG/Source/PCG/Public/PCGCommon.h
+++ b/Engine/Plugins/Experimental/PCG/Source/PCG/Public/PCGCommon.h
@@ -104,6 +104,8 @@ namespace PCGObjectMetadata
 	const FName Overridable = TEXT("PCG_Overridable");
 	const FName NotOverridable = TEXT("PCG_NotOverridable");
 	const FName OverrideAliases = TEXT("PCG_OverrideAliases");
+	const FName DiscardPropertySelection = TEXT("PCG_DiscardPropertySelection");
+	const FName DiscardExtraSelection = TEXT("PCG_DiscardExtraSelection");
 
 	// Metadata usable in UPROPERTY for customizing the behavior when displaying the property in a property panel or graph node
 	enum
@@ -116,6 +118,12 @@ namespace PCGObjectMetadata
 
 		/// [PropertyMetadata] Extra names to match for a given property.
 		PCG_OverrideAliases,
+
+		/// [PropertyMetadata] For FPCGAttributePropertySelector, won't display the point property items in the dropdown
+		PCG_DiscardPropertySelection,
+
+		/// [PropertyMetadata] For FPCGAttributePropertySelector, won't display the extra property items in the dropdown
+		PCG_DiscardExtraSelection,
 	};
 }
 
diff --git a/Engine/Plugins/Experimental/PCG/Source/PCG/Public/PCGCustomVersion.h b/Engine/Plugins/Experimental/PCG/Source/PCG/Public/PCGCustomVersion.h
--- a/Engine/Plugins/Experimental/PCG/Source/PCG/Public/PCGCustomVersion.h
+++ b/Engine/Plugins/Experimental/PCG/Source/PCG/Public/PCGCustomVersion.h
@@ -50,6 +50,9 @@ struct PCG_API FPCGCustomVersion
 		// Difference node now iterates on the source pin and unions the differences pin
 		DifferenceNodeIterateOnSourceAndUnionDifferences = 13,
 
+		// Update AddAttribute with selectors
+		UpdateAddAttributeWithSelectors = 14,
+
 		// -----<new versions can be added above this line>-------------------------------------------------
 		VersionPlusOne,
 		LatestVersion = VersionPlusOne - 1
diff --git a/Engine/Plugins/Experimental/PCG/Source/PCGEditor/Private/Details/PCGAttributePropertySelectorDetails.cpp b/Engine/Plugins/Experimental/PCG/Source/PCGEditor/Private/Details/PCGAttributePropertySelectorDetails.cpp
--- a/Engine/Plugins/Experimental/PCG/Source/PCGEditor/Private/Details/PCGAttributePropertySelectorDetails.cpp
+++ b/Engine/Plugins/Experimental/PCG/Source/PCGEditor/Private/Details/PCGAttributePropertySelectorDetails.cpp
@@ -2,6 +2,7 @@
 
 #include "Details/PCGAttributePropertySelectorDetails.h"
 
+#include "PCGCommon.h"
 #include "Metadata/PCGAttributePropertySelector.h"
 #include "Metadata/PCGMetadataAttribute.h"
 
@@ -76,15 +77,19 @@ void FPCGAttributePropertySelectorDetails::CustomizeHeader(TSharedRef<IPropertyH
 		return (Selector && !Selector->IsValid()) ? FStyleColors::AccentRed : FSlateColor::UseForeground();
 	};
 
+	TSharedRef<SWidget> PropertyNameWidget = PropertyHandle->CreatePropertyNameWidget();
+	PropertyNameWidget->SetEnabled(TAttribute<bool>(this, &FPCGAttributePropertySelectorDetails::IsEnabled));
+
 	HeaderRow
 		.NameContent()
 		[
-			PropertyHandle->CreatePropertyNameWidget()
+			PropertyNameWidget
 		]
 		.ValueContent()
 		.MinDesiredWidth(350.0f)
 		[
 			SNew(SHorizontalBox)
+			.IsEnabled_Raw(this, &FPCGAttributePropertySelectorDetails::IsEnabled)
 			+ SHorizontalBox::Slot()
 			[
 				SNew(SEditableTextBox)
@@ -99,6 +104,7 @@ void FPCGAttributePropertySelectorDetails::CustomizeHeader(TSharedRef<IPropertyH
 			[
 				SNew(SComboButton)
 				.ButtonStyle(FAppStyle::Get(), "SimpleButton")
+				.Visibility_Raw(this, &FPCGAttributePropertySelectorDetails::ExtraMenuVisibility)
 				.OnGetMenuContent(this, &FPCGAttributePropertySelectorDetails::GenerateExtraMenu)
 				.HasDownArrow(false)
 				.ButtonContent()
@@ -111,6 +117,35 @@ void FPCGAttributePropertySelectorDetails::CustomizeHeader(TSharedRef<IPropertyH
 		];
 }
 
+EVisibility FPCGAttributePropertySelectorDetails::ExtraMenuVisibility() const
+{
+	FStructProperty* StructProperty = CastFieldChecked<FStructProperty>(PropertyHandle->GetProperty());
+
+	const bool bIsVisible = StructProperty &&
+		(StructProperty->Struct == FPCGAttributePropertyInputSelector::StaticStruct() ||
+			StructProperty->Struct == FPCGAttributePropertyOutputSelector::StaticStruct() ||
+			!StructProperty->HasMetaData(PCGObjectMetadata::DiscardPropertySelection) ||
+			!StructProperty->HasMetaData(PCGObjectMetadata::DiscardExtraSelection));
+
+	return bIsVisible ? EVisibility::Visible : EVisibility::Hidden;
+}
+
+bool FPCGAttributePropertySelectorDetails::IsEnabled() const
+{
+	FStructProperty* StructProperty = CastFieldChecked<FStructProperty>(PropertyHandle->GetProperty());
+	TArray<UObject*> Outers;
+	PropertyHandle->GetOuterObjects(Outers);
+	for (UObject* Outer : Outers)
+	{
+		if (Outer && !Outer->CanEditChange(StructProperty))
+		{
+			return false;
+		}
+	}
+
+	return true;
+}
+
 TSharedRef<SWidget> FPCGAttributePropertySelectorDetails::GenerateExtraMenu()
 {
 	// Clear the focus on the text box. It is preventing from updating the text.
@@ -118,9 +153,10 @@ TSharedRef<SWidget> FPCGAttributePropertySelectorDetails::GenerateExtraMenu()
 
 	FMenuBuilder MenuBuilder(true, nullptr);
 
+	FStructProperty* StructProperty = CastFieldChecked<FStructProperty>(PropertyHandle->GetProperty());
+
 	MenuBuilder.BeginSection("Attributes", LOCTEXT("AttributesHeader", "Attributes"));
 	{
-		FStructProperty* StructProperty = CastFieldChecked<FStructProperty>(PropertyHandle->GetProperty());
 		if (StructProperty->Struct == FPCGAttributePropertyInputSelector::StaticStruct())
 		{
 			MenuBuilder.AddMenuEntry(
@@ -147,17 +183,23 @@ TSharedRef<SWidget> FPCGAttributePropertySelectorDetails::GenerateExtraMenu()
 	}
 	MenuBuilder.EndSection();
 
-	MenuBuilder.BeginSection("PointProperties", LOCTEXT("PointPropertiesHeader", "Point Properties"));
+	if (!StructProperty->HasMetaData(PCGObjectMetadata::DiscardPropertySelection))
 	{
-		PCGAttributePropertySelectorDetails::BuildMenuSectionFromEnum<EPCGPointProperties>(MenuBuilder, this, &FPCGAttributePropertySelectorDetails::SetPointProperty);
+		MenuBuilder.BeginSection("PointProperties", LOCTEXT("PointPropertiesHeader", "Point Properties"));
+		{
+			PCGAttributePropertySelectorDetails::BuildMenuSectionFromEnum<EPCGPointProperties>(MenuBuilder, this, &FPCGAttributePropertySelectorDetails::SetPointProperty);
+		}
+		MenuBuilder.EndSection();
 	}
-	MenuBuilder.EndSection();
 
-	MenuBuilder.BeginSection("OtherProperties", LOCTEXT("OtherPropertiesHeader", "Other Properties"));
+	if (!StructProperty->HasMetaData(PCGObjectMetadata::DiscardExtraSelection))
 	{
-		PCGAttributePropertySelectorDetails::BuildMenuSectionFromEnum<EPCGExtraProperties>(MenuBuilder, this, &FPCGAttributePropertySelectorDetails::SetExtraProperty);
+		MenuBuilder.BeginSection("OtherProperties", LOCTEXT("OtherPropertiesHeader", "Other Properties"));
+		{
+			PCGAttributePropertySelectorDetails::BuildMenuSectionFromEnum<EPCGExtraProperties>(MenuBuilder, this, &FPCGAttributePropertySelectorDetails::SetExtraProperty);
+		}
+		MenuBuilder.EndSection();
 	}
-	MenuBuilder.EndSection();
 
 	return MenuBuilder.MakeWidget();
 }
diff --git a/Engine/Plugins/Experimental/PCG/Source/PCGEditor/Public/Details/PCGAttributePropertySelectorDetails.h b/Engine/Plugins/Experimental/PCG/Source/PCGEditor/Public/Details/PCGAttributePropertySelectorDetails.h
--- a/Engine/Plugins/Experimental/PCG/Source/PCGEditor/Public/Details/PCGAttributePropertySelectorDetails.h
+++ b/Engine/Plugins/Experimental/PCG/Source/PCGEditor/Public/Details/PCGAttributePropertySelectorDetails.h
@@ -3,7 +3,7 @@
 #pragma once
 
 #include "IPropertyTypeCustomization.h"
-
+#include "Layout/Visibility.h"
 
 enum class EPCGPointProperties : uint8;
 enum class EPCGExtraProperties : uint8;
@@ -30,6 +30,8 @@ protected:
 	const FPCGAttributePropertySelector* GetStruct() const;
 
 	TSharedRef<SWidget> GenerateExtraMenu();
+	EVisibility ExtraMenuVisibility() const;
+	bool IsEnabled() const;
 
 	FText GetText() const;
 	void SetText(const FText& NewText, ETextCommit::Type CommitInfo);
