[PCG] MetadataOps nodes will now issue a node error when there is a user-fixable oversight that aborts operation

diff --git a/Engine/Plugins/Experimental/PCG/Source/PCG/Private/Elements/Metadata/PCGMetadataOpElementBase.cpp b/Engine/Plugins/Experimental/PCG/Source/PCG/Private/Elements/Metadata/PCGMetadataOpElementBase.cpp
--- a/Engine/Plugins/Experimental/PCG/Source/PCG/Private/Elements/Metadata/PCGMetadataOpElementBase.cpp
+++ b/Engine/Plugins/Experimental/PCG/Source/PCG/Private/Elements/Metadata/PCGMetadataOpElementBase.cpp
@@ -301,8 +301,8 @@ bool FPCGMetadataElementBase::ExecuteInternal(FPCGContext* Context) const
 
 		if (InputData.IsEmpty())
 		{
-			// Absence of data not worth broadcasting to user as visual warning
-			PCGE_LOG(Warning, LogOnly, FText::Format(LOCTEXT("MissingInputDataForPin", "No data provided on pin {0}"), FText::FromName(PinLabel)));
+			// Visually warn the user, since this is causing execution to be aborted
+			PCGE_LOG(Warning, GraphAndLog, FText::Format(LOCTEXT("MissingInputDataForPin", "No data provided on pin {0}"), FText::FromName(PinLabel)));
 			return true;
 		}
 		else if (InputData.Num() > 1)
@@ -325,7 +325,8 @@ bool FPCGMetadataElementBase::ExecuteInternal(FPCGContext* Context) const
 		}
 		else
 		{
-			PCGE_LOG(Warning, LogOnly, FText::Format(LOCTEXT("InvalidInputDataTypeForPin", "Invalid data provided on pin {0}, must be of type Spatial or Attribute Set"), FText::FromName(PinLabel)));
+			// Since this aborts execution, and the user can fix it, it should be a node error
+			PCGE_LOG(Error, GraphAndLog, FText::Format(LOCTEXT("InvalidInputDataTypeForPin", "Invalid data provided on pin {0}, must be of type Spatial or Attribute Set"), FText::FromName(PinLabel)));
 			return true;
 		}
 	}
