[PCG] Overridable properties in settings that are in structs are now edit const if the override pin is connected.

#rnx

diff --git a/Engine/Plugins/Experimental/PCG/Source/PCG/Private/PCGSettings.cpp b/Engine/Plugins/Experimental/PCG/Source/PCG/Private/PCGSettings.cpp
--- a/Engine/Plugins/Experimental/PCG/Source/PCG/Private/PCGSettings.cpp
+++ b/Engine/Plugins/Experimental/PCG/Source/PCG/Private/PCGSettings.cpp
@@ -553,19 +553,33 @@ void UPCGSettings::DirtyCache()
 	}
 }
 
-bool UPCGSettings::CanEditChange(const FProperty* InProperty) const
+bool UPCGSettings::CanEditChange(const FEditPropertyChain& InPropertyChain) const
 {
-	if (!InProperty || !Super::CanEditChange(InProperty))
+	// Property is the actual property currently checked
+	FProperty* Property = InPropertyChain.GetActiveNode() ? InPropertyChain.GetActiveNode()->GetValue() : nullptr;
+	// Member property is the same as Property if it is not in a struct/array, otherwise it is the struct/array.
+	FProperty* MemberProperty = InPropertyChain.GetActiveMemberNode() ? InPropertyChain.GetActiveMemberNode()->GetValue() : nullptr;
+
+	// No property/member property, or if the property is marked edit const (in metadata), it is not editable.
+	if (!Property || !MemberProperty || !Super::CanEditChange(InPropertyChain))
 	{
 		return false;
 	}
 
-	if (!InProperty->HasMetaData(PCGObjectMetadata::Overridable))
+	// If it is not marked edit const, properties that are not marked overridable or child properties explicitly marked non overridable are always editable.
+	if (!MemberProperty->HasMetaData(PCGObjectMetadata::Overridable) || Property->HasMetaData(PCGObjectMetadata::NotOverridable))
 	{
 		return true;
 	}
 
-	return !IsPropertyOverriddenByPin(InProperty);
+	// If the property can be overridden, mark it as EditConst if it is actually overridden 
+	// (ie. the override pin associated with this property is connected)
+	return !IsPropertyOverriddenByPin(Property);
+}
+
+bool UPCGSettings::CanEditChange(const FProperty* InProperty) const
+{
+	return Super::CanEditChange(InProperty);
 }
 
 bool UPCGSettings::IsPropertyOverriddenByPin(const FProperty* InProperty) const
diff --git a/Engine/Plugins/Experimental/PCG/Source/PCG/Public/PCGSettings.h b/Engine/Plugins/Experimental/PCG/Source/PCG/Public/PCGSettings.h
--- a/Engine/Plugins/Experimental/PCG/Source/PCG/Public/PCGSettings.h
+++ b/Engine/Plugins/Experimental/PCG/Source/PCG/Public/PCGSettings.h
@@ -346,6 +346,10 @@ protected:
 	/** Method that can be called to dirty the cache data from this settings objects if the operator== does not allow to detect changes */
 	void DirtyCache();
 
+	// We need the more complex function (with PropertyChain) to detect child properties in structs, if they are overridable
+	virtual bool CanEditChange(const FEditPropertyChain& InPropertyChain) const override;
+
+	// Passthrough for the simpler method, to avoid modifying the child settings already overriding this method.
 	virtual bool CanEditChange(const FProperty* InProperty) const override;
 #endif
 
