[PCG] Remove double reset arrow in graph parameters override

#rnx

diff --git a/Engine/Plugins/Experimental/PCG/Source/PCGEditor/Private/Details/PCGInstancedPropertyBagOverrideDetails.cpp b/Engine/Plugins/Experimental/PCG/Source/PCGEditor/Private/Details/PCGInstancedPropertyBagOverrideDetails.cpp
--- a/Engine/Plugins/Experimental/PCG/Source/PCGEditor/Private/Details/PCGInstancedPropertyBagOverrideDetails.cpp
+++ b/Engine/Plugins/Experimental/PCG/Source/PCGEditor/Private/Details/PCGInstancedPropertyBagOverrideDetails.cpp
@@ -123,13 +123,19 @@ void FPCGOverrideInstancedPropertyBagDataDetails::OnChildRowAdded(IDetailPropert
 		return;
 	}
 
+	TSharedPtr<IPropertyHandle> ChildPropertyHandle = ChildRow.GetPropertyHandle();
+	if (!ChildPropertyHandle.IsValid())
+	{
+		return;
+	}
+
+	ChildPropertyHandle->SetInstanceMetaData("NoResetToDefault", "true");
+
 	TSharedPtr<SWidget> NameWidget;
 	TSharedPtr<SWidget> ValueWidget;
 	FDetailWidgetRow Row;
 	ChildRow.GetDefaultWidgets(NameWidget, ValueWidget, Row);
 
-	TSharedPtr<IPropertyHandle> ChildPropertyHandle = ChildRow.GetPropertyHandle();
-
 	FGuid PropertyID;
 
 	const bool bIsOverridden = Owner->IsPropertyOverridden(ChildPropertyHandle->GetProperty());
