[PCG] Fixes to projected splines which were not working when used in a Difference

diff --git a/Engine/Plugins/Experimental/PCG/Source/PCG/Private/Data/PCGProjectionData.cpp b/Engine/Plugins/Experimental/PCG/Source/PCG/Private/Data/PCGProjectionData.cpp
--- a/Engine/Plugins/Experimental/PCG/Source/PCG/Private/Data/PCGProjectionData.cpp
+++ b/Engine/Plugins/Experimental/PCG/Source/PCG/Private/Data/PCGProjectionData.cpp
@@ -6,6 +6,7 @@
 #include "Data/PCGPointData.h"
 #include "Data/PCGSplineData.h"
 #include "Data/PCGSpatialData.h"
+#include "Elements/PCGProjectionParams.h"
 #include "Helpers/PCGAsync.h"
 #include "Metadata/PCGMetadataAccessor.h"
 
@@ -62,13 +63,24 @@ FBox UPCGProjectionData::ProjectBounds(const FBox& InBounds) const
 {
 	FBox Bounds(EForceInit::ForceInit);
 
+	const FBox PointAABB = FBox::BuildAABB(FVector::ZeroVector, FVector::ZeroVector);
+
 	for (int Corner = 0; Corner < 8; ++Corner)
 	{
-		Bounds += Target->TransformPosition(
-			FVector(
-				(Corner / 4) ? InBounds.Max.X : InBounds.Min.X,
-				((Corner / 2) % 2) ? InBounds.Max.Y : InBounds.Min.Y,
-				(Corner % 2) ? InBounds.Max.Z : InBounds.Min.Z));
+		const FVector CornerPoint = FVector(
+			(Corner / 4) ? InBounds.Max.X : InBounds.Min.X,
+			((Corner / 2) % 2) ? InBounds.Max.Y : InBounds.Min.Y,
+			(Corner % 2) ? InBounds.Max.Z : InBounds.Min.Z);
+
+		FPCGPoint ProjectedPoint;
+		if (Target->ProjectPoint(FTransform(CornerPoint), PointAABB, ProjectionParams, ProjectedPoint, nullptr))
+		{
+			Bounds += ProjectedPoint.Transform.GetLocation();
+		}
+		else
+		{
+			Bounds += CornerPoint;
+		}
 	}
 
 	// Fixup the Z direction, as transforming the corners is not sufficient
diff --git a/Engine/Plugins/Experimental/PCG/Source/PCG/Private/Data/PCGSpatialData.cpp b/Engine/Plugins/Experimental/PCG/Source/PCG/Private/Data/PCGSpatialData.cpp
--- a/Engine/Plugins/Experimental/PCG/Source/PCG/Private/Data/PCGSpatialData.cpp
+++ b/Engine/Plugins/Experimental/PCG/Source/PCG/Private/Data/PCGSpatialData.cpp
@@ -102,19 +102,6 @@ float UPCGSpatialData::GetDensityAtPosition(const FVector& InPosition) const
 	}
 }
 
-FVector UPCGSpatialData::TransformPosition(const FVector& InPosition) const
-{
-	FPCGPoint TemporaryPoint;
-	if (SamplePoint(FTransform(InPosition), FBox::BuildAABB(FVector::ZeroVector, FVector::ZeroVector), TemporaryPoint, nullptr))
-	{
-		return TemporaryPoint.Transform.GetLocation();
-	}
-	else
-	{
-		return InPosition;
-	}
-}
-
 bool UPCGSpatialData::ProjectPoint(const FTransform& InTransform, const FBox& InBounds, const FPCGProjectionParams& InParams, FPCGPoint& OutPoint, UPCGMetadata* OutMetadata) const
 {
 	// Fallback implementation - calls SamplePoint because SamplePoint was being used for projection previously.
diff --git a/Engine/Plugins/Experimental/PCG/Source/PCG/Private/Data/PCGSplineData.cpp b/Engine/Plugins/Experimental/PCG/Source/PCG/Private/Data/PCGSplineData.cpp
--- a/Engine/Plugins/Experimental/PCG/Source/PCG/Private/Data/PCGSplineData.cpp
+++ b/Engine/Plugins/Experimental/PCG/Source/PCG/Private/Data/PCGSplineData.cpp
@@ -225,12 +225,16 @@ bool UPCGSplineProjectionData::SamplePoint(const FTransform& InTransform, const
 	const FVector NearestPointOnSpline = Spline.GetLocationAtSplineInputKey(NearestInputKey, ESplineCoordinateSpace::World);
 	const FVector PointOnLine = FMath::ClosestPointOnInfiniteLine(InPosition, InPosition + SurfaceNormal, NearestPointOnSpline);
 
+	// In the following statements we check if point lies in projection of spline onto landscape, which is true if:
+	//  * When we hoist the point up to the nearest point on the unprojected spline, it overlaps the spline
+	//  * The point is on the landscape
+	
 	// TODO: this is super inefficient, could be done in 2D if we duplicate the sampling code
 	FPCGPoint SplinePoint;
 	if (GetSpline()->SamplePoint(FTransform(PointOnLine), InBounds, SplinePoint, OutMetadata))
 	{
 		FPCGPoint SurfacePoint;
-		if (GetSurface()->SamplePoint(SplinePoint.Transform, InBounds, SurfacePoint, OutMetadata))
+		if (GetSurface()->SamplePoint(InTransform, InBounds, SurfacePoint, OutMetadata))
 		{
 			OutPoint = SplinePoint;
 
@@ -264,7 +268,6 @@ FVector2D UPCGSplineProjectionData::Project(const FVector& InVector) const
 	// Find the largest coordinate of the normal and use as the projection axis
 	int BiggestCoordinateAxis = 0;
 	FVector::FReal BiggestCoordinate = FMath::Abs(SurfaceNormal[BiggestCoordinateAxis]);
-
 	for (int Axis = 1; Axis < 3; ++Axis)
 	{
 		FVector::FReal AbsoluteCoordinateValue = FMath::Abs(SurfaceNormal[Axis]);
diff --git a/Engine/Plugins/Experimental/PCG/Source/PCG/Public/Data/PCGDifferenceData.h b/Engine/Plugins/Experimental/PCG/Source/PCG/Public/Data/PCGDifferenceData.h
--- a/Engine/Plugins/Experimental/PCG/Source/PCG/Public/Data/PCGDifferenceData.h
+++ b/Engine/Plugins/Experimental/PCG/Source/PCG/Public/Data/PCGDifferenceData.h
@@ -14,7 +14,7 @@ UENUM()
 enum class EPCGDifferenceDensityFunction : uint8
 {
 	Minimum,
-	ClampedSubstraction,
+	ClampedSubstraction UMETA(DisplayName = "Clamped Subtraction"),
 	Binary
 };
 
diff --git a/Engine/Plugins/Experimental/PCG/Source/PCG/Public/Data/PCGSpatialData.h b/Engine/Plugins/Experimental/PCG/Source/PCG/Public/Data/PCGSpatialData.h
--- a/Engine/Plugins/Experimental/PCG/Source/PCG/Public/Data/PCGSpatialData.h
+++ b/Engine/Plugins/Experimental/PCG/Source/PCG/Public/Data/PCGSpatialData.h
@@ -89,9 +89,6 @@ public:
 
 	virtual const UPCGPointData* ToPointData(FPCGContext* Context, const FBox& InBounds = FBox(EForceInit::ForceInit)) const PURE_VIRTUAL(UPCGSpatialData::ToPointData, return nullptr;);
 
-	/** Transform a world-space position to a world-space position in relation to the current data. (ex: projection on surface) */
-	FVector TransformPosition(const FVector& InPosition) const;
-
 	/** Sample rotation, scale and other attributes from this data at the query position. Returns true if Transform location and Bounds overlaps this data. */
 	UFUNCTION(BlueprintCallable, Category = SpatialData)
 	virtual bool SamplePoint(const FTransform& Transform, const FBox& Bounds, FPCGPoint& OutPoint, UPCGMetadata* OutMetadata) const PURE_VIRTUAL(UPCGSpatialData::SamplePoint, return false;);
