[PCG] Guard better against tracked actors that could be invalid, because of weak pointers.

#rnx

diff --git a/Engine/Plugins/Experimental/PCG/Source/PCG/Private/PCGActorAndComponentMapping.cpp b/Engine/Plugins/Experimental/PCG/Source/PCG/Private/PCGActorAndComponentMapping.cpp
--- a/Engine/Plugins/Experimental/PCG/Source/PCG/Private/PCGActorAndComponentMapping.cpp
+++ b/Engine/Plugins/Experimental/PCG/Source/PCG/Private/PCGActorAndComponentMapping.cpp
@@ -1028,9 +1028,11 @@ void UPCGActorAndComponentMapping::OnObjectPropertyChanged(UObject* InObject, FP
 		{
 			if (TrackedActor.Value.Contains(InObject))
 			{
-				AActor* ActorToChange = TrackedActor.Key.Get();
-				OnActorChanged(ActorToChange, /*bInHasMoved=*/ false);
-				UpdateActorDependencies(ActorToChange);
+				if (AActor* ActorToChange = TrackedActor.Key.Get())
+				{
+					OnActorChanged(ActorToChange, /*bInHasMoved=*/ false);
+					UpdateActorDependencies(ActorToChange);
+				}
 			}
 		}
 
@@ -1212,6 +1214,11 @@ void UPCGActorAndComponentMapping::OnActorChanged(AActor* InActor, bool bInHasMo
 
 void UPCGActorAndComponentMapping::OnLandscapeChanged(ALandscapeProxy* InLandscape, const FLandscapeProxyComponentDataChangedParams& InChangeParams)
 {
+	if (!InLandscape)
+	{
+		return;
+	}
+
 	// We don't know if the landscape moved, only that it has changed. Since `bHasMoved` is doing a bit more, always assume that the landscape has moved.
 	OnActorChanged(InLandscape, /*bHasMoved=*/true);
 
@@ -1244,7 +1251,7 @@ void UPCGActorAndComponentMapping::UpdateActorDependencies(AActor* InActor)
 
 void UPCGActorAndComponentMapping::OnPCGGraphGeneratedOrCleaned(UPCGComponent* InComponent)
 {
-	if (!InComponent)
+	if (!InComponent || !InComponent->GetOwner())
 	{
 		return;
 	}
