[PCG] Minor cleanup while working on another CL

#rnx

diff --git a/Engine/Plugins/Experimental/PCG/Source/PCG/Public/Elements/PCGCollapseElement.h b/Engine/Plugins/Experimental/PCG/Source/PCG/Public/Elements/PCGCollapseElement.h
--- a/Engine/Plugins/Experimental/PCG/Source/PCG/Public/Elements/PCGCollapseElement.h
+++ b/Engine/Plugins/Experimental/PCG/Source/PCG/Public/Elements/PCGCollapseElement.h
@@ -17,7 +17,7 @@ public:
 	//~Begin UPCGSettings interface
 #if WITH_EDITOR
 	virtual FName GetDefaultNodeName() const override { return FName(TEXT("ToPoint")); }
-	virtual FText GetDefaultNodeTitle() const override { return NSLOCTEXT("PCGCollapseSettings", "NodeTitle", "To Point"); }
+	virtual FText GetDefaultNodeTitle() const override { return NSLOCTEXT("PCGCollapseElement", "NodeTitle", "To Point"); }
 	virtual EPCGSettingsType GetType() const override { return EPCGSettingsType::Spatial; }
 	virtual bool ShouldDrawNodeCompact() const override { return true; }
 #endif
