Significantly better behavior of cancel-on-refresh to not cancel generation tasks that haven't started.
Fixed a landscaspe data dirty issue on the component (related to usage of the landscape pins only)
Added a dependency only pin on the Gather node so that execution can be sequenced without require to filter data

diff --git a/Engine/Plugins/Experimental/PCG/Source/PCG/Private/Elements/PCGGather.cpp b/Engine/Plugins/Experimental/PCG/Source/PCG/Private/Elements/PCGGather.cpp
--- a/Engine/Plugins/Experimental/PCG/Source/PCG/Private/Elements/PCGGather.cpp
+++ b/Engine/Plugins/Experimental/PCG/Source/PCG/Private/Elements/PCGGather.cpp
@@ -4,24 +4,29 @@
 
 #include "PCGContext.h"
 #include "PCGPin.h"
-	
+
+#define LOCTEXT_NAMESPACE "PCGGatherElement"
+
 EPCGDataType UPCGGatherSettings::GetCurrentPinTypes(const UPCGPin* InPin) const
 {
 	check(InPin);
-	if (!InPin->IsOutputPin())
+
+	if (InPin->Properties.Label == PCGPinConstants::DefaultDependencyOnlyLabel)
 	{
 		return Super::GetCurrentPinTypes(InPin);
 	}
-
-	// Output pin narrows to union of inputs on first pin
-	const EPCGDataType InputTypeUnion = GetTypeUnionOfIncidentEdges(PCGPinConstants::DefaultInputLabel);
-	return InputTypeUnion != EPCGDataType::None ? InputTypeUnion : EPCGDataType::Any;
+	else
+	{
+		const EPCGDataType InputTypeUnion = GetTypeUnionOfIncidentEdges(PCGPinConstants::DefaultInputLabel);
+		return InputTypeUnion != EPCGDataType::None ? InputTypeUnion : EPCGDataType::Any;
+	}
 }
 
 TArray<FPCGPinProperties> UPCGGatherSettings::InputPinProperties() const
 {
 	TArray<FPCGPinProperties> PinProperties;
 	PinProperties.Emplace(PCGPinConstants::DefaultInputLabel, EPCGDataType::Any);
+	PinProperties.Emplace(PCGPinConstants::DefaultDependencyOnlyLabel, EPCGDataType::Any, /*bAllowMultipleConnections=*/true, /*bAllowMultipleData=*/true, LOCTEXT("DependencyPinTooltip", "Data passed to this pin will be used to order execution but will otherwise not contribute to the results of this node."));
 
 	return PinProperties;
 }
@@ -43,7 +48,22 @@ bool FPCGGatherElement::ExecuteInternal(FPCGContext* Context) const
 {
 	TRACE_CPUPROFILER_EVENT_SCOPE(FPCGGatherElement::Execute);
 
-	Context->OutputData = Context->InputData;
+	TArray<FPCGTaggedData> GatheredData = Context->InputData.GetInputsByPin(PCGPinConstants::DefaultInputLabel);
+	if (GatheredData.Num() == Context->InputData.TaggedData.Num())
+	{
+		Context->OutputData = Context->InputData;
+	}
+	else
+	{
+		Context->OutputData.TaggedData = MoveTemp(GatheredData);
+	}
+
+	for(FPCGTaggedData& TaggedData : Context->OutputData.TaggedData)
+	{
+		TaggedData.Pin = PCGPinConstants::DefaultOutputLabel;
+	}
 
 	return true;
 }
+
+#undef LOCTEXT_NAMESPACE
\ No newline at end of file
diff --git a/Engine/Plugins/Experimental/PCG/Source/PCG/Private/Graph/PCGGraphExecutor.cpp b/Engine/Plugins/Experimental/PCG/Source/PCG/Private/Graph/PCGGraphExecutor.cpp
--- a/Engine/Plugins/Experimental/PCG/Source/PCG/Private/Graph/PCGGraphExecutor.cpp
+++ b/Engine/Plugins/Experimental/PCG/Source/PCG/Private/Graph/PCGGraphExecutor.cpp
@@ -1035,17 +1035,23 @@ void FPCGGraphExecutor::QueueNextTasks(FPCGTaskId FinishedTask)
 		for (FPCGTaskId Successor : *Successors)
 		{
 			bool bAllPrerequisitesMet = true;
-			FPCGGraphTask& SuccessorTask = Tasks[Successor];
+			FPCGGraphTask* SuccessorTaskPtr = Tasks.Find(Successor);
 
-			for (const FPCGGraphTaskInput& Input : SuccessorTask.Inputs)
+			// This should never be null, but later recovery should be able to cleanup this properly
+			if (ensure(SuccessorTaskPtr))
 			{
-				bAllPrerequisitesMet &= OutputData.Contains(Input.TaskId);
-			}
+				FPCGGraphTask& SuccessorTask = *SuccessorTaskPtr;
 
-			if (bAllPrerequisitesMet)
-			{
-				ReadyTasks.Emplace(MoveTemp(SuccessorTask));
-				Tasks.Remove(Successor);
+				for (const FPCGGraphTaskInput& Input : SuccessorTask.Inputs)
+				{
+					bAllPrerequisitesMet &= OutputData.Contains(Input.TaskId);
+				}
+
+				if (bAllPrerequisitesMet)
+				{
+					ReadyTasks.Emplace(MoveTemp(SuccessorTask));
+					Tasks.Remove(Successor);
+				}
 			}
 		}
 
@@ -1398,6 +1404,10 @@ bool FPCGFetchInputElement::ExecuteInternal(FPCGContext* Context) const
 		return true;
 	}
 
+#if WITH_EDITOR
+	Component->StartGenerationInProgress();
+#endif
+
 	check(Context->Node);
 
 	if (Context->Node->IsOutputPinConnected(PCGPinConstants::DefaultInputLabel))
diff --git a/Engine/Plugins/Experimental/PCG/Source/PCG/Private/PCGComponent.cpp b/Engine/Plugins/Experimental/PCG/Source/PCG/Private/PCGComponent.cpp
--- a/Engine/Plugins/Experimental/PCG/Source/PCG/Private/PCGComponent.cpp
+++ b/Engine/Plugins/Experimental/PCG/Source/PCG/Private/PCGComponent.cpp
@@ -405,6 +405,8 @@ void UPCGComponent::PostProcessGraph(const FBox& InNewBounds, bool bInGenerated,
 		FPropertyChangedEvent GeneratedOutputChangedEvent(GeneratedOutputProperty, EPropertyChangeType::ValueSet);
 		FCoreUObjectDelegates::OnObjectPropertyChanged.Broadcast(this, GeneratedOutputChangedEvent);
 	}
+
+	StopGenerationInProgress();
 #endif
 }
 
@@ -499,6 +501,8 @@ void UPCGComponent::OnProcessGraphAborted(bool bQuiet)
 	CurrentRefreshTask = InvalidPCGTaskId;
 	// Implementation note: while it may seem logical to clear the bDirtyGenerated flag here, 
 	// the component is still considered dirty if we aborted processing, hence it should stay this way.
+
+	StopGenerationInProgress();
 #endif
 }
 
@@ -1144,7 +1148,7 @@ void UPCGComponent::RefreshAfterGraphChanged(UPCGGraphInterface* InGraph, bool b
 		DirtyGenerated(bDirtyInputs ? (EPCGComponentDirtyFlag::Actor | EPCGComponentDirtyFlag::Landscape) : EPCGComponentDirtyFlag::None);
 		if (bHasGraph)
 		{
-			Refresh();
+			Refresh(bIsStructural);
 		}
 		else
 		{
@@ -1265,7 +1269,7 @@ void UPCGComponent::PostEditUndo()
 
 	if (bGenerated)
 	{
-		Refresh();
+		Refresh(/*bIsStructural=*/true);
 	}
 }
 
@@ -1322,13 +1326,9 @@ void UPCGComponent::DirtyGenerated(EPCGComponentDirtyFlag DirtyFlag, const bool
 	if (!!(DirtyFlag & EPCGComponentDirtyFlag::Actor))
 	{
 		CachedActorData = nullptr;
-		
-		if (Cast<ALandscapeProxy>(GetOwner()))
-		{
-			CachedLandscapeData = nullptr;
-			CachedLandscapeHeightData = nullptr;
-		}
-
+		// Since landscape data is related on the bounds of the current actor, when we dirty the actor data, we need to dirty the landscape data as well
+		CachedLandscapeData = nullptr;
+		CachedLandscapeHeightData = nullptr;
 		CachedInputData = nullptr;
 		CachedPCGData = nullptr;
 	}
@@ -1438,7 +1438,7 @@ const FPCGDataCollection* UPCGComponent::GetInspectionData(const FString& InStac
 	return InspectionCache.Find(InStackPath);
 }
 
-void UPCGComponent::Refresh()
+void UPCGComponent::Refresh(bool bStructural)
 {
 	// Disable auto-refreshing on preview actors until we have something more robust on the execution side.
 	if (GetOwner() && GetOwner()->bIsEditorPreviewActor)
@@ -1446,22 +1446,68 @@ void UPCGComponent::Refresh()
 		return;
 	}
 
+	// If the component is tagged as not to regenerate in the editor, only exceptional cases should trigger a refresh
+	// namely: the component is deactivated.
+	// Note that the component changing its IsPartitioned state is already covered in the PostEditChangeProperty
+	// Note that even if this is force refresh/structural change, we will NOT refresh
+	if (!bRegenerateInEditor && bActivated)
+	{
+		// We still need to trigger component registration event otherwise further generations will fail if this is moved.
+		// Note that we pass in false here to remove everything when moving a partitioned graph because we would otherwise need to do a reversible stamp to support this
+		if (UPCGSubsystem* Subsystem = GetSubsystem())
+		{
+			Subsystem->RegisterOrUpdatePCGComponent(this, bGenerated);
+		}
+
+		return;
+	}
+
 	// Discard any refresh if have already one scheduled.
 	if (UPCGSubsystem* Subsystem = GetSubsystem())
 	{
-		//TDOO: uncomment this once we have the proper things in place to prevent cancellation-by-landscape changes when we have a graph that does change the landscape
-		/*if (CurrentGenerationTask != InvalidPCGTaskId)
+		// Cancel an already existing generation if either the change is structural in nature (which requires a recompilation, so a full-rescheduling)
+		// or if the generation is already started
+		const bool bGenerationWasInProgress = IsGenerationInProgress();
+		if (CurrentGenerationTask != InvalidPCGTaskId && (bStructural || bGenerationWasInProgress))
 		{
 			CancelGeneration();
-		}*/
+		}
+
+		// Calling a new refresh here might not be sufficient; if the current component was generating but was not previously generated,
+		// then the bGenerated flag will be false, which will prevent a subsequent update here
+		if (CurrentRefreshTask == InvalidPCGTaskId && CurrentCleanupTask == InvalidPCGTaskId)
+		{
+			CurrentRefreshTask = Subsystem->ScheduleRefresh(this, bGenerationWasInProgress);
+		}
+	}
+}
+
+void UPCGComponent::StartGenerationInProgress()
+{
+	// Implementation detail:
+	// Since the original component is not guaranteed to run the FetchInput element, local components are "allowed" to mark generation in progress on their original component.
+	// However, the PostProcessGraph on the original component will be guaranteed to be called at the end of the execution so we do not need this mechanism in that case.
+	bGenerationInProgress = true;
 
-		if (CurrentRefreshTask == InvalidPCGTaskId)
+	if (IsLocalComponent())
+	{
+		if (UPCGComponent* OriginalComponent = CastChecked<APCGPartitionActor>(GetOwner())->GetOriginalComponent(this))
 		{
-			CurrentRefreshTask = Subsystem->ScheduleRefresh(this);
+			OriginalComponent->bGenerationInProgress = true;
 		}
 	}
 }
 
+void UPCGComponent::StopGenerationInProgress()
+{
+	bGenerationInProgress = false;
+}
+
+bool UPCGComponent::IsGenerationInProgress()
+{
+	return bGenerationInProgress;
+}
+
 bool UPCGComponent::ShouldGenerateBPPCGAddedToWorld() const
 {
 	if (PCGHelpers::IsRuntimeOrPIE())
@@ -1507,7 +1553,7 @@ bool UPCGComponent::IsActorTracked(AActor* InActor, bool& bOutIsCulled) const
 	return bFound;
 }
 
-void UPCGComponent::OnRefresh()
+void UPCGComponent::OnRefresh(bool bForceRefresh)
 {
 	// Mark the refresh task invalid to allow re-triggering refreshes
 	CurrentRefreshTask = InvalidPCGTaskId;
@@ -1515,7 +1561,8 @@ void UPCGComponent::OnRefresh()
 	// Before doing a refresh, update the component to the subsystem if we are partitioned
 	// Only redo the mapping if we are generated
 	UPCGSubsystem* Subsystem = GetSubsystem();
-	bool bWasGenerated = bGenerated;
+	const bool bWasGenerated = bGenerated;
+	const bool bWasGeneratedOrGenerating = bWasGenerated || bForceRefresh;
 
 	if (IsPartitioned())
 	{
@@ -1528,7 +1575,7 @@ void UPCGComponent::OnRefresh()
 
 	if (Subsystem)
 	{
-		Subsystem->RegisterOrUpdatePCGComponent(this, /*bDoActorMapping=*/ bWasGenerated);
+		Subsystem->RegisterOrUpdatePCGComponent(this, /*bDoActorMapping=*/ bWasGeneratedOrGenerating);
 	}
 
 	// Following a change in some properties or in some spatial information related to this component,
@@ -1542,7 +1589,7 @@ void UPCGComponent::OnRefresh()
 	{
 		// If we just cleaned up resources, call back generate. Only do this for original component, which will then trigger
 		// generation of local components. Also, for BPs, we ask if we should generate, to support generate on added to world.
-		if ((bWasGenerated || ShouldGenerateBPPCGAddedToWorld()) && !IsLocalComponent() && (!bGenerated || bRegenerateInEditor))
+		if ((bWasGeneratedOrGenerating || ShouldGenerateBPPCGAddedToWorld()) && !IsLocalComponent() && (!bGenerated || bRegenerateInEditor))
 		{
 			GenerateLocal(/*bForce=*/false);
 		}
@@ -2253,7 +2300,7 @@ void FPCGComponentInstanceData::ApplyToComponent(UActorComponent* Component, con
 		// Note that we only do this if we are not currently loading
 		if (!SourceComponent || !SourceComponent->HasAllFlags(RF_WasLoaded))
 		{
-			PCGComponent->Refresh();
+			PCGComponent->Refresh(/*bIsStructural=*/true);
 		}
 #endif // WITH_EDITOR
 	}
diff --git a/Engine/Plugins/Experimental/PCG/Source/PCG/Private/PCGSubsystem.cpp b/Engine/Plugins/Experimental/PCG/Source/PCG/Private/PCGSubsystem.cpp
--- a/Engine/Plugins/Experimental/PCG/Source/PCG/Private/PCGSubsystem.cpp
+++ b/Engine/Plugins/Experimental/PCG/Source/PCG/Private/PCGSubsystem.cpp
@@ -937,14 +937,14 @@ FPCGTaskId UPCGSubsystem::ProcessGraph(UPCGComponent* Component, const FBox& InP
 	}
 }
 
-FPCGTaskId UPCGSubsystem::ScheduleRefresh(UPCGComponent* Component)
+FPCGTaskId UPCGSubsystem::ScheduleRefresh(UPCGComponent* Component, bool bForceRegen)
 {
 	TWeakObjectPtr<UPCGComponent> ComponentPtr(Component);
 
-	auto RefreshTask = [ComponentPtr]() {
+	auto RefreshTask = [ComponentPtr, bForceRegen]() {
 		if (UPCGComponent* Component = ComponentPtr.Get())
 		{
-			Component->OnRefresh();
+			Component->OnRefresh(bForceRegen);
 		}
 		return true;
 	};
diff --git a/Engine/Plugins/Experimental/PCG/Source/PCG/Public/Elements/PCGGather.h b/Engine/Plugins/Experimental/PCG/Source/PCG/Public/Elements/PCGGather.h
--- a/Engine/Plugins/Experimental/PCG/Source/PCG/Public/Elements/PCGGather.h
+++ b/Engine/Plugins/Experimental/PCG/Source/PCG/Public/Elements/PCGGather.h
@@ -18,6 +18,7 @@ public:
 #if WITH_EDITOR
 	virtual FName GetDefaultNodeName() const override { return FName(TEXT("Gather")); }
 	virtual FText GetDefaultNodeTitle() const override { return NSLOCTEXT("PCGGatherSettings", "NodeTitle", "Gather"); }
+	virtual FText GetNodeTooltipText() const override { return NSLOCTEXT("PCGGatherSettings", "NodeTooltip", "Gathers multiple data in a single collection. Can also be used to order execution through the dependency-only pin."); }
 	virtual EPCGSettingsType GetType() const override { return EPCGSettingsType::Spatial; }
 	virtual bool HasDynamicPins() const override { return true; }
 #endif
diff --git a/Engine/Plugins/Experimental/PCG/Source/PCG/Public/PCGCommon.h b/Engine/Plugins/Experimental/PCG/Source/PCG/Public/PCGCommon.h
--- a/Engine/Plugins/Experimental/PCG/Source/PCG/Public/PCGCommon.h
+++ b/Engine/Plugins/Experimental/PCG/Source/PCG/Public/PCGCommon.h
@@ -91,6 +91,7 @@ namespace PCGPinConstants
 	const FName DefaultInputLabel = TEXT("In");
 	const FName DefaultOutputLabel = TEXT("Out");
 	const FName DefaultParamsLabel = TEXT("Overrides");
+	const FName DefaultDependencyOnlyLabel = TEXT("Dependency Only");
 
 namespace Private
 {
diff --git a/Engine/Plugins/Experimental/PCG/Source/PCG/Public/PCGComponent.h b/Engine/Plugins/Experimental/PCG/Source/PCG/Public/PCGComponent.h
--- a/Engine/Plugins/Experimental/PCG/Source/PCG/Public/PCGComponent.h
+++ b/Engine/Plugins/Experimental/PCG/Source/PCG/Public/PCGComponent.h
@@ -231,8 +231,12 @@ public:
 	FPCGTaskId GetGenerationTaskId() const { return CurrentGenerationTask; }
 
 #if WITH_EDITOR
-	void Refresh();
-	void OnRefresh();
+	void Refresh(bool bStructural = false);
+	void OnRefresh(bool bForceRefresh);
+
+	void StartGenerationInProgress();
+	void StopGenerationInProgress();
+	bool IsGenerationInProgress();
 
 	/** Dirty generated data depending on the flag. By default the call is forwarded to the local components.
 	    We don't forward if the local component has callbacks that would dirty them too.
@@ -394,6 +398,8 @@ private:
 
 #if WITH_EDITOR
 	FPCGTaskId CurrentRefreshTask = InvalidPCGTaskId;
+
+	bool bGenerationInProgress = false;
 #endif // WITH_EDITOR
 
 	UPROPERTY(VisibleAnywhere, Transient, Category = Properties, meta = (EditCondition = false, EditConditionHides))
diff --git a/Engine/Plugins/Experimental/PCG/Source/PCG/Public/PCGSubsystem.h b/Engine/Plugins/Experimental/PCG/Source/PCG/Public/PCGSubsystem.h
--- a/Engine/Plugins/Experimental/PCG/Source/PCG/Public/PCGSubsystem.h
+++ b/Engine/Plugins/Experimental/PCG/Source/PCG/Public/PCGSubsystem.h
@@ -153,7 +153,7 @@ public:
 #if WITH_EDITOR
 public:
 	/** Schedule refresh on the current or next frame */
-	FPCGTaskId ScheduleRefresh(UPCGComponent* SourceComponent);
+	FPCGTaskId ScheduleRefresh(UPCGComponent* SourceComponent, bool bForceRefresh);
 
 	/** Schedules an operation to cleanup the graph in the given bounds */
 	FPCGTaskId CleanupGraph(UPCGComponent* Component, const FBox& InBounds, bool bRemoveComponents, bool bSave);
