[PCG] Fix test failures

* CopyPoints.Basic - Logic seemed simply wrong
* DensityRemap.Basic - Logic wrong for new default setting value (exclude outside range == false)
* Distance.PointToPoint - Don't know why 2 outputs would be expected. Perhaps settings pass-through.
* Noise.CalcLocalCoordinates2D - Typo - when edge blend disabled, blend weight should be 1 (not 0)
* General - initialised TestData.InputData.Crc to a simple valid Crc to fix ensure() fails

#rnx

diff --git a/Engine/Plugins/Experimental/PCG/Source/PCG/Private/Tests/Elements/PCGCopyPointsTest.cpp b/Engine/Plugins/Experimental/PCG/Source/PCG/Private/Tests/Elements/PCGCopyPointsTest.cpp
--- a/Engine/Plugins/Experimental/PCG/Source/PCG/Private/Tests/Elements/PCGCopyPointsTest.cpp
+++ b/Engine/Plugins/Experimental/PCG/Source/PCG/Private/Tests/Elements/PCGCopyPointsTest.cpp
@@ -267,15 +267,15 @@ bool FPCGCopyPointsTest::RunTest(const FString& Parameters)
 				Color = TargetPoint.Color;
 			}
 
-			if (Settings->ColorInheritance == EPCGCopyPointsInheritanceMode::Relative)
+			if (Settings->SeedInheritance == EPCGCopyPointsInheritanceMode::Relative)
 			{
-				Seed = PCGHelpers::ComputeSeed((int)Location.X, (int)Location.Y, (int)Location.Z);
+				Seed = PCGHelpers::ComputeSeed(SourcePoint.Seed, TargetPoint.Seed);
 			}
-			else if (Settings->ColorInheritance == EPCGCopyPointsInheritanceMode::Source)
+			else if (Settings->SeedInheritance == EPCGCopyPointsInheritanceMode::Source)
 			{
 				Seed = SourcePoint.Seed;
 			}
-			else // if (Settings->ColorInheritance == EPCGCopyPointsInheritanceMode::Target)
+			else // if (Settings->SeedInheritance == EPCGCopyPointsInheritanceMode::Target)
 			{
 				Seed = TargetPoint.Seed;
 			}
@@ -284,7 +284,7 @@ bool FPCGCopyPointsTest::RunTest(const FString& Parameters)
 			bTestPassed &= TestEqual("Valid scale", Scale, OutTransform.GetScale3D());
 			bTestPassed &= TestEqual("Valid color", Color, OutPoint.Color);
 			bTestPassed &= TestEqual("Valid location", Location, OutTransform.GetLocation());
-			bTestPassed &= TestEqual("Valid location", Seed, OutPoint.Seed);		
+			bTestPassed &= TestEqual("Valid seed", Seed, OutPoint.Seed);		
 
 			// Validate point value keys
 			for (int AttributeIndex = 0; AttributeIndex < InheritedAttributes.Num(); ++AttributeIndex)
diff --git a/Engine/Plugins/Experimental/PCG/Source/PCG/Private/Tests/Elements/PCGDensityRemapTest.cpp b/Engine/Plugins/Experimental/PCG/Source/PCG/Private/Tests/Elements/PCGDensityRemapTest.cpp
--- a/Engine/Plugins/Experimental/PCG/Source/PCG/Private/Tests/Elements/PCGDensityRemapTest.cpp
+++ b/Engine/Plugins/Experimental/PCG/Source/PCG/Private/Tests/Elements/PCGDensityRemapTest.cpp
@@ -102,6 +102,7 @@ bool FPCGDensityRemapTest::RunTest(const FString& Parameters)
 		Settings->InRangeMax = 1.f;
 		Settings->OutRangeMin = 0.f;
 		Settings->OutRangeMax = 1.f;
+		Settings->bExcludeValuesOutsideInputRange = true;
 		bTestPassed &= TestTrue("Input and Output are identical when InRange and OutRange are identical", ValidateDensityRemap({ 0.f, 0.2f, 0.4f, 0.6f, 0.8f, 1.f }));
 	}
 
@@ -111,12 +112,14 @@ bool FPCGDensityRemapTest::RunTest(const FString& Parameters)
 		Settings->InRangeMax = 0.4f;
 		Settings->OutRangeMin = 0.f;
 		Settings->OutRangeMax = 1.f;
+		Settings->bExcludeValuesOutsideInputRange = true;
 		bTestPassed &= TestTrue("Valid densities for partial InRange", ValidateDensityRemap({ 0.f, 0.5f, 1.f, 0.6f, 0.8f, 1.f }));
 
 		Settings->InRangeMin = 0.4f;
 		Settings->InRangeMax = 0.f;
 		Settings->OutRangeMin = 1.f;
 		Settings->OutRangeMax = 0.f;
+		Settings->bExcludeValuesOutsideInputRange = true;
 		bTestPassed &= TestTrue("Inverting ranges does not effect output", ValidateDensityRemap({ 0.f, 0.5f, 1.f, 0.6f, 0.8f, 1.f }));
 	}
 
@@ -126,12 +129,14 @@ bool FPCGDensityRemapTest::RunTest(const FString& Parameters)
 		Settings->InRangeMax = 1.f;
 		Settings->OutRangeMin = 0.5f;
 		Settings->OutRangeMax = 1.f;
+		Settings->bExcludeValuesOutsideInputRange = true;
 		bTestPassed &= TestTrue("Valid densities for partial OutRange", ValidateDensityRemap({ 0.f, 0.2f, 3.f / 6.f, 4.f / 6.f, 5.f / 6.f, 1.f }));
 
 		Settings->InRangeMin = 1.f;
 		Settings->InRangeMax = 0.4f;
 		Settings->OutRangeMin = 1.f;
 		Settings->OutRangeMax = 0.5f;
+		Settings->bExcludeValuesOutsideInputRange = true;
 		bTestPassed &= TestTrue("Inverting ranges does not effect output", ValidateDensityRemap({ 0.f, 0.2f, 3.f / 6.f, 4.f / 6.f, 5.f / 6.f, 1.f }));
 	}
 
@@ -142,10 +147,7 @@ bool FPCGDensityRemapTest::RunTest(const FString& Parameters)
 		Settings->OutRangeMin = 0.5f;
 		Settings->OutRangeMax = 1.f;
 		Settings->bExcludeValuesOutsideInputRange = false;
-
 		bTestPassed &= TestTrue("All values are remapped when Range Exclusion is disabled", ValidateDensityRemap({ 1.f / 6.f, 2.f / 6.f, 3.f / 6.f, 4.f / 6.f, 5.f / 6.f, 1.f }));
-
-		Settings->bExcludeValuesOutsideInputRange = true;
 	}
 
 	// Test Point to Point, Point to Range, and Range to Point
@@ -154,18 +156,21 @@ bool FPCGDensityRemapTest::RunTest(const FString& Parameters)
 		Settings->InRangeMax = 0.2f;
 		Settings->OutRangeMin = 0.f;
 		Settings->OutRangeMax = 1.f;
+		Settings->bExcludeValuesOutsideInputRange = true;
 		bTestPassed &= TestTrue("Point input to Range output", ValidateDensityRemap({ 0.f, 0.5f, 0.4f, 0.6f, 0.8f, 1.f }));
 
 		Settings->InRangeMin = 0.2f;
 		Settings->InRangeMax = 0.2f;
 		Settings->OutRangeMin = 0.5f;
 		Settings->OutRangeMax = 0.5f;
+		Settings->bExcludeValuesOutsideInputRange = true;
 		bTestPassed &= TestTrue("Point input to Point output", ValidateDensityRemap({ 0.f, 0.5f, 0.4f, 0.6f, 0.8f, 1.f }));
 
 		Settings->InRangeMin = 0.2f;
 		Settings->InRangeMax = 1.f;
 		Settings->OutRangeMin = 0.5f;
 		Settings->OutRangeMax = 0.5f;
+		Settings->bExcludeValuesOutsideInputRange = true;
 		bTestPassed &= TestTrue("Range input to Point output", ValidateDensityRemap({ 0.f, 0.5f, 0.5f, 0.5f, 0.5f, 0.5f }));
 	}
 
diff --git a/Engine/Plugins/Experimental/PCG/Source/PCG/Private/Tests/Elements/PCGDistanceTest.cpp b/Engine/Plugins/Experimental/PCG/Source/PCG/Private/Tests/Elements/PCGDistanceTest.cpp
--- a/Engine/Plugins/Experimental/PCG/Source/PCG/Private/Tests/Elements/PCGDistanceTest.cpp
+++ b/Engine/Plugins/Experimental/PCG/Source/PCG/Private/Tests/Elements/PCGDistanceTest.cpp
@@ -46,7 +46,7 @@ bool FPCGDistanceTest_PointToPoint::RunTest(const FString& Parameters)
 	// test our point data
 	const TArray<FPCGTaggedData>& Outputs = Context->OutputData.GetInputsByPin(PCGPinConstants::DefaultOutputLabel);
 
-	UTEST_EQUAL("Output count", Outputs.Num(), 2);
+	UTEST_EQUAL("Output count", Outputs.Num(), 1);
 
 	UPCGPointData *OutPointData = Cast<UPCGPointData>(Outputs[0].Data);
 
@@ -106,7 +106,7 @@ bool FPCGDistanceTest_SetDensity::RunTest(const FString& Parameters)
 	// test our point data
 	const TArray<FPCGTaggedData>& Outputs = Context->OutputData.GetInputsByPin(PCGPinConstants::DefaultOutputLabel);
 
-	UTEST_EQUAL("Output count", Outputs.Num(), 2);
+	UTEST_EQUAL("Output count", Outputs.Num(), 1);
 
 	UPCGPointData *OutPointData = Cast<UPCGPointData>(Outputs[0].Data);
 
@@ -168,7 +168,7 @@ bool FPCGDistanceTest_PointToSphere::RunTest(const FString& Parameters)
 	// test our point data
 	const TArray<FPCGTaggedData>& Outputs = Context->OutputData.GetInputsByPin(PCGPinConstants::DefaultOutputLabel);
 
-	UTEST_EQUAL("Output count", Outputs.Num(), 2);
+	UTEST_EQUAL("Output count", Outputs.Num(), 1);
 
 	UPCGPointData *OutPointData = Cast<UPCGPointData>(Outputs[0].Data);
 
@@ -233,7 +233,7 @@ bool FPCGDistanceTest_PointToBox::RunTest(const FString& Parameters)
 	// test our point data
 	const TArray<FPCGTaggedData>& Outputs = Context->OutputData.GetInputsByPin(PCGPinConstants::DefaultOutputLabel);
 
-	UTEST_EQUAL("Output count", Outputs.Num(), 2);
+	UTEST_EQUAL("Output count", Outputs.Num(), 1);
 
 	UPCGPointData *OutPointData = Cast<UPCGPointData>(Outputs[0].Data);
 
diff --git a/Engine/Plugins/Experimental/PCG/Source/PCG/Private/Tests/Elements/PCGSpatialNoiseTest.cpp b/Engine/Plugins/Experimental/PCG/Source/PCG/Private/Tests/Elements/PCGSpatialNoiseTest.cpp
--- a/Engine/Plugins/Experimental/PCG/Source/PCG/Private/Tests/Elements/PCGSpatialNoiseTest.cpp
+++ b/Engine/Plugins/Experimental/PCG/Source/PCG/Private/Tests/Elements/PCGSpatialNoiseTest.cpp
@@ -13,8 +13,8 @@ bool FPCGSpatialNoise_CalcLocalCoordinates2D::RunTest(const FString& Parameters)
 	const FBox TestBox(-FVector::One(), FVector::One());
 	const FVector2D Scale = FVector2D::One();
 
-	UTEST_EQUAL("Disabled Edge value", CalcEdgeBlendAmount2D(PCGSpatialNoise::CalcLocalCoordinates2D(TestBox, FTransform::Identity, Scale, FVector::One()), 0.0), 0.0);
-	UTEST_EQUAL("Disabled Edge value", CalcEdgeBlendAmount2D(PCGSpatialNoise::CalcLocalCoordinates2D(TestBox, FTransform::Identity, Scale, FVector::Zero()), 0.0), 0.0);
+	UTEST_EQUAL("Disabled Edge value", CalcEdgeBlendAmount2D(PCGSpatialNoise::CalcLocalCoordinates2D(TestBox, FTransform::Identity, Scale, FVector::One()), 0.0), 1.0);
+	UTEST_EQUAL("Disabled Edge value", CalcEdgeBlendAmount2D(PCGSpatialNoise::CalcLocalCoordinates2D(TestBox, FTransform::Identity, Scale, FVector::Zero()), 0.0), 1.0);
 
 	PCGSpatialNoise::FLocalCoordinates2D Test0 = PCGSpatialNoise::CalcLocalCoordinates2D(TestBox, FTransform::Identity, Scale, FVector::One());
 	const double EdgeBlendAmount0 = CalcEdgeBlendAmount2D(Test0, 1.0);
diff --git a/Engine/Plugins/Experimental/PCG/Source/PCG/Private/Tests/PCGTestsCommon.cpp b/Engine/Plugins/Experimental/PCG/Source/PCG/Private/Tests/PCGTestsCommon.cpp
--- a/Engine/Plugins/Experimental/PCG/Source/PCG/Private/Tests/PCGTestsCommon.cpp
+++ b/Engine/Plugins/Experimental/PCG/Source/PCG/Private/Tests/PCGTestsCommon.cpp
@@ -4,15 +4,16 @@
 
 #include "PCGComponent.h"
 #include "PCGContext.h"
+#include "PCGCrc.h"
 #include "PCGElement.h"
 #include "PCGGraph.h"
 #include "PCGParamData.h"
+#include "PCGPin.h"
 #include "Data/PCGPointData.h"
 #include "Data/PCGPolyLineData.h" // IWYU pragma: keep
 #include "Data/PCGPrimitiveData.h" // IWYU pragma: keep
 #include "Data/PCGSurfaceData.h" // IWYU pragma: keep
 #include "Data/PCGVolumeData.h"
-#include "PCGPin.h"
 
 #if WITH_EDITOR
 #include "Editor.h"
@@ -48,6 +49,9 @@ namespace PCGTestsCommon
 		UPCGGraph* TestGraph = NewObject<UPCGGraph>(TestPCGComponent, FName(TEXT("Test PCG Graph")), RF_Transient);
 		check(TestGraph);
 		TestPCGComponent->SetGraphLocal(TestGraph);
+
+		// Tests don't currently need CRCs, so just assigned a trivial valid CRC to avoid asserts.
+		InputData.Crc = FPCGCrc(0);
 #else
 		TestActor = nullptr;
 		TestPCGComponent = nullptr;
