[PCG] Add Set operator to Attribute Math operations

#rnx

diff --git a/Engine/Plugins/Experimental/PCG/Source/PCG/Private/Elements/Metadata/PCGMetadataMathsOpElement.cpp b/Engine/Plugins/Experimental/PCG/Source/PCG/Private/Elements/Metadata/PCGMetadataMathsOpElement.cpp
--- a/Engine/Plugins/Experimental/PCG/Source/PCG/Private/Elements/Metadata/PCGMetadataMathsOpElement.cpp
+++ b/Engine/Plugins/Experimental/PCG/Source/PCG/Private/Elements/Metadata/PCGMetadataMathsOpElement.cpp
@@ -133,6 +133,8 @@ namespace PCGMetadataMathsSettings
 			return PCGMetadataMaths::Pow(Value1, Value2);
 		case EPCGMedadataMathsOperation::Modulo:
 			return PCGMetadataMaths::Modulo(Value1, Value2);
+		case EPCGMedadataMathsOperation::Set:
+			return Value2;
 		default:
 			return T{};
 		}
diff --git a/Engine/Plugins/Experimental/PCG/Source/PCG/Public/Elements/Metadata/PCGMetadataMathsOpElement.h b/Engine/Plugins/Experimental/PCG/Source/PCG/Public/Elements/Metadata/PCGMetadataMathsOpElement.h
--- a/Engine/Plugins/Experimental/PCG/Source/PCG/Public/Elements/Metadata/PCGMetadataMathsOpElement.h
+++ b/Engine/Plugins/Experimental/PCG/Source/PCG/Public/Elements/Metadata/PCGMetadataMathsOpElement.h
@@ -33,6 +33,7 @@ enum class EPCGMedadataMathsOperation : uint16
 	ClampMin,
 	ClampMax,
 	Modulo,
+	Set,
 
 	// Ternary op
 	TernaryOp = 1 << 12 UMETA(Hidden),
