[PCG] Removed automatic "shape finding" functionality from the Surface Sampler node

- Surface Sampler no longer implicitly tries to find a surface
- `FindShapeFromNetwork` virtual function removed from PCGSpatialData

#rnx

diff --git a/Engine/Plugins/Experimental/PCG/Source/PCG/Private/Data/PCGIntersectionData.cpp b/Engine/Plugins/Experimental/PCG/Source/PCG/Private/Data/PCGIntersectionData.cpp
--- a/Engine/Plugins/Experimental/PCG/Source/PCG/Private/Data/PCGIntersectionData.cpp
+++ b/Engine/Plugins/Experimental/PCG/Source/PCG/Private/Data/PCGIntersectionData.cpp
@@ -170,47 +170,6 @@ bool UPCGIntersectionData::HasNonTrivialTransform() const
 	return GetA()->HasNonTrivialTransform() || GetB()->HasNonTrivialTransform();
 }
 
-const UPCGSpatialData* UPCGIntersectionData::FindShapeFromNetwork(const int InDimension) const
-{
-	check(GetA() && GetB());
-
-	if (InDimension >= 0)
-	{
-		// Return first candidate that matches Dimension
-		if (const UPCGSpatialData* CandidateA = GetA()->FindShapeFromNetwork(InDimension))
-		{
-			return CandidateA;
-		}
-
-		if (const UPCGSpatialData* CandidateB = GetB()->FindShapeFromNetwork(InDimension))
-		{
-			return CandidateB;
-		}
-	}
-	else
-	{
-		// Return lowest Dimension candidate
-		const UPCGSpatialData* Result = nullptr;
-
-		int LowestDimension = MAX_uint32;
-		const UPCGSpatialData* LowestDimensionalShape = nullptr;
-
-		if (const UPCGSpatialData* CandidateA = GetA()->FindShapeFromNetwork(InDimension))
-		{
-			Result = CandidateA;
-			LowestDimension = CandidateA->GetDimension();
-		}
-
-		const UPCGSpatialData* CandidateB = GetB()->FindShapeFromNetwork(InDimension);
-		if (CandidateB && CandidateB->GetDimension() < LowestDimension)
-		{
-			Result = CandidateB;
-		}
-	}
-
-	return nullptr;
-}
-
 const UPCGSpatialData* UPCGIntersectionData::FindFirstConcreteShapeFromNetwork() const
 {
 	check(GetA() && GetB());
diff --git a/Engine/Plugins/Experimental/PCG/Source/PCG/Private/Data/PCGUnionData.cpp b/Engine/Plugins/Experimental/PCG/Source/PCG/Private/Data/PCGUnionData.cpp
--- a/Engine/Plugins/Experimental/PCG/Source/PCG/Private/Data/PCGUnionData.cpp
+++ b/Engine/Plugins/Experimental/PCG/Source/PCG/Private/Data/PCGUnionData.cpp
@@ -218,43 +218,6 @@ bool UPCGUnionData::HasNonTrivialTransform() const
 	return (FirstNonTrivialTransformData != nullptr || Super::HasNonTrivialTransform());
 }
 
-const UPCGSpatialData* UPCGUnionData::FindShapeFromNetwork(const int InDimension) const
-{
-	if (InDimension >= 0)
-	{
-		// Return first candidate that matches Dimension
-		for (TObjectPtr<const UPCGSpatialData> Datum : Data)
-		{
-			const UPCGSpatialData* Candidate = Datum ? Datum->FindShapeFromNetwork(InDimension) : nullptr;
-			if (Candidate)
-			{
-				return Candidate;
-			}
-		}
-	}
-	else
-	{
-		// Return lowest Dimension candidate
-		const UPCGSpatialData* Result = nullptr;
-		int LowestDimension = MAX_uint32;
-
-		for (TObjectPtr<const UPCGSpatialData> Datum : Data)
-		{
-			const UPCGSpatialData* Candidate = Datum ? Datum->FindShapeFromNetwork(InDimension) : nullptr;
-
-			if (Candidate && Candidate->GetDimension() < LowestDimension)
-			{
-				LowestDimension = Candidate->GetDimension();
-				Result = Candidate;
-			}
-		}
-
-		return Result;
-	}
-
-	return nullptr;
-}
-
 const UPCGSpatialData* UPCGUnionData::FindFirstConcreteShapeFromNetwork() const
 {
 	// Return first concrete candidate data.
diff --git a/Engine/Plugins/Experimental/PCG/Source/PCG/Private/Elements/PCGSurfaceSampler.cpp b/Engine/Plugins/Experimental/PCG/Source/PCG/Private/Elements/PCGSurfaceSampler.cpp
--- a/Engine/Plugins/Experimental/PCG/Source/PCG/Private/Elements/PCGSurfaceSampler.cpp
+++ b/Engine/Plugins/Experimental/PCG/Source/PCG/Private/Elements/PCGSurfaceSampler.cpp
@@ -157,7 +157,7 @@ namespace PCGSurfaceSampler
 		// Drop points slightly by an epsilon otherwise point can be culled. If the sampler has a volume connected as the Bounding Shape,
 		// the volume will call through to PCGHelpers::IsInsideBounds() which is a one sided test and points at the top of the volume
 		// will fail it. TODO perhaps the one-sided check can be isolated to component-bounds
-		const FVector::FReal ZMultiplier = 1.0 - UE_DOUBLE_SMALL_NUMBER;
+		constexpr FVector::FReal ZMultiplier = 1.0 - UE_DOUBLE_SMALL_NUMBER;
 		// Try to use a multiplier instead of a simply offset to combat loss of precision in floats. However if MaxZ is very small,
 		// then multiplier will not work, so just use an offset.
 		FVector::FReal SampleZ = (FMath::Abs(LoopData.InputBoundsMaxZ) > UE_DOUBLE_SMALL_NUMBER) ? LoopData.InputBoundsMaxZ * ZMultiplier : -UE_DOUBLE_SMALL_NUMBER;
@@ -360,30 +360,31 @@ bool FPCGSurfaceSamplerElement::AddGeneratingShapesToContext(FPCGSurfaceSamplerC
 
 	TArray<FPCGTaggedData> SurfaceInputs = InContext->InputData.GetInputsByPin(PCGSurfaceSamplerConstants::SurfaceLabel);
 
-	// Construct a list of shapes to generate samples from. Prefer to get these directly from the first input pin.
+	// If there are no surfaces to sample, early out
+	if (SurfaceInputs.IsEmpty())
+	{
+		return false;
+	}
+
+	// Construct a list of shapes to generate samples from. Get these directly from the first input pin.
 	for (FPCGTaggedData& TaggedData : SurfaceInputs)
 	{
-		if (const UPCGSpatialData* SpatialData = Cast<UPCGSpatialData>(TaggedData.Data))
+		if (const UPCGSpatialData* SurfaceData = Cast<UPCGSpatialData>(TaggedData.Data))
 		{
-			// Find a concrete shape for sampling. Prefer a 2D surface if we can find one.
-			if (const UPCGSpatialData* SurfaceData = SpatialData->FindShapeFromNetwork(/*InDimension=*/2))
-			{
-				InContext->GeneratingShapes.Add(SurfaceData);
-				Outputs.Add(TaggedData);
-			}
-			else if (const UPCGSpatialData* ConcreteData = SpatialData->FindFirstConcreteShapeFromNetwork())
-			{
-				// Alternatively surface-sample any concrete data - can be used to sprinkle samples down onto shapes like volumes.
-				// Searching like this allows the user to plonk in any composite network and it will often find the shape of interest.
-				// A potential extension would be to find all (unique?) concrete shapes and use all of them rather than just the first.
-				InContext->GeneratingShapes.Add(ConcreteData);
-				Outputs.Add(TaggedData);
-			}
+			InContext->GeneratingShapes.Add(SurfaceData);
+			Outputs.Add(TaggedData);
 		}
 	}
 
+	// If there are no generating shapes, early out
+	if (InContext->GeneratingShapes.IsEmpty())
+	{
+		PCGE_LOG_C(Warning, GraphAndLog, InContext, LOCTEXT("NoSurfaceFound", "No valid surfaces found from which to generate"));
+		return false;
+	}
+
 	// Grab the Bounding Shape input if there is one.
-	TArray<FPCGTaggedData> BoundingShapeInputs = InContext->InputData.GetInputsByPin(PCGSurfaceSamplerConstants::BoundingShapeLabel);
+	const TArray<FPCGTaggedData> BoundingShapeInputs = InContext->InputData.GetInputsByPin(PCGSurfaceSamplerConstants::BoundingShapeLabel);
 
 	if (!Settings->bUnbounded)
 	{
@@ -411,25 +412,6 @@ bool FPCGSurfaceSamplerElement::AddGeneratingShapesToContext(FPCGSurfaceSamplerC
 		InContext->BoundingShapeBounds = InContext->BoundingShapeSpatialInput->GetBounds();
 	}
 
-	// If no shapes were obtained from the first input pin, try to find a shape to sample from nodes connected to the second pin.
-	if (InContext->GeneratingShapes.Num() == 0 && InContext->BoundingShapeSpatialInput)
-	{
-		if (const UPCGSpatialData* GeneratorFromBoundingShapeInput = InContext->BoundingShapeSpatialInput->FindShapeFromNetwork(/*InDimension=*/2))
-		{
-			InContext->GeneratingShapes.Add(GeneratorFromBoundingShapeInput);
-
-			// If there was a bounding shape input, use it as the starting point to get the tags
-			if (BoundingShapeInputs.Num() > 0)
-			{
-				Outputs.Add(BoundingShapeInputs[0]);
-			}
-			else
-			{
-				Outputs.Emplace();
-			}
-		}
-	}
-
 	// Warn if something is connected but no shape could be obtained for sampling
 	if (InContext->GeneratingShapes.Num() == 0 && (BoundingShapeInputs.Num() > 0 || SurfaceInputs.Num() > 0))
 	{
diff --git a/Engine/Plugins/Experimental/PCG/Source/PCG/Public/Data/PCGDifferenceData.h b/Engine/Plugins/Experimental/PCG/Source/PCG/Public/Data/PCGDifferenceData.h
--- a/Engine/Plugins/Experimental/PCG/Source/PCG/Public/Data/PCGDifferenceData.h
+++ b/Engine/Plugins/Experimental/PCG/Source/PCG/Public/Data/PCGDifferenceData.h
@@ -69,8 +69,8 @@ public:
 	virtual FBox GetStrictBounds() const override;
 	virtual bool SamplePoint(const FTransform& Transform, const FBox& Bounds, FPCGPoint& OutPoint, UPCGMetadata* OutMetadata) const override;
 	virtual bool HasNonTrivialTransform() const override;
-	virtual const UPCGSpatialData* FindShapeFromNetwork(const int InDimension) const override { return GetSource() ? GetSource()->FindShapeFromNetwork(InDimension) : nullptr; }
 	virtual const UPCGSpatialData* FindFirstConcreteShapeFromNetwork() const override { return GetSource() ? GetSource()->FindFirstConcreteShapeFromNetwork() : nullptr; }
+
 protected:
 	virtual UPCGSpatialData* CopyInternal() const override;
 	//~End UPCGSpatialData interface
diff --git a/Engine/Plugins/Experimental/PCG/Source/PCG/Public/Data/PCGIntersectionData.h b/Engine/Plugins/Experimental/PCG/Source/PCG/Public/Data/PCGIntersectionData.h
--- a/Engine/Plugins/Experimental/PCG/Source/PCG/Public/Data/PCGIntersectionData.h
+++ b/Engine/Plugins/Experimental/PCG/Source/PCG/Public/Data/PCGIntersectionData.h
@@ -46,7 +46,6 @@ public:
 	virtual FBox GetStrictBounds() const override;
 	virtual bool SamplePoint(const FTransform& Transform, const FBox& Bounds, FPCGPoint& OutPoint, UPCGMetadata* OutMetadata) const override;
 	virtual bool HasNonTrivialTransform() const override;
-	virtual const UPCGSpatialData* FindShapeFromNetwork(const int InDimension) const override;
 	virtual const UPCGSpatialData* FindFirstConcreteShapeFromNetwork() const override;
 
 protected:
diff --git a/Engine/Plugins/Experimental/PCG/Source/PCG/Public/Data/PCGSpatialData.h b/Engine/Plugins/Experimental/PCG/Source/PCG/Public/Data/PCGSpatialData.h
--- a/Engine/Plugins/Experimental/PCG/Source/PCG/Public/Data/PCGSpatialData.h
+++ b/Engine/Plugins/Experimental/PCG/Source/PCG/Public/Data/PCGSpatialData.h
@@ -131,9 +131,6 @@ public:
 	/** True if this operation does not have an inverse and cannot be queried analytically/implicitly, and therefore must be collapsed to an explicit point representation. */
 	virtual bool RequiresCollapseToSample() const { return false; }
 
-	/** A call that is made recursively up through the graph to find the best candidate shape for point generation. If InDimension is -1, finds lowest dimensional shape. */
-	virtual const UPCGSpatialData* FindShapeFromNetwork(const int InDimension) const { return (InDimension == -1 || GetDimension() == InDimension) ? this : nullptr; }
-
 	/** Find the first concrete (non-composite) shape in the network. Depth first search. */
 	virtual const UPCGSpatialData* FindFirstConcreteShapeFromNetwork() const { return !!(GetDataType() & EPCGDataType::Concrete) ? this : nullptr; }
 
diff --git a/Engine/Plugins/Experimental/PCG/Source/PCG/Public/Data/PCGUnionData.h b/Engine/Plugins/Experimental/PCG/Source/PCG/Public/Data/PCGUnionData.h
--- a/Engine/Plugins/Experimental/PCG/Source/PCG/Public/Data/PCGUnionData.h
+++ b/Engine/Plugins/Experimental/PCG/Source/PCG/Public/Data/PCGUnionData.h
@@ -52,8 +52,8 @@ public:
 	virtual FBox GetStrictBounds() const override;
 	virtual bool SamplePoint(const FTransform& Transform, const FBox& Bounds, FPCGPoint& OutPoint, UPCGMetadata* OutMetadata) const override;
 	virtual bool HasNonTrivialTransform() const override;
-	virtual const UPCGSpatialData* FindShapeFromNetwork(const int InDimension) const override;
 	virtual const UPCGSpatialData* FindFirstConcreteShapeFromNetwork() const override;
+
 protected:
 	virtual UPCGSpatialData* CopyInternal() const override;
 	//~End UPCGSpatialData interface
