[PCG] Revert change on UserParameterGet, should fix crashes in this element.

#rnx

diff --git a/Engine/Plugins/Experimental/PCG/Source/PCG/Private/Elements/PCGUserParameterGet.cpp b/Engine/Plugins/Experimental/PCG/Source/PCG/Private/Elements/PCGUserParameterGet.cpp
--- a/Engine/Plugins/Experimental/PCG/Source/PCG/Private/Elements/PCGUserParameterGet.cpp
+++ b/Engine/Plugins/Experimental/PCG/Source/PCG/Private/Elements/PCGUserParameterGet.cpp
@@ -97,7 +97,7 @@ namespace PCGUserParameterGetSettings
 
 		// Finally, we go down the graph instance chain to find the first valid layout.
 		// In most cases, we should never enter this loop, as layouts should match. But it is a safeguard.
-		while (GraphInstance && !WantedScriptStruct->CompareScriptStruct(WantedScriptStruct, GraphParameters.GetScriptStruct(), /*PortFlags=*/0))
+		while (GraphInstance && GraphParameters.GetScriptStruct() != WantedScriptStruct)
 		{
 			GraphInstance = Cast<UPCGGraphInstance>(GraphInstance->Graph);
 			if (!GraphInstance)
