Fix NoPCH compilation

diff --git a/Engine/Plugins/Experimental/PCG/Source/PCGCompute/Public/PCGComputeModule.h b/Engine/Plugins/Experimental/PCG/Source/PCGCompute/Public/PCGComputeModule.h
--- a/Engine/Plugins/Experimental/PCG/Source/PCGCompute/Public/PCGComputeModule.h
+++ b/Engine/Plugins/Experimental/PCG/Source/PCGCompute/Public/PCGComputeModule.h
@@ -2,7 +2,7 @@
 
 #pragma once
 
-#include "Stats/Stats2.h"
+#include "Stats/Stats.h"
 
 // Logs
 PCGCOMPUTE_API DECLARE_LOG_CATEGORY_EXTERN(LogPCGCompute, Log, All);
