[PCG] New "Branch" node implemented, similar to the native blueprint version.

The branch node will pass along one of two inputs, based on a user defined or overridden flag.

+ New node type added: ControlFlow
+ New node added: Branch

diff --git a/Engine/Plugins/Experimental/PCG/Source/PCG/Private/Elements/PCGBranch.cpp b/Engine/Plugins/Experimental/PCG/Source/PCG/Private/Elements/PCGBranch.cpp
new file mode 100644
--- /dev/null
+++ b/Engine/Plugins/Experimental/PCG/Source/PCG/Private/Elements/PCGBranch.cpp
@@ -0,0 +1,91 @@
+// Copyright Epic Games, Inc. All Rights Reserved.
+
+#include "Elements/PCGBranch.h"
+
+#include "PCGCommon.h"
+#include "PCGContext.h"
+#include "PCGPin.h"
+#include "Elements/PCGGather.h"
+
+#define LOCTEXT_NAMESPACE "FPCGBranchElement"
+
+namespace PCGBranchConstants
+{
+	const FName InputLabelA = TEXT("Input A");
+	const FName InputLabelB = TEXT("Input B");
+	const FText NodeTitleBase = LOCTEXT("NodeTitle", "Branch");
+}
+
+#if WITH_EDITOR
+FText UPCGBranchSettings::GetDefaultNodeTitle() const
+{
+	// TODO: This should statically update or dynamically update, if overridden, for which branch was taken, ie. Branch (A)
+	return PCGBranchConstants::NodeTitleBase;
+}
+
+FText UPCGBranchSettings::GetNodeTooltipText() const
+{
+	return LOCTEXT("NodeTooltip", "Control flow node that will allow all input data on either Pin A or Pin B only, based on the 'Use Input B' property - which can also be overridden.");
+}
+
+EPCGDataType UPCGBranchSettings::GetCurrentPinTypes(const UPCGPin* InPin) const
+{
+	check(InPin);
+	if (!InPin->IsOutputPin())
+	{
+		return Super::GetCurrentPinTypes(InPin);
+	}
+
+	// Output pin narrows to union of inputs on both
+	const EPCGDataType InputTypeUnion = GetTypeUnionOfIncidentEdges(PCGBranchConstants::InputLabelA) | GetTypeUnionOfIncidentEdges(PCGBranchConstants::InputLabelB);
+	return InputTypeUnion != EPCGDataType::None ? InputTypeUnion : EPCGDataType::Any;
+}
+#endif // WITH_EDITOR
+
+
+TArray<FPCGPinProperties> UPCGBranchSettings::InputPinProperties() const
+{
+	TArray<FPCGPinProperties> PinProperties;
+	PinProperties.Emplace(PCGBranchConstants::InputLabelA,
+		EPCGDataType::Any,
+		/*bInAllowMultipleConnections=*/true,
+		/*bAllowMultipleData=*/true,
+		LOCTEXT("FirstInputPinTooltip", "Will only be used if 'Use Input B' (overridable) is false"));
+	PinProperties.Emplace(PCGBranchConstants::InputLabelB,
+		EPCGDataType::Any,
+		/*bInAllowMultipleConnections=*/true,
+		/*bAllowMultipleData=*/true,
+		LOCTEXT("SecondInputPinTooltip", "Will only be used if 'Use Input B' (overridable) is true"));
+
+	return PinProperties;
+}
+
+TArray<FPCGPinProperties> UPCGBranchSettings::OutputPinProperties() const
+{
+	TArray<FPCGPinProperties> PinProperties;
+	PinProperties.Emplace(PCGPinConstants::DefaultOutputLabel, EPCGDataType::Any, /*bInAllowMultipleConnections=*/true, /*bAllowMultipleData=*/true, LOCTEXT("OutputPinTooltip", "All input will gathered into a single data collection"));
+
+	return PinProperties;
+}
+
+FPCGElementPtr UPCGBranchSettings::CreateElement() const
+{
+	return MakeShared<FPCGBranchElement>();
+}
+
+bool FPCGBranchElement::ExecuteInternal(FPCGContext* Context) const
+{
+	TRACE_CPUPROFILER_EVENT_SCOPE(FPCGBranchElement::ExecuteInternal);
+
+	const UPCGBranchSettings* Settings = Context->GetInputSettings<UPCGBranchSettings>();
+	check(Settings);
+
+	const FName SelectedPinLabel = Settings->bUseInputB ? PCGBranchConstants::InputLabelB : PCGBranchConstants::InputLabelA;
+
+	// Reuse the functionality of the Gather Node
+	Context->OutputData = PCGGather::GatherDataForPin(Context->InputData, SelectedPinLabel);
+
+	return true;
+}
+
+#undef LOCTEXT_NAMESPACE
\ No newline at end of file
diff --git a/Engine/Plugins/Experimental/PCG/Source/PCG/Private/Elements/PCGGather.cpp b/Engine/Plugins/Experimental/PCG/Source/PCG/Private/Elements/PCGGather.cpp
--- a/Engine/Plugins/Experimental/PCG/Source/PCG/Private/Elements/PCGGather.cpp
+++ b/Engine/Plugins/Experimental/PCG/Source/PCG/Private/Elements/PCGGather.cpp
@@ -44,24 +44,41 @@ FPCGElementPtr UPCGGatherSettings::CreateElement() const
 	return MakeShared<FPCGGatherElement>();
 }
 
-bool FPCGGatherElement::ExecuteInternal(FPCGContext* Context) const
+namespace PCGGather
 {
-	TRACE_CPUPROFILER_EVENT_SCOPE(FPCGGatherElement::Execute);
-
-	TArray<FPCGTaggedData> GatheredData = Context->InputData.GetInputsByPin(PCGPinConstants::DefaultInputLabel);
-	if (GatheredData.Num() == Context->InputData.TaggedData.Num())
-	{
-		Context->OutputData = Context->InputData;
-	}
-	else
+	FPCGDataCollection GatherDataForPin(const FPCGDataCollection& InputData, const FName InputLabel, const FName OutputLabel)
 	{
-		Context->OutputData.TaggedData = MoveTemp(GatheredData);
-	}
+		TArray<FPCGTaggedData> GatheredData = InputData.GetInputsByPin(InputLabel);
+		FPCGDataCollection Output;
 
-	for(FPCGTaggedData& TaggedData : Context->OutputData.TaggedData)
-	{
-		TaggedData.Pin = PCGPinConstants::DefaultOutputLabel;
+		if (GatheredData.IsEmpty())
+		{
+			return Output;
+		}
+	
+		if (GatheredData.Num() == InputData.TaggedData.Num())
+		{
+			Output = InputData;
+		}
+		else
+		{
+			Output.TaggedData = MoveTemp(GatheredData);
+		}
+
+		for(FPCGTaggedData& TaggedData : Output.TaggedData)
+		{
+			TaggedData.Pin = OutputLabel;
+		}
+
+		return Output;
 	}
+}
+
+bool FPCGGatherElement::ExecuteInternal(FPCGContext* Context) const
+{
+	TRACE_CPUPROFILER_EVENT_SCOPE(FPCGGatherElement::Execute);
+
+	Context->OutputData = PCGGather::GatherDataForPin(Context->InputData);
 
 	return true;
 }
diff --git a/Engine/Plugins/Experimental/PCG/Source/PCG/Public/Elements/PCGBranch.h b/Engine/Plugins/Experimental/PCG/Source/PCG/Public/Elements/PCGBranch.h
new file mode 100644
--- /dev/null
+++ b/Engine/Plugins/Experimental/PCG/Source/PCG/Public/Elements/PCGBranch.h
@@ -0,0 +1,43 @@
+// Copyright Epic Games, Inc. All Rights Reserved.
+#pragma once
+
+#include "PCGSettings.h"
+
+#include "PCGBranch.generated.h"
+
+/**
+ * Selects data from either input pin, based on a boolean condition.
+ */
+UCLASS(BlueprintType, ClassGroup = (Procedural), meta=(Keywords = "if bool branch"))
+class UPCGBranchSettings : public UPCGSettings
+{
+	GENERATED_BODY()
+
+public:
+	//~Begin UPCGSettings interface
+#if WITH_EDITOR
+	virtual FName GetDefaultNodeName() const override { return FName(TEXT("Branch")); }
+	virtual FText GetDefaultNodeTitle() const override;
+	virtual FText GetNodeTooltipText() const override;
+	virtual EPCGSettingsType GetType() const override { return EPCGSettingsType::ControlFlow; }
+	virtual bool HasDynamicPins() const override { return true; }
+#endif
+
+	virtual EPCGDataType GetCurrentPinTypes(const UPCGPin* InPin) const override;
+
+protected:
+	virtual TArray<FPCGPinProperties> InputPinProperties() const override;
+	virtual TArray<FPCGPinProperties> OutputPinProperties() const override;
+	virtual FPCGElementPtr CreateElement() const override;
+	//~End UPCGSettings interface
+
+public:
+	UPROPERTY(BlueprintReadWrite, EditAnywhere, Category=Settings, meta=(PCG_Overridable))
+	bool bUseInputB = false;
+};
+
+class FPCGBranchElement : public FSimplePCGElement
+{
+protected:
+	virtual bool ExecuteInternal(FPCGContext* Context) const override;
+};
diff --git a/Engine/Plugins/Experimental/PCG/Source/PCG/Public/Elements/PCGGather.h b/Engine/Plugins/Experimental/PCG/Source/PCG/Public/Elements/PCGGather.h
--- a/Engine/Plugins/Experimental/PCG/Source/PCG/Public/Elements/PCGGather.h
+++ b/Engine/Plugins/Experimental/PCG/Source/PCG/Public/Elements/PCGGather.h
@@ -37,6 +37,12 @@ protected:
 	virtual bool ExecuteInternal(FPCGContext* Context) const override;
 };
 
+namespace PCGGather
+{
+	/** Gathers the input data into a single data collection and updates the tags */
+	FPCGDataCollection GatherDataForPin(const FPCGDataCollection& InputData, const FName InputLabel = PCGPinConstants::DefaultInputLabel, const FName OutputLabel = PCGPinConstants::DefaultOutputLabel);
+}
+
 
 #if UE_ENABLE_INCLUDE_ORDER_DEPRECATED_IN_5_2
 #include "CoreMinimal.h"
diff --git a/Engine/Plugins/Experimental/PCG/Source/PCG/Public/PCGSettings.h b/Engine/Plugins/Experimental/PCG/Source/PCG/Public/PCGSettings.h
--- a/Engine/Plugins/Experimental/PCG/Source/PCG/Public/PCGSettings.h
+++ b/Engine/Plugins/Experimental/PCG/Source/PCG/Public/PCGSettings.h
@@ -47,6 +47,7 @@ enum class EPCGSettingsType : uint8
 	Generic,
 	Param,
 	HierarchicalGeneration,
+	ControlFlow
 };
 
 #if WITH_EDITOR
diff --git a/Engine/Plugins/Experimental/PCG/Source/PCGEditor/Private/PCGEditorSettings.cpp b/Engine/Plugins/Experimental/PCG/Source/PCGEditor/Private/PCGEditorSettings.cpp
--- a/Engine/Plugins/Experimental/PCG/Source/PCGEditor/Private/PCGEditorSettings.cpp
+++ b/Engine/Plugins/Experimental/PCG/Source/PCGEditor/Private/PCGEditorSettings.cpp
@@ -23,7 +23,9 @@ UPCGEditorSettings::UPCGEditorSettings(const FObjectInitializer& ObjectInitializ
 	SubgraphNodeColor = FLinearColor(1.0f, 0.05f, 0.05f);
 	ParamDataNodeColor = FLinearColor(1.0f, 0.38f, 0.02f);
 	DebugNodeColor = FLinearColor(1.0f, 0.0f, 1.0f);
-	HierarchicalGenerationNodeColor = FLinearColor(1.0f, 0.4f, 0.0f);
+	ControlFlowNodeColor = FLinearColor(0.66f, .6f, 0.15f);
+	// HiGen are also a subset of Control Flow
+	HierarchicalGenerationNodeColor = ControlFlowNodeColor;
 
 	DefaultPinColor = FLinearColor(0.29f, 0.29f, 0.29f);
 	SpatialDataPinColor = FLinearColor(1.0f, 1.0f, 1.0f);
@@ -49,64 +51,46 @@ FLinearColor UPCGEditorSettings::GetColor(UPCGSettings* Settings) const
 	{
 		return DefaultNodeColor;
 	}
+
 	// First: check if there's an override
-	else if (const FLinearColor* Override = OverrideNodeColorByClass.Find(Settings->GetClass()))
+	if (const FLinearColor* Override = OverrideNodeColorByClass.Find(Settings->GetClass()))
 	{
 		return *Override;
 	}
+
 	// Otherwise, check against the classes we know
-	else if (Settings->GetType() == EPCGSettingsType::InputOutput)
-	{
-		return InputOutputNodeColor;
-	}
-	else if (Settings->GetType() == EPCGSettingsType::Spatial)
-	{
-		return SetOperationNodeColor;
-	}
-	else if (Settings->GetType() == EPCGSettingsType::Density)
-	{
-		return DensityOperationNodeColor;
-	}
-	else if (Settings->GetType() == EPCGSettingsType::Blueprint)
-	{
-		return BlueprintNodeColor;
-	}
-	else if (Settings->GetType() == EPCGSettingsType::Metadata)
-	{
-		return MetadataNodeColor;
-	}
-	else if (Settings->GetType() == EPCGSettingsType::Filter)
-	{
-		return FilterNodeColor;
-	}
-	else if (Settings->GetType() == EPCGSettingsType::Sampler)
-	{
-		return SamplerNodeColor;
-	}
-	else if (Settings->GetType() == EPCGSettingsType::Spawner)
-	{
-		return SpawnerNodeColor;
-	}
-	else if (Settings->GetType() == EPCGSettingsType::Subgraph)
-	{
-		return SubgraphNodeColor;
-	}
-	else if (Settings->GetType() == EPCGSettingsType::Debug)
-	{
-		return DebugNodeColor;
-	}
-	else if (Settings->GetType() == EPCGSettingsType::Param)
-	{
-		return ParamDataNodeColor;
-	}
-	else if (Settings->GetType() == EPCGSettingsType::HierarchicalGeneration)
-	{
-		return HierarchicalGenerationNodeColor;
-	}
-	else
-	{
-		// Finally, we couldn't find any match, so return the default value
-		return DefaultNodeColor;
+	switch (Settings->GetType())
+	{
+		case EPCGSettingsType::InputOutput:
+			return InputOutputNodeColor;
+		case EPCGSettingsType::Spatial:
+			return SetOperationNodeColor;
+		case EPCGSettingsType::Density:
+			return DensityOperationNodeColor;
+		case EPCGSettingsType::Blueprint:
+			return BlueprintNodeColor;
+		case EPCGSettingsType::Metadata:
+			return MetadataNodeColor;
+		case EPCGSettingsType::Filter:
+			return FilterNodeColor;
+		case EPCGSettingsType::Sampler:
+			return SamplerNodeColor;
+		case EPCGSettingsType::Spawner:
+			return SpawnerNodeColor;
+		case EPCGSettingsType::Subgraph:
+			return SubgraphNodeColor;
+		case EPCGSettingsType::Debug:
+			return DebugNodeColor;
+		case EPCGSettingsType::Param:
+			return ParamDataNodeColor;
+		case EPCGSettingsType::HierarchicalGeneration:
+			return HierarchicalGenerationNodeColor;
+		case EPCGSettingsType::ControlFlow:
+			return ControlFlowNodeColor;
+		case EPCGSettingsType::Generic: // falls through
+		default:
+			// Finally, we couldn't find any match, so return the default value
+			return DefaultNodeColor;
 	}
 }
 
diff --git a/Engine/Plugins/Experimental/PCG/Source/PCGEditor/Public/PCGEditorSettings.h b/Engine/Plugins/Experimental/PCG/Source/PCGEditor/Public/PCGEditorSettings.h
--- a/Engine/Plugins/Experimental/PCG/Source/PCGEditor/Public/PCGEditorSettings.h
+++ b/Engine/Plugins/Experimental/PCG/Source/PCGEditor/Public/PCGEditorSettings.h
@@ -77,6 +77,10 @@ public:
 	UPROPERTY(EditAnywhere, config, Category = Node, meta = (HideAlphaChannel))
 	FLinearColor DebugNodeColor;
 
+	/** Color used for control flow operations */
+	UPROPERTY(EditAnywhere, config, Category = Node, meta = (HideAlphaChannel))
+	FLinearColor ControlFlowNodeColor;
+
 	/** Color used for hierarchical generation operations */
 	UPROPERTY(EditAnywhere, config, Category = Node, meta = (HideAlphaChannel))
 	FLinearColor HierarchicalGenerationNodeColor;
