[PCG] Moving away from WeakPtr map keys in tracking + move unregister of the component earlier in the destroy process.

#rnx

diff --git a/Engine/Plugins/Experimental/PCG/Source/PCG/Private/PCGActorAndComponentMapping.cpp b/Engine/Plugins/Experimental/PCG/Source/PCG/Private/PCGActorAndComponentMapping.cpp
--- a/Engine/Plugins/Experimental/PCG/Source/PCG/Private/PCGActorAndComponentMapping.cpp
+++ b/Engine/Plugins/Experimental/PCG/Source/PCG/Private/PCGActorAndComponentMapping.cpp
@@ -707,7 +707,7 @@ void UPCGActorAndComponentMapping::UnregisterTracking(UPCGComponent* InComponent
 		return;
 	}
 
-	TSet<TWeakObjectPtr<AActor>> CandidatesForUntrack;
+	TSet<TObjectKey<AActor>> CandidatesForUntrack;
 	TSet<FPCGActorSelectionKey> KeysToRemove;
 
 	auto RemoveFromMap = [InComponent](auto& InMap, auto& InCandidateToRemove)
@@ -732,17 +732,17 @@ void UPCGActorAndComponentMapping::UnregisterTracking(UPCGComponent* InComponent
 	}
 
 	// We also need to untrack actors that doesn't have any component that tracks them.
-	auto ShouldBeRemoved = [](const TWeakObjectPtr<AActor>& InActor, TMap<TWeakObjectPtr<AActor>, TSet<TObjectPtr<UPCGComponent>>>& InMap)
+	auto ShouldBeRemoved = [](const TObjectKey<AActor>& InActor, TMap<TObjectKey<AActor>, TSet<TObjectPtr<UPCGComponent>>>& InMap)
 	{
 		TSet<TObjectPtr<UPCGComponent>>* RegisteredComponents = InMap.Find(InActor);
 		return !RegisteredComponents || RegisteredComponents->IsEmpty();
 	};
 
-	for (TWeakObjectPtr<AActor> Candidate : CandidatesForUntrack)
+	for (TObjectKey<AActor> Candidate : CandidatesForUntrack)
 	{
 		if (ShouldBeRemoved(Candidate, CulledTrackedActorsToComponentsMap) && ShouldBeRemoved(Candidate, AlwaysTrackedActorsToComponentsMap))
 		{
-			UnregisterActor(Candidate.Get());
+			UnregisterActor(Candidate.ResolveObjectPtr());
 		}
 	}
 
@@ -783,9 +783,9 @@ void UPCGActorAndComponentMapping::AddActorsPostInit()
 		return;
 	}
 
-	for (TWeakObjectPtr<AActor>& ActorPtr : DelayedAddedActors)
+	for (TObjectKey<AActor>& ActorPtr : DelayedAddedActors)
 	{
-		OnActorAdded(ActorPtr.Get());
+		OnActorAdded(ActorPtr.ResolveObjectPtr());
 	}
 
 	DelayedAddedActors.Empty();
@@ -1024,11 +1024,11 @@ void UPCGActorAndComponentMapping::OnObjectPropertyChanged(UObject* InObject, FP
 	// If we don't find any actor, try to see if it is a dependency
 	if (!Actor)
 	{
-		for (const TPair<TWeakObjectPtr<AActor>, TSet<TObjectPtr<UObject>>>& TrackedActor : TrackedActorsToDependenciesMap)
+		for (const TPair<TObjectKey<AActor>, TSet<TObjectPtr<UObject>>>& TrackedActor : TrackedActorsToDependenciesMap)
 		{
 			if (TrackedActor.Value.Contains(InObject))
 			{
-				if (AActor* ActorToChange = TrackedActor.Key.Get())
+				if (AActor* ActorToChange = TrackedActor.Key.ResolveObjectPtr())
 				{
 					OnActorChanged(ActorToChange, /*bInHasMoved=*/ false);
 					UpdateActorDependencies(ActorToChange);
diff --git a/Engine/Plugins/Experimental/PCG/Source/PCG/Private/PCGComponent.cpp b/Engine/Plugins/Experimental/PCG/Source/PCG/Private/PCGComponent.cpp
--- a/Engine/Plugins/Experimental/PCG/Source/PCG/Private/PCGComponent.cpp
+++ b/Engine/Plugins/Experimental/PCG/Source/PCG/Private/PCGComponent.cpp
@@ -935,11 +935,9 @@ void UPCGComponent::EndPlay(const EEndPlayReason::Type EndPlayReason)
 	Super::EndPlay(EndPlayReason);
 }
 
-void UPCGComponent::OnComponentDestroyed(bool bDestroyingHierarchy)
+void UPCGComponent::OnUnregister()
 {
 #if WITH_EDITOR
-	// This is inspired by UChildActorComponent::DestroyChildActor()
-	// In the case of level change or exit, the subsystem will be null
 	if (UPCGSubsystem* Subsystem = GetSubsystem())
 	{
 		if (!PCGHelpers::IsRuntimeOrPIE())
@@ -948,7 +946,14 @@ void UPCGComponent::OnComponentDestroyed(bool bDestroyingHierarchy)
 			Subsystem->UnregisterPCGComponent(this);
 		}
 	}
+#endif // WITH_EDITOR
+
+	Super::OnUnregister();
+}
 
+void UPCGComponent::OnComponentDestroyed(bool bDestroyingHierarchy)
+{
+#if WITH_EDITOR
 	// BeginDestroy is not called immediately when a component is destroyed. Therefore callbacks are not cleaned
 	// until GC is ran, and can stack up with BP reconstruction scripts. Force the removal of callbacks here. If the component
 	// is dead, we don't want to react to callbacks anyway.
diff --git a/Engine/Plugins/Experimental/PCG/Source/PCG/Public/PCGActorAndComponentMapping.h b/Engine/Plugins/Experimental/PCG/Source/PCG/Public/PCGActorAndComponentMapping.h
--- a/Engine/Plugins/Experimental/PCG/Source/PCG/Public/PCGActorAndComponentMapping.h
+++ b/Engine/Plugins/Experimental/PCG/Source/PCG/Public/PCGActorAndComponentMapping.h
@@ -10,6 +10,7 @@
 #include "Containers/Map.h"
 #include "Containers/Set.h"
 #include "Templates/Function.h"
+#include "UObject/ObjectKey.h"
 #include "UObject/ObjectPtr.h"
 #include "UObject/SoftObjectPtr.h"
 #include "UObject/WeakObjectPtr.h"
@@ -170,23 +171,23 @@ private:
 	FPCGComponentOctreeAndMap NonPartitionedOctree;
 
 	/** Keep a mapping between tracked actors and the components that track them, and the tracking needs to be culled.*/
-	TMap<TWeakObjectPtr<AActor>, TSet<TObjectPtr<UPCGComponent>>> CulledTrackedActorsToComponentsMap;
+	TMap<TObjectKey<AActor>, TSet<TObjectPtr<UPCGComponent>>> CulledTrackedActorsToComponentsMap;
 
 	/** Same mapping but for always tracked actors */
-	TMap<TWeakObjectPtr<AActor>, TSet<TObjectPtr<UPCGComponent>>> AlwaysTrackedActorsToComponentsMap;
+	TMap<TObjectKey<AActor>, TSet<TObjectPtr<UPCGComponent>>> AlwaysTrackedActorsToComponentsMap;
 
 	/** Keep the list of the keys already tracked, to avoid requerying the actors everytime */
 	TMap<FPCGActorSelectionKey, TSet<TObjectPtr<UPCGComponent>>> KeysToComponentsMap;
 
 	/** Finally keep a mapping between actors and their position to know if an actor move in/out of a component tracking bounds. Only kept for actors that need to be culled. */
-	TMap<TWeakObjectPtr<AActor>, FBox> TrackedActorToPositionMap;
+	TMap<TObjectKey<AActor>, FBox> TrackedActorToPositionMap;
 
 	mutable FRWLock TrackedComponentsLock;
 
-	TSet<TWeakObjectPtr<AActor>> DelayedAddedActors;
+	TSet<TObjectKey<AActor>> DelayedAddedActors;
 
 	/** Keep a mapping between tracked actors and their dependencies. */
-	TMap<TWeakObjectPtr<AActor>, TSet<TObjectPtr<UObject>>> TrackedActorsToDependenciesMap;
+	TMap<TObjectKey<AActor>, TSet<TObjectPtr<UObject>>> TrackedActorsToDependenciesMap;
 
 	/** Transient list of tags, kept when there is a tag change on a tracked Actor. */
 	TSet<FName> TempTrackedActorTags;
diff --git a/Engine/Plugins/Experimental/PCG/Source/PCG/Public/PCGComponent.h b/Engine/Plugins/Experimental/PCG/Source/PCG/Public/PCGComponent.h
--- a/Engine/Plugins/Experimental/PCG/Source/PCG/Public/PCGComponent.h
+++ b/Engine/Plugins/Experimental/PCG/Source/PCG/Public/PCGComponent.h
@@ -91,6 +91,7 @@ public:
 	virtual void EndPlay(const EEndPlayReason::Type EndPlayReason) override;
 	virtual void OnComponentDestroyed(bool bDestroyingHierarchy) override;
 	virtual void OnRegister() override;
+	virtual void OnUnregister() override;
 
 protected:
 	friend struct FPCGComponentInstanceData;
