[PCG] Add Make Concrete conversion type for going to Concrete, performs a collapse of non-concrete Spatial data.

Added tests for pin type connectivity.

diff --git a/Engine/Plugins/Experimental/PCG/Source/PCG/Private/Elements/PCGMakeConcreteElement.cpp b/Engine/Plugins/Experimental/PCG/Source/PCG/Private/Elements/PCGMakeConcreteElement.cpp
new file mode 100644
--- /dev/null
+++ b/Engine/Plugins/Experimental/PCG/Source/PCG/Private/Elements/PCGMakeConcreteElement.cpp
@@ -0,0 +1,68 @@
+// Copyright Epic Games, Inc. All Rights Reserved.
+
+#include "Elements/PCGMakeConcreteElement.h"
+
+#include "PCGContext.h"
+#include "Data/PCGPointData.h"
+#include "Data/PCGSpatialData.h"
+
+#define LOCTEXT_NAMESPACE "PCGMakeConcreteElement"
+
+#if WITH_EDITOR
+FText UPCGMakeConcreteSettings::GetDefaultNodeTitle() const
+{
+	return LOCTEXT("NodeTitle", "Make Concrete");
+}
+
+FText UPCGMakeConcreteSettings::GetNodeTooltipText() const
+{
+	return LOCTEXT("NodeTooltip",
+		"Concrete data is passed through (e.g. Point, Curve, Landscape). "
+		"Spatial data (e.g. Intersection, Difference, Union) is collapsed to Point. "
+		"Non-Spatial data (e.g. Attribute Set) is discarded.");
+}
+#endif // WITH_EDITOR
+
+TArray<FPCGPinProperties> UPCGMakeConcreteSettings::OutputPinProperties() const
+{
+	TArray<FPCGPinProperties> PinProperties;
+	PinProperties.Emplace(PCGPinConstants::DefaultOutputLabel, EPCGDataType::Concrete);
+
+	return PinProperties;
+}
+
+FPCGElementPtr UPCGMakeConcreteSettings::CreateElement() const
+{
+	return MakeShared<FPCGMakeConcreteElement>();
+}
+
+bool FPCGMakeConcreteElement::ExecuteInternal(FPCGContext* InContext) const
+{
+	TRACE_CPUPROFILER_EVENT_SCOPE(FPCGMakeConcreteElement::Execute);
+
+	TArray<FPCGTaggedData> SpatialInputs = InContext->InputData.GetInputs();
+	for (const FPCGTaggedData& Input : SpatialInputs)
+	{
+		const UPCGSpatialData* InputSpatialData = Cast<UPCGSpatialData>(Input.Data);
+		// Verify data really Spatial and reported typed of data also within Spatial.
+		if (!ensure(InputSpatialData) || !ensure(!(Input.Data->GetDataType() & ~EPCGDataType::Spatial)))
+		{
+			continue;
+		}
+		
+		// Determine if we have data that needs to be made concrete.
+		const bool bDataIsConcrete = !(Input.Data->GetDataType() & ~EPCGDataType::Concrete);
+
+		// Add output, doing a to-point to make concrete if required.
+		FPCGTaggedData& Output = InContext->OutputData.TaggedData.Emplace_GetRef(Input);
+		Output.Pin = PCGPinConstants::DefaultOutputLabel;
+		if (!bDataIsConcrete)
+		{
+			Output.Data = InputSpatialData->ToPointData(InContext);
+		}
+	}
+
+	return true;
+}
+
+#undef LOCTEXT_NAMESPACE
diff --git a/Engine/Plugins/Experimental/PCG/Source/PCG/Private/PCGPin.cpp b/Engine/Plugins/Experimental/PCG/Source/PCG/Private/PCGPin.cpp
--- a/Engine/Plugins/Experimental/PCG/Source/PCG/Private/PCGPin.cpp
+++ b/Engine/Plugins/Experimental/PCG/Source/PCG/Private/PCGPin.cpp
@@ -330,82 +330,55 @@ bool UPCGPin::CanConnect(const UPCGPin* OtherPin) const
 	return OtherPin && (Edges.IsEmpty() || AllowMultipleConnections());
 }
 
-bool UPCGPin::RequiresPointConversionToConnect(const UPCGPin* OtherPin) const
+EPCGTypeConversion UPCGPin::GetRequiredTypeConversion(const UPCGPin* InOtherPin) const
 {
-	if (!OtherPin)
-	{
-		return false;
-	}
-
-	const UPCGPin* UpstreamPin = nullptr;
-	const UPCGPin* DownstreamPin = nullptr;
-	EPCGDataType UpstreamTypes = EPCGDataType::None, DownstreamTypes = EPCGDataType::None;
-	if (!PCGPin::SortPinsAndRetrieveTypes(this, OtherPin, UpstreamPin, DownstreamPin, UpstreamTypes, DownstreamTypes))
-	{
-		return false;
-	}
-	check(UpstreamPin && DownstreamPin);
-
-	// This conversion requires downstream Point
-	if (DownstreamTypes != EPCGDataType::Point)
-	{
-		return false;
-	}
-
-	// Only collapse from Spatial to Point
-	const bool bUpstreamInSpatial = !(UpstreamTypes & ~EPCGDataType::Spatial);
-	if (!bUpstreamInSpatial)
+	if (!InOtherPin)
 	{
-		return false;
-	}
-
-	// Requires conversion if upstream pin is broader than point
-	return !!(UpstreamTypes & (~EPCGDataType::Point));
-}
-
-bool UPCGPin::RequiresFilterToConnect(const UPCGPin* OtherPin) const
-{
-	if (!OtherPin)
-	{
-		return false;
+		return EPCGTypeConversion::Failed;
 	}
 
 	const UPCGPin* UpstreamPin = nullptr;
 	const UPCGPin* DownstreamPin = nullptr;
-	EPCGDataType UpstreamTypes = EPCGDataType::None, DownstreamTypes = EPCGDataType::None;
-	if (!PCGPin::SortPinsAndRetrieveTypes(this, OtherPin, UpstreamPin, DownstreamPin, UpstreamTypes, DownstreamTypes))
+	EPCGDataType UpstreamTypes = EPCGDataType::None;
+	EPCGDataType DownstreamTypes = EPCGDataType::None;
+	if (!PCGPin::SortPinsAndRetrieveTypes(this, InOtherPin, UpstreamPin, DownstreamPin, UpstreamTypes, DownstreamTypes))
 	{
-		return false;
+		return EPCGTypeConversion::Failed;
 	}
 	check(UpstreamPin && DownstreamPin);
 
-	// Same types - no filter
+	// Types same - trivial early out.
 	if (DownstreamTypes == UpstreamTypes)
 	{
-		return false;
+		return EPCGTypeConversion::NoConversionRequired;
 	}
 
-	// No type - early out
-	if (!DownstreamTypes)
+	// Type missing, or types are the same - trivial early out.
+	if (!DownstreamTypes || !UpstreamTypes)
 	{
-		return false;
+		return EPCGTypeConversion::Failed;
 	}
 
-	// Spatial -> Point - that uses conversion, not filter
+	// Spatial -> Point - "To Point" conversion.
 	const bool bUpstreamInSpatial = !(UpstreamTypes & ~EPCGDataType::Spatial);
 	if (bUpstreamInSpatial && DownstreamTypes == EPCGDataType::Point)
 	{
-		return false;
+		return EPCGTypeConversion::CollapseToPoint;
 	}
 
-	// Upstream type should be Any, Spatial, Concrete
-	if (UpstreamTypes != EPCGDataType::Any
-		&& UpstreamTypes != EPCGDataType::Spatial
-		&& UpstreamTypes != EPCGDataType::Concrete)
+	// Any or Spatial -> Concrete - "Make Concrete" conversion. We decided to support Any as it is a superset of Concrete,
+	// and some pins are Any until they dynamically narrow, and other pins may stay Any but still want to be supported.
+	bool bUpstreamIsSpatialOrAny = (UpstreamTypes == EPCGDataType::Any) || (UpstreamTypes == EPCGDataType::Spatial);
+	if (bUpstreamIsSpatialOrAny && DownstreamTypes == EPCGDataType::Concrete)
 	{
-		return false;
+		return EPCGTypeConversion::MakeConcrete;
+	}
+
+	// Requires filter if upstream type is broader than downstream type.
+	if (!!(UpstreamTypes & (~DownstreamTypes)))
+	{
+		return EPCGTypeConversion::Filter;
 	}
 
-	// Requires conversion if upstream type is broader than downstream type
-	return !!(UpstreamTypes & (~DownstreamTypes));
+	return EPCGTypeConversion::NoConversionRequired;
 }
diff --git a/Engine/Plugins/Experimental/PCG/Source/PCG/Private/Tests/GraphAuthoring/PCGGraphAuthoringTest.cpp b/Engine/Plugins/Experimental/PCG/Source/PCG/Private/Tests/GraphAuthoring/PCGGraphAuthoringTest.cpp
new file mode 100644
--- /dev/null
+++ b/Engine/Plugins/Experimental/PCG/Source/PCG/Private/Tests/GraphAuthoring/PCGGraphAuthoringTest.cpp
@@ -0,0 +1,123 @@
+// Copyright Epic Games, Inc. All Rights Reserved.
+
+#if WITH_EDITOR
+
+#include "PCGComponent.h"
+#include "PCGGraph.h"
+#include "PCGNode.h"
+#include "PCGPin.h"
+#include "PCGSettings.h"
+#include "Tests/PCGTestsCommon.h"
+#include "Tests/GraphAuthoring/PCGGraphAuthoringTestHelperSettings.h"
+
+IMPLEMENT_CUSTOM_SIMPLE_AUTOMATION_TEST(FPCGGraphAuthoringTypeSystemTest, FPCGTestBaseClass, "pcg.GraphAuthoring.TypeSystem", PCGTestsCommon::TestFlags)
+
+bool FPCGGraphAuthoringTypeSystemTest::RunTest(const FString& Parameters)
+{
+	PCGTestsCommon::FTestData TestData(0);
+	UPCGGraph* Graph = TestData.TestPCGComponent->GetGraph();
+
+	UPCGSettings* UpstreamNodeUntypedSettings = nullptr;
+	UPCGSettings* DownstreamNodeUntypedSettings = nullptr;
+	UPCGNode* UpstreamNode = Graph->AddNodeOfType(UPCGGraphAuthoringTestHelperSettings::StaticClass(), UpstreamNodeUntypedSettings);
+	UPCGNode* DownstreamNode = Graph->AddNodeOfType(UPCGGraphAuthoringTestHelperSettings::StaticClass(), DownstreamNodeUntypedSettings);
+	UPCGGraphAuthoringTestHelperSettings* UpstreamNodeSettings = CastChecked<UPCGGraphAuthoringTestHelperSettings>(UpstreamNodeUntypedSettings);
+	UPCGGraphAuthoringTestHelperSettings* DownstreamNodeSettings = CastChecked<UPCGGraphAuthoringTestHelperSettings>(DownstreamNodeUntypedSettings);
+
+	auto ValidateConversion = [&](
+		const FString& InWhat,
+		EPCGDataType InUpstreamType,
+		EPCGDataType InDownstreamType,
+		bool bInExpectedCompatible,
+		EPCGTypeConversion InExpectedConversion = EPCGTypeConversion::NoConversionRequired)
+	{
+		UpstreamNodeSettings->PinType = InUpstreamType;
+		UpstreamNode->SetSettingsInterface(UpstreamNodeSettings);
+
+		DownstreamNodeSettings->PinType = InDownstreamType;
+		DownstreamNode->SetSettingsInterface(DownstreamNodeSettings);
+
+		const UPCGPin* UpstreamPin = UpstreamNode->GetOutputPin(PCGPinConstants::DefaultOutputLabel);
+		const UPCGPin* DownstreamPin = DownstreamNode->GetInputPin(PCGPinConstants::DefaultInputLabel);
+		check(UpstreamPin && DownstreamPin);
+
+		const bool bCompatible = UpstreamPin->IsCompatible(DownstreamPin);
+		bool bValidationPassed = TestEqual(InWhat, bCompatible, bInExpectedCompatible);
+
+		// If compatibility test was successful, and if types are compatible, check their conversion.
+		if (bCompatible && bInExpectedCompatible)
+		{
+			bValidationPassed &= TestEqual(InWhat, UpstreamPin->GetRequiredTypeConversion(DownstreamPin), InExpectedConversion);
+		}
+
+		return bValidationPassed;
+	};
+
+	bool bTestPassed = true;
+
+	// Trivial
+	bTestPassed &= ValidateConversion("None to none", EPCGDataType::None, EPCGDataType::None, false);
+	bTestPassed &= ValidateConversion("Point to none", EPCGDataType::Point, EPCGDataType::None, false);
+	bTestPassed &= ValidateConversion("None to point", EPCGDataType::None, EPCGDataType::Point, false);
+
+	// No conversion needed
+	bTestPassed &= ValidateConversion("Matching types: Point to Point", EPCGDataType::Point, EPCGDataType::Point, true);
+	bTestPassed &= ValidateConversion("Matching types: Landscape to Landscape", EPCGDataType::Landscape, EPCGDataType::Landscape, true);
+	bTestPassed &= ValidateConversion("Matching types: Surface to Surface", EPCGDataType::Surface, EPCGDataType::Surface, true);
+	bTestPassed &= ValidateConversion("Matching types: Concrete to Concrete", EPCGDataType::Concrete, EPCGDataType::Concrete, true);
+	bTestPassed &= ValidateConversion("Matching types: Spatial to Spatial", EPCGDataType::Spatial, EPCGDataType::Spatial, true);
+	bTestPassed &= ValidateConversion("Matching types: Other to Other", EPCGDataType::Other, EPCGDataType::Other, true);
+	bTestPassed &= ValidateConversion("Matching types: Attribute Set to Attribute Set", EPCGDataType::Param, EPCGDataType::Param, true);
+	bTestPassed &= ValidateConversion("Matching types: Any to Any", EPCGDataType::Any, EPCGDataType::Any, true);
+	bTestPassed &= ValidateConversion("Compatible types: Landscape to Surface", EPCGDataType::Landscape, EPCGDataType::Surface, true);
+	bTestPassed &= ValidateConversion("Compatible types: Point to Concrete", EPCGDataType::Point, EPCGDataType::Concrete, true);
+	bTestPassed &= ValidateConversion("Compatible types: Landscape to Concrete", EPCGDataType::Landscape, EPCGDataType::Concrete, true);
+	bTestPassed &= ValidateConversion("Compatible types: Surface to Concrete", EPCGDataType::Surface, EPCGDataType::Concrete, true);
+	bTestPassed &= ValidateConversion("Compatible types: Point to Spatial", EPCGDataType::Point, EPCGDataType::Spatial, true);
+	bTestPassed &= ValidateConversion("Compatible types: Landscape to Spatial", EPCGDataType::Landscape, EPCGDataType::Spatial, true);
+	bTestPassed &= ValidateConversion("Compatible types: Surface to Spatial", EPCGDataType::Surface, EPCGDataType::Spatial, true);
+	bTestPassed &= ValidateConversion("Compatible types: Concrete to Spatial", EPCGDataType::Concrete, EPCGDataType::Spatial, true);
+	bTestPassed &= ValidateConversion("Compatible types: Point to Any", EPCGDataType::Point, EPCGDataType::Any, true);
+	bTestPassed &= ValidateConversion("Compatible types: Landscape to Any", EPCGDataType::Landscape, EPCGDataType::Any, true);
+	bTestPassed &= ValidateConversion("Compatible types: Surface to Any", EPCGDataType::Surface, EPCGDataType::Any, true);
+	bTestPassed &= ValidateConversion("Compatible types: Concrete to Any", EPCGDataType::Concrete, EPCGDataType::Any, true);
+	bTestPassed &= ValidateConversion("Compatible types: Spatial to Any", EPCGDataType::Spatial, EPCGDataType::Any, true);
+	bTestPassed &= ValidateConversion("Compatible types: Attribute Set to Any", EPCGDataType::Param, EPCGDataType::Any, true);
+	bTestPassed &= ValidateConversion("Compatible types: Other to Any", EPCGDataType::Other, EPCGDataType::Any, true);
+
+	// Collapse
+	bTestPassed &= ValidateConversion("Collapse: Spline to Point", EPCGDataType::Spline, EPCGDataType::Point, true, EPCGTypeConversion::CollapseToPoint);
+	bTestPassed &= ValidateConversion("Collapse: Landscape to Point", EPCGDataType::Landscape, EPCGDataType::Point, true, EPCGTypeConversion::CollapseToPoint);
+	bTestPassed &= ValidateConversion("Collapse: Surface to Point", EPCGDataType::Surface, EPCGDataType::Point, true, EPCGTypeConversion::CollapseToPoint);
+	bTestPassed &= ValidateConversion("Collapse: Concrete to Point", EPCGDataType::Concrete, EPCGDataType::Point, true, EPCGTypeConversion::CollapseToPoint);
+	bTestPassed &= ValidateConversion("Collapse: Spatial to Point", EPCGDataType::Spatial, EPCGDataType::Point, true, EPCGTypeConversion::CollapseToPoint);
+
+	// Filter
+	bTestPassed &= ValidateConversion("Filter: Any to Point", EPCGDataType::Any, EPCGDataType::Point, true, EPCGTypeConversion::Filter);
+	bTestPassed &= ValidateConversion("Filter: Any to Landscape", EPCGDataType::Any, EPCGDataType::Landscape, true, EPCGTypeConversion::Filter);
+	bTestPassed &= ValidateConversion("Filter: Any to Surface", EPCGDataType::Any, EPCGDataType::Surface, true, EPCGTypeConversion::Filter);
+	bTestPassed &= ValidateConversion("Filter: Any to Spatial", EPCGDataType::Any, EPCGDataType::Spatial, true, EPCGTypeConversion::Filter);
+	bTestPassed &= ValidateConversion("Filter: Any to Attribute Set", EPCGDataType::Any, EPCGDataType::Param, true, EPCGTypeConversion::Filter);
+	bTestPassed &= ValidateConversion("Filter: Spatial to Surface", EPCGDataType::Spatial, EPCGDataType::Surface, true, EPCGTypeConversion::Filter);
+	bTestPassed &= ValidateConversion("Filter: Spatial to Landscape", EPCGDataType::Spatial, EPCGDataType::Landscape, true, EPCGTypeConversion::Filter);
+
+	// Make Concrete
+	bTestPassed &= ValidateConversion("Make Concrete: Spatial to Concrete", EPCGDataType::Spatial, EPCGDataType::Concrete, true, EPCGTypeConversion::MakeConcrete);
+	bTestPassed &= ValidateConversion("Make Concrete: Any to Concrete", EPCGDataType::Any, EPCGDataType::Concrete, true, EPCGTypeConversion::MakeConcrete);
+
+	// Incompatible
+	bTestPassed &= ValidateConversion("Incompatible: Attribute Set to Point", EPCGDataType::Param, EPCGDataType::Point, false);
+	bTestPassed &= ValidateConversion("Incompatible: Attribute Set to Landscape", EPCGDataType::Param, EPCGDataType::Landscape, false);
+	bTestPassed &= ValidateConversion("Incompatible: Attribute Set to Surface", EPCGDataType::Param, EPCGDataType::Surface, false);
+	bTestPassed &= ValidateConversion("Incompatible: Attribute Set to Concrete", EPCGDataType::Param, EPCGDataType::Concrete, false);
+	bTestPassed &= ValidateConversion("Incompatible: Attribute Set to Spatial", EPCGDataType::Param, EPCGDataType::Spatial, false);
+	bTestPassed &= ValidateConversion("Incompatible: Other to Point", EPCGDataType::Other, EPCGDataType::Point, false);
+	bTestPassed &= ValidateConversion("Incompatible: Other to Landscape", EPCGDataType::Other, EPCGDataType::Landscape, false);
+	bTestPassed &= ValidateConversion("Incompatible: Other to Surface", EPCGDataType::Other, EPCGDataType::Surface, false);
+	bTestPassed &= ValidateConversion("Incompatible: Other to Concrete", EPCGDataType::Other, EPCGDataType::Concrete, false);
+	bTestPassed &= ValidateConversion("Incompatible: Other to Spatial", EPCGDataType::Other, EPCGDataType::Spatial, false);
+
+	return bTestPassed;
+}
+
+#endif // WITH_EDITOR
diff --git a/Engine/Plugins/Experimental/PCG/Source/PCG/Private/Tests/GraphAuthoring/PCGGraphAuthoringTestHelperSettings.cpp b/Engine/Plugins/Experimental/PCG/Source/PCG/Private/Tests/GraphAuthoring/PCGGraphAuthoringTestHelperSettings.cpp
new file mode 100644
--- /dev/null
+++ b/Engine/Plugins/Experimental/PCG/Source/PCG/Private/Tests/GraphAuthoring/PCGGraphAuthoringTestHelperSettings.cpp
@@ -0,0 +1,19 @@
+// Copyright Epic Games, Inc. All Rights Reserved.
+
+#include "Tests/GraphAuthoring/PCGGraphAuthoringTestHelperSettings.h"
+
+TArray<FPCGPinProperties> UPCGGraphAuthoringTestHelperSettings::InputPinProperties() const
+{
+	TArray<FPCGPinProperties> PinProperties;
+	PinProperties.Emplace_GetRef(PCGPinConstants::DefaultInputLabel, PinType);
+
+	return PinProperties;
+}
+
+TArray<FPCGPinProperties> UPCGGraphAuthoringTestHelperSettings::OutputPinProperties() const
+{
+	TArray<FPCGPinProperties> PinProperties;
+	PinProperties.Emplace_GetRef(PCGPinConstants::DefaultOutputLabel, PinType);
+
+	return PinProperties;
+}
diff --git a/Engine/Plugins/Experimental/PCG/Source/PCG/Private/Tests/GraphAuthoring/PCGGraphAuthoringTestHelperSettings.h b/Engine/Plugins/Experimental/PCG/Source/PCG/Private/Tests/GraphAuthoring/PCGGraphAuthoringTestHelperSettings.h
new file mode 100644
--- /dev/null
+++ b/Engine/Plugins/Experimental/PCG/Source/PCG/Private/Tests/GraphAuthoring/PCGGraphAuthoringTestHelperSettings.h
@@ -0,0 +1,22 @@
+// Copyright Epic Games, Inc. All Rights Reserved.
+#pragma once
+
+#include "PCGSettings.h"
+
+#include "PCGPin.h"
+
+#include "PCGGraphAuthoringTestHelperSettings.generated.h"
+
+/** Testing helper - generates a node with a single input and output pin of the stipulated type. */
+UCLASS(NotBlueprintable, NotBlueprintType, ClassGroup = (Procedural))
+class UPCGGraphAuthoringTestHelperSettings : public UPCGSettings
+{
+	GENERATED_BODY()
+
+protected:
+	virtual TArray<FPCGPinProperties> InputPinProperties() const override;
+	virtual TArray<FPCGPinProperties> OutputPinProperties() const override;
+
+public:
+	EPCGDataType PinType;
+};
diff --git a/Engine/Plugins/Experimental/PCG/Source/PCG/Public/Elements/PCGMakeConcreteElement.h b/Engine/Plugins/Experimental/PCG/Source/PCG/Public/Elements/PCGMakeConcreteElement.h
new file mode 100644
--- /dev/null
+++ b/Engine/Plugins/Experimental/PCG/Source/PCG/Public/Elements/PCGMakeConcreteElement.h
@@ -0,0 +1,41 @@
+// Copyright Epic Games, Inc. All Rights Reserved.
+
+#pragma once
+
+#include "PCGSettings.h"
+
+#include "PCGPin.h"
+
+#include "PCGMakeConcreteElement.generated.h"
+
+/** Makes data Concrete, collapsing to Point if necessary. Discards non-Spatial data. */
+UCLASS(BlueprintType, ClassGroup = (Procedural))
+class PCG_API UPCGMakeConcreteSettings : public UPCGSettings
+{
+	GENERATED_BODY()
+
+public:
+	//~Begin UPCGSettings interface
+#if WITH_EDITOR
+	virtual FName GetDefaultNodeName() const override { return FName(TEXT("MakeConcrete")); }
+	virtual FText GetDefaultNodeTitle() const override;
+	virtual FText GetNodeTooltipText() const override;
+	virtual EPCGSettingsType GetType() const override { return EPCGSettingsType::Spatial; }
+	virtual bool ShouldDrawNodeCompact() const override { return true; }
+#endif
+
+protected:
+	virtual TArray<FPCGPinProperties> OutputPinProperties() const override;
+	virtual FPCGElementPtr CreateElement() const override;
+	//~End UPCGSettings interface
+};
+
+class FPCGMakeConcreteElement : public FSimplePCGElement
+{
+public:
+	// Might be sampling spline/landscape or other external data, worth computing a full CRC in case we can halt change propagation/re-executions
+	virtual bool ShouldComputeFullOutputDataCrc(FPCGContext* InContext) const override { return true; }
+
+protected:
+	virtual bool ExecuteInternal(FPCGContext* InContext) const override;
+};
diff --git a/Engine/Plugins/Experimental/PCG/Source/PCG/Public/PCGPin.h b/Engine/Plugins/Experimental/PCG/Source/PCG/Public/PCGPin.h
--- a/Engine/Plugins/Experimental/PCG/Source/PCG/Public/PCGPin.h
+++ b/Engine/Plugins/Experimental/PCG/Source/PCG/Public/PCGPin.h
@@ -43,6 +43,16 @@ struct PCG_API FPCGPinProperties
 #endif
 };
 
+UENUM()
+enum class EPCGTypeConversion : uint8
+{
+	NoConversionRequired,
+	CollapseToPoint,
+	Filter,
+	MakeConcrete,
+	Failed,
+};
+
 UCLASS()
 class PCG_API UPCGBlueprintPinHelpers : public UBlueprintFunctionLibrary
 {
@@ -160,11 +170,7 @@ public:
 	/** Returns the current pin types, which can either be the static types from the pin properties, or a dynamic type based on connected edges. */
 	EPCGDataType GetCurrentTypes() const;
 
-	/** Returns true if a collapse is required to convert from upstream type to downstream type */
-	bool RequiresPointConversionToConnect(const UPCGPin* OtherPin) const;
-
-	/** Returns true if a data filter is required to convert from upstream type to downstream type */
-	bool RequiresFilterToConnect(const UPCGPin* OtherPin) const;
+	EPCGTypeConversion GetRequiredTypeConversion(const UPCGPin* InOtherPin) const;
 };
 
 #if UE_ENABLE_INCLUDE_ORDER_DEPRECATED_IN_5_2
diff --git a/Engine/Plugins/Experimental/PCG/Source/PCGEditor/Private/PCGEditorGraphSchema.cpp b/Engine/Plugins/Experimental/PCG/Source/PCGEditor/Private/PCGEditorGraphSchema.cpp
--- a/Engine/Plugins/Experimental/PCG/Source/PCGEditor/Private/PCGEditorGraphSchema.cpp
+++ b/Engine/Plugins/Experimental/PCG/Source/PCGEditor/Private/PCGEditorGraphSchema.cpp
@@ -8,6 +8,7 @@
 #include "Elements/PCGCollapseElement.h"
 #include "Elements/PCGExecuteBlueprint.h"
 #include "Elements/PCGFilterByType.h"
+#include "Elements/PCGMakeConcreteElement.h"
 #include "Elements/PCGUserParameterGet.h"
 
 #include "PCGEditorCommon.h"
@@ -118,15 +119,19 @@ const FPinConnectionResponse UPCGEditorGraphSchema::CanCreateConnection(const UE
 		return FPinConnectionResponse(CONNECT_RESPONSE_DISALLOW, LOCTEXT("ConnectionTypesIncompatible", "Pins are incompatible"));
 	}
 
-	if (InputPin->RequiresPointConversionToConnect(OutputPin))
+	const EPCGTypeConversion RequiredTypeConversion = InputPin->GetRequiredTypeConversion(OutputPin);
+	if (RequiredTypeConversion == EPCGTypeConversion::CollapseToPoint)
 	{
 		return FPinConnectionResponse(CONNECT_RESPONSE_MAKE_WITH_CONVERSION_NODE, LOCTEXT("ConnectionConversionToPoint", "Convert to Point"));
 	}
-
-	if (InputPin->RequiresFilterToConnect(OutputPin))
+	else if (RequiredTypeConversion == EPCGTypeConversion::Filter)
 	{
 		return FPinConnectionResponse(CONNECT_RESPONSE_MAKE_WITH_CONVERSION_NODE, LOCTEXT("ConnectionUsingFilter", "Filter data to match type"));
 	}
+	else if (RequiredTypeConversion == EPCGTypeConversion::MakeConcrete)
+	{
+		return FPinConnectionResponse(CONNECT_RESPONSE_MAKE_WITH_CONVERSION_NODE, LOCTEXT("ConnectionUsingMakeConcrete", "Make data concrete"));
+	}
 
 	if (!InputPin->AllowMultipleConnections() && InputPin->EdgeCount() > 0)
 	{
@@ -239,14 +244,15 @@ bool UPCGEditorGraphSchema::TryCreateConnectionInternal(UEdGraphPin* InA, UEdGra
 		return bModifiedA || bModifiedB;
 	};
 
-	if (bAddConversionNodeIfNeeded && PCGPinA->RequiresPointConversionToConnect(PCGPinB))
+	const EPCGTypeConversion Conversion = bAddConversionNodeIfNeeded ? PCGPinA->GetRequiredTypeConversion(PCGPinB) : EPCGTypeConversion::NoConversionRequired;
+	if (Conversion == EPCGTypeConversion::CollapseToPoint)
 	{
 		UPCGSettings* NodeSettings = nullptr;
 		UPCGNode* ConversionPCGNode = PCGGraph->AddNodeOfType(UPCGCollapseSettings::StaticClass(), NodeSettings);
 
 		return ConnectViaIntermediate(ConversionPCGNode);
 	}
-	else if (bAddConversionNodeIfNeeded && PCGPinA->RequiresFilterToConnect(PCGPinB))
+	else if (Conversion == EPCGTypeConversion::Filter)
 	{
 		UPCGSettings* NodeSettings = nullptr;
 		UPCGNode* ConversionPCGNode = PCGGraph->AddNodeOfType(UPCGFilterByTypeSettings::StaticClass(), NodeSettings);
@@ -258,6 +264,13 @@ bool UPCGEditorGraphSchema::TryCreateConnectionInternal(UEdGraphPin* InA, UEdGra
 
 		return ConnectViaIntermediate(ConversionPCGNode);
 	}
+	else if (Conversion == EPCGTypeConversion::MakeConcrete)
+	{
+		UPCGSettings* NodeSettings = nullptr;
+		UPCGNode* ConversionPCGNode = PCGGraph->AddNodeOfType(UPCGMakeConcreteSettings::StaticClass(), NodeSettings);
+
+		return ConnectViaIntermediate(ConversionPCGNode);
+	}
 	else
 	{
 		const bool bModified = Super::TryCreateConnection(InA, InB);
