[PCG] Fix Linux compile error

#rnx

diff --git a/Engine/Plugins/Experimental/PCG/Source/PCG/Private/Tests/GraphAuthoring/PCGGraphAuthoringTestHelperSettings.h b/Engine/Plugins/Experimental/PCG/Source/PCG/Private/Tests/GraphAuthoring/PCGGraphAuthoringTestHelperSettings.h
--- a/Engine/Plugins/Experimental/PCG/Source/PCG/Private/Tests/GraphAuthoring/PCGGraphAuthoringTestHelperSettings.h
+++ b/Engine/Plugins/Experimental/PCG/Source/PCG/Private/Tests/GraphAuthoring/PCGGraphAuthoringTestHelperSettings.h
@@ -17,6 +17,9 @@ protected:
 	virtual TArray<FPCGPinProperties> InputPinProperties() const override;
 	virtual TArray<FPCGPinProperties> OutputPinProperties() const override;
 
+	// Does not support execution currently, testing purposes only.
+	virtual FPCGElementPtr CreateElement() const override { return nullptr; }
+
 public:
 	EPCGDataType PinType;
 };
