[PCG] Fixing Undo/Redo on Graph parameters override

#rnx

diff --git a/Engine/Plugins/Experimental/PCG/Source/PCG/Private/PCGGraph.cpp b/Engine/Plugins/Experimental/PCG/Source/PCG/Private/PCGGraph.cpp
--- a/Engine/Plugins/Experimental/PCG/Source/PCG/Private/PCGGraph.cpp
+++ b/Engine/Plugins/Experimental/PCG/Source/PCG/Private/PCGGraph.cpp
@@ -228,6 +228,10 @@ UPCGGraph::UPCGGraph(const FObjectInitializer& ObjectInitializer)
 	// but not when using a blueprint construct script.
 	//InputNode->ConnectTo(OutputNode);
 	//OutputNode->ConnectFrom(InputNode);
+
+	// Force the user parameters to have an empty property bag. It is necessary to catch the first
+	// add property into the undo/redo history.
+	UserParameters.MigrateToNewBagStruct(UPropertyBag::GetOrCreateFromDescs({}));
 }
 
 void UPCGGraph::PostLoad()
@@ -1243,6 +1247,9 @@ void UPCGGraphInstance::PostEditUndo()
 		Graph->OnGraphGridSizesChangedDelegate.AddUObject(this, &UPCGGraphInstance::OnGraphGridSizesChanged);
 		Graph->OnGraphParametersChangedDelegate.AddUObject(this, &UPCGGraphInstance::OnGraphParametersChanged);
 	}
+
+	// Since we don't know what happened, we need to notify any changes
+	NotifyGraphParametersChanged(EPCGGraphParameterEvent::GraphChanged, NAME_None);
 }
 
 void UPCGGraphInstance::OnGraphChanged(UPCGGraphInterface* InGraph, EPCGChangeType ChangeType)
@@ -1431,7 +1438,6 @@ void UPCGGraphInstance::UpdatePropertyOverride(const FProperty* InProperty, bool
 #if WITH_EDITOR
 		// If it is true, it means that the value has changed, so propagate the changes, in Editor
 		NotifyGraphParametersChanged(EPCGGraphParameterEvent::ValueModifiedLocally, InProperty->GetFName());
-		OnGraphChangedDelegate.Broadcast(this, EPCGChangeType::Settings);
 #endif // WITH_EDITOR
 	}
 }
@@ -1468,7 +1474,6 @@ void UPCGGraphInstance::ResetPropertyToDefault(const FProperty* InProperty)
 
 #if WITH_EDITOR
 	NotifyGraphParametersChanged(EPCGGraphParameterEvent::ValueModifiedLocally, InProperty->GetFName());
-	OnGraphChangedDelegate.Broadcast(this, EPCGChangeType::Settings);
 #endif // WITH_EDITOR
 }
 
diff --git a/Engine/Plugins/Experimental/PCG/Source/PCGEditor/Private/Details/PCGInstancedPropertyBagOverrideDetails.cpp b/Engine/Plugins/Experimental/PCG/Source/PCGEditor/Private/Details/PCGInstancedPropertyBagOverrideDetails.cpp
--- a/Engine/Plugins/Experimental/PCG/Source/PCGEditor/Private/Details/PCGInstancedPropertyBagOverrideDetails.cpp
+++ b/Engine/Plugins/Experimental/PCG/Source/PCGEditor/Private/Details/PCGInstancedPropertyBagOverrideDetails.cpp
@@ -103,6 +103,13 @@ EVisibility FPCGOverrideInstancedPropertyBagDataDetails::IsResetVisible(TSharedP
 
 FReply FPCGOverrideInstancedPropertyBagDataDetails::OnResetToDefaultValue(TSharedPtr<IPropertyHandle> InPropertyHandle) const
 {
+	if (!InPropertyHandle.IsValid() || !InPropertyHandle->GetProperty())
+	{
+		return FReply::Handled();
+	}
+
+	FScopedTransaction Transaction(FText::Format(LOCTEXT("OnResetToDefaultValue", "Reset Override for {0}"), FText::FromName(InPropertyHandle->GetProperty()->GetFName())));
+
 	if (Owner.IsValid())
 	{
 		Owner->ResetPropertyToDefault(InPropertyHandle->GetProperty());
