[PCG] Re-enable warning with override name clashes, now that data is fixed

#rnx

diff --git a/Engine/Plugins/Experimental/PCG/Source/PCG/Private/PCGSettings.cpp b/Engine/Plugins/Experimental/PCG/Source/PCG/Private/PCGSettings.cpp
--- a/Engine/Plugins/Experimental/PCG/Source/PCG/Private/PCGSettings.cpp
+++ b/Engine/Plugins/Experimental/PCG/Source/PCG/Private/PCGSettings.cpp
@@ -407,10 +407,10 @@ void UPCGSettings::FillOverridableParamsPins(TArray<FPCGPinProperties>& OutPins)
 	{
 		if (InputPinsLabelsAndTypes.Contains(OverridableParam.Label))
 		{
-			//const FString ParamsName = OverridableParam.Label.ToString();
-			//UE_LOG(LogPCG, Warning, TEXT("[%s-%s] While automatically adding override pins, an existing pin was found with conflicting name '%s'. "
-			//	"Rename or remove this pin to allow the automatic override pin to be added. Automatic override pin '%s' skipped."),
-			//	*GraphName, *NodeName, *ParamsName, *ParamsName);
+			const FString ParamsName = OverridableParam.Label.ToString();
+			UE_LOG(LogPCG, Warning, TEXT("[%s-%s] While automatically adding override pins, an existing pin was found with conflicting name '%s'. "
+				"Rename or remove this pin to allow the automatic override pin to be added. Automatic override pin '%s' skipped."),
+				*GraphName, *NodeName, *ParamsName, *ParamsName);
 			continue;
 		}
 
