Fixed cast error when passing in spatial data in the attribute noise (would have been an issue for old data)
Improved default values for spatial noise
Small QoL improvements in PCG BP related to data collections

diff --git a/Engine/Plugins/Experimental/PCG/Source/PCG/Private/Elements/PCGAttributeNoise.cpp b/Engine/Plugins/Experimental/PCG/Source/PCG/Private/Elements/PCGAttributeNoise.cpp
--- a/Engine/Plugins/Experimental/PCG/Source/PCG/Private/Elements/PCGAttributeNoise.cpp
+++ b/Engine/Plugins/Experimental/PCG/Source/PCG/Private/Elements/PCGAttributeNoise.cpp
@@ -297,7 +297,9 @@ bool FPCGAttributeNoiseElement::ExecuteInternal(FPCGContext* InContext) const
 		else
 		{
 			OutputPoints = &CastChecked<UPCGPointData>(Outputs[CurrentInput].Data)->GetMutablePoints();
-			InputPoints = &CastChecked<const UPCGPointData>(Inputs[CurrentInput].Data)->GetPoints();
+			// Note: for deprecation purposes, we consider that the input here is a spatial data (even though the pin is typed to be a Point Data)
+			// hence the need to call ToPointData, as otherwise a direct case to UPCGPointData would fail.
+			InputPoints = &(CastChecked<const UPCGSpatialData>(Inputs[CurrentInput].Data)->ToPointData(Context))->GetPoints();
 		}
 
 		check(InputPoints && OutputPoints);
diff --git a/Engine/Plugins/Experimental/PCG/Source/PCG/Private/Elements/PCGSpatialNoise.cpp b/Engine/Plugins/Experimental/PCG/Source/PCG/Private/Elements/PCGSpatialNoise.cpp
--- a/Engine/Plugins/Experimental/PCG/Source/PCG/Private/Elements/PCGSpatialNoise.cpp
+++ b/Engine/Plugins/Experimental/PCG/Source/PCG/Private/Elements/PCGSpatialNoise.cpp
@@ -21,6 +21,11 @@
 
 #define LOCTEXT_NAMESPACE "PCGSpatialNoise"
 
+UPCGSpatialNoiseSettings::UPCGSpatialNoiseSettings()
+{
+	ValueTarget.SetPointProperty(EPCGPointProperties::Density);
+}
+
 TArray<FPCGPinProperties> UPCGSpatialNoiseSettings::InputPinProperties() const
 {
 	return Super::DefaultPointInputPinProperties();
diff --git a/Engine/Plugins/Experimental/PCG/Source/PCG/Private/PCGData.cpp b/Engine/Plugins/Experimental/PCG/Source/PCG/Private/PCGData.cpp
--- a/Engine/Plugins/Experimental/PCG/Source/PCG/Private/PCGData.cpp
+++ b/Engine/Plugins/Experimental/PCG/Source/PCG/Private/PCGData.cpp
@@ -531,7 +531,7 @@ TArray<FPCGTaggedData> UPCGDataFunctionLibrary::GetInputs(const FPCGDataCollecti
 	return InCollection.GetInputs();
 }
 
-TArray<FPCGTaggedData> UPCGDataFunctionLibrary::GetInputsByPinLabel(const FPCGDataCollection& InCollection, const FName& InPinLabel)
+TArray<FPCGTaggedData> UPCGDataFunctionLibrary::GetInputsByPinLabel(const FPCGDataCollection& InCollection, const FName InPinLabel)
 {
 	return InCollection.GetInputsByPin(InPinLabel);
 }
@@ -546,7 +546,7 @@ TArray<FPCGTaggedData> UPCGDataFunctionLibrary::GetParams(const FPCGDataCollecti
 	return InCollection.GetAllParams();
 }
 
-TArray<FPCGTaggedData> UPCGDataFunctionLibrary::GetParamsByPinLabel(const FPCGDataCollection& InCollection, const FName& InPinLabel)
+TArray<FPCGTaggedData> UPCGDataFunctionLibrary::GetParamsByPinLabel(const FPCGDataCollection& InCollection, const FName InPinLabel)
 {
 	return InCollection.GetParamsByPin(InPinLabel);
 }
diff --git a/Engine/Plugins/Experimental/PCG/Source/PCG/Public/Elements/PCGSpatialNoise.h b/Engine/Plugins/Experimental/PCG/Source/PCG/Public/Elements/PCGSpatialNoise.h
--- a/Engine/Plugins/Experimental/PCG/Source/PCG/Public/Elements/PCGSpatialNoise.h
+++ b/Engine/Plugins/Experimental/PCG/Source/PCG/Public/Elements/PCGSpatialNoise.h
@@ -54,6 +54,8 @@ class UPCGSpatialNoiseSettings : public UPCGSettings
 	GENERATED_BODY()
 
 public:
+	UPCGSpatialNoiseSettings();
+
 	//~Begin UPCGSettings interface
 #if WITH_EDITOR
 	virtual FName GetDefaultNodeName() const override { return FName(TEXT("Spatial Noise")); }
diff --git a/Engine/Plugins/Experimental/PCG/Source/PCG/Public/PCGData.h b/Engine/Plugins/Experimental/PCG/Source/PCG/Public/PCGData.h
--- a/Engine/Plugins/Experimental/PCG/Source/PCG/Public/PCGData.h
+++ b/Engine/Plugins/Experimental/PCG/Source/PCG/Public/PCGData.h
@@ -228,7 +228,7 @@ public:
 	static TArray<FPCGTaggedData> GetInputs(const FPCGDataCollection& InCollection);
 
 	UFUNCTION(BlueprintCallable, Category = Data)
-	static TArray<FPCGTaggedData> GetInputsByPinLabel(const FPCGDataCollection& InCollection, const FName& InPinLabel);
+	static TArray<FPCGTaggedData> GetInputsByPinLabel(const FPCGDataCollection& InCollection, const FName InPinLabel);
 
 	UFUNCTION(BlueprintCallable, Category = Data)
 	static TArray<FPCGTaggedData> GetInputsByTag(const FPCGDataCollection& InCollection, const FString& InTag);
@@ -237,7 +237,7 @@ public:
 	static TArray<FPCGTaggedData> GetParams(const FPCGDataCollection& InCollection);
 
 	UFUNCTION(BlueprintCallable, Category = Data)
-	static TArray<FPCGTaggedData> GetParamsByPinLabel(const FPCGDataCollection& InCollection, const FName& InPinLabel);
+	static TArray<FPCGTaggedData> GetParamsByPinLabel(const FPCGDataCollection& InCollection, const FName InPinLabel);
 
 	UFUNCTION(BlueprintCallable, Category = Data)
 	static TArray<FPCGTaggedData> GetParamsByTag(const FPCGDataCollection& InCollection, const FString& InTag);
