Exposed flatten method that simplifies metadata serialization

#rnx

diff --git a/Engine/Plugins/Experimental/PCG/Source/PCG/Public/Metadata/PCGMetadata.h b/Engine/Plugins/Experimental/PCG/Source/PCG/Public/Metadata/PCGMetadata.h
--- a/Engine/Plugins/Experimental/PCG/Source/PCG/Public/Metadata/PCGMetadata.h
+++ b/Engine/Plugins/Experimental/PCG/Source/PCG/Public/Metadata/PCGMetadata.h
@@ -82,6 +82,7 @@ public:
 	bool HasParent(const UPCGMetadata* InTentativeParent) const;
 
 	/** Unparents current metadata by flattening the attributes (values, entries, etc.) */
+	UFUNCTION(BlueprintCallable, Category = "PCG|Metadata|Advanced")
 	void Flatten();
 
 	/** Create new streams */
