[PCG] Debug flag applies per-debug-object by default, to allow debugging particular partition actors.

Setting added to revert back to 'always debug all objects'.

diff --git a/Engine/Plugins/Experimental/PCG/Source/PCG/Private/Graph/PCGGraphExecutor.cpp b/Engine/Plugins/Experimental/PCG/Source/PCG/Private/Graph/PCGGraphExecutor.cpp
--- a/Engine/Plugins/Experimental/PCG/Source/PCG/Private/Graph/PCGGraphExecutor.cpp
+++ b/Engine/Plugins/Experimental/PCG/Source/PCG/Private/Graph/PCGGraphExecutor.cpp
@@ -1571,32 +1571,20 @@ namespace PCGGraphExecutor
 				return false;
 			}
 
-			APCGPartitionActor* PartitionActor = Cast<APCGPartitionActor>(InContext->SourceComponent->GetOwner());
-			UPCGComponent* OriginalComponent = PartitionActor ? PartitionActor->GetOriginalComponent(InContext->SourceComponent.Get()) : nullptr;
-
 			UPCGComponent* ComponentWithData = nullptr;
 			bool bComponentWithDataIsOriginalComponent = false;
 			if (FromGridSize == PCGHiGenGrid::UnboundedGridSize())
 			{
-				// Unbounded means not generated on grid, instead use the original component volume
-				ComponentWithData = OriginalComponent;
+				ComponentWithData = InContext->SourceComponent->GetOriginalComponent();
 				bComponentWithDataIsOriginalComponent = true;
 			}
-			else if (OriginalComponent)
+			else
 			{
-				// Now iterate through local components and find a component from the correct from-grid, and overlapping this component.
-				// Would be more efficient if we had a spatial index. The current component octree seems to only store original components.
-				Subsystem->ForAllRegisteredLocalComponents(OriginalComponent, [FromGridSize, &ComponentWithData, InContext](UPCGComponent* InLocalComponent)
+				Subsystem->ForAllOverlappingComponentsInHierarchy(InContext->SourceComponent.Get(), [FromGridSize, &ComponentWithData](UPCGComponent* InLocalComponent)
 				{
 					if (InLocalComponent->GetGenerationGridSize() == FromGridSize)
 					{
-						const FBox OtherBounds = InLocalComponent->GetGridBounds();
-						const FBox ThisBounds = InContext->SourceComponent->GetGridBounds();
-						const FBox Overlap = OtherBounds.Overlap(ThisBounds);
-						if (Overlap.GetVolume() > 0)
-						{
-							ComponentWithData = InLocalComponent;
-						}
+						ComponentWithData = InLocalComponent;
 					}
 				});
 			}
diff --git a/Engine/Plugins/Experimental/PCG/Source/PCG/Private/PCGComponent.cpp b/Engine/Plugins/Experimental/PCG/Source/PCG/Private/PCGComponent.cpp
--- a/Engine/Plugins/Experimental/PCG/Source/PCG/Private/PCGComponent.cpp
+++ b/Engine/Plugins/Experimental/PCG/Source/PCG/Private/PCGComponent.cpp
@@ -1695,6 +1695,17 @@ UPCGData* UPCGComponent::GetOriginalActorPCGData()
 	return nullptr;
 }
 
+UPCGComponent* UPCGComponent::GetOriginalComponent()
+{
+	if (!IsLocalComponent())
+	{
+		return this;
+	}
+
+	APCGPartitionActor* PartitionActor = Cast<APCGPartitionActor>(GetOwner());
+	return ensure(PartitionActor) ? PartitionActor->GetOriginalComponent(this) : this;
+}
+
 UPCGData* UPCGComponent::CreateActorPCGData()
 {
 	return CreateActorPCGData(GetOwner(), bParseActorComponents);
diff --git a/Engine/Plugins/Experimental/PCG/Source/PCG/Private/PCGElement.cpp b/Engine/Plugins/Experimental/PCG/Source/PCG/Private/PCGElement.cpp
--- a/Engine/Plugins/Experimental/PCG/Source/PCG/Private/PCGElement.cpp
+++ b/Engine/Plugins/Experimental/PCG/Source/PCG/Private/PCGElement.cpp
@@ -5,9 +5,11 @@
 #include "PCGComponent.h"
 #include "PCGContext.h"
 #include "PCGGraph.h"
+#include "PCGSubsystem.h"
 #include "Data/PCGPointData.h"
 #include "Elements/PCGDebugElement.h"
 #include "Elements/PCGSelfPruning.h"
+#include "Grid/PCGPartitionActor.h"
 
 #include "HAL/IConsoleManager.h"
 #include "Utils/PCGExtraCapture.h"
@@ -330,22 +332,62 @@ void IPCGElement::DisabledPassThroughData(FPCGContext* Context) const
 #if WITH_EDITOR
 void IPCGElement::DebugDisplay(FPCGContext* Context) const
 {
+	// Check Debug flag.
 	const UPCGSettingsInterface* SettingsInterface = Context->GetInputSettingsInterface();
-	if (SettingsInterface && SettingsInterface->bDebug)
+	if (!SettingsInterface || !SettingsInterface->bDebug)
 	{
-		FPCGDataCollection ElementInputs = Context->InputData;
-		FPCGDataCollection ElementOutputs = Context->OutputData;
+		return;
+	}
 
-		Context->InputData = ElementOutputs;
-		Context->OutputData = FPCGDataCollection();
+	// If graph is being inspected, only display Debug if the component is being inspected, or in the HiGen case also display if
+	// this component is a parent of an inspected component (because this data is available to child components).
+
+	// If the graph is not being inspected, or the current component is being inspected, then we know we should display
+	// debug, if not then we do further checks.
+	UPCGGraph* Graph = Context->SourceComponent.Get() ? Context->SourceComponent->GetGraph() : nullptr;
+	if (Graph && Graph->IsInspecting() && !Context->SourceComponent->IsInspecting() && Graph->DebugFlagAppliesToIndividualComponents())
+	{
+		// If we're no doing HiGen, or if the current component is not a local component (and therefore will not have children),
+		// then do not display debug.
+		if (!Graph->IsHierarchicalGenerationEnabled())
+		{
+			return;
+		}
 
-		PCGDebugElement::ExecuteDebugDisplay(Context);
+		// If a child of this component is being inspected (a local component on smaller grid and overlapping) then we still show debug,
+		// because this data is available to that child for use.
+		if (UPCGSubsystem* Subsystem = UPCGSubsystem::GetInstance(Context->SourceComponent->GetWorld()))
+		{
+			const uint32 ThisGenerationGridSize = Context->SourceComponent->GetGenerationGridSize();
 
-		Context->InputData = ElementInputs;
-		Context->OutputData = ElementOutputs;
+			bool bFoundInspectedChildComponent = false;
+			Subsystem->ForAllOverlappingComponentsInHierarchy(Context->SourceComponent.Get(), [ThisGenerationGridSize, &bFoundInspectedChildComponent](UPCGComponent* InLocalComponent)
+			{
+				if (InLocalComponent->GetGenerationGridSize() < ThisGenerationGridSize && InLocalComponent->IsInspecting())
+				{
+					bFoundInspectedChildComponent = true;
+				}
+			});
+
+			// If no inspected child component then don't display debug.
+			if (!bFoundInspectedChildComponent)
+			{
+				return;
+			}
+		}
 	}
-}
 
+	FPCGDataCollection ElementInputs = Context->InputData;
+	FPCGDataCollection ElementOutputs = Context->OutputData;
+
+	Context->InputData = ElementOutputs;
+	Context->OutputData = FPCGDataCollection();
+
+	PCGDebugElement::ExecuteDebugDisplay(Context);
+
+	Context->InputData = ElementInputs;
+	Context->OutputData = ElementOutputs;
+}
 #endif // WITH_EDITOR
 
 void IPCGElement::CleanupAndValidateOutput(FPCGContext* Context) const
diff --git a/Engine/Plugins/Experimental/PCG/Source/PCG/Private/PCGGraph.cpp b/Engine/Plugins/Experimental/PCG/Source/PCG/Private/PCGGraph.cpp
--- a/Engine/Plugins/Experimental/PCG/Source/PCG/Private/PCGGraph.cpp
+++ b/Engine/Plugins/Experimental/PCG/Source/PCG/Private/PCGGraph.cpp
@@ -792,6 +792,29 @@ void UPCGGraph::SetExtraEditorNodes(const TArray<TObjectPtr<const UObject>>& InN
 	}
 }
 
+bool UPCGGraph::IsInspecting() const
+{
+	return InspectionCounter > 0;
+}
+
+void UPCGGraph::EnableInspection()
+{
+	if (!ensure(InspectionCounter >= 0))
+	{
+		InspectionCounter = 0;
+	}
+
+	++InspectionCounter;
+}
+
+void UPCGGraph::DisableInspection()
+{
+	if (ensure(InspectionCounter > 0))
+	{
+		--InspectionCounter;
+	}
+};
+
 FPCGActorSelectionKeyToSettingsMap UPCGGraph::GetTrackedActorKeysToSettings() const
 {
 	FPCGActorSelectionKeyToSettingsMap TagsToSettings;
diff --git a/Engine/Plugins/Experimental/PCG/Source/PCG/Private/PCGSubsystem.cpp b/Engine/Plugins/Experimental/PCG/Source/PCG/Private/PCGSubsystem.cpp
--- a/Engine/Plugins/Experimental/PCG/Source/PCG/Private/PCGSubsystem.cpp
+++ b/Engine/Plugins/Experimental/PCG/Source/PCG/Private/PCGSubsystem.cpp
@@ -323,7 +323,7 @@ FPCGTaskId UPCGSubsystem::ScheduleComponent(UPCGComponent* PCGComponent, bool bS
 			Dependencies.Add(OriginalComponentTask);
 		}
 
-		auto LocalGenerateTask = [OriginalComponent = PCGComponent, &Dependencies](UPCGComponent* LocalComponent)
+		auto LocalGenerateTask = [OriginalComponent = PCGComponent, &Dependencies, bSave](UPCGComponent* LocalComponent)
 		{
 			// If the local component is currently generating, it's probably because it was requested by a refresh.
 			// Wait after this one instead
@@ -342,7 +342,8 @@ FPCGTaskId UPCGSubsystem::ScheduleComponent(UPCGComponent* PCGComponent, bool bS
 				LocalComponent->CleanupLocalImmediate(true);
 			}
 
-			return LocalComponent->GenerateInternal(/*bForce=*/ false, EPCGComponentGenerationTrigger::GenerateOnDemand, Dependencies);
+			// TODO If i don't forward the force flag, then nothing gets scheduled. Also why is it called bSave?
+			return LocalComponent->GenerateInternal(/*bForce=*/bSave, EPCGComponentGenerationTrigger::GenerateOnDemand, Dependencies);
 		};
 
 		AllTasks.Append(ActorAndComponentMapping.DispatchToRegisteredLocalComponents(PCGComponent, LocalGenerateTask));
@@ -587,6 +588,22 @@ void UPCGSubsystem::ForAllRegisteredLocalComponents(UPCGComponent* OriginalCompo
 	ActorAndComponentMapping.DispatchToRegisteredLocalComponents(OriginalComponent, WrapperFunc);
 }
 
+void UPCGSubsystem::ForAllOverlappingComponentsInHierarchy(UPCGComponent* InComponent, const TFunction<void(UPCGComponent*)>& InFunc) const
+{
+	UPCGComponent* OriginalComponent = InComponent->GetOriginalComponent();
+
+	ForAllRegisteredLocalComponents(OriginalComponent, [InComponent, &InFunc](UPCGComponent* InLocalComponent)
+	{
+		const FBox OtherBounds = InLocalComponent->GetGridBounds();
+		const FBox ThisBounds = InComponent->GetGridBounds();
+		const FBox Overlap = OtherBounds.Overlap(ThisBounds);
+		if (Overlap.GetVolume() > 0)
+		{
+			InFunc(InLocalComponent);
+		}
+	});
+}
+
 bool UPCGSubsystem::IsGraphCacheDebuggingEnabled() const
 {
 	return GraphExecutor && GraphExecutor->IsGraphCacheDebuggingEnabled();
diff --git a/Engine/Plugins/Experimental/PCG/Source/PCG/Public/PCGComponent.h b/Engine/Plugins/Experimental/PCG/Source/PCG/Public/PCGComponent.h
--- a/Engine/Plugins/Experimental/PCG/Source/PCG/Public/PCGComponent.h
+++ b/Engine/Plugins/Experimental/PCG/Source/PCG/Public/PCGComponent.h
@@ -106,6 +106,9 @@ public:
 	UPCGData* GetLandscapeHeightPCGData();
 	UPCGData* GetOriginalActorPCGData();
 
+	/** If this is a local component returns self, otherwise returns the original component. */
+	UPCGComponent* GetOriginalComponent();
+
 	bool CanPartition() const;
 
 	UPCGGraph* GetGraph() const;
diff --git a/Engine/Plugins/Experimental/PCG/Source/PCG/Public/PCGGraph.h b/Engine/Plugins/Experimental/PCG/Source/PCG/Public/PCGGraph.h
--- a/Engine/Plugins/Experimental/PCG/Source/PCG/Public/PCGGraph.h
+++ b/Engine/Plugins/Experimental/PCG/Source/PCG/Public/PCGGraph.h
@@ -222,6 +222,11 @@ public:
 
 	const TArray<TObjectPtr<UObject>>& GetExtraEditorNodes() const { return ExtraEditorNodes; }
 	void SetExtraEditorNodes(const TArray<TObjectPtr<const UObject>>& InNodes);
+
+	bool IsInspecting() const;
+	void EnableInspection();
+	void DisableInspection();
+	bool DebugFlagAppliesToIndividualComponents() const { return bDebugFlagAppliesToIndividualComponents; }
 #endif
 
 #if WITH_EDITOR
@@ -277,6 +282,14 @@ protected:
 	mutable TMap<const UPCGNode*, uint32> NodeToGridSize;
 	mutable FRWLock NodeToGridSizeLock;
 
+#if WITH_EDITORONLY_DATA
+	/** When true the Debug flag in the graph editor will display debug information contextually for the selected debug object. Otherwise
+	* debug information is displayed for all components using a graph (requires regenerate).
+	*/
+	UPROPERTY(EditAnywhere, Category = Debug)
+	bool bDebugFlagAppliesToIndividualComponents = true;
+#endif // WITH_EDITORONLY_DATA
+
 public:
 	virtual const FInstancedPropertyBag* GetUserParametersStruct() const override { return &UserParameters; }
 
@@ -294,6 +307,7 @@ private:
 	bool bIsNotifying = false;
 	bool bUserPausedNotificationsInGraphEditor = false;
 	int32 NumberOfUserParametersPreEdit = 0;
+	int32 InspectionCounter = 0;
 #endif // WITH_EDITOR
 };
 
diff --git a/Engine/Plugins/Experimental/PCG/Source/PCG/Public/PCGSubsystem.h b/Engine/Plugins/Experimental/PCG/Source/PCG/Public/PCGSubsystem.h
--- a/Engine/Plugins/Experimental/PCG/Source/PCG/Public/PCGSubsystem.h
+++ b/Engine/Plugins/Experimental/PCG/Source/PCG/Public/PCGSubsystem.h
@@ -147,6 +147,9 @@ public:
 	/* Call the InFunc function to all local component registered to the original component. Thread safe*/
 	void ForAllRegisteredLocalComponents(UPCGComponent* OriginalComponent, const TFunction<void(UPCGComponent*)>& InFunc) const;
 
+	/** Traverses the hierarchy associated with the given component and calls InFunc for each overlapping component. */
+	void ForAllOverlappingComponentsInHierarchy(UPCGComponent* InComponent, const TFunction<void(UPCGComponent*)>& InFunc) const;
+
 	/** True if graph cache debugging is enabled. */
 	bool IsGraphCacheDebuggingEnabled() const;
 
diff --git a/Engine/Plugins/Experimental/PCG/Source/PCGEditor/Private/PCGEditor.cpp b/Engine/Plugins/Experimental/PCG/Source/PCGEditor/Private/PCGEditor.cpp
--- a/Engine/Plugins/Experimental/PCG/Source/PCGEditor/Private/PCGEditor.cpp
+++ b/Engine/Plugins/Experimental/PCG/Source/PCGEditor/Private/PCGEditor.cpp
@@ -13,6 +13,7 @@
 #include "PCGInputOutputSettings.h"
 #include "PCGPin.h"
 #include "PCGSubsystem.h"
+#include "Algo/AnyOf.h"
 #include "Rendering/SlateRenderer.h"
 #include "Tests/Determinism/PCGDeterminismNativeTests.h"
 #include "Tests/Determinism/PCGDeterminismTestBlueprintBase.h"
@@ -181,27 +182,36 @@ UPCGEditorGraph* FPCGEditor::GetPCGEditorGraph()
 
 void FPCGEditor::SetComponentAndStackBeingInspected(UPCGComponent* InPCGComponent, const FPCGStack& InPCGStack)
 {
-	if (PCGComponentBeingInspected != InPCGComponent)
+	UPCGComponent* OldComponent = PCGComponentBeingInspected.Get();
+	UPCGComponent* NewComponent = InPCGComponent;
+
+	if (OldComponent != NewComponent)
 	{
-		if (PCGComponentBeingInspected.IsValid())
+		if (OldComponent)
 		{
-			PCGComponentBeingInspected->DisableInspection();
+			OldComponent->DisableInspection();
+			if (PCGGraphBeingEdited)
+			{
+				PCGGraphBeingEdited->DisableInspection();
+			}
 		}
 
-		PCGComponentBeingInspected = InPCGComponent;
-		OnInspectedComponentChangedDelegate.Broadcast(PCGComponentBeingInspected.Get());
+		const bool bNewComponentStartedInspecting = NewComponent && !NewComponent->IsInspecting();
 
-		if (PCGComponentBeingInspected.IsValid())
+		PCGComponentBeingInspected = NewComponent;
+		OnInspectedComponentChangedDelegate.Broadcast(NewComponent);
+
+		if (NewComponent)
 		{
-			if (!PCGComponentBeingInspected->IsInspecting())
+			PCGComponentBeingInspected->EnableInspection();
+			if (PCGGraphBeingEdited)
 			{
-				// We need to force generation so that we create debug data
-				PCGComponentBeingInspected->GenerateLocal(/*bForce=*/true);
+				PCGGraphBeingEdited->EnableInspection();
 			}
-			
-			PCGComponentBeingInspected->EnableInspection();
 		}
 
+		UpdateDebugAfterComponentSelection(OldComponent, NewComponent, bNewComponentStartedInspecting);
+
 		check(PCGEditorGraph);
 		for (UEdGraphNode* Node : PCGEditorGraph->Nodes)
 		{
@@ -219,13 +229,64 @@ void FPCGEditor::SetComponentAndStackBeingInspected(UPCGComponent* InPCGComponen
 		}
 	}
 
-	if(InPCGStack != StackBeingInspected)
+	if (InPCGStack != StackBeingInspected)
 	{
 		StackBeingInspected = InPCGStack;
 		OnInspectedStackChangedDelegate.Broadcast(StackBeingInspected);
 	}
 }
 
+void FPCGEditor::UpdateDebugAfterComponentSelection(UPCGComponent* InOldComponent, UPCGComponent* InNewComponent, bool bInNewComponentStartedInspecting)
+{
+	if (!ensure(PCGGraphBeingEdited))
+	{
+		return;
+	}
+
+	// If individual component debugging is disabled, just generate the new component if required.
+	if (!PCGGraphBeingEdited->DebugFlagAppliesToIndividualComponents())
+	{
+		if (InNewComponent && bInNewComponentStartedInspecting)
+		{
+			InNewComponent->GenerateLocal(/*bForce=*/true);
+		}
+
+		return;
+	}
+
+	// Trigger necessary generation(s) for per-component debugging.
+	if (!InOldComponent)
+	{
+		if (InNewComponent && bInNewComponentStartedInspecting)
+		{
+			// Transition from 'null' to 'any component not already inspecting' - generate to create debug/inspection info.
+			// If we have null selected, all components are displaying debug. Go to Original component so that all refresh.
+			InNewComponent->GetOriginalComponent()->GenerateLocal(/*bForce=*/true);
+		}
+	}
+	else
+	{
+		const bool bDebugFlagSetOnAnyNode = Algo::AnyOf(PCGGraphBeingEdited->GetNodes(), [](const UPCGNode* InNode)
+		{
+			return InNode && InNode->GetSettings() && InNode->GetSettings()->bDebug;
+		});
+
+		// Regenerate to clear debug info if switching components, or if changing from a component to null.
+		if (InNewComponent || bDebugFlagSetOnAnyNode)
+		{
+			// Use original component - debug can be displayed both by the local component and parent local components.
+			InOldComponent->GetOriginalComponent()->GenerateLocal(/*bForce=*/true);
+		}
+
+		// Debug new component if it wasn't already
+		if (InNewComponent && bInNewComponentStartedInspecting)
+		{
+			// Use original component - debug can be displayed both by the local component and parent local components.
+			InNewComponent->GetOriginalComponent()->GenerateLocal(/*bForce=*/true);
+		}
+	}
+}
+
 void FPCGEditor::JumpToNode(const UEdGraphNode* InNode)
 {
 	if (GraphEditorWidget.IsValid())
diff --git a/Engine/Plugins/Experimental/PCG/Source/PCGEditor/Private/PCGEditor.h b/Engine/Plugins/Experimental/PCG/Source/PCGEditor/Private/PCGEditor.h
--- a/Engine/Plugins/Experimental/PCG/Source/PCGEditor/Private/PCGEditor.h
+++ b/Engine/Plugins/Experimental/PCG/Source/PCGEditor/Private/PCGEditor.h
@@ -275,6 +275,9 @@ private:
 
 	void OnGraphGridSizesChanged(UPCGGraphInterface* InGraph);
 
+	/** Trigger any generation required to ensure debug display is up to date. */
+	void UpdateDebugAfterComponentSelection(UPCGComponent* InOldComponent, UPCGComponent* InNewComponent, bool bNewComponentStartedInspecting);
+
 	/** Helper to get to the subsystem. */
 	static class UPCGSubsystem* GetSubsystem();
 
