[PCG] Created a native node that returns the number of point

diff --git a/Engine/Plugins/Experimental/PCG/Source/PCG/Private/Elements/PCGNumberOfPoints.cpp b/Engine/Plugins/Experimental/PCG/Source/PCG/Private/Elements/PCGNumberOfPoints.cpp
new file mode 100644
--- /dev/null
+++ b/Engine/Plugins/Experimental/PCG/Source/PCG/Private/Elements/PCGNumberOfPoints.cpp
@@ -0,0 +1,86 @@
+// Copyright Epic Games, Inc. All Rights Reserved.
+
+#include "Elements/PCGNumberOfPoints.h"
+
+#include "PCGContext.h"
+#include "PCGCustomVersion.h"
+#include "PCGPin.h"
+#include "PCGParamData.h"
+#include "Data/PCGPointData.h"
+#include "Metadata/PCGMetadata.h"
+
+#define LOCTEXT_NAMESPACE "PCGNumberOfPointsSettings"
+
+
+TArray<FPCGPinProperties> UPCGNumberOfPointsSettings::InputPinProperties() const
+{
+	TArray<FPCGPinProperties> PinProperties;
+	PinProperties.Emplace(PCGPinConstants::DefaultInputLabel, EPCGDataType::Point, /*bInAllowMultipleConnections=*/ false, /*bAllowMultipleData=*/false);
+
+	return PinProperties;
+}
+
+TArray<FPCGPinProperties> UPCGNumberOfPointsSettings::OutputPinProperties() const
+{
+	TArray<FPCGPinProperties> PinProperties;
+	PinProperties.Emplace(PCGPinConstants::DefaultOutputLabel, EPCGDataType::Param);
+
+	return PinProperties;
+}
+
+FPCGElementPtr UPCGNumberOfPointsSettings::CreateElement() const
+{
+	return MakeShared<FPCGNumberOfPointsElement>();
+}
+
+bool FPCGNumberOfPointsElement::ExecuteInternal(FPCGContext* Context) const
+{
+	TRACE_CPUPROFILER_EVENT_SCOPE(FPCGNumberOfPointsElement::Execute);
+	
+	check(Context);
+
+	const UPCGNumberOfPointsSettings* Settings = Context->GetInputSettings<UPCGNumberOfPointsSettings>();
+	check(Settings);
+
+	TArray<FPCGTaggedData> Inputs = Context->InputData.GetInputsByPin(PCGPinConstants::DefaultInputLabel);
+
+	if (Inputs.Num() != 1)
+	{
+		PCGE_LOG(Warning, LogOnly, FText::Format(LOCTEXT("WrongNumberOfInputs", "Input pin expected to have one input data element, encountered {0}"), Inputs.Num()));
+		return true;
+	}
+
+	const UPCGPointData* PointData = Cast<UPCGPointData>(Inputs[0].Data);
+
+	if (!PointData)
+	{
+		PCGE_LOG(Error, GraphAndLog, LOCTEXT("InputNotPointData", "Input is not a point data"));
+		return true;
+	}
+
+	TArray<FPCGTaggedData>& Outputs = Context->OutputData.TaggedData;
+	UPCGParamData* OutputParamData = NewObject<UPCGParamData>();
+	
+	const int32 NumPoints = PointData->GetPoints().Num();
+
+	FPCGMetadataAttribute<int32>* NewAttribute =
+		OutputParamData->Metadata->CreateAttribute<int32>(Settings->OutputAttributeName, NumPoints, /*bAllowInterpolation=*/true, /*bOverrideParent=*/false);
+
+	if (NewAttribute) 
+	{
+		NewAttribute->SetValue(OutputParamData->Metadata->AddEntry(), NumPoints);
+
+		FPCGTaggedData& Output = Outputs.Emplace_GetRef();
+		Output.Data = OutputParamData;
+		Output.Pin = PCGPinConstants::DefaultOutputLabel;
+	}
+	else 
+	{
+		PCGE_LOG(Error, GraphAndLog, LOCTEXT("AttributeNotFound", "New Attribute failed to create."));
+	}
+
+
+	return true;
+}
+
+#undef LOCTEXT_NAMESPACE
diff --git a/Engine/Plugins/Experimental/PCG/Source/PCG/Public/Elements/PCGNumberOfPoints.h b/Engine/Plugins/Experimental/PCG/Source/PCG/Public/Elements/PCGNumberOfPoints.h
new file mode 100644
--- /dev/null
+++ b/Engine/Plugins/Experimental/PCG/Source/PCG/Public/Elements/PCGNumberOfPoints.h
@@ -0,0 +1,43 @@
+// Copyright Epic Games, Inc. All Rights Reserved.
+#pragma once
+
+#include "PCGSettings.h"
+
+#include "UObject/NameTypes.h"
+#include "PCGNumberOfPoints.generated.h"
+
+/**
+ * Used to wrangle multiple input wires into one output wire for organizational purposes.
+ */
+
+UCLASS(BlueprintType, ClassGroup = (Procedural))
+class UPCGNumberOfPointsSettings : public UPCGSettings
+{
+	GENERATED_BODY()
+
+public:
+	//~Begin UPCGSettings interface
+#if WITH_EDITOR
+	virtual FName GetDefaultNodeName() const override { return FName(TEXT("Get Number of Points")); }
+	virtual FText GetDefaultNodeTitle() const override { return NSLOCTEXT("PCGNumberOfPointsSettings", "NodeTitle", "Get Number of Points"); }
+	virtual FText GetNodeTooltipText() const override { return NSLOCTEXT("PCGNumberOfPointsSettings", "NodeTooltip", "Return the number of points in a point data."); }
+	virtual EPCGSettingsType GetType() const override { return EPCGSettingsType::Spatial; }
+#endif
+
+
+protected:
+	virtual TArray<FPCGPinProperties> InputPinProperties() const override;
+	virtual TArray<FPCGPinProperties> OutputPinProperties() const override;
+	virtual FPCGElementPtr CreateElement() const override;
+	//~End UPCGSettings interface
+
+public:
+	UPROPERTY(BlueprintReadWrite, EditAnywhere, Category = Settings, meta = (PCG_Overridable))
+	FName OutputAttributeName = TEXT("NumPoints");
+};
+
+class FPCGNumberOfPointsElement : public FSimplePCGElement
+{
+protected:
+	virtual bool ExecuteInternal(FPCGContext* Context) const override;
+};
\ No newline at end of file
