Added node to count number of data in collection, useful for special case processing or existence tests

diff --git a/Engine/Plugins/Experimental/PCG/Source/PCG/Private/Elements/PCGDataNum.cpp b/Engine/Plugins/Experimental/PCG/Source/PCG/Private/Elements/PCGDataNum.cpp
new file mode 100644
--- /dev/null
+++ b/Engine/Plugins/Experimental/PCG/Source/PCG/Private/Elements/PCGDataNum.cpp
@@ -0,0 +1,76 @@
+// Copyright Epic Games, Inc. All Rights Reserved.
+
+#include "Elements/PCGDataNum.h"
+
+#include "PCGContext.h"
+#include "PCGParamData.h"
+#include "PCGPin.h"
+
+#include UE_INLINE_GENERATED_CPP_BY_NAME(PCGDataNum)
+
+#define LOCTEXT_NAMESPACE "PCGDataNumElement"
+
+#if WITH_EDITOR
+FName UPCGDataNumSettings::GetDefaultNodeName() const
+{
+	return TEXT("DataNum");
+}
+
+FText UPCGDataNumSettings::GetDefaultNodeTitle() const
+{
+	return LOCTEXT("NodeTitle", "Data Count");
+}
+
+FText UPCGDataNumSettings::GetNodeTooltipText() const
+{
+	return LOCTEXT("NodeTooltip", "Returns the count of data in the input data collection.");
+}
+#endif // WITH_EDITOR
+
+TArray<FPCGPinProperties> UPCGDataNumSettings::OutputPinProperties() const
+{
+	TArray<FPCGPinProperties> PinProperties;
+	PinProperties.Emplace(
+		PCGPinConstants::DefaultOutputLabel,
+		EPCGDataType::Param,
+		/*bInAllowMultipleConnections=*/true,
+		/*bInAllowMultipleData=*/true,
+		LOCTEXT("OutParamTooltip", "Attribute set containing the data count from the input collection"));
+
+	return PinProperties;
+}
+
+FPCGElementPtr UPCGDataNumSettings::CreateElement() const
+{
+	return MakeShared<FPCGDataNumElement>();
+}
+
+bool FPCGDataNumElement::ExecuteInternal(FPCGContext* Context) const
+{
+	TRACE_CPUPROFILER_EVENT_SCOPE(FPCGDataNumElement::Execute);
+
+	check(Context);
+
+	const UPCGDataNumSettings* Settings = Context->GetInputSettings<UPCGDataNumSettings>();
+	check(Settings);
+
+	const int32 InputDataCount = Context->InputData.GetInputsByPin(PCGPinConstants::DefaultInputLabel).Num();
+
+	UPCGParamData* OutputParamData = NewObject<UPCGParamData>();
+	FPCGMetadataAttribute<int32>* NumAttribute = OutputParamData->Metadata->CreateAttribute<int32>(Settings->OutputAttributeName, InputDataCount, /*bAllowInterpolation=*/false, /*bOverrideParent=*/false);
+
+	if (!NumAttribute)
+	{
+		PCGE_LOG(Error, GraphAndLog, FText::Format(LOCTEXT("AttributeCreationFailed", "Failed to create attribute {0}"), FText::FromName(Settings->OutputAttributeName)));
+		return true;
+	}
+
+	NumAttribute->SetValue(OutputParamData->Metadata->AddEntry(), InputDataCount);
+
+	FPCGTaggedData& Output = Context->OutputData.TaggedData.Emplace_GetRef();
+	Output.Data = OutputParamData;
+
+	return true;
+}
+
+#undef LOCTEXT_NAMESPACE
\ No newline at end of file
diff --git a/Engine/Plugins/Experimental/PCG/Source/PCG/Public/Elements/PCGDataNum.h b/Engine/Plugins/Experimental/PCG/Source/PCG/Public/Elements/PCGDataNum.h
new file mode 100644
--- /dev/null
+++ b/Engine/Plugins/Experimental/PCG/Source/PCG/Public/Elements/PCGDataNum.h
@@ -0,0 +1,38 @@
+// Copyright Epic Games, Inc. All Rights Reserved.
+
+#pragma once
+
+#include "PCGSettings.h"
+
+#include "PCGDataNum.generated.h"
+
+/** Counts number of data in a data collection */
+UCLASS(BlueprintType, ClassGroup = (Procedural))
+class PCG_API UPCGDataNumSettings : public UPCGSettings
+{
+	GENERATED_BODY()
+
+public:
+	//~Begin UPCGSettings interface
+#if WITH_EDITOR
+	virtual FName GetDefaultNodeName() const override;
+	virtual FText GetDefaultNodeTitle() const override;
+	virtual FText GetNodeTooltipText() const override;
+	virtual EPCGSettingsType GetType() const override { return EPCGSettingsType::Param; }
+#endif
+	
+protected:
+	virtual TArray<FPCGPinProperties> OutputPinProperties() const override;
+	virtual FPCGElementPtr CreateElement() const override;
+	//~End UPCGSettings interface
+
+public:
+	UPROPERTY(BlueprintReadWrite, EditAnywhere, Category = Settings, meta = (PCG_Overridable))
+	FName OutputAttributeName = NAME_None;
+};
+
+class FPCGDataNumElement : public FSimplePCGElement
+{
+protected:
+	virtual bool ExecuteInternal(FPCGContext* Context) const override;
+};
\ No newline at end of file
