Attach behavior improvements needed for proper streaming, including the use of actor folders.
Added some actor folder cleanup & management on the PCG component side of things to make it nice.
Fixed some uproperty tags that were invalid.
Applied same template vs class logic in the CreateTargetActor as we have done in the SpawnActor node.

diff --git a/Engine/Plugins/Experimental/PCG/Source/PCG/Private/Elements/PCGCreateSpline.cpp b/Engine/Plugins/Experimental/PCG/Source/PCG/Private/Elements/PCGCreateSpline.cpp
--- a/Engine/Plugins/Experimental/PCG/Source/PCG/Private/Elements/PCGCreateSpline.cpp
+++ b/Engine/Plugins/Experimental/PCG/Source/PCG/Private/Elements/PCGCreateSpline.cpp
@@ -139,7 +139,6 @@ bool FPCGCreateSplineElement::ExecuteInternal(FPCGContext* Context) const
 		else if(Settings->Mode == EPCGCreateSplineMode::CreateNewActor)
 		{
 			FActorSpawnParameters SpawnParams;
-			SpawnParams.Owner = TargetActor;
 			if (PCGHelpers::IsRuntimeOrPIE())
 			{
 				SpawnParams.ObjectFlags |= RF_Transient;
@@ -161,8 +160,7 @@ bool FPCGCreateSplineElement::ExecuteInternal(FPCGContext* Context) const
 
 			SplineActor->Tags = TargetActor->Tags;
 			SplineActor->Tags.AddUnique(PCGHelpers::DefaultPCGActorTag);
-			SplineActor->AttachToActor(TargetActor, FAttachmentTransformRules::KeepWorldTransform);
-
+			PCGHelpers::AttachToParent(SplineActor, TargetActor, Settings->AttachOptions);
 			SplineData->TargetActor = SplineActor;
 
 			ManagedActors->GeneratedActors.Add(SplineActor);
diff --git a/Engine/Plugins/Experimental/PCG/Source/PCG/Private/Elements/PCGCreateTargetActor.cpp b/Engine/Plugins/Experimental/PCG/Source/PCG/Private/Elements/PCGCreateTargetActor.cpp
--- a/Engine/Plugins/Experimental/PCG/Source/PCG/Private/Elements/PCGCreateTargetActor.cpp
+++ b/Engine/Plugins/Experimental/PCG/Source/PCG/Private/Elements/PCGCreateTargetActor.cpp
@@ -99,6 +99,10 @@ void UPCGCreateTargetActor::PostEditChangeProperty(FPropertyChangedEvent& Proper
 			SetupBlueprintEvent();
 			RefreshTemplateActor();
 		}
+		else if (PropertyName == GET_MEMBER_NAME_CHECKED(UPCGCreateTargetActor, bAllowTemplateActorEditing))
+		{
+			RefreshTemplateActor();
+		}
 	}
 
 	Super::PostEditChangeProperty(PropertyChangedEvent);
@@ -129,23 +133,46 @@ void UPCGCreateTargetActor::OnBlueprintChanged(UBlueprint* InBlueprint)
 
 void UPCGCreateTargetActor::RefreshTemplateActor()
 {
-	if (TemplateActorClass)
+	// Implementation note: this is similar to the child actor component implementation
+	if (TemplateActorClass && bAllowTemplateActorEditing)
 	{
-		AActor* NewTemplateActor = NewObject<AActor>(this, TemplateActorClass, NAME_None, RF_ArchetypeObject | RF_Transactional | RF_Public);
+		const bool bCreateNewTemplateActor = (!TemplateActor || TemplateActor->GetClass() != TemplateActorClass);
 
-		if (TemplateActor)
+		if (bCreateNewTemplateActor)
 		{
-			UEngine::FCopyPropertiesForUnrelatedObjectsParams Options;
-			Options.bNotifyObjectReplacement = true;
-			UEngine::CopyPropertiesForUnrelatedObjects(TemplateActor, NewTemplateActor, Options);
+			AActor* NewTemplateActor = NewObject<AActor>(GetTransientPackage(), TemplateActorClass, NAME_None, RF_ArchetypeObject | RF_Transactional | RF_Public);
 
-			TemplateActor->Rename(nullptr, GetTransientPackage(), REN_DontCreateRedirectors | REN_ForceNoResetLoaders);
-		}
+			if (TemplateActor)
+			{
+				UEngine::FCopyPropertiesForUnrelatedObjectsParams Options;
+				Options.bNotifyObjectReplacement = true;
+				UEngine::CopyPropertiesForUnrelatedObjects(TemplateActor, NewTemplateActor, Options);
+
+				TemplateActor->Rename(nullptr, GetTransientPackage(), REN_DontCreateRedirectors);
+
+				TMap<UObject*, UObject*> OldToNew;
+				OldToNew.Emplace(TemplateActor, NewTemplateActor);
+				GEngine->NotifyToolsOfObjectReplacement(OldToNew);
+
+				TemplateActor->MarkAsGarbage();
+			}
+
+			TemplateActor = NewTemplateActor;
 
-		TemplateActor = NewTemplateActor;
+			// Record initial object state in case we're in a transaction context.
+			TemplateActor->Modify();
+
+			// Outer to this object
+			TemplateActor->Rename(nullptr, this, REN_DoNotDirty | REN_DontCreateRedirectors | REN_ForceNoResetLoaders);
+		}
 	}
 	else
 	{
+		if (TemplateActor)
+		{
+			TemplateActor->MarkAsGarbage();
+		}
+
 		TemplateActor = nullptr;
 	}
 }
@@ -155,6 +182,12 @@ void UPCGCreateTargetActor::PostLoad()
 	Super::PostLoad();
 
 #if WITH_EDITOR
+	// Since the template actor editing is set to false by default, this needs to be corrected on post-load for proper deprecation
+	if (TemplateActor)
+	{
+		bAllowTemplateActorEditing = true;
+	}
+
 	SetupBlueprintEvent();
 
 	if (TemplateActorClass)
@@ -183,14 +216,14 @@ bool FPCGCreateTargetActorElement::ExecuteInternal(FPCGContext* Context) const
 	check(Settings);
 
 	// Early out if the template actor isn't valid
-	if (!Settings->TemplateActorClass || Settings->TemplateActorClass->HasAnyClassFlags(CLASS_Abstract))
+	if (!Settings->TemplateActorClass || Settings->TemplateActorClass->HasAnyClassFlags(CLASS_Abstract) || !Settings->TemplateActorClass->GetDefaultObject()->IsA<AActor>())
 	{
 		const FText ClassName = Settings->TemplateActorClass ? FText::FromString(Settings->TemplateActorClass->GetFName().ToString()) : FText::FromName(NAME_None);
 		PCGE_LOG(Error, GraphAndLog, FText::Format(LOCTEXT("InvalidTemplateActorClass", "Invalid template actor class '{0}'"), ClassName));
 		return true;
 	}
 
-	if (!ensure(Settings->TemplateActor && Settings->TemplateActor->IsA(Settings->TemplateActorClass)))
+	if (!ensure(!Settings->TemplateActor || Settings->TemplateActor->IsA(Settings->TemplateActorClass)))
 	{
 		return true;
 	}
@@ -203,7 +236,7 @@ bool FPCGCreateTargetActorElement::ExecuteInternal(FPCGContext* Context) const
 	}
 
 	const bool bHasAuthority = !Context->SourceComponent.IsValid() || (Context->SourceComponent->GetOwner() && Context->SourceComponent->GetOwner()->HasAuthority());
-	const bool bSpawnedActorRequiresAuthority = Settings->TemplateActor->GetIsReplicated();
+	const bool bSpawnedActorRequiresAuthority = CastChecked<AActor>(Settings->TemplateActorClass->GetDefaultObject())->GetIsReplicated();
 
 	if (!bHasAuthority && bSpawnedActorRequiresAuthority)
 	{
@@ -214,7 +247,6 @@ bool FPCGCreateTargetActorElement::ExecuteInternal(FPCGContext* Context) const
 	AActor* TemplateActor = Settings->TemplateActor.Get();
 
 	FActorSpawnParameters SpawnParams;
-	SpawnParams.Owner = TargetActor;
 	SpawnParams.Template = TemplateActor;
 	SpawnParams.SpawnCollisionHandlingOverride = ESpawnActorCollisionHandlingMethod::AlwaysSpawn;
 
@@ -229,7 +261,7 @@ bool FPCGCreateTargetActorElement::ExecuteInternal(FPCGContext* Context) const
 		Transform = Settings->ActorPivot;
 	}
 
-	AActor* GeneratedActor = UPCGActorHelpers::SpawnDefaultActor(TargetActor->GetWorld(), Settings->TemplateActorClass, Transform, SpawnParams, TargetActor);
+	AActor* GeneratedActor = UPCGActorHelpers::SpawnDefaultActor(TargetActor->GetWorld(), Settings->TemplateActorClass, Transform, SpawnParams);
 
 	if (!GeneratedActor)
 	{
@@ -237,6 +269,9 @@ bool FPCGCreateTargetActorElement::ExecuteInternal(FPCGContext* Context) const
 		return true;
 	}
 
+	// Always attach if root actor is provided
+	PCGHelpers::AttachToParent(GeneratedActor, TargetActor, Settings->RootActor.Get() ? EPCGAttachOptions::Attached : Settings->AttachOptions);
+
 #if WITH_EDITOR
 	if (Settings->ActorLabel != FString())
 	{
diff --git a/Engine/Plugins/Experimental/PCG/Source/PCG/Private/Elements/PCGSpawnActor.cpp b/Engine/Plugins/Experimental/PCG/Source/PCG/Private/Elements/PCGSpawnActor.cpp
--- a/Engine/Plugins/Experimental/PCG/Source/PCG/Private/Elements/PCGSpawnActor.cpp
+++ b/Engine/Plugins/Experimental/PCG/Source/PCG/Private/Elements/PCGSpawnActor.cpp
@@ -833,7 +833,6 @@ void FPCGSpawnActorElement::SpawnActors(FPCGSubgraphContext* Context, AActor* Ta
 	PCGSpawnActorHelpers::FActorOverrides ActorOverrides(Settings->ActorOverrides, TemplateActor, PointData);
 
 	FActorSpawnParameters SpawnParams;
-	SpawnParams.Owner = TargetActor;
 	SpawnParams.Template = TemplateActor;
 	SpawnParams.SpawnCollisionHandlingOverride = ESpawnActorCollisionHandlingMethod::AlwaysSpawn;
 
@@ -910,6 +909,12 @@ void FPCGSpawnActorElement::SpawnActors(FPCGSubgraphContext* Context, AActor* Ta
 		UPCGManagedActors* ManagedActors = NewObject<UPCGManagedActors>(Context->SourceComponent.Get());
 		ManagedActors->SetCrc(InputDependenciesCrc);
 
+		// If generated actors are not directly attached, place them in a subfolder for tidiness.
+		FString GeneratedActorsFolderPath;
+#if WITH_EDITOR
+		PCGHelpers::GetGeneratedActorsFolderPath(TargetActor, GeneratedActorsFolderPath);
+#endif
+
 		for (int32 i = 0; i < Points.Num(); ++i)
 		{
 			const FPCGPoint& Point = Points[i];
@@ -927,7 +932,7 @@ void FPCGSpawnActorElement::SpawnActors(FPCGSubgraphContext* Context, AActor* Ta
 			// HACK: until UE-62747 is fixed, we have to force set the scale after spawning the actor
 			GeneratedActor->SetActorRelativeScale3D(Point.Transform.GetScale3D());
 			GeneratedActor->Tags.Append(NewActorTags);
-			GeneratedActor->AttachToActor(TargetActor, FAttachmentTransformRules::KeepWorldTransform);
+			PCGHelpers::AttachToParent(GeneratedActor, TargetActor, Settings->AttachOptions, GeneratedActorsFolderPath);
 
 			for (UFunction* PostSpawnFunction : PostSpawnFunctions)
 			{
diff --git a/Engine/Plugins/Experimental/PCG/Source/PCG/Private/Helpers/PCGActorHelpers.cpp b/Engine/Plugins/Experimental/PCG/Source/PCG/Private/Helpers/PCGActorHelpers.cpp
--- a/Engine/Plugins/Experimental/PCG/Source/PCG/Private/Helpers/PCGActorHelpers.cpp
+++ b/Engine/Plugins/Experimental/PCG/Source/PCG/Private/Helpers/PCGActorHelpers.cpp
@@ -295,7 +295,6 @@ AActor* UPCGActorHelpers::SpawnDefaultActor(UWorld* World, TSubclassOf<AActor> A
 
 	FActorSpawnParameters SpawnParams;
 	SpawnParams.Name = MakeUniqueObjectName(World->GetCurrentLevel(), ActorClass, BaseName);
-	SpawnParams.Owner = Parent;
 
 	if (PCGHelpers::IsRuntimeOrPIE())
 	{
diff --git a/Engine/Plugins/Experimental/PCG/Source/PCG/Private/Helpers/PCGHelpers.cpp b/Engine/Plugins/Experimental/PCG/Source/PCG/Private/Helpers/PCGHelpers.cpp
--- a/Engine/Plugins/Experimental/PCG/Source/PCG/Private/Helpers/PCGHelpers.cpp
+++ b/Engine/Plugins/Experimental/PCG/Source/PCG/Private/Helpers/PCGHelpers.cpp
@@ -439,4 +439,53 @@ namespace PCGHelpers
 			return false;
 		}
 	}
+
+#if WITH_EDITOR
+	void GetGeneratedActorsFolderPath(const AActor* InTargetActor, FString& OutFolderPath)
+	{
+		if (!InTargetActor)
+		{
+			return;
+		}
+
+		// Reserves reasonable max string length on stack, overflows to heap if exceeded.
+		TStringBuilderWithBuffer<TCHAR, 1024> GeneratedActorsFolder;
+
+		FName TargetActorFolder = InTargetActor->GetFolderPath();
+		if (TargetActorFolder != NAME_None)
+		{
+			GeneratedActorsFolder << TargetActorFolder.ToString() << "/";
+		}
+
+		GeneratedActorsFolder << InTargetActor->GetActorLabel() << "_Generated";
+		OutFolderPath = GeneratedActorsFolder;
+	}
+#endif
+
+	void AttachToParent(AActor* InActorToAttach, AActor* InParent, EPCGAttachOptions AttachOptions, const FString& InGeneratedPath)
+	{
+		if (!InParent)
+		{
+			return;
+		}
+
+		if (AttachOptions == EPCGAttachOptions::Attached)
+		{
+			InActorToAttach->AttachToActor(InParent, FAttachmentTransformRules::KeepWorldTransform);
+		}
+#if WITH_EDITOR
+		else if(AttachOptions == EPCGAttachOptions::InFolder)
+		{
+			FString DefaultFolderPath;
+
+			if (InGeneratedPath.IsEmpty())
+			{
+				GetGeneratedActorsFolderPath(InParent, DefaultFolderPath);
+			}
+
+			const FString& FolderPath = (InGeneratedPath.IsEmpty() ? DefaultFolderPath : InGeneratedPath);
+			InActorToAttach->SetFolderPath(*FolderPath);
+		}
+#endif
+	}
 }
diff --git a/Engine/Plugins/Experimental/PCG/Source/PCG/Private/PCGComponent.cpp b/Engine/Plugins/Experimental/PCG/Source/PCG/Private/PCGComponent.cpp
--- a/Engine/Plugins/Experimental/PCG/Source/PCG/Private/PCGComponent.cpp
+++ b/Engine/Plugins/Experimental/PCG/Source/PCG/Private/PCGComponent.cpp
@@ -45,6 +45,7 @@
 #include UE_INLINE_GENERATED_CPP_BY_NAME(PCGComponent)
 
 #if WITH_EDITOR
+#include "EditorActorFolders.h"
 #include "ScopedTransaction.h"
 #endif
 
@@ -610,6 +611,33 @@ AActor* UPCGComponent::ClearPCGLink(UClass* TemplateActor)
 		}
 	}
 
+#if WITH_EDITOR
+	// If there is an associated generated folder from this actor, rename it according to the stamp name
+	if (World && NewActor)
+	{
+		FString GeneratedFolderPath;
+		PCGHelpers::GetGeneratedActorsFolderPath(GetOwner(), GeneratedFolderPath);
+		
+		FString GeneratedStampFolder;
+		PCGHelpers::GetGeneratedActorsFolderPath(NewActor, GeneratedStampFolder);
+
+		if (!GeneratedFolderPath.IsEmpty() && !GeneratedStampFolder.IsEmpty())
+		{
+			FFolder GeneratedFolder(FFolder::GetWorldRootFolder(World).GetRootObject(), *GeneratedFolderPath);
+			FFolder StampFolder(FFolder::GetWorldRootFolder(World).GetRootObject(), *GeneratedStampFolder);
+
+			const bool bGeneratedFolderExists = GeneratedFolder.IsValid() && FActorFolders::Get().ContainsFolder(*World, GeneratedFolder);
+			const bool bStampFolderExists = FActorFolders::Get().ContainsFolder(*World, StampFolder);
+
+			// TODO: improve behavior when target stamp folder would exist
+			if (bGeneratedFolderExists && !bStampFolderExists)
+			{
+				FActorFolders::Get().RenameFolderInWorld(*World, GeneratedFolder, StampFolder);
+			}
+		}
+	}
+#endif
+
 	return NewActor;
 }
 
@@ -839,8 +867,65 @@ FPCGTaskId UPCGComponent::CreateCleanupTask(bool bRemoveComponents, const TArray
 			UPCGActorHelpers::DeleteActors(World, Context->ActorsToDelete.Array());
 		}
 
-		if (ThisComponentWeakPtr.Get())
+		if (UPCGComponent* ThisComponent = ThisComponentWeakPtr.Get())
 		{
+#if WITH_EDITOR
+			if (UWorld* ThisWorld = ThisComponent->GetWorld())
+			{
+				// Look for a nested generated results subfolder and remove it if it exists
+				FString FolderPath;
+				PCGHelpers::GetGeneratedActorsFolderPath(ThisComponent->GetOwner(), FolderPath);
+
+				if (!FolderPath.IsEmpty())
+				{
+					FFolder GeneratedFolder(FFolder::GetWorldRootFolder(ThisWorld).GetRootObject(), *FolderPath);
+
+					const bool bFolderExists = GeneratedFolder.IsValid() && FActorFolders::Get().ContainsFolder(*ThisWorld, GeneratedFolder);
+					bool bFoundActors = false;
+
+					if (bFolderExists)
+					{
+						TArray<FName> Paths = { *FolderPath };
+						FActorFolders::ForEachActorInFolders(*ThisWorld, Paths, [&bFoundActors](AActor* InActor)
+						{
+							if (InActor)
+							{
+								bFoundActors = true;
+								return false;
+							}
+							else
+							{
+								return true;
+							}
+						});
+					}
+
+					if (bFolderExists && !bFoundActors)
+					{
+						// Delete all subfolders
+						TArray<FFolder> SubfoldersToDelete;
+						FActorFolders::Get().ForEachFolder(*ThisWorld, [&GeneratedFolder, &SubfoldersToDelete](const FFolder& InFolder)
+						{
+							if (InFolder.IsChildOf(GeneratedFolder))
+							{
+								SubfoldersToDelete.Add(InFolder);
+							}
+
+							return true;
+						});
+
+						for (const FFolder& FolderToDelete : SubfoldersToDelete)
+						{
+							FActorFolders::Get().DeleteFolder(*ThisWorld, FolderToDelete);
+						}
+
+						// Finally, delete folder
+						FActorFolders::Get().DeleteFolder(*ThisWorld, GeneratedFolder);
+					}
+				}
+			}
+#endif
+
 			PCGGeneratedResourcesLogging::LogCreateCleanupTaskFinished(ThisComponentWeakPtr->GeneratedResources);
 		}
 
diff --git a/Engine/Plugins/Experimental/PCG/Source/PCG/Private/PCGManagedResource.cpp b/Engine/Plugins/Experimental/PCG/Source/PCG/Private/PCGManagedResource.cpp
--- a/Engine/Plugins/Experimental/PCG/Source/PCG/Private/PCGManagedResource.cpp
+++ b/Engine/Plugins/Experimental/PCG/Source/PCG/Private/PCGManagedResource.cpp
@@ -127,9 +127,14 @@ bool UPCGManagedActors::MoveResourceToNewActor(AActor* NewActor)
 			continue;
 		}
 
-		Actor->DetachFromActor(FDetachmentTransformRules::KeepWorldTransform);
-		Actor->SetOwner(nullptr);
-		Actor->AttachToActor(NewActor, FAttachmentTransformRules::KeepWorldTransform);
+		const bool bWasAttached = (Actor->GetAttachParentActor() != nullptr);
+
+		if (bWasAttached)
+		{
+			Actor->DetachFromActor(FDetachmentTransformRules::KeepWorldTransform);
+			Actor->SetOwner(nullptr);
+			Actor->AttachToActor(NewActor, FAttachmentTransformRules::KeepWorldTransform);
+		}
 	}
 
 	GeneratedActors.Empty();
diff --git a/Engine/Plugins/Experimental/PCG/Source/PCG/Public/Elements/PCGCreateSpline.h b/Engine/Plugins/Experimental/PCG/Source/PCG/Public/Elements/PCGCreateSpline.h
--- a/Engine/Plugins/Experimental/PCG/Source/PCG/Public/Elements/PCGCreateSpline.h
+++ b/Engine/Plugins/Experimental/PCG/Source/PCG/Public/Elements/PCGCreateSpline.h
@@ -62,8 +62,11 @@ public:
 	UPROPERTY(BlueprintReadWrite, EditAnywhere, Category = Settings, meta = (EditCondition = "!bLinear && bApplyCustomTangents"))
 	FName LeaveTangentAttribute;
 
-	UPROPERTY(BlueprintReadWrite, Category = Settings, meta = (PCG_Overridable))
+	UPROPERTY(meta = (PCG_Overridable))
 	TSoftObjectPtr<AActor> TargetActor;
+
+	UPROPERTY(BlueprintReadWrite, EditAnywhere, Category = Settings, meta = (EditCondition = "Mode==EPCGCreateSplineMode::CreateNewActor", EditConditionHides))
+	EPCGAttachOptions AttachOptions = EPCGAttachOptions::Attached;
 };
 
 class FPCGCreateSplineElement : public FSimplePCGElement
diff --git a/Engine/Plugins/Experimental/PCG/Source/PCG/Public/Elements/PCGCreateTargetActor.h b/Engine/Plugins/Experimental/PCG/Source/PCG/Public/Elements/PCGCreateTargetActor.h
--- a/Engine/Plugins/Experimental/PCG/Source/PCG/Public/Elements/PCGCreateTargetActor.h
+++ b/Engine/Plugins/Experimental/PCG/Source/PCG/Public/Elements/PCGCreateTargetActor.h
@@ -54,16 +54,23 @@ public:
 	UPROPERTY(BlueprintReadOnly, EditAnywhere, Category = Settings, meta = (OnlyPlaceable, DisallowCreateNew))
 	TSubclassOf<AActor> TemplateActorClass = nullptr;
 
-	UPROPERTY(EditAnywhere, BlueprintReadWrite, Instanced, Category = Settings, meta = (ShowInnerProperties))
+	UPROPERTY(EditAnywhere, BlueprintReadWrite, Instanced, Category = Settings, meta = (ShowInnerProperties, EditCondition = "bAllowTemplateActorEditing", EditConditionHides))
 	TObjectPtr<AActor> TemplateActor;
 
-	UPROPERTY(BlueprintReadWrite, Category = Settings, meta = (PCG_Overridable))
+	// TODO: make this InlineEditConditionToggle, not done because property changed event does not propagate correctly so we can't track accurately the need to create the target actor
+	UPROPERTY(EditAnywhere, BlueprintReadWrite, Category = Settings)
+	bool bAllowTemplateActorEditing = false;
+
+	UPROPERTY(BlueprintReadWrite, EditAnywhere, Category = Settings)
+	EPCGAttachOptions AttachOptions = EPCGAttachOptions::Attached;
+
+	UPROPERTY(meta = (PCG_Overridable))
 	TSoftObjectPtr<AActor> RootActor;
 
-	UPROPERTY(BlueprintReadWrite, Category = Settings, meta = (PCG_Overridable))
+	UPROPERTY(meta = (PCG_Overridable))
 	FString ActorLabel;
 
-	UPROPERTY(BlueprintReadWrite, Category = Settings, meta = (PCG_Overridable))
+	UPROPERTY(meta = (PCG_Overridable))
 	FTransform ActorPivot;
 };
 
diff --git a/Engine/Plugins/Experimental/PCG/Source/PCG/Public/Elements/PCGSpawnActor.h b/Engine/Plugins/Experimental/PCG/Source/PCG/Public/Elements/PCGSpawnActor.h
--- a/Engine/Plugins/Experimental/PCG/Source/PCG/Public/Elements/PCGSpawnActor.h
+++ b/Engine/Plugins/Experimental/PCG/Source/PCG/Public/Elements/PCGSpawnActor.h
@@ -86,9 +86,12 @@ public:
 	UPROPERTY(EditAnywhere, BlueprintReadWrite, Category = Settings, meta = (EditCondition = "Option != EPCGSpawnActorOption::CollapseActors", EditConditionHides))
 	TArray<FPCGActorPropertyOverride> ActorOverrides;
 
-	UPROPERTY(BlueprintReadWrite, Category = Settings, meta = (PCG_Overridable))
+	UPROPERTY(meta = (PCG_Overridable))
 	TSoftObjectPtr<AActor> RootActor;
 
+	UPROPERTY(BlueprintReadWrite, EditAnywhere, Category = Settings, meta = (EditCondition = "Option != EPCGSpawnActorOption::CollapseActors", EditConditionHides))
+	EPCGAttachOptions AttachOptions = EPCGAttachOptions::Attached;
+
 	UPROPERTY(BlueprintReadWrite, EditAnywhere, Category = Settings, meta = (InlineEditConditionToggle))
 	bool bSpawnByAttribute = false;
 
diff --git a/Engine/Plugins/Experimental/PCG/Source/PCG/Public/Helpers/PCGHelpers.h b/Engine/Plugins/Experimental/PCG/Source/PCG/Public/Helpers/PCGHelpers.h
--- a/Engine/Plugins/Experimental/PCG/Source/PCG/Public/Helpers/PCGHelpers.h
+++ b/Engine/Plugins/Experimental/PCG/Source/PCG/Public/Helpers/PCGHelpers.h
@@ -63,4 +63,10 @@ namespace PCGHelpers
 
 	/** If hierarchical generation is enabled, returns all relevant grid sizes for the graph, otherwise returns partition grid size from world actor. */
 	PCG_API bool GetGenerationGridSizes(const UPCGGraph* InGraph, const APCGWorldActor* InWorldActor, PCGHiGenGrid::FSizeArray& OutGridSizes, bool& bOutHasUnbounded);
+
+#if WITH_EDITOR
+	PCG_API void GetGeneratedActorsFolderPath(const AActor* InTargetActor, FString& OutFolderPath);
+#endif
+
+	PCG_API void AttachToParent(AActor* InActorToAttach, AActor* InParent, EPCGAttachOptions AttachOptions, const FString& GeneratedPath = FString());
 };
diff --git a/Engine/Plugins/Experimental/PCG/Source/PCG/Public/PCGCommon.h b/Engine/Plugins/Experimental/PCG/Source/PCG/Public/PCGCommon.h
--- a/Engine/Plugins/Experimental/PCG/Source/PCG/Public/PCGCommon.h
+++ b/Engine/Plugins/Experimental/PCG/Source/PCG/Public/PCGCommon.h
@@ -178,3 +178,11 @@ namespace PCGHiGenGrid
 	PCG_API uint32 UninitializedGridSize();
 	PCG_API uint32 UnboundedGridSize();
 }
+
+UENUM()
+enum class EPCGAttachOptions : uint32
+{
+	NotAttached UMETA(Tooltip="Actor will not be attached to the target actor nor placed in an actor folder"),
+	Attached UMETA(Tooltip="Actor will be attached to the target actor in the given node"),
+	InFolder UMETA(Tooltip="Actor will be placed in an actor folder containing the name of the target actor.")
+};
\ No newline at end of file
