[PCG] Support transfering properties in Transfer Attributes node

#rnx

diff --git a/Engine/Plugins/Experimental/PCG/Source/PCG/Private/Elements/PCGAttributeTransferElement.cpp b/Engine/Plugins/Experimental/PCG/Source/PCG/Private/Elements/PCGAttributeTransferElement.cpp
--- a/Engine/Plugins/Experimental/PCG/Source/PCG/Private/Elements/PCGAttributeTransferElement.cpp
+++ b/Engine/Plugins/Experimental/PCG/Source/PCG/Private/Elements/PCGAttributeTransferElement.cpp
@@ -2,15 +2,37 @@
 
 #include "Elements/PCGAttributeTransferElement.h"
 
-#include "Data/PCGPointData.h"
-#include "Data/PCGSpatialData.h"
 #include "PCGContext.h"
+#include "PCGCustomVersion.h"
 #include "PCGPin.h"
+#include "Data/PCGPointData.h"
+#include "Data/PCGSpatialData.h"
+#include "Elements/Metadata/PCGMetadataElementCommon.h"
+#include "Metadata/Accessors/PCGAttributeAccessorHelpers.h"
 
 #include UE_INLINE_GENERATED_CPP_BY_NAME(PCGAttributeTransferElement)
 
 #define LOCTEXT_NAMESPACE "PCGAttributeTransferElement"
 
+void UPCGAttributeTransferSettings::PostLoad()
+{
+	Super::PostLoad();
+
+#if WITH_EDITOR
+	if (SourceAttributeName_DEPRECATED != NAME_None)
+	{
+		SourceAttributeProperty.SetAttributeName(SourceAttributeName_DEPRECATED);
+		SourceAttributeName_DEPRECATED = NAME_None;
+	}
+
+	if (TargetAttributeName_DEPRECATED != NAME_None)
+	{
+		TargetAttributeProperty.SetAttributeName(TargetAttributeName_DEPRECATED);
+		TargetAttributeName_DEPRECATED = NAME_None;
+	}
+#endif // WITH_EDITOR
+}
+
 #if WITH_EDITOR
 FName UPCGAttributeTransferSettings::GetDefaultNodeName() const
 {
@@ -28,9 +50,23 @@ FText UPCGAttributeTransferSettings::GetNodeTooltipText() const
 		"Transfer an attribute from a source metadata to a target data. Only supports Spatial to Spatial and Points to Points, and they need to match.\n\n"
 		" - For Spatial data, number of entries in the metadata should be the same between source and target.\n"
 		" - For Point data, number of points should be the same between source and target.\n\n"
-		"The output will be the target data with the updated metadata. If the TargetAttributeName is None, it will use SourceAttributeName instead."
+		"The output will be the target data with the updated metadata."
 	);
 }
+
+void UPCGAttributeTransferSettings::ApplyDeprecation(UPCGNode* InOutNode)
+{
+	if (DataVersion < FPCGCustomVersion::UpdateTransferAttributeWithSelectors
+		&& TargetAttributeProperty.GetSelection() == EPCGAttributePropertySelection::Attribute
+		&& TargetAttributeProperty.GetAttributeName() == PCGMetadataAttributeConstants::SourceAttributeName)
+	{
+		// Previous behavior of the output target for this node was:
+		// None => SourceName
+		TargetAttributeProperty.SetAttributeName(PCGMetadataAttributeConstants::SourceNameAttributeName);
+	}
+
+	Super::ApplyDeprecation(InOutNode);
+}
 #endif
 
 EPCGDataType UPCGAttributeTransferSettings::GetCurrentPinTypes(const UPCGPin* InPin) const
@@ -44,6 +80,9 @@ EPCGDataType UPCGAttributeTransferSettings::GetCurrentPinTypes(const UPCGPin* In
 
 FName UPCGAttributeTransferSettings::AdditionalTaskName() const
 {
+	const FName SourceAttributeName = SourceAttributeProperty.GetName();
+	const FName TargetAttributeName = TargetAttributeProperty.GetName();
+
 	if (SourceAttributeName != TargetAttributeName && TargetAttributeName != NAME_None)
 	{
 		return FName(FString::Printf(TEXT("%s %s to %s"), *PCGAttributeTransferConstants::NodeName.ToString(), *SourceAttributeName.ToString(), *TargetAttributeName.ToString()));
@@ -90,12 +129,12 @@ bool FPCGAttributeTransferElement::ExecuteInternal(FPCGContext* Context) const
 
 	if (SourceInputs.Num() != 1 || TargetInputs.Num() != 1)
 	{
-		PCGE_LOG(Warning, LogOnly, LOCTEXT("WrongNumberOfInputs", "Source input contains {0} data elements, Target inputs contain {1} data elements, but both should contain precisely 1 data element"));
+		PCGE_LOG(Warning, LogOnly, FText::Format(LOCTEXT("WrongNumberOfInputs", "Source input contains {0} data elements, Target inputs contain {1} data elements, but both should contain precisely 1 data element"), SourceInputs.Num(), TargetInputs.Num()));
 		return true;
 	}
 
-	const UPCGSpatialData* SourceData = Cast<UPCGSpatialData>(SourceInputs[0].Data);
-	const UPCGSpatialData* TargetData = Cast<UPCGSpatialData>(TargetInputs[0].Data);
+	const UPCGSpatialData* SourceData = Cast<const UPCGSpatialData>(SourceInputs[0].Data);
+	const UPCGSpatialData* TargetData = Cast<const UPCGSpatialData>(TargetInputs[0].Data);
 
 	if (!SourceData || !TargetData || SourceData->IsA<UPCGPointData>() != TargetData->IsA<UPCGPointData>())
 	{
@@ -103,124 +142,189 @@ bool FPCGAttributeTransferElement::ExecuteInternal(FPCGContext* Context) const
 		return true;
 	}
 
-	if (!SourceData->Metadata)
+	const UPCGPointData* SourcePointData = Cast<const UPCGPointData>(SourceData);
+	const UPCGPointData* TargetPointData = Cast<const UPCGPointData>(TargetData);
+
+	const bool bIsPointData = SourcePointData && TargetPointData;
+
+	// TODO: Provide some kind of interpolation when point counts are not equal?
+	if (bIsPointData && SourcePointData->GetPoints().Num() != TargetPointData->GetPoints().Num())
 	{
-		PCGE_LOG(Error, GraphAndLog, LOCTEXT("SourceMissingMetadata", "Source does not have metadata"));
+		PCGE_LOG(Error, GraphAndLog, LOCTEXT("MismatchingPointCounts", "Source and target do not have the same number of points"));
 		return true;
 	}
 
-	if (!SourceData->Metadata->HasAttribute(Settings->SourceAttributeName))
+	if (!SourceData->Metadata)
 	{
-		PCGE_LOG(Error, GraphAndLog, FText::Format(LOCTEXT("SourceMissingAttribute", "Source does not have the attribute '{0}'"), FText::FromName(Settings->SourceAttributeName)));
+		PCGE_LOG(Error, GraphAndLog, LOCTEXT("SourceMissingMetadata", "Source does not have metadata"));
 		return true;
 	}
 
-	UPCGData* OutputData = nullptr;
+	const FPCGAttributePropertyInputSelector SourceAttributeProperty = Settings->SourceAttributeProperty.CopyAndFixLast(SourceData);
+	const FPCGAttributePropertyOutputSelector TargetAttributeProperty = Settings->TargetAttributeProperty.CopyAndFixSource(&SourceAttributeProperty, SourceData);
 
-	// We only supports spatial to spatial or point to points, with the same number of points
-	if (SourceData->IsA<UPCGPointData>() && TargetData->IsA<UPCGPointData>())
-	{
-		OutputData = TransferPointToPoint(Context, CastChecked<UPCGPointData>(SourceData), CastChecked<UPCGPointData>(TargetData));
-	}
-	else
+	const FName SourceAttributeName = SourceAttributeProperty.GetName();
+	const FName TargetAttributeName = TargetAttributeProperty.GetName();
+
+	if (SourceAttributeProperty.GetSelection() == EPCGAttributePropertySelection::Attribute && !SourceData->Metadata->HasAttribute(SourceAttributeName))
 	{
-		OutputData = TransferSpatialToSpatial(Context, SourceData, TargetData);
+		PCGE_LOG(Error, GraphAndLog, FText::Format(LOCTEXT("SourceMissingAttribute", "Source does not have the attribute '{0}'"), FText::FromName(SourceAttributeName)));
+		return true;
 	}
 
-	if (OutputData)
+	UPCGSpatialData* OutputData = TargetData->DuplicateData();
+	check(OutputData->Metadata);
+
+	auto AppendOutputData = [Context, OutputData]()
 	{
 		TArray<FPCGTaggedData>& Outputs = Context->OutputData.TaggedData;
 		FPCGTaggedData& Output = Outputs.Emplace_GetRef();
 		Output.Data = OutputData;
-	}
+	};
 
-	return true;
-}
+	// If it is attribute to attribute, just copy the attributes, if they exist and are valid
+	// Only do that if it is really attribute to attribute, without any extra accessor. Any extra accessor will behave as a property.
+	const bool bInputHasAnyExtra = !SourceAttributeProperty.GetExtraNames().IsEmpty();
+	const bool bOutputHasAnyExtra = !TargetAttributeProperty.GetExtraNames().IsEmpty();
+	if (!bInputHasAnyExtra && !bOutputHasAnyExtra && SourceAttributeProperty.GetSelection() == EPCGAttributePropertySelection::Attribute && TargetAttributeProperty.GetSelection() == EPCGAttributePropertySelection::Attribute)
+	{
+		if (SourceData == TargetData && SourceAttributeName == TargetAttributeName)
+		{
+			// Nothing to do if we try to copy an attribute into itself
+			AppendOutputData();
+			return true;
+		}
+
+		const FPCGMetadataAttributeBase* SourceAttribute = SourceData->Metadata->GetConstAttribute(SourceAttributeName);
+		check(SourceAttribute);
+
+		if (bIsPointData)
+		{
+			UPCGPointData* OutPointData = CastChecked<UPCGPointData>(OutputData);
+
+			const TArray<FPCGPoint>& SourcePoints = SourcePointData->GetPoints();
+			TArray<FPCGPoint>& TargetPoints = OutPointData->GetMutablePoints();
+
+			// Making sure the target attribute doesn't exist in the target
+			if (OutPointData->Metadata->HasAttribute(TargetAttributeName))
+			{
+				OutPointData->Metadata->DeleteAttribute(TargetAttributeName);
+			}
+
+			FPCGMetadataAttributeBase* TargetAttribute = OutPointData->Metadata->CopyAttribute(SourceAttribute, TargetAttributeName, /*bKeepParent=*/ false, /*bCopyEntries=*/ false, /*bCopyValues=*/ true);
+
+			if (!TargetAttribute)
+			{
+				PCGE_LOG(Error, GraphAndLog, FText::Format(LOCTEXT("ErrorCreatingTargetAttribute", "Error while creating target attribute '{0}'"), FText::FromName(TargetAttributeName)));
+				return true;
+			}
+
+			// For Point -> Point, the entry keys may not match the source points, so we explicitly write for all the target points
+			for (int32 I = 0; I < SourcePoints.Num(); ++I)
+			{
+				PCGMetadataEntryKey& TargetKey = TargetPoints[I].MetadataEntry;
+				OutPointData->Metadata->InitializeOnSet(TargetKey);
+				check(TargetKey != PCGInvalidEntryKey);
+				TargetAttribute->SetValueFromValueKey(TargetKey, SourceAttribute->GetValueKey(SourcePoints[I].MetadataEntry));
+			}
+		}
+		else
+		{
+			// If we are copying Spatial -> Spatial, we can just copy the attribute across
+			if (!OutputData->Metadata->CopyAttribute(SourceAttribute, TargetAttributeName, /*bKeepParent=*/ false, /*bCopyEntries=*/ true, /*bCopyValues=*/ true))
+			{
+				PCGE_LOG(Error, GraphAndLog, FText::Format(LOCTEXT("FailedCopyToNewAttribute", "Failed to copy to new attribute '{0}'"), FText::FromName(TargetAttributeName)));
+			}
+		}
+
+		AppendOutputData();
 
-UPCGSpatialData* FPCGAttributeTransferElement::TransferSpatialToSpatial(FPCGContext* Context, const UPCGSpatialData* SourceData, const UPCGSpatialData* TargetData) const
-{
-	check(Context && SourceData && TargetData);
+		return true;
+	}
 
-	const UPCGAttributeTransferSettings* Settings = Context->GetInputSettings<UPCGAttributeTransferSettings>();
-	check(Settings);
+	TUniquePtr<const IPCGAttributeAccessor> InputAccessor = PCGAttributeAccessorHelpers::CreateConstAccessor(SourceData, SourceAttributeProperty);
+	TUniquePtr<const IPCGAttributeAccessorKeys> InputKeys = PCGAttributeAccessorHelpers::CreateConstKeys(SourceData, SourceAttributeProperty);
 
-	// When transferring attribute from spatial to spatial, copy the attribute entries and values.
-	// They need to match their number of entries
-	if (SourceData->Metadata->GetItemCountForChild() != TargetData->Metadata->GetItemCountForChild())
+	if (!InputAccessor.IsValid() || !InputKeys.IsValid())
 	{
-		PCGE_LOG(Error, GraphAndLog, LOCTEXT("MismatchingEntryCounts", "Source and target do not have the same number of metadata entries"));
-		return nullptr;
+		PCGE_LOG(Error, GraphAndLog, LOCTEXT("FailedToCreateInputAccessor", "Failed to create input accessor or iterator"));
+		return true;
 	}
 
-	UPCGSpatialData* NewSpatialData = TargetData->DuplicateData();
-
-	const FName TargetAttributeName = (Settings->TargetAttributeName == NAME_None) ? Settings->SourceAttributeName : Settings->TargetAttributeName;
-
-	// Making sure the target attribute doesn't exists in the target
-	if (NewSpatialData->Metadata->HasAttribute(TargetAttributeName))
+	// If the target is an attribute, only create a new one if the attribute doesn't already exist or we have any extra.
+	// If it exist or have any extra, it will try to write to it.
+	if (!bOutputHasAnyExtra && TargetAttributeProperty.GetSelection() == EPCGAttributePropertySelection::Attribute && !OutputData->Metadata->HasAttribute(TargetAttributeName))
 	{
-		NewSpatialData->Metadata->DeleteAttribute(TargetAttributeName);
+		auto CreateAttribute = [OutputData, TargetAttributeName](auto Dummy)
+		{
+			using AttributeType = decltype(Dummy);
+			return PCGMetadataElementCommon::ClearOrCreateAttribute(OutputData->Metadata, TargetAttributeName, AttributeType{}) != nullptr;
+		};
+		
+		if (!PCGMetadataAttribute::CallbackWithRightType(InputAccessor->GetUnderlyingType(), CreateAttribute))
+		{
+			PCGE_LOG(Error, GraphAndLog, FText::Format(LOCTEXT("FailedToCreateNewAttribute", "Failed to create new attribute '{0}'"), FText::FromName(TargetAttributeName)));
+			return true;
+		}
 	}
 
-	NewSpatialData->Metadata->CopyAttribute(SourceData->Metadata, Settings->SourceAttributeName, TargetAttributeName);
+	TUniquePtr<IPCGAttributeAccessor> OutputAccessor = PCGAttributeAccessorHelpers::CreateAccessor(OutputData, TargetAttributeProperty);
+	TUniquePtr<IPCGAttributeAccessorKeys> OutputKeys = PCGAttributeAccessorHelpers::CreateKeys(OutputData, TargetAttributeProperty);
 
-	if (!NewSpatialData->Metadata->HasAttribute(TargetAttributeName))
+	if (!OutputAccessor.IsValid() || !OutputKeys.IsValid())
 	{
-		PCGE_LOG(Error, GraphAndLog, FText::Format(LOCTEXT("ErrorCreatingAttribute", "Error while creating target attribute '{0}'"), FText::FromName(TargetAttributeName)));
-		return nullptr;
+		PCGE_LOG(Error, GraphAndLog, LOCTEXT("FailedToCreateOutputAccessor", "Failed to create output accessor or iterator"));
+		return true;
 	}
 
-	return NewSpatialData;
-}
-
-UPCGPointData* FPCGAttributeTransferElement::TransferPointToPoint(FPCGContext* Context, const UPCGPointData* SourceData, const UPCGPointData* TargetData) const
-{
-	check(Context && SourceData && TargetData);
-
-	const UPCGAttributeTransferSettings* Settings = Context->GetInputSettings<UPCGAttributeTransferSettings>();
-	check(Settings);
-
-	// When transferring attribute point to point, iterate over the source point and set the attribute to the target point.
-	// They need to match their number of points
-	if (SourceData->GetPoints().Num() != TargetData->GetPoints().Num())
+	if (OutputAccessor->IsReadOnly())
 	{
-		PCGE_LOG(Error, GraphAndLog, LOCTEXT("MismatchingPointCounts", "Source and target do not have the same number of points"));
-		return nullptr;
+		PCGE_LOG(Error, GraphAndLog, FText::Format(LOCTEXT("OutputAccessorIsReadOnly", "Attribute/Property '{0}' is read only."), TargetAttributeProperty.GetDisplayText()));
+		return true;
 	}
 
-	UPCGPointData* NewPointData = Cast<UPCGPointData>(TargetData->DuplicateData());
+	// At this point, we are ready.
+	auto Operation = [&InputAccessor, &InputKeys, &OutputAccessor, &OutputKeys, Context](auto Dummy)
+	{
+		using OutputType = decltype(Dummy);
+		OutputType Value{};
 
-	const TArray<FPCGPoint>& SourcePoints = SourceData->GetPoints();
-	TArray<FPCGPoint>& TargetPoints = NewPointData->GetMutablePoints();
+		EPCGAttributeAccessorFlags Flags = EPCGAttributeAccessorFlags::AllowBroadcast | EPCGAttributeAccessorFlags::AllowConstructible;
 
-	const FName TargetAttributeName = (Settings->TargetAttributeName == NAME_None) ? Settings->SourceAttributeName : Settings->TargetAttributeName;
+		const int32 NumberOfElements = OutputKeys->GetNum();
+		constexpr int32 ChunkSize = 256;
 
-	// Making sure the target attribute doesn't exist in the target
-	if (NewPointData->Metadata->HasAttribute(TargetAttributeName))
-	{
-		NewPointData->Metadata->DeleteAttribute(TargetAttributeName);
-	}
+		TArray<OutputType, TInlineAllocator<ChunkSize>> TempValues;
+		TempValues.SetNum(ChunkSize);
 
-	const FPCGMetadataAttributeBase* SourceAttribute = SourceData->Metadata->GetConstAttribute(Settings->SourceAttributeName);
-	FPCGMetadataAttributeBase* TargetAttribute = NewPointData->Metadata->CopyAttribute(SourceAttribute, TargetAttributeName, /*bKeepParent=*/ false, /*bCopyEntries=*/ false, /*bCopyValues=*/ true);
+		const int32 NumberOfIterations = (NumberOfElements + ChunkSize - 1) / ChunkSize;
 
-	if (!TargetAttribute)
-	{
-		PCGE_LOG(Error, GraphAndLog, FText::Format(LOCTEXT("ErrorCreatingTargetAttribute", "Error while creating target attribute '{0}'"), FText::FromName(TargetAttributeName)));
-		return nullptr;
-	}
+		for (int32 i = 0; i < NumberOfIterations; ++i)
+		{
+			const int32 StartIndex = i * ChunkSize;
+			const int32 Range = FMath::Min(NumberOfElements - StartIndex, ChunkSize);
+			TArrayView<OutputType> View(TempValues.GetData(), Range);
 
-	for (int32 i = 0; i < SourcePoints.Num(); ++i)
-	{
-		PCGMetadataEntryKey& TargetKey = TargetPoints[i].MetadataEntry;
-		NewPointData->Metadata->InitializeOnSet(TargetKey);
+			if (!InputAccessor->GetRange<OutputType>(View, StartIndex, *InputKeys, Flags)
+				|| !OutputAccessor->SetRange<OutputType>(View, StartIndex, *OutputKeys, Flags))
+			{
+				PCGE_LOG_C(Error, GraphAndLog, Context, LOCTEXT("ConversionFailed", "Source attribute/property cannot be converted to target attribute/property"));
+				return false;
+			}
+		}
 
-		check(TargetKey != PCGInvalidEntryKey);
+		return true;
+	};
 
-		TargetAttribute->SetValueFromValueKey(TargetKey, SourceAttribute->GetValueKey(SourcePoints[i].MetadataEntry));
+	if (!PCGMetadataAttribute::CallbackWithRightType(OutputAccessor->GetUnderlyingType(), Operation))
+	{
+		PCGE_LOG(Error, GraphAndLog, LOCTEXT("ErrorGettingSettingValues", "Error while getting/setting values"));
+		return true;
 	}
 
-	return NewPointData;
+	AppendOutputData();
+
+	return true;
 }
 
 #undef LOCTEXT_NAMESPACE
diff --git a/Engine/Plugins/Experimental/PCG/Source/PCG/Public/Elements/PCGAttributeTransferElement.h b/Engine/Plugins/Experimental/PCG/Source/PCG/Public/Elements/PCGAttributeTransferElement.h
--- a/Engine/Plugins/Experimental/PCG/Source/PCG/Public/Elements/PCGAttributeTransferElement.h
+++ b/Engine/Plugins/Experimental/PCG/Source/PCG/Public/Elements/PCGAttributeTransferElement.h
@@ -23,9 +23,6 @@ namespace PCGAttributeTransferConstants
 *  - For Spatial data, number of entries in the metadata should be the same between source and target.
 *  - For Point data, number of points should be the same between source and target.
 * 
-* The output will be the target data with the updated metadata.
-* If the TargetAttributeName is None, it will use SourceAttributeName instead.
-* 
 * To do the same but with a Source param data, use CreateAttribute.
 */
 UCLASS(BlueprintType, ClassGroup = (Procedural))
@@ -34,6 +31,10 @@ class PCG_API UPCGAttributeTransferSettings : public UPCGSettings
 	GENERATED_BODY()
 
 public:
+	// ~Begin UObject interface
+	virtual void PostLoad() override;
+	// ~End UObject interface
+
 	//~Begin UPCGSettings interface
 #if WITH_EDITOR
 	virtual FName GetDefaultNodeName() const override;
@@ -41,6 +42,7 @@ public:
 	virtual EPCGSettingsType GetType() const override { return EPCGSettingsType::Metadata; }
 	virtual FText GetNodeTooltipText() const override;
 	virtual bool HasDynamicPins() const override { return true; }
+	virtual void ApplyDeprecation(UPCGNode* InOutNode) override;
 #endif
 	virtual EPCGDataType GetCurrentPinTypes(const UPCGPin* InPin) const override;
 	virtual FName AdditionalTaskName() const override;
@@ -53,18 +55,22 @@ protected:
 
 public:
 	UPROPERTY(BlueprintReadWrite, EditAnywhere, Category = Settings)
-	FName SourceAttributeName;
+	FPCGAttributePropertyInputSelector SourceAttributeProperty;
 
 	UPROPERTY(BlueprintReadWrite, EditAnywhere, Category = Settings)
-	FName TargetAttributeName;
-};
+	FPCGAttributePropertyOutputSelector TargetAttributeProperty;
 
+#if WITH_EDITORONLY_DATA
+	UPROPERTY()
+	FName SourceAttributeName_DEPRECATED;
+
+	UPROPERTY()
+	FName TargetAttributeName_DEPRECATED;
+#endif
+};
 
 class FPCGAttributeTransferElement : public FSimplePCGElement
 {
 protected:
 	virtual bool ExecuteInternal(FPCGContext* Context) const override;
-
-	UPCGSpatialData* TransferSpatialToSpatial(FPCGContext* Context, const UPCGSpatialData* SourceData, const UPCGSpatialData* TargetData) const;
-	UPCGPointData* TransferPointToPoint(FPCGContext* Context, const UPCGPointData* SourceData, const UPCGPointData* TargetData) const;
 };
diff --git a/Engine/Plugins/Experimental/PCG/Source/PCG/Public/PCGCustomVersion.h b/Engine/Plugins/Experimental/PCG/Source/PCG/Public/PCGCustomVersion.h
--- a/Engine/Plugins/Experimental/PCG/Source/PCG/Public/PCGCustomVersion.h
+++ b/Engine/Plugins/Experimental/PCG/Source/PCG/Public/PCGCustomVersion.h
@@ -53,6 +53,9 @@ struct PCG_API FPCGCustomVersion
 		// Update AddAttribute with selectors
 		UpdateAddAttributeWithSelectors = 14,
 
+		// Update TransferAttribute with selectors
+		UpdateTransferAttributeWithSelectors = 15,
+
 		// -----<new versions can be added above this line>-------------------------------------------------
 		VersionPlusOne,
 		LatestVersion = VersionPlusOne - 1
