Fixed missing landscape metadata when using world ray hit.
Added hit actor & physical material support as well.

diff --git a/Engine/Plugins/Experimental/PCG/Source/PCG/Private/Data/PCGWorldData.cpp b/Engine/Plugins/Experimental/PCG/Source/PCG/Private/Data/PCGWorldData.cpp
--- a/Engine/Plugins/Experimental/PCG/Source/PCG/Private/Data/PCGWorldData.cpp
+++ b/Engine/Plugins/Experimental/PCG/Source/PCG/Private/Data/PCGWorldData.cpp
@@ -59,6 +59,8 @@ bool UPCGWorldVolumetricData::SamplePoint(const FTransform& InTransform, const F
 
 	check(World.IsValid());
 
+	FPCGMetadataAttribute<FString>* ActorOverlappedAttribute = ((OutMetadata && QueryParams.bGetReferenceToActorHit) ? OutMetadata->GetMutableTypedAttribute<FString>(PCGPointDataConstants::ActorReferenceAttribute) : nullptr);
+
 	FCollisionObjectQueryParams ObjectQueryParams(QueryParams.CollisionChannel);
 	FCollisionShape CollisionShape = FCollisionShape::MakeBox(InBounds.GetExtent() * InTransform.GetScale3D());
 	FCollisionQueryParams Params; // TODO: apply properties from the settings when/if they exist
@@ -136,6 +138,13 @@ bool UPCGWorldVolumetricData::SamplePoint(const FTransform& InTransform, const F
 			OutPoint = FPCGPoint(InTransform, 1.0f, 0);
 			UPCGBlueprintHelpers::SetSeedFromPosition(OutPoint);
 			OutPoint.SetLocalBounds(InBounds);
+
+			if (ActorOverlappedAttribute && OverlappedComponent->GetOwner())
+			{
+				OutMetadata->InitializeOnSet(OutPoint.MetadataEntry);
+				ActorOverlappedAttribute->SetValue(OutPoint.MetadataEntry, FSoftObjectPath(OverlappedComponent->GetOwner()).ToString());
+			}
+
 			return true;
 		}
 		else
@@ -225,10 +234,14 @@ bool UPCGWorldRayHitData::SamplePoint(const FTransform& InTransform, const FBox&
 	// TODO: This seems to be a projection - along a direction. I suspect that UPCGWorldVolumetricData is the SamplePoint(), and this is the ProjectPoint() (in a direction)?
 	check(World.IsValid());
 
+	FPCGMetadataAttribute<FString>* ActorHitAttribute = ((OutMetadata && QueryParams.bGetReferenceToActorHit) ? OutMetadata->GetMutableTypedAttribute<FString>(PCGPointDataConstants::ActorReferenceAttribute) : nullptr);
+	FPCGMetadataAttribute<FString>* PhysicalMaterialAttribute = ((OutMetadata && QueryParams.bGetReferenceToPhysicalMaterial) ? OutMetadata->GetMutableTypedAttribute<FString>(PCGWorldRayHitConstants::PhysicalMaterialReferenceAttribute) : nullptr);
+
 	// Todo: consider prebuilding this
 	FCollisionObjectQueryParams ObjectQueryParams(QueryParams.CollisionChannel);
 	FCollisionQueryParams Params; // TODO: apply properties from the settings when/if they exist
 	Params.bTraceComplex = QueryParams.bTraceComplex;
+	Params.bReturnPhysicalMaterial = (PhysicalMaterialAttribute != nullptr);
 
 	// Project the InTransform location on the ray origin plane
 	const FVector PointLocation = InTransform.GetLocation();
@@ -330,6 +343,18 @@ bool UPCGWorldRayHitData::SamplePoint(const FTransform& InTransform, const FBox&
 			}
 		}
 
+		if (ActorHitAttribute && HitComponent->GetOwner())
+		{
+			OutMetadata->InitializeOnSet(OutPoint.MetadataEntry);
+			ActorHitAttribute->SetValue(OutPoint.MetadataEntry, FSoftObjectPath(HitComponent->GetOwner()).ToString());
+		}
+
+		if (PhysicalMaterialAttribute && Hit.PhysMaterial.Get())
+		{
+			OutMetadata->InitializeOnSet(OutPoint.MetadataEntry);
+			PhysicalMaterialAttribute->SetValue(OutPoint.MetadataEntry, FSoftObjectPath(Hit.PhysMaterial.Get()).ToString());
+		}
+
 		return true;
 	}
 
diff --git a/Engine/Plugins/Experimental/PCG/Source/PCG/Private/Elements/PCGWorldQuery.cpp b/Engine/Plugins/Experimental/PCG/Source/PCG/Private/Elements/PCGWorldQuery.cpp
--- a/Engine/Plugins/Experimental/PCG/Source/PCG/Private/Elements/PCGWorldQuery.cpp
+++ b/Engine/Plugins/Experimental/PCG/Source/PCG/Private/Elements/PCGWorldQuery.cpp
@@ -4,9 +4,13 @@
 
 #include "PCGComponent.h"
 #include "PCGContext.h"
+#include "PCGSubsystem.h"
+#include "Data/PCGPointData.h"
+#include "Grid/PCGLandscapeCache.h"
 #include "Helpers/PCGHelpers.h"
 
 #include "GameFramework/Actor.h"
+#include "LandscapeProxy.h"
 
 #include UE_INLINE_GENERATED_CPP_BY_NAME(PCGWorldQuery)
 
@@ -129,6 +133,74 @@ bool FPCGWorldRayHitQueryElement::ExecuteInternal(FPCGContext* Context) const
 	Data->QueryParams.Initialize();
 	Data->OriginatingComponent = Context->SourceComponent;
 
+	if (Data->QueryParams.bGetReferenceToActorHit && Data->Metadata)
+	{
+		Data->Metadata->FindOrCreateAttribute(PCGPointDataConstants::ActorReferenceAttribute, FString(), /*bAllowInterpolation=*/false, /*bOverrideParent=*/false);
+	}
+	else
+	{
+		Data->QueryParams.bGetReferenceToActorHit = false;
+	}
+
+	if (Data->QueryParams.bGetReferenceToPhysicalMaterial && Data->Metadata)
+	{
+		Data->Metadata->FindOrCreateAttribute(PCGWorldRayHitConstants::PhysicalMaterialReferenceAttribute, FString(), /*bAllowInterpolation=*/false, /*bOverrideParent=*/false);
+	}
+	else
+	{
+		Data->QueryParams.bGetReferenceToPhysicalMaterial = false;
+	}
+
+	bool bHasLandscapeMetadata = false;
+	if (QueryParams.bApplyMetadataFromLandscape && Data->Metadata)
+	{
+		UPCGSubsystem* PCGSubsystem = UPCGSubsystem::GetInstance(World);
+		UPCGLandscapeCache* LandscapeCache = PCGSubsystem ? PCGSubsystem->GetLandscapeCache() : nullptr;
+
+		if (LandscapeCache)
+		{
+			TFunction<bool(const AActor*)> BoundsCheck = [](const AActor*) -> bool { return true; };
+			TFunction<bool(const AActor*)> SelfIgnoreCheck = [](const AActor*) -> bool { return true; };
+
+			FPCGActorSelectorSettings ActorSelector;
+			ActorSelector.ActorFilter = EPCGActorFilter::AllWorldActors;
+			ActorSelector.ActorSelection = EPCGActorSelection::ByClass;
+			ActorSelector.ActorSelectionClass = ALandscapeProxy::StaticClass();
+			ActorSelector.bSelectMultiple = true;
+
+			if (Data->Bounds.IsValid)
+			{
+				BoundsCheck = [Data, Component=Context->SourceComponent.Get()](const AActor* OtherActor) -> bool
+				{
+					const FBox OtherActorBounds = OtherActor ? PCGHelpers::GetGridBounds(OtherActor, Component) : FBox(EForceInit::ForceInit);
+					return OtherActorBounds.Intersect(Data->Bounds);
+				};
+			}
+	
+			TArray<AActor*> LandscapeActors = PCGActorSelector::FindActors(ActorSelector, Context->SourceComponent.Get(), BoundsCheck, SelfIgnoreCheck);
+			for (AActor* Landscape : LandscapeActors)
+			{
+				if (ALandscapeProxy* LandscapeProxy = Cast<ALandscapeProxy>(Landscape))
+				{
+					const TArray<FName> Layers = LandscapeCache->GetLayerNames(LandscapeProxy);
+					for (FName Layer : Layers)
+					{
+						if (!Data->Metadata->HasAttribute(Layer))
+						{
+							Data->Metadata->CreateFloatAttribute(Layer, 0.0f, /*bAllowInterpolation=*/true);
+							bHasLandscapeMetadata = true;
+						}
+					}
+				}
+			}
+		}
+	}
+
+	if (!bHasLandscapeMetadata)
+	{
+		Data->QueryParams.bApplyMetadataFromLandscape = false;
+	}
+
 	FPCGTaggedData& Output = Context->OutputData.TaggedData.Emplace_GetRef();
 	Output.Data = Data;
 
diff --git a/Engine/Plugins/Experimental/PCG/Source/PCG/Public/Data/PCGWorldData.h b/Engine/Plugins/Experimental/PCG/Source/PCG/Public/Data/PCGWorldData.h
--- a/Engine/Plugins/Experimental/PCG/Source/PCG/Public/Data/PCGWorldData.h
+++ b/Engine/Plugins/Experimental/PCG/Source/PCG/Public/Data/PCGWorldData.h
@@ -23,6 +23,11 @@ enum class EPCGWorldQueryFilterByTag
 	ExcludeTagged
 };
 
+namespace PCGWorldRayHitConstants
+{
+	const FName PhysicalMaterialReferenceAttribute = TEXT("PhysicalMaterial");
+}
+
 USTRUCT(BlueprintType)
 struct FPCGWorldCommonQueryParams
 {
@@ -51,6 +56,9 @@ struct FPCGWorldCommonQueryParams
 	UPROPERTY(BlueprintReadWrite, EditAnywhere, Category = "Data|Filtering", meta = (PCG_Overridable))
 	bool bIgnoreLandscapeHits = false;
 
+	UPROPERTY(BlueprintReadWrite, EditAnywhere, Category = Data, meta = (PCG_Overridable))
+	bool bGetReferenceToActorHit = false;
+
 	// Not exposed, will be filled in when initializing this
 	UPROPERTY()
 	TSet<FName> ParsedActorTagsList;
@@ -131,6 +139,9 @@ struct FPCGWorldRayHitQueryParams : public FPCGWorldCommonQueryParams
 
 	UPROPERTY(BlueprintReadWrite, EditAnywhere, Category = Data, meta = (PCG_Overridable))
 	bool bApplyMetadataFromLandscape = false;
+
+	UPROPERTY(BlueprintReadWrite, EditAnywhere, Category = Data, meta = (PCG_Overridable))
+	bool bGetReferenceToPhysicalMaterial = false;
 };
 
 /** Executes collision queries against world collision. */
