Fix non-unity build (missing include)

#rnx

diff --git a/Engine/Plugins/Experimental/PCG/Source/PCG/Private/Data/PCGWorldData.cpp b/Engine/Plugins/Experimental/PCG/Source/PCG/Private/Data/PCGWorldData.cpp
--- a/Engine/Plugins/Experimental/PCG/Source/PCG/Private/Data/PCGWorldData.cpp
+++ b/Engine/Plugins/Experimental/PCG/Source/PCG/Private/Data/PCGWorldData.cpp
@@ -16,6 +16,7 @@
 #include "LandscapeProxy.h"
 #include "Components/BrushComponent.h"
 #include "Engine/World.h"
+#include "PhysicalMaterials/PhysicalMaterial.h"
 
 #include UE_INLINE_GENERATED_CPP_BY_NAME(PCGWorldData)
 
