[PCG] Use raw ptr instead of TObjectPtr in set, to counter TObjectPtr going out of scope before being unregister.

* This should fix the issue with invalid resources crash (for Jimmy)
* This should fix the issue with level instances crash (for Quentin)

#rnx

diff --git a/Engine/Plugins/Experimental/PCG/Source/PCG/Private/Grid/PCGComponentOctree.cpp b/Engine/Plugins/Experimental/PCG/Source/PCG/Private/Grid/PCGComponentOctree.cpp
--- a/Engine/Plugins/Experimental/PCG/Source/PCG/Private/Grid/PCGComponentOctree.cpp
+++ b/Engine/Plugins/Experimental/PCG/Source/PCG/Private/Grid/PCGComponentOctree.cpp
@@ -167,11 +167,11 @@ bool FPCGComponentOctreeAndMap::RemoveComponent(UPCGComponent* InComponent)
 	return true;
 }
 
-TSet<TObjectPtr<UPCGComponent>> FPCGComponentOctreeAndMap::GetAllComponents() const
+TSet<UPCGComponent*> FPCGComponentOctreeAndMap::GetAllComponents() const
 {
 	FReadScopeLock ReadLock(Lock);
 
-	TSet<TObjectPtr<UPCGComponent>> Components;
+	TSet<UPCGComponent*> Components;
 	ComponentToIdMap.GetKeys(Components);
 
 	return Components;
diff --git a/Engine/Plugins/Experimental/PCG/Source/PCG/Private/PCGActorAndComponentMapping.cpp b/Engine/Plugins/Experimental/PCG/Source/PCG/Private/PCGActorAndComponentMapping.cpp
--- a/Engine/Plugins/Experimental/PCG/Source/PCG/Private/PCGActorAndComponentMapping.cpp
+++ b/Engine/Plugins/Experimental/PCG/Source/PCG/Private/PCGActorAndComponentMapping.cpp
@@ -54,7 +54,7 @@ UPCGActorAndComponentMapping::UPCGActorAndComponentMapping(UPCGSubsystem* InPCGS
 
 void UPCGActorAndComponentMapping::Tick()
 {
-	TSet<TObjectPtr<UPCGComponent>> ComponentToUnregister;
+	TSet<UPCGComponent*> ComponentToUnregister;
 	{
 		FScopeLock Lock(&DelayedComponentToUnregisterLock);
 		ComponentToUnregister = MoveTemp(DelayedComponentToUnregister);
@@ -296,14 +296,20 @@ void UPCGActorAndComponentMapping::UnregisterPCGComponent(UPCGComponent* InCompo
 			DelayedComponentToUnregister.Add(InComponent);
 			return;
 		}
+
+#if WITH_EDITOR
+		UnregisterTracking(InComponent);
+#endif // WITH_EDITOR
 	}
 
 	UnregisterPartitionedPCGComponent(InComponent);
 	UnregisterNonPartitionedPCGComponent(InComponent);
 
-#if WITH_EDITOR
-	UnregisterTracking(InComponent);
-#endif // WITH_EDITOR
+	FScopeLock Lock(&DelayedComponentToUnregisterLock);
+	if (DelayedComponentToUnregister.Contains(InComponent))
+	{
+		DelayedComponentToUnregister.Remove(InComponent);
+	}
 }
 
 void UPCGActorAndComponentMapping::UnregisterPartitionedPCGComponent(UPCGComponent* InComponent)
@@ -552,19 +558,19 @@ void UPCGActorAndComponentMapping::DeleteMappingPCGComponentPartitionActor(UPCGC
 	}
 }
 
-TSet<TObjectPtr<UPCGComponent>> UPCGActorAndComponentMapping::GetAllRegisteredPartitionedComponents() const
+TSet<UPCGComponent*> UPCGActorAndComponentMapping::GetAllRegisteredPartitionedComponents() const
 {
 	return PartitionedOctree.GetAllComponents();
 }
 
-TSet<TObjectPtr<UPCGComponent>> UPCGActorAndComponentMapping::GetAllRegisteredNonPartitionedComponents() const
+TSet<UPCGComponent*> UPCGActorAndComponentMapping::GetAllRegisteredNonPartitionedComponents() const
 {
 	return NonPartitionedOctree.GetAllComponents();
 }
 
-TSet<TObjectPtr<UPCGComponent>> UPCGActorAndComponentMapping::GetAllRegisteredComponents() const
+TSet<UPCGComponent*> UPCGActorAndComponentMapping::GetAllRegisteredComponents() const
 {
-	TSet<TObjectPtr<UPCGComponent>> Res = GetAllRegisteredPartitionedComponents();
+	TSet<UPCGComponent*> Res = GetAllRegisteredPartitionedComponents();
 	Res.Append(GetAllRegisteredNonPartitionedComponents());
 	return Res;
 }
@@ -588,7 +594,7 @@ void UPCGActorAndComponentMapping::RegisterOrUpdateTracking(UPCGComponent* InCom
 
 	// Components owner needs to be always tracked
 	RegisterActor(ComponentOwner);
-	TSet<TObjectPtr<UPCGComponent>>& AllComponents = AlwaysTrackedActorsToComponentsMap.FindOrAdd(ComponentOwner);
+	TSet<UPCGComponent*>& AllComponents = AlwaysTrackedActorsToComponentsMap.FindOrAdd(ComponentOwner);
 	AllComponents.Add(InComponent);
 
 	UWorld* World = PCGSubsystem->GetWorld();
@@ -732,9 +738,9 @@ void UPCGActorAndComponentMapping::UnregisterTracking(UPCGComponent* InComponent
 	}
 
 	// We also need to untrack actors that doesn't have any component that tracks them.
-	auto ShouldBeRemoved = [](const TObjectKey<AActor>& InActor, TMap<TObjectKey<AActor>, TSet<TObjectPtr<UPCGComponent>>>& InMap)
+	auto ShouldBeRemoved = [](const TObjectKey<AActor>& InActor, TMap<TObjectKey<AActor>, TSet<UPCGComponent*>>& InMap)
 	{
-		TSet<TObjectPtr<UPCGComponent>>* RegisteredComponents = InMap.Find(InActor);
+		TSet<UPCGComponent*>* RegisteredComponents = InMap.Find(InActor);
 		return !RegisteredComponents || RegisteredComponents->IsEmpty();
 	};
 
@@ -835,13 +841,22 @@ bool UPCGActorAndComponentMapping::AddOrUpdateTrackedActor(AActor* InActor)
 	}
 
 	// Gather all components, and check if they want to track this one
-	TSet<TObjectPtr<UPCGComponent>> AllComponents = GetAllRegisteredComponents();
+	TSet<UPCGComponent*> AllComponents = GetAllRegisteredComponents();
 
-	TSet<TObjectPtr<UPCGComponent>>* CulledTrackedComponents = nullptr;
-	TSet<TObjectPtr<UPCGComponent>>* AlwaysTrackedComponents = nullptr;
+	TSet<UPCGComponent*>* CulledTrackedComponents = nullptr;
+	TSet<UPCGComponent*>* AlwaysTrackedComponents = nullptr;
 	
 	for (UPCGComponent* PCGComponent : AllComponents)
 	{
+		// Making sure that they are not currently waiting to die
+		{
+			FScopeLock Lock(&DelayedComponentToUnregisterLock);
+			if (DelayedComponentToUnregister.Contains(PCGComponent))
+			{
+				continue;
+			}
+		}
+
 		bool bTrackingIsCulled = false;
 		if (PCGComponent && PCGComponent->IsActorTracked(InActor, bTrackingIsCulled))
 		{
@@ -873,12 +888,18 @@ bool UPCGActorAndComponentMapping::AddOrUpdateTrackedActor(AActor* InActor)
 		RegisterActor(InActor);
 		return true;
 	}
-	else
+	else if (CulledTrackedActorsToComponentsMap.Contains(InActor) || AlwaysTrackedActorsToComponentsMap.Contains(InActor))
 	{
-		// Do some cleanup. We will force the refresh here, so return false to make sure we don't refresh it twice.
+		// Do some cleanup if the actor was tracked. 
+		// We will force the refresh here, so return false to make sure we don't refresh it twice.
 		OnActorDeleted(InActor);
 		return false;
 	}
+	else
+	{
+		// If it is not tracked, and should not be tracked, just do nothing.
+		return false;
+	}
 }
 
 void UPCGActorAndComponentMapping::RegisterActor(AActor* InActor)
@@ -1073,7 +1094,7 @@ void UPCGActorAndComponentMapping::OnActorChanged(AActor* InActor, bool bInHasMo
 	// Check if we have a change of tag too
 	TSet<FName> RemovedTags = TempTrackedActorTags.Difference(TSet<FName>(InActor->Tags));
 
-	if (TSet<TObjectPtr<UPCGComponent>>* CulledTrackedComponents = CulledTrackedActorsToComponentsMap.Find(InActor))
+	if (TSet<UPCGComponent*>* CulledTrackedComponents = CulledTrackedActorsToComponentsMap.Find(InActor))
 	{
 		// Not const, since it will be updated with old actor bounds
 		FBox ActorBounds = PCGActorAndComponentMapping::GetActorBounds(InActor);
@@ -1160,7 +1181,7 @@ void UPCGActorAndComponentMapping::OnActorChanged(AActor* InActor, bool bInHasMo
 	}
 
 	// Finally, dirty all components that always track this actor that are not yet notified.
-	if (TSet<TObjectPtr<UPCGComponent>>* AlwaysTrackedComponents = AlwaysTrackedActorsToComponentsMap.Find(InActor))
+	if (TSet<UPCGComponent*>* AlwaysTrackedComponents = AlwaysTrackedActorsToComponentsMap.Find(InActor))
 	{
 		TRACE_CPUPROFILER_EVENT_SCOPE(UPCGActorAndComponentMapping::OnActorChanged::AlwaysTrackedUpdate);
 
diff --git a/Engine/Plugins/Experimental/PCG/Source/PCG/Private/PCGComponent.cpp b/Engine/Plugins/Experimental/PCG/Source/PCG/Private/PCGComponent.cpp
--- a/Engine/Plugins/Experimental/PCG/Source/PCG/Private/PCGComponent.cpp
+++ b/Engine/Plugins/Experimental/PCG/Source/PCG/Private/PCGComponent.cpp
@@ -70,6 +70,18 @@ UPCGComponent::UPCGComponent(const FObjectInitializer& InObjectInitializer)
 #endif // WITH_EDITOR
 }
 
+UPCGComponent::~UPCGComponent()
+{
+#if WITH_EDITOR
+	// For the special case where a component is part of a reconstruction script (from a BP),
+	// but gets destroyed immediately, we need to force the unregistering. 
+	if (UPCGSubsystem* PCGSubsystem = UPCGSubsystem::GetSubsystemForCurrentWorld())
+	{
+		PCGSubsystem->UnregisterPCGComponent(this, /*bForce=*/true);
+	}
+#endif // WITH_EDITOR
+}
+
 bool UPCGComponent::CanPartition() const
 {
 	// Support/Force partitioning on non-PCG partition actors in WP worlds.
@@ -946,7 +958,7 @@ void UPCGComponent::CleanupUnusedManagedResources()
 		check(!GeneratedResourcesInaccessible);
 		for (int32 ResourceIndex = GeneratedResources.Num() - 1; ResourceIndex >= 0; --ResourceIndex)
 		{
-			UPCGManagedResource* Resource = GeneratedResources[ResourceIndex];
+			UPCGManagedResource* Resource = GetValid(GeneratedResources[ResourceIndex]);
 
 			PCGGeneratedResourcesLogging::LogCleanupUnusedManagedResourcesResource(Resource);
 
diff --git a/Engine/Plugins/Experimental/PCG/Source/PCG/Private/PCGSubsystem.cpp b/Engine/Plugins/Experimental/PCG/Source/PCG/Private/PCGSubsystem.cpp
--- a/Engine/Plugins/Experimental/PCG/Source/PCG/Private/PCGSubsystem.cpp
+++ b/Engine/Plugins/Experimental/PCG/Source/PCG/Private/PCGSubsystem.cpp
@@ -38,28 +38,7 @@ namespace PCGSubsystemConsole
 		TEXT("Clears the PCG results cache."),
 		FConsoleCommandDelegate::CreateLambda([]()
 			{
-				UWorld* World = nullptr;
-
-#if WITH_EDITOR
-				if (GEditor)
-				{
-					if (GEditor->PlayWorld)
-					{
-						World = GEditor->PlayWorld;
-					}
-					else
-					{
-						World = GEditor->GetEditorWorldContext().World();
-					}
-				}
-				else
-#endif
-				if (GEngine)
-				{
-					World = GEngine->GetCurrentPlayWorld();
-				}
-
-				if (UPCGSubsystem* PCGSubsystem = UPCGSubsystem::GetInstance(World))
+				if (UPCGSubsystem* PCGSubsystem = UPCGSubsystem::GetSubsystemForCurrentWorld())
 				{
 					PCGSubsystem->FlushCache();
 				}
@@ -71,8 +50,7 @@ namespace PCGSubsystemConsole
 		TEXT("Builds the landscape cache in the current world."),
 		FConsoleCommandDelegate::CreateLambda([]()
 			{
-				UWorld* World = (GEditor ? (GEditor->PlayWorld ? GEditor->PlayWorld.Get() : GEditor->GetEditorWorldContext().World()) : (GEngine ? GEngine->GetCurrentPlayWorld() : nullptr));
-				if (UPCGSubsystem* PCGSubsystem = UPCGSubsystem::GetInstance(World))
+				if (UPCGSubsystem* PCGSubsystem = UPCGSubsystem::GetSubsystemForCurrentWorld())
 				{
 					PCGSubsystem->BuildLandscapeCache();
 				}
@@ -83,10 +61,9 @@ namespace PCGSubsystemConsole
 		TEXT("Clear the landscape cache in the current world."),
 		FConsoleCommandDelegate::CreateLambda([]()
 			{
-				UWorld* World = (GEditor ? (GEditor->PlayWorld ? GEditor->PlayWorld.Get() : GEditor->GetEditorWorldContext().World()) : (GEngine ? GEngine->GetCurrentPlayWorld() : nullptr));
-				if(World && World->GetSubsystem<UPCGSubsystem>())
+				if (UPCGSubsystem* PCGSubsystem = UPCGSubsystem::GetSubsystemForCurrentWorld())
 				{
-					World->GetSubsystem<UPCGSubsystem>()->ClearLandscapeCache();
+					PCGSubsystem->ClearLandscapeCache();
 				}
 			}));
 
@@ -116,6 +93,32 @@ UPCGSubsystem::UPCGSubsystem()
 #endif // WITH_EDITOR
 }
 
+UPCGSubsystem* UPCGSubsystem::GetSubsystemForCurrentWorld()
+{
+	UWorld* World = nullptr;
+
+#if WITH_EDITOR
+	if (GEditor)
+	{
+		if (GEditor->PlayWorld)
+		{
+			World = GEditor->PlayWorld;
+		}
+		else
+		{
+			World = GEditor->GetEditorWorldContext().World();
+		}
+	}
+	else
+#endif
+	if (GEngine)
+	{
+		World = GEngine->GetCurrentPlayWorld();
+	}
+
+	return UPCGSubsystem::GetInstance(World);
+}
+
 void UPCGSubsystem::Deinitialize()
 {
 	// Cancel all tasks
diff --git a/Engine/Plugins/Experimental/PCG/Source/PCG/Public/Grid/PCGComponentOctree.h b/Engine/Plugins/Experimental/PCG/Source/PCG/Public/Grid/PCGComponentOctree.h
--- a/Engine/Plugins/Experimental/PCG/Source/PCG/Public/Grid/PCGComponentOctree.h
+++ b/Engine/Plugins/Experimental/PCG/Source/PCG/Public/Grid/PCGComponentOctree.h
@@ -61,7 +61,7 @@ struct PCG_API FPCGComponentRefSemantics
 };
 
 using FPCGComponentOctree = TOctree2<FPCGComponentRef, FPCGComponentRefSemantics> ;
-using FPCGComponentToIdMap = TMap<TObjectPtr<UPCGComponent>, FPCGComponentOctreeIDSharedRef>;
+using FPCGComponentToIdMap = TMap<UPCGComponent*, FPCGComponentOctreeIDSharedRef>;
 
 class FPCGComponentOctreeAndMap
 {
@@ -71,7 +71,7 @@ public:
 
 	void Reset(const FVector& InOrigin, FVector::FReal InExtent);
 
-	TSet<TObjectPtr<UPCGComponent>> GetAllComponents() const;
+	TSet<UPCGComponent*> GetAllComponents() const;
 
 	template<typename IterateBoundsFunc>
 	inline void FindElementsWithBoundsTest(const FBoxCenterAndExtent& BoxBounds, const IterateBoundsFunc& Func) const
diff --git a/Engine/Plugins/Experimental/PCG/Source/PCG/Public/PCGActorAndComponentMapping.h b/Engine/Plugins/Experimental/PCG/Source/PCG/Public/PCGActorAndComponentMapping.h
--- a/Engine/Plugins/Experimental/PCG/Source/PCG/Public/PCGActorAndComponentMapping.h
+++ b/Engine/Plugins/Experimental/PCG/Source/PCG/Public/PCGActorAndComponentMapping.h
@@ -70,13 +70,13 @@ public:
 	void UnregisterPartitionActor(APCGPartitionActor* InActor);
 
 	/** Return a copy of all the registered partitioned components. Thread safe */
-	TSet<TObjectPtr<UPCGComponent>> GetAllRegisteredPartitionedComponents() const;
+	TSet<UPCGComponent*> GetAllRegisteredPartitionedComponents() const;
 
 	/** Return a copy of all the registered non-partitioned components. Thread safe */
-	TSet<TObjectPtr<UPCGComponent>> GetAllRegisteredNonPartitionedComponents() const;
+	TSet<UPCGComponent*> GetAllRegisteredNonPartitionedComponents() const;
 
 	/** Return a copy of all the registered components. Thread safe */
-	TSet<TObjectPtr<UPCGComponent>> GetAllRegisteredComponents() const;
+	TSet<UPCGComponent*> GetAllRegisteredComponents() const;
 
 private:
 	// This class is only meant to be used as part of the PCG Subsytem and owned by it.
@@ -159,11 +159,11 @@ private:
 	mutable FRWLock PartitionActorsMapLock;
 
 	/** Mapping between original components and its overlapping partition actors. */
-	TMap<TObjectPtr<const UPCGComponent>, TSet<TObjectPtr<APCGPartitionActor>>> ComponentToPartitionActorsMap;
+	TMap<const UPCGComponent*, TSet<TObjectPtr<APCGPartitionActor>>> ComponentToPartitionActorsMap;
 	mutable FRWLock ComponentToPartitionActorsMapLock;
 
 	/** Components to be unregister at the next frame. cf. UnregisterComponent for a better understanding on why it is needed. */
-	TSet<TObjectPtr<UPCGComponent>> DelayedComponentToUnregister;
+	TSet<UPCGComponent*> DelayedComponentToUnregister;
 	mutable FCriticalSection DelayedComponentToUnregisterLock;
 
 	// Tracking actors
@@ -171,13 +171,13 @@ private:
 	FPCGComponentOctreeAndMap NonPartitionedOctree;
 
 	/** Keep a mapping between tracked actors and the components that track them, and the tracking needs to be culled.*/
-	TMap<TObjectKey<AActor>, TSet<TObjectPtr<UPCGComponent>>> CulledTrackedActorsToComponentsMap;
+	TMap<TObjectKey<AActor>, TSet<UPCGComponent*>> CulledTrackedActorsToComponentsMap;
 
 	/** Same mapping but for always tracked actors */
-	TMap<TObjectKey<AActor>, TSet<TObjectPtr<UPCGComponent>>> AlwaysTrackedActorsToComponentsMap;
+	TMap<TObjectKey<AActor>, TSet<UPCGComponent*>> AlwaysTrackedActorsToComponentsMap;
 
 	/** Keep the list of the keys already tracked, to avoid requerying the actors everytime */
-	TMap<FPCGActorSelectionKey, TSet<TObjectPtr<UPCGComponent>>> KeysToComponentsMap;
+	TMap<FPCGActorSelectionKey, TSet<UPCGComponent*>> KeysToComponentsMap;
 
 	/** Finally keep a mapping between actors and their position to know if an actor move in/out of a component tracking bounds. Only kept for actors that need to be culled. */
 	TMap<TObjectKey<AActor>, FBox> TrackedActorToPositionMap;
diff --git a/Engine/Plugins/Experimental/PCG/Source/PCG/Public/PCGComponent.h b/Engine/Plugins/Experimental/PCG/Source/PCG/Public/PCGComponent.h
--- a/Engine/Plugins/Experimental/PCG/Source/PCG/Public/PCGComponent.h
+++ b/Engine/Plugins/Experimental/PCG/Source/PCG/Public/PCGComponent.h
@@ -67,6 +67,7 @@ UCLASS(BlueprintType, ClassGroup = (Procedural), meta = (BlueprintSpawnableCompo
 class PCG_API UPCGComponent : public UActorComponent
 {
 	UPCGComponent(const FObjectInitializer& InObjectInitializer);
+	~UPCGComponent();
 
 	GENERATED_BODY()
 
diff --git a/Engine/Plugins/Experimental/PCG/Source/PCG/Public/PCGSubsystem.h b/Engine/Plugins/Experimental/PCG/Source/PCG/Public/PCGSubsystem.h
--- a/Engine/Plugins/Experimental/PCG/Source/PCG/Public/PCGSubsystem.h
+++ b/Engine/Plugins/Experimental/PCG/Source/PCG/Public/PCGSubsystem.h
@@ -45,6 +45,9 @@ public:
 
 	UPCGSubsystem();
 
+	/** To be used when a PCG component can not have a world anymore, to unregister itself. */
+	static UPCGSubsystem* GetSubsystemForCurrentWorld();
+
 	//~ Begin USubsystem Interface.
 	virtual void Deinitialize() override;
 	//~ End USubsystem Interface
@@ -139,7 +142,7 @@ public:
 	/** Unregister a Partition actor, will be removed from the map and remove itself to all intersecting volumes. Thread safe */
 	void UnregisterPartitionActor(APCGPartitionActor* InActor) { ActorAndComponentMapping.UnregisterPartitionActor(InActor); }
 
-	TSet<TObjectPtr<UPCGComponent>> GetAllRegisteredPartitionedComponents() const { return ActorAndComponentMapping.GetAllRegisteredPartitionedComponents(); }
+	TSet<UPCGComponent*> GetAllRegisteredPartitionedComponents() const { return ActorAndComponentMapping.GetAllRegisteredPartitionedComponents(); }
 
 	/** Flushes the graph cache completely, use only for debugging */
 	void FlushCache();
