[PCG] Fix assert in graph cache memory accounting

#rnx

diff --git a/Engine/Plugins/Experimental/PCG/Source/PCG/Private/Graph/PCGGraphCache.cpp b/Engine/Plugins/Experimental/PCG/Source/PCG/Private/Graph/PCGGraphCache.cpp
--- a/Engine/Plugins/Experimental/PCG/Source/PCG/Private/Graph/PCGGraphCache.cpp
+++ b/Engine/Plugins/Experimental/PCG/Source/PCG/Private/Graph/PCGGraphCache.cpp
@@ -276,7 +276,7 @@ void FPCGGraphCache::RemoveFromMemoryTotal(const FPCGDataCollection& InCollectio
 					if (Record->InstanceCount == 0)
 					{
 						// Last instance removed, update accordingly
-						if (ensure(TotalMemoryUsed > Record->MemoryPerInstance))
+						if (ensure(TotalMemoryUsed >= Record->MemoryPerInstance))
 						{
 							TotalMemoryUsed -= Record->MemoryPerInstance;
 						}
