[PCG] Grow graph cache entry count when max capacity hit.

#rnx

diff --git a/Engine/Plugins/Experimental/PCG/Source/PCG/Private/Graph/PCGGraphCache.cpp b/Engine/Plugins/Experimental/PCG/Source/PCG/Private/Graph/PCGGraphCache.cpp
--- a/Engine/Plugins/Experimental/PCG/Source/PCG/Private/Graph/PCGGraphCache.cpp
+++ b/Engine/Plugins/Experimental/PCG/Source/PCG/Private/Graph/PCGGraphCache.cpp
@@ -24,16 +24,11 @@ static TAutoConsoleVariable<bool> CVarCacheMemoryBudgetEnabled(
 	true,
 	TEXT("Whether memory budget is enforced (items purged from cache to respect pcg.Cache.MemoryBudgetMB."));
 
-static int32 GPCGGraphCacheMaxElements = 65536;
-static FAutoConsoleVariableRef CVarNiagaraGraphDataCacheSize(
-	TEXT("pcg.Cache.GraphCacheMaxElements"),
-	GPCGGraphCacheMaxElements,
-	TEXT("Maximum number of elements to store within the graph cache."),
-	ECVF_ReadOnly
-);
+// Initial max number of entries graph cache
+static const int32 GPCGGraphCacheInitialCapacity = 65536;
 
 FPCGGraphCache::FPCGGraphCache(TWeakObjectPtr<UObject> InOwner, FPCGRootSet* InRootSet)
-	: CacheData(GPCGGraphCacheMaxElements)
+	: CacheData(GPCGGraphCacheInitialCapacity)
 	, Owner(InOwner)
 	, RootSet(InRootSet)
 {
@@ -100,6 +95,11 @@ void FPCGGraphCache::StoreInCache(const IPCGElement* InElement, const FPCGCrc& I
 		TRACE_CPUPROFILER_EVENT_SCOPE(FPCGGraphCache::StoreInCache);
 		FWriteScopeLock ScopedWriteLock(CacheLock);
 
+		if (CacheData.Num() == CacheData.Max())
+		{
+			GrowCache_Unsafe();
+		}
+
 		FPCGCacheEntryKey CacheKey(InElement, InDependenciesCrc);
 		CacheData.Add(CacheKey, InOutput);
 		
@@ -123,7 +123,7 @@ void FPCGGraphCache::ClearCache()
 	TotalMemoryUsed = 0;
 
 	// Remove all entries
-	CacheData.Empty(GPCGGraphCacheMaxElements);
+	CacheData.Empty(CacheData.Max());
 }
 
 bool FPCGGraphCache::EnforceMemoryBudget()
@@ -135,11 +135,6 @@ bool FPCGGraphCache::EnforceMemoryBudget()
 		return false;
 	}
 
-	if (CacheData.Num() == GPCGGraphCacheMaxElements)
-	{
-		UE_LOG(LogPCG, Warning, TEXT("Graph cache full (%d entries). Consider increasing GPCGGraphCacheMaxElements in code."), GPCGGraphCacheMaxElements);
-	}
-
 	const uint64 MemoryBudget = static_cast<uint64>(CVarCacheMemoryBudgetMB.GetValueOnAnyThread()) * 1024 * 1024;
 	if (TotalMemoryUsed <= MemoryBudget)
 	{
@@ -219,6 +214,29 @@ uint32 FPCGGraphCache::GetGraphCacheEntryCount(IPCGElement* InElement) const
 }
 #endif // WITH_EDITOR
 
+void FPCGGraphCache::GrowCache_Unsafe()
+{
+	TRACE_CPUPROFILER_EVENT_SCOPE(FPCGGraphCache::GrowCache_Unsafe);
+
+	TLruCache<FPCGCacheEntryKey, FPCGDataCollection> CacheDataCopy(CacheData.Max());
+
+	// Iteration begins from most recent, so this reverses the temporal order.
+	for (TLruCache<FPCGCacheEntryKey, FPCGDataCollection>::TIterator It(CacheData); It; ++It)
+	{
+		CacheDataCopy.Add(It.Key(), It.Value());
+	}
+
+	// Resize and flush
+	CacheData.Empty(CacheData.Num() * 2);
+	UE_LOG(LogPCG, Warning, TEXT("Graph cache doubled in capacity to %d entries."), CacheData.Max());
+
+	// Copy back. Restore temporal order.
+	for (TLruCache<FPCGCacheEntryKey, FPCGDataCollection>::TIterator It(CacheDataCopy); It; ++It)
+	{
+		CacheData.Add(It.Key(), It.Value());
+	}
+}
+
 void FPCGGraphCache::AddDataToAccountedMemory(const FPCGDataCollection& InCollection)
 {
 	for (const FPCGTaggedData& Data : InCollection.TaggedData)
diff --git a/Engine/Plugins/Experimental/PCG/Source/PCG/Private/Graph/PCGGraphCache.h b/Engine/Plugins/Experimental/PCG/Source/PCG/Private/Graph/PCGGraphCache.h
--- a/Engine/Plugins/Experimental/PCG/Source/PCG/Private/Graph/PCGGraphCache.h
+++ b/Engine/Plugins/Experimental/PCG/Source/PCG/Private/Graph/PCGGraphCache.h
@@ -89,6 +89,9 @@ struct FPCGGraphCache
 #endif
 
 private:
+	// Grow cache entry capacity, preserving current entries.
+	void GrowCache_Unsafe();
+
 	/** For each data element in collection, keep records of memory usage and instance count. */
 	void AddDataToAccountedMemory(const FPCGDataCollection& InCollection);
 
