[PCG] Update inspection/debug on significant events like changing map and deleting actors with PCG components.

These were the primary ways I could repro the jira below - and it generally feels much neater to tidy up the state rather than continue inspecting components/data that no longer exist.

#rnx

diff --git a/Engine/Plugins/Experimental/PCG/Source/PCGEditor/Private/PCGEditor.cpp b/Engine/Plugins/Experimental/PCG/Source/PCGEditor/Private/PCGEditor.cpp
--- a/Engine/Plugins/Experimental/PCG/Source/PCGEditor/Private/PCGEditor.cpp
+++ b/Engine/Plugins/Experimental/PCG/Source/PCGEditor/Private/PCGEditor.cpp
@@ -42,6 +42,7 @@
 #include "EditorAssetLibrary.h"
 #include "GraphEditorActions.h"
 #include "IDetailsView.h"
+#include "LevelEditor.h"
 #include "PropertyEditorModule.h"
 #include "SNodePanel.h"
 #include "ScopedTransaction.h"
@@ -108,6 +109,7 @@ void FPCGEditor::Initialize(const EToolkitMode::Type InMode, const TSharedPtr<cl
 	GraphEditorWidget = CreateGraphEditorWidget();
 	PaletteWidget = CreatePaletteWidget();
 	DebugObjectWidget = CreateDebugObjectWidget();
+	DebugObjectTreeWidget = CreateDebugObjectTreeWidget();
 	FindWidget = CreateFindWidget();
 	AttributesWidget = CreateAttributesWidget();
 	DeterminismWidget = CreateDeterminismWidget();
@@ -173,6 +175,14 @@ void FPCGEditor::Initialize(const EToolkitMode::Type InMode, const TSharedPtr<cl
 	const FName PCGGraphEditorAppName = FName(TEXT("PCGEditorApp"));
 
 	InitAssetEditor(InMode, InToolkitHost, PCGGraphEditorAppName, StandaloneDefaultLayout, /*bCreateDefaultStandaloneMenu=*/ true, /*bCreateDefaultToolbar=*/ true, InPCGGraph);
+
+	// Hook to map change / delete actor to refresh debug object selection list, to help prevent it going stale.
+	FLevelEditorModule& LevelEditor = FModuleManager::LoadModuleChecked<FLevelEditorModule>("LevelEditor");
+	LevelEditor.OnMapChanged().AddRaw(this, &FPCGEditor::OnMapChanged);
+	if (GEngine)
+	{
+		GEngine->OnLevelActorDeleted().AddRaw(this, &FPCGEditor::OnLevelActorDeleted);
+	}
 }
 
 UPCGEditorGraph* FPCGEditor::GetPCGEditorGraph()
@@ -470,7 +480,10 @@ void FPCGEditor::RegisterToolbar() const
 			const UPCGEditorMenuContext* Context = InSection.FindContext<UPCGEditorMenuContext>();
 			if (Context && Context->PCGEditor.IsValid())
 			{
-				InSection.AddEntry(FToolMenuEntry::InitWidget("SelectedDebugObjectWidget", SNew(SPCGEditorGraphDebugObjectWidget, Context->PCGEditor.Pin()), FText::GetEmpty()));	
+				InSection.AddEntry(FToolMenuEntry::InitWidget(
+					"SelectedDebugObjectWidget",
+					Context->PCGEditor.Pin()->DebugObjectWidget.ToSharedRef(),
+					FText::GetEmpty()));
 			}
 		}));
 
@@ -2186,6 +2199,15 @@ void FPCGEditor::OnClose()
 		PCGEditorGraph->OnClose();
 	}
 
+	if (FLevelEditorModule* LevelEditor = FModuleManager::GetModulePtr<FLevelEditorModule>(TEXT("LevelEditor")))
+	{
+		LevelEditor->OnMapChanged().RemoveAll(this);
+	}
+	if (GEngine)
+	{
+		GEngine->OnLevelActorDeleted().RemoveAll(this);
+	}
+
 	// Extra nodes are replicated on editor close, to be saved in the underlying PCGGraph
 	ReplicateExtraNodes();
 
@@ -2216,7 +2238,12 @@ TSharedRef<SPCGEditorGraphNodePalette> FPCGEditor::CreatePaletteWidget()
 	return SNew(SPCGEditorGraphNodePalette);
 }
 
-TSharedRef<SPCGEditorGraphDebugObjectTree> FPCGEditor::CreateDebugObjectWidget()
+TSharedRef<SPCGEditorGraphDebugObjectWidget> FPCGEditor::CreateDebugObjectWidget()
+{
+	return SNew(SPCGEditorGraphDebugObjectWidget, SharedThis(this));
+}
+
+TSharedRef<SPCGEditorGraphDebugObjectTree> FPCGEditor::CreateDebugObjectTreeWidget()
 {
 	return SNew(SPCGEditorGraphDebugObjectTree, SharedThis(this));
 }
@@ -2420,6 +2447,35 @@ UPCGSubsystem* FPCGEditor::GetSubsystem()
 	return UPCGSubsystem::GetInstance(World);
 }
 
+void FPCGEditor::OnMapChanged(UWorld* InWorld, EMapChangeType InMapChangedType)
+{
+	if (InMapChangedType != EMapChangeType::SaveMap)
+	{
+		if (DebugObjectTreeWidget.IsValid())
+		{
+			DebugObjectTreeWidget->RequestRefresh();
+		}
+		if (DebugObjectWidget.IsValid())
+		{
+			DebugObjectWidget->RefreshDebugObjects();
+		}
+	}
+}
+
+void FPCGEditor::OnLevelActorDeleted(AActor* InActor)
+{
+	if (DebugObjectTreeWidget.IsValid())
+	{
+		DebugObjectTreeWidget->RequestRefresh();
+	}
+
+	// Forward call as this makes an effort to retain the selection if the selected component has not been deleted.
+	if (DebugObjectWidget.IsValid())
+	{
+		DebugObjectWidget->OnLevelActorDeleted(InActor);
+	}
+}
+
 TSharedRef<SDockTab> FPCGEditor::SpawnTab_GraphEditor(const FSpawnTabArgs& Args)
 {
 	return SNew(SDockTab)
@@ -2456,7 +2512,7 @@ TSharedRef<SDockTab> FPCGEditor::SpawnTab_DebugObject(const FSpawnTabArgs& Args)
 		.Label(LOCTEXT("PCGDebugObjectTitle", "Debug Object"))
 		.TabColorScale(GetTabColorScale())
 		[
-			DebugObjectWidget.ToSharedRef()
+			DebugObjectTreeWidget.ToSharedRef()
 		];
 }
 
diff --git a/Engine/Plugins/Experimental/PCG/Source/PCGEditor/Private/PCGEditor.h b/Engine/Plugins/Experimental/PCG/Source/PCGEditor/Private/PCGEditor.h
--- a/Engine/Plugins/Experimental/PCG/Source/PCGEditor/Private/PCGEditor.h
+++ b/Engine/Plugins/Experimental/PCG/Source/PCGEditor/Private/PCGEditor.h
@@ -17,6 +17,7 @@ class FUICommandList;
 class IDetailsView;
 class SGraphEditor;
 class SPCGEditorGraphAttributeListView;
+class SPCGEditorGraphDebugObjectWidget;
 class SPCGEditorGraphDebugObjectTree;
 class SPCGEditorGraphDeterminismListView;
 class SPCGEditorGraphFind;
@@ -226,8 +227,11 @@ private:
 	/** Create new palette widget */
 	TSharedRef<SPCGEditorGraphNodePalette> CreatePaletteWidget();
 
+	/** Create debug object combo box */
+	TSharedRef<SPCGEditorGraphDebugObjectWidget> CreateDebugObjectWidget();
+
 	/** Create new debug object tree widget */
-	TSharedRef<SPCGEditorGraphDebugObjectTree> CreateDebugObjectWidget();
+	TSharedRef<SPCGEditorGraphDebugObjectTree> CreateDebugObjectTreeWidget();
 
 	/** Create new find widget */
 	TSharedRef<SPCGEditorGraphFind> CreateFindWidget();
@@ -281,6 +285,9 @@ private:
 	/** Helper to get to the subsystem. */
 	static class UPCGSubsystem* GetSubsystem();
 
+	void OnMapChanged(UWorld* InWorld, EMapChangeType InMapChangedType);
+	void OnLevelActorDeleted(AActor* InActor);
+
 	TSharedRef<SDockTab> SpawnTab_GraphEditor(const FSpawnTabArgs& Args);
 	TSharedRef<SDockTab> SpawnTab_PropertyDetails(const FSpawnTabArgs& Args);
 	TSharedRef<SDockTab> SpawnTab_Palette(const FSpawnTabArgs& Args);
@@ -294,7 +301,8 @@ private:
 	TSharedPtr<SGraphEditor> GraphEditorWidget;
 	TSharedPtr<IDetailsView> PropertyDetailsWidget;
 	TSharedPtr<SPCGEditorGraphNodePalette> PaletteWidget;
-	TSharedPtr<SPCGEditorGraphDebugObjectTree> DebugObjectWidget;
+	TSharedPtr<SPCGEditorGraphDebugObjectWidget> DebugObjectWidget;
+	TSharedPtr<SPCGEditorGraphDebugObjectTree> DebugObjectTreeWidget;
 	TSharedPtr<SPCGEditorGraphFind> FindWidget;
 	TSharedPtr<SPCGEditorGraphAttributeListView> AttributesWidget;
 	TSharedPtr<SPCGEditorGraphDeterminismListView> DeterminismWidget;
diff --git a/Engine/Plugins/Experimental/PCG/Source/PCGEditor/Private/SPCGEditorGraphDebugObjectWidget.cpp b/Engine/Plugins/Experimental/PCG/Source/PCGEditor/Private/SPCGEditorGraphDebugObjectWidget.cpp
--- a/Engine/Plugins/Experimental/PCG/Source/PCGEditor/Private/SPCGEditorGraphDebugObjectWidget.cpp
+++ b/Engine/Plugins/Experimental/PCG/Source/PCGEditor/Private/SPCGEditorGraphDebugObjectWidget.cpp
@@ -97,7 +97,7 @@ void SPCGEditorGraphDebugObjectWidget::Construct(const FArguments& InArgs, TShar
 	];
 }
 
-void SPCGEditorGraphDebugObjectWidget::OnComboBoxOpening()
+void SPCGEditorGraphDebugObjectWidget::RefreshDebugObjects()
 {
 	DebugObjects.Empty();
 	DebugObjectsComboBox->RefreshOptions();
@@ -166,6 +166,11 @@ void SPCGEditorGraphDebugObjectWidget::OnComboBoxOpening()
 	}
 }
 
+void SPCGEditorGraphDebugObjectWidget::OnComboBoxOpening()
+{
+	RefreshDebugObjects();
+}
+
 void SPCGEditorGraphDebugObjectWidget::OnSelectionChanged(TSharedPtr<FPCGEditorGraphDebugObjectInstance> NewSelection, ESelectInfo::Type SelectInfo) const
 {
 	if (NewSelection.IsValid())
@@ -324,4 +329,59 @@ bool SPCGEditorGraphDebugObjectWidget::IsSetDebugObjectFromSelectionButtonEnable
 	return false;
 }
 
+void SPCGEditorGraphDebugObjectWidget::OnLevelActorDeleted(const AActor* InActor)
+{
+	// Iterate over all PCG components on the deleted actor and remove them from the debug objects list. If
+	// the currently-selected debug object is deleted, reset the selection to item 0.
+	if (!ensure(InActor))
+	{
+		return;
+	}
+
+	TInlineComponentArray<const UPCGComponent*, 2> PCGComponents;
+	InActor->GetComponents(PCGComponents);
+	if (PCGComponents.IsEmpty())
+	{
+		return;
+	}
+
+	const TSharedPtr<FPCGEditorGraphDebugObjectInstance> SelectedItem = DebugObjectsComboBox->GetSelectedItem();
+
+	bool bDebugObjectWasRemoved = false;
+	bool bSelectedDebugObjectWasRemoved = false;
+	for (const UPCGComponent* DeletedComponent : PCGComponents)
+	{
+		if (!DeletedComponent)
+		{
+			continue;
+		}
+
+		if (SelectedItem && SelectedItem->GetPCGComponent().Get() == DeletedComponent)
+		{
+			bSelectedDebugObjectWasRemoved = true;
+		}
+
+		// First item is always "No debug object selected" item, so decrement down to index 1.
+		for (int i = DebugObjects.Num() - 1; i >= 1; --i)
+		{
+			const UPCGComponent* DebugObjectComponent = DebugObjects[i]->GetPCGComponent().Get();
+			if (DebugObjectComponent == DeletedComponent)
+			{
+				DebugObjects.RemoveAt(i);
+				bDebugObjectWasRemoved = true;
+			}
+		}
+	}
+
+	if (bSelectedDebugObjectWasRemoved && ensure(DebugObjects.Num() > 0))
+	{
+		DebugObjectsComboBox->SetSelectedItem(DebugObjects[0]);
+	}
+
+	if (bDebugObjectWasRemoved)
+	{
+		DebugObjectsComboBox->RefreshOptions();
+	}
+}
+
 #undef LOCTEXT_NAMESPACE
diff --git a/Engine/Plugins/Experimental/PCG/Source/PCGEditor/Private/SPCGEditorGraphDebugObjectWidget.h b/Engine/Plugins/Experimental/PCG/Source/PCGEditor/Private/SPCGEditorGraphDebugObjectWidget.h
--- a/Engine/Plugins/Experimental/PCG/Source/PCGEditor/Private/SPCGEditorGraphDebugObjectWidget.h
+++ b/Engine/Plugins/Experimental/PCG/Source/PCGEditor/Private/SPCGEditorGraphDebugObjectWidget.h
@@ -42,6 +42,10 @@ public:
 
 	void Construct(const FArguments& InArgs, TSharedPtr<FPCGEditor> InPCGEditor);
 
+	void RefreshDebugObjects();
+
+	void OnLevelActorDeleted(const AActor* InActor);
+
 private:
 	void OnComboBoxOpening();
 	void OnSelectionChanged(TSharedPtr<FPCGEditorGraphDebugObjectInstance> NewSelection, ESelectInfo::Type SelectInfo) const;
