Added some robustness in cases where the world isn't available, esp. for blueprint actors.

#rnx

diff --git a/Engine/Plugins/Experimental/PCG/Source/PCG/Private/PCGActorAndComponentMapping.cpp b/Engine/Plugins/Experimental/PCG/Source/PCG/Private/PCGActorAndComponentMapping.cpp
--- a/Engine/Plugins/Experimental/PCG/Source/PCG/Private/PCGActorAndComponentMapping.cpp
+++ b/Engine/Plugins/Experimental/PCG/Source/PCG/Private/PCGActorAndComponentMapping.cpp
@@ -597,10 +597,16 @@ void UPCGActorAndComponentMapping::RegisterOrUpdateTracking(UPCGComponent* InCom
 	TSet<UPCGComponent*>& AllComponents = AlwaysTrackedActorsToComponentsMap.FindOrAdd(ComponentOwner);
 	AllComponents.Add(InComponent);
 
-	UWorld* World = PCGSubsystem->GetWorld();
+	UWorld* World = PCGSubsystem ? PCGSubsystem->GetWorld() : nullptr;
+
+	if (!World)
+	{
+		return;
+	}
+
 	APCGWorldActor* PCGWorldActor = PCGSubsystem ? PCGSubsystem->GetPCGWorldActor() : nullptr;
 
-	if (!World || !PCGWorldActor)
+	if (!PCGWorldActor)
 	{
 		return;
 	}
diff --git a/Engine/Plugins/Experimental/PCG/Source/PCG/Private/PCGWorldActor.cpp b/Engine/Plugins/Experimental/PCG/Source/PCG/Private/PCGWorldActor.cpp
--- a/Engine/Plugins/Experimental/PCG/Source/PCG/Private/PCGWorldActor.cpp
+++ b/Engine/Plugins/Experimental/PCG/Source/PCG/Private/PCGWorldActor.cpp
@@ -29,6 +29,7 @@ APCGWorldActor::APCGWorldActor(const FObjectInitializer& ObjectInitializer)
 void APCGWorldActor::BeginCacheForCookedPlatformData(const ITargetPlatform* TargetPlatform)
 {
 	Super::BeginCacheForCookedPlatformData(TargetPlatform);
+	check(LandscapeCacheObject);
 	LandscapeCacheObject->PrimeCache();
 }
 #endif
@@ -123,9 +124,17 @@ void APCGWorldActor::GetGridGuids(PCGHiGenGrid::FSizeToGuidMap& OutSizeToGuidMap
 #if WITH_EDITOR
 APCGWorldActor* APCGWorldActor::CreatePCGWorldActor(UWorld* InWorld)
 {
-	check(InWorld);
-	APCGWorldActor* PCGActor = InWorld->SpawnActor<APCGWorldActor>();
-	PCGActor->RegisterToSubsystem();
+	APCGWorldActor* PCGActor = nullptr;
+
+	if (InWorld)
+	{
+		PCGActor = InWorld->SpawnActor<APCGWorldActor>();
+
+		if (PCGActor)
+		{
+			PCGActor->RegisterToSubsystem();
+		}
+	}
 
 	return PCGActor;
 }
