[PCG] Fixed a bug with newly created Landscapes not being cached or tracked

+ LandscapeCache will now track LandscapeProxy's as they are created
+ The actor tracking system will now appropriately track Landscapes as they are created

diff --git a/Engine/Plugins/Experimental/PCG/Source/PCG/Private/Elements/PCGActorSelector.cpp b/Engine/Plugins/Experimental/PCG/Source/PCG/Private/Elements/PCGActorSelector.cpp
--- a/Engine/Plugins/Experimental/PCG/Source/PCG/Private/Elements/PCGActorSelector.cpp
+++ b/Engine/Plugins/Experimental/PCG/Source/PCG/Private/Elements/PCGActorSelector.cpp
@@ -244,7 +244,14 @@ TArray<FPCGActorSelectionKey> FPCGActorSelectionKey::GenerateAllKeysForActor(con
 	}
 
 	Result.Reserve(InActor->Tags.Num() + 1);
-	Result.Emplace(InActor->GetClass());
+	// Register keys all the way up the chain
+	UClass* CurrentClass = InActor->GetClass();
+	while (CurrentClass && CurrentClass != AActor::StaticClass())
+	{
+		Result.Emplace(CurrentClass);
+		CurrentClass = CurrentClass->GetSuperClass();
+	}
+
 	for (FName Tag : InActor->Tags)
 	{
 		Result.Emplace(Tag);
diff --git a/Engine/Plugins/Experimental/PCG/Source/PCG/Private/Grid/PCGLandscapeCache.cpp b/Engine/Plugins/Experimental/PCG/Source/PCG/Private/Grid/PCGLandscapeCache.cpp
--- a/Engine/Plugins/Experimental/PCG/Source/PCG/Private/Grid/PCGLandscapeCache.cpp
+++ b/Engine/Plugins/Experimental/PCG/Source/PCG/Private/Grid/PCGLandscapeCache.cpp
@@ -53,6 +53,12 @@ namespace PCGLandscapeCache
 
 		return Result;
 	}
+
+	FIntPoint GetCoordinates(const ULandscapeComponent* LandscapeComponent)
+	{
+		check(LandscapeComponent && LandscapeComponent->ComponentSizeQuads != 0);
+		return FIntPoint(LandscapeComponent->SectionBaseX / LandscapeComponent->ComponentSizeQuads, LandscapeComponent->SectionBaseY / LandscapeComponent->ComponentSizeQuads);
+	}
 }
 
 #if WITH_EDITOR
@@ -582,8 +588,7 @@ void UPCGLandscapeCache::PrimeCache()
 						continue;
 					}
 
-					FIntPoint Coordinate(LandscapeComponent->SectionBaseX / LandscapeComponent->ComponentSizeQuads, LandscapeComponent->SectionBaseY / LandscapeComponent->ComponentSizeQuads);
-					TPair<FGuid, FIntPoint> ComponentKey(LandscapeGuid, Coordinate);
+					TPair<FGuid, FIntPoint> ComponentKey(LandscapeGuid, PCGLandscapeCache::GetCoordinates(LandscapeComponent));
 
 					if (!CachedData.Contains(ComponentKey))
 					{
@@ -757,10 +762,12 @@ void UPCGLandscapeCache::SetupLandscapeCallbacks()
 		Landscape->OnComponentDataChanged.AddUObject(this, &UPCGLandscapeCache::OnLandscapeChanged);
 	}
 
-	// Also track when the landscape is moved
+	// Also track when the landscape is moved, added, or deleted
 	if (GEngine)
 	{
 		GEngine->OnActorMoved().AddUObject(this, &UPCGLandscapeCache::OnLandscapeMoved);
+		GEngine->OnLevelActorAdded().AddUObject(this, &UPCGLandscapeCache::OnLandscapeAdded);
+		GEngine->OnLevelActorDeleted().AddUObject(this, &UPCGLandscapeCache::OnLandscapeDeleted);
 	}
 }
 
@@ -777,6 +784,8 @@ void UPCGLandscapeCache::TeardownLandscapeCallbacks()
 	if (GEngine)
 	{
 		GEngine->OnActorMoved().RemoveAll(this);
+		GEngine->OnLevelActorAdded().RemoveAll(this);
+		GEngine->OnLevelActorDeleted().RemoveAll(this);
 	}
 }
 
@@ -794,8 +803,7 @@ void UPCGLandscapeCache::OnLandscapeChanged(ALandscapeProxy* InLandscape, const
 	{
 		if (LandscapeComponent)
 		{
-			FIntPoint Coordinate(LandscapeComponent->SectionBaseX / LandscapeComponent->ComponentSizeQuads, LandscapeComponent->SectionBaseY / LandscapeComponent->ComponentSizeQuads);
-			TPair<FGuid, FIntPoint> ComponentKey(InLandscape->GetOriginalLandscapeGuid(), Coordinate);
+			const TPair<FGuid, FIntPoint> ComponentKey(InLandscape->GetOriginalLandscapeGuid(), PCGLandscapeCache::GetCoordinates(LandscapeComponent));
 
 			FPCGLandscapeCacheEntry* EntryToDelete = nullptr;
 
@@ -822,26 +830,41 @@ void UPCGLandscapeCache::OnLandscapeMoved(AActor* InActor)
 	CacheLock.WriteLock();
 
 	// Just remove these from the cache, they'll be added back on demand
-	Landscape->ForEachComponent<ULandscapeComponent>(/*bIncludeFromChildActors=*/false, [this, Landscape](const ULandscapeComponent* LandscapeComponent)
+	RemoveComponentFromCache(Landscape);
+
+	CacheLayerNames(Landscape);
+
+	CacheLock.WriteUnlock();
+}
+
+void UPCGLandscapeCache::OnLandscapeDeleted(AActor* Actor)
+{
+	ALandscapeProxy* LandscapeProxy = Cast<ALandscapeProxy>(Actor);
+	if (!LandscapeProxy || !Landscapes.Contains(LandscapeProxy))
 	{
-		if (LandscapeComponent)
-		{
-			FIntPoint Coordinate(LandscapeComponent->SectionBaseX / LandscapeComponent->ComponentSizeQuads, LandscapeComponent->SectionBaseY / LandscapeComponent->ComponentSizeQuads);
-			TPair<FGuid, FIntPoint> ComponentKey(Landscape->GetOriginalLandscapeGuid(), Coordinate);
+		return;
+	}
 
-			if (FPCGLandscapeCacheEntry** FoundEntry = CachedData.Find(ComponentKey))
-			{
-				delete* FoundEntry;
-				CachedData.Remove(ComponentKey);
-			}
-		}
-	});
+	CacheLock.WriteLock();
 
-	CacheLayerNames(Landscape);
+	RemoveComponentFromCache(LandscapeProxy);
+
+	Landscapes.Remove(LandscapeProxy);
+	LandscapeProxy->OnComponentDataChanged.RemoveAll(this);
 
 	CacheLock.WriteUnlock();
 }
 
+void UPCGLandscapeCache::OnLandscapeAdded(AActor* Actor)
+{
+	if (ALandscapeProxy* LandscapeProxy = Cast<ALandscapeProxy>(Actor))
+	{
+		Landscapes.Add(LandscapeProxy);
+		LandscapeProxy->OnComponentDataChanged.AddUObject(this, &UPCGLandscapeCache::OnLandscapeChanged);
+		// Note: Landscape Proxies have no components at this stage, so they will need to be added on demand
+	}
+}
+
 void UPCGLandscapeCache::CacheLayerNames()
 {
 	CachedLayerNames.Reset();
@@ -855,6 +878,25 @@ void UPCGLandscapeCache::CacheLayerNames()
 	}
 }
 
+void UPCGLandscapeCache::RemoveComponentFromCache(const ALandscapeProxy* LandscapeProxy)
+{
+	check(LandscapeProxy);
+
+	LandscapeProxy->ForEachComponent<ULandscapeComponent>(/*bIncludeFromChildActors=*/false, [this, LandscapeProxy](const ULandscapeComponent* LandscapeComponent)
+	{
+		if (LandscapeComponent)
+		{
+			const TPair<FGuid, FIntPoint> ComponentKey(LandscapeProxy->GetOriginalLandscapeGuid(), PCGLandscapeCache::GetCoordinates(LandscapeComponent));
+
+			if (FPCGLandscapeCacheEntry** FoundEntry = CachedData.Find(ComponentKey))
+			{
+				CachedData.Remove(ComponentKey);
+				delete* FoundEntry;
+			}
+		}
+	});
+}
+
 void UPCGLandscapeCache::CacheLayerNames(ALandscapeProxy* InLandscape)
 {
 	check(InLandscape);
diff --git a/Engine/Plugins/Experimental/PCG/Source/PCG/Private/PCGComponent.cpp b/Engine/Plugins/Experimental/PCG/Source/PCG/Private/PCGComponent.cpp
--- a/Engine/Plugins/Experimental/PCG/Source/PCG/Private/PCGComponent.cpp
+++ b/Engine/Plugins/Experimental/PCG/Source/PCG/Private/PCGComponent.cpp
@@ -1633,6 +1633,20 @@ bool UPCGComponent::IsActorTracked(AActor* InActor, bool& bOutIsCulled) const
 		return false;
 	}
 
+	// We should always track the owner of the component, without culling
+	if (GetOwner() == InActor)
+	{
+		bOutIsCulled = false;
+		return true;
+	}
+
+	// If we track the landscape using legacy methods and it is a landscape, it should be tracked as culled
+	if (InActor->IsA<ALandscapeProxy>() && ShouldTrackLandscape())
+	{
+		bOutIsCulled = true;
+		return true;
+	}
+
 	bool bFound = false;
 
 	for (const FPCGActorSelectionKey& Key : FPCGActorSelectionKey::GenerateAllKeysForActor(InActor))
diff --git a/Engine/Plugins/Experimental/PCG/Source/PCG/Public/Grid/PCGLandscapeCache.h b/Engine/Plugins/Experimental/PCG/Source/PCG/Public/Grid/PCGLandscapeCache.h
--- a/Engine/Plugins/Experimental/PCG/Source/PCG/Public/Grid/PCGLandscapeCache.h
+++ b/Engine/Plugins/Experimental/PCG/Source/PCG/Public/Grid/PCGLandscapeCache.h
@@ -131,8 +131,11 @@ private:
 	void TeardownLandscapeCallbacks();
 	void OnLandscapeChanged(ALandscapeProxy* InLandscape, const FLandscapeProxyComponentDataChangedParams& InChangeParams);
 	void OnLandscapeMoved(AActor* InActor);
+	void OnLandscapeAdded(AActor* Actor);
+	void OnLandscapeDeleted(AActor* Actor);
 	void CacheLayerNames(ALandscapeProxy* InLandscape);
 	void CacheLayerNames();
+	void RemoveComponentFromCache(const ALandscapeProxy* LandscapeProxy);
 #endif
 
 	TMap<TPair<FGuid, FIntPoint>, FPCGLandscapeCacheEntry*> CachedData;
