[PCG] Add Attribute node now works with extra accessors

diff --git a/Engine/Plugins/Experimental/PCG/Source/PCG/Private/Elements/PCGCreateAttribute.cpp b/Engine/Plugins/Experimental/PCG/Source/PCG/Private/Elements/PCGCreateAttribute.cpp
--- a/Engine/Plugins/Experimental/PCG/Source/PCG/Private/Elements/PCGCreateAttribute.cpp
+++ b/Engine/Plugins/Experimental/PCG/Source/PCG/Private/Elements/PCGCreateAttribute.cpp
@@ -8,6 +8,7 @@
 #include "PCGPin.h"
 #include "Data/PCGSpatialData.h"
 #include "Elements/Metadata/PCGMetadataElementCommon.h"
+#include "Metadata/Accessors/PCGAttributeAccessorHelpers.h"
 
 #include UE_INLINE_GENERATED_CPP_BY_NAME(PCGCreateAttribute)
 
@@ -352,8 +353,34 @@ bool FPCGCreateAttributeElement::ExecuteInternal(FPCGContext* Context) const
 
 		if (SourceParamData)
 		{
-			const FPCGMetadataAttributeBase* SourceAttribute = SourceParamData->Metadata->GetConstAttribute(SourceParamAttributeName);
-			Attribute = Metadata->CopyAttribute(SourceAttribute, OutputAttributeName, /*bKeepParent=*/false, /*bCopyEntries=*/bShouldAddNewEntry, /*bCopyValues=*/bShouldAddNewEntry);
+			const FPCGAttributePropertyInputSelector InputSource = Settings->InputSource.CopyAndFixLast(SourceParamData);
+
+			// If no field accessor, copy over the attribute
+			if (InputSource.GetExtraNames().IsEmpty())
+			{
+				const FPCGMetadataAttributeBase* SourceAttribute = Metadata->GetConstAttribute(SourceParamAttributeName);
+				Attribute = Metadata->CopyAttribute(SourceAttribute, OutputAttributeName, /*bKeepParent=*/false, /*bCopyEntries=*/bShouldAddNewEntry, /*bCopyValues=*/bShouldAddNewEntry);
+			}
+			else // Create a new attribute of the accessed field's type manually
+			{
+				TUniquePtr<const IPCGAttributeAccessor> InputAccessor = PCGAttributeAccessorHelpers::CreateConstAccessor(SourceParamData, InputSource);
+				TUniquePtr<const IPCGAttributeAccessorKeys> InputKeys = PCGAttributeAccessorHelpers::CreateConstKeys(SourceParamData, InputSource);
+				if (!InputAccessor.IsValid() || !InputKeys.IsValid())
+				{
+					PCGE_LOG(Error, GraphAndLog, LOCTEXT("FailedToCreateInputAccessor", "Failed to create input accessor"));
+					return true;
+				}
+
+				auto CreateOutputAttribute = [Metadata, OutputAttributeName, &InputAccessor, &InputKeys]<typename Type>(Type Dummy) -> FPCGMetadataAttributeBase*
+				{
+					// Get the value from the input accessor and pass that as the default value
+					Type Value{};
+					InputAccessor->Get<Type>(Value, *InputKeys);
+					return PCGMetadataElementCommon::ClearOrCreateAttribute<Type>(Metadata, OutputAttributeName, Value);
+				};
+
+				Attribute = PCGMetadataAttribute::CallbackWithRightType(InputAccessor->GetUnderlyingType(), CreateOutputAttribute);
+			}
 		}
 		else
 		{
