[PCG] Renamed the Branch Node to Select Node and added new Branch Node that better reflects the BP version

* Branch Node is now renamed "Select Node", but otherwise maintains the same function
+ Branch Node (new) created as one single input, that branches into one of two outputs, based on an overridable property toggle

Future work will include a Select node (by enum, int, etc) with dynamic number of inputs and a Switch (by enum, int, etc) with dynamic outputs

diff --git a/Engine/Plugins/Experimental/PCG/Source/PCG/Private/Elements/PCGBooleanSelect.cpp b/Engine/Plugins/Experimental/PCG/Source/PCG/Private/Elements/PCGBooleanSelect.cpp
new file mode 100644
--- /dev/null
+++ b/Engine/Plugins/Experimental/PCG/Source/PCG/Private/Elements/PCGBooleanSelect.cpp
@@ -0,0 +1,93 @@
+// Copyright Epic Games, Inc. All Rights Reserved.
+
+#include "Elements/PCGBooleanSelect.h"
+
+#include "PCGCommon.h"
+#include "PCGContext.h"
+#include "PCGPin.h"
+#include "Elements/PCGGather.h"
+
+#define LOCTEXT_NAMESPACE "FPCGBooleanSelectElement"
+
+namespace PCGBooleanSelectConstants
+{
+	const FName InputLabelA = TEXT("Input A");
+	const FName InputLabelB = TEXT("Input B");
+	const FText NodeTitleBase = LOCTEXT("NodeTitle", "Select");
+}
+
+#if WITH_EDITOR
+FText UPCGBooleanSelectSettings::GetDefaultNodeTitle() const
+{
+	// TODO: This should statically update or dynamically update, if overridden, for which input was selected, ie. Select (Input A)
+	return PCGBooleanSelectConstants::NodeTitleBase;
+}
+
+FText UPCGBooleanSelectSettings::GetNodeTooltipText() const
+{
+	return LOCTEXT("NodeTooltip", "Control flow node that will select all input data on either Pin A or Pin B only, based on the 'Use Input B' property - which can also be overridden.");
+}
+#endif // WITH_EDITOR
+
+EPCGDataType UPCGBooleanSelectSettings::GetCurrentPinTypes(const UPCGPin* Pin) const
+{
+	// Output pin depends on the union of both input pins 
+	if (Pin->IsOutputPin())
+	{
+		const EPCGDataType InputTypeUnion = GetTypeUnionOfIncidentEdges(PCGBooleanSelectConstants::InputLabelA) | GetTypeUnionOfIncidentEdges(PCGBooleanSelectConstants::InputLabelB);
+		return InputTypeUnion != EPCGDataType::None ? InputTypeUnion : EPCGDataType::Any;
+	}
+
+	return Super::GetCurrentPinTypes(Pin);
+}
+
+TArray<FPCGPinProperties> UPCGBooleanSelectSettings::InputPinProperties() const
+{
+	TArray<FPCGPinProperties> PinProperties;
+	PinProperties.Emplace(PCGBooleanSelectConstants::InputLabelA,
+		EPCGDataType::Any,
+		/*bInAllowMultipleConnections=*/true,
+		/*bAllowMultipleData=*/true,
+		LOCTEXT("FirstInputPinTooltip", "Will only be used if 'Use Input B' (overridable) is false"));
+	PinProperties.Emplace(PCGBooleanSelectConstants::InputLabelB,
+		EPCGDataType::Any,
+		/*bInAllowMultipleConnections=*/true,
+		/*bAllowMultipleData=*/true,
+		LOCTEXT("SecondInputPinTooltip", "Will only be used if 'Use Input B' (overridable) is true"));
+
+	return PinProperties;
+}
+
+TArray<FPCGPinProperties> UPCGBooleanSelectSettings::OutputPinProperties() const
+{
+	TArray<FPCGPinProperties> PinProperties;
+	PinProperties.Emplace(PCGPinConstants::DefaultOutputLabel,
+		EPCGDataType::Any,
+		/*bInAllowMultipleConnections=*/true,
+		/*bAllowMultipleData=*/true,
+		LOCTEXT("OutputPinTooltip", "All input will gathered into a single data collection"));
+
+	return PinProperties;
+}
+
+FPCGElementPtr UPCGBooleanSelectSettings::CreateElement() const
+{
+	return MakeShared<FPCGBooleanSelectElement>();
+}
+
+bool FPCGBooleanSelectElement::ExecuteInternal(FPCGContext* Context) const
+{
+	TRACE_CPUPROFILER_EVENT_SCOPE(FPCGBooleanSelectElement::ExecuteInternal);
+
+	const UPCGBooleanSelectSettings* Settings = Context->GetInputSettings<UPCGBooleanSelectSettings>();
+	check(Settings);
+
+	const FName SelectedPinLabel = Settings->bUseInputB ? PCGBooleanSelectConstants::InputLabelB : PCGBooleanSelectConstants::InputLabelA;
+
+	// Reuse the functionality of the Gather Node
+	Context->OutputData = PCGGather::GatherDataForPin(Context->InputData, SelectedPinLabel);
+
+	return true;
+}
+
+#undef LOCTEXT_NAMESPACE
\ No newline at end of file
diff --git a/Engine/Plugins/Experimental/PCG/Source/PCG/Private/Elements/PCGBranch.cpp b/Engine/Plugins/Experimental/PCG/Source/PCG/Private/Elements/PCGBranch.cpp
--- a/Engine/Plugins/Experimental/PCG/Source/PCG/Private/Elements/PCGBranch.cpp
+++ b/Engine/Plugins/Experimental/PCG/Source/PCG/Private/Elements/PCGBranch.cpp
@@ -11,59 +11,49 @@
 
 namespace PCGBranchConstants
 {
-	const FName InputLabelA = TEXT("Input A");
-	const FName InputLabelB = TEXT("Input B");
+	const FName OutputLabelA = TEXT("Output A");
+	const FName OutputLabelB = TEXT("Output B");
 	const FText NodeTitleBase = LOCTEXT("NodeTitle", "Branch");
 }
 
 #if WITH_EDITOR
 FText UPCGBranchSettings::GetDefaultNodeTitle() const
 {
-	// TODO: This should statically update or dynamically update, if overridden, for which branch was taken, ie. Branch (A)
+	// TODO: This should statically update or dynamically update, if overridden, for which branch was taken, ie. Branch (Output A)
 	return PCGBranchConstants::NodeTitleBase;
 }
 
 FText UPCGBranchSettings::GetNodeTooltipText() const
 {
-	return LOCTEXT("NodeTooltip", "Control flow node that will allow all input data on either Pin A or Pin B only, based on the 'Use Input B' property - which can also be overridden.");
+	return LOCTEXT("NodeTooltip", "Control flow node that will route the input to either Output A or Output B, based on the 'Output To B' property - which can also be overridden.");
 }
 #endif // WITH_EDITOR
 
-EPCGDataType UPCGBranchSettings::GetCurrentPinTypes(const UPCGPin* InPin) const
+EPCGDataType UPCGBranchSettings::GetCurrentPinTypes(const UPCGPin* Pin) const
 {
-	check(InPin);
-	if (!InPin->IsOutputPin())
+	check(Pin);
+	if (Pin->IsOutputPin())
 	{
-		return Super::GetCurrentPinTypes(InPin);
+		const EPCGDataType InputTypeUnion = GetTypeUnionOfIncidentEdges(PCGPinConstants::DefaultInputLabel);
+		return InputTypeUnion != EPCGDataType::None ? InputTypeUnion : EPCGDataType::Any;
 	}
 
-	// Output pin narrows to union of inputs on both
-	const EPCGDataType InputTypeUnion = GetTypeUnionOfIncidentEdges(PCGBranchConstants::InputLabelA) | GetTypeUnionOfIncidentEdges(PCGBranchConstants::InputLabelB);
-	return InputTypeUnion != EPCGDataType::None ? InputTypeUnion : EPCGDataType::Any;
+	return Super::GetCurrentPinTypes(Pin);
 }
 
-
-TArray<FPCGPinProperties> UPCGBranchSettings::InputPinProperties() const
+TArray<FPCGPinProperties> UPCGBranchSettings::OutputPinProperties() const
 {
 	TArray<FPCGPinProperties> PinProperties;
-	PinProperties.Emplace(PCGBranchConstants::InputLabelA,
+	PinProperties.Emplace(PCGBranchConstants::OutputLabelA,
 		EPCGDataType::Any,
 		/*bInAllowMultipleConnections=*/true,
 		/*bAllowMultipleData=*/true,
-		LOCTEXT("FirstInputPinTooltip", "Will only be used if 'Use Input B' (overridable) is false"));
-	PinProperties.Emplace(PCGBranchConstants::InputLabelB,
+		LOCTEXT("OutputPinTooltipA", "Will only route input if 'Output To B' (overridable) is false"));
+	PinProperties.Emplace(PCGBranchConstants::OutputLabelB,
 		EPCGDataType::Any,
 		/*bInAllowMultipleConnections=*/true,
 		/*bAllowMultipleData=*/true,
-		LOCTEXT("SecondInputPinTooltip", "Will only be used if 'Use Input B' (overridable) is true"));
-
-	return PinProperties;
-}
-
-TArray<FPCGPinProperties> UPCGBranchSettings::OutputPinProperties() const
-{
-	TArray<FPCGPinProperties> PinProperties;
-	PinProperties.Emplace(PCGPinConstants::DefaultOutputLabel, EPCGDataType::Any, /*bInAllowMultipleConnections=*/true, /*bAllowMultipleData=*/true, LOCTEXT("OutputPinTooltip", "All input will gathered into a single data collection"));
+		LOCTEXT("OutputPinTooltipB", "Will only route input if 'Output To B' (overridable) is true"));
 
 	return PinProperties;
 }
@@ -80,10 +70,10 @@ bool FPCGBranchElement::ExecuteInternal(FPCGContext* Context) const
 	const UPCGBranchSettings* Settings = Context->GetInputSettings<UPCGBranchSettings>();
 	check(Settings);
 
-	const FName SelectedPinLabel = Settings->bUseInputB ? PCGBranchConstants::InputLabelB : PCGBranchConstants::InputLabelA;
+	const FName SelectedPinLabel = Settings->bOutputToB ? PCGBranchConstants::OutputLabelB : PCGBranchConstants::OutputLabelA;
 
 	// Reuse the functionality of the Gather Node
-	Context->OutputData = PCGGather::GatherDataForPin(Context->InputData, SelectedPinLabel);
+	Context->OutputData = PCGGather::GatherDataForPin(Context->InputData, PCGPinConstants::DefaultInputLabel, SelectedPinLabel);
 
 	return true;
 }
diff --git a/Engine/Plugins/Experimental/PCG/Source/PCG/Public/Elements/PCGBooleanSelect.h b/Engine/Plugins/Experimental/PCG/Source/PCG/Public/Elements/PCGBooleanSelect.h
new file mode 100644
--- /dev/null
+++ b/Engine/Plugins/Experimental/PCG/Source/PCG/Public/Elements/PCGBooleanSelect.h
@@ -0,0 +1,43 @@
+// Copyright Epic Games, Inc. All Rights Reserved.
+#pragma once
+
+#include "PCGSettings.h"
+
+#include "PCGBooleanSelect.generated.h"
+
+/**
+ * Selects data from either input pin, based on a boolean condition.
+ */
+UCLASS(BlueprintType, ClassGroup = (Procedural), meta=(Keywords = "if bool select"))
+class UPCGBooleanSelectSettings : public UPCGSettings
+{
+	GENERATED_BODY()
+
+public:
+	//~Begin UPCGSettings interface
+#if WITH_EDITOR
+	virtual FName GetDefaultNodeName() const override { return FName(TEXT("Select")); }
+	virtual FText GetDefaultNodeTitle() const override;
+	virtual FText GetNodeTooltipText() const override;
+	virtual EPCGSettingsType GetType() const override { return EPCGSettingsType::ControlFlow; }
+	virtual bool HasDynamicPins() const override { return true; }
+#endif
+
+	virtual EPCGDataType GetCurrentPinTypes(const UPCGPin* Pin) const override;
+
+protected:
+	virtual TArray<FPCGPinProperties> InputPinProperties() const override;
+	virtual TArray<FPCGPinProperties> OutputPinProperties() const override;
+	virtual FPCGElementPtr CreateElement() const override;
+	//~End UPCGSettings interface
+
+public:
+	UPROPERTY(BlueprintReadWrite, EditAnywhere, Category=Settings, meta=(PCG_Overridable))
+	bool bUseInputB = false;
+};
+
+class FPCGBooleanSelectElement : public FSimplePCGElement
+{
+protected:
+	virtual bool ExecuteInternal(FPCGContext* Context) const override;
+};
diff --git a/Engine/Plugins/Experimental/PCG/Source/PCG/Public/Elements/PCGBranch.h b/Engine/Plugins/Experimental/PCG/Source/PCG/Public/Elements/PCGBranch.h
--- a/Engine/Plugins/Experimental/PCG/Source/PCG/Public/Elements/PCGBranch.h
+++ b/Engine/Plugins/Experimental/PCG/Source/PCG/Public/Elements/PCGBranch.h
@@ -6,7 +6,7 @@
 #include "PCGBranch.generated.h"
 
 /**
- * Selects data from either input pin, based on a boolean condition.
+ * Routes input data to one of two outputs,based on a boolean condition.
  */
 UCLASS(BlueprintType, ClassGroup = (Procedural), meta=(Keywords = "if bool branch"))
 class UPCGBranchSettings : public UPCGSettings
@@ -23,17 +23,16 @@ public:
 	virtual bool HasDynamicPins() const override { return true; }
 #endif
 
-	virtual EPCGDataType GetCurrentPinTypes(const UPCGPin* InPin) const override;
+	virtual EPCGDataType GetCurrentPinTypes(const UPCGPin* Pin) const override;
 
 protected:
-	virtual TArray<FPCGPinProperties> InputPinProperties() const override;
 	virtual TArray<FPCGPinProperties> OutputPinProperties() const override;
 	virtual FPCGElementPtr CreateElement() const override;
 	//~End UPCGSettings interface
 
 public:
 	UPROPERTY(BlueprintReadWrite, EditAnywhere, Category=Settings, meta=(PCG_Overridable))
-	bool bUseInputB = false;
+	bool bOutputToB = false;
 };
 
 class FPCGBranchElement : public FSimplePCGElement
