[PCG] Hotfix for minor bug in Create Attribute Node .

#rnx

diff --git a/Engine/Plugins/Experimental/PCG/Source/PCG/Private/Elements/PCGCreateAttribute.cpp b/Engine/Plugins/Experimental/PCG/Source/PCG/Private/Elements/PCGCreateAttribute.cpp
--- a/Engine/Plugins/Experimental/PCG/Source/PCG/Private/Elements/PCGCreateAttribute.cpp
+++ b/Engine/Plugins/Experimental/PCG/Source/PCG/Private/Elements/PCGCreateAttribute.cpp
@@ -358,7 +358,7 @@ bool FPCGCreateAttributeElement::ExecuteInternal(FPCGContext* Context) const
 			// If no field accessor, copy over the attribute
 			if (InputSource.GetExtraNames().IsEmpty())
 			{
-				const FPCGMetadataAttributeBase* SourceAttribute = Metadata->GetConstAttribute(SourceParamAttributeName);
+				const FPCGMetadataAttributeBase* SourceAttribute = SourceParamData->Metadata->GetConstAttribute(SourceParamAttributeName);
 				Attribute = Metadata->CopyAttribute(SourceAttribute, OutputAttributeName, /*bKeepParent=*/false, /*bCopyEntries=*/bShouldAddNewEntry, /*bCopyValues=*/bShouldAddNewEntry);
 			}
 			else // Create a new attribute of the accessed field's type manually
