[PCG] Update tracking keys to support other than AllWorldActors + add extra dependency objects for tracking

* Fixes generation infinite loop in a case where an actor has multiple PCG components.

#rnx

diff --git a/Engine/Plugins/Experimental/PCG/Source/PCG/Private/Elements/PCGActorSelector.cpp b/Engine/Plugins/Experimental/PCG/Source/PCG/Private/Elements/PCGActorSelector.cpp
--- a/Engine/Plugins/Experimental/PCG/Source/PCG/Private/Elements/PCGActorSelector.cpp
+++ b/Engine/Plugins/Experimental/PCG/Source/PCG/Private/Elements/PCGActorSelector.cpp
@@ -204,24 +204,50 @@ namespace PCGActorSelector
 	}
 }
 
+FPCGActorSelectionKey::FPCGActorSelectionKey(EPCGActorFilter InFilter)
+{
+	check(InFilter != EPCGActorFilter::AllWorldActors);
+	ActorFilter = InFilter;
+}
+
 FPCGActorSelectionKey::FPCGActorSelectionKey(FName InTag)
 {
 	Selection = EPCGActorSelection::ByTag;
 	Tag = InTag;
+	ActorFilter = EPCGActorFilter::AllWorldActors;
 }
 
 FPCGActorSelectionKey::FPCGActorSelectionKey(TSubclassOf<AActor> InSelectionClass)
 {
 	Selection = EPCGActorSelection::ByClass;
 	ActorSelectionClass = InSelectionClass;
+	ActorFilter = EPCGActorFilter::AllWorldActors;
 }
 
-bool FPCGActorSelectionKey::IsMatching(const AActor* InActor) const
+void FPCGActorSelectionKey::SetExtraDependency(const UClass* InExtraDependency)
+{
+	OptionalExtraDependency = InExtraDependency;
+}
+
+bool FPCGActorSelectionKey::IsMatching(const AActor* InActor, const UPCGComponent* InComponent) const
 {
 	if (!InActor)
 	{
 		return false;
 	}
+	
+	// If we filter something else than all world actors, matching depends on the component.
+	// Re-use the same mecanism than Get Actor Data, which should be cheap since we don't look for all actors in the world.
+	if (ActorFilter != EPCGActorFilter::AllWorldActors)
+	{
+		// InKey provide the info for selecting a given actor.
+		// We reconstruct the selector settings from this key, and we also force it to SelectMultiple, since
+		// we want to gather all the actors that matches this given key, to find if ours matches.
+		FPCGActorSelectorSettings SelectorSettings = FPCGActorSelectorSettings::ReconstructFromKey(*this);
+		SelectorSettings.bSelectMultiple = true;
+		TArray<AActor*> AllActors = PCGActorSelector::FindActors(SelectorSettings, InComponent, [](const AActor*) { return true; }, [](const AActor*) { return true; });
+		return AllActors.Contains(InActor);
+	}
 
 	switch (Selection)
 	{
@@ -234,35 +260,9 @@ bool FPCGActorSelectionKey::IsMatching(const AActor* InActor) const
 	}
 }
 
-TArray<FPCGActorSelectionKey> FPCGActorSelectionKey::GenerateAllKeysForActor(const AActor* InActor)
-{
-	TArray<FPCGActorSelectionKey> Result;
-
-	if (!InActor)
-	{
-		return Result;
-	}
-
-	Result.Reserve(InActor->Tags.Num() + 1);
-	// Register keys all the way up the chain
-	UClass* CurrentClass = InActor->GetClass();
-	while (CurrentClass && CurrentClass != AActor::StaticClass())
-	{
-		Result.Emplace(CurrentClass);
-		CurrentClass = CurrentClass->GetSuperClass();
-	}
-
-	for (FName Tag : InActor->Tags)
-	{
-		Result.Emplace(Tag);
-	}
-
-	return Result;
-}
-
 bool FPCGActorSelectionKey::operator==(const FPCGActorSelectionKey& InOther) const
 {
-	if (Selection != InOther.Selection)
+	if (ActorFilter != InOther.ActorFilter || Selection != InOther.Selection || OptionalExtraDependency != InOther.OptionalExtraDependency)
 	{
 		return false;
 	}
@@ -280,7 +280,12 @@ bool FPCGActorSelectionKey::operator==(const FPCGActorSelectionKey& InOther) con
 
 uint32 GetTypeHash(const FPCGActorSelectionKey& In)
 {
-	return HashCombine(HashCombine(GetTypeHash(In.Selection), GetTypeHash(In.Tag)), GetTypeHash(In.ActorSelectionClass));
+	uint32 HashResult = HashCombine(GetTypeHash(In.ActorFilter), GetTypeHash(In.Selection));
+	HashResult = HashCombine(HashResult, GetTypeHash(In.Tag));
+	HashResult = HashCombine(HashResult, GetTypeHash(In.ActorSelectionClass));
+	HashResult = HashCombine(HashResult, GetTypeHash(In.OptionalExtraDependency));
+
+	return HashResult;
 }
 
 #if WITH_EDITOR
@@ -313,10 +318,9 @@ FName FPCGActorSelectorSettings::GetTaskName(const FText& Prefix) const
 
 FPCGActorSelectionKey FPCGActorSelectorSettings::GetAssociatedKey() const
 {
-	// If we don't look for AllWorldActors, it means we track the pcg component, which should be already picked up by the tracking system.
 	if (ActorFilter != EPCGActorFilter::AllWorldActors)
 	{
-		return FPCGActorSelectionKey();
+		return FPCGActorSelectionKey(ActorFilter);
 	}
 
 	switch (ActorSelection)
@@ -333,7 +337,7 @@ FPCGActorSelectionKey FPCGActorSelectorSettings::GetAssociatedKey() const
 FPCGActorSelectorSettings FPCGActorSelectorSettings::ReconstructFromKey(const FPCGActorSelectionKey& InKey)
 {
 	FPCGActorSelectorSettings Result{};
-	Result.ActorFilter = EPCGActorFilter::AllWorldActors;
+	Result.ActorFilter = InKey.ActorFilter;
 	Result.ActorSelection = InKey.Selection;
 	Result.ActorSelectionTag = InKey.Tag;
 	Result.ActorSelectionClass = InKey.ActorSelectionClass;
diff --git a/Engine/Plugins/Experimental/PCG/Source/PCG/Private/Elements/PCGDataFromActor.cpp b/Engine/Plugins/Experimental/PCG/Source/PCG/Private/Elements/PCGDataFromActor.cpp
--- a/Engine/Plugins/Experimental/PCG/Source/PCG/Private/Elements/PCGDataFromActor.cpp
+++ b/Engine/Plugins/Experimental/PCG/Source/PCG/Private/Elements/PCGDataFromActor.cpp
@@ -22,10 +22,13 @@
 #if WITH_EDITOR
 void UPCGDataFromActorSettings::GetTrackedActorKeys(FPCGActorSelectionKeyToSettingsMap& OutKeysToSettings, TArray<TObjectPtr<const UPCGGraph>>& OutVisitedGraphs) const
 {
-	if (ActorSelector.ActorFilter == EPCGActorFilter::AllWorldActors)
+	FPCGActorSelectionKey Key = ActorSelector.GetAssociatedKey();
+	if (Mode == EPCGGetDataFromActorMode::GetDataFromPCGComponent || Mode == EPCGGetDataFromActorMode::GetDataFromPCGComponentOrParseComponents)
 	{
-		OutKeysToSettings.FindOrAdd(ActorSelector.GetAssociatedKey()).Emplace(this, bTrackActorsOnlyWithinBounds);
+		Key.SetExtraDependency(UPCGComponent::StaticClass());
 	}
+
+	OutKeysToSettings.FindOrAdd(Key).Emplace(this, bTrackActorsOnlyWithinBounds);
 }
 
 FText UPCGDataFromActorSettings::GetNodeTooltipText() const
diff --git a/Engine/Plugins/Experimental/PCG/Source/PCG/Private/Elements/PCGPropertyToParamData.cpp b/Engine/Plugins/Experimental/PCG/Source/PCG/Private/Elements/PCGPropertyToParamData.cpp
--- a/Engine/Plugins/Experimental/PCG/Source/PCG/Private/Elements/PCGPropertyToParamData.cpp
+++ b/Engine/Plugins/Experimental/PCG/Source/PCG/Private/Elements/PCGPropertyToParamData.cpp
@@ -17,10 +17,7 @@
 #if WITH_EDITOR
 void UPCGPropertyToParamDataSettings::GetTrackedActorKeys(FPCGActorSelectionKeyToSettingsMap& OutKeysToSettings, TArray<TObjectPtr<const UPCGGraph>>& OutVisitedGraphs) const
 {
-	if (ActorSelector.ActorFilter == EPCGActorFilter::AllWorldActors)
-	{
-		OutKeysToSettings.FindOrAdd(ActorSelector.GetAssociatedKey()).Emplace(this, bTrackActorsOnlyWithinBounds);
-	}
+	OutKeysToSettings.FindOrAdd(ActorSelector.GetAssociatedKey()).Emplace(this, bTrackActorsOnlyWithinBounds);
 }
 
 void UPCGPropertyToParamDataSettings::PostEditChangeProperty(FPropertyChangedEvent& PropertyChangedEvent)
diff --git a/Engine/Plugins/Experimental/PCG/Source/PCG/Private/PCGActorAndComponentMapping.cpp b/Engine/Plugins/Experimental/PCG/Source/PCG/Private/PCGActorAndComponentMapping.cpp
--- a/Engine/Plugins/Experimental/PCG/Source/PCG/Private/PCGActorAndComponentMapping.cpp
+++ b/Engine/Plugins/Experimental/PCG/Source/PCG/Private/PCGActorAndComponentMapping.cpp
@@ -592,11 +592,6 @@ void UPCGActorAndComponentMapping::RegisterOrUpdateTracking(UPCGComponent* InCom
 		return;
 	}
 
-	// Components owner needs to be always tracked
-	RegisterActor(ComponentOwner);
-	TSet<UPCGComponent*>& AllComponents = AlwaysTrackedActorsToComponentsMap.FindOrAdd(ComponentOwner);
-	AllComponents.Add(InComponent);
-
 	UWorld* World = PCGSubsystem ? PCGSubsystem->GetWorld() : nullptr;
 
 	if (!World)
@@ -611,6 +606,14 @@ void UPCGActorAndComponentMapping::RegisterOrUpdateTracking(UPCGComponent* InCom
 		return;
 	}
 
+	// Components owner needs to be always tracked
+	RegisterActor(ComponentOwner);
+	TSet<UPCGComponent*>& AllComponents = AlwaysTrackedActorsToComponentsMap.FindOrAdd(ComponentOwner);
+	AllComponents.Add(InComponent);
+
+	FPCGActorSelectionKey OwnerKey = FPCGActorSelectionKey(EPCGActorFilter::Self);
+	KeysToComponentsMap.FindOrAdd(OwnerKey).Add(InComponent);
+
 	// And we also need to find all actors that should be tracked
 	if (UPCGGraph* PCGGraph = InComponent->GetGraph())
 	{
@@ -650,7 +653,7 @@ void UPCGActorAndComponentMapping::RegisterOrUpdateTracking(UPCGComponent* InCom
 				if (bInShouldDirtyActors)
 				{
 					// If we need to force dirty, disregard culling (always intersect).
-					InComponent->DirtyTrackedActor(Actor, /*bIntersect=*/ true, {});
+					InComponent->DirtyTrackedActor(Actor, /*bIntersect=*/ true, /*InRemovedTags=*/ {}, /*InOriginatingChangeObject=*/ nullptr);
 				}
 			}
 		};
@@ -1040,7 +1043,6 @@ void UPCGActorAndComponentMapping::OnObjectPropertyChanged(UObject* InObject, FP
 		return;
 	}
 
-	AActor* Actor = Cast<AActor>(InObject);
 	APCGWorldActor* PCGWorldActor = PCGSubsystem ? PCGSubsystem->FindPCGWorldActor() : nullptr;
 
 	if (!PCGWorldActor)
@@ -1048,6 +1050,18 @@ void UPCGActorAndComponentMapping::OnObjectPropertyChanged(UObject* InObject, FP
 		return;
 	}
 
+	// First check if it is an actor
+	AActor* Actor = Cast<AActor>(InObject);
+
+	// Otherwise, if it's an actor component, track it as well
+	if (!Actor)
+	{
+		if (UActorComponent* ActorComponent = Cast<UActorComponent>(InObject))
+		{
+			Actor = ActorComponent->GetOwner();
+		}
+	}
+
 	// If we don't find any actor, try to see if it is a dependency
 	if (!Actor)
 	{
@@ -1057,7 +1071,7 @@ void UPCGActorAndComponentMapping::OnObjectPropertyChanged(UObject* InObject, FP
 			{
 				if (AActor* ActorToChange = TrackedActor.Key.ResolveObjectPtr())
 				{
-					OnActorChanged(ActorToChange, /*bInHasMoved=*/ false);
+					OnActorChanged(ActorToChange, /*bInHasMoved=*/ false, /*InOriginatingChangeObject=*/ InObject);
 					UpdateActorDependencies(ActorToChange);
 				}
 			}
@@ -1075,7 +1089,7 @@ void UPCGActorAndComponentMapping::OnObjectPropertyChanged(UObject* InObject, FP
 
 	if (bShouldChange)
 	{
-		OnActorChanged(Actor, /*bInHasMoved=*/ false);
+		OnActorChanged(Actor, /*bInHasMoved=*/ false, /*InOriginatingChangeObject=*/ InObject);
 	}
 	else
 	{
@@ -1084,7 +1098,7 @@ void UPCGActorAndComponentMapping::OnObjectPropertyChanged(UObject* InObject, FP
 	}
 }
 
-void UPCGActorAndComponentMapping::OnActorChanged(AActor* InActor, bool bInHasMoved, const UPCGComponent* InComponentChanged)
+void UPCGActorAndComponentMapping::OnActorChanged(AActor* InActor, bool bInHasMoved, const UObject* InOriginatingChangeObject)
 {
 	TRACE_CPUPROFILER_EVENT_SCOPE(UPCGActorAndComponentMapping::OnActorChanged);
 
@@ -1108,19 +1122,19 @@ void UPCGActorAndComponentMapping::OnActorChanged(AActor* InActor, bool bInHasMo
 		// Then do an octree find to get all components that intersect with this actor.
 		// If the actor has moved, we also need to find components that intersected with it before
 		// We first do it for non-partitioned, then we do it for partitioned
-		auto UpdateNonPartitioned = [&DirtyComponents, InActor, CulledTrackedComponents, &RemovedTags, DirtyFlag, InComponentChanged](const FPCGComponentRef& ComponentRef) -> void
+		auto UpdateNonPartitioned = [&DirtyComponents, InActor, CulledTrackedComponents, &RemovedTags, DirtyFlag, InOriginatingChangeObject](const FPCGComponentRef& ComponentRef) -> void
 		{
 			TRACE_CPUPROFILER_EVENT_SCOPE(UPCGActorAndComponentMapping::OnActorChanged::UpdateNonPartitioned);
 
 			// Don't dirty if the component was already dirtied, not tracked, or the origin of the change.
 			if (DirtyComponents.Contains(ComponentRef.Component) || 
 				!CulledTrackedComponents->Contains(ComponentRef.Component) ||
-				InComponentChanged == ComponentRef.Component)
+				InOriginatingChangeObject == ComponentRef.Component)
 			{
 				return;
 			}
 
-			if (ComponentRef.Component->DirtyTrackedActor(InActor, /*bIntersect=*/true, RemovedTags))
+			if (ComponentRef.Component->DirtyTrackedActor(InActor, /*bIntersect=*/true, RemovedTags, InOriginatingChangeObject))
 			{
 				ComponentRef.Component->DirtyGenerated(DirtyFlag);
 				DirtyComponents.Add(ComponentRef.Component);
@@ -1130,7 +1144,7 @@ void UPCGActorAndComponentMapping::OnActorChanged(AActor* InActor, bool bInHasMo
 		NonPartitionedOctree.FindElementsWithBoundsTest(ActorBounds, UpdateNonPartitioned);
 
 		// For partitioned, we first need to find all components that intersect with our actor and then forward the dirty call to all local components that intersect.
-		auto UpdatePartitioned = [this, &DirtyComponents, InActor, CulledTrackedComponents, &ActorBounds, &RemovedTags, DirtyFlag, InComponentChanged](const FPCGComponentRef& ComponentRef)  -> void
+		auto UpdatePartitioned = [this, &DirtyComponents, InActor, CulledTrackedComponents, &ActorBounds, &RemovedTags, DirtyFlag, InOriginatingChangeObject](const FPCGComponentRef& ComponentRef)  -> void
 		{
 			TRACE_CPUPROFILER_EVENT_SCOPE(UPCGActorAndComponentMapping::OnActorChanged::UpdatePartitioned);
 
@@ -1138,7 +1152,7 @@ void UPCGActorAndComponentMapping::OnActorChanged(AActor* InActor, bool bInHasMo
 			// We can "re-dirty" it because changes can impact different local components, from the same
 			// original component.
 			if (!CulledTrackedComponents->Contains(ComponentRef.Component) ||
-				InComponentChanged == ComponentRef.Component)
+				InOriginatingChangeObject == ComponentRef.Component)
 			{
 				return;
 			}
@@ -1146,11 +1160,11 @@ void UPCGActorAndComponentMapping::OnActorChanged(AActor* InActor, bool bInHasMo
 			const FBox Overlap = ActorBounds.Overlap(ComponentRef.Bounds.GetBox());
 			bool bWasDirtied = false;
 
-			ForAllIntersectingPartitionActors(Overlap, [InActor, Component = ComponentRef.Component, &RemovedTags, &bWasDirtied, DirtyFlag](APCGPartitionActor* InPartitionActor) -> void
+			ForAllIntersectingPartitionActors(Overlap, [InActor, Component = ComponentRef.Component, &RemovedTags, &bWasDirtied, DirtyFlag, InOriginatingChangeObject](APCGPartitionActor* InPartitionActor) -> void
 			{
 				if (UPCGComponent* LocalComponent = InPartitionActor->GetLocalComponent(Component))
 				{
-					if (LocalComponent->DirtyTrackedActor(InActor, /*bIntersect=*/true, RemovedTags))
+					if (LocalComponent->DirtyTrackedActor(InActor, /*bIntersect=*/true, RemovedTags, InOriginatingChangeObject))
 					{
 						bWasDirtied = true;
 						LocalComponent->DirtyGenerated(DirtyFlag);
@@ -1193,21 +1207,23 @@ void UPCGActorAndComponentMapping::OnActorChanged(AActor* InActor, bool bInHasMo
 
 		for (UPCGComponent* PCGComponent : *AlwaysTrackedComponents)
 		{
-			if (!PCGComponent || PCGComponent == InComponentChanged)
+			if (!PCGComponent || PCGComponent == InOriginatingChangeObject)
 			{
 				continue;
 			}
 
-			const bool bOwnerChanged = PCGComponent->GetOwner() == InActor;
+			// Don't mark "Owner changed" if the change originate from a PCG Component. It will be delegated to the DirtyTrackedActor
+			// It is necessary to avoid infine loops when there are multiple PCG components on one actor, and one component was generated.
+			const bool bOwnerChanged = (PCGComponent->GetOwner() == InActor) && (!InOriginatingChangeObject || !InOriginatingChangeObject->IsA<UPCGComponent>());
 			bool bWasDirtied = false;
 
 			if (!DirtyComponents.Contains(PCGComponent) && !bOwnerChanged)
 			{
 				if (PCGComponent->IsPartitioned())
 				{
-					DispatchToRegisteredLocalComponents(PCGComponent, [InActor, &RemovedTags, &bWasDirtied, DirtyFlag](UPCGComponent* InLocalComponent) -> FPCGTaskId
+					DispatchToRegisteredLocalComponents(PCGComponent, [InActor, &RemovedTags, &bWasDirtied, DirtyFlag, InOriginatingChangeObject](UPCGComponent* InLocalComponent) -> FPCGTaskId
 					{
-						if (InLocalComponent->DirtyTrackedActor(InActor, /*bIntersect=*/false, RemovedTags))
+						if (InLocalComponent->DirtyTrackedActor(InActor, /*bIntersect=*/false, RemovedTags, InOriginatingChangeObject))
 						{
 							bWasDirtied = true;
 							InLocalComponent->DirtyGenerated(DirtyFlag);
@@ -1217,7 +1233,7 @@ void UPCGActorAndComponentMapping::OnActorChanged(AActor* InActor, bool bInHasMo
 				}
 				else
 				{
-					bWasDirtied = PCGComponent->DirtyTrackedActor(InActor, /*bIntersect=*/false, RemovedTags);
+					bWasDirtied = PCGComponent->DirtyTrackedActor(InActor, /*bIntersect=*/false, RemovedTags, InOriginatingChangeObject);
 				}
 			}
 
diff --git a/Engine/Plugins/Experimental/PCG/Source/PCG/Private/PCGComponent.cpp b/Engine/Plugins/Experimental/PCG/Source/PCG/Private/PCGComponent.cpp
--- a/Engine/Plugins/Experimental/PCG/Source/PCG/Private/PCGComponent.cpp
+++ b/Engine/Plugins/Experimental/PCG/Source/PCG/Private/PCGComponent.cpp
@@ -1649,11 +1649,11 @@ bool UPCGComponent::IsActorTracked(AActor* InActor, bool& bOutIsCulled) const
 
 	bool bFound = false;
 
-	for (const FPCGActorSelectionKey& Key : FPCGActorSelectionKey::GenerateAllKeysForActor(InActor))
+	for (const TPair<FPCGActorSelectionKey, bool>& It : CachedTrackedKeysToCulling)
 	{
-		if (const bool* ShouldCullPtr = CachedTrackedKeysToCulling.Find(Key))
+		if (It.Key.IsMatching(InActor, this))
 		{
-			bOutIsCulled = *ShouldCullPtr;
+			bOutIsCulled = It.Value;
 			bFound = true;
 
 			// Check for other tags that might not be culled
@@ -2222,7 +2222,7 @@ UPCGSubsystem* UPCGComponent::GetSubsystem() const
 }
 
 #if WITH_EDITOR
-bool UPCGComponent::DirtyTrackedActor(AActor* InActor, bool bIntersect, const TSet<FName>& InRemovedTags)
+bool UPCGComponent::DirtyTrackedActor(AActor* InActor, bool bIntersect, const TSet<FName>& InRemovedTags, const UObject* InOriginatingChangeObject)
 {
 	TRACE_CPUPROFILER_EVENT_SCOPE(UPCGComponent::DirtyTrackedActor);
 
@@ -2233,35 +2233,22 @@ bool UPCGComponent::DirtyTrackedActor(AActor* InActor, bool bIntersect, const TS
 
 	bool bWasDirtied = false;
 
-	TArray<FPCGActorSelectionKey> AllKeys = FPCGActorSelectionKey::GenerateAllKeysForActor(InActor);
-
-	// We also need to force a refresh on tags that were removed.
-	for (FName RemovedTag : InRemovedTags)
+	for (const auto& It : CachedTrackedKeysToSettings)
 	{
-		AllKeys.Emplace(RemovedTag);
-	}
+		const FPCGActorSelectionKey& Key = It.Key;
 
-	for (FPCGActorSelectionKey Key : AllKeys)
-	{
-		const TArray<FPCGSettingsAndCulling>* SettingsAndCullingPtr = nullptr;
-		
-		while (!SettingsAndCullingPtr)
-		{
-			SettingsAndCullingPtr = CachedTrackedKeysToSettings.Find(Key);
+		const bool bRemovedTagIsTracked = (Key.Selection == EPCGActorSelection::ByTag) && InRemovedTags.Contains(Key.Tag);
 
-			if (!SettingsAndCullingPtr && Key.Selection == EPCGActorSelection::ByClass && Key.ActorSelectionClass != nullptr)
-			{
-				Key.ActorSelectionClass = (Key.ActorSelectionClass == AActor::StaticClass()) ? nullptr : Key.ActorSelectionClass->GetSuperClass();
-			}
-			else
+		if (It.Key.IsMatching(InActor, this) || bRemovedTagIsTracked)
+		{
+			// Extra care if the change originates from a PCGComponent. Only dirty if we are tracking a PCG component.
+			if (InOriginatingChangeObject && InOriginatingChangeObject->IsA<UPCGComponent>() 
+				&& (!It.Key.OptionalExtraDependency || !It.Key.OptionalExtraDependency->IsChildOf(UPCGComponent::StaticClass())))
 			{
-				break;
+				continue;
 			}
-		} 
 
-		if (SettingsAndCullingPtr)
-		{
-			for (const FPCGSettingsAndCulling& SettingsAndCulling : *SettingsAndCullingPtr)
+			for (const FPCGSettingsAndCulling& SettingsAndCulling : It.Value)
 			{
 				if (SettingsAndCulling.Value && !bIntersect)
 				{
diff --git a/Engine/Plugins/Experimental/PCG/Source/PCG/Public/Elements/PCGActorSelector.h b/Engine/Plugins/Experimental/PCG/Source/PCG/Public/Elements/PCGActorSelector.h
--- a/Engine/Plugins/Experimental/PCG/Source/PCG/Public/Elements/PCGActorSelector.h
+++ b/Engine/Plugins/Experimental/PCG/Source/PCG/Public/Elements/PCGActorSelector.h
@@ -17,41 +17,49 @@ enum class EPCGActorSelection : uint8
 	ByTag,
 	// Deprecated - actor labels are unavailable in shipping builds
 	ByName UMETA(Hidden),
-	ByClass
+	ByClass,
+	Unknown UMETA(Hidden)
+};
+
+UENUM()
+enum class EPCGActorFilter : uint8
+{
+	/** This actor (either the original PCG actor or the partition actor if partitioning is enabled). */
+	Self,
+	/** The parent of this actor in the hierarchy. */
+	Parent,
+	/** The top most parent of this actor in the hierarchy. */
+	Root,
+	/** All actors in world. */
+	AllWorldActors,
+	/** The source PCG actor (rather than the generated partition actor). */
+	Original,
 };
 
 struct FPCGActorSelectionKey
 {
 	FPCGActorSelectionKey() = default;
 
+	// For all filters others than AllWorldActor. For AllWorldActors Filter, use the other constructors.
+	explicit FPCGActorSelectionKey(EPCGActorFilter InFilter);
+
 	explicit FPCGActorSelectionKey(FName InTag);
 	explicit FPCGActorSelectionKey(TSubclassOf<AActor> InSelectionClass);
 
 	bool operator==(const FPCGActorSelectionKey& InOther) const;
 
 	friend uint32 GetTypeHash(const FPCGActorSelectionKey& In);
-	bool IsMatching(const AActor* InActor) const;
+	bool IsMatching(const AActor* InActor, const UPCGComponent* InComponent) const;
 
-	static TArray<FPCGActorSelectionKey> GenerateAllKeysForActor(const AActor* InActor);
+	void SetExtraDependency(const UClass* InExtraDependency);
 
-	EPCGActorSelection Selection = EPCGActorSelection::ByTag;
+	EPCGActorFilter ActorFilter = EPCGActorFilter::AllWorldActors;
+	EPCGActorSelection Selection = EPCGActorSelection::Unknown;
 	FName Tag = NAME_None;
 	TSubclassOf<AActor> ActorSelectionClass = nullptr;
-};
 
-UENUM()
-enum class EPCGActorFilter : uint8
-{
-	/** This actor (either the original PCG actor or the partition actor if partitioning is enabled). */
-	Self,
-	/** The parent of this actor in the hierarchy. */
-	Parent,
-	/** The top most parent of this actor in the hierarchy. */
-	Root,
-	/** All actors in world. */
-	AllWorldActors,
-	/** The source PCG actor (rather than the generated partition actor). */
-	Original,
+	// If it should track a specific object dependency instead of an actor. For example, GetActorData with GetPCGComponent data.
+	const UClass* OptionalExtraDependency = nullptr;
 };
 
 USTRUCT(BlueprintType)
diff --git a/Engine/Plugins/Experimental/PCG/Source/PCG/Public/PCGActorAndComponentMapping.h b/Engine/Plugins/Experimental/PCG/Source/PCG/Public/PCGActorAndComponentMapping.h
--- a/Engine/Plugins/Experimental/PCG/Source/PCG/Public/PCGActorAndComponentMapping.h
+++ b/Engine/Plugins/Experimental/PCG/Source/PCG/Public/PCGActorAndComponentMapping.h
@@ -138,9 +138,9 @@ private:
 
 	/** Trigger an update when the actor changed. 
 	* Can specify if the actor has moved to also update components that were at its previous position.
-	* Can also specify an optional component, originating the change, to avoid re-dirtying it.
+	* Can also specify an optional object, originating the change, to avoid re-dirtying a component if it was the origin.
 	*/
-	void OnActorChanged(AActor* InActor, bool bInHasMoved, const UPCGComponent* InComponentChanged = nullptr);
+	void OnActorChanged(AActor* InActor, bool bInHasMoved, const UObject* InOriginatingChangeObject = nullptr);
 
 	/** Update dependencies for a given tracked actor. */
 	void UpdateActorDependencies(AActor* InActor);
diff --git a/Engine/Plugins/Experimental/PCG/Source/PCG/Public/PCGComponent.h b/Engine/Plugins/Experimental/PCG/Source/PCG/Public/PCGComponent.h
--- a/Engine/Plugins/Experimental/PCG/Source/PCG/Public/PCGComponent.h
+++ b/Engine/Plugins/Experimental/PCG/Source/PCG/Public/PCGComponent.h
@@ -350,10 +350,9 @@ private:
 	void SetupCallbacksOnCreation();
 
 	void UpdateTrackingCache();
-
-	bool ActorIsTracked(AActor* InActor) const;
 	
-	bool DirtyTrackedActor(AActor* InActor, bool bIntersect, const TSet<FName>& InRemovedTags);
+	// InOriginatingChangeObject can be null
+	bool DirtyTrackedActor(AActor* InActor, bool bIntersect, const TSet<FName>& InRemovedTags, const UObject* InOriginatingChangeObject);
 
 	bool ShouldTrackLandscape() const;
 #endif
