[PCG] Add inspection support for dynamically executed graphs

diff --git a/Engine/Plugins/Experimental/PCG/Source/PCG/Private/Elements/PCGLoopElement.cpp b/Engine/Plugins/Experimental/PCG/Source/PCG/Private/Elements/PCGLoopElement.cpp
--- a/Engine/Plugins/Experimental/PCG/Source/PCG/Private/Elements/PCGLoopElement.cpp
+++ b/Engine/Plugins/Experimental/PCG/Source/PCG/Private/Elements/PCGLoopElement.cpp
@@ -117,6 +117,10 @@ bool FPCGLoopElement::ExecuteInternal(FPCGContext* InContext) const
 			StackFrames.Emplace(Context->Node);
 			StackFrames.Emplace(EntryIndex);
 
+#if WITH_EDITOR
+			Subgraph->OnGraphDynamicallyExecutedDelegate.Broadcast(Subgraph, Context->SourceComponent, InvocationStack);
+#endif
+
 			FPCGTaskId SubgraphTaskId = Subsystem->ScheduleGraph(Subgraph, Context->SourceComponent.Get(), PreGraphElement, MakeShared<FPCGInputForwardingElement>(InputDataCollection), {}, &InvocationStack);
 
 			Context->SubgraphTaskIds.Add(SubgraphTaskId);
diff --git a/Engine/Plugins/Experimental/PCG/Source/PCG/Private/Elements/PCGSpawnActor.cpp b/Engine/Plugins/Experimental/PCG/Source/PCG/Private/Elements/PCGSpawnActor.cpp
--- a/Engine/Plugins/Experimental/PCG/Source/PCG/Private/Elements/PCGSpawnActor.cpp
+++ b/Engine/Plugins/Experimental/PCG/Source/PCG/Private/Elements/PCGSpawnActor.cpp
@@ -613,7 +613,13 @@ bool FPCGSpawnActorElement::SpawnAndPrepareSubgraphs(FPCGSubgraphContext* Contex
 							// Prepare the invocation stack - which is the stack up to this node, and then this node, then a loop index
 							FPCGStack InvocationStack = ensure(Context->Stack) ? *Context->Stack : FPCGStack();
 
-							FPCGTaskId SubgraphTaskId = Subsystem->ScheduleGraph(GraphInterface->GetGraph(),
+							UPCGGraph* Graph = GraphInterface->GetGraph();
+
+#if WITH_EDITOR
+							Graph->OnGraphDynamicallyExecutedDelegate.Broadcast(Graph, Context->SourceComponent, InvocationStack);
+#endif
+
+							FPCGTaskId SubgraphTaskId = Subsystem->ScheduleGraph(Graph,
 								Context->SourceComponent.Get(),
 								MakeShared<FPCGTrivialElement>(),// TODO: prepare user parameters like in subgraph/loop
 								MakeShared<FPCGInputForwardingElement>(SubgraphInputData),
diff --git a/Engine/Plugins/Experimental/PCG/Source/PCG/Private/Graph/PCGGraphExecutor.h b/Engine/Plugins/Experimental/PCG/Source/PCG/Private/Graph/PCGGraphExecutor.h
--- a/Engine/Plugins/Experimental/PCG/Source/PCG/Private/Graph/PCGGraphExecutor.h
+++ b/Engine/Plugins/Experimental/PCG/Source/PCG/Private/Graph/PCGGraphExecutor.h
@@ -18,7 +18,7 @@ class UPCGGraph;
 class UPCGNode;
 class UPCGComponent;
 class FPCGGraphCompiler;
-class FPCGStack;
+struct FPCGStack;
 class FPCGStackContext;
 class FTextFormat;
 
diff --git a/Engine/Plugins/Experimental/PCG/Source/PCG/Private/Graph/PCGStackContext.cpp b/Engine/Plugins/Experimental/PCG/Source/PCG/Private/Graph/PCGStackContext.cpp
--- a/Engine/Plugins/Experimental/PCG/Source/PCG/Private/Graph/PCGStackContext.cpp
+++ b/Engine/Plugins/Experimental/PCG/Source/PCG/Private/Graph/PCGStackContext.cpp
@@ -10,6 +10,8 @@
 #include "Containers/UnrealString.h"
 #include "Misc/StringBuilder.h"
 
+#include UE_INLINE_GENERATED_CPP_BY_NAME(PCGStackContext)
+
 const FPCGStack* FPCGStackContext::GetStack(int32 InStackIndex) const
 {
 	if (ensure(Stacks.IsValidIndex(InStackIndex)))
@@ -113,6 +115,18 @@ bool FPCGStack::operator==(const FPCGStack& Other) const
 	return true;
 }
 
+uint32 GetTypeHash(const FPCGStack& In)
+{
+	uint32 Hash = 0;
+
+	for (const FPCGStackFrame& Frame : In.StackFrames)
+	{
+		Hash = HashCombine(Hash, GetTypeHash(Frame));
+	}
+
+	return Hash;
+}
+
 int32 FPCGStackContext::PushFrame(const UObject* InFrameObject)
 {
 	if (CurrentStackIndex == INDEX_NONE)
diff --git a/Engine/Plugins/Experimental/PCG/Source/PCG/Private/PCGComponent.cpp b/Engine/Plugins/Experimental/PCG/Source/PCG/Private/PCGComponent.cpp
--- a/Engine/Plugins/Experimental/PCG/Source/PCG/Private/PCGComponent.cpp
+++ b/Engine/Plugins/Experimental/PCG/Source/PCG/Private/PCGComponent.cpp
@@ -1508,25 +1508,27 @@ void UPCGComponent::StoreInspectionData(const FPCGStack* InStack, const UPCGNode
 		return;
 	}
 
-	auto StorePinInspectionData = [InStack, InNode](const TArray<TObjectPtr<UPCGPin>>& InPins, const FPCGDataCollection& InData, TMap<FString, FPCGDataCollection>& InOutInspectionCache)
+	auto StorePinInspectionData = [InStack, InNode](const TArray<TObjectPtr<UPCGPin>>& InPins, const FPCGDataCollection& InData, TMap<FPCGStack, FPCGDataCollection>& InOutInspectionCache)
 	{
 		for (const UPCGPin* Pin : InPins)
 		{
-			FString StackFramePath;
-			if (!ensure(InStack->CreateStackFramePath(StackFramePath, InNode, Pin)))
-			{
-				continue;
-			}
+			FPCGStack Stack = *InStack;
+
+			// Append the Node and Pin to the current Stack to uniquely identify each DataCollection
+			TArray<FPCGStackFrame>& StackFrames = Stack.GetStackFramesMutable();
+			StackFrames.Reserve(StackFrames.Num() + 2);
+			StackFrames.Emplace(InNode);
+			StackFrames.Emplace(Pin);
 
 			FPCGDataCollection PinDataCollection;
 			PinDataCollection.TaggedData = InData.GetInputsByPin(Pin->Properties.Label);
 			if (!PinDataCollection.TaggedData.IsEmpty())
 			{
-				InOutInspectionCache.Add(StackFramePath, PinDataCollection);
+				InOutInspectionCache.Add(Stack, PinDataCollection);
 			}
 			else
 			{
-				InOutInspectionCache.Remove(StackFramePath);
+				InOutInspectionCache.Remove(Stack);
 			}
 		}
 	};
@@ -1535,9 +1537,9 @@ void UPCGComponent::StoreInspectionData(const FPCGStack* InStack, const UPCGNode
 	StorePinInspectionData(InNode->GetOutputPins(), InOutputData, InspectionCache);
 }
 
-const FPCGDataCollection* UPCGComponent::GetInspectionData(const FString& InStackPath) const
+const FPCGDataCollection* UPCGComponent::GetInspectionData(const FPCGStack& InStack) const
 {
-	return InspectionCache.Find(InStackPath);
+	return InspectionCache.Find(InStack);
 }
 
 void UPCGComponent::Refresh(bool bStructural)
diff --git a/Engine/Plugins/Experimental/PCG/Source/PCG/Private/PCGSubgraph.cpp b/Engine/Plugins/Experimental/PCG/Source/PCG/Private/PCGSubgraph.cpp
--- a/Engine/Plugins/Experimental/PCG/Source/PCG/Private/PCGSubgraph.cpp
+++ b/Engine/Plugins/Experimental/PCG/Source/PCG/Private/PCGSubgraph.cpp
@@ -495,6 +495,10 @@ bool FPCGSubgraphElement::ExecuteInternal(FPCGContext* InContext) const
 				FPCGStack InvocationStack = ensure(Context->Stack) ? *Context->Stack : FPCGStack();
 				InvocationStack.GetStackFramesMutable().Emplace(Context->Node);
 
+#if WITH_EDITOR
+				Subgraph->OnGraphDynamicallyExecutedDelegate.Broadcast(Subgraph, Context->SourceComponent, InvocationStack);
+#endif
+
 				FPCGTaskId SubgraphTaskId = Subsystem->ScheduleGraph(Subgraph, Context->SourceComponent.Get(), MakeShared<FPCGInputForwardingElement>(PreSubgraphInputData), MakeShared<FPCGInputForwardingElement>(SubgraphInputData), {}, &InvocationStack);
 
 				if (SubgraphTaskId != InvalidPCGTaskId)
diff --git a/Engine/Plugins/Experimental/PCG/Source/PCG/Public/Graph/PCGStackContext.h b/Engine/Plugins/Experimental/PCG/Source/PCG/Public/Graph/PCGStackContext.h
--- a/Engine/Plugins/Experimental/PCG/Source/PCG/Public/Graph/PCGStackContext.h
+++ b/Engine/Plugins/Experimental/PCG/Source/PCG/Public/Graph/PCGStackContext.h
@@ -5,6 +5,8 @@
 #include "Containers/Array.h"
 #include "UObject/WeakObjectPtr.h"
 
+#include "PCGStackContext.generated.h"
+
 class FString;
 class UObject;
 class UPCGGraph;
@@ -12,27 +14,44 @@ class UPCGNode;
 class UPCGPin;
 
 /** A single frame of a call stack, represented as a pointer to the associated object (graph/subgraph or node) or a loop index. */
+USTRUCT()
 struct PCG_API FPCGStackFrame
 {
+	GENERATED_BODY()
+
+	FPCGStackFrame() {}
+
 	explicit FPCGStackFrame(TWeakObjectPtr<const UObject> InObject)
 		: Object(InObject)
-	{}
+	{
+		Hash = PointerHash(Object.Get());
+	}
 
 	explicit FPCGStackFrame(int32 InLoopIndex)
 		: LoopIndex(InLoopIndex)
-	{}
+	{
+		Hash = GetTypeHash(LoopIndex);
+	}
 
 	bool operator==(const FPCGStackFrame& Other) const { return Object == Other.Object && LoopIndex == Other.LoopIndex; }
 	bool operator!=(const FPCGStackFrame& Other) const { return !(*this == Other); };
 
+	friend uint32 GetTypeHash(const FPCGStackFrame& In) { return In.Hash; }
+
 	// A valid frame should either point to an object or have a loop index >= 0.
 	TWeakObjectPtr<const UObject> Object;
 	int32 LoopIndex = INDEX_NONE;
+
+private:
+	uint32 Hash = 0;
 };
 
 /** A call stack, represented as an array of stack frames. */
-class PCG_API FPCGStack
+USTRUCT()
+struct PCG_API FPCGStack
 {
+	GENERATED_BODY()
+
 	friend class FPCGStackContext;
 
 public:
@@ -51,6 +70,8 @@ public:
 	bool operator==(const FPCGStack& Other) const;
 	bool operator!=(const FPCGStack& Other) const { return !(*this == Other); }
 
+	friend uint32 GetTypeHash(const FPCGStack& In);
+
 private:
 	TArray<FPCGStackFrame> StackFrames;
 };
diff --git a/Engine/Plugins/Experimental/PCG/Source/PCG/Public/PCGComponent.h b/Engine/Plugins/Experimental/PCG/Source/PCG/Public/PCGComponent.h
--- a/Engine/Plugins/Experimental/PCG/Source/PCG/Public/PCGComponent.h
+++ b/Engine/Plugins/Experimental/PCG/Source/PCG/Public/PCGComponent.h
@@ -2,12 +2,13 @@
 
 #pragma once
 
-#include "ComponentInstanceDataCache.h"
-#include "Components/ActorComponent.h"
 #include "PCGNode.h"
 #include "PCGSettings.h"
+#include "Graph/PCGStackContext.h"
 #include "Utils/PCGExtraCapture.h"
 
+#include "ComponentInstanceDataCache.h"
+#include "Components/ActorComponent.h"
 #include "Containers/Map.h"
 
 #include "PCGComponent.generated.h"
@@ -16,7 +17,6 @@ namespace EEndPlayReason { enum Type : int; }
 
 class APCGPartitionActor;
 struct FPCGContext;
-class FPCGStack;
 class UPCGActorAndComponentMapping;
 class UPCGComponent;
 class UPCGGraph;
@@ -256,7 +256,7 @@ public:
 	void EnableInspection();
 	void DisableInspection();
 	void StoreInspectionData(const FPCGStack* InStack, const UPCGNode* InNode, const FPCGDataCollection& InInputData, const FPCGDataCollection& InOutputData);
-	const FPCGDataCollection* GetInspectionData(const FString& InStackPath) const;
+	const FPCGDataCollection* GetInspectionData(const FPCGStack& InStack) const;
 
 	bool IsActorTracked(AActor* InActor, bool& bOutIsCulled) const;
 
@@ -428,7 +428,7 @@ private:
 	TSet<TWeakObjectPtr<AActor>> CachedTrackedActors_DEPRECATED;
 
 	UPROPERTY(Transient)
-	TMap<FString, FPCGDataCollection> InspectionCache;
+	TMap<FPCGStack, FPCGDataCollection> InspectionCache;
 #endif
 
 	mutable FCriticalSection GeneratedResourcesLock;
diff --git a/Engine/Plugins/Experimental/PCG/Source/PCG/Public/PCGContext.h b/Engine/Plugins/Experimental/PCG/Source/PCG/Public/PCGContext.h
--- a/Engine/Plugins/Experimental/PCG/Source/PCG/Public/PCGContext.h
+++ b/Engine/Plugins/Experimental/PCG/Source/PCG/Public/PCGContext.h
@@ -14,7 +14,7 @@ class UPCGSettingsInterface;
 class UPCGSpatialData;
 struct FPCGGraphCache;
 struct FPCGSettingsOverridableParam;
-class FPCGStack;
+struct FPCGStack;
 
 namespace PCGContextHelpers
 {
diff --git a/Engine/Plugins/Experimental/PCG/Source/PCG/Public/PCGGraph.h b/Engine/Plugins/Experimental/PCG/Source/PCG/Public/PCGGraph.h
--- a/Engine/Plugins/Experimental/PCG/Source/PCG/Public/PCGGraph.h
+++ b/Engine/Plugins/Experimental/PCG/Source/PCG/Public/PCGGraph.h
@@ -5,6 +5,7 @@
 #include "PCGCommon.h"
 #include "PCGNode.h"
 #include "PCGSettings.h"
+#include "Graph/PCGStackContext.h"
 
 #include "PropertyBag.h"
 #include "UObject/ObjectPtr.h"
@@ -30,6 +31,7 @@ enum class EPCGGraphParameterEvent
 DECLARE_MULTICAST_DELEGATE_TwoParams(FOnPCGGraphChanged, UPCGGraphInterface* /*Graph*/, EPCGChangeType /*ChangeType*/);
 DECLARE_MULTICAST_DELEGATE_OneParam(FOnPCGGraphGridSizesChanged, UPCGGraphInterface* /*Graph*/);
 DECLARE_MULTICAST_DELEGATE_ThreeParams(FOnPCGGraphParametersChanged, UPCGGraphInterface* /*Graph*/, EPCGGraphParameterEvent /*ChangeType*/, FName /*ChangedPropertyName*/);
+DECLARE_MULTICAST_DELEGATE_ThreeParams(FOnPCGGraphDynamicallyExecuted, UPCGGraphInterface* /*Graph*/, const TWeakObjectPtr<UPCGComponent> /*SourceComponent*/, FPCGStack /*InvocationStack*/);
 #endif // WITH_EDITOR
 
 /**
@@ -102,6 +104,7 @@ public:
 	FOnPCGGraphChanged OnGraphChangedDelegate;
 	FOnPCGGraphGridSizesChanged OnGraphGridSizesChangedDelegate;
 	FOnPCGGraphParametersChanged OnGraphParametersChangedDelegate;
+	FOnPCGGraphDynamicallyExecuted OnGraphDynamicallyExecutedDelegate;
 #endif // WITH_EDITOR
 };
 
diff --git a/Engine/Plugins/Experimental/PCG/Source/PCG/Public/PCGSubsystem.h b/Engine/Plugins/Experimental/PCG/Source/PCG/Public/PCGSubsystem.h
--- a/Engine/Plugins/Experimental/PCG/Source/PCG/Public/PCGSubsystem.h
+++ b/Engine/Plugins/Experimental/PCG/Source/PCG/Public/PCGSubsystem.h
@@ -24,7 +24,7 @@ class FPCGGraphCompiler;
 class FPCGGraphExecutor;
 struct FPCGContext;
 struct FPCGDataCollection;
-class FPCGStack;
+struct FPCGStack;
 class UPCGSettings;
 
 class IPCGElement;
diff --git a/Engine/Plugins/Experimental/PCG/Source/PCGEditor/Private/PCGEditor.cpp b/Engine/Plugins/Experimental/PCG/Source/PCGEditor/Private/PCGEditor.cpp
--- a/Engine/Plugins/Experimental/PCG/Source/PCGEditor/Private/PCGEditor.cpp
+++ b/Engine/Plugins/Experimental/PCG/Source/PCGEditor/Private/PCGEditor.cpp
@@ -86,6 +86,7 @@ void FPCGEditor::Initialize(const EToolkitMode::Type InMode, const TSharedPtr<cl
 	if (PCGGraphBeingEdited)
 	{
 		PCGGraphBeingEdited->OnGraphGridSizesChangedDelegate.AddRaw(this, &FPCGEditor::OnGraphGridSizesChanged);
+		PCGGraphBeingEdited->OnGraphDynamicallyExecutedDelegate.AddRaw(this, &FPCGEditor::OnGraphDynamicallyExecuted);
 	}
 
 	PCGEditorGraph = NewObject<UPCGEditorGraph>(PCGGraphBeingEdited, UPCGEditorGraph::StaticClass(), NAME_None, RF_Transactional | RF_Transient);
@@ -2221,6 +2222,7 @@ void FPCGEditor::OnClose()
 		}
 
 		PCGGraphBeingEdited->OnGraphGridSizesChangedDelegate.RemoveAll(this);
+		PCGGraphBeingEdited->OnGraphDynamicallyExecutedDelegate.RemoveAll(this);
 	}
 }
 
@@ -2441,6 +2443,21 @@ void FPCGEditor::OnGraphGridSizesChanged(UPCGGraphInterface* InGraph)
 	}
 }
 
+void FPCGEditor::OnGraphDynamicallyExecuted(UPCGGraphInterface* InGraphInterface, const TWeakObjectPtr<UPCGComponent> InSourceComponent, FPCGStack InvocationStack)
+{
+	InvocationStack.GetStackFramesMutable().Emplace(InGraphInterface);
+
+	if (DebugObjectWidget.IsValid())
+	{
+		DebugObjectWidget->AddDynamicStack(InSourceComponent, InvocationStack);
+	}
+
+	if (DebugObjectTreeWidget.IsValid())
+	{
+		DebugObjectTreeWidget->AddDynamicStack(InSourceComponent, InvocationStack);
+	}
+}
+
 UPCGSubsystem* FPCGEditor::GetSubsystem()
 {
 	UWorld* World = (GEditor ? (GEditor->PlayWorld ? GEditor->PlayWorld.Get() : GEditor->GetEditorWorldContext().World()) : nullptr);
diff --git a/Engine/Plugins/Experimental/PCG/Source/PCGEditor/Private/PCGEditor.h b/Engine/Plugins/Experimental/PCG/Source/PCGEditor/Private/PCGEditor.h
--- a/Engine/Plugins/Experimental/PCG/Source/PCGEditor/Private/PCGEditor.h
+++ b/Engine/Plugins/Experimental/PCG/Source/PCGEditor/Private/PCGEditor.h
@@ -278,6 +278,7 @@ private:
 	bool IsVisibleProperty(const FPropertyAndParent& InPropertyAndParent, IDetailsView* InDetailsView) const;
 
 	void OnGraphGridSizesChanged(UPCGGraphInterface* InGraph);
+	void OnGraphDynamicallyExecuted(UPCGGraphInterface* InGraphInterface, const TWeakObjectPtr<UPCGComponent> InSourceComponent, FPCGStack InvocationStack);
 
 	/** Trigger any generation required to ensure debug display is up to date. */
 	void UpdateDebugAfterComponentSelection(UPCGComponent* InOldComponent, UPCGComponent* InNewComponent, bool bNewComponentStartedInspecting);
diff --git a/Engine/Plugins/Experimental/PCG/Source/PCGEditor/Private/SPCGEditorGraphAttributeListView.cpp b/Engine/Plugins/Experimental/PCG/Source/PCGEditor/Private/SPCGEditorGraphAttributeListView.cpp
--- a/Engine/Plugins/Experimental/PCG/Source/PCGEditor/Private/SPCGEditorGraphAttributeListView.cpp
+++ b/Engine/Plugins/Experimental/PCG/Source/PCGEditor/Private/SPCGEditorGraphAttributeListView.cpp
@@ -544,13 +544,14 @@ const FPCGDataCollection* SPCGEditorGraphAttributeListView::GetInspectionData()
 		return nullptr;
 	}
 
-	FString ResourcePath;
-	if (!PCGStack.CreateStackFramePath(ResourcePath, PCGNode, Pin))
-	{
-		return nullptr;
-	}
-
-	return PCGComponent->GetInspectionData(ResourcePath);
+	// Create a temporary stack with Node+Pin to query the exact DataCollection we are inspecting
+	FPCGStack Stack = PCGStack;
+	TArray<FPCGStackFrame>& StackFrames = Stack.GetStackFramesMutable();
+	StackFrames.Reserve(StackFrames.Num() + 2);
+	StackFrames.Emplace(PCGNode);
+	StackFrames.Emplace(Pin);
+
+	return PCGComponent->GetInspectionData(Stack);
 }
 
 void SPCGEditorGraphAttributeListView::RefreshAttributeList()
diff --git a/Engine/Plugins/Experimental/PCG/Source/PCGEditor/Private/SPCGEditorGraphDebugObjectTree.cpp b/Engine/Plugins/Experimental/PCG/Source/PCGEditor/Private/SPCGEditorGraphDebugObjectTree.cpp
--- a/Engine/Plugins/Experimental/PCG/Source/PCGEditor/Private/SPCGEditorGraphDebugObjectTree.cpp
+++ b/Engine/Plugins/Experimental/PCG/Source/PCGEditor/Private/SPCGEditorGraphDebugObjectTree.cpp
@@ -7,6 +7,7 @@
 #include "PCGEditorGraph.h"
 #include "PCGEditorModule.h"
 #include "PCGGraph.h"
+#include "Elements/PCGLoopElement.h"
 
 #include "PropertyCustomizationHelpers.h"
 #include "Selection.h"
@@ -51,6 +52,7 @@ FString FPCGEditorGraphDebugObjectItem_PCGGraph::GetLabel() const
 {
 	return PCGGraph.IsValid() ? PCGGraph->GetName() : FString();
 }
+
 FString FPCGEditorGraphDebugObjectItem_PCGSubgraph::GetLabel() const
 {
 	if (PCGNode.IsValid() && PCGGraph.IsValid())
@@ -61,6 +63,16 @@ FString FPCGEditorGraphDebugObjectItem_PCGSubgraph::GetLabel() const
 	return FString();
 }
 
+FString FPCGEditorGraphDebugObjectItem_PCGLoopIndex::GetLabel() const
+{
+	if (PCGNode.IsValid())
+	{
+		return PCGNode->GetName() + FString::Format(TEXT(" - {0}"), { LoopIndex });
+	}
+
+	return FString();
+}
+
 void SPCGEditorGraphDebugObjectItemRow::Construct(const FArguments& InArgs, const TSharedRef<STableViewBase>& InOwnerTableView, FPCGEditorGraphDebugObjectItemPtr InItem)
 {
 	Item = InItem;
@@ -208,6 +220,17 @@ void SPCGEditorGraphDebugObjectTree::Tick(const FGeometry& AllottedGeometry, con
 	}
 }
 
+void SPCGEditorGraphDebugObjectTree::AddDynamicStack(const TWeakObjectPtr<UPCGComponent> InComponent, const FPCGStack& InvocationStack)
+{
+	TArray<FPCGStack>& Stacks = DynamicInvocationStacks.FindOrAdd(InComponent);
+
+	if (!Stacks.Contains(InvocationStack))
+	{
+		Stacks.Add(InvocationStack);
+		RequestRefresh();
+	}
+}
+
 void SPCGEditorGraphDebugObjectTree::SelectedDebugObject_OnClicked() const
 {
 	if (UPCGComponent* PCGComponent = PCGEditor.Pin()->GetPCGComponentBeingInspected())
@@ -381,17 +404,19 @@ void SPCGEditorGraphDebugObjectTree::RefreshTree()
 		{
 			continue;
 		}
-		
+
 		FPCGStackContext StackContext = FPCGStackContext::CreateStackContextFromGraph(PCGComponentGraph);
 
 		TMap<const UPCGGraph*, FPCGEditorGraphDebugObjectItemPtr> GraphItems;
 		TArray<FPCGEditorGraphDebugObjectItemPtr> NodeItems;
 
+		// Process statically generated stacks
 		for (const FPCGStack& Stack : StackContext.GetStacks())
 		{
 			const UPCGGraph* TopStackGraph = Cast<const UPCGGraph>(Stack.GetStackFrames().Top().Object);
 			if (TopStackGraph == PCGGraph)
 			{
+				// Find or Create ActorItem
 				TSharedPtr<FPCGEditorGraphDebugObjectItem_Actor> ActorItem;
 				if (TSharedPtr<FPCGEditorGraphDebugObjectItem_Actor>* FoundActorItem = ActorItems.Find(Actor))
 				{
@@ -401,7 +426,8 @@ void SPCGEditorGraphDebugObjectTree::RefreshTree()
 				{
 					ActorItem = ActorItems.Emplace(Actor, MakeShared<FPCGEditorGraphDebugObjectItem_Actor>(Actor));
 				}
-				
+
+				// Find or Create ComponentItem
 				TSharedPtr<FPCGEditorGraphDebugObjectItem_PCGComponent> ComponentItem;
 				if (TSharedPtr<FPCGEditorGraphDebugObjectItem_PCGComponent>* FoundComponentItem = ComponentItems.Find(PCGComponent))
 				{
@@ -411,20 +437,20 @@ void SPCGEditorGraphDebugObjectTree::RefreshTree()
 				{
 					ComponentItem = ComponentItems.Emplace(PCGComponent, MakeShared<FPCGEditorGraphDebugObjectItem_PCGComponent>(PCGComponent));
 				}
-				
+
 				if (!ComponentItem->GetParent())
 				{
-					ActorItem->AddChild(ComponentItem.ToSharedRef());	
+					ActorItem->AddChild(ComponentItem.ToSharedRef());
 				}
 
 				const TArray<FPCGStackFrame>& StackFrames = Stack.GetStackFrames();
-				
+
 				FString StackPath;
 				Stack.CreateStackFramePath(StackPath);
-					
+
 				FPCGEditorGraphDebugObjectItemPtr PreviousItem = ComponentItem;
 				for (int32 StackIndex = 0; StackIndex < StackFrames.Num(); StackIndex++)
-				{	
+				{
 					const FPCGStackFrame& StackFrame = StackFrames[StackIndex];
 					FPCGEditorGraphDebugObjectItemPtr CurrentItem;
 					if (const UPCGGraph* StackGraph = Cast<const UPCGGraph>(StackFrame.Object))
@@ -440,6 +466,15 @@ void SPCGEditorGraphDebugObjectTree::RefreshTree()
 					}
 					else if (const UPCGNode* StackNode = Cast<const UPCGNode>(StackFrame.Object))
 					{
+						// We cannot handle dynamically executed graphs in this path. They must be provided through DynamicInvocationStacks
+						if (const UPCGBaseSubgraphSettings* SubgraphSettings = Cast<const UPCGBaseSubgraphSettings>(StackNode->GetSettings()))
+						{
+							if (SubgraphSettings->IsDynamicGraph())
+							{
+								break;
+							}
+						}
+
 						const int32 NextStackIndex = StackIndex + 1;
 						if (StackFrames.IsValidIndex(NextStackIndex))
 						{
@@ -464,10 +499,113 @@ void SPCGEditorGraphDebugObjectTree::RefreshTree()
 							}
 						}
 					}
-					
+
+					if (!CurrentItem)
+					{
+						break;
+					}
+
+					if (!CurrentItem->GetParent())
+					{
+						PreviousItem->AddChild(CurrentItem.ToSharedRef());
+					}
+
+					AllGraphItems.Add(CurrentItem);
+					PreviousItem = CurrentItem;
+				}
+			}
+		}
+
+		if (const TArray<FPCGStack>* DynamicStacks = DynamicInvocationStacks.Find(PCGComponent))
+		{
+			for (const FPCGStack& Stack : *DynamicStacks)
+			{
+				TSharedPtr<FPCGEditorGraphDebugObjectItem_Actor> ActorItem;
+				if (TSharedPtr<FPCGEditorGraphDebugObjectItem_Actor>* FoundActorItem = ActorItems.Find(Actor))
+				{
+					ActorItem = *FoundActorItem;
+				}
+				else
+				{
+					ActorItem = ActorItems.Emplace(Actor, MakeShared<FPCGEditorGraphDebugObjectItem_Actor>(Actor));
+				}
+
+				TSharedPtr<FPCGEditorGraphDebugObjectItem_PCGComponent> ComponentItem;
+				if (TSharedPtr<FPCGEditorGraphDebugObjectItem_PCGComponent>* FoundComponentItem = ComponentItems.Find(PCGComponent))
+				{
+					ComponentItem = *FoundComponentItem;
+				}
+				else
+				{
+					ComponentItem = ComponentItems.Emplace(PCGComponent, MakeShared<FPCGEditorGraphDebugObjectItem_PCGComponent>(PCGComponent));
+				}
+
+				if (!ComponentItem->GetParent())
+				{
+					ActorItem->AddChild(ComponentItem.ToSharedRef());
+				}
+
+				const TArray<FPCGStackFrame>& StackFrames = Stack.GetStackFrames();
+
+				FString StackPath;
+				Stack.CreateStackFramePath(StackPath);
+
+				FPCGEditorGraphDebugObjectItemPtr PreviousItem = ComponentItem;
+				for (int32 StackIndex = 0; StackIndex < StackFrames.Num(); StackIndex++)
+				{
+					const FPCGStackFrame& StackFrame = StackFrames[StackIndex];
+					FPCGEditorGraphDebugObjectItemPtr CurrentItem;
+					if (const UPCGGraph* StackGraph = Cast<const UPCGGraph>(StackFrame.Object))
+					{
+						if (FPCGEditorGraphDebugObjectItemPtr* GraphItem = GraphItems.Find(StackGraph))
+						{
+							CurrentItem = *GraphItem;
+						}
+						else
+						{
+							CurrentItem = GraphItems.Emplace(StackGraph, MakeShared<FPCGEditorGraphDebugObjectItem_PCGGraph>(StackGraph, StackGraph == PCGGraph ? Stack : FPCGStack()));
+						}
+					}
+					else if (const UPCGNode* StackNode = Cast<const UPCGNode>(StackFrame.Object))
+					{
+						const int32 NextStackIndex = StackIndex + 1;
+						if (StackFrames.IsValidIndex(NextStackIndex))
+						{
+							const FPCGEditorGraphDebugObjectItemPtr* NodeItem = NodeItems.FindByPredicate([&StackNode, &Stack](const FPCGEditorGraphDebugObjectItemPtr& InItem)
+								{
+									return (InItem->GetObject() == StackNode) && (*InItem->GetPCGStack() == Stack);
+								});
+
+							if (NodeItem)
+							{
+								CurrentItem = *NodeItem;
+							}
+							else
+							{
+								const FPCGStackFrame& NextStackFrame = StackFrames[NextStackIndex];
+
+								if (const UPCGGraph* NextStackGraph = Cast<const UPCGGraph>(NextStackFrame.Object))
+								{
+									CurrentItem = NodeItems.Add_GetRef(MakeShared<FPCGEditorGraphDebugObjectItem_PCGSubgraph>(StackNode, NextStackGraph, Stack));
+								}
+								else if (NextStackFrame.LoopIndex != INDEX_NONE)
+								{
+									CurrentItem = NodeItems.Add_GetRef(MakeShared<FPCGEditorGraphDebugObjectItem_PCGLoopIndex>(StackNode, NextStackFrame.LoopIndex, Stack));
+								}
+							}
+
+							StackIndex++;
+						}
+					}
+
+					if (!CurrentItem)
+					{
+						break;
+					}
+
 					if (!CurrentItem->GetParent())
 					{
-						PreviousItem->AddChild(CurrentItem.ToSharedRef());	
+						PreviousItem->AddChild(CurrentItem.ToSharedRef());
 					}
 
 					AllGraphItems.Add(CurrentItem);
diff --git a/Engine/Plugins/Experimental/PCG/Source/PCGEditor/Private/SPCGEditorGraphDebugObjectTree.h b/Engine/Plugins/Experimental/PCG/Source/PCGEditor/Private/SPCGEditorGraphDebugObjectTree.h
--- a/Engine/Plugins/Experimental/PCG/Source/PCGEditor/Private/SPCGEditorGraphDebugObjectTree.h
+++ b/Engine/Plugins/Experimental/PCG/Source/PCGEditor/Private/SPCGEditorGraphDebugObjectTree.h
@@ -105,6 +105,27 @@ protected:
 	FPCGStack PCGStack;
 };
 
+class FPCGEditorGraphDebugObjectItem_PCGLoopIndex : public FPCGEditorGraphDebugObjectItem
+{
+public:
+	FPCGEditorGraphDebugObjectItem_PCGLoopIndex(TWeakObjectPtr<const UPCGNode> InPCGNode, int32 InLoopIndex, const FPCGStack& InPCGStack)
+		: PCGNode(InPCGNode)
+		, LoopIndex(InLoopIndex)
+		, PCGStack(InPCGStack)
+	{}
+
+	virtual FString GetLabel() const override;
+	virtual UPCGComponent* GetPCGComponent() const override { return GetParent() ? GetParent()->GetPCGComponent() : nullptr; }
+	virtual const UObject* GetObject() const override { return PCGNode.Get(); };
+	virtual const FPCGStack* GetPCGStack() const override { return &PCGStack; }
+	virtual int32 GetLoopIndex() const { return LoopIndex; }
+	
+protected:
+	TWeakObjectPtr<const UPCGNode> PCGNode = nullptr;
+	int32 LoopIndex = INDEX_NONE;
+	FPCGStack PCGStack;
+};
+
 class SPCGEditorGraphDebugObjectItemRow : public SCompoundWidget
 {
 public:
@@ -130,6 +151,8 @@ public:
 	virtual void Tick(const FGeometry& AllottedGeometry, const double InCurrentTime, const float InDeltaTime) override;
 	
 	void RequestRefresh() { bNeedsRefresh = true; }
+
+	void AddDynamicStack(TWeakObjectPtr<UPCGComponent> InComponent, const FPCGStack& InvocationStack);
 	
 private:
 	void SelectedDebugObject_OnClicked() const;
@@ -157,5 +180,7 @@ private:
 	TArray<FPCGEditorGraphDebugObjectItemPtr> RootItems;
 	TArray<FPCGEditorGraphDebugObjectItemPtr> AllGraphItems;
 
+	TMap<const TWeakObjectPtr<UPCGComponent>, TArray<FPCGStack>> DynamicInvocationStacks;
+
 	bool bNeedsRefresh = false;
 };
diff --git a/Engine/Plugins/Experimental/PCG/Source/PCGEditor/Private/SPCGEditorGraphDebugObjectWidget.cpp b/Engine/Plugins/Experimental/PCG/Source/PCGEditor/Private/SPCGEditorGraphDebugObjectWidget.cpp
--- a/Engine/Plugins/Experimental/PCG/Source/PCGEditor/Private/SPCGEditorGraphDebugObjectWidget.cpp
+++ b/Engine/Plugins/Experimental/PCG/Source/PCGEditor/Private/SPCGEditorGraphDebugObjectWidget.cpp
@@ -28,7 +28,14 @@ FPCGEditorGraphDebugObjectInstance::FPCGEditorGraphDebugObjectInstance(TWeakObje
 
 	for (const FPCGStackFrame& StackFrame : PCGStack.GetStackFrames())
 	{
-		PathLabel /= StackFrame.Object->GetName();
+		if (StackFrame.Object.IsValid())
+		{
+			PathLabel /= StackFrame.Object->GetName();
+		}
+		else
+		{
+			PathLabel /= FString::Format(TEXT("{0}"), { StackFrame.LoopIndex });
+		}
 	}
 
 	Label = FText::FromString(PathLabel);
@@ -163,6 +170,20 @@ void SPCGEditorGraphDebugObjectWidget::RefreshDebugObjects()
 				}
 			}
 		}
+
+		if (const TArray<FPCGStack>* DynamicStacks = DynamicInvocationStacks.Find(PCGComponent))
+		{
+			for (const FPCGStack& Stack : *DynamicStacks)
+			{
+				const TSharedPtr<FPCGEditorGraphDebugObjectInstance> DebugInstance = MakeShared<FPCGEditorGraphDebugObjectInstance>(PCGComponent, Stack);
+				DebugObjects.Add(DebugInstance);
+
+				if (SelectedItem.IsValid() && SelectedItem->GetPCGComponent() == PCGComponent && SelectedItem->GetStack() == Stack)
+				{
+					DebugObjectsComboBox->SetSelectedItem(DebugInstance);
+				}
+			}
+		}
 	}
 }
 
@@ -384,4 +405,15 @@ void SPCGEditorGraphDebugObjectWidget::OnLevelActorDeleted(const AActor* InActor
 	}
 }
 
+void SPCGEditorGraphDebugObjectWidget::AddDynamicStack(const TWeakObjectPtr<UPCGComponent> InComponent, const FPCGStack& InvocationStack)
+{
+	TArray<FPCGStack>& Stacks = DynamicInvocationStacks.FindOrAdd(InComponent);
+
+	if (!Stacks.Contains(InvocationStack))
+	{
+		Stacks.Add(InvocationStack);
+		RefreshDebugObjects();
+	}
+}
+
 #undef LOCTEXT_NAMESPACE
diff --git a/Engine/Plugins/Experimental/PCG/Source/PCGEditor/Private/SPCGEditorGraphDebugObjectWidget.h b/Engine/Plugins/Experimental/PCG/Source/PCGEditor/Private/SPCGEditorGraphDebugObjectWidget.h
--- a/Engine/Plugins/Experimental/PCG/Source/PCGEditor/Private/SPCGEditorGraphDebugObjectWidget.h
+++ b/Engine/Plugins/Experimental/PCG/Source/PCGEditor/Private/SPCGEditorGraphDebugObjectWidget.h
@@ -46,6 +46,8 @@ public:
 
 	void OnLevelActorDeleted(const AActor* InActor);
 
+	void AddDynamicStack(const TWeakObjectPtr<UPCGComponent> InComponent, const FPCGStack& InvocationStack);
+
 private:
 	void OnComboBoxOpening();
 	void OnSelectionChanged(TSharedPtr<FPCGEditorGraphDebugObjectInstance> NewSelection, ESelectInfo::Type SelectInfo) const;
@@ -65,4 +67,6 @@ private:
 
 	TArray<TSharedPtr<FPCGEditorGraphDebugObjectInstance>> DebugObjects;
 	TSharedPtr<SComboBox<TSharedPtr<FPCGEditorGraphDebugObjectInstance>>> DebugObjectsComboBox;
+
+	TMap<const TWeakObjectPtr<UPCGComponent>, TArray<FPCGStack>> DynamicInvocationStacks;
 };
