Subgraph nodes that were statically compiled were running the same element twice, which caused an issue in the subgraph element particularly as, since the addition of dynamic dispatch we would do some processing on the context data which wasn't using then the output pins but only the input ones.

#rnx

diff --git a/Engine/Plugins/Experimental/PCG/Source/PCG/Private/Graph/PCGGraphCompiler.cpp b/Engine/Plugins/Experimental/PCG/Source/PCG/Private/Graph/PCGGraphCompiler.cpp
--- a/Engine/Plugins/Experimental/PCG/Source/PCG/Private/Graph/PCGGraphCompiler.cpp
+++ b/Engine/Plugins/Experimental/PCG/Source/PCG/Private/Graph/PCGGraphCompiler.cpp
@@ -153,6 +153,9 @@ TArray<FPCGGraphTask> FPCGGraphCompiler::CompileGraph(UPCGGraph* InGraph, FPCGTa
 			PostTask.Node = Node;
 			PostTask.NodeId = PostId;
 			PostTask.StackIndex = InOutStackContext.GetCurrentStackIndex();
+			// Implementation note: since we`ve already executed the node once, we normally don`t need to execute it a second time
+			// especially since we cannot distinguish between the pre and post during execution so any data filtering related to pins is bound to fail.
+			PostTask.Element = MakeShared<FPCGTrivialElement>();
 
 			// Add subgraph output node task as input to the post-task
 			if (OutputNodeTask)
