[PCG] Loop node does not respect SubgraphOverride pin

#rnx

diff --git a/Engine/Plugins/Experimental/PCG/Source/PCG/Private/Elements/PCGLoopElement.cpp b/Engine/Plugins/Experimental/PCG/Source/PCG/Private/Elements/PCGLoopElement.cpp
--- a/Engine/Plugins/Experimental/PCG/Source/PCG/Private/Elements/PCGLoopElement.cpp
+++ b/Engine/Plugins/Experimental/PCG/Source/PCG/Private/Elements/PCGLoopElement.cpp
@@ -54,6 +54,13 @@ bool FPCGLoopElement::ExecuteInternal(FPCGContext* InContext) const
 	// If we haven't scheduled the subgraph yet (1st call) - this is part of the FPCGSubgraphContext, see FPCGSubgraphElement for more info
 	if (!Context->bScheduledSubgraph)
 	{
+		if (Settings->SubgraphInstance && Settings->SubgraphOverride)
+		{
+			// If OriginalSettings is null, then we ARE the original settings, and writing over the existing graph is incorrect (and potentially a race condition)
+			check(Settings->OriginalSettings);
+			Settings->SubgraphInstance->SetGraph(Settings->SubgraphOverride);
+		}
+
 		UPCGGraph* Subgraph = Settings->GetSubgraph();
 		UPCGSubsystem* Subsystem = Context->SourceComponent.IsValid() ? Context->SourceComponent->GetSubsystem() : nullptr;
 
