[PCG] Fix broken Crc when the generated data is read into a downstream graph which would prevent the downstream from updating properly

#rnx

diff --git a/Engine/Plugins/Experimental/PCG/Source/PCG/Private/PCGComponent.cpp b/Engine/Plugins/Experimental/PCG/Source/PCG/Private/PCGComponent.cpp
--- a/Engine/Plugins/Experimental/PCG/Source/PCG/Private/PCGComponent.cpp
+++ b/Engine/Plugins/Experimental/PCG/Source/PCG/Private/PCGComponent.cpp
@@ -1522,6 +1522,10 @@ void UPCGComponent::StoreInspectionData(const FPCGStack* InStack, const UPCGNode
 
 			FPCGDataCollection PinDataCollection;
 			PinDataCollection.TaggedData = InData.GetInputsByPin(Pin->Properties.Label);
+			// The data collection for each pin is given the Crc from the data collection. This is to enable inspecting the normal node output Crc
+			// when cache debugging is enabled.
+			PinDataCollection.Crc = InData.Crc;
+
 			if (!PinDataCollection.TaggedData.IsEmpty())
 			{
 				InOutInspectionCache.Add(Stack, PinDataCollection);
diff --git a/Engine/Plugins/Experimental/PCG/Source/PCG/Private/PCGData.cpp b/Engine/Plugins/Experimental/PCG/Source/PCG/Private/PCGData.cpp
--- a/Engine/Plugins/Experimental/PCG/Source/PCG/Private/PCGData.cpp
+++ b/Engine/Plugins/Experimental/PCG/Source/PCG/Private/PCGData.cpp
@@ -85,7 +85,7 @@ UPCGData::UPCGData(const FObjectInitializer& ObjectInitializer)
 {
 	if (!HasAnyFlags(RF_ClassDefaultObject))
 	{
-		UID = ++UIDCounter;
+		InitUID();
 	}
 }
 
@@ -126,6 +126,11 @@ void UPCGData::VisitDataNetwork(TFunctionRef<void(const UPCGData*)> Action) cons
 	Action(this);
 }
 
+void UPCGData::InitUID()
+{
+	UID = ++UIDCounter;
+}
+
 bool FPCGTaggedData::operator==(const FPCGTaggedData& Other) const
 {
 	return Data == Other.Data &&
diff --git a/Engine/Plugins/Experimental/PCG/Source/PCG/Public/PCGData.h b/Engine/Plugins/Experimental/PCG/Source/PCG/Public/PCGData.h
--- a/Engine/Plugins/Experimental/PCG/Source/PCG/Public/PCGData.h
+++ b/Engine/Plugins/Experimental/PCG/Source/PCG/Public/PCGData.h
@@ -50,6 +50,10 @@ public:
 	virtual FPCGAttributePropertyInputSelector GetCachedLastSelector() const { return FPCGAttributePropertyInputSelector{}; }
 	virtual void SetLastSelector(const FPCGAttributePropertySelector& InSelector) {};
 
+	// ~Begin UObject interface
+	virtual void PostDuplicate(bool bDuplicateForPIE) override { InitUID(); }
+	// ~End UObject interface
+
 protected:
 	/** Computes Crc for this and any connected data. */
 	virtual FPCGCrc ComputeCrc(bool bFullDataCrc) const;
@@ -61,6 +65,8 @@ protected:
 	bool PropagateCrcThroughBooleanData() const;
 
 private:
+	void InitUID();
+
 	/** Serves unique ID values to instances of this object. */
 	static inline std::atomic<uint64> UIDCounter{ 1 };
 };
