[PCG] Fixed a bug causing the PCG Attribute List View to crash when sorting a FString-based attribute...

#rnx

diff --git a/Engine/Plugins/Experimental/PCG/Source/PCG/Public/Metadata/PCGMetadataAttributeTraits.h b/Engine/Plugins/Experimental/PCG/Source/PCG/Public/Metadata/PCGMetadataAttributeTraits.h
--- a/Engine/Plugins/Experimental/PCG/Source/PCG/Public/Metadata/PCGMetadataAttributeTraits.h
+++ b/Engine/Plugins/Experimental/PCG/Source/PCG/Public/Metadata/PCGMetadataAttributeTraits.h
@@ -325,6 +325,7 @@ namespace PCG
 		{
 			enum { CompressData = false };
 			enum { CanSearchString = false };
+			enum { NeedsConstruction = false };
 		};
 
 		template<>
@@ -332,6 +333,7 @@ namespace PCG
 		{
 			enum { CompressData = false };
 			enum { CanSearchString = false };
+			enum { NeedsConstruction = false };
 
 			static bool Equal(const float& A, const float& B)
 			{
@@ -344,6 +346,7 @@ namespace PCG
 		{
 			enum { CompressData = false };
 			enum { CanSearchString = false };
+			enum { NeedsConstruction = false };
 
 			static bool Equal(const double& A, const double& B)
 			{
@@ -361,6 +364,7 @@ namespace PCG
 			enum { CanInterpolate = false };
 			enum { CanCompare = true };
 			enum { CanSearchString = false };
+			enum { NeedsConstruction = false };
 
 			static bool ZeroValue()
 			{
@@ -422,6 +426,7 @@ namespace PCG
 			enum { CanInterpolate = true };
 			enum { CanCompare = true };
 			enum { CanSearchString = false };
+			enum { NeedsConstruction = false };
 
 			static T ZeroValue()
 			{
@@ -557,6 +562,7 @@ namespace PCG
 			enum { CanInterpolate = true };
 			enum { CanCompare = false };
 			enum { CanSearchString = false };
+			enum { NeedsConstruction = false };
 
 			static bool Equal(const FQuat& A, const FQuat& B)
 			{
@@ -611,6 +617,7 @@ namespace PCG
 			enum { CanInterpolate = true };
 			enum { CanCompare = false };
 			enum { CanSearchString = false };
+			enum { NeedsConstruction = false };
 
 			static bool Equal(const FRotator& A, const FRotator& B)
 			{
@@ -660,6 +667,7 @@ namespace PCG
 			enum { CanInterpolate = true };
 			enum { CanCompare = false };
 			enum { CanSearchString = false };
+			enum { NeedsConstruction = false };
 
 			static bool Equal(const FTransform& A, const FTransform& B)
 			{
@@ -711,6 +719,7 @@ namespace PCG
 			enum { CanMulDiv = false };
 			enum { CanInterpolate = false };
 			enum { CanSearchString = true };
+			enum { NeedsConstruction = true };
 
 			static bool Equal(const FString& A, const FString& B)
 			{
@@ -743,6 +752,7 @@ namespace PCG
 			enum { CanInterpolate = false };
 			enum { CanCompare = true };
 			enum { CanSearchString = true };
+			enum { NeedsConstruction = false };
 
 			static bool Equal(const FName& A, const FName& B)
 			{
diff --git a/Engine/Plugins/Experimental/PCG/Source/PCGEditor/Private/SPCGEditorGraphAttributeListView.cpp b/Engine/Plugins/Experimental/PCG/Source/PCGEditor/Private/SPCGEditorGraphAttributeListView.cpp
--- a/Engine/Plugins/Experimental/PCG/Source/PCGEditor/Private/SPCGEditorGraphAttributeListView.cpp
+++ b/Engine/Plugins/Experimental/PCG/Source/PCGEditor/Private/SPCGEditorGraphAttributeListView.cpp
@@ -1101,7 +1101,16 @@ void SPCGEditorGraphAttributeListView::RefreshSorting()
 				if constexpr (PCG::Private::MetadataTraits<ValueType>::CanCompare)
 				{
 					TArray<ValueType> CachedValues;
-					CachedValues.SetNumUninitialized(ListViewItems.Num());
+					// Strings need initialized values for later assignment
+					if constexpr (PCG::Private::MetadataTraits<ValueType>::NeedsConstruction)
+					{
+						CachedValues.SetNum(ListViewItems.Num());
+					}
+					else
+					{
+						CachedValues.SetNumUninitialized(ListViewItems.Num());
+					}
+
 					ColumnData->DataAccessor->GetRange(TArrayView<ValueType>(CachedValues), 0, *ColumnData->DataKeys);
 
 					auto SortAscending = [&CachedValues] (const PCGListviewItemPtr& LHS, const PCGListviewItemPtr& RHS)
