[PCG] Adds a node that applies a new random seed from point input using its position and user seed input.

diff --git a/Engine/Plugins/Experimental/PCG/Source/PCG/Private/Elements/PCGMutateSeed.cpp b/Engine/Plugins/Experimental/PCG/Source/PCG/Private/Elements/PCGMutateSeed.cpp
new file mode 100644
--- /dev/null
+++ b/Engine/Plugins/Experimental/PCG/Source/PCG/Private/Elements/PCGMutateSeed.cpp
@@ -0,0 +1,49 @@
+// Copyright Epic Games, Inc. All Rights Reserved.
+
+#include "Elements/PCGMutateSeed.h"
+
+#include "PCGContext.h"
+#include "PCGPin.h" 
+#include "PCGParamData.h"
+#include "Data/PCGPointData.h"
+#include "Helpers/PCGBlueprintHelpers.h"
+#include "Helpers/PCGHelpers.h"
+#include "Metadata/PCGMetadata.h"
+
+#define LOCTEXT_NAMESPACE "PCGMutateSeedSettings"
+
+UPCGMutateSeedSettings::UPCGMutateSeedSettings() 
+{
+	bUseSeed = true;
+}
+
+FPCGElementPtr UPCGMutateSeedSettings::CreateElement() const
+{
+	return MakeShared<FPCGMutateSeedElement>();
+}
+
+bool FPCGMutateSeedElement::ExecuteInternal(FPCGContext* Context) const
+{
+	TRACE_CPUPROFILER_EVENT_SCOPE(FPCGMutateSeedElement::Execute);
+	
+	check(Context);
+
+	const UPCGMutateSeedSettings* Settings = Context->GetInputSettings<UPCGMutateSeedSettings>();
+	check(Settings);
+
+	TArray<FPCGTaggedData> Inputs = Context->InputData.GetInputs();
+	TArray<FPCGTaggedData>& Outputs = Context->OutputData.TaggedData;
+
+	const int Seed = Context->GetSeed();
+
+	ProcessPoints(Context, Inputs, Outputs, [Seed](const FPCGPoint& InPoint, FPCGPoint& OutPoint)
+	{
+		OutPoint = InPoint;
+		OutPoint.Seed = PCGHelpers::ComputeSeed(UPCGBlueprintHelpers::ComputeSeedFromPosition(OutPoint.Transform.GetLocation()), Seed, OutPoint.Seed);
+		return true;
+	});
+	
+	return true;
+}
+
+#undef LOCTEXT_NAMESPACE
diff --git a/Engine/Plugins/Experimental/PCG/Source/PCG/Public/Elements/PCGMutateSeed.h b/Engine/Plugins/Experimental/PCG/Source/PCG/Public/Elements/PCGMutateSeed.h
new file mode 100644
--- /dev/null
+++ b/Engine/Plugins/Experimental/PCG/Source/PCG/Public/Elements/PCGMutateSeed.h
@@ -0,0 +1,43 @@
+// Copyright Epic Games, Inc. All Rights Reserved.
+
+#pragma once
+
+#include "PCGSettings.h"
+
+#include "Elements/PCGPointProcessingElementBase.h"
+
+#include "UObject/NameTypes.h"
+
+#include "PCGMutateSeed.generated.h"
+
+/**
+ * Generates a new seed for each point using its position and user seed input and applies to all points.
+ */
+UCLASS(BlueprintType, ClassGroup = (Procedural))
+class UPCGMutateSeedSettings : public UPCGSettings
+{
+	GENERATED_BODY()
+
+public:
+	UPCGMutateSeedSettings();
+
+	//~Begin UPCGSettings interface
+#if WITH_EDITOR
+	virtual FName GetDefaultNodeName() const override { return FName(TEXT("MutateSeed")); }
+	virtual FText GetDefaultNodeTitle() const override { return NSLOCTEXT("PCGMutateSeedElement", "NodeTitle", "MutateSeed"); }
+	virtual FText GetNodeTooltipText() const override { return NSLOCTEXT("PCGMutateSeedElement", "NodeTooltip", "Applies a new random seed from point input."); }
+	virtual EPCGSettingsType GetType() const override { return EPCGSettingsType::Spatial; }
+#endif
+
+protected:
+	virtual TArray<FPCGPinProperties> InputPinProperties() const override { return Super::DefaultPointInputPinProperties(); };
+	virtual TArray<FPCGPinProperties> OutputPinProperties() const override { return Super::DefaultPointOutputPinProperties(); };
+	virtual FPCGElementPtr CreateElement() const override;
+	//~End UPCGSettings interface
+};
+
+class FPCGMutateSeedElement : public FPCGPointProcessingElementBase
+{
+protected:
+	virtual bool ExecuteInternal(FPCGContext* Context) const override;
+};
\ No newline at end of file
