[PCG] Tentative build-fix following mutate seed node changes
#rnx

diff --git a/Engine/Plugins/Experimental/PCG/Source/PCG/Private/Elements/PCGMutateSeed.cpp b/Engine/Plugins/Experimental/PCG/Source/PCG/Private/Elements/PCGMutateSeed.cpp
--- a/Engine/Plugins/Experimental/PCG/Source/PCG/Private/Elements/PCGMutateSeed.cpp
+++ b/Engine/Plugins/Experimental/PCG/Source/PCG/Private/Elements/PCGMutateSeed.cpp
@@ -3,7 +3,6 @@
 #include "Elements/PCGMutateSeed.h"
 
 #include "PCGContext.h"
-#include "PCGPin.h" 
 #include "PCGParamData.h"
 #include "Data/PCGPointData.h"
 #include "Helpers/PCGBlueprintHelpers.h"
diff --git a/Engine/Plugins/Experimental/PCG/Source/PCG/Public/Elements/PCGMutateSeed.h b/Engine/Plugins/Experimental/PCG/Source/PCG/Public/Elements/PCGMutateSeed.h
--- a/Engine/Plugins/Experimental/PCG/Source/PCG/Public/Elements/PCGMutateSeed.h
+++ b/Engine/Plugins/Experimental/PCG/Source/PCG/Public/Elements/PCGMutateSeed.h
@@ -4,6 +4,7 @@
 
 #include "PCGSettings.h"
 
+#include "PCGPin.h"
 #include "Elements/PCGPointProcessingElementBase.h"
 
 #include "UObject/NameTypes.h"
