[PCG] Add post-paste operation to UPCGSettings to support copy/paste of nodes using graph parameter overrides

#rnx

diff --git a/Engine/Plugins/Experimental/PCG/Source/PCG/Private/PCGSettings.cpp b/Engine/Plugins/Experimental/PCG/Source/PCG/Private/PCGSettings.cpp
--- a/Engine/Plugins/Experimental/PCG/Source/PCG/Private/PCGSettings.cpp
+++ b/Engine/Plugins/Experimental/PCG/Source/PCG/Private/PCGSettings.cpp
@@ -9,6 +9,7 @@
 #include "PCGGraph.h"
 #include "PCGModule.h"
 #include "PCGPin.h"
+#include "PCGSubgraph.h"
 #include "PCGSubsystem.h"
 #include "Helpers/PCGHelpers.h"
 #include "Helpers/PCGSettingsHelpers.h"
@@ -605,6 +606,16 @@ bool UPCGSettings::IsPropertyOverriddenByPin(const FProperty* InProperty) const
 	return false;
 }
 
+void UPCGSettings::PostPaste()
+{
+	// Graph parameter overrides are not properly built when importing from pasted text (because param
+	// Properties are transient/unserialized), so for subgraphs we should reinitialize in post
+	if (UPCGSubgraphSettings* SubgraphSettings = Cast<UPCGSubgraphSettings>(this))
+	{
+		SubgraphSettings->InitializeCachedOverridableParams(/*bReset=*/true);
+	}
+}
+
 void UPCGSettings::ApplyDeprecationBeforeUpdatePins(UPCGNode* InOutNode, TArray<TObjectPtr<UPCGPin>>& InputPins, TArray<TObjectPtr<UPCGPin>>& OutputPins)
 {
 	check(InOutNode);
diff --git a/Engine/Plugins/Experimental/PCG/Source/PCG/Public/PCGSettings.h b/Engine/Plugins/Experimental/PCG/Source/PCG/Public/PCGSettings.h
--- a/Engine/Plugins/Experimental/PCG/Source/PCG/Public/PCGSettings.h
+++ b/Engine/Plugins/Experimental/PCG/Source/PCG/Public/PCGSettings.h
@@ -274,6 +274,9 @@ public:
 
 	/* If there are preconfigured info, we can skip the default settings and only expose pre-configured actions in the editor palette */
 	virtual bool OnlyExposePreconfiguredSettings() const { return false; }
+
+	/* Perform post-operations when an editor node is copied */
+	virtual void PostPaste();
 #endif
 
 	virtual void ApplyPreconfiguredSettings(const FPCGPreConfiguredSettingsInfo& PreconfigureInfo) {}
diff --git a/Engine/Plugins/Experimental/PCG/Source/PCGEditor/Private/PCGEditor.cpp b/Engine/Plugins/Experimental/PCG/Source/PCGEditor/Private/PCGEditor.cpp
--- a/Engine/Plugins/Experimental/PCG/Source/PCGEditor/Private/PCGEditor.cpp
+++ b/Engine/Plugins/Experimental/PCG/Source/PCGEditor/Private/PCGEditor.cpp
@@ -1961,6 +1961,11 @@ void FPCGEditor::PasteNodesHere(const FVector2D& Location)
 		if (UPCGNode* PastedPCGNode = PastedPCGGraphNode ? PastedPCGGraphNode->GetPCGNode() : nullptr)
 		{
 			PastedPCGGraphNode->PostPaste();
+
+			if (UPCGSettings* Settings = PastedPCGNode->GetSettings())
+			{
+				Settings->PostPaste();
+			}
 		}
 	}
 
