[PCG] Fix bug where Debug flag stops functioning for a graph

#rnx

diff --git a/Engine/Plugins/Experimental/PCG/Source/PCG/Private/PCGGraph.cpp b/Engine/Plugins/Experimental/PCG/Source/PCG/Private/PCGGraph.cpp
--- a/Engine/Plugins/Experimental/PCG/Source/PCG/Private/PCGGraph.cpp
+++ b/Engine/Plugins/Experimental/PCG/Source/PCG/Private/PCGGraph.cpp
@@ -792,29 +792,6 @@ void UPCGGraph::SetExtraEditorNodes(const TArray<TObjectPtr<const UObject>>& InN
 	}
 }
 
-bool UPCGGraph::IsInspecting() const
-{
-	return InspectionCounter > 0;
-}
-
-void UPCGGraph::EnableInspection()
-{
-	if (!ensure(InspectionCounter >= 0))
-	{
-		InspectionCounter = 0;
-	}
-
-	++InspectionCounter;
-}
-
-void UPCGGraph::DisableInspection()
-{
-	if (ensure(InspectionCounter > 0))
-	{
-		--InspectionCounter;
-	}
-};
-
 FPCGActorSelectionKeyToSettingsMap UPCGGraph::GetTrackedActorKeysToSettings() const
 {
 	FPCGActorSelectionKeyToSettingsMap TagsToSettings;
diff --git a/Engine/Plugins/Experimental/PCG/Source/PCG/Public/PCGGraph.h b/Engine/Plugins/Experimental/PCG/Source/PCG/Public/PCGGraph.h
--- a/Engine/Plugins/Experimental/PCG/Source/PCG/Public/PCGGraph.h
+++ b/Engine/Plugins/Experimental/PCG/Source/PCG/Public/PCGGraph.h
@@ -226,9 +226,9 @@ public:
 	const TArray<TObjectPtr<UObject>>& GetExtraEditorNodes() const { return ExtraEditorNodes; }
 	void SetExtraEditorNodes(const TArray<TObjectPtr<const UObject>>& InNodes);
 
-	bool IsInspecting() const;
-	void EnableInspection();
-	void DisableInspection();
+	bool IsInspecting() const { return bIsInspecting; }
+	void EnableInspection() { bIsInspecting = true; }
+	void DisableInspection() { bIsInspecting = false; }
 	bool DebugFlagAppliesToIndividualComponents() const { return bDebugFlagAppliesToIndividualComponents; }
 #endif
 
@@ -310,7 +310,7 @@ private:
 	bool bIsNotifying = false;
 	bool bUserPausedNotificationsInGraphEditor = false;
 	int32 NumberOfUserParametersPreEdit = 0;
-	int32 InspectionCounter = 0;
+	bool bIsInspecting = false;
 #endif // WITH_EDITOR
 };
 
diff --git a/Engine/Plugins/Experimental/PCG/Source/PCGEditor/Private/PCGEditor.cpp b/Engine/Plugins/Experimental/PCG/Source/PCGEditor/Private/PCGEditor.cpp
--- a/Engine/Plugins/Experimental/PCG/Source/PCGEditor/Private/PCGEditor.cpp
+++ b/Engine/Plugins/Experimental/PCG/Source/PCGEditor/Private/PCGEditor.cpp
@@ -2219,6 +2219,16 @@ void FPCGEditor::OnClose()
 
 	FAssetEditorToolkit::OnClose();
 
+	if (PCGComponentBeingInspected.IsValid())
+	{
+		PCGComponentBeingInspected->DisableInspection();
+
+		if (PCGGraphBeingEdited)
+		{
+			PCGGraphBeingEdited->DisableInspection();
+		}
+	}
+
 	if (PCGGraphBeingEdited)
 	{
 		if (PCGGraphBeingEdited->NotificationsForEditorArePausedByUser())
