[PCG] Disable ensure()'s that occur if no PCGWorldActor is present

The ensure was more a development aid rather than something abnormal or broken in code. This fixes an engine test failure.

#rnx

diff --git a/Engine/Plugins/Experimental/PCG/Source/PCG/Private/Helpers/PCGHelpers.cpp b/Engine/Plugins/Experimental/PCG/Source/PCG/Private/Helpers/PCGHelpers.cpp
--- a/Engine/Plugins/Experimental/PCG/Source/PCG/Private/Helpers/PCGHelpers.cpp
+++ b/Engine/Plugins/Experimental/PCG/Source/PCG/Private/Helpers/PCGHelpers.cpp
@@ -434,9 +434,8 @@ namespace PCGHelpers
 		}
 		else
 		{
-			ensure(false);
-			OutGridSizes.Add(PCGHiGenGrid::UninitializedGridSize());
-			return false;
+			OutGridSizes.Add(PCGHiGenGrid::UnboundedGridSize());
+			return true;
 		}
 	}
 
