[PCG] Fixed a crash in the Static Mesh Spawner, related to garbage collection

* Instance packer blueprint override implementation was invoking a default copy assignment op, which in turn was being destructed and flagging the OverrideSettings for GC

#rnx

diff --git a/Engine/Plugins/Experimental/PCG/Source/PCG/Private/Elements/PCGStaticMeshSpawner.cpp b/Engine/Plugins/Experimental/PCG/Source/PCG/Private/Elements/PCGStaticMeshSpawner.cpp
--- a/Engine/Plugins/Experimental/PCG/Source/PCG/Private/Elements/PCGStaticMeshSpawner.cpp
+++ b/Engine/Plugins/Experimental/PCG/Source/PCG/Private/Elements/PCGStaticMeshSpawner.cpp
@@ -230,10 +230,16 @@ bool FPCGStaticMeshSpawnerElement::PrepareDataInternal(FPCGContext* InContext) c
 
 			if (Settings->InstanceDataPackerParameters)
 			{
+				// This blueprint overridden implementation will create an implicit copy of the context
+				// and when the copy destructs, will cause the SettingsOverride to unroot and mark for GC
+				const bool bShouldUnrootSettingsOnDelete = Context->bShouldUnrootSettingsOnDelete;
+				Context->bShouldUnrootSettingsOnDelete = false;
 				for (int32 InstanceListIndex = 0; InstanceListIndex < MeshInstances.Num(); ++InstanceListIndex)
 				{
 					Settings->InstanceDataPackerParameters->PackInstances(*Context, Context->CurrentPointData, MeshInstances[InstanceListIndex], PackedCustomData[InstanceListIndex]);
 				}
+				// Put back the proper unroot flag
+				Context->bShouldUnrootSettingsOnDelete = bShouldUnrootSettingsOnDelete;
 			}
 		}
 
