[PCG] Fix a bug that created a gap in the Surface Sampler generated points across Landscapes with Looseness at 0

#rnx

diff --git a/Engine/Plugins/Experimental/PCG/Source/PCG/Private/Data/PCGLandscapeData.cpp b/Engine/Plugins/Experimental/PCG/Source/PCG/Private/Data/PCGLandscapeData.cpp
--- a/Engine/Plugins/Experimental/PCG/Source/PCG/Private/Data/PCGLandscapeData.cpp
+++ b/Engine/Plugins/Experimental/PCG/Source/PCG/Private/Data/PCGLandscapeData.cpp
@@ -354,7 +354,7 @@ const ULandscapeInfo* UPCGLandscapeData::GetLandscapeInfo(const FVector& InPosit
 	// TODO: we could maybe improve on this if we find the "nearest" landscape on a Z perspective, but this might still lead to issues
 	for (const TPair<FBox, ULandscapeInfo*>& LandscapeInfoPair : LandscapeInfos)
 	{
-		if (LandscapeInfoPair.Key.IsInsideXY(InPosition))
+		if (PCGHelpers::IsInsideBoundsXY(LandscapeInfoPair.Key, InPosition))
 		{
 			return LandscapeInfoPair.Value;
 		}
