[PCG] Apply UPCGSettings::PostPaste() to all settings instead of only SubgraphSettings

#rnx

diff --git a/Engine/Plugins/Experimental/PCG/Source/PCG/Private/PCGContext.cpp b/Engine/Plugins/Experimental/PCG/Source/PCG/Private/PCGContext.cpp
--- a/Engine/Plugins/Experimental/PCG/Source/PCG/Private/PCGContext.cpp
+++ b/Engine/Plugins/Experimental/PCG/Source/PCG/Private/PCGContext.cpp
@@ -155,7 +155,11 @@ void FPCGContext::OverrideSettings()
 
 	for (const FPCGSettingsOverridableParam& Param : OriginalSettings->OverridableParams())
 	{
-		check(!Param.Properties.IsEmpty());
+		if (!ensure(!Param.Properties.IsEmpty()))
+		{
+			PCGE_LOG_C(Error, GraphAndLog, this, FText::Format(LOCTEXT("ParamPropertyIsEmpty", "Override pin '{0}' has no property set, we can't override it."), FText::FromName(Param.Label)));
+			continue;
+		}
 
 		// Verification that container is valid and we have the right class.
 		void* Container = nullptr;
diff --git a/Engine/Plugins/Experimental/PCG/Source/PCG/Private/PCGSettings.cpp b/Engine/Plugins/Experimental/PCG/Source/PCG/Private/PCGSettings.cpp
--- a/Engine/Plugins/Experimental/PCG/Source/PCG/Private/PCGSettings.cpp
+++ b/Engine/Plugins/Experimental/PCG/Source/PCG/Private/PCGSettings.cpp
@@ -608,12 +608,9 @@ bool UPCGSettings::IsPropertyOverriddenByPin(const FProperty* InProperty) const
 
 void UPCGSettings::PostPaste()
 {
-	// Graph parameter overrides are not properly built when importing from pasted text (because param
-	// Properties are transient/unserialized), so for subgraphs we should reinitialize in post
-	if (UPCGSubgraphSettings* SubgraphSettings = Cast<UPCGSubgraphSettings>(this))
-	{
-		SubgraphSettings->InitializeCachedOverridableParams(/*bReset=*/true);
-	}
+	// Params are not properly built when importing from pasted text (because param
+	// Properties are transient/unserialized), so we should reinitialize in post
+	InitializeCachedOverridableParams(/*bReset=*/true);
 }
 
 void UPCGSettings::ApplyDeprecationBeforeUpdatePins(UPCGNode* InOutNode, TArray<TObjectPtr<UPCGPin>>& InputPins, TArray<TObjectPtr<UPCGPin>>& OutputPins)
