*Fixed a crash when selecting an execute blueprint node after undoing set blueprint property
*Fixed an issue were undo wouldnt work after changing spawn actor settings

diff --git a/Engine/Plugins/Experimental/PCG/Source/PCG/Private/Elements/PCGCreateTargetActor.cpp b/Engine/Plugins/Experimental/PCG/Source/PCG/Private/Elements/PCGCreateTargetActor.cpp
--- a/Engine/Plugins/Experimental/PCG/Source/PCG/Private/Elements/PCGCreateTargetActor.cpp
+++ b/Engine/Plugins/Experimental/PCG/Source/PCG/Private/Elements/PCGCreateTargetActor.cpp
@@ -82,6 +82,8 @@ void UPCGCreateTargetActor::TeardownBlueprintEvent()
 
 void UPCGCreateTargetActor::PreEditChange(FProperty* PropertyAboutToChange)
 {
+	Super::PreEditChange(PropertyAboutToChange);
+
 	if (PropertyAboutToChange && PropertyAboutToChange->GetFName() == GET_MEMBER_NAME_CHECKED(UPCGCreateTargetActor, TemplateActorClass))
 	{
 		TeardownBlueprintEvent();
diff --git a/Engine/Plugins/Experimental/PCG/Source/PCG/Private/Elements/PCGExecuteBlueprint.cpp b/Engine/Plugins/Experimental/PCG/Source/PCG/Private/Elements/PCGExecuteBlueprint.cpp
--- a/Engine/Plugins/Experimental/PCG/Source/PCG/Private/Elements/PCGExecuteBlueprint.cpp
+++ b/Engine/Plugins/Experimental/PCG/Source/PCG/Private/Elements/PCGExecuteBlueprint.cpp
@@ -388,6 +388,8 @@ void UPCGBlueprintSettings::BeginDestroy()
 #if WITH_EDITOR
 void UPCGBlueprintSettings::PreEditChange(FProperty* PropertyAboutToChange)
 {
+	Super::PreEditChange(PropertyAboutToChange);
+
 	if (PropertyAboutToChange)
 	{
 		if (PropertyAboutToChange->GetFName() == GET_MEMBER_NAME_CHECKED(UPCGBlueprintSettings, BlueprintElementType))
diff --git a/Engine/Plugins/Experimental/PCG/Source/PCG/Private/Elements/PCGSpawnActor.cpp b/Engine/Plugins/Experimental/PCG/Source/PCG/Private/Elements/PCGSpawnActor.cpp
--- a/Engine/Plugins/Experimental/PCG/Source/PCG/Private/Elements/PCGSpawnActor.cpp
+++ b/Engine/Plugins/Experimental/PCG/Source/PCG/Private/Elements/PCGSpawnActor.cpp
@@ -331,6 +331,8 @@ TObjectPtr<UPCGGraphInterface> UPCGSpawnActorNode::GetSubgraphInterface() const
 #if WITH_EDITOR
 void UPCGSpawnActorSettings::PreEditChange(FProperty* PropertyAboutToChange)
 {
+	Super::PreEditChange(PropertyAboutToChange);
+
 	if (PropertyAboutToChange && PropertyAboutToChange->GetFName() == GET_MEMBER_NAME_CHECKED(UPCGSpawnActorSettings, TemplateActorClass))
 	{
 		TeardownBlueprintEvent();
