Fixed a few outstanding issues found when testing in Electric Dreams:
- In PIE, actors from the editor world would react to callbacks from the PIE world, fixed by separating actor tracking by world
- When dragging in a temporary actor, the actor tracking could trigger changes, disabled reacting on preview actors
- The refresh in the PCG component instance data should not have been a structural one as that broke generate-on-placement
- Fixed a potential crash issue when the PCG component holds resources to remote actors by removing pointer caching
- Fixed streaming issue by making the cached tracked actors map transient
- Fixed movable warning when changing the transform on ISMs

#rnx

diff --git a/Engine/Plugins/Experimental/PCG/Source/PCG/Private/PCGActorAndComponentMapping.cpp b/Engine/Plugins/Experimental/PCG/Source/PCG/Private/PCGActorAndComponentMapping.cpp
--- a/Engine/Plugins/Experimental/PCG/Source/PCG/Private/PCGActorAndComponentMapping.cpp
+++ b/Engine/Plugins/Experimental/PCG/Source/PCG/Private/PCGActorAndComponentMapping.cpp
@@ -149,6 +149,12 @@ bool UPCGActorAndComponentMapping::RegisterOrUpdatePCGComponent(UPCGComponent* I
 		bHasChanged = RegisterOrUpdateNonPartitionedPCGComponent(InComponent);
 	}
 
+	// If the component was previously marked as to be unregistered, remove it here.
+	{
+		FScopeLock Lock(&DelayedComponentToUnregisterLock);
+		DelayedComponentToUnregister.Remove(InComponent);
+	}
+
 	// And finally handle the tracking. Only do it when the component is registered for the first time.
 #if WITH_EDITOR
 	if (!bWasAlreadyRegistered && bHasChanged)
@@ -822,11 +828,18 @@ void UPCGActorAndComponentMapping::OnActorAdded(AActor* InActor)
 	TRACE_CPUPROFILER_EVENT_SCOPE(UPCGActorAndComponentMapping::OnActorAdded);
 
 	// We have to make sure to not create a infinite loop
-	if (!InActor || InActor->IsA<APCGWorldActor>() || !PCGSubsystem)
+	if (!InActor || InActor->IsA<APCGWorldActor>() || !PCGSubsystem || InActor->GetWorld() != PCGSubsystem->GetWorld())
 	{
 		return;
 	}
 
+#if WITH_EDITOR
+	if (InActor->bIsEditorPreviewActor)
+	{
+		return;
+	}
+#endif
+
 	// If the subsystem is not initialized, wait for it to be, and store all the actors to check
 	if (!PCGSubsystem->IsInitialized())
 	{
@@ -844,10 +857,17 @@ void UPCGActorAndComponentMapping::OnActorAdded(AActor* InActor)
 bool UPCGActorAndComponentMapping::AddOrUpdateTrackedActor(AActor* InActor)
 {
 	// We have to make sure to not create a infinite loop
-	if (!InActor || InActor->IsA<APCGWorldActor>() || !PCGSubsystem || !PCGSubsystem->FindPCGWorldActor())
+	if (!InActor || InActor->IsA<APCGWorldActor>() || !PCGSubsystem || !PCGSubsystem->FindPCGWorldActor() || InActor->GetWorld() != PCGSubsystem->GetWorld())
+	{
+		return false;
+	}
+
+#if WITH_EDITOR
+	if (InActor->bIsEditorPreviewActor)
 	{
 		return false;
 	}
+#endif
 
 	// Gather all components, and check if they want to track this one
 	TSet<UPCGComponent*> AllComponents = GetAllRegisteredComponents();
@@ -921,7 +941,6 @@ void UPCGActorAndComponentMapping::RegisterActor(AActor* InActor)
 
 	if (!PCGWorldActor->CachedTrackedActors.Contains(InActor))
 	{
-		PCGWorldActor->Modify();
 		PCGWorldActor->CachedTrackedActors.Add(InActor);
 	}
 
@@ -950,7 +969,6 @@ bool UPCGActorAndComponentMapping::UnregisterActor(AActor* InActor)
 
 	if (PCGWorldActor->CachedTrackedActors.Contains(InActor))
 	{
-		PCGWorldActor->Modify();
 		PCGWorldActor->CachedTrackedActors.Remove(InActor);
 		TrackedActorToPositionMap.Remove(InActor);
 		CulledTrackedActorsToComponentsMap.Remove(InActor);
@@ -974,8 +992,20 @@ void UPCGActorAndComponentMapping::OnActorDeleted(AActor* InActor)
 {
 	TRACE_CPUPROFILER_EVENT_SCOPE(UPCGActorAndComponentMapping::OnActorDeleted);
 
-	APCGWorldActor* PCGWorldActor = PCGSubsystem ? PCGSubsystem->FindPCGWorldActor() : nullptr;
-	if (!PCGWorldActor || !InActor || !PCGWorldActor->CachedTrackedActors.Contains(InActor))
+	if (!InActor || !PCGSubsystem || InActor->GetWorld() != PCGSubsystem->GetWorld())
+	{
+		return;
+	}
+
+#if WITH_EDITOR
+	if (InActor->bIsEditorPreviewActor)
+	{
+		return;
+	}
+#endif
+
+	APCGWorldActor* PCGWorldActor = PCGSubsystem->FindPCGWorldActor();
+	if (!PCGWorldActor || !PCGWorldActor->CachedTrackedActors.Contains(InActor))
 	{
 		return;
 	}
@@ -992,11 +1022,18 @@ void UPCGActorAndComponentMapping::OnActorMoved(AActor* InActor)
 	TRACE_CPUPROFILER_EVENT_SCOPE(UPCGActorAndComponentMapping::OnActorMoved);
 
 	APCGWorldActor* PCGWorldActor = PCGSubsystem ? PCGSubsystem->FindPCGWorldActor() : nullptr;
-	if (!PCGWorldActor || !InActor || !PCGWorldActor->CachedTrackedActors.Contains(InActor))
+	if (!PCGWorldActor || !InActor || (PCGSubsystem && InActor->GetWorld() != PCGSubsystem->GetWorld()) || !PCGWorldActor->CachedTrackedActors.Contains(InActor))
 	{
 		return;
 	}
 
+#if WITH_EDITOR
+	if (InActor->bIsEditorPreviewActor)
+	{
+		return;
+	}
+#endif
+
 	// Notify all components
 	OnActorChanged(InActor, /*bInHasMoved=*/ true);
 
@@ -1014,10 +1051,17 @@ void UPCGActorAndComponentMapping::OnPreObjectPropertyChanged(UObject* InObject,
 	FProperty* MemberProperty = InEditPropertyChain.GetActiveMemberNode() ? InEditPropertyChain.GetActiveMemberNode()->GetValue() : nullptr;
 	AActor* Actor = Cast<AActor>(InObject);
 
-	if (!Actor || !MemberProperty || MemberProperty->GetFName() != GET_MEMBER_NAME_CHECKED(AActor, Tags))
+	if (!Actor || (PCGSubsystem && Actor->GetWorld() != PCGSubsystem->GetWorld()) || !MemberProperty || MemberProperty->GetFName() != GET_MEMBER_NAME_CHECKED(AActor, Tags))
+	{
+		return;
+	}
+
+#if WITH_EDITOR
+	if (Actor->bIsEditorPreviewActor)
 	{
 		return;
 	}
+#endif
 
 	TempTrackedActorTags = TSet<FName>(Actor->Tags);
 }
@@ -1080,6 +1124,18 @@ void UPCGActorAndComponentMapping::OnObjectPropertyChanged(UObject* InObject, FP
 		return;
 	}
 
+	if (PCGSubsystem && Actor->GetWorld() != PCGSubsystem->GetWorld())
+	{
+		return;
+	}
+
+#if WITH_EDITOR
+	if (Actor->bIsEditorPreviewActor)
+	{
+		return;
+	}
+#endif
+
 	// Check if we are not tracking it or is a tag change.
 	bool bShouldChange = true;
 	if (!PCGWorldActor->CachedTrackedActors.Contains(Actor) || bActorTagChange)
@@ -1103,6 +1159,7 @@ void UPCGActorAndComponentMapping::OnActorChanged(AActor* InActor, bool bInHasMo
 	TRACE_CPUPROFILER_EVENT_SCOPE(UPCGActorAndComponentMapping::OnActorChanged);
 
 	check(InActor);
+	ensure(!PCGSubsystem || InActor->GetWorld() == PCGSubsystem->GetWorld());
 	TSet<UPCGComponent*> DirtyComponents;
 
 	EPCGComponentDirtyFlag DirtyFlag = EPCGComponentDirtyFlag::Actor;
diff --git a/Engine/Plugins/Experimental/PCG/Source/PCG/Private/PCGComponent.cpp b/Engine/Plugins/Experimental/PCG/Source/PCG/Private/PCGComponent.cpp
--- a/Engine/Plugins/Experimental/PCG/Source/PCG/Private/PCGComponent.cpp
+++ b/Engine/Plugins/Experimental/PCG/Source/PCG/Private/PCGComponent.cpp
@@ -2420,7 +2420,7 @@ void FPCGComponentInstanceData::ApplyToComponent(UActorComponent* Component, con
 		// Note that we only do this if we are not currently loading
 		if (!SourceComponent || !SourceComponent->HasAllFlags(RF_WasLoaded))
 		{
-			PCGComponent->Refresh(/*bIsStructural=*/true);
+			PCGComponent->Refresh();
 		}
 #endif // WITH_EDITOR
 	}
diff --git a/Engine/Plugins/Experimental/PCG/Source/PCG/Private/PCGManagedResource.cpp b/Engine/Plugins/Experimental/PCG/Source/PCG/Private/PCGManagedResource.cpp
--- a/Engine/Plugins/Experimental/PCG/Source/PCG/Private/PCGManagedResource.cpp
+++ b/Engine/Plugins/Experimental/PCG/Source/PCG/Private/PCGManagedResource.cpp
@@ -385,7 +385,10 @@ void UPCGManagedISMComponent::MarkAsUsed()
 		}
 
 		// Reset the rotation/scale to be identity otherwise if the root component transform has changed, the final transform will be wrong
+		// Since this is technically 'moving' the ISM, we need to unregister it before moving otherwise we could get a warning that we're moving a component with static mobility
+		ISMC->UnregisterComponent();
 		ISMC->SetWorldTransform(FTransform(FQuat::Identity, RootLocation, FVector::OneVector));
+		ISMC->RegisterComponent();
 	}
 }
 
@@ -406,7 +409,10 @@ void UPCGManagedISMComponent::MarkAsReused()
 			}
 		}
 
+		// Since this is technically 'moving' the ISM, we need to unregister it before moving otherwise we could get a warning that we're moving a component with static mobility
+		ISMC->UnregisterComponent();
 		ISMC->SetWorldTransform(FTransform(FQuat::Identity, TentativeRootLocation, FVector::OneVector));
+		ISMC->RegisterComponent();
 	}
 }
 
@@ -420,7 +426,16 @@ UInstancedStaticMeshComponent* UPCGManagedISMComponent::GetComponent() const
 {
 	if (!CachedRawComponentPtr)
 	{
-		CachedRawComponentPtr = Cast<UInstancedStaticMeshComponent>(GeneratedComponent.Get());
+		UInstancedStaticMeshComponent* GeneratedComponentPtr = Cast<UInstancedStaticMeshComponent>(GeneratedComponent.Get());
+
+		// Implementation note:
+		// There is no surefire way to make sure that we can use the raw pointer UNLESS it is from the same owner
+		if (GeneratedComponentPtr && Cast<UPCGComponent>(GetOuter()) && GeneratedComponentPtr->GetOwner() == Cast<UPCGComponent>(GetOuter())->GetOwner())
+		{
+			CachedRawComponentPtr = GeneratedComponentPtr;
+		}
+
+		return GeneratedComponentPtr;
 	}
 
 	return CachedRawComponentPtr;
diff --git a/Engine/Plugins/Experimental/PCG/Source/PCG/Public/PCGWorldActor.h b/Engine/Plugins/Experimental/PCG/Source/PCG/Public/PCGWorldActor.h
--- a/Engine/Plugins/Experimental/PCG/Source/PCG/Public/PCGWorldActor.h
+++ b/Engine/Plugins/Experimental/PCG/Source/PCG/Public/PCGWorldActor.h
@@ -63,7 +63,7 @@ public:
 
 #if WITH_EDITORONLY_DATA
 	// Cached tracked actors list is serialized because we can't get it at postload time
-	UPROPERTY()
+	UPROPERTY(Transient)
 	TSet<TWeakObjectPtr<AActor>> CachedTrackedActors;
 #endif // WITH_EDITORONLY_DATA
 
