Disabled usage of target actor in artifact creation to sidestep caching issues until we fully remove the target actor concept

#rnx

diff --git a/Engine/Plugins/Experimental/PCG/Source/PCG/Private/Elements/PCGCopyPoints.cpp b/Engine/Plugins/Experimental/PCG/Source/PCG/Private/Elements/PCGCopyPoints.cpp
--- a/Engine/Plugins/Experimental/PCG/Source/PCG/Private/Elements/PCGCopyPoints.cpp
+++ b/Engine/Plugins/Experimental/PCG/Source/PCG/Private/Elements/PCGCopyPoints.cpp
@@ -152,8 +152,8 @@ bool FPCGCopyPointsElement::ExecuteInternal(FPCGContext* Context) const
 		RootMetadata = NonRootMetadata = nullptr;
 	}
 
-	// Priorize use the target actor from the target, irrespective of the source
-	OutPointData->TargetActor = TargetPointData->TargetActor.IsValid() ? TargetPointData->TargetActor : SourcePointData->TargetActor;
+	// Always use the target actor from the target, irrespective of the source
+	OutPointData->TargetActor = TargetPointData->TargetActor;
 
 	UPCGMetadata* OutMetadata = OutPointData->Metadata;
 	check(OutMetadata);
diff --git a/Engine/Plugins/Experimental/PCG/Source/PCG/Private/Elements/PCGCreateSpline.cpp b/Engine/Plugins/Experimental/PCG/Source/PCG/Private/Elements/PCGCreateSpline.cpp
--- a/Engine/Plugins/Experimental/PCG/Source/PCG/Private/Elements/PCGCreateSpline.cpp
+++ b/Engine/Plugins/Experimental/PCG/Source/PCG/Private/Elements/PCGCreateSpline.cpp
@@ -83,7 +83,7 @@ bool FPCGCreateSplineElement::ExecuteInternal(FPCGContext* Context) const
 			continue;
 		}
 
-		AActor* TargetActor = Settings->TargetActor.Get() ? Settings->TargetActor.Get() : Context->GetTargetActor(SpatialData);
+		AActor* TargetActor = Settings->TargetActor.Get() ? Settings->TargetActor.Get() : Context->GetTargetActor(nullptr);
 		if (!TargetActor)
 		{
 			PCGE_LOG(Error, GraphAndLog, LOCTEXT("InvalidTargetActor", "Invalid target actor. Ensure TargetActor member is initialized when creating SpatialData."));
diff --git a/Engine/Plugins/Experimental/PCG/Source/PCG/Private/Elements/PCGDebugElement.cpp b/Engine/Plugins/Experimental/PCG/Source/PCG/Private/Elements/PCGDebugElement.cpp
--- a/Engine/Plugins/Experimental/PCG/Source/PCG/Private/Elements/PCGDebugElement.cpp
+++ b/Engine/Plugins/Experimental/PCG/Source/PCG/Private/Elements/PCGDebugElement.cpp
@@ -93,7 +93,7 @@ namespace PCGDebugElement
 				continue;
 			}
 
-			AActor* TargetActor = InTargetActor ? InTargetActor : Context->GetTargetActor(SpatialData);
+			AActor* TargetActor = InTargetActor ? InTargetActor : Context->GetTargetActor(nullptr);
 
 			if (!TargetActor)
 			{
diff --git a/Engine/Plugins/Experimental/PCG/Source/PCG/Private/Elements/PCGSpawnActor.cpp b/Engine/Plugins/Experimental/PCG/Source/PCG/Private/Elements/PCGSpawnActor.cpp
--- a/Engine/Plugins/Experimental/PCG/Source/PCG/Private/Elements/PCGSpawnActor.cpp
+++ b/Engine/Plugins/Experimental/PCG/Source/PCG/Private/Elements/PCGSpawnActor.cpp
@@ -536,7 +536,7 @@ bool FPCGSpawnActorElement::SpawnAndPrepareSubgraphs(FPCGSubgraphContext* Contex
 			continue;
 		}
 
-		AActor* TargetActor = Settings->RootActor.Get() ? Settings->RootActor.Get() : Context->GetTargetActor(SpatialData);
+		AActor* TargetActor = Settings->RootActor.Get() ? Settings->RootActor.Get() : Context->GetTargetActor(nullptr);
 
 		if (!TargetActor)
 		{
diff --git a/Engine/Plugins/Experimental/PCG/Source/PCG/Private/Elements/PCGStaticMeshSpawner.cpp b/Engine/Plugins/Experimental/PCG/Source/PCG/Private/Elements/PCGStaticMeshSpawner.cpp
--- a/Engine/Plugins/Experimental/PCG/Source/PCG/Private/Elements/PCGStaticMeshSpawner.cpp
+++ b/Engine/Plugins/Experimental/PCG/Source/PCG/Private/Elements/PCGStaticMeshSpawner.cpp
@@ -163,7 +163,7 @@ bool FPCGStaticMeshSpawnerElement::PrepareDataInternal(FPCGContext* InContext) c
 				continue;
 			}
 
-			AActor* TargetActor = Settings->TargetActor.Get() ? Settings->TargetActor.Get() : Context->GetTargetActor(PointData);
+			AActor* TargetActor = Settings->TargetActor.Get() ? Settings->TargetActor.Get() : Context->GetTargetActor(nullptr);
 			if (!TargetActor)
 			{
 				PCGE_LOG(Error, GraphAndLog, LOCTEXT("InvalidTargetActor", "Invalid target actor. Ensure TargetActor member is initialized when creating SpatialData."));
