[PCG] Make [REDACTED_HANDLE]/[REDACTED_HANDLE] more robust and consistent

* Extra accessors (like MyPos.X) now works on [REDACTED_HANDLE] and [REDACTED_HANDLE].
* [REDACTED_HANDLE] will be the lastest attribute modified WITH extra accessors (was already the case).
* Latest attribute is shown in the attribute list with its extra accessors (was already the case).
* Fix overlook in metadata operation when last selector was set on const input data.

#rnx

diff --git a/Engine/Plugins/Experimental/PCG/Source/PCG/Private/Elements/Metadata/PCGMetadataOpElementBase.cpp b/Engine/Plugins/Experimental/PCG/Source/PCG/Private/Elements/Metadata/PCGMetadataOpElementBase.cpp
--- a/Engine/Plugins/Experimental/PCG/Source/PCG/Private/Elements/Metadata/PCGMetadataOpElementBase.cpp
+++ b/Engine/Plugins/Experimental/PCG/Source/PCG/Private/Elements/Metadata/PCGMetadataOpElementBase.cpp
@@ -508,13 +508,14 @@ bool FPCGMetadataElementBase::ExecuteInternal(FPCGContext* Context) const
 		else
 		{
 			// In case of property or attribute with extra accessor, we need to validate that the property/attribute can accept the output type.
-			// Verify this before duplicating.
-			OperationData.OutputAccessors[OutputIndex] = PCGAttributeAccessorHelpers::CreateAccessor(Cast<UPCGData>(OutputData.Data), OutputTarget);
+			// Verify this before duplicating, because an extra allocation is certainly less costly than duplicating the data.
+			// Do it with a const accessor, since OutputData.Data is still pointing on the const input data.
+			TUniquePtr<const IPCGAttributeAccessor> TempConstAccessor = PCGAttributeAccessorHelpers::CreateConstAccessor(Cast<const UPCGData>(OutputData.Data), OutputTarget);
 
-			if (OperationData.OutputAccessors[OutputIndex].IsValid())
+			if (TempConstAccessor.IsValid())
 			{
 				// We matched an attribute/property, check if the output type is valid.
-				if (!PCG::Private::IsBroadcastable(PCG::Private::MetadataTypes<AttributeType>::Id, OperationData.OutputAccessors[OutputIndex]->GetUnderlyingType()))
+				if (!PCG::Private::IsBroadcastable(PCG::Private::MetadataTypes<AttributeType>::Id, TempConstAccessor->GetUnderlyingType()))
 				{
 					PCGE_LOG(Error, GraphAndLog, FText::Format(LOCTEXT("AttributeTypeBroadcastFailed", "Attribute/Property '{0}' cannot be broadcasted to match types for input"), FText::FromName(OutputName)));
 					return false;
diff --git a/Engine/Plugins/Experimental/PCG/Source/PCG/Private/Metadata/PCGAttributePropertySelector.cpp b/Engine/Plugins/Experimental/PCG/Source/PCG/Private/Metadata/PCGAttributePropertySelector.cpp
--- a/Engine/Plugins/Experimental/PCG/Source/PCG/Private/Metadata/PCGAttributePropertySelector.cpp
+++ b/Engine/Plugins/Experimental/PCG/Source/PCG/Private/Metadata/PCGAttributePropertySelector.cpp
@@ -258,9 +258,12 @@ FPCGAttributePropertyInputSelector FPCGAttributePropertyInputSelector::CopyAndFi
 {
 	if (Selection == EPCGAttributePropertySelection::Attribute)
 	{
+		// For each case, append extra names to the newly created selector.
 		if (AttributeName == PCGMetadataAttributeConstants::LastAttributeName && InData && InData->HasCachedLastSelector())
 		{
-			return InData->GetCachedLastSelector();
+			FPCGAttributePropertyInputSelector Selector = InData->GetCachedLastSelector();
+			Selector.ExtraNames.Append(ExtraNames);
+			return Selector;
 		}
 		else if (AttributeName == PCGMetadataAttributeConstants::LastCreatedAttributeName && InData)
 		{
@@ -268,6 +271,7 @@ FPCGAttributePropertyInputSelector FPCGAttributePropertyInputSelector::CopyAndFi
 			{
 				FPCGAttributePropertyInputSelector Selector;
 				Selector.SetAttributeName(Metadata->GetLatestAttributeNameOrNone());
+				Selector.ExtraNames.Append(ExtraNames);
 				return Selector;
 			}
 		}
@@ -307,14 +311,18 @@ FPCGAttributePropertyOutputSelector FPCGAttributePropertyOutputSelector::CopyAnd
 {
 	if (Selection == EPCGAttributePropertySelection::Attribute)
 	{
+		// For each case, append extra names to the newly created selector.
 		if (AttributeName == PCGMetadataAttributeConstants::SourceAttributeName && InSourceSelector)
 		{
-			return FPCGAttributePropertySelector::CreateFromOtherSelector<FPCGAttributePropertyOutputSelector>(*InSourceSelector);
+			FPCGAttributePropertyOutputSelector Selector = FPCGAttributePropertySelector::CreateFromOtherSelector<FPCGAttributePropertyOutputSelector>(*InSourceSelector);
+			Selector.ExtraNames.Append(ExtraNames);
+			return Selector;
 		}
 		else if (AttributeName == PCGMetadataAttributeConstants::SourceNameAttributeName && InSourceSelector)
 		{
 			FPCGAttributePropertyOutputSelector Selector;
 			Selector.SetAttributeName(InSourceSelector->GetName());
+			Selector.ExtraNames.Append(ExtraNames);
 			return Selector;
 		}
 		// Only for deprecation
@@ -324,6 +332,7 @@ FPCGAttributePropertyOutputSelector FPCGAttributePropertyOutputSelector::CopyAnd
 			{
 				FPCGAttributePropertyOutputSelector Selector;
 				Selector.SetAttributeName(Metadata->GetLatestAttributeNameOrNone());
+				Selector.ExtraNames.Append(ExtraNames);
 				return Selector;
 			}
 		}
