[PCG] Fixed a crash when inspecting a transient debug graph object created from the Determinism Tests

* Also fixed an issue where SettingsWithOverride was incorrectly setting itself as Transient

#rnx

diff --git a/Engine/Plugins/Experimental/PCG/Source/PCG/Private/PCGContext.cpp b/Engine/Plugins/Experimental/PCG/Source/PCG/Private/PCGContext.cpp
--- a/Engine/Plugins/Experimental/PCG/Source/PCG/Private/PCGContext.cpp
+++ b/Engine/Plugins/Experimental/PCG/Source/PCG/Private/PCGContext.cpp
@@ -130,7 +130,8 @@ void FPCGContext::InitializeSettings()
 
 			if (bHasParamConnected)
 			{
-				SettingsWithOverride = Cast<UPCGSettings>(StaticDuplicateObject(NodeSettings, GetTransientPackage(), NAME_None, RF_Transient));
+				SettingsWithOverride = Cast<UPCGSettings>(StaticDuplicateObject(NodeSettings, GetTransientPackage()));
+				SettingsWithOverride->SetFlags(RF_Transient);
 				SettingsWithOverride->AddToRoot();
 
 				// Force seed copy to prevent issue due to delta serialization vs. Seed being initialized in the constructor only for new nodes
diff --git a/Engine/Plugins/Experimental/PCG/Source/PCG/Private/Tests/Determinism/PCGDeterminismTestsCommon.cpp b/Engine/Plugins/Experimental/PCG/Source/PCG/Private/Tests/Determinism/PCGDeterminismTestsCommon.cpp
--- a/Engine/Plugins/Experimental/PCG/Source/PCG/Private/Tests/Determinism/PCGDeterminismTestsCommon.cpp
+++ b/Engine/Plugins/Experimental/PCG/Source/PCG/Private/Tests/Determinism/PCGDeterminismTestsCommon.cpp
@@ -89,6 +89,7 @@ namespace PCGDeterminismTests
 		PCGActorCopy->SetFlags(RF_Transient);
 		UPCGComponent* PCGComponentCopy = PCGActorCopy->FindComponentByClass<UPCGComponent>();
 		check(PCGComponentCopy);
+		PCGComponentCopy->SetFlags(RF_Transient);
 		PCGComponentCopy->SetIsPartitioned(false);
 
 		FPCGGraphExecutor Executor = FPCGGraphExecutor(PCGComponentCopy);
diff --git a/Engine/Plugins/Experimental/PCG/Source/PCGEditor/Private/SPCGEditorGraphDebugObjectTree.cpp b/Engine/Plugins/Experimental/PCG/Source/PCGEditor/Private/SPCGEditorGraphDebugObjectTree.cpp
--- a/Engine/Plugins/Experimental/PCG/Source/PCGEditor/Private/SPCGEditorGraphDebugObjectTree.cpp
+++ b/Engine/Plugins/Experimental/PCG/Source/PCGEditor/Private/SPCGEditorGraphDebugObjectTree.cpp
@@ -236,7 +236,7 @@ void SPCGEditorGraphDebugObjectTree::SelectedDebugObject_OnClicked() const
 	if (UPCGComponent* PCGComponent = PCGEditor.Pin()->GetPCGComponentBeingInspected())
 	{
 		AActor* Actor = PCGComponent->GetOwner();
-		if (Actor && GEditor && GUnrealEd)
+		if (Actor && GEditor && GUnrealEd && GEditor->CanSelectActor(Actor, /*bInSelected=*/true))
 		{
 			GEditor->SelectNone(/*bNoteSelectionChange=*/false, /*bDeselectBSPSurfs=*/true, /*WarnAboutManyActors=*/false);
 			GEditor->SelectActor(Actor, /*bInSelected=*/true, /*bNotify=*/true, /*bSelectEvenIfHidden=*/true);
@@ -382,7 +382,7 @@ void SPCGEditorGraphDebugObjectTree::RefreshTree()
 	
 	for (UObject* PCGComponentObject : PCGComponents)
 	{
-		if (!IsValid(PCGComponentObject))
+		if (!IsValid(PCGComponentObject) || PCGComponentObject->HasAnyFlags(RF_Transient))
 		{
 			continue;
 		}
