[PCG] Refactor collapse to subgraph

* Move all collapse code into a helper on the plugin side
* Fix the collapse to include correctly graph parameters
* Add some unit tests
* Introduced a "ReconstructGraph" if there are changes on the model, but not on the view.

#rnx

diff --git a/Engine/Plugins/Experimental/PCG/Source/PCG/Private/Helpers/PCGSubgraphHelpers.cpp b/Engine/Plugins/Experimental/PCG/Source/PCG/Private/Helpers/PCGSubgraphHelpers.cpp
new file mode 100644
--- /dev/null
+++ b/Engine/Plugins/Experimental/PCG/Source/PCG/Private/Helpers/PCGSubgraphHelpers.cpp
@@ -0,0 +1,524 @@
+// Copyright Epic Games, Inc. All Rights Reserved.
+
+#include "Helpers/PCGSubgraphHelpers.h"
+
+#include "PCGEdge.h"
+#include "PCGGraph.h"
+#include "PCGInputOutputSettings.h"
+#include "PCGModule.h"
+#include "PCGNode.h"
+#include "PCGPin.h"
+#include "PCGSettings.h"
+#include "PCGSubgraph.h"
+#include "Elements/PCGUserParameterGet.h"
+
+#include "Algo/AnyOf.h"
+#include "Algo/Transform.h"
+
+#define LOCTEXT_NAMESPACE "PCGSubgraphHelpers"
+
+namespace PCGSubgraphHelpersExtra
+{
+	// Function that will create a new pin in the input/output node of the subgraph, with a unique name that will
+	// be the related to the pin passed as argument. It will be formatted like this:
+	// "{NodeName} {PinName} {OptionalIndex}"
+	FName CreateNewCustomPin(UPCGGraph* InGraph, UPCGPin* InPinToClone, bool bIsInput, TMap<FString, int>& InOutNameCollisionMapping)
+	{
+		UPCGNode* NewInputOutputNode = bIsInput ? InGraph->GetInputNode() : InGraph->GetOutputNode();
+		UPCGGraphInputOutputSettings* NewInputOutputSettings = CastChecked<UPCGGraphInputOutputSettings>(NewInputOutputNode->GetSettings());
+		FString NewName = InPinToClone->Node->GetNodeTitle().ToString() + " " + InPinToClone->Properties.Label.ToString();
+		if (InOutNameCollisionMapping.Contains(NewName))
+		{
+			if (!bIsInput)
+			{
+				// Output pin can be re-used. If there is a clash, just return the new name, the pin was already added.
+				return FName(NewName);
+			}
+
+			NewName += " " + FString::FormatAsNumber(++InOutNameCollisionMapping[NewName]);
+		}
+		else
+		{
+			InOutNameCollisionMapping.Emplace(NewName, 1);
+		}
+
+		FPCGPinProperties NewProperties = InPinToClone->Properties;
+
+		// For the pin type, narrow it down to the input edges (if it is an input pin)
+		if (bIsInput)
+		{
+			const EPCGDataType FirstPinTypeUnion = InPinToClone->Node->GetSettings()->GetTypeUnionOfIncidentEdges(InPinToClone->Properties.Label);
+			if (FirstPinTypeUnion != EPCGDataType::None)
+			{
+				NewProperties.AllowedTypes = FirstPinTypeUnion;
+			}
+		}
+
+		NewProperties.Label = FName(NewName);
+		NewProperties = NewInputOutputSettings->AddCustomPin(NewProperties);
+		NewInputOutputNode->UpdateAfterSettingsChangeDuringCreation();
+		return NewProperties.Label;
+	}
+
+	struct FCollapsingInformation
+	{
+#if WITH_EDITOR
+		// Keep track of all node positions to know where to spawn our subgraph node (editor only)
+		FVector2D AveragePosition = FVector2D::ZeroVector;
+		int32 MinX = std::numeric_limits<int32>::max();
+		int32 MaxX = std::numeric_limits<int32>::min();
+#endif // WITH_EDITOR
+
+		// Those 3 sets are used to keep track of the subgraph pins
+		// and will be used to extract pins that will be outside the subgraph
+		// and will need special treatment.
+		TSet<UPCGPin*> InputFromSubgraphPins;
+		TSet<UPCGPin*> OutputToSubgraphPins;
+		TSet<UPCGPin*> AllSubgraphPins;
+
+		// Also keep track of all the edges in the subgraph. Use a set to be able to add the
+		// same edge twice without duplicates.
+		TSet<UPCGEdge*> PCGGraphEdges;
+
+		// Keep track of all graph parameters used, they will need to be forwarded to the new graph.
+		TArray<const FPropertyBagPropertyDesc*> GraphParametersUsed;
+
+		TArray<UPCGNode*> ValidNodesToCollapse;
+
+		// Keep track of all user parameters output pins
+		TMap<FGuid, UPCGPin*> GetUserParametersOutputPins;
+		// Also keep track of duplicated get user parameters, to only keep a single one for each user parameter.
+		TArray<UPCGNode*> SuperfluousNodes;
+	};
+
+	FCollapsingInformation GatheringCollapseInformation(const TArray<UPCGNode*>& InNodesToCollapse, const FInstancedPropertyBag& InGraphParameters)
+	{
+		FCollapsingInformation CollapseInfo{};
+
+		for (UPCGNode* PCGNode : InNodesToCollapse)
+		{
+			check(PCGNode);
+
+			const UPCGSettings* NodeSettings = PCGNode->GetSettings();
+
+			// Exclude input and output nodes from the subgraph.
+			if (!NodeSettings || NodeSettings->IsA<UPCGGraphInputOutputSettings>())
+			{
+				continue;
+			}
+
+			if (const UPCGUserParameterGetSettings* GetSettings = Cast<UPCGUserParameterGetSettings>(NodeSettings))
+			{
+				if (const FPropertyBagPropertyDesc* PropertyDesc = InGraphParameters.FindPropertyDescByID(GetSettings->PropertyGuid))
+				{
+					CollapseInfo.GraphParametersUsed.AddUnique(PropertyDesc);
+
+					if (CollapseInfo.GetUserParametersOutputPins.Contains(PropertyDesc->ID))
+					{
+						CollapseInfo.SuperfluousNodes.Add(PCGNode);
+					}
+					else
+					{
+						CollapseInfo.GetUserParametersOutputPins.Emplace(PropertyDesc->ID, PCGNode->GetOutputPins()[0]);
+					}
+				}
+				else
+				{
+					// Invalid node
+					continue;
+				}
+			}
+
+			// For each pin of the subgraph, gather all its edges in a set and store the other pin of
+			// the edge as an input from or output to the subgraph. It will be useful for tracking pins
+			// that are outside the subgraph.
+			for (UPCGPin* InputPin : PCGNode->GetInputPins())
+			{
+				check(InputPin);
+
+				for (UPCGEdge* Edge : InputPin->Edges)
+				{
+					check(Edge);
+					CollapseInfo.PCGGraphEdges.Add(Edge);
+					CollapseInfo.OutputToSubgraphPins.Add(Edge->InputPin);
+				}
+
+				CollapseInfo.AllSubgraphPins.Add(InputPin);
+			}
+
+			for (UPCGPin* OutputPin : PCGNode->GetOutputPins())
+			{
+				check(OutputPin);
+
+				for (UPCGEdge* Edge : OutputPin->Edges)
+				{
+					check(Edge);
+					CollapseInfo.PCGGraphEdges.Add(Edge);
+					CollapseInfo.InputFromSubgraphPins.Add(Edge->OutputPin);
+				}
+
+				CollapseInfo.AllSubgraphPins.Add(OutputPin);
+			}
+
+			CollapseInfo.ValidNodesToCollapse.Add(PCGNode);
+
+#if WITH_EDITOR
+			// And do all the computation to get the min, max and mean position.
+			CollapseInfo.AveragePosition.X += PCGNode->PositionX;
+			CollapseInfo.AveragePosition.Y += PCGNode->PositionY;
+			CollapseInfo.MinX = FMath::Min(CollapseInfo.MinX, PCGNode->PositionX);
+			CollapseInfo.MaxX = FMath::Max(CollapseInfo.MaxX, PCGNode->PositionY);
+#endif // WITH_EDITOR
+		}
+
+#if WITH_EDITOR
+		if (CollapseInfo.ValidNodesToCollapse.Num() > 1)
+		{
+			// Compute the average position
+			CollapseInfo.AveragePosition /= CollapseInfo.ValidNodesToCollapse.Num();
+		}
+#endif // WITH_EDITOR
+
+		// Gather the pins that are outside the subgraph
+		CollapseInfo.InputFromSubgraphPins = CollapseInfo.InputFromSubgraphPins.Difference(CollapseInfo.AllSubgraphPins);
+		CollapseInfo.OutputToSubgraphPins = CollapseInfo.OutputToSubgraphPins.Difference(CollapseInfo.AllSubgraphPins);
+
+		return CollapseInfo;
+	}
+}
+
+UPCGGraph* FPCGSubgraphHelpers::CollapseIntoSubgraph(UPCGGraph* InOriginalGraph, const TArray<UPCGNode*>& InNodesToCollapse, const TArray<UObject*>& InExtraNodesToCollapse, UPCGGraph* OptionalPreAllocatedGraph)
+{
+	// Don't collapse into a subgraph if you don't have at least 2 nodes
+	if (!InOriginalGraph || InNodesToCollapse.Num() < 2)
+	{
+		return nullptr;
+	}
+
+	const FInstancedPropertyBag* GraphParameters = InOriginalGraph->GetUserParametersStruct();
+	check(GraphParameters);
+
+	PCGSubgraphHelpersExtra::FCollapsingInformation CollapseInfo = PCGSubgraphHelpersExtra::GatheringCollapseInformation(InNodesToCollapse, *GraphParameters);
+
+	// If we have at most 1 valid node to collapse, just exit
+	if (CollapseInfo.ValidNodesToCollapse.Num() <= 1)
+	{
+		UE_LOG(LogPCG, Warning, TEXT("There were less than 2 PCG nodes selected, abort"));
+		return nullptr;
+	}
+
+	// 2. Create a new subgraph if necessary.
+	UPCGGraph* NewPCGGraph = OptionalPreAllocatedGraph;
+	if (!NewPCGGraph)
+	{
+		NewPCGGraph = NewObject<UPCGGraph>();
+	}
+
+#if WITH_EDITOR
+	// Do some clean-up on input/output nodes
+	constexpr int32 Padding = 200;
+	NewPCGGraph->GetInputNode()->PositionX = CollapseInfo.MinX - Padding;
+	NewPCGGraph->GetInputNode()->PositionY = CollapseInfo.AveragePosition.Y;
+	NewPCGGraph->GetOutputNode()->PositionX = CollapseInfo.MaxX + Padding;
+	NewPCGGraph->GetOutputNode()->PositionY = CollapseInfo.AveragePosition.Y;
+#endif // WITH_EDITOR
+
+	// 4. Duplicate all the nodes, and keep a mapping between the old pins and new pins
+	TMap<UPCGPin*, UPCGPin*> PinMapping;
+	for (const UPCGNode* PCGNode : CollapseInfo.ValidNodesToCollapse)
+	{
+		check(PCGNode);
+
+		// Reconstruct a new node, same as PCGNode, but without any edges in the new graph 
+		UPCGNode* NewNode = NewPCGGraph->ReconstructNewNode(PCGNode);
+
+		// Safeguard: We should have a 1 for 1 matching between pins labels between the original node
+		// and the copied node. If for some reason we don't (perhaps the node was not updated correctly after pins were added/removed)
+		// we will log an error and try to connect as best as we can (probably breaking some edges on the process).
+
+		auto Mapping = [&PinMapping, PCGNode](const TArray<UPCGPin*>& OriginalPins, const TArray<UPCGPin*>& NewPins)
+		{
+			TSet<FName> UnmatchedOriginal;
+			TMap<FName, UPCGPin*> NewMapping;
+
+			for (UPCGPin* NewPin : NewPins)
+			{
+				NewMapping.Emplace(NewPin->Properties.Label, NewPin);
+			}
+
+			for (UPCGPin* OriginalPin : OriginalPins)
+			{
+				FName PinLabel = OriginalPin->Properties.Label;
+				if (UPCGPin** NewPinPtr = NewMapping.Find(PinLabel))
+				{
+					PinMapping.Emplace(OriginalPin, *NewPinPtr);
+				}
+				else if (OriginalPin->IsConnected())
+				{
+					// It is only problematic if the pin was connected
+					UE_LOG(LogPCG, Error, TEXT("[CollapseInSubgraph - %s] %s pin %s does not exist anymore. Edges will be broken."),
+						*PCGNode->GetNodeTitle().ToString(), (OriginalPin->IsOutputPin() ? TEXT("Output") : TEXT("Input")), *PinLabel.ToString());
+				}
+			}
+		};
+
+		Mapping(PCGNode->GetInputPins(), NewNode->GetInputPins());
+		Mapping(PCGNode->GetOutputPins(), NewNode->GetOutputPins());
+	}
+
+#if WITH_EDITOR
+	// Also duplicate the extra nodes and assign them to the new graph
+	TArray<TObjectPtr<const UObject>> NewExtraGraphNodes;
+	for (const UObject* ExtraNode : InExtraNodesToCollapse)
+	{
+		NewExtraGraphNodes.Add(DuplicateObject(ExtraNode, NewPCGGraph));
+	}
+
+	NewPCGGraph->SetExtraEditorNodes(NewExtraGraphNodes);
+#endif ///WITH_EDITOR
+
+	// 5. Iterate over all the edges and create edges "placeholders"
+	// Most of them will already be complete, but for those that needs to be connected to the new
+	// subgraph nodes, the pins don't exist yet. Therefore we identify them with their pin labels.
+	//
+	// The logic behind the new pins is this:
+	// -> If the pin is connected to the simple pin of the input/output node, let it like this
+	// -> If the pin is part of the original input/output advanced pins, we will trigger the advanced pins flag on the input/output node in the subgraph
+	// -> If the pin is connected to a get user parameters that is outside of the subgraph, we'll create a new node and connect it to it.
+	// -> Otherwise, we add a new custom pin, with the name of the node, the name of the pin and a number if there is name collision
+	struct EdgePlaceholder
+	{
+		UPCGPin* InputPin = nullptr;
+		UPCGPin* OutputPin = nullptr;
+		FName InputPinLabel;
+		FName OutputPinLabel;
+	};
+
+	TArray<EdgePlaceholder> EdgePlaceholders;
+	TMap<FString, int> NameCollisionMapping;
+
+	// For all internal parameters, create extra edges to connect those nodes to their override pin counterparts
+	for (const TPair<FGuid, UPCGPin*>& It : CollapseInfo.GetUserParametersOutputPins)
+	{
+		check(It.Value);
+
+		EdgePlaceholder OutsideSubgraphEdge;
+		OutsideSubgraphEdge.InputPin = It.Value;
+		// Gymnastics to make sure we have a 1 to 1 match between labels, because subgraph override uses the display text name.
+		OutsideSubgraphEdge.OutputPinLabel = FName(FName::NameToDisplayString(OutsideSubgraphEdge.InputPin->Properties.Label.ToString(), /*bIsBool=*/ false));
+
+		EdgePlaceholders.Add(std::move(OutsideSubgraphEdge));
+	}
+
+	for (UPCGEdge* Edge : CollapseInfo.PCGGraphEdges)
+	{
+		UPCGPin* const* InPin = PinMapping.Find(Edge->InputPin);
+		UPCGPin* const* OutPin = PinMapping.Find(Edge->OutputPin);
+
+		check(InPin || OutPin);
+
+		if (InPin == nullptr)
+		{
+			// The edge comes from outside the graph.
+			// If it is from the input node, we have a special behavior
+			// Same if it is from a get user parameter
+			EdgePlaceholder OutsideSubgraphEdge;
+			EdgePlaceholder InsideSubgraphEdge;
+			bool bProcessed = false;
+
+			OutsideSubgraphEdge.InputPin = Edge->InputPin;
+			InsideSubgraphEdge.OutputPin = *OutPin;
+
+			if (Edge->InputPin->Node == InOriginalGraph->GetInputNode())
+			{
+				const UPCGGraphInputOutputSettings* Settings = Cast<const UPCGGraphInputOutputSettings>(Edge->InputPin->Node->GetSettings());
+
+				if (Settings && !Settings->IsCustomPin(Edge->InputPin))
+				{
+					OutsideSubgraphEdge.OutputPinLabel = Edge->InputPin->Properties.Label;
+					InsideSubgraphEdge.InputPin = NewPCGGraph->GetInputNode()->GetOutputPin(Edge->InputPin->Properties.Label);
+					bProcessed = true;
+				}
+			}
+			else if (Edge->InputPin->Node && Edge->InputPin->Node->GetSettings() && Edge->InputPin->Node->GetSettings()->IsA<UPCGUserParameterGetSettings>())
+			{
+				const UPCGUserParameterGetSettings* OldSettings = CastChecked<const UPCGUserParameterGetSettings>(Edge->InputPin->Node->GetSettings());
+				if (CollapseInfo.GetUserParametersOutputPins.Contains(OldSettings->PropertyGuid))
+				{
+					if (CollapseInfo.GetUserParametersOutputPins[OldSettings->PropertyGuid] != Edge->InputPin->Node->GetOutputPins()[0])
+					{
+						CollapseInfo.SuperfluousNodes.Add(Edge->InputPin->Node);
+					}
+				}
+				else
+				{
+					CollapseInfo.GetUserParametersOutputPins.Emplace(OldSettings->PropertyGuid, Edge->InputPin->Node->GetOutputPins()[0]);
+				}
+
+				check(CollapseInfo.GetUserParametersOutputPins.Contains(OldSettings->PropertyGuid));
+
+				// The label will be the same as the input pin, since it will be an override pin.
+				OutsideSubgraphEdge.InputPin = CollapseInfo.GetUserParametersOutputPins[OldSettings->PropertyGuid];
+				// Gymnastics to make sure we have a 1 to 1 match between labels, because subgraph override uses the display text name.
+				OutsideSubgraphEdge.OutputPinLabel = FName(FName::NameToDisplayString(OutsideSubgraphEdge.InputPin->Properties.Label.ToString(), /*bIsBool=*/ false));
+
+				// Create a new node and connect it to it.
+				UPCGSettings* NewSettings = nullptr;
+				UPCGNode* NewNode = NewPCGGraph->AddNodeCopy(Edge->InputPin->Node->GetSettings(), NewSettings);
+				check(NewNode && NewSettings && NewNode->GetOutputPins().Num() == 1);
+
+				// Also need to add this parameter to the list.
+				if (const FPropertyBagPropertyDesc* PropertyDesc = GraphParameters->FindPropertyDescByID(OldSettings->PropertyGuid))
+				{
+					CollapseInfo.GraphParametersUsed.AddUnique(PropertyDesc);
+				}
+
+				InsideSubgraphEdge.InputPin = NewNode->GetOutputPins()[0];
+
+				bProcessed = true;
+#if WITH_EDITOR
+				// Place the newly created node slightly to the left of the output pin node.
+				if (ensure(Edge->OutputPin->Node))
+				{
+					constexpr int32 OffsetX = -200;
+					NewNode->PositionX = Edge->OutputPin->Node->PositionX + OffsetX;
+					NewNode->PositionY = Edge->OutputPin->Node->PositionY;
+				}
+#endif // WITH_EDITOR
+			}
+
+			if (!bProcessed)
+			{
+				FName NewPinName = PCGSubgraphHelpersExtra::CreateNewCustomPin(NewPCGGraph, Edge->OutputPin, true, NameCollisionMapping);
+				OutsideSubgraphEdge.OutputPinLabel = NewPinName;
+				InsideSubgraphEdge.InputPin = NewPCGGraph->GetInputNode()->GetOutputPin(NewPinName);
+			}
+
+			EdgePlaceholders.Add(std::move(OutsideSubgraphEdge));
+			EdgePlaceholders.Add(std::move(InsideSubgraphEdge));
+		}
+		else if (OutPin == nullptr)
+		{
+			// The edge comes from outside the graph.
+			// If it is from the output node, we have a special behavior
+			EdgePlaceholder OutsideSubgraphEdge;
+			EdgePlaceholder InsideSubgraphEdge;
+			bool bProcessed = false;
+
+			OutsideSubgraphEdge.OutputPin = Edge->OutputPin;
+			InsideSubgraphEdge.InputPin = *InPin;
+
+			if (Edge->OutputPin->Node == InOriginalGraph->GetOutputNode())
+			{
+				const UPCGGraphInputOutputSettings* Settings = Cast<const UPCGGraphInputOutputSettings>(Edge->OutputPin->Node->GetSettings());
+
+				if (Settings && !Settings->IsCustomPin(Edge->OutputPin))
+				{
+					OutsideSubgraphEdge.InputPinLabel = Edge->OutputPin->Properties.Label;
+					InsideSubgraphEdge.OutputPin = NewPCGGraph->GetOutputNode()->GetInputPin(Edge->OutputPin->Properties.Label);
+					bProcessed = true;
+				}
+			}
+
+			if (!bProcessed)
+			{
+				FName NewPinName = PCGSubgraphHelpersExtra::CreateNewCustomPin(NewPCGGraph, Edge->InputPin, false, NameCollisionMapping);
+				OutsideSubgraphEdge.InputPinLabel = NewPinName;
+				InsideSubgraphEdge.OutputPin = NewPCGGraph->GetOutputNode()->GetInputPin(NewPinName);
+			}
+
+			EdgePlaceholders.Add(std::move(OutsideSubgraphEdge));
+			EdgePlaceholders.Add(std::move(InsideSubgraphEdge));
+		}
+		else
+		{
+			// Both nodes are inside
+			EdgePlaceholders.Add(EdgePlaceholder{ *InPin, *OutPin });
+		}
+	}
+
+	// Transfer the values from the original graph parameters to the new graph
+	TArray<FPropertyBagPropertyDesc> GraphParametersUsed;
+	Algo::Transform(CollapseInfo.GraphParametersUsed, GraphParametersUsed, [](const FPropertyBagPropertyDesc* In) { check(In); return *In; });
+	NewPCGGraph->AddUserParameters(GraphParametersUsed, InOriginalGraph);
+
+	// 6. Create subgraph and delete old nodes and superfluous nodes
+	UPCGNode* SubgraphNode = nullptr;
+	{
+		for (UPCGNode* PCGNode : CollapseInfo.ValidNodesToCollapse)
+		{
+			const UPCGSettings* NodeSettings = PCGNode->GetSettings();
+
+			// Don't delete input, output and get user parameters nodes from the subgraph.
+			if (!NodeSettings || NodeSettings->IsA<UPCGGraphInputOutputSettings>() || NodeSettings->IsA<UPCGUserParameterGetSettings>())
+			{
+				continue;
+			}
+
+			InOriginalGraph->RemoveNode(PCGNode);
+		}
+
+#if WITH_EDITOR
+		for (const UObject* ExtraNode : InExtraNodesToCollapse)
+		{
+			InOriginalGraph->RemoveExtraEditorNode(ExtraNode);
+		}
+#endif // WITH_EDITOR
+
+		for (UPCGNode* SuperfluousNode : CollapseInfo.SuperfluousNodes)
+		{
+			check(SuperfluousNode);
+			// Only remove them if they have no edges anymore
+			bool bHasAnyEdges = Algo::AnyOf(SuperfluousNode->GetInputPins(), [](const UPCGPin* InPin) { return InPin && InPin->EdgeCount() > 0; })
+				|| Algo::AnyOf(SuperfluousNode->GetOutputPins(), [](const UPCGPin* InPin) { return InPin && InPin->EdgeCount() > 0; });
+
+			if (!bHasAnyEdges)
+			{
+				InOriginalGraph->RemoveNode(SuperfluousNode);
+			}
+		}
+
+		UPCGSettings* DefaultNodeSettings = nullptr;
+		SubgraphNode = InOriginalGraph->AddNodeOfType(UPCGSubgraphSettings::StaticClass(), DefaultNodeSettings);
+		UPCGSubgraphSettings* DefaultSubgraphSettings = CastChecked<UPCGSubgraphSettings>(DefaultNodeSettings);
+		DefaultSubgraphSettings->SetSubgraph(NewPCGGraph);
+
+#if WITH_EDITOR
+		SubgraphNode->PositionX = CollapseInfo.AveragePosition.X;
+		SubgraphNode->PositionY = CollapseInfo.AveragePosition.Y;
+#endif // WITH_EDITOR
+
+		SubgraphNode->UpdateAfterSettingsChangeDuringCreation();
+	}
+
+	// 8. Connect all the edges
+	TSet<UPCGNode*> TouchedNodes;
+	for (EdgePlaceholder& Edge : EdgePlaceholders)
+	{
+		if (Edge.InputPin == nullptr)
+		{
+			SubgraphNode->GetOutputPin(Edge.InputPinLabel)->AddEdgeTo(Edge.OutputPin, &TouchedNodes);
+		}
+		else if (Edge.OutputPin == nullptr)
+		{
+			Edge.InputPin->AddEdgeTo(SubgraphNode->GetInputPin(Edge.OutputPinLabel), &TouchedNodes);
+		}
+		else
+		{
+			Edge.InputPin->AddEdgeTo(Edge.OutputPin, &TouchedNodes);
+		}
+	}
+
+#if WITH_EDITOR
+	InOriginalGraph->NotifyGraphChanged(EPCGChangeType::Structural);
+	NewPCGGraph->NotifyGraphChanged(EPCGChangeType::Structural);
+	for (UPCGNode* TouchedNode : TouchedNodes)
+	{
+		TouchedNode->OnNodeChangedDelegate.Broadcast(TouchedNode, EPCGChangeType::Node);
+	}
+#endif // WITH_EDITOR
+
+	return NewPCGGraph;
+}
+
+#undef LOCTEXT_NAMESPACE
\ No newline at end of file
diff --git a/Engine/Plugins/Experimental/PCG/Source/PCG/Private/PCGGraph.cpp b/Engine/Plugins/Experimental/PCG/Source/PCG/Private/PCGGraph.cpp
--- a/Engine/Plugins/Experimental/PCG/Source/PCG/Private/PCGGraph.cpp
+++ b/Engine/Plugins/Experimental/PCG/Source/PCG/Private/PCGGraph.cpp
@@ -792,6 +792,11 @@ void UPCGGraph::SetExtraEditorNodes(const TArray<TObjectPtr<const UObject>>& InN
 	}
 }
 
+void UPCGGraph::RemoveExtraEditorNode(const UObject* InNode)
+{
+	ExtraEditorNodes.Remove(const_cast<UObject*>(InNode));
+}
+
 FPCGActorSelectionKeyToSettingsMap UPCGGraph::GetTrackedActorKeysToSettings() const
 {
 	FPCGActorSelectionKeyToSettingsMap TagsToSettings;
@@ -1143,6 +1148,22 @@ uint32 UPCGGraph::CalculateNodeGridSizeRecursive_Unsafe(const UPCGNode* InNode,
 	return GridSize;
 }
 
+void UPCGGraph::AddUserParameters(const TArray<FPropertyBagPropertyDesc>& InDescs, const UPCGGraph* InOptionalOriginalGraph)
+{
+	UserParameters.AddProperties(InDescs);
+	if (InOptionalOriginalGraph)
+	{
+		if (const FInstancedPropertyBag* OriginalPropertyBag = InOptionalOriginalGraph->GetUserParametersStruct())
+		{
+			UserParameters.CopyMatchingValuesByID(*OriginalPropertyBag);
+		}
+	}
+
+#if WITH_EDITOR
+	OnGraphParametersChanged(EPCGGraphParameterEvent::MultiplePropertiesAdded, NAME_None);
+#endif // WITH_EDITOR
+}
+
 /****************************
 * UPCGGraphInstance
 ****************************/
@@ -1568,7 +1589,10 @@ bool FPCGOverrideInstancedPropertyBag::RefreshParameters(const FInstancedPropert
 		}
 		break;
 	}
-	case EPCGGraphParameterEvent::GraphPostLoad:
+	// Do the same thing in case of post load and multiple properties added.
+	// We have 2 enums to avoid puzzling someone that wonders why we would call GraphPostLoad when we add multiple properties.
+	case EPCGGraphParameterEvent::GraphPostLoad: // fall-through
+	case EPCGGraphParameterEvent::MultiplePropertiesAdded:
 	{
 		// Check if the property struct mismatch. If so, do the migration
 		if (Parameters.GetPropertyBagStruct() != ParentUserParameters->GetPropertyBagStruct())
diff --git a/Engine/Plugins/Experimental/PCG/Source/PCG/Private/PCGSubgraph.cpp b/Engine/Plugins/Experimental/PCG/Source/PCG/Private/PCGSubgraph.cpp
--- a/Engine/Plugins/Experimental/PCG/Source/PCG/Private/PCGSubgraph.cpp
+++ b/Engine/Plugins/Experimental/PCG/Source/PCG/Private/PCGSubgraph.cpp
@@ -83,6 +83,9 @@ void UPCGBaseSubgraphSettings::SetSubgraph(UPCGGraphInterface* InGraph)
 		Subgraph->OnGraphChangedDelegate.AddUObject(this, &UPCGBaseSubgraphSettings::OnSubgraphChanged);
 	}
 #endif // WITH_EDITOR
+
+	// Also, reconstruct overrides
+	InitializeCachedOverridableParams(/*bReset=*/true);
 }
 
 void UPCGBaseSubgraphSettings::BeginDestroy()
diff --git a/Engine/Plugins/Experimental/PCG/Source/PCG/Private/Tests/Subgraphs/PCGCollapseSubgraphTest.cpp b/Engine/Plugins/Experimental/PCG/Source/PCG/Private/Tests/Subgraphs/PCGCollapseSubgraphTest.cpp
new file mode 100644
--- /dev/null
+++ b/Engine/Plugins/Experimental/PCG/Source/PCG/Private/Tests/Subgraphs/PCGCollapseSubgraphTest.cpp
@@ -0,0 +1,225 @@
+// Copyright Epic Games, Inc. All Rights Reserved.
+
+#include "PCGContext.h"
+#include "PCGEdge.h"
+#include "PCGGraph.h"
+#include "PCGSubgraph.h"
+#include "PCGInputOutputSettings.h"
+#include "Elements/PCGAttributeNoise.h"
+#include "Elements/PCGSurfaceSampler.h"
+#include "Elements/PCGUserParameterGet.h"
+#include "Helpers/PCGSubgraphHelpers.h"
+#include "Tests/PCGTestsCommon.h"
+
+IMPLEMENT_CUSTOM_SIMPLE_AUTOMATION_TEST(FPCGCollapseSubgraphSimple, FPCGTestBaseClass, "pcg.tests.Subgraph.Collapse.Simple", PCGTestsCommon::TestFlags);
+IMPLEMENT_CUSTOM_SIMPLE_AUTOMATION_TEST(FPCGCollapseSubgraphWithParams, FPCGTestBaseClass, "pcg.tests.Subgraph.Collapse.WithParams", PCGTestsCommon::TestFlags);
+
+namespace PCGCollapseSubgraphTests
+{
+	/**
+	* Utility function to test if an edge exists between two nodes of class InputClass and OutputClass.
+	* You need to set at least InLabel or OutLabel to be different to NAME_None, if you only care if a pin is connected to a node and not a specific pin.
+	* Only works if there is a single node of type InputClass or OutputClass (exception made for input and output node).
+	*/
+	template <typename InputClass, typename OutputClass>
+	bool ValidateNodeAndEdge(FPCGTestBaseClass* InTestClass, UPCGGraph* InGraph, FName InLabel, FName OutLabel)
+	{
+		check(InGraph);
+
+		if (!InTestClass->TestTrue("Test is not illed formed", (InLabel != NAME_None) || (OutLabel != NAME_None))) { return false; }
+
+		UPCGNode* InNode = nullptr;
+		UPCGNode* OutNode = nullptr;
+
+		if constexpr (std::is_same_v<InputClass, UPCGGraphInputOutputSettings>)
+		{
+			InNode = InGraph->GetInputNode();
+		}
+		else
+		{
+			UPCGNode* const* NodePtr = InGraph->GetNodes().FindByPredicate([](const UPCGNode* Node) { return Node && Node->GetSettings() && Node->GetSettings()->IsA<InputClass>(); });
+			InNode = NodePtr ? *NodePtr : nullptr;
+		}
+
+		if constexpr (std::is_same_v<OutputClass, UPCGGraphInputOutputSettings>)
+		{
+			OutNode = InGraph->GetOutputNode();
+		}
+		else
+		{
+			UPCGNode* const* NodePtr = InGraph->GetNodes().FindByPredicate([](const UPCGNode* Node) { return Node && Node->GetSettings() && Node->GetSettings()->IsA<OutputClass>(); });
+			OutNode = NodePtr ? *NodePtr : nullptr;
+		}
+
+		if (!InTestClass->TestNotNull(*FString::Printf(TEXT("Input Node of class %s is valid"), *InputClass::StaticClass()->GetName()), InNode)) { return false; }
+		if (!InTestClass->TestNotNull(*FString::Printf(TEXT("Output Node of class %s is valid"), *OutputClass::StaticClass()->GetName()), OutNode)) { return false; }
+
+		UPCGPin* InPin = InNode->GetOutputPin(InLabel);
+		UPCGPin* OutPin = OutNode->GetInputPin(OutLabel);
+
+		if (!InTestClass->TestTrue(*FString::Printf(TEXT("InLabel (%s) exists on node %s"), *InLabel.ToString(), *InNode->GetNodeTitle().ToString()), (InLabel == NAME_None) || InPin)) { return false; }
+		if (!InTestClass->TestTrue(*FString::Printf(TEXT("OutLabel (%s) exists on node %s"), *OutLabel.ToString(), *OutNode->GetNodeTitle().ToString()), (OutLabel == NAME_None) || OutPin)) { return false; }
+
+		const TArray<TObjectPtr<UPCGEdge>>& Edges = InPin ? InPin->Edges : OutPin->Edges;
+
+		auto Predicate = [InPin, OutPin](const TObjectPtr<UPCGEdge>& InEdge)
+		{
+			if (InPin && OutPin)
+			{
+				return (InEdge->InputPin == InPin) && (InEdge->OutputPin == OutPin);
+			}
+			else if (OutPin)
+			{
+				return (InEdge->OutputPin == OutPin) && InEdge->InputPin && InEdge->InputPin->Node && InEdge->InputPin->Node->GetSettings() && InEdge->InputPin->Node->GetSettings()->IsA<InputClass>();
+			}
+			else
+			{
+				return (InEdge->InputPin == InPin) && InEdge->OutputPin && InEdge->OutputPin->Node && InEdge->OutputPin->Node->GetSettings() && InEdge->OutputPin->Node->GetSettings()->IsA<OutputClass>();
+			}
+		};
+
+		return InTestClass->TestTrue(*FString::Printf(TEXT("Edge between %s(%s) and %s(%s) exists"), *InNode->GetNodeTitle().ToString(), *InLabel.ToString(), *OutNode->GetNodeTitle().ToString(), *OutLabel.ToString()), Edges.ContainsByPredicate(Predicate));
+	}
+
+	UPCGNode* AddGetGraphParameterNode(UPCGGraph* InGraph, const FName InParamName)
+	{
+		const FPropertyBagPropertyDesc* Desc = InGraph->GetUserParametersStruct()->FindPropertyDescByName(InParamName);
+		if (!Desc)
+		{
+			return nullptr;
+		}
+
+		UPCGUserParameterGetSettings* Settings = nullptr;
+		UPCGNode* GetParameterNode = InGraph->AddNodeOfType<UPCGUserParameterGetSettings>(Settings);
+
+		Settings->PropertyName = InParamName;
+		Settings->PropertyGuid = Desc->ID;
+
+		GetParameterNode->UpdateAfterSettingsChangeDuringCreation();
+
+		return GetParameterNode;
+	}
+}
+
+#define VALIDATE_EDGE(InSettings, OutSettings, Graph, InLabel, OutLabel) if (!PCGCollapseSubgraphTests::ValidateNodeAndEdge<InSettings, OutSettings>(this, Graph, InLabel, OutLabel)) { return false; }
+
+bool FPCGCollapseSubgraphSimple::RunTest(const FString& Parameters)
+{
+	UPCGGraph* MainGraph = NewObject<UPCGGraph>();
+	MainGraph->SetFlags(RF_Transient);
+
+	// Add a few nodes:
+	// * Surface sampler connected to the landscape pin
+	// * Attribute noise on the output of the surface sampler, connected to the output
+
+	UPCGNode* InputNode = MainGraph->GetInputNode();
+	UPCGNode* OutputNode = MainGraph->GetOutputNode();
+
+	UTEST_NOT_NULL("InputNode is valid", InputNode);
+	UTEST_NOT_NULL("OutputNode is valid", OutputNode);
+
+	UPCGSurfaceSamplerSettings* SurfaceSamplerSettings = nullptr;
+	UPCGAttributeNoiseSettings* AttributeNoiseSettings = nullptr;
+	UPCGNode* SurfaceSamplerNode = MainGraph->AddNodeOfType<UPCGSurfaceSamplerSettings>(SurfaceSamplerSettings);
+	UPCGNode* AttributeNoiseNode = MainGraph->AddNodeOfType<UPCGAttributeNoiseSettings>(AttributeNoiseSettings);
+
+	UTEST_NOT_NULL("SurfaceSampler node is valid", SurfaceSamplerNode);
+	UTEST_NOT_NULL("AttributeNoise node is valid", AttributeNoiseNode);
+	UTEST_EQUAL("There are 2 nodes", MainGraph->GetNodes().Num(), 2);
+
+	MainGraph->AddLabeledEdge(InputNode, PCGInputOutputConstants::DefaultLandscapeLabel, SurfaceSamplerNode, TEXT("Surface"));
+	MainGraph->AddLabeledEdge(SurfaceSamplerNode, PCGPinConstants::DefaultOutputLabel, AttributeNoiseNode, PCGPinConstants::DefaultInputLabel);
+	MainGraph->AddLabeledEdge(AttributeNoiseNode, PCGPinConstants::DefaultOutputLabel, OutputNode, PCGPinConstants::DefaultOutputLabel);
+
+	VALIDATE_EDGE(UPCGGraphInputOutputSettings, UPCGSurfaceSamplerSettings, MainGraph, PCGInputOutputConstants::DefaultLandscapeLabel, TEXT("Surface"))
+	VALIDATE_EDGE(UPCGSurfaceSamplerSettings, UPCGAttributeNoiseSettings, MainGraph, PCGPinConstants::DefaultOutputLabel, PCGPinConstants::DefaultInputLabel)
+	VALIDATE_EDGE(UPCGAttributeNoiseSettings, UPCGGraphInputOutputSettings, MainGraph, PCGPinConstants::DefaultOutputLabel, PCGPinConstants::DefaultOutputLabel)
+
+	// Try to collapse all the nodes. Input/Output should stay, the other not.
+	TArray<UPCGNode*> NodesToCollapse = { InputNode, OutputNode, SurfaceSamplerNode, AttributeNoiseNode };
+	UPCGGraph* Subgraph = FPCGSubgraphHelpers::CollapseIntoSubgraph(MainGraph, NodesToCollapse, {});
+
+	UTEST_NOT_NULL("Subgraph collapse succeeded", Subgraph);
+	UTEST_EQUAL("There is 1 node in the main graph", MainGraph->GetNodes().Num(), 1);
+
+	UPCGNode* SubgraphNode = MainGraph->GetNodes()[0];
+	UTEST_NOT_NULL("SubgraphNode is valid", SubgraphNode);
+
+	UPCGSubgraphSettings* SubgraphSettings = Cast<UPCGSubgraphSettings>(SubgraphNode->GetSettings());
+	UTEST_NOT_NULL("Node is a subgraph node", SubgraphSettings);
+	UTEST_EQUAL("Subgraph graph is the right one", SubgraphSettings->SubgraphInstance->GetGraph(), Subgraph);
+
+	// Landscape pin from input is connected to landscape pin of subgraph node
+	VALIDATE_EDGE(UPCGGraphInputOutputSettings, UPCGSubgraphSettings, MainGraph, PCGInputOutputConstants::DefaultLandscapeLabel, PCGInputOutputConstants::DefaultLandscapeLabel)
+
+	// Output node is connected to the subgraph node
+	VALIDATE_EDGE(UPCGSubgraphSettings, UPCGGraphInputOutputSettings, MainGraph, NAME_None, PCGPinConstants::DefaultOutputLabel)
+
+	// Same verification for subgraph than the initial main graph, expect for output pin
+	UTEST_EQUAL("There is 2 nodes in the subgraph", Subgraph->GetNodes().Num(), 2);
+
+	VALIDATE_EDGE(UPCGGraphInputOutputSettings, UPCGSurfaceSamplerSettings, Subgraph, PCGInputOutputConstants::DefaultLandscapeLabel, TEXT("Surface"))
+	VALIDATE_EDGE(UPCGSurfaceSamplerSettings, UPCGAttributeNoiseSettings, Subgraph, PCGPinConstants::DefaultOutputLabel, PCGPinConstants::DefaultInputLabel)
+	VALIDATE_EDGE(UPCGAttributeNoiseSettings, UPCGGraphInputOutputSettings, Subgraph, PCGPinConstants::DefaultOutputLabel, NAME_None);
+
+	return true;
+}
+
+bool FPCGCollapseSubgraphWithParams::RunTest(const FString& Parameters)
+{
+	UPCGGraph* MainGraph = NewObject<UPCGGraph>();
+	MainGraph->SetFlags(RF_Transient);
+
+	const FName Param1Name = TEXT("FloatParam");
+	const FName Param2Name = TEXT("VectorParam");
+	const float Param1Value = 0.1f;
+	const FVector Param2Value = FVector(1.0, 2.0, 3.0);
+
+	TArray<FPropertyBagPropertyDesc> GraphParameters;
+	GraphParameters.Emplace(Param1Name, EPropertyBagPropertyType::Float);
+	GraphParameters.Emplace(Param2Name, EPropertyBagPropertyType::Struct, TBaseStructure<FVector>::Get());
+
+	MainGraph->AddUserParameters(GraphParameters);
+	UTEST_EQUAL("Graph has 2 parameters", MainGraph->GetUserParametersStruct()->GetNumPropertiesInBag(), 2);
+
+	// Set the values
+	MainGraph->GetMutableUserParametersStruct_Unsafe()->SetValueFloat(Param1Name, Param1Value);
+	MainGraph->GetMutableUserParametersStruct_Unsafe()->SetValueStruct(Param2Name, Param2Value);
+
+	UPCGSurfaceSamplerSettings* SurfaceSamplerSettings = nullptr;
+	UPCGNode* SurfaceSamplerNode = MainGraph->AddNodeOfType<UPCGSurfaceSamplerSettings>(SurfaceSamplerSettings);
+
+	UPCGNode* GetParameter1Node = PCGCollapseSubgraphTests::AddGetGraphParameterNode(MainGraph, Param1Name);
+	UPCGNode* GetParameter2Node = PCGCollapseSubgraphTests::AddGetGraphParameterNode(MainGraph, Param2Name);
+
+	MainGraph->AddLabeledEdge(GetParameter1Node, Param1Name, SurfaceSamplerNode, TEXT("Looseness"));
+	MainGraph->AddLabeledEdge(GetParameter2Node, Param2Name, SurfaceSamplerNode, TEXT("Point Extents"));
+
+	TArray<UPCGNode*> NodesToCollapse = { SurfaceSamplerNode, GetParameter1Node, GetParameter2Node };
+
+	UPCGGraph* Subgraph = FPCGSubgraphHelpers::CollapseIntoSubgraph(MainGraph, NodesToCollapse, {});
+
+	UTEST_NOT_NULL("Subgraph collapse succeeded", Subgraph);
+	UTEST_EQUAL("There are 3 node in the main graph", MainGraph->GetNodes().Num(), 3);
+
+	// Get parameters nodes are still there and plugged
+	UTEST_TRUE("GetParam1 node is still in the main graph", MainGraph->GetNodes().Contains(GetParameter1Node));
+	UTEST_TRUE("GetParam2 node is still in the main graph", MainGraph->GetNodes().Contains(GetParameter2Node));
+
+	UTEST_TRUE("Edge between param 1 and subgraph exists", GetParameter1Node->GetOutputPin(Param1Name)->Edges.ContainsByPredicate([GetParameter1Node](const UPCGEdge* Edge) { return Edge && Edge->InputPin && Edge->OutputPin && (Edge->InputPin->Node == GetParameter1Node) && Edge->OutputPin->Node && Edge->OutputPin->Node->GetSettings() && Edge->OutputPin->Node->GetSettings()->IsA<UPCGSubgraphSettings>(); }));
+	UTEST_TRUE("Edge between param 2 and subgraph exists", GetParameter2Node->GetOutputPin(Param2Name)->Edges.ContainsByPredicate([GetParameter2Node](const UPCGEdge* Edge) { return Edge && Edge->InputPin && Edge->OutputPin && (Edge->InputPin->Node == GetParameter2Node) && Edge->OutputPin->Node && Edge->OutputPin->Node->GetSettings() && Edge->OutputPin->Node->GetSettings()->IsA<UPCGSubgraphSettings>(); }));
+
+	// Also verify that the subgraph has the 2 parameters, with the right value, and their getters are connected to the surface sampler.
+	UTEST_NOT_NULL("Subgraph has Param1 as parameter", Subgraph->GetUserParametersStruct()->FindPropertyDescByName(Param1Name));
+	UTEST_NOT_NULL("Subgraph has Param2 as parameter", Subgraph->GetUserParametersStruct()->FindPropertyDescByName(Param2Name));
+
+	UTEST_EQUAL("Param1 has the right value", Subgraph->GetUserParametersStruct()->GetValueFloat(Param1Name).GetValue(), Param1Value);
+	UTEST_EQUAL("Param2 has the right value", *Subgraph->GetUserParametersStruct()->GetValueStruct<FVector>(Param2Name).GetValue(), Param2Value);
+
+	VALIDATE_EDGE(UPCGUserParameterGetSettings, UPCGSurfaceSamplerSettings, Subgraph, NAME_None, TEXT("Looseness"))
+	VALIDATE_EDGE(UPCGUserParameterGetSettings, UPCGSurfaceSamplerSettings, Subgraph, NAME_None, TEXT("Point Extents"))
+
+	return true;
+}
+
+#undef VALIDATE_EDGE
\ No newline at end of file
diff --git a/Engine/Plugins/Experimental/PCG/Source/PCG/Public/Helpers/PCGSubgraphHelpers.h b/Engine/Plugins/Experimental/PCG/Source/PCG/Public/Helpers/PCGSubgraphHelpers.h
new file mode 100644
--- /dev/null
+++ b/Engine/Plugins/Experimental/PCG/Source/PCG/Public/Helpers/PCGSubgraphHelpers.h
@@ -0,0 +1,15 @@
+// Copyright Epic Games, Inc. All Rights Reserved.
+
+#pragma once
+
+#include "Containers/Array.h"
+
+class UPCGGraph;
+class UPCGNode;
+class UObject;
+
+class FPCGSubgraphHelpers
+{
+public:
+	static PCG_API UPCGGraph* CollapseIntoSubgraph(UPCGGraph* InOriginalGraph, const TArray<UPCGNode*>& InNodesToCollapse, const TArray<UObject*>& InExtraNodesToCollapse, UPCGGraph* OptionalPreAllocatedGraph = nullptr);
+};
\ No newline at end of file
diff --git a/Engine/Plugins/Experimental/PCG/Source/PCG/Public/PCGGraph.h b/Engine/Plugins/Experimental/PCG/Source/PCG/Public/PCGGraph.h
--- a/Engine/Plugins/Experimental/PCG/Source/PCG/Public/PCGGraph.h
+++ b/Engine/Plugins/Experimental/PCG/Source/PCG/Public/PCGGraph.h
@@ -24,7 +24,8 @@ enum class EPCGGraphParameterEvent
 	Removed,
 	PropertyModified,
 	ValueModifiedLocally,
-	ValueModifiedByParent
+	ValueModifiedByParent,
+	MultiplePropertiesAdded
 };
 
 #if WITH_EDITOR
@@ -95,6 +96,10 @@ public:
 
 	virtual const FInstancedPropertyBag* GetUserParametersStruct() const PURE_VIRTUAL(UPCGGraphInterface::GetUserParametersStruct, return nullptr;)
 
+	// Mutable version should not be used outside of testing, since there are callbacks fired when parameters changes.
+	// TODO: Make it safe to change parameters from the outside.
+	FInstancedPropertyBag* GetMutableUserParametersStruct_Unsafe() const { return const_cast<FInstancedPropertyBag*>(GetUserParametersStruct()); }
+
 	bool IsInstance() const;
 
 	/** A graph interface is equivalent to another graph interface if they are the same (same ptr), or if they have the same graph. Will be overridden when graph instance supports overrides. */
@@ -113,6 +118,7 @@ class PCG_API UPCGGraph : public UPCGGraphInterface
 {
 #if WITH_EDITOR
 	friend class FPCGEditor;
+	friend class FPCGSubgraphHelpers;
 #endif // WITH_EDITOR
 
 	GENERATED_BODY()
@@ -164,6 +170,9 @@ public:
 	UFUNCTION(BlueprintCallable, Category = Graph, meta=(DeterminesOutputType = "InSettingsClass", DynamicOutputParam = "DefaultNodeSettings"))
 	UPCGNode* AddNodeOfType(TSubclassOf<class UPCGSettings> InSettingsClass, UPCGSettings*& DefaultNodeSettings);
 
+	template <typename T, typename = typename std::enable_if_t<std::is_base_of_v<UPCGSettings, T>>>
+	UPCGNode* AddNodeOfType(T*& DefaultNodeSettings);
+
 	/** Creates a node containing an instance to the given settings. Returns the created node. */
 	UFUNCTION(BlueprintCallable, Category = Graph)
 	UPCGNode* AddNodeInstance(UPCGSettings* InSettings);
@@ -230,6 +239,7 @@ public:
 	void EnableInspection() { bIsInspecting = true; }
 	void DisableInspection() { bIsInspecting = false; }
 	bool DebugFlagAppliesToIndividualComponents() const { return bDebugFlagAppliesToIndividualComponents; }
+	void RemoveExtraEditorNode(const UObject* InNode);
 #endif
 
 #if WITH_EDITOR
@@ -296,6 +306,12 @@ protected:
 public:
 	virtual const FInstancedPropertyBag* GetUserParametersStruct() const override { return &UserParameters; }
 
+	// Add new user parameters using an array of descriptors. Can also provide an original graph to copy the values.
+	// Original Graph needs to have the properties.
+	// Be careful if there is any overlap between existing parameters, that also exists in the original graph, they will be overridden by the original.
+	// Best used on a brand new PCG Graph.
+	void AddUserParameters(const TArray<FPropertyBagPropertyDesc>& InDescs, const UPCGGraph* InOptionalOriginalGraph = nullptr);
+
 #if WITH_EDITOR
 private:
 	void NotifyGraphChanged(EPCGChangeType ChangeType);
@@ -376,6 +392,15 @@ private:
 #endif // WITH_EDITORONLY_DATA
 };
 
+template <typename T, typename>
+UPCGNode* UPCGGraph::AddNodeOfType(T*& DefaultNodeSettings)
+{
+	UPCGSettings* TempSettings = DefaultNodeSettings;
+	UPCGNode* Node = AddNodeOfType(T::StaticClass(), TempSettings);
+	DefaultNodeSettings = Cast<T>(TempSettings);
+	return Node;
+}
+
 #if UE_ENABLE_INCLUDE_ORDER_DEPRECATED_IN_5_2
 #include "CoreMinimal.h"
 #endif
diff --git a/Engine/Plugins/Experimental/PCG/Source/PCGEditor/Private/PCGEditor.cpp b/Engine/Plugins/Experimental/PCG/Source/PCGEditor/Private/PCGEditor.cpp
--- a/Engine/Plugins/Experimental/PCG/Source/PCGEditor/Private/PCGEditor.cpp
+++ b/Engine/Plugins/Experimental/PCG/Source/PCGEditor/Private/PCGEditor.cpp
@@ -13,7 +13,7 @@
 #include "PCGInputOutputSettings.h"
 #include "PCGPin.h"
 #include "PCGSubsystem.h"
-#include "Algo/AnyOf.h"
+#include "Helpers/PCGSubgraphHelpers.h"
 #include "Rendering/SlateRenderer.h"
 #include "Tests/Determinism/PCGDeterminismNativeTests.h"
 #include "Tests/Determinism/PCGDeterminismTestBlueprintBase.h"
@@ -52,6 +52,7 @@
 #include "ToolMenuSection.h"
 #include "ToolMenus.h"
 #include "UnrealEdGlobals.h"
+#include "Algo/AnyOf.h"
 #include "Editor/UnrealEdEngine.h"
 #include "Fonts/FontMeasure.h"
 #include "Framework/Application/SlateApplication.h"
@@ -892,27 +893,9 @@ void FPCGEditor::OnCollapseNodesInSubgraph()
 		return;
 	}
 
-	// 1. Gather all nodes that will be included in the subgraph, and the extra nodes
-	// Also keep track of all node positions to know where to spawn our subgraph node
-	FVector2D AveragePosition = FVector2D::ZeroVector;
-	int32 MinX = std::numeric_limits<int32>::max();
-	int32 MaxX = std::numeric_limits<int32>::min();
-
-	// We keep track of all the PCG nodes that will go into the subgraph
-	TArray<TObjectPtr<UPCGEditorGraphNodeBase>> PCGEditorGraphNodes;
-	// Also move the extra nodes (like comments)
-	TArray<TObjectPtr<UEdGraphNode>> ExtraGraphNodes;
-
-	// Those 3 sets are used to keep track of the subgraph pins
-	// and will be used to extract pins that will be outside the subgraph
-	// and will need special treatment.
-	TSet<TObjectPtr<UPCGPin>> InputFromSubgraphPins;
-	TSet<TObjectPtr<UPCGPin>> OutputToSubgraphPins;
-	TSet<TObjectPtr<UPCGPin>> AllSubgraphPins;
-
-	// Also keep track of all the edges in the subgraph. Use a set to be able to add the
-	// same edge twice without duplicates.
-	TSet<TObjectPtr<UPCGEdge>> PCGGraphEdges;
+	// Gather all nodes that will be included in the subgraph, and the extra nodes
+	TArray<UPCGNode*> NodesToCollapse;
+	TArray<UObject*> ExtraNodesToCollapse;
 
 	for (UObject* Object : GraphEditorWidget->GetSelectedNodes())
 	{
@@ -928,314 +911,60 @@ void FPCGEditor::OnCollapseNodesInSubgraph()
 		{
 			UPCGNode* PCGNode = PCGEditorGraphNode->GetPCGNode();
 			check(PCGNode);
-
-			// For each pin of the subgraph, gather all its edges in a set and store the other pin of
-			// the edge as an input from or output to the subgraph. It will be useful for tracking pins
-			// that are outside the subgraph.
-			for (UPCGPin* InputPin : PCGNode->GetInputPins())
-			{
-				check(InputPin);
-
-				for (UPCGEdge* Edge : InputPin->Edges)
-				{
-					check(Edge);
-					PCGGraphEdges.Add(Edge);
-					OutputToSubgraphPins.Add(Edge->InputPin);
-				}
-
-				AllSubgraphPins.Add(InputPin);
-			}
-
-			for (UPCGPin* OutputPin : PCGNode->GetOutputPins())
-			{
-				check(OutputPin);
-
-				for (UPCGEdge* Edge : OutputPin->Edges)
-				{
-					check(Edge);
-					PCGGraphEdges.Add(Edge);
-					InputFromSubgraphPins.Add(Edge->OutputPin);
-				}
-
-				AllSubgraphPins.Add(OutputPin);
-			}
-
-			PCGEditorGraphNodes.Add(PCGEditorGraphNode);
-
-			// And do all the computation to get the min, max and mean position.
-			AveragePosition.X += PCGEditorGraphNode->NodePosX;
-			AveragePosition.Y += PCGEditorGraphNode->NodePosY;
-			MinX = FMath::Min(MinX, PCGEditorGraphNode->NodePosX);
-			MaxX = FMath::Max(MaxX, PCGEditorGraphNode->NodePosX);
+			NodesToCollapse.Add(PCGNode);
 		}
 		else if (UEdGraphNode* GraphNode = Cast<UEdGraphNode>(Object))
 		{
-			ExtraGraphNodes.Add(GraphNode);
+			ExtraNodesToCollapse.Add(GraphNode);
 		}
 	}
 
 	// If we have at most 1 node to collapse, just exit
-	if (PCGEditorGraphNodes.Num() <= 1)
+	if (NodesToCollapse.Num() <= 1)
 	{
 		UE_LOG(LogPCGEditor, Warning, TEXT("There were less than 2 PCG nodes selected, abort"));
 		return;
 	}
 
-	// Compute the average position
-	AveragePosition /= PCGEditorGraphNodes.Num();
-
-	// 2. Create a new subgraph, by creating a new PCGGraph asset.
-	IAssetTools& AssetTools = FModuleManager::Get().LoadModuleChecked<FAssetToolsModule>("AssetTools").Get();
-
-	TObjectPtr<UPCGGraphFactory> Factory = NewObject<UPCGGraphFactory>();
-
-	FString NewPackageName;
-	FString NewAssetName;
-	PCGEditorUtils::GetParentPackagePathAndUniqueName(PCGGraph, LOCTEXT("NewPCGSubgraphAsset", "NewPCGSubgraph").ToString(), NewPackageName, NewAssetName);
-
-	TObjectPtr<UPCGGraph> NewPCGGraph = Cast<UPCGGraph>(AssetTools.CreateAssetWithDialog(NewAssetName, NewPackageName, PCGGraph->GetClass(), Factory, "PCGEditor_CollapseInSubgraph"));
-
-	if (NewPCGGraph == nullptr)
-	{
-		UE_LOG(LogPCGEditor, Warning, TEXT("Subgraph asset creation was aborted or failed, abort."));
-		return;
-	}
-
-	// Do some clean-up on input/output nodes
-	constexpr int32 Padding = 200;
-	NewPCGGraph->GetInputNode()->PositionX = MinX - Padding;
-	NewPCGGraph->GetInputNode()->PositionY = AveragePosition.Y;
-	NewPCGGraph->GetOutputNode()->PositionX = MaxX + Padding;
-	NewPCGGraph->GetOutputNode()->PositionY = AveragePosition.Y;
-
-	// 3. Define some lambdas to create/adapt pins on the new subgraph.
-
-	// Gather the pins that are outside the subgraph
-	InputFromSubgraphPins = InputFromSubgraphPins.Difference(AllSubgraphPins);
-	OutputToSubgraphPins = OutputToSubgraphPins.Difference(AllSubgraphPins);
-
-	// Function that will create a new pin in the input/output node of the subgraph, with a unique name that will
-	// be the related to the pin passed as argument. It will be formatted like this:
-	// "{NodeName} {PinName} {OptionalIndex}"
-	TMap<FString, int> NameCollisionMapping;
-	auto CreateNewCustomPin = [&NewPCGGraph, &NameCollisionMapping](UPCGPin* Pin, bool bIsInput)
-	{
-		TObjectPtr<UPCGNode> NewInputOutputNode = bIsInput ? NewPCGGraph->GetInputNode() : NewPCGGraph->GetOutputNode();
-		TObjectPtr<UPCGGraphInputOutputSettings> NewInputOutputSettings = CastChecked<UPCGGraphInputOutputSettings>(NewInputOutputNode->GetSettings());
-		FString NewName = Pin->Node->GetNodeTitle().ToString() + " " + Pin->Properties.Label.ToString();
-		if (NameCollisionMapping.Contains(NewName))
-		{
-			NewName += " " + FString::FormatAsNumber(++NameCollisionMapping[NewName]);
-		}
-		else
-		{
-			NameCollisionMapping.Emplace(NewName, 1);
-		}
-
-		FPCGPinProperties NewProperties = Pin->Properties;
-		NewProperties.Label = FName(NewName);
-		NewProperties = NewInputOutputSettings->AddCustomPin(NewProperties);
-		NewInputOutputNode->UpdateAfterSettingsChangeDuringCreation();
-		return NewProperties.Label;
-	};
-
-	// 4. Duplicate all the nodes, and keep a mapping between the old pins and new pins
-	TMap<TObjectPtr<UPCGPin>, TObjectPtr<UPCGPin>> PinMapping;
-	for (UPCGEditorGraphNodeBase* PCGEditorGraphNode : PCGEditorGraphNodes)
-	{
-		const UPCGNode* PCGNode = PCGEditorGraphNode->GetPCGNode();
-		check(PCGNode);
-
-		// Reconstruct a new node, same as PCGNode, but without any edges in the new graph 
-		TObjectPtr<UPCGNode> NewNode = NewPCGGraph->ReconstructNewNode(PCGNode);
-
-		// Safeguard: We should have a 1 for 1 matching between pins labels between the original node
-		// and the copied node. If for some reason we don't (perhaps the node was not updated correctly after pins were added/removed)
-		// we will log an error and try to connect as best as we can (probably breaking some edges on the process).
-
-		auto Mapping = [&PinMapping, PCGNode](const TArray<UPCGPin*>& OriginalPins, const TArray<UPCGPin*>& NewPins)
-		{
-			TSet<FName> UnmatchedOriginal;
-			TMap<FName, UPCGPin*> NewMapping;
-
-			for (UPCGPin* NewPin : NewPins)
-			{
-				NewMapping.Emplace(NewPin->Properties.Label, NewPin);
-			}
-
-			for (UPCGPin* OriginalPin : OriginalPins)
-			{
-				FName PinLabel = OriginalPin->Properties.Label;
-				if (UPCGPin** NewPinPtr = NewMapping.Find(PinLabel))
-				{
-					PinMapping.Emplace(OriginalPin, *NewPinPtr);
-				}
-				else if (OriginalPin->IsConnected())
-				{
-					// It is only problematic if the pin was connected
-					UE_LOG(LogPCG, Error, TEXT("[CollapseInSubgraph - %s] %s pin %s does not exist anymore. Edges will be broken."),
-						*PCGNode->GetNodeTitle().ToString(), (OriginalPin->IsOutputPin() ? TEXT("Output") : TEXT("Input")), *PinLabel.ToString());
-				}
-			}
-		};
-		
-		Mapping(PCGNode->GetInputPins(), NewNode->GetInputPins());
-		Mapping(PCGNode->GetOutputPins(), NewNode->GetOutputPins());
-	}
-
-	// Also duplicate the extra nodes and assign them to the new graph
-	TArray<TObjectPtr<const UObject>> NewExtraGraphNodes;
-	for (const UEdGraphNode* ExtraNode : ExtraGraphNodes)
+	// Create a new subgraph, by creating a new PCGGraph asset.
+	TObjectPtr<UPCGGraph> NewPCGGraph = nullptr;
 	{
-		NewExtraGraphNodes.Add(DuplicateObject(ExtraNode, NewPCGGraph));
-	}
+		FScopedTransaction Transaction(LOCTEXT("PCGCollapseInSubgraphMessage", "[PCG] Collapse into Subgraph"));
 
-	NewPCGGraph->SetExtraEditorNodes(NewExtraGraphNodes);
+		IAssetTools& AssetTools = FModuleManager::Get().LoadModuleChecked<FAssetToolsModule>("AssetTools").Get();
 
-	// 5. Iterate over all the edges and create edges "placeholders"
-	// Most of them will already be complete, but for those that needs to be connected to the new
-	// subgraph nodes, the pins don't exist yet. Therefore we identify them with their pin labels.
-	//
-	// The logic behind the new pins is this:
-	// -> If the pin is connected to the simple pin of the input/output node, let it like this
-	// -> If the pin is part of the original input/output advanced pins, we will trigger the advanced pins flag on the input/output node in the subgraph
-	// -> Otherwise, we add a new custom pin, with the name of the node, the name of the pin and a number if there is name collision
-	struct EdgePlaceholder
-	{
-		TObjectPtr<UPCGPin> InputPin = nullptr;
-		TObjectPtr<UPCGPin> OutputPin = nullptr;
-		FName InputPinLabel;
-		FName OutputPinLabel;
-	};
+		TObjectPtr<UPCGGraphFactory> Factory = NewObject<UPCGGraphFactory>();
 
-	TArray<EdgePlaceholder> EdgePlaceholders;
-	for (UPCGEdge* Edge : PCGGraphEdges)
-	{
-		const TObjectPtr<UPCGPin>* InPin = PinMapping.Find(Edge->InputPin);
-		const TObjectPtr<UPCGPin>* OutPin = PinMapping.Find(Edge->OutputPin);
+		FString NewPackageName;
+		FString NewAssetName;
+		PCGEditorUtils::GetParentPackagePathAndUniqueName(PCGGraph, LOCTEXT("NewPCGSubgraphAsset", "NewPCGSubgraph").ToString(), NewPackageName, NewAssetName);
 
-		check(InPin || OutPin)
+		NewPCGGraph = Cast<UPCGGraph>(AssetTools.CreateAssetWithDialog(NewAssetName, NewPackageName, PCGGraph->GetClass(), Factory, "PCGEditor_CollapseInSubgraph"));
 
-		if (InPin == nullptr)
+		if (NewPCGGraph == nullptr)
 		{
-			// The edge comes from outside the graph.
-			// If it is from the input node, we have a special behavior
-			EdgePlaceholder OutsideSubgraphEdge;
-			EdgePlaceholder InsideSubgraphEdge;
-			bool bProcessed = false;
-
-			OutsideSubgraphEdge.InputPin = Edge->InputPin;
-			InsideSubgraphEdge.OutputPin = *OutPin;
-
-			if (Edge->InputPin->Node == PCGGraph->GetInputNode())
-			{
-				const UPCGGraphInputOutputSettings* Settings = Cast<const UPCGGraphInputOutputSettings>(Edge->InputPin->Node->GetSettings());
-
-				if (Settings && !Settings->IsCustomPin(Edge->InputPin))
-				{
-					OutsideSubgraphEdge.OutputPinLabel = Edge->InputPin->Properties.Label;
-					InsideSubgraphEdge.InputPin = NewPCGGraph->GetInputNode()->GetOutputPin(Edge->InputPin->Properties.Label);
-					bProcessed = true;
-				}
-			}
-
-			if (!bProcessed)
-			{
-				FName NewPinName = CreateNewCustomPin(Edge->OutputPin, true);
-				OutsideSubgraphEdge.OutputPinLabel = NewPinName;
-				InsideSubgraphEdge.InputPin = NewPCGGraph->GetInputNode()->GetOutputPin(NewPinName);
-			}
-
-			EdgePlaceholders.Add(OutsideSubgraphEdge);
-			EdgePlaceholders.Add(InsideSubgraphEdge);
+			UE_LOG(LogPCGEditor, Warning, TEXT("Subgraph asset creation was aborted or failed, abort."));
+			Transaction.Cancel();
+			return;
 		}
-		else if (OutPin == nullptr)
-		{
-			// The edge comes from outside the graph.
-			// If it is from the output node, we have a special behavior
-			EdgePlaceholder OutsideSubgraphEdge;
-			EdgePlaceholder InsideSubgraphEdge;
-			bool bProcessed = false;
 
-			OutsideSubgraphEdge.OutputPin = Edge->OutputPin;
-			InsideSubgraphEdge.InputPin = *InPin;
+		NewPCGGraph = FPCGSubgraphHelpers::CollapseIntoSubgraph(PCGGraph, NodesToCollapse, ExtraNodesToCollapse, NewPCGGraph);
 
-			if (Edge->OutputPin->Node == PCGGraph->GetOutputNode())
-			{
-				const UPCGGraphInputOutputSettings* Settings = Cast<const UPCGGraphInputOutputSettings>(Edge->OutputPin->Node->GetSettings());
-
-				if (Settings && !Settings->IsCustomPin(Edge->OutputPin))
-				{
-					OutsideSubgraphEdge.InputPinLabel = Edge->OutputPin->Properties.Label;
-					InsideSubgraphEdge.OutputPin = NewPCGGraph->GetOutputNode()->GetInputPin(Edge->OutputPin->Properties.Label);
-					bProcessed = true;
-				}
-			}
-
-			if (!bProcessed)
-			{
-				FName NewPinName = CreateNewCustomPin(Edge->InputPin, false);
-				OutsideSubgraphEdge.InputPinLabel = NewPinName;
-				InsideSubgraphEdge.OutputPin = NewPCGGraph->GetOutputNode()->GetInputPin(NewPinName);
-			}
-
-			EdgePlaceholders.Add(OutsideSubgraphEdge);
-			EdgePlaceholders.Add(InsideSubgraphEdge);
-		}
-		else
+		if (NewPCGGraph == nullptr)
 		{
-			// Both nodes are inside
-			EdgePlaceholders.Add(EdgePlaceholder{*InPin, *OutPin});
+			UE_LOG(LogPCGEditor, Warning, TEXT("Subgraph collapse failed, abort."));
+			return;
 		}
-	}
 
-	// 6. Create subgraph and delete old nodes
-	// Done within a transaction to be undoable (subgraph will stay though)
-	TObjectPtr<UPCGEditorGraphNodeBase> SubgraphEditorNode;
-	TObjectPtr<UPCGNode> SubgraphNode;
-	{
-		const FScopedTransaction Transaction(*FPCGEditorCommon::ContextIdentifier, LOCTEXT("PCGEditorCollapseInSubgraphMessage", "PCG Editor: Collapse into Subgraph"), nullptr);
-
-		DeleteSelectedNodes();
-
-		FPCGEditorGraphSchemaAction_NewSubgraphElement SubgraphAddAction;
-		SubgraphAddAction.SubgraphObjectPath = NewPCGGraph->GetPathName();
-		SubgraphEditorNode = Cast<UPCGEditorGraphNodeBase>(SubgraphAddAction.PerformAction(PCGEditorGraph, nullptr, AveragePosition, true));
-		SubgraphNode = SubgraphEditorNode->GetPCGNode();
+		// Force a refresh
+		PCGEditorGraph->ReconstructGraph();
 	}
 
-	// 7. Connect all the edges
-	TSet<UPCGNode*> TouchedNodes;
-	for (EdgePlaceholder& Edge : EdgePlaceholders)
-	{
-		if (Edge.InputPin == nullptr)
-		{
-			SubgraphNode->GetOutputPin(Edge.InputPinLabel)->AddEdgeTo(Edge.OutputPin, &TouchedNodes);
-		}
-		else if (Edge.OutputPin == nullptr)
-		{
-			Edge.InputPin->AddEdgeTo(SubgraphNode->GetInputPin(Edge.OutputPinLabel), &TouchedNodes);
-		}
-		else
-		{
-			Edge.InputPin->AddEdgeTo(Edge.OutputPin, &TouchedNodes);
-		}
-	}
-
-	// 8. Finalize the operation
-	SubgraphEditorNode->ReconstructNode();
-
 	// Save the new asset
 	UEditorAssetLibrary::SaveLoadedAsset(NewPCGGraph);
 
-	// And notify everyone
+	// Notify the widget
 	GraphEditorWidget->NotifyGraphChanged();
-	PCGGraph->NotifyGraphChanged(EPCGChangeType::Structural);
-	NewPCGGraph->NotifyGraphChanged(EPCGChangeType::Structural);
-	for (UPCGNode* TouchedNode : TouchedNodes)
-	{
-		TouchedNode->OnNodeChangedDelegate.Broadcast(TouchedNode, EPCGChangeType::Node);
-	}
 }
 
 bool FPCGEditor::CanExportNodes() const
diff --git a/Engine/Plugins/Experimental/PCG/Source/PCGEditor/Private/PCGEditorGraph.cpp b/Engine/Plugins/Experimental/PCG/Source/PCGEditor/Private/PCGEditorGraph.cpp
--- a/Engine/Plugins/Experimental/PCG/Source/PCGEditor/Private/PCGEditorGraph.cpp
+++ b/Engine/Plugins/Experimental/PCG/Source/PCGEditor/Private/PCGEditorGraph.cpp
@@ -22,6 +22,25 @@ void UPCGEditorGraph::InitFromNodeGraph(UPCGGraph* InPCGGraph)
 
 	PCGGraph->OnGraphParametersChangedDelegate.AddUObject(this, &UPCGEditorGraph::OnGraphUserParametersChanged);
 
+	ReconstructGraph();
+}
+
+void UPCGEditorGraph::ReconstructGraph()
+{
+	check(PCGGraph);
+
+	// If there are already some nodes, remove all of them.
+	if (!Nodes.IsEmpty())
+	{
+		Modify();
+
+		TArray<TObjectPtr<class UEdGraphNode>> NodesCopy = Nodes;
+		for (UEdGraphNode* Node : NodesCopy)
+		{
+			RemoveNode(Node);
+		}
+	}
+
 	TMap<UPCGNode*, UPCGEditorGraphNodeBase*> NodeLookup;
 	const bool bSelectNewNode = false;
 
diff --git a/Engine/Plugins/Experimental/PCG/Source/PCGEditor/Private/PCGEditorGraph.h b/Engine/Plugins/Experimental/PCG/Source/PCGEditor/Private/PCGEditorGraph.h
--- a/Engine/Plugins/Experimental/PCG/Source/PCGEditor/Private/PCGEditorGraph.h
+++ b/Engine/Plugins/Experimental/PCG/Source/PCGEditor/Private/PCGEditorGraph.h
@@ -26,6 +26,9 @@ public:
 	/** Initialize the editor graph from a PCGGraph */
 	void InitFromNodeGraph(UPCGGraph* InPCGGraph);
 
+	/** If the underlying graph changed without UI interaction, use this function to reconstruct the UI elements. */
+	void ReconstructGraph();
+
 	/** When the editor is closing */
 	void OnClose();
 
