Fixed a few issues related to difference:
- Surface was not provided point bounds appropriately, which made some differences not be considered
- Fixed error in the overlap computation not using the right referential in one of the two comparisons
- Fixed error related to using non-uniform scales with Transforms, passed to using matrices instead
Also set single point from actor is now given a steepness of 1 to be consistent with user expectations

diff --git a/Engine/Plugins/Experimental/PCG/Source/PCG/Private/Data/PCGPointData.cpp b/Engine/Plugins/Experimental/PCG/Source/PCG/Private/Data/PCGPointData.cpp
--- a/Engine/Plugins/Experimental/PCG/Source/PCG/Private/Data/PCGPointData.cpp
+++ b/Engine/Plugins/Experimental/PCG/Source/PCG/Private/Data/PCGPointData.cpp
@@ -90,27 +90,33 @@ namespace PCGPointHelpers
 			(DenominatorExtent.Z > 0 ? NumeratorExtent.Z / DenominatorExtent.Z : 1.0));
 	}
 
-	FVector::FReal VolumeOverlap(const FPCGPoint& InPoint, const FBox& InBounds, const FTransform& InTransform)
+	FVector::FReal VolumeOverlap(const FPCGPoint& InPoint, const FBox& InBounds, const FMatrix& InInverseTransform)
 	{
 		// This is similar in idea to SAT considering we have two boxes - since we will test all 6 axes.
 		// However, there is some uncertainty due to rotation, and using the overlap value as-is is an overestimation, which might not be critical in this case
 		// TODO: investigate if we should do a 8-pt test instead (would be more precise, but significantly more costly).
+		// Implementation note: we are using FMatrix here because we want to support non-uniform scales
 		const FBox PointBounds = InPoint.GetLocalDensityBounds();
-		const FTransform& PointTransform = InPoint.Transform;
 
-		const FBox FirstOverlap = PointBounds.Overlap(InBounds.TransformBy(InTransform.GetRelativeTransform(PointTransform)));
+		FMatrix PointTransformToInTransform = InPoint.Transform.ToMatrixWithScale() * InInverseTransform;
+		const FBox PointBoundsTransformed = PointBounds.TransformBy(PointTransformToInTransform);
+
+		const FBox FirstOverlap = InBounds.Overlap(PointBoundsTransformed);
 		if (!FirstOverlap.IsValid)
 		{
 			return 0;
 		}
 
-		const FBox SecondOverlap = InBounds.Overlap(PointBounds.TransformBy(PointTransform.GetRelativeTransform(InTransform)));
+		FMatrix InTransformToPointTransform = PointTransformToInTransform.Inverse();
+		const FBox InBoundsTransformed = InBounds.TransformBy(InTransformToPointTransform);
+
+		const FBox SecondOverlap = InBoundsTransformed.Overlap(PointBounds);
 		if (!SecondOverlap.IsValid)
 		{
 			return 0;
 		}
-		
-		return FMath::Min(ComputeOverlapRatio(FirstOverlap, InBounds), ComputeOverlapRatio(SecondOverlap, InBounds));
+
+		return FMath::Min(ComputeOverlapRatio(FirstOverlap, InBounds), ComputeOverlapRatio(SecondOverlap, InBoundsTransformed));
 	}
 
 	/** Helper function for additive blending of quaternions (copied from ControlRig) */
@@ -395,6 +401,7 @@ void UPCGPointData::AddSinglePointFromActor(AActor* InActor)
 	check(InActor);
 
 	FPCGPoint& Point = GetMutablePoints().Emplace_GetRef();
+	Point.Steepness = 1.0f;
 	Point.Transform = InActor->GetActorTransform();
 
 	const FVector& Position = Point.Transform.GetLocation();
@@ -499,8 +506,10 @@ bool UPCGPointData::ProjectPoint(const FTransform& InTransform, const FBox& InBo
 	else
 	{
 		FBox TransformedBounds = InBounds.TransformBy(InTransform);
-		Octree.FindElementsWithBoundsTest(FBoxCenterAndExtent(TransformedBounds.GetCenter(), TransformedBounds.GetExtent()), [&InBounds, &InTransform, &Contributions](const FPCGPointRef& InPointRef) {
-			float Contribution = PCGPointHelpers::VolumeOverlap(*InPointRef.Point, InBounds, InTransform);
+		FMatrix InTransformInverseMatrix = InTransform.ToMatrixWithScale().Inverse();
+
+		Octree.FindElementsWithBoundsTest(FBoxCenterAndExtent(TransformedBounds.GetCenter(), TransformedBounds.GetExtent()), [&InBounds, &InTransformInverseMatrix, &Contributions](const FPCGPointRef& InPointRef) {
+			float Contribution = PCGPointHelpers::VolumeOverlap(*InPointRef.Point, InBounds, InTransformInverseMatrix);
 			if (Contribution > 0)
 			{
 				Contributions.Emplace(InPointRef.Point, Contribution);
diff --git a/Engine/Plugins/Experimental/PCG/Source/PCG/Private/Elements/PCGSurfaceSampler.cpp b/Engine/Plugins/Experimental/PCG/Source/PCG/Private/Elements/PCGSurfaceSampler.cpp
--- a/Engine/Plugins/Experimental/PCG/Source/PCG/Private/Elements/PCGSurfaceSampler.cpp
+++ b/Engine/Plugins/Experimental/PCG/Source/PCG/Private/Elements/PCGSurfaceSampler.cpp
@@ -201,6 +201,10 @@ namespace PCGSurfaceSampler
 				return false;
 			}
 
+			// Set physical properties that are needed for the bounding shape checks, etc.
+			OutPoint.SetExtents(LoopData.PointExtents);
+			OutPoint.Steepness = LoopData.PointSteepness;
+
 			// Now run gauntlet of shape network (if there is one) to accept or reject the point.
 			if (InBoundingShape)
 			{
@@ -219,9 +223,7 @@ namespace PCGSurfaceSampler
 			}
 
 			// Apply final parameters on the point
-			OutPoint.SetExtents(LoopData.PointExtents);
 			OutPoint.Density *= (LoopData.bApplyDensityToPoints ? ((Ratio - Chance) / Ratio) : 1.0f);
-			OutPoint.Steepness = LoopData.PointSteepness;
 			OutPoint.Seed = RandomSource.GetCurrentSeed();
 
 			return true;
