[PCG] Removing PCG World Actor dependency in tracking as the new system does not require it.

* It was added as a 1 to 1 with what was previously done in PCG Component but it brought more harm than good.
* New tracking system is now tracking more effectively any actor/component added
* Removed Cached tracking on PCG World Actor and deprecated versions in PCG Component

#rnx

diff --git a/Engine/Plugins/Experimental/PCG/Source/PCG/Private/PCGActorAndComponentMapping.cpp b/Engine/Plugins/Experimental/PCG/Source/PCG/Private/PCGActorAndComponentMapping.cpp
--- a/Engine/Plugins/Experimental/PCG/Source/PCG/Private/PCGActorAndComponentMapping.cpp
+++ b/Engine/Plugins/Experimental/PCG/Source/PCG/Private/PCGActorAndComponentMapping.cpp
@@ -605,13 +605,6 @@ void UPCGActorAndComponentMapping::RegisterOrUpdateTracking(UPCGComponent* InCom
 		return;
 	}
 
-	APCGWorldActor* PCGWorldActor = PCGSubsystem ? PCGSubsystem->GetPCGWorldActor() : nullptr;
-
-	if (!PCGWorldActor)
-	{
-		return;
-	}
-
 	// Components owner needs to be always tracked
 	RegisterActor(ComponentOwner);
 	TSet<UPCGComponent*>& AllComponents = AlwaysTrackedActorsToComponentsMap.FindOrAdd(ComponentOwner);
@@ -623,7 +616,7 @@ void UPCGActorAndComponentMapping::RegisterOrUpdateTracking(UPCGComponent* InCom
 	// And we also need to find all actors that should be tracked
 	if (UPCGGraph* PCGGraph = InComponent->GetGraph())
 	{
-		auto FindActorsAndTrack = [this, InComponent, PCGWorldActor, bInShouldDirtyActors](const FPCGActorSelectionKey& InKey, const TArray<FPCGSettingsAndCulling>& InSettingsAndCulling)
+		auto FindActorsAndTrack = [this, InComponent, bInShouldDirtyActors](const FPCGActorSelectionKey& InKey, const TArray<FPCGSettingsAndCulling>& InSettingsAndCulling)
 		{
 			// InKey provide the info for selecting a given actor.
 			// We reconstruct the selector settings from this key, and we also force it to SelectMultiple, since
@@ -810,17 +803,6 @@ void UPCGActorAndComponentMapping::AddActorsPostInit()
 	}
 
 	DelayedAddedActors.Empty();
-
-	// Also add the one tracked by the World Actor
-	if (APCGWorldActor* PCGWorldActor = PCGSubsystem->FindPCGWorldActor())
-	{
-		// Making a copy, since CachedTrackedActors can be modified (if an actor is no longer tracked)
-		TSet<TWeakObjectPtr<AActor>> CachedTrackedActorsCopy = PCGWorldActor->CachedTrackedActors;
-		for (TWeakObjectPtr<AActor>& ActorPtr : CachedTrackedActorsCopy)
-		{
-			AddOrUpdateTrackedActor(ActorPtr.Get());
-		}
-	}
 }
 
 void UPCGActorAndComponentMapping::OnActorAdded(AActor* InActor)
@@ -857,7 +839,7 @@ void UPCGActorAndComponentMapping::OnActorAdded(AActor* InActor)
 bool UPCGActorAndComponentMapping::AddOrUpdateTrackedActor(AActor* InActor)
 {
 	// We have to make sure to not create a infinite loop
-	if (!InActor || InActor->IsA<APCGWorldActor>() || !PCGSubsystem || !PCGSubsystem->FindPCGWorldActor() || InActor->GetWorld() != PCGSubsystem->GetWorld())
+	if (!InActor || InActor->IsA<APCGWorldActor>() || !PCGSubsystem || InActor->GetWorld() != PCGSubsystem->GetWorld())
 	{
 		return false;
 	}
@@ -917,7 +899,7 @@ bool UPCGActorAndComponentMapping::AddOrUpdateTrackedActor(AActor* InActor)
 		RegisterActor(InActor);
 		return true;
 	}
-	else if (CulledTrackedActorsToComponentsMap.Contains(InActor) || AlwaysTrackedActorsToComponentsMap.Contains(InActor))
+	else if (IsActorTracked(InActor))
 	{
 		// Do some cleanup if the actor was tracked. 
 		// We will force the refresh here, so return false to make sure we don't refresh it twice.
@@ -933,17 +915,11 @@ bool UPCGActorAndComponentMapping::AddOrUpdateTrackedActor(AActor* InActor)
 
 void UPCGActorAndComponentMapping::RegisterActor(AActor* InActor)
 {
-	APCGWorldActor* PCGWorldActor = PCGSubsystem ? PCGSubsystem->FindPCGWorldActor() : nullptr;
-	if (!PCGWorldActor || !InActor)
+	if (!InActor)
 	{
 		return;
 	}
 
-	if (!PCGWorldActor->CachedTrackedActors.Contains(InActor))
-	{
-		PCGWorldActor->CachedTrackedActors.Add(InActor);
-	}
-
 	if (ALandscapeProxy* LandscapeProxy = Cast<ALandscapeProxy>(InActor))
 	{
 		// Only add it once.
@@ -961,15 +937,13 @@ void UPCGActorAndComponentMapping::RegisterActor(AActor* InActor)
 
 bool UPCGActorAndComponentMapping::UnregisterActor(AActor* InActor)
 {
-	APCGWorldActor* PCGWorldActor = PCGSubsystem ? PCGSubsystem->FindPCGWorldActor() : nullptr;
-	if (!PCGWorldActor || !InActor)
+	if (!InActor)
 	{
 		return false;
 	}
 
-	if (PCGWorldActor->CachedTrackedActors.Contains(InActor))
+	if (IsActorTracked(InActor))
 	{
-		PCGWorldActor->CachedTrackedActors.Remove(InActor);
 		TrackedActorToPositionMap.Remove(InActor);
 		CulledTrackedActorsToComponentsMap.Remove(InActor);
 		AlwaysTrackedActorsToComponentsMap.Remove(InActor);
@@ -1004,8 +978,7 @@ void UPCGActorAndComponentMapping::OnActorDeleted(AActor* InActor)
 	}
 #endif
 
-	APCGWorldActor* PCGWorldActor = PCGSubsystem->FindPCGWorldActor();
-	if (!PCGWorldActor || !PCGWorldActor->CachedTrackedActors.Contains(InActor))
+	if (!IsActorTracked(InActor))
 	{
 		return;
 	}
@@ -1021,8 +994,7 @@ void UPCGActorAndComponentMapping::OnActorMoved(AActor* InActor)
 {
 	TRACE_CPUPROFILER_EVENT_SCOPE(UPCGActorAndComponentMapping::OnActorMoved);
 
-	APCGWorldActor* PCGWorldActor = PCGSubsystem ? PCGSubsystem->FindPCGWorldActor() : nullptr;
-	if (!PCGWorldActor || !InActor || (PCGSubsystem && InActor->GetWorld() != PCGSubsystem->GetWorld()) || !PCGWorldActor->CachedTrackedActors.Contains(InActor))
+	if (!InActor || (PCGSubsystem && InActor->GetWorld() != PCGSubsystem->GetWorld()) || !IsActorTracked(InActor))
 	{
 		return;
 	}
@@ -1087,13 +1059,6 @@ void UPCGActorAndComponentMapping::OnObjectPropertyChanged(UObject* InObject, FP
 		return;
 	}
 
-	APCGWorldActor* PCGWorldActor = PCGSubsystem ? PCGSubsystem->FindPCGWorldActor() : nullptr;
-
-	if (!PCGWorldActor)
-	{
-		return;
-	}
-
 	// First check if it is an actor
 	AActor* Actor = Cast<AActor>(InObject);
 
@@ -1138,7 +1103,7 @@ void UPCGActorAndComponentMapping::OnObjectPropertyChanged(UObject* InObject, FP
 
 	// Check if we are not tracking it or is a tag change.
 	bool bShouldChange = true;
-	if (!PCGWorldActor->CachedTrackedActors.Contains(Actor) || bActorTagChange)
+	if (!IsActorTracked(Actor) || bActorTagChange)
 	{
 		bShouldChange = AddOrUpdateTrackedActor(Actor);
 	}
@@ -1359,4 +1324,9 @@ void UPCGActorAndComponentMapping::OnPCGGraphGeneratedOrCleaned(UPCGComponent* I
 	OnActorChanged(InComponent->GetOwner(), /*bInHasMoved=*/false, InComponent);
 }
 
+bool UPCGActorAndComponentMapping::IsActorTracked(const AActor* InActor) const
+{
+	return InActor && (CulledTrackedActorsToComponentsMap.Contains(InActor) || AlwaysTrackedActorsToComponentsMap.Contains(InActor));
+}
+
 #endif // WITH_EDITOR
\ No newline at end of file
diff --git a/Engine/Plugins/Experimental/PCG/Source/PCG/Private/PCGComponent.cpp b/Engine/Plugins/Experimental/PCG/Source/PCG/Private/PCGComponent.cpp
--- a/Engine/Plugins/Experimental/PCG/Source/PCG/Private/PCGComponent.cpp
+++ b/Engine/Plugins/Experimental/PCG/Source/PCG/Private/PCGComponent.cpp
@@ -1119,18 +1119,6 @@ void UPCGComponent::PostLoad()
 		Graph_DEPRECATED = nullptr;
 	}
 
-	if (!CachedTrackedActors_DEPRECATED.IsEmpty() || !TrackedLandscapes_DEPRECATED.IsEmpty())
-	{
-		if (UPCGSubsystem* Subsystem = GetSubsystem())
-		{
-			Subsystem->AddSerializedTrackedActorsForDeprecation(CachedTrackedActors_DEPRECATED);
-			Subsystem->AddSerializedTrackedActorsForDeprecation(TrackedLandscapes_DEPRECATED);
-		}
-
-		CachedTrackedActors_DEPRECATED.Empty();
-		TrackedLandscapes_DEPRECATED.Empty();
-	}
-
 	SetupCallbacksOnCreation();
 #endif
 }
diff --git a/Engine/Plugins/Experimental/PCG/Source/PCG/Public/PCGActorAndComponentMapping.h b/Engine/Plugins/Experimental/PCG/Source/PCG/Public/PCGActorAndComponentMapping.h
--- a/Engine/Plugins/Experimental/PCG/Source/PCG/Public/PCGActorAndComponentMapping.h
+++ b/Engine/Plugins/Experimental/PCG/Source/PCG/Public/PCGActorAndComponentMapping.h
@@ -126,10 +126,6 @@ private:
 	void OnObjectPropertyChanged(UObject* InObject, FPropertyChangedEvent& InEvent);
 	void OnPCGGraphGeneratedOrCleaned(UPCGComponent* InComponent);
 
-	/** For deprecation only, register tracked actors from serialized components */
-	template <typename ActorContainer>
-	void AddSerializedTrackedActorsForDeprecation(const ActorContainer& InActors);
-
 	/** Remap the tracking in case of BP components. */
 	void RemapTracking(const UPCGComponent* InOldComponent, UPCGComponent* InNewComponent);
 
@@ -144,6 +140,8 @@ private:
 
 	/** Update dependencies for a given tracked actor. */
 	void UpdateActorDependencies(AActor* InActor);
+
+	bool IsActorTracked(const AActor* InActor) const;
 #endif // WITH_EDITOR
 
 private:
@@ -191,22 +189,4 @@ private:
 
 	/** Transient list of tags, kept when there is a tag change on a tracked Actor. */
 	TSet<FName> TempTrackedActorTags;
-};
-
-#if WITH_EDITOR
-template <typename ActorContainer>
-void UPCGActorAndComponentMapping::AddSerializedTrackedActorsForDeprecation(const ActorContainer& InActors)
-{
-	using PointerType = typename ActorContainer::ElementType;
-	using ElementType = typename PointerType::ElementType;
-	static_assert(std::is_base_of_v<AActor, ElementType>);
-
-	for (auto ActorPtr : InActors)
-	{
-		if (AActor* Actor = ActorPtr.Get())
-		{
-			OnActorAdded(Actor);
-		}
-	}
-}
-#endif // WITH_EDITOR
\ No newline at end of file
+};
\ No newline at end of file
diff --git a/Engine/Plugins/Experimental/PCG/Source/PCG/Public/PCGComponent.h b/Engine/Plugins/Experimental/PCG/Source/PCG/Public/PCGComponent.h
--- a/Engine/Plugins/Experimental/PCG/Source/PCG/Public/PCGComponent.h
+++ b/Engine/Plugins/Experimental/PCG/Source/PCG/Public/PCGComponent.h
@@ -418,15 +418,6 @@ private:
 #endif
 
 #if WITH_EDITORONLY_DATA
-	UPROPERTY()
-	TArray<TSoftObjectPtr<ALandscapeProxy>> TrackedLandscapes_DEPRECATED;
-#endif
-
-#if WITH_EDITORONLY_DATA
-	// Cached tracked actors list is serialized because we can't get it at postload time
-	UPROPERTY()
-	TSet<TWeakObjectPtr<AActor>> CachedTrackedActors_DEPRECATED;
-
 	UPROPERTY(Transient)
 	TMap<FPCGStack, FPCGDataCollection> InspectionCache;
 #endif
diff --git a/Engine/Plugins/Experimental/PCG/Source/PCG/Public/PCGSubsystem.h b/Engine/Plugins/Experimental/PCG/Source/PCG/Public/PCGSubsystem.h
--- a/Engine/Plugins/Experimental/PCG/Source/PCG/Public/PCGSubsystem.h
+++ b/Engine/Plugins/Experimental/PCG/Source/PCG/Public/PCGSubsystem.h
@@ -177,12 +177,6 @@ public:
 	/** Update the tracking on a given component. */
 	void UpdateComponentTracking(UPCGComponent* InComponent, bool bInShouldDirtyActors = false) { ActorAndComponentMapping.RegisterOrUpdateTracking(InComponent, bInShouldDirtyActors); }
 
-#if WITH_EDITOR
-	/** For deprecation only, register tracked actors from serialized components */
-	template <typename ActorContainer>
-	void AddSerializedTrackedActorsForDeprecation(const ActorContainer& InActors) { ActorAndComponentMapping.AddSerializedTrackedActorsForDeprecation(InActors); }
-#endif // WITH_EDITOR
-
 	/** Cleans up the graph cache on an element basis. InSettings is used for debugging and is optional. */
 	void CleanFromCache(const IPCGElement* InElement, const UPCGSettings* InSettings = nullptr);
 
diff --git a/Engine/Plugins/Experimental/PCG/Source/PCG/Public/PCGWorldActor.h b/Engine/Plugins/Experimental/PCG/Source/PCG/Public/PCGWorldActor.h
--- a/Engine/Plugins/Experimental/PCG/Source/PCG/Public/PCGWorldActor.h
+++ b/Engine/Plugins/Experimental/PCG/Source/PCG/Public/PCGWorldActor.h
@@ -61,12 +61,6 @@ public:
 	UPROPERTY(config, EditAnywhere, Category = GenerationSettings)
 	bool bUse2DGrid = true;
 
-#if WITH_EDITORONLY_DATA
-	// Cached tracked actors list is serialized because we can't get it at postload time
-	UPROPERTY(Transient)
-	TSet<TWeakObjectPtr<AActor>> CachedTrackedActors;
-#endif // WITH_EDITORONLY_DATA
-
 private:
 	void RegisterToSubsystem();
 	void UnregisterFromSubsystem();
