[PCG] Fixed a crash where Determinism tests ran nodes that required a World to exist (such as World Ray Hit Query)

#rnx

diff --git a/Engine/Plugins/Experimental/PCG/Source/PCG/Private/Tests/Determinism/PCGDeterminismTestsCommon.cpp b/Engine/Plugins/Experimental/PCG/Source/PCG/Private/Tests/Determinism/PCGDeterminismTestsCommon.cpp
--- a/Engine/Plugins/Experimental/PCG/Source/PCG/Private/Tests/Determinism/PCGDeterminismTestsCommon.cpp
+++ b/Engine/Plugins/Experimental/PCG/Source/PCG/Private/Tests/Determinism/PCGDeterminismTestsCommon.cpp
@@ -82,7 +82,7 @@ namespace PCGDeterminismTests
 		OutResult.DataTypesTested |= InputPCGData->GetDataType();
 
 		// Clone the actor and component
-		AActor* PCGActor = InPCGComponent->GetOwner();
+		const AActor* PCGActor = InPCGComponent->GetOwner();
 		check(PCGActor);
 
 		AActor* PCGActorCopy = DuplicateObject<AActor>(PCGActor, GetTransientPackage());
@@ -91,6 +91,10 @@ namespace PCGDeterminismTests
 		check(PCGComponentCopy);
 		PCGComponentCopy->SetFlags(RF_Transient);
 		PCGComponentCopy->SetIsPartitioned(false);
+		if (PCGActor->GetWorld())
+		{
+			PCGComponentCopy->RegisterComponentWithWorld(PCGActor->GetWorld());
+		}
 
 		FPCGGraphExecutor Executor = FPCGGraphExecutor(PCGComponentCopy);
 
@@ -205,6 +209,10 @@ namespace PCGDeterminismTests
 
 		// Clean up anything generated
 		PCGComponentCopy->Cleanup();
+		if (PCGComponentCopy->IsRegistered())
+		{
+			PCGComponentCopy->UnregisterComponent();
+		}
 
 		// Finalize the results
 		UpdateTestResults(Defaults::GraphResultName, OutResult, HighestDeterminismLevel);
