Fixed issue where attribute creation in the data table loader would not actually use the proper data and cause a crash for complex types (e.g. string)

diff --git a/Engine/Plugins/Experimental/PCG/Source/PCG/Private/Elements/IO/PCGDataTableElement.cpp b/Engine/Plugins/Experimental/PCG/Source/PCG/Private/Elements/IO/PCGDataTableElement.cpp
--- a/Engine/Plugins/Experimental/PCG/Source/PCG/Private/Elements/IO/PCGDataTableElement.cpp
+++ b/Engine/Plugins/Experimental/PCG/Source/PCG/Private/Elements/IO/PCGDataTableElement.cpp
@@ -98,7 +98,8 @@ bool FPCGLoadDataTableElement::PrepareLoad(FPCGExternalDataContext* Context) con
 		if (PointPropertySelector.GetSelection() == EPCGAttributePropertySelection::Attribute)
 		{
 			// Create attribute with the appropriate type
-			PointMetadata->CreateAttributeFromProperty(FName(FieldName), RowStruct, *FieldIt);
+			const void* FirstRow = DataTable->GetRowMap().CreateConstIterator().Value();
+			PointMetadata->CreateAttributeFromDataProperty(FName(FieldName), FirstRow, *FieldIt);
 		}
 
 		TUniquePtr<IPCGAttributeAccessor> PointPropertyAccessor = PCGAttributeAccessorHelpers::CreateAccessor(PointData, PointPropertySelector);
