[PCG] Fix broken callbacks on duplicate

#rnx

diff --git a/Engine/Plugins/Experimental/PCG/Source/PCG/Private/PCGComponent.cpp b/Engine/Plugins/Experimental/PCG/Source/PCG/Private/PCGComponent.cpp
--- a/Engine/Plugins/Experimental/PCG/Source/PCG/Private/PCGComponent.cpp
+++ b/Engine/Plugins/Experimental/PCG/Source/PCG/Private/PCGComponent.cpp
@@ -2380,7 +2380,7 @@ void FPCGComponentInstanceData::ApplyToComponent(UActorComponent* Component, con
 		// Reconnect callbacks
 		if (PCGComponent->GraphInstance)
 		{
-			PCGComponent->GraphInstance->FixCallbacks();
+			PCGComponent->GraphInstance->SetupCallbacks();
 			PCGComponent->GraphInstance->OnGraphChangedDelegate.RemoveAll(PCGComponent);
 			PCGComponent->GraphInstance->OnGraphChangedDelegate.AddUObject(PCGComponent, &UPCGComponent::OnGraphChanged);
 		}
diff --git a/Engine/Plugins/Experimental/PCG/Source/PCG/Private/PCGGraph.cpp b/Engine/Plugins/Experimental/PCG/Source/PCG/Private/PCGGraph.cpp
--- a/Engine/Plugins/Experimental/PCG/Source/PCG/Private/PCGGraph.cpp
+++ b/Engine/Plugins/Experimental/PCG/Source/PCG/Private/PCGGraph.cpp
@@ -1180,29 +1180,37 @@ void UPCGGraphInstance::PostLoad()
 	RefreshParameters(EPCGGraphParameterEvent::GraphPostLoad);
 
 #if WITH_EDITOR
-	if (Graph)
-	{
-		Graph->OnGraphChangedDelegate.AddUObject(this, &UPCGGraphInstance::OnGraphChanged);
-		Graph->OnGraphGridSizesChangedDelegate.AddUObject(this, &UPCGGraphInstance::OnGraphGridSizesChanged);
-		Graph->OnGraphParametersChangedDelegate.AddUObject(this, &UPCGGraphInstance::OnGraphParametersChanged);
-	}
+	SetupCallbacks();
 #endif // WITH_EDITOR
 }
 
 void UPCGGraphInstance::BeginDestroy()
 {
 #if WITH_EDITOR
-	if (Graph)
-	{
-		Graph->OnGraphChangedDelegate.RemoveAll(this);
-		Graph->OnGraphGridSizesChangedDelegate.RemoveAll(this);
-		Graph->OnGraphParametersChangedDelegate.RemoveAll(this);
-	}
+	TeardownCallbacks();
 #endif // WITH_EDITOR
 
 	Super::BeginDestroy();
 }
 
+void UPCGGraphInstance::PostDuplicate(bool bDuplicateForPIE)
+{
+	Super::PostDuplicate(bDuplicateForPIE);
+
+#if WITH_EDITOR
+	SetupCallbacks();
+#endif // WITH_EDITOR
+}
+
+void UPCGGraphInstance::PostEditImport()
+{
+	Super::PostEditImport();
+
+#if WITH_EDITOR
+	SetupCallbacks();
+#endif // WITH_EDITOR
+}
+
 #if WITH_EDITOR
 void UPCGGraphInstance::PreEditChange(FProperty* InProperty)
 {
@@ -1215,9 +1223,7 @@ void UPCGGraphInstance::PreEditChange(FProperty* InProperty)
 
 	if (InProperty->GetFName() == GET_MEMBER_NAME_CHECKED(UPCGGraphInstance, Graph) && Graph)
 	{
-		Graph->OnGraphChangedDelegate.RemoveAll(this);
-		Graph->OnGraphGridSizesChangedDelegate.RemoveAll(this);
-		Graph->OnGraphParametersChangedDelegate.RemoveAll(this);
+		TeardownCallbacks();
 	}
 }
 
@@ -1232,9 +1238,7 @@ void UPCGGraphInstance::PostEditChangeProperty(FPropertyChangedEvent& PropertyCh
 
 	if (PropertyChangedEvent.GetPropertyName() == GET_MEMBER_NAME_CHECKED(UPCGGraphInstance, Graph) && Graph)
 	{
-		Graph->OnGraphChangedDelegate.AddUObject(this, &UPCGGraphInstance::OnGraphChanged);
-		Graph->OnGraphGridSizesChangedDelegate.AddUObject(this, &UPCGGraphInstance::OnGraphGridSizesChanged);
-		Graph->OnGraphParametersChangedDelegate.AddUObject(this, &UPCGGraphInstance::OnGraphParametersChanged);
+		SetupCallbacks();
 
 		RefreshParameters(EPCGGraphParameterEvent::GraphChanged);
 	}
@@ -1248,12 +1252,7 @@ void UPCGGraphInstance::PreEditUndo()
 {
 	Super::PreEditUndo();
 
-	if (Graph)
-	{
-		Graph->OnGraphChangedDelegate.RemoveAll(this);
-		Graph->OnGraphGridSizesChangedDelegate.RemoveAll(this);
-		Graph->OnGraphParametersChangedDelegate.RemoveAll(this);
-	}
+	TeardownCallbacks();
 
 	UndoRedoGraphCache = Graph;
 }
@@ -1262,12 +1261,7 @@ void UPCGGraphInstance::PostEditUndo()
 {
 	Super::PostEditUndo();
 
-	if (Graph)
-	{
-		Graph->OnGraphChangedDelegate.AddUObject(this, &UPCGGraphInstance::OnGraphChanged);
-		Graph->OnGraphGridSizesChangedDelegate.AddUObject(this, &UPCGGraphInstance::OnGraphGridSizesChanged);
-		Graph->OnGraphParametersChangedDelegate.AddUObject(this, &UPCGGraphInstance::OnGraphParametersChanged);
-	}
+	SetupCallbacks();
 
 	// Since we don't know what happened, we need to notify any changes
 	NotifyGraphParametersChanged(EPCGGraphParameterEvent::GraphChanged, NAME_None);
@@ -1319,12 +1313,9 @@ void UPCGGraphInstance::TeardownCallbacks()
 	}
 }
 
-void UPCGGraphInstance::FixCallbacks()
+void UPCGGraphInstance::SetupCallbacks()
 {
-	// Start from a clean state.
-	TeardownCallbacks();
-
-	if (Graph)
+	if (Graph && !Graph->OnGraphChangedDelegate.IsBoundToObject(this))
 	{
 		Graph->OnGraphChangedDelegate.AddUObject(this, &UPCGGraphInstance::OnGraphChanged);
 		Graph->OnGraphGridSizesChangedDelegate.AddUObject(this, &UPCGGraphInstance::OnGraphGridSizesChanged);
@@ -1348,23 +1339,13 @@ void UPCGGraphInstance::SetGraph(UPCGGraphInterface* InGraph)
 	}
 
 #if WITH_EDITOR
-	if (Graph)
-	{
-		Graph->OnGraphChangedDelegate.RemoveAll(this);
-		Graph->OnGraphGridSizesChangedDelegate.RemoveAll(this);
-		Graph->OnGraphParametersChangedDelegate.RemoveAll(this);
-	}
+	TeardownCallbacks();
 #endif // WITH_EDITOR
 
 	Graph = InGraph;
 
 #if WITH_EDITOR
-	if (Graph)
-	{
-		Graph->OnGraphChangedDelegate.AddUObject(this, &UPCGGraphInstance::OnGraphChanged);
-		Graph->OnGraphGridSizesChangedDelegate.AddUObject(this, &UPCGGraphInstance::OnGraphGridSizesChanged);
-		Graph->OnGraphParametersChangedDelegate.AddUObject(this, &UPCGGraphInstance::OnGraphParametersChanged);
-	}
+	SetupCallbacks();
 #endif // WITH_EDITOR
 
 #if WITH_EDITOR
diff --git a/Engine/Plugins/Experimental/PCG/Source/PCG/Public/PCGGraph.h b/Engine/Plugins/Experimental/PCG/Source/PCG/Public/PCGGraph.h
--- a/Engine/Plugins/Experimental/PCG/Source/PCG/Public/PCGGraph.h
+++ b/Engine/Plugins/Experimental/PCG/Source/PCG/Public/PCGGraph.h
@@ -342,6 +342,8 @@ public:
 
 	// ~Begin UObject interface
 	virtual void PostLoad() override;
+	virtual void PostDuplicate(bool bDuplicateForPIE) override;
+	virtual void PostEditImport() override;
 	virtual void BeginDestroy() override;
 
 #if WITH_EDITOR
@@ -352,9 +354,7 @@ public:
 	virtual void PostEditUndo() override;
 	// ~End UObject interface
 
-	// Reconstruction script specific, to fix callbacks on reconstructed component
-	// and teardown callback on trashed component.
-	void FixCallbacks();
+	void SetupCallbacks();
 	void TeardownCallbacks();
 #endif
 
