[PCG] Fixed oversight in ActorSelector + set all unregister as delayed

* In operator== in the ActorSelector, the switch case was missing a case. Which was growing the register map indefinitly.
* In case of unregister/register in the same frame, we should not remove the component. Add it around a CVar if it as unintended consequences.

#rnx

diff --git a/Engine/Plugins/Experimental/PCG/Source/PCG/Private/Elements/PCGActorSelector.cpp b/Engine/Plugins/Experimental/PCG/Source/PCG/Private/Elements/PCGActorSelector.cpp
--- a/Engine/Plugins/Experimental/PCG/Source/PCG/Private/Elements/PCGActorSelector.cpp
+++ b/Engine/Plugins/Experimental/PCG/Source/PCG/Private/Elements/PCGActorSelector.cpp
@@ -273,8 +273,14 @@ bool FPCGActorSelectionKey::operator==(const FPCGActorSelectionKey& InOther) con
 		return Tag == InOther.Tag;
 	case EPCGActorSelection::ByClass:
 		return ActorSelectionClass == InOther.ActorSelectionClass;
+	case EPCGActorSelection::Unknown: // Fall-through
+	case EPCGActorSelection::ByName:
+		return true;
 	default:
-		return false;
+	{
+		checkNoEntry();
+		return true;
+	}
 	}
 }
 
diff --git a/Engine/Plugins/Experimental/PCG/Source/PCG/Private/PCGActorAndComponentMapping.cpp b/Engine/Plugins/Experimental/PCG/Source/PCG/Private/PCGActorAndComponentMapping.cpp
--- a/Engine/Plugins/Experimental/PCG/Source/PCG/Private/PCGActorAndComponentMapping.cpp
+++ b/Engine/Plugins/Experimental/PCG/Source/PCG/Private/PCGActorAndComponentMapping.cpp
@@ -38,6 +38,11 @@ namespace PCGActorAndComponentMapping
 		false,
 		TEXT("If depencencies are being unstable, disable the tracking, allowing people to continue working while we investigate."));
 #endif // WITH_EDITOR
+
+	static TAutoConsoleVariable<bool> CVarDisableDelayedUnregister(
+		TEXT("pcg.DisableDelayedUnregister"),
+		false,
+		TEXT("If delayed unregister for all is introducing bad behavior, disables it, allowing people to continue working while we investigate."));
 }
 
 UPCGActorAndComponentMapping::UPCGActorAndComponentMapping(UPCGSubsystem* InPCGSubsystem)
@@ -291,12 +296,18 @@ void UPCGActorAndComponentMapping::UnregisterPCGComponent(UPCGComponent* InCompo
 
 	if ((PartitionedOctree.Contains(InComponent) || NonPartitionedOctree.Contains(InComponent)))
 	{
-		// We also need to check that our current PCG Component is not deleted while being reconstructed by a construction script.
-		// If so, it will be "re-created" at some point with the same properties.
-		// In this particular case, we don't remove the PCG component from the octree and we won't delete the mapping, but mark it to be removed
-		// at next Subsystem tick. If we call "RemapPCGComponent" before, we will re-connect everything correctly.
-		// Ignore this if we force (aka when we actually unregister the delayed one)
-		if (InComponent->IsCreatedByConstructionScript() && !bForce)
+		bool bShouldBeDelayed = true;
+		if (PCGActorAndComponentMapping::CVarDisableDelayedUnregister.GetValueOnAnyThread())
+		{
+			// We also need to check that our current PCG Component is not deleted while being reconstructed by a construction script.
+			// If so, it will be "re-created" at some point with the same properties.
+			// In this particular case, we don't remove the PCG component from the octree and we won't delete the mapping, but mark it to be removed
+			// at next Subsystem tick. If we call "RemapPCGComponent" before, we will re-connect everything correctly.
+			// Ignore this if we force (aka when we actually unregister the delayed one)
+			bShouldBeDelayed = InComponent->IsCreatedByConstructionScript();
+		}
+
+		if (!bForce && bShouldBeDelayed)
 		{
 			FScopeLock Lock(&DelayedComponentToUnregisterLock);
 			DelayedComponentToUnregister.Add(InComponent);
