[PCG] Fixed PCGGraphs not generating on drop

Also fix broken node colours that was impacting graph readability

#rnx

diff --git a/Engine/Plugins/Experimental/PCG/Source/PCGEditor/PCGEditor.Build.cs b/Engine/Plugins/Experimental/PCG/Source/PCGEditor/PCGEditor.Build.cs
--- a/Engine/Plugins/Experimental/PCG/Source/PCGEditor/PCGEditor.Build.cs
+++ b/Engine/Plugins/Experimental/PCG/Source/PCGEditor/PCGEditor.Build.cs
@@ -54,12 +54,14 @@ namespace UnrealBuildTool.Rules
 					"KismetWidgets",
 					"PCG",
 					"PropertyEditor",
-					"StructUtils",
-					"StructUtilsEditor",
 					"Slate",
 					"SlateCore",
 					"SourceControl",
+					"StructUtils",
+					"StructUtilsEditor",
 					"ToolMenus",
+					"TypedElementFramework",
+					"TypedElementRuntime",
 					"UnrealEd",
 				}
 			);
diff --git a/Engine/Plugins/Experimental/PCG/Source/PCGEditor/Private/PCGEditorSettings.cpp b/Engine/Plugins/Experimental/PCG/Source/PCGEditor/Private/PCGEditorSettings.cpp
--- a/Engine/Plugins/Experimental/PCG/Source/PCGEditor/Private/PCGEditorSettings.cpp
+++ b/Engine/Plugins/Experimental/PCG/Source/PCGEditor/Private/PCGEditorSettings.cpp
@@ -23,9 +23,8 @@ UPCGEditorSettings::UPCGEditorSettings(const FObjectInitializer& ObjectInitializ
 	SubgraphNodeColor = FLinearColor(1.0f, 0.05f, 0.05f);
 	ParamDataNodeColor = FLinearColor(1.0f, 0.38f, 0.02f);
 	DebugNodeColor = FLinearColor(1.0f, 0.0f, 1.0f);
-	ControlFlowNodeColor = FLinearColor(0.66f, .6f, 0.15f);
-	// HiGen are also a subset of Control Flow
-	HierarchicalGenerationNodeColor = ControlFlowNodeColor;
+	ControlFlowNodeColor = FLinearColor(0.0f, 1.0f, 0.0f);
+	HierarchicalGenerationNodeColor = FLinearColor(1.0f, 0.132868f, 0.0f);
 
 	DefaultPinColor = FLinearColor(0.29f, 0.29f, 0.29f);
 	SpatialDataPinColor = FLinearColor(1.0f, 1.0f, 1.0f);
diff --git a/Engine/Plugins/Experimental/PCG/Source/PCGEditor/Private/PCGVolumeFactory.cpp b/Engine/Plugins/Experimental/PCG/Source/PCGEditor/Private/PCGVolumeFactory.cpp
--- a/Engine/Plugins/Experimental/PCG/Source/PCGEditor/Private/PCGVolumeFactory.cpp
+++ b/Engine/Plugins/Experimental/PCG/Source/PCGEditor/Private/PCGVolumeFactory.cpp
@@ -5,8 +5,13 @@
 #include "PCGComponent.h"
 #include "PCGEngineSettings.h"
 #include "PCGGraph.h"
+#include "PCGSubsystem.h"
 #include "PCGVolume.h"
 
+#include "Elements/Framework/TypedElementRegistry.h"
+#include "Elements/Interfaces/TypedElementObjectInterface.h"
+#include "Subsystems/PlacementSubsystem.h"
+
 #define LOCTEXT_NAMESPACE "PCGVolumeFactory"
 
 UPCGVolumeFactory::UPCGVolumeFactory(const FObjectInitializer& ObjectInitializer)
@@ -33,29 +38,63 @@ bool UPCGVolumeFactory::CanCreateActorFrom(const FAssetData& AssetData, FText& O
 	return true;
 }
 
-void UPCGVolumeFactory::PostSpawnActor(UObject* Asset, AActor* NewActor)
+void UPCGVolumeFactory::PostPlaceAsset(TArrayView<const FTypedElementHandle> InElementHandles, const FAssetPlacementInfo& InPlacementInfo, const FPlacementOptions& InPlacementOptions)
 {
-	Super::PostSpawnActor(Asset, NewActor);
+	Super::PostPlaceAsset(InElementHandles, InPlacementInfo, InPlacementOptions);
 
-	// Disable auto-refreshing on preview actors until we have something more robust on the execution side.
-	if (NewActor && NewActor->bIsEditorPreviewActor)
+	// Preview elements are created while dragging an asset, but before dropping.
+	if (InPlacementOptions.bIsCreatingPreviewElements)
 	{
 		return;
 	}
 
-	if (UPCGGraphInterface* PCGGraph = Cast<UPCGGraphInterface>(Asset))
+	for (const FTypedElementHandle& PlacedElement : InElementHandles)
 	{
+		TTypedElement<ITypedElementObjectInterface> ObjectInterface = UTypedElementRegistry::GetInstance()->GetElement<ITypedElementObjectInterface>(PlacedElement);
+		AActor* NewActor = ObjectInterface ? ObjectInterface.GetObjectAs<AActor>() : nullptr;
+		if (!NewActor)
+		{
+			continue;
+		}
+
+		UPCGGraphInterface* PCGGraph = Cast<UPCGGraphInterface>(InPlacementInfo.AssetToPlace.GetAsset());
 		const UPCGEngineSettings* Settings = GetDefault<UPCGEngineSettings>();
-		
+		if (!PCGGraph || !Settings)
+		{
+			continue;
+		}
+
 		APCGVolume* PCGVolume = CastChecked<APCGVolume>(NewActor);
 		PCGVolume->SetActorScale3D(Settings->VolumeScale);
-		
+
 		UPCGComponent* PCGComponent = CastChecked<UPCGComponent>(PCGVolume->GetComponentByClass(UPCGComponent::StaticClass()));
 		PCGComponent->SetGraph(PCGGraph);
 
 		if (Settings->bGenerateOnDrop)
 		{
-			PCGComponent->Generate();
+			if (UPCGSubsystem* Subsystem = PCGComponent->GetSubsystem())
+			{
+				TWeakObjectPtr<UPCGComponent> ComponentPtr(PCGComponent);
+
+				// Schedule a task to generate this component.
+				// We cannot generate the component right away because after PostSpawnActor is called, all
+				// actor components are unregistered and re-registered, which cancels component generation
+				Subsystem->ScheduleGeneric([ComponentPtr]()
+				{
+					if (UPCGComponent* Component = ComponentPtr.Get())
+					{
+						// If the component is not valid anymore, just early out.
+						if (!IsValid(Component))
+						{
+							return true;
+						}
+
+						Component->Generate();
+					}
+
+					return true;
+				}, PCGComponent, /*TaskDependencies=*/{});
+			}
 		}
 	}
 }
diff --git a/Engine/Plugins/Experimental/PCG/Source/PCGEditor/Private/PCGVolumeFactory.h b/Engine/Plugins/Experimental/PCG/Source/PCGEditor/Private/PCGVolumeFactory.h
--- a/Engine/Plugins/Experimental/PCG/Source/PCGEditor/Private/PCGVolumeFactory.h
+++ b/Engine/Plugins/Experimental/PCG/Source/PCGEditor/Private/PCGVolumeFactory.h
@@ -6,6 +6,9 @@
 
 #include "PCGVolumeFactory.generated.h"
 
+struct FPlacementOptions;
+struct FTypedElementHandle;
+
 UCLASS(MinimalAPI, config=Editor)
 class UPCGVolumeFactory : public UActorFactoryBoxVolume
 {
@@ -16,6 +19,6 @@ public:
 
 	//~ Begin UActorFactory Interface
 	virtual bool CanCreateActorFrom(const FAssetData& AssetData, FText& OutErrorMsg) override;
-	virtual void PostSpawnActor(UObject* Asset, AActor* NewActor) override;
+	virtual void PostPlaceAsset(TArrayView<const FTypedElementHandle> InElementHandles, const FAssetPlacementInfo& InPlacementInfo, const FPlacementOptions& InPlacementOptions) override;
 	//~ End UActorFactory Interface
 };
