Early review for undo/redo issue fix. Editor team is looking into adding support to only update the subobject tree, if we cant get something in place for 5.3 we should disable the refresh entirely with this change.

[PCG] Added console variable to toggle component changed calls to the leveleditor and made it disabled by default. This was causing issues with undo/redo transaction blocks while scrubbing values since it would regenerate the graph and refresh the details during interaction. Disabling this will cause the subobject tree to not refresh and show generated ISMC's until object has been reselected.

diff --git a/Engine/Plugins/Experimental/PCG/Source/PCGEditor/Private/Details/PCGComponentDetails.cpp b/Engine/Plugins/Experimental/PCG/Source/PCGEditor/Private/Details/PCGComponentDetails.cpp
--- a/Engine/Plugins/Experimental/PCG/Source/PCGEditor/Private/Details/PCGComponentDetails.cpp
+++ b/Engine/Plugins/Experimental/PCG/Source/PCGEditor/Private/Details/PCGComponentDetails.cpp
@@ -12,6 +12,11 @@
 
 #define LOCTEXT_NAMESPACE "PCGComponentDetails"
 
+static TAutoConsoleVariable<bool> CVarBroadcastComponentChanged(
+		TEXT("pcg.ComponentDetails.BroadcastComponentChanged"),
+		false,
+		TEXT("Controls whether we should broadcast component changes to the level editor"));
+
 TSharedRef<IDetailCustomization> FPCGComponentDetails::MakeInstance()
 {
 	return MakeShareable(new FPCGComponentDetails());
@@ -208,8 +213,11 @@ void FPCGComponentDetails::OnGraphChanged(UPCGComponent* InComponent)
 	}
 
 	// Notify editor that some components might have changed
-	FLevelEditorModule& LevelEditor = FModuleManager::LoadModuleChecked<FLevelEditorModule>("LevelEditor");
-	LevelEditor.BroadcastComponentsEdited();
+	if (CVarBroadcastComponentChanged.GetValueOnGameThread())
+	{
+		FLevelEditorModule& LevelEditor = FModuleManager::LoadModuleChecked<FLevelEditorModule>("LevelEditor");
+		LevelEditor.BroadcastComponentsEdited();
+	}
 }
 
 FReply FPCGComponentDetails::OnCleanupClicked()
