[PCG] Generated resources cleanup

* Un-outer managed resource objects to stop them hanging around
* Clean out managed resources once released

diff --git a/Engine/Plugins/Experimental/PCG/Source/PCG/Private/PCGComponent.cpp b/Engine/Plugins/Experimental/PCG/Source/PCG/Private/PCGComponent.cpp
--- a/Engine/Plugins/Experimental/PCG/Source/PCG/Private/PCGComponent.cpp
+++ b/Engine/Plugins/Experimental/PCG/Source/PCG/Private/PCGComponent.cpp
@@ -772,6 +772,11 @@ void UPCGComponent::CleanupLocalImmediate(bool bRemoveComponents)
 
 			if (!Resource || Resource->Release(bRemoveComponents, ActorsToDelete))
 			{
+				if (Resource)
+				{
+					Resource->Rename(nullptr, nullptr, REN_ForceNoResetLoaders | REN_DoNotDirty | REN_DontCreateRedirectors | REN_NonTransactional);
+				}
+
 				GeneratedResources.RemoveAtSwap(ResourceIndex);
 			}
 		}
@@ -860,6 +865,11 @@ FPCGTaskId UPCGComponent::CreateCleanupTask(bool bRemoveComponents, const TArray
 
 				if (!Resource || Resource->Release(bRemoveComponents, Context->ActorsToDelete))
 				{
+					if (Resource)
+					{
+						Resource->Rename(nullptr, nullptr, REN_ForceNoResetLoaders | REN_DoNotDirty | REN_DontCreateRedirectors | REN_NonTransactional);
+					}
+
 					ThisComponent->GeneratedResources.RemoveAtSwap(Context->ResourceIndex);
 				}
 
@@ -969,6 +979,11 @@ void UPCGComponent::CleanupUnusedManagedResources()
 
 			if (!Resource || Resource->ReleaseIfUnused(ActorsToDelete))
 			{
+				if (Resource)
+				{
+					Resource->Rename(nullptr, nullptr, REN_ForceNoResetLoaders | REN_DoNotDirty | REN_DontCreateRedirectors | REN_NonTransactional);
+				}
+
 				GeneratedResources.RemoveAtSwap(ResourceIndex);
 			}
 		}
diff --git a/Engine/Plugins/Experimental/PCG/Source/PCG/Private/PCGManagedResource.cpp b/Engine/Plugins/Experimental/PCG/Source/PCG/Private/PCGManagedResource.cpp
--- a/Engine/Plugins/Experimental/PCG/Source/PCG/Private/PCGManagedResource.cpp
+++ b/Engine/Plugins/Experimental/PCG/Source/PCG/Private/PCGManagedResource.cpp
@@ -206,13 +206,14 @@ bool UPCGManagedComponent::Release(bool bHardRelease, TSet<TSoftObjectPtr<AActor
 	TRACE_CPUPROFILER_EVENT_SCOPE(UPCGManagedComponent::Release);
 
 	const bool bSupportsComponentReset = SupportsComponentReset();
-	const bool bDeleteComponent = bHardRelease || !bSupportsComponentReset;
+	bool bDeleteComponent = bHardRelease || !bSupportsComponentReset;
 
 	if (GeneratedComponent.IsValid())
 	{
 		if (bDeleteComponent)
 		{
 			GeneratedComponent->DestroyComponent();
+			ForgetComponent();
 		}
 		else
 		{
@@ -221,6 +222,11 @@ bool UPCGManagedComponent::Release(bool bHardRelease, TSet<TSoftObjectPtr<AActor
 			GeneratedComponent->ComponentTags.Add(PCGHelpers::MarkedForCleanupPCGTag);
 		}
 	}
+	else
+	{
+		// Dead component reference - clear it out.
+		bDeleteComponent = true;
+	}
 
 	return bDeleteComponent;
 }
