[PCG] Fix issue with first move from a tracked actor not refreshing its old position correctly

* Bounds were not valid when the actor was registered for tracking, because its components were not registered yet
* We now delay the register of the actor if its components are not registered yet
* It will be registered on the next tick (or on the next tick when those components are registered)
* It's behind a CVar, so we can disable it if it has unwanted side effects.

#rnx

diff --git a/Engine/Plugins/Experimental/PCG/Source/PCG/Private/PCGActorAndComponentMapping.cpp b/Engine/Plugins/Experimental/PCG/Source/PCG/Private/PCGActorAndComponentMapping.cpp
--- a/Engine/Plugins/Experimental/PCG/Source/PCG/Private/PCGActorAndComponentMapping.cpp
+++ b/Engine/Plugins/Experimental/PCG/Source/PCG/Private/PCGActorAndComponentMapping.cpp
@@ -37,6 +37,11 @@ namespace PCGActorAndComponentMapping
 		TEXT("pcg.DisableObjectDependenciesTracking"),
 		false,
 		TEXT("If depencencies are being unstable, disable the tracking, allowing people to continue working while we investigate."));
+
+	static TAutoConsoleVariable<bool> CVarDisableDelayedActorRegistering(
+		TEXT("pcg.DisableDelayedActorRegistering"),
+		false,
+		TEXT("If delayed actor registering when their components aren't registered yet is introducing bad behavior, disables it, allowing people to continue working while we investigate."));
 #endif // WITH_EDITOR
 
 	static TAutoConsoleVariable<bool> CVarDisableDelayedUnregister(
@@ -69,6 +74,10 @@ void UPCGActorAndComponentMapping::Tick()
 	{
 		UnregisterPCGComponent(Component, /*bForce=*/true);
 	}
+
+#if WITH_EDITOR
+	AddDelayedActors();
+#endif // WITH_EDITOR
 }
 
 TArray<FPCGTaskId> UPCGActorAndComponentMapping::DispatchToRegisteredLocalComponents(UPCGComponent* OriginalComponent, const TFunction<FPCGTaskId(UPCGComponent*)>& InFunc) const
@@ -624,10 +633,12 @@ void UPCGActorAndComponentMapping::RegisterOrUpdateTracking(UPCGComponent* InCom
 	FPCGActorSelectionKey OwnerKey = FPCGActorSelectionKey(EPCGActorFilter::Self);
 	KeysToComponentsMap.FindOrAdd(OwnerKey).Add(InComponent);
 
+	const bool bDisableDelayedActorRegistering = PCGActorAndComponentMapping::CVarDisableDelayedActorRegistering.GetValueOnAnyThread();
+
 	// And we also need to find all actors that should be tracked
 	if (UPCGGraph* PCGGraph = InComponent->GetGraph())
 	{
-		auto FindActorsAndTrack = [this, InComponent, bInShouldDirtyActors](const FPCGActorSelectionKey& InKey, const TArray<FPCGSettingsAndCulling>& InSettingsAndCulling)
+		auto FindActorsAndTrack = [this, InComponent, bInShouldDirtyActors, bDisableDelayedActorRegistering](const FPCGActorSelectionKey& InKey, const TArray<FPCGSettingsAndCulling>& InSettingsAndCulling)
 		{
 			// InKey provide the info for selecting a given actor.
 			// We reconstruct the selector settings from this key, and we also force it to SelectMultiple, since
@@ -649,6 +660,17 @@ void UPCGActorAndComponentMapping::RegisterOrUpdateTracking(UPCGComponent* InCom
 
 			for (AActor* Actor : AllActors)
 			{
+				if (!Actor)
+				{
+					continue;
+				}
+
+				if (!Actor->HasActorRegisteredAllComponents() && !bDisableDelayedActorRegistering)
+				{
+					DelayedAddedActors.Emplace({ Actor, bInShouldDirtyActors });
+					continue;
+				}
+
 				if (bShouldCull)
 				{
 					CulledTrackedActorsToComponentsMap.FindOrAdd(Actor).Add(InComponent);
@@ -800,23 +822,44 @@ void UPCGActorAndComponentMapping::TeardownTrackingCallbacks()
 	FCoreUObjectDelegates::OnPreObjectPropertyChanged.RemoveAll(this);
 }
 
-void UPCGActorAndComponentMapping::AddActorsPostInit()
+void UPCGActorAndComponentMapping::AddDelayedActors()
 {
 	// Safeguard, we can't add delayed actors if the subsystem is not initialized
-	if (!PCGSubsystem || !PCGSubsystem->IsInitialized())
+	if (!PCGSubsystem || !PCGSubsystem->IsInitialized() || DelayedAddedActors.IsEmpty())
 	{
 		return;
 	}
 
-	for (TObjectKey<AActor>& ActorPtr : DelayedAddedActors)
+	TSet<TTuple<TObjectKey<AActor>, bool>> StillDelayedActors;
+	const bool bDisableDelayedActorRegistering = PCGActorAndComponentMapping::CVarDisableDelayedActorRegistering.GetValueOnAnyThread();
+
+	for (TTuple<TObjectKey<AActor>, bool>& ActorPtrAndShouldDirty : DelayedAddedActors)
 	{
-		OnActorAdded(ActorPtr.ResolveObjectPtr());
+		AActor* Actor = ActorPtrAndShouldDirty.Get<0>().ResolveObjectPtr();
+		if (!Actor)
+		{
+			continue;
+		}
+
+		if (!Actor->HasActorRegisteredAllComponents() && !bDisableDelayedActorRegistering)
+		{
+			StillDelayedActors.Add(ActorPtrAndShouldDirty);
+		}
+		else
+		{
+			OnActorAdded_Internal(Actor, ActorPtrAndShouldDirty.Get<1>());
+		}
 	}
 
-	DelayedAddedActors.Empty();
+	DelayedAddedActors = MoveTemp(StillDelayedActors);
 }
 
 void UPCGActorAndComponentMapping::OnActorAdded(AActor* InActor)
+{
+	OnActorAdded_Internal(InActor);
+}
+
+void UPCGActorAndComponentMapping::OnActorAdded_Internal(AActor* InActor, bool bShouldDirty)
 {
 	TRACE_CPUPROFILER_EVENT_SCOPE(UPCGActorAndComponentMapping::OnActorAdded);
 
@@ -836,11 +879,11 @@ void UPCGActorAndComponentMapping::OnActorAdded(AActor* InActor)
 	// If the subsystem is not initialized, wait for it to be, and store all the actors to check
 	if (!PCGSubsystem->IsInitialized())
 	{
-		DelayedAddedActors.Add(InActor);
+		DelayedAddedActors.Emplace({ InActor, bShouldDirty });
 		return;
 	}
 
-	if (AddOrUpdateTrackedActor(InActor))
+	if (AddOrUpdateTrackedActor(InActor) && bShouldDirty)
 	{
 		// Finally notify them all
 		OnActorChanged(InActor, /*bInHasMoved=*/ false);
diff --git a/Engine/Plugins/Experimental/PCG/Source/PCG/Private/PCGSubsystem.cpp b/Engine/Plugins/Experimental/PCG/Source/PCG/Private/PCGSubsystem.cpp
--- a/Engine/Plugins/Experimental/PCG/Source/PCG/Private/PCGSubsystem.cpp
+++ b/Engine/Plugins/Experimental/PCG/Source/PCG/Private/PCGSubsystem.cpp
@@ -156,11 +156,6 @@ void UPCGSubsystem::PostInitialize()
 			});
 		}
 	}
-
-#if WITH_EDITOR
-	// When the PCGWorldActor is loaded, we can add delayed actors.
-	ActorAndComponentMapping.AddActorsPostInit();
-#endif // WITH_EDITOR
 }
 
 UPCGSubsystem* UPCGSubsystem::GetInstance(UWorld* World)
diff --git a/Engine/Plugins/Experimental/PCG/Source/PCG/Public/PCGActorAndComponentMapping.h b/Engine/Plugins/Experimental/PCG/Source/PCG/Public/PCGActorAndComponentMapping.h
--- a/Engine/Plugins/Experimental/PCG/Source/PCG/Public/PCGActorAndComponentMapping.h
+++ b/Engine/Plugins/Experimental/PCG/Source/PCG/Public/PCGActorAndComponentMapping.h
@@ -48,7 +48,7 @@ public:
 	void RegisterTrackingCallbacks();
 	void TeardownTrackingCallbacks();
 
-	void AddActorsPostInit();
+	void AddDelayedActors();
 
 	/** Will register/update tracking if a component was registered/updated. */
 	void RegisterOrUpdateTracking(UPCGComponent* InComponent, bool bInShouldDirtyActors);
@@ -119,6 +119,7 @@ private:
 	bool UnregisterActor(AActor* InActor);
 
 	void OnActorAdded(AActor* InActor);
+	void OnActorAdded_Internal(AActor* InActor, bool bShouldDirty = true);
 	void OnActorDeleted(AActor* InActor);
 	void OnActorMoved(AActor* InActor);
 	void OnLandscapeChanged(ALandscapeProxy* InLandscape, const FLandscapeProxyComponentDataChangedParams& InChangeParams);
@@ -182,7 +183,8 @@ private:
 
 	mutable FRWLock TrackedComponentsLock;
 
-	TSet<TObjectKey<AActor>> DelayedAddedActors;
+	// Keep track of actors that aren't yet ready (or if the subsystem is not yet ready) and add them in next tick.
+	TSet<TTuple<TObjectKey<AActor>, bool>> DelayedAddedActors;
 
 	/** Keep a mapping between tracked actors and their dependencies. */
 	TMap<TObjectKey<AActor>, TSet<TObjectPtr<UObject>>> TrackedActorsToDependenciesMap;
