[PCG] Fix static analysis in unit test

#rnx

diff --git a/Engine/Plugins/Experimental/PCG/Source/PCG/Private/Tests/Subgraphs/PCGCollapseSubgraphTest.cpp b/Engine/Plugins/Experimental/PCG/Source/PCG/Private/Tests/Subgraphs/PCGCollapseSubgraphTest.cpp
--- a/Engine/Plugins/Experimental/PCG/Source/PCG/Private/Tests/Subgraphs/PCGCollapseSubgraphTest.cpp
+++ b/Engine/Plugins/Experimental/PCG/Source/PCG/Private/Tests/Subgraphs/PCGCollapseSubgraphTest.cpp
@@ -54,12 +54,16 @@ namespace PCGCollapseSubgraphTests
 		if (!InTestClass->TestNotNull(*FString::Printf(TEXT("Input Node of class %s is valid"), *InputClass::StaticClass()->GetName()), InNode)) { return false; }
 		if (!InTestClass->TestNotNull(*FString::Printf(TEXT("Output Node of class %s is valid"), *OutputClass::StaticClass()->GetName()), OutNode)) { return false; }
 
+		check(InNode && OutNode);
+
 		UPCGPin* InPin = InNode->GetOutputPin(InLabel);
 		UPCGPin* OutPin = OutNode->GetInputPin(OutLabel);
 
 		if (!InTestClass->TestTrue(*FString::Printf(TEXT("InLabel (%s) exists on node %s"), *InLabel.ToString(), *InNode->GetNodeTitle().ToString()), (InLabel == NAME_None) || InPin)) { return false; }
 		if (!InTestClass->TestTrue(*FString::Printf(TEXT("OutLabel (%s) exists on node %s"), *OutLabel.ToString(), *OutNode->GetNodeTitle().ToString()), (OutLabel == NAME_None) || OutPin)) { return false; }
 
+		check(InPin || OutPin);
+
 		const TArray<TObjectPtr<UPCGEdge>>& Edges = InPin ? InPin->Edges : OutPin->Edges;
 
 		auto Predicate = [InPin, OutPin](const TObjectPtr<UPCGEdge>& InEdge)
