[PCG] Fix outstanding bug in the attribute list view

* Any None attribute that has extractor (like Vector) were not displaying properly

#rnx

diff --git a/Engine/Plugins/Experimental/PCG/Source/PCGEditor/Private/SPCGEditorGraphAttributeListView.cpp b/Engine/Plugins/Experimental/PCG/Source/PCGEditor/Private/SPCGEditorGraphAttributeListView.cpp
--- a/Engine/Plugins/Experimental/PCG/Source/PCGEditor/Private/SPCGEditorGraphAttributeListView.cpp
+++ b/Engine/Plugins/Experimental/PCG/Source/PCGEditor/Private/SPCGEditorGraphAttributeListView.cpp
@@ -1278,7 +1278,7 @@ void SPCGEditorGraphAttributeListView::AddMetadataColumn(const UPCGData* InPCGDa
 	FString ColumnIdString(OriginalColumnIdString);
 	if (InColumnId == NAME_None)
 	{
-		ColumnIdString = PCGEditorGraphAttributeListView::NoneAttributeId;
+		ColumnIdString = FString::Printf(TEXT("%s%s"), *PCGEditorGraphAttributeListView::NoneAttributeId, PostFix);
 		ColumnLabel = FText::Format(LOCTEXT("NoneLabelFormat", "{0}{1}"), UEnum::GetDisplayValueAsText(InMetadataType), FText::FromString(PostFix));
 	}
 	
