[PCG] Made the PCG graph object selectable by clicking the node editor background.

diff --git a/Engine/Plugins/Experimental/PCG/Source/PCGEditor/Private/PCGEditor.cpp b/Engine/Plugins/Experimental/PCG/Source/PCGEditor/Private/PCGEditor.cpp
--- a/Engine/Plugins/Experimental/PCG/Source/PCGEditor/Private/PCGEditor.cpp
+++ b/Engine/Plugins/Experimental/PCG/Source/PCGEditor/Private/PCGEditor.cpp
@@ -2022,26 +2022,33 @@ TSharedRef<SPCGEditorGraphLogView> FPCGEditor::CreateLogWidget()
 void FPCGEditor::OnSelectedNodesChanged(const TSet<UObject*>& NewSelection)
 {
 	TArray<TWeakObjectPtr<UObject>> SelectedObjects;
-	
-	for (UObject* Object : NewSelection)
+
+	if (NewSelection.Num() == 0)
+	{
+		SelectedObjects.Add(PCGGraphBeingEdited);
+	}
+	else
 	{
-		if (UPCGEditorGraphNodeBase* PCGGraphNode = Cast<UPCGEditorGraphNodeBase>(Object))
+		for (UObject* Object : NewSelection)
 		{
-			if (UPCGNode* PCGNode = PCGGraphNode->GetPCGNode())
+			if (UPCGEditorGraphNodeBase* PCGGraphNode = Cast<UPCGEditorGraphNodeBase>(Object))
 			{
-				if (PCGNode->IsInstance())
-				{
-					SelectedObjects.Add(Cast<UPCGSettingsInstance>(PCGNode->GetSettingsInterface()));
-				}
-				else
+				if (UPCGNode* PCGNode = PCGGraphNode->GetPCGNode())
 				{
-					SelectedObjects.Add(PCGNode->GetSettings());
+					if (PCGNode->IsInstance())
+					{
+						SelectedObjects.Add(Cast<UPCGSettingsInstance>(PCGNode->GetSettingsInterface()));
+					}
+					else
+					{
+						SelectedObjects.Add(PCGNode->GetSettings());
+					}
 				}
 			}
-		}
-		else if (UEdGraphNode* GraphNode = Cast<UEdGraphNode>(Object))
-		{
-			SelectedObjects.Add(GraphNode);
+			else if (UEdGraphNode* GraphNode = Cast<UEdGraphNode>(Object))
+			{
+				SelectedObjects.Add(GraphNode);
+			}
 		}
 	}
 
