Added a pcg asset loader node (proxy) that can be overriden.
Added some flexibility around data table to output param data instead of point data.
Added node to get single param data inside of a multi param data.

diff --git a/Engine/Plugins/Experimental/PCG/Source/PCG/Private/Elements/IO/PCGDataTableElement.cpp b/Engine/Plugins/Experimental/PCG/Source/PCG/Private/Elements/IO/PCGDataTableElement.cpp
--- a/Engine/Plugins/Experimental/PCG/Source/PCG/Private/Elements/IO/PCGDataTableElement.cpp
+++ b/Engine/Plugins/Experimental/PCG/Source/PCG/Private/Elements/IO/PCGDataTableElement.cpp
@@ -4,6 +4,7 @@
 
 #include "PCGComponent.h"
 #include "PCGModule.h"
+#include "PCGParamData.h"
 #include "Data/PCGPointData.h"
 #include "Elements/IO/PCGExternalDataContext.h"
 
@@ -25,6 +26,41 @@ FText UPCGLoadDataTableSettings::GetNodeTooltipText() const
 
 #endif // WITH_EDITOR
 
+TArray<FPCGPinProperties> UPCGLoadDataTableSettings::OutputPinProperties() const
+{
+	TArray<FPCGPinProperties> PinProperties;
+	PinProperties.Emplace(PCGPinConstants::DefaultOutputLabel, EPCGDataType::Any);
+
+	return PinProperties;
+}
+
+EPCGDataType UPCGLoadDataTableSettings::GetCurrentPinTypes(const UPCGPin* InPin) const
+{
+	check(InPin);
+	if (!InPin->IsOutputPin())
+	{
+		return Super::GetCurrentPinTypes(InPin);
+	}
+
+	const UEnum* DataTypeEnum = StaticEnum<EPCGDataType>();
+	const UEnum* ExclusiveDataTypeEnum = StaticEnum<EPCGExclusiveDataType>();
+
+	if (DataTypeEnum && ExclusiveDataTypeEnum)
+	{
+		FName ExclusiveDataTypeName = ExclusiveDataTypeEnum->GetNameByValue(static_cast<__underlying_type(EPCGExclusiveDataType)>(OutputType));
+		if (ExclusiveDataTypeName != NAME_None)
+		{
+			const int64 MatchingDataType = DataTypeEnum->GetValueByName(ExclusiveDataTypeName);
+			if (ensure(MatchingDataType != INDEX_NONE))
+			{
+				return static_cast<EPCGDataType>(MatchingDataType);
+			}
+		}
+	}
+
+	return Super::GetCurrentPinTypes(InPin);
+}
+
 FPCGElementPtr UPCGLoadDataTableSettings::CreateElement() const
 {
 	return MakeShared<FPCGLoadDataTableElement>();
@@ -60,13 +96,33 @@ bool FPCGLoadDataTableElement::PrepareLoad(FPCGExternalDataContext* Context) con
 	}
 
 	// 3-e. Create point data, and the mapping object
-	UPCGPointData* PointData = NewObject<UPCGPointData>();
-	check(PointData);
+	UPCGData* OutData = nullptr;
+	UPCGMetadata* OutMetadata = nullptr;
 
-	FPCGExternalDataContext::FPointDataAccessorsMapping& PointDataAccessorMapping = Context->PointDataAccessorsMapping.Emplace_GetRef();
-	PointDataAccessorMapping.PointData = PointData;
+	if (Settings->OutputType == EPCGExclusiveDataType::Point)
+	{
+		UPCGPointData* PointData = NewObject<UPCGPointData>();
+		check(PointData);
+		OutData = PointData;
+		OutMetadata = PointData->MutableMetadata();
+	}
+	else if(Settings->OutputType == EPCGExclusiveDataType::Param)
+	{
+		UPCGParamData* ParamData = NewObject<UPCGParamData>();
+		check(ParamData);
+		OutData = ParamData;
+		OutMetadata = ParamData->MutableMetadata();
+	}
+	else
+	{
+		ensure(0);
+		PCGE_LOG(Error, GraphAndLog, LOCTEXT("Output format not supported", "Unable to load data table to selected output type"));
+		return true;
+	}
 
-	UPCGMetadata* PointMetadata = PointData->MutableMetadata();
+	FPCGExternalDataContext::FPointDataAccessorsMapping& PointDataAccessorMapping = Context->PointDataAccessorsMapping.Emplace_GetRef();
+	PointDataAccessorMapping.Data = OutData;
+	PointDataAccessorMapping.Metadata = OutMetadata;
 
 	// 3. Build property to accessor mappings, including renamings, and create new attributes as needed
 	for (TFieldIterator<const FProperty> FieldIt(RowStruct, EFieldIterationFlags::IncludeSuper); FieldIt; ++FieldIt)
@@ -95,14 +151,14 @@ bool FPCGLoadDataTableElement::PrepareLoad(FPCGExternalDataContext* Context) con
 			PointPropertySelector.Update(FieldName);
 		}
 
-		if (PointPropertySelector.GetSelection() == EPCGAttributePropertySelection::Attribute)
+ 		if (PointPropertySelector.GetSelection() == EPCGAttributePropertySelection::Attribute)
 		{
 			// Create attribute with the appropriate type
 			const void* FirstRow = DataTable->GetRowMap().CreateConstIterator().Value();
 			PointMetadata->CreateAttributeFromDataProperty(FName(FieldName), FirstRow, *FieldIt);
 		}
 
-		TUniquePtr<IPCGAttributeAccessor> PointPropertyAccessor = PCGAttributeAccessorHelpers::CreateAccessor(PointData, PointPropertySelector);
+		TUniquePtr<IPCGAttributeAccessor> PointPropertyAccessor = PCGAttributeAccessorHelpers::CreateAccessor(OutData, PointPropertySelector);
 
 		if (!PointPropertyAccessor)
 		{
@@ -125,8 +181,8 @@ bool FPCGLoadDataTableElement::PrepareLoad(FPCGExternalDataContext* Context) con
 
 	if (!Context->PointDataAccessorsMapping.IsEmpty())
 	{
-		FPCGTaggedData& OutData = Context->OutputData.TaggedData.Emplace_GetRef();
-		OutData.Data = PointData;
+		FPCGTaggedData& OutTaggedData = Context->OutputData.TaggedData.Emplace_GetRef();
+		OutTaggedData.Data = OutData;
 
 		Context->bDataPrepared = true;
 	}
@@ -156,14 +212,41 @@ bool FPCGLoadDataTableElement::ExecuteLoad(FPCGExternalDataContext* Context) con
 	FPCGExternalDataContext::FPointDataAccessorsMapping& PointDataAccessorMapping = Context->PointDataAccessorsMapping[0];
 
 	// Preallocate the point array, and create the row keys
-	UPCGPointData* PointData = PointDataAccessorMapping.PointData;
-	check(PointData);
+	UPCGData* Data = PointDataAccessorMapping.Data;
+	check(Data);
 
-	TArray<FPCGPoint>& Points = PointDataAccessorMapping.PointData->GetMutablePoints();
-	if (Points.IsEmpty())
+	if (UPCGPointData* PointData = Cast<UPCGPointData>(Data))
+	{
+		TArray<FPCGPoint>& Points = PointData->GetMutablePoints();
+		if (Points.IsEmpty())
+		{
+			// Dummy implementation - note that we do initialize the points here in cases where some of the properties would not be set
+			Points.SetNum(DataTable->GetRowMap().Num());
+		}
+	}
+	else if (UPCGParamData* ParamData = Cast<UPCGParamData>(Data))
 	{
-		// Dummy implementation - note that we do initialize the points here in cases where some of the properties would not be set
-		Points.SetNum(DataTable->GetRowMap().Num());
+		check(ParamData->Metadata);
+		TArray<FName> AttributeNames;
+		TArray<EPCGMetadataTypes> AttributeTypes;
+		ParamData->Metadata->GetAttributes(AttributeNames, AttributeTypes);
+
+		TArray<FPCGMetadataAttributeBase*> Attributes;
+		for (const FName AttributeName : AttributeNames)
+		{
+			Attributes.Add(ParamData->Metadata->GetMutableAttribute(AttributeName));
+		}
+
+		for (int32 EntryIndex = 0; EntryIndex < DataTable->GetRowMap().Num(); ++EntryIndex)
+		{
+			const PCGMetadataEntryKey EntryKey = ParamData->Metadata->AddEntry();
+
+			// Need to create a value entry for each newly created entry in the attributes, otherwise we will not write properly to the output
+			for (FPCGMetadataAttributeBase* Attribute : Attributes)
+			{
+				Attribute->SetValueFromValueKey(EntryKey, PCGDefaultValueKey);
+			}
+		}
 	}
 
 	TArray<void*> AllRows;
diff --git a/Engine/Plugins/Experimental/PCG/Source/PCG/Private/Elements/IO/PCGExternalData.cpp b/Engine/Plugins/Experimental/PCG/Source/PCG/Private/Elements/IO/PCGExternalData.cpp
--- a/Engine/Plugins/Experimental/PCG/Source/PCG/Private/Elements/IO/PCGExternalData.cpp
+++ b/Engine/Plugins/Experimental/PCG/Source/PCG/Private/Elements/IO/PCGExternalData.cpp
@@ -48,13 +48,12 @@ bool FPCGExternalDataElement::ExecuteLoad(FPCGExternalDataContext* Context) cons
 	// Do in parallel according to the resources given in the context
 	for (FPCGExternalDataContext::FPointDataAccessorsMapping& PointDataAccessorMapping : Context->PointDataAccessorsMapping)
 	{
-		UPCGPointData* PointData = PointDataAccessorMapping.PointData;
-		UPCGMetadata* PointMetadata = PointData->MutableMetadata();
+		UPCGData* Data = PointDataAccessorMapping.Data;
 		TUniquePtr<const IPCGAttributeAccessorKeys>& RowKeys = PointDataAccessorMapping.RowKeys;
 
 		for (FPCGExternalDataContext::FRowToPointAccessors& RowToPointAccessor : PointDataAccessorMapping.RowToPointAccessors)
 		{
-			TUniquePtr<IPCGAttributeAccessorKeys> PointKeys = PCGAttributeAccessorHelpers::CreateKeys(PointData, RowToPointAccessor.Selector);
+			TUniquePtr<IPCGAttributeAccessorKeys> PointKeys = PCGAttributeAccessorHelpers::CreateKeys(Data, RowToPointAccessor.Selector);
 
 			auto Operation = [&RowKeys, &RowToPointAccessor, &PointKeys](auto Dummy)
 			{
diff --git a/Engine/Plugins/Experimental/PCG/Source/PCG/Private/Elements/PCGAttributeGetFromIndexElement.cpp b/Engine/Plugins/Experimental/PCG/Source/PCG/Private/Elements/PCGAttributeGetFromIndexElement.cpp
new file mode 100644
--- /dev/null
+++ b/Engine/Plugins/Experimental/PCG/Source/PCG/Private/Elements/PCGAttributeGetFromIndexElement.cpp
@@ -0,0 +1,96 @@
+// Copyright Epic Games, Inc. All Rights Reserved.
+
+#include "Elements/PCGAttributeGetFromIndexElement.h"
+
+#include "PCGContext.h"
+#include "PCGParamData.h"
+#include "PCGPin.h"
+
+#include UE_INLINE_GENERATED_CPP_BY_NAME(PCGAttributeGetFromIndexElement)
+
+#define LOCTEXT_NAMESPACE "PCGAttributeGetFromIndexElement"
+
+#if WITH_EDITOR
+FName UPCGAttributeGetFromIndexSettings::GetDefaultNodeName() const
+{
+	return TEXT("GetAttributeFromIndex");
+}
+
+FText UPCGAttributeGetFromIndexSettings::GetDefaultNodeTitle() const
+{
+	return LOCTEXT("NodeTitle", "Get Attribute Set from Index");
+}
+
+FText UPCGAttributeGetFromIndexSettings::GetNodeTooltipText() const
+{
+	return LOCTEXT("NodeTooltip", "Retrieves a single entry from an Attribute Set.");
+}
+#endif // WITH_EDITOR
+
+TArray<FPCGPinProperties> UPCGAttributeGetFromIndexSettings::InputPinProperties() const
+{
+	TArray<FPCGPinProperties> PinProperties;
+	PinProperties.Emplace(PCGPinConstants::DefaultInputLabel, EPCGDataType::Param);
+
+	return PinProperties;
+}
+
+TArray<FPCGPinProperties> UPCGAttributeGetFromIndexSettings::OutputPinProperties() const
+{
+	TArray<FPCGPinProperties> PinProperties;
+	PinProperties.Emplace(PCGPinConstants::DefaultOutputLabel, EPCGDataType::Param);
+
+	return PinProperties;
+}
+
+FPCGElementPtr UPCGAttributeGetFromIndexSettings::CreateElement() const
+{
+	return MakeShared<FPCGAttributeGetFromIndexElement>();
+}
+
+bool FPCGAttributeGetFromIndexElement::ExecuteInternal(FPCGContext* Context) const
+{
+	TRACE_CPUPROFILER_EVENT_SCOPE(FPCGAttributeGetFromIndexElement::Execute);
+
+	check(Context);
+
+	const UPCGAttributeGetFromIndexSettings* Settings = Context->GetInputSettings<UPCGAttributeGetFromIndexSettings>();
+	check(Settings);
+
+	TArray<FPCGTaggedData> Inputs = Context->InputData.GetInputsByPin(PCGPinConstants::DefaultInputLabel);
+	const int32 Index = Settings->Index;
+
+	for (const FPCGTaggedData& Input : Inputs)
+	{
+		const UPCGParamData* Param = Cast<const UPCGParamData>(Input.Data);
+
+		if (!Param || !Param->Metadata)
+		{
+			continue;
+		}
+
+		const UPCGMetadata* ParamMetadata = Param->Metadata;
+		const int64 ParamItemCount = ParamMetadata->GetLocalItemCount();
+
+		if (Index >= ParamItemCount)
+		{
+			PCGE_LOG(Warning, GraphAndLog, FText::Format(LOCTEXT("InvalidParamIndex", "Unable to retrieve entry {0} because there are {1} entries in the AttributeSet"), Index, ParamItemCount));
+			continue;
+		}
+
+		UPCGParamData* SubParam = NewObject<UPCGParamData>();
+		check(SubParam->Metadata);
+		SubParam->Metadata->AddAttributes(ParamMetadata);
+
+		const PCGMetadataEntryKey OriginalEntryKey = Index;
+		PCGMetadataEntryKey SingleEntryKey = SubParam->Metadata->AddEntry();
+		SubParam->Metadata->SetAttributes(OriginalEntryKey, ParamMetadata, SingleEntryKey);
+
+		FPCGTaggedData& TaggedData = Context->OutputData.TaggedData.Emplace_GetRef(Input);
+		TaggedData.Data = SubParam;
+	}
+
+	return true;
+}
+
+#undef LOCTEXT_NAMESPACE
\ No newline at end of file
diff --git a/Engine/Plugins/Experimental/PCG/Source/PCG/Private/Elements/PCGIndirectionElement.cpp b/Engine/Plugins/Experimental/PCG/Source/PCG/Private/Elements/PCGIndirectionElement.cpp
new file mode 100644
--- /dev/null
+++ b/Engine/Plugins/Experimental/PCG/Source/PCG/Private/Elements/PCGIndirectionElement.cpp
@@ -0,0 +1,173 @@
+// Copyright Epic Games, Inc. All Rights Reserved.
+
+#include "Elements/PCGIndirectionElement.h"
+
+#include "PCGPin.h"
+
+#include UE_INLINE_GENERATED_CPP_BY_NAME(PCGIndirectionElement)
+
+#define LOCTEXT_NAMESPACE "PCGIndirectionElement"
+
+#if WITH_EDITOR
+
+FName UPCGIndirectionSettings::GetDefaultNodeName() const
+{
+	return FName(TEXT("Proxy"));
+}
+
+FText UPCGIndirectionSettings::GetDefaultNodeTitle() const
+{
+	return LOCTEXT("NodeTitle", "Proxy");
+}
+
+FText UPCGIndirectionSettings::GetNodeTooltipText() const
+{
+	return LOCTEXT("NodeTooltip", "Executes another settings object, which can be overriden.");
+}
+
+#endif // WITH_EDITOR
+
+TArray<FPCGPinProperties> UPCGIndirectionSettings::OutputPinProperties() const
+{
+	TArray<FPCGPinProperties> PinProperties;
+	PinProperties.Emplace(PCGPinConstants::DefaultOutputLabel, EPCGDataType::Any);
+
+	return PinProperties;
+}
+
+FPCGElementPtr UPCGIndirectionSettings::CreateElement() const
+{
+	return MakeShared<FPCGIndirectionElement>();
+}
+
+bool FPCGIndirectionElement::CanExecuteOnlyOnMainThread(FPCGContext* InContext) const
+{
+	if (!InContext)
+	{
+		return true;
+	}
+
+	FPCGIndirectionContext* Context = static_cast<FPCGIndirectionContext*>(InContext);
+
+	if (Context->InnerElement)
+	{
+		return Context->InnerElement->CanExecuteOnlyOnMainThread(Context->InnerContext);
+	}
+	else
+	{
+		return true;
+	}
+}
+
+FPCGIndirectionContext::~FPCGIndirectionContext()
+{
+	if (bNeedsToUnrootInnerSettings)
+	{
+		const UPCGIndirectionSettings* Settings = GetInputSettings<UPCGIndirectionSettings>();
+		check(Settings);
+
+		if (UPCGSettings* InnerSettings = Settings->Settings.Get())
+		{
+			InnerSettings->RemoveFromRoot();
+		}
+	}
+
+	delete InnerContext;
+	InnerContext = nullptr;
+}
+
+FPCGContext* FPCGIndirectionElement::Initialize(const FPCGDataCollection& InputData, TWeakObjectPtr<UPCGComponent> SourceComponent, const UPCGNode* Node)
+{
+	FPCGIndirectionContext* Context = new FPCGIndirectionContext();
+	Context->InputData = InputData;
+	Context->SourceComponent = SourceComponent;
+	Context->Node = Node;
+
+	return Context;
+}
+
+bool FPCGIndirectionElement::PrepareDataInternal(FPCGContext* InContext) const
+{
+	check(InContext);
+	FPCGIndirectionContext* Context = static_cast<FPCGIndirectionContext*>(InContext);
+
+	const UPCGIndirectionSettings* Settings = Context->GetInputSettings<UPCGIndirectionSettings>();
+	check(Settings);
+
+	if (UPCGSettings* InnerSettings = Settings->Settings.LoadSynchronous())
+	{
+		// TODO: while we can root it here, if this node or multiple indirection node were to execute in parallel,
+		// the lifetime of the inner settings would be a bit unclear - we need to improve this.
+		if (!InnerSettings->IsRooted())
+		{
+			InnerSettings->AddToRoot();
+			Context->bNeedsToUnrootInnerSettings = true;
+		}
+
+		Context->InnerElement = InnerSettings->GetElement();
+		check(Context->InnerElement);
+
+		// Note: we need to pass a null node here + add the inner settings as part of the input so that they are retrieved properly
+		FPCGDataCollection InnerInput = Context->InputData;
+		InnerInput.TaggedData.Emplace_GetRef().Data = InnerSettings;
+
+		// TODO: there are some types of settings that might require a bit more information to be able to do their processing correctly
+		// namely (dynamic) subgraphs, so YMMV.
+		Context->InnerContext = Context->InnerElement->Initialize(InnerInput, Context->SourceComponent, nullptr);
+		check(Context->InnerContext);
+		Context->InnerContext->InitializeSettings();
+
+		// Unclear whether we should give those new values
+		Context->InnerContext->TaskId = Context->TaskId;
+		Context->InnerContext->CompiledTaskId = Context->CompiledTaskId;
+		Context->InnerContext->DependenciesCrc = Context->DependenciesCrc;
+		Context->InnerContext->GenerationGrid = Context->GenerationGrid;
+
+		Context->InnerContext->AsyncState = Context->AsyncState;
+	}
+
+	return true;
+}
+
+bool FPCGIndirectionElement::ExecuteInternal(FPCGContext* InContext) const
+{
+	TRACE_CPUPROFILER_EVENT_SCOPE(FPCGIndirectionElement::Execute);
+
+	check(InContext);
+	FPCGIndirectionContext* Context = static_cast<FPCGIndirectionContext*>(InContext);
+
+	const UPCGIndirectionSettings* Settings = Context->GetInputSettings<UPCGIndirectionSettings>();
+	check(Settings);
+
+	if (!Context->InnerElement || !Context->InnerContext)
+	{
+		return true;
+	}
+
+	// TODO: use caching when possible
+	// TODO: see what we can do for inspection data
+	// TODO: support pausing in inner element, might require some upstream changes in the graph executor
+	Context->InnerContext->AsyncState = Context->AsyncState;
+	bool bElementDone = Context->InnerElement->Execute(Context->InnerContext);
+
+	// Implementation note: to make sure everything is clean vs. the root set, we need to copy the output data
+	// regardless of whether the element is done or not
+	Context->OutputData = Context->InnerContext->OutputData;
+
+	// Finally, move the inner context pin data (which does not exist as-is on this node) to tags if required
+	if (bElementDone && Settings->bTagOutputsBasedOnOutputPins)
+	{
+		for (FPCGTaggedData& TaggedData : Context->OutputData.TaggedData)
+		{
+			if (TaggedData.Pin != NAME_None && TaggedData.Pin != PCGPinConstants::DefaultOutputLabel)
+			{
+				TaggedData.Tags.Add(TaggedData.Pin.ToString());
+				TaggedData.Pin = NAME_None;
+			}
+		}
+	}
+
+	return bElementDone;
+}
+
+#undef LOCTEXT_NAMESPACE
\ No newline at end of file
diff --git a/Engine/Plugins/Experimental/PCG/Source/PCG/Private/Metadata/PCGMetadata.cpp b/Engine/Plugins/Experimental/PCG/Source/PCG/Private/Metadata/PCGMetadata.cpp
--- a/Engine/Plugins/Experimental/PCG/Source/PCG/Private/Metadata/PCGMetadata.cpp
+++ b/Engine/Plugins/Experimental/PCG/Source/PCG/Private/Metadata/PCGMetadata.cpp
@@ -747,6 +747,12 @@ int64 UPCGMetadata::GetItemCountForChild() const
 	return ParentKeys.Num() + ItemKeyOffset;
 }
 
+int64 UPCGMetadata::GetLocalItemCount() const
+{
+	FReadScopeLock ScopeLock(ItemLock);
+	return ParentKeys.Num();
+}
+
 int64 UPCGMetadata::AddEntry(int64 ParentEntry)
 {
 	FWriteScopeLock ScopeLock(ItemLock);
diff --git a/Engine/Plugins/Experimental/PCG/Source/PCG/Public/Elements/IO/PCGDataTableElement.h b/Engine/Plugins/Experimental/PCG/Source/PCG/Public/Elements/IO/PCGDataTableElement.h
--- a/Engine/Plugins/Experimental/PCG/Source/PCG/Public/Elements/IO/PCGDataTableElement.h
+++ b/Engine/Plugins/Experimental/PCG/Source/PCG/Public/Elements/IO/PCGDataTableElement.h
@@ -4,6 +4,8 @@
 
 #include "PCGExternalData.h"
 
+#include "PCGCommon.h"
+
 class UDataTable;
 
 #include "PCGDataTableElement.generated.h"
@@ -19,15 +21,22 @@ public:
 	virtual FName GetDefaultNodeName() const override { return FName(TEXT("LoadDataTable")); }
 	virtual FText GetDefaultNodeTitle() const override;
 	virtual FText GetNodeTooltipText() const override;
+	virtual bool HasDynamicPins() const override { return true; }
 #endif
 
+	virtual EPCGDataType GetCurrentPinTypes(const UPCGPin* InPin) const override;
+
 protected:
+	virtual TArray<FPCGPinProperties> OutputPinProperties() const override;
 	virtual FPCGElementPtr CreateElement() const override;
 	// ~End UPCGSettings interface
 
 public:
-	UPROPERTY(BlueprintReadWrite, EditAnywhere, Category = "Settings", meta = (PCG_Overridable))
+	UPROPERTY(BlueprintReadWrite, EditAnywhere, Category = Settings, meta = (PCG_Overridable))
 	TSoftObjectPtr<UDataTable> DataTable;
+
+	UPROPERTY(BlueprintReadWrite, EditAnywhere, Category = Settings, meta = (ValidEnumValues = "Point, Param"))
+	EPCGExclusiveDataType OutputType = EPCGExclusiveDataType::Point;
 };
 
 class FPCGLoadDataTableElement : public FPCGExternalDataElement
diff --git a/Engine/Plugins/Experimental/PCG/Source/PCG/Public/Elements/IO/PCGExternalDataContext.h b/Engine/Plugins/Experimental/PCG/Source/PCG/Public/Elements/IO/PCGExternalDataContext.h
--- a/Engine/Plugins/Experimental/PCG/Source/PCG/Public/Elements/IO/PCGExternalDataContext.h
+++ b/Engine/Plugins/Experimental/PCG/Source/PCG/Public/Elements/IO/PCGExternalDataContext.h
@@ -32,7 +32,8 @@ struct FPCGExternalDataContext : public FPCGContext
 
 	struct FPointDataAccessorsMapping
 	{
-		UPCGPointData* PointData = nullptr;
+		UPCGData* Data = nullptr;
+		UPCGMetadata* Metadata = nullptr;
 		TArray<FRowToPointAccessors> RowToPointAccessors;
 		TUniquePtr<const IPCGAttributeAccessorKeys> RowKeys;
 	};
diff --git a/Engine/Plugins/Experimental/PCG/Source/PCG/Public/Elements/PCGAttributeGetFromIndexElement.h b/Engine/Plugins/Experimental/PCG/Source/PCG/Public/Elements/PCGAttributeGetFromIndexElement.h
new file mode 100644
--- /dev/null
+++ b/Engine/Plugins/Experimental/PCG/Source/PCG/Public/Elements/PCGAttributeGetFromIndexElement.h
@@ -0,0 +1,38 @@
+// Copyright Epic Games, Inc. All Rights Reserved.
+
+#pragma once
+
+#include "PCGSettings.h"
+
+#include "PCGAttributeGetFromIndexElement.generated.h"
+
+UCLASS(BlueprintType, ClassGroup=(Procedural))
+class PCG_API UPCGAttributeGetFromIndexSettings : public UPCGSettings
+{
+	GENERATED_BODY()
+
+public:
+	//~Begin UPCGSettings interface
+#if WITH_EDITOR
+	virtual FName GetDefaultNodeName() const override;
+	virtual FText GetDefaultNodeTitle() const override;
+	virtual FText GetNodeTooltipText() const override;
+	virtual EPCGSettingsType GetType() const override { return EPCGSettingsType::Param; }
+#endif
+
+protected:
+	virtual TArray<FPCGPinProperties> InputPinProperties() const override;
+	virtual TArray<FPCGPinProperties> OutputPinProperties() const override;
+	virtual FPCGElementPtr CreateElement() const override;
+	//~End UPCGSettings interface
+
+public:
+	UPROPERTY(BlueprintReadWrite, EditAnywhere, Category = Settings, meta = (PCG_Overridable))
+	int32 Index = 0;
+};
+
+class FPCGAttributeGetFromIndexElement : public FSimplePCGElement
+{
+protected:
+	virtual bool ExecuteInternal(FPCGContext* Context) const override;
+};
\ No newline at end of file
diff --git a/Engine/Plugins/Experimental/PCG/Source/PCG/Public/Elements/PCGIndirectionElement.h b/Engine/Plugins/Experimental/PCG/Source/PCG/Public/Elements/PCGIndirectionElement.h
new file mode 100644
--- /dev/null
+++ b/Engine/Plugins/Experimental/PCG/Source/PCG/Public/Elements/PCGIndirectionElement.h
@@ -0,0 +1,58 @@
+// Copyright Epic Games, Inc. All Rights Reserved.
+
+#pragma once
+
+#include "PCGSettings.h"
+#include "PCGContext.h"
+
+#include "PCGIndirectionElement.generated.h"
+
+UCLASS(BlueprintType, ClassGroup=(Procedural))
+class PCG_API UPCGIndirectionSettings : public UPCGSettings
+{
+	GENERATED_BODY()
+
+public:
+	//~Begin UPCGSettings interface
+#if WITH_EDITOR
+	virtual FName GetDefaultNodeName() const override;
+	virtual FText GetDefaultNodeTitle() const override;
+	virtual FText GetNodeTooltipText() const override;
+	virtual EPCGSettingsType GetType() const override { return EPCGSettingsType::Generic; }
+#endif
+
+protected:
+	virtual TArray<FPCGPinProperties> OutputPinProperties() const override;
+
+	virtual FPCGElementPtr CreateElement() const override;
+	//~End UPCGSettings interface
+
+public:
+	UPROPERTY(BlueprintReadWrite, EditAnywhere, Category = Settings, meta = (PCG_Overridable))
+	TSoftObjectPtr<UPCGSettings> Settings;
+
+	UPROPERTY(BlueprintReadWrite, EditAnywhere, Category = Settings, meta = (PCG_Overridable))
+	bool bTagOutputsBasedOnOutputPins = true;
+};
+
+struct FPCGIndirectionContext : public FPCGContext
+{
+	virtual ~FPCGIndirectionContext();
+
+	FPCGElementPtr InnerElement;
+	FPCGContext* InnerContext = nullptr;
+	bool bNeedsToUnrootInnerSettings = false;
+};
+
+class FPCGIndirectionElement : public IPCGElement
+{
+public:
+	virtual FPCGContext* Initialize(const FPCGDataCollection& InputData, TWeakObjectPtr<UPCGComponent> SourceComponent, const UPCGNode* Node) override;
+	// TODO: investigate how we could make this cacheable, might require to pass in context instead of settings
+	virtual bool IsCacheable(const UPCGSettings* InSettings) const override { return false; }
+	virtual bool CanExecuteOnlyOnMainThread(FPCGContext* Context) const override;
+
+protected:
+	virtual bool PrepareDataInternal(FPCGContext* InContext) const override;
+	virtual bool ExecuteInternal(FPCGContext* InContext) const override;
+};
\ No newline at end of file
diff --git a/Engine/Plugins/Experimental/PCG/Source/PCG/Public/Metadata/PCGMetadata.h b/Engine/Plugins/Experimental/PCG/Source/PCG/Public/Metadata/PCGMetadata.h
--- a/Engine/Plugins/Experimental/PCG/Source/PCG/Public/Metadata/PCGMetadata.h
+++ b/Engine/Plugins/Experimental/PCG/Source/PCG/Public/Metadata/PCGMetadata.h
@@ -271,6 +271,7 @@ public:
 
 	int64 GetItemKeyCountForParent() const;
 	int64 GetItemCountForChild() const;
+	int64 GetLocalItemCount() const;
 
 	/**
 	* Create a new attribute. If the attribute already exists, it will raise a warning (use FindOrCreateAttribute if this usecase can arise)
diff --git a/Engine/Plugins/Experimental/PCGExternalDataInterop/Source/PCGExternalDataInterop/Private/Alembic/PCGAlembicInterop.cpp b/Engine/Plugins/Experimental/PCGExternalDataInterop/Source/PCGExternalDataInterop/Private/Alembic/PCGAlembicInterop.cpp
--- a/Engine/Plugins/Experimental/PCGExternalDataInterop/Source/PCGExternalDataInterop/Private/Alembic/PCGAlembicInterop.cpp
+++ b/Engine/Plugins/Experimental/PCGExternalDataInterop/Source/PCGExternalDataInterop/Private/Alembic/PCGAlembicInterop.cpp
@@ -251,7 +251,8 @@ void ParseAlembicObject(FPCGExternalDataContext* Context, const Alembic::Abc::IO
 			UPCGMetadata* PointMetadata = PointData->MutableMetadata();
 
 			FPCGExternalDataContext::FPointDataAccessorsMapping& PointDataAccessorMapping = Context->PointDataAccessorsMapping.Emplace_GetRef();
-			PointDataAccessorMapping.PointData = PointData;
+			PointDataAccessorMapping.Data = PointData;
+			PointDataAccessorMapping.Metadata = PointMetadata;
 			// We're not going to use the input keys, but we still need to provide something
 			PointDataAccessorMapping.RowKeys = MakeUnique<FPCGAttributeAccessorKeysEntries>(PCGInvalidEntryKey);
 
diff --git a/Engine/Plugins/Experimental/PCGExternalDataInterop/Source/PCGExternalDataInterop/Private/Elements/PCGLoadAlembicElement.cpp b/Engine/Plugins/Experimental/PCGExternalDataInterop/Source/PCGExternalDataInterop/Private/Elements/PCGLoadAlembicElement.cpp
--- a/Engine/Plugins/Experimental/PCGExternalDataInterop/Source/PCGExternalDataInterop/Private/Elements/PCGLoadAlembicElement.cpp
+++ b/Engine/Plugins/Experimental/PCGExternalDataInterop/Source/PCGExternalDataInterop/Private/Elements/PCGLoadAlembicElement.cpp
@@ -101,7 +101,7 @@ bool FPCGLoadAlembicElement::PrepareLoad(FPCGExternalDataContext* InContext) con
 		for (const FPCGExternalDataContext::FPointDataAccessorsMapping& DataMapping : Context->PointDataAccessorsMapping)
 		{
 			FPCGTaggedData& OutData = Context->OutputData.TaggedData.Emplace_GetRef();
-			OutData.Data = DataMapping.PointData;
+			OutData.Data = DataMapping.Data;
 		}
 
 		Context->bDataPrepared = true;
@@ -136,7 +136,7 @@ bool FPCGLoadAlembicElement::ExecuteLoad(FPCGExternalDataContext* InContext) con
 	{
 		for (const FPCGExternalDataContext::FPointDataAccessorsMapping& DataMapping : Context->PointDataAccessorsMapping)
 		{
-			UPCGPointData* PointData = DataMapping.PointData;
+			UPCGPointData* PointData = Cast<UPCGPointData>(DataMapping.Data);
 
 			if (!PointData)
 			{
