Tentative fix for code inconsistency post-merge

#rnx

diff --git a/Engine/Plugins/Experimental/PCG/Source/PCG/Private/Elements/IO/PCGDataTableElement.cpp b/Engine/Plugins/Experimental/PCG/Source/PCG/Private/Elements/IO/PCGDataTableElement.cpp
--- a/Engine/Plugins/Experimental/PCG/Source/PCG/Private/Elements/IO/PCGDataTableElement.cpp
+++ b/Engine/Plugins/Experimental/PCG/Source/PCG/Private/Elements/IO/PCGDataTableElement.cpp
@@ -155,7 +155,7 @@ bool FPCGLoadDataTableElement::PrepareLoad(FPCGExternalDataContext* Context) con
 		{
 			// Create attribute with the appropriate type
 			const void* FirstRow = DataTable->GetRowMap().CreateConstIterator().Value();
-			PointMetadata->CreateAttributeFromDataProperty(FName(FieldName), FirstRow, *FieldIt);
+			OutMetadata->CreateAttributeFromDataProperty(FName(FieldName), FirstRow, *FieldIt);
 		}
 
 		TUniquePtr<IPCGAttributeAccessor> PointPropertyAccessor = PCGAttributeAccessorHelpers::CreateAccessor(OutData, PointPropertySelector);
