Check all colliding landscapes when querying the landscape cache.  If two landscape are right next to each other, the query was only testing the first one, which lead to the cache missing because a component was not found.

diff --git a/Engine/Plugins/Experimental/PCG/Source/PCG/Private/Helpers/PCGBlueprintHelpers.cpp b/Engine/Plugins/Experimental/PCG/Source/PCG/Private/Helpers/PCGBlueprintHelpers.cpp
--- a/Engine/Plugins/Experimental/PCG/Source/PCG/Private/Helpers/PCGBlueprintHelpers.cpp
+++ b/Engine/Plugins/Experimental/PCG/Source/PCG/Private/Helpers/PCGBlueprintHelpers.cpp
@@ -150,46 +150,64 @@ TArray<FPCGLandscapeLayerWeight> UPCGBlueprintHelpers::GetInterpolatedPCGLandsca
 	FBox Bounds(&Location, 1);
 	TArray<TWeakObjectPtr<ALandscapeProxy>> Landscapes = PCGHelpers::GetLandscapeProxies(World, Bounds);
 
-	if (Landscapes.IsEmpty() || !Landscapes[0].Get())
+	if (Landscapes.IsEmpty())
 	{
 		return {};
 	}
 
-	ALandscapeProxy* Landscape = Landscapes[0].Get();
-	check(Landscape);
+	FString FailureReason;
 
-	ULandscapeInfo* LandscapeInfo = Landscape->GetLandscapeInfo();
-	if (!LandscapeInfo)
+
+	for (TWeakObjectPtr<ALandscapeProxy> LandscapePtr : Landscapes)
 	{
-		UE_LOG(LogPCG, Warning, TEXT("Unable to get landscape layer weights because the landscape info is not available (landscape not registered yet?"));
-		return {};
-	}
 
-	const FVector LocalPoint = Landscape->GetTransform().InverseTransformPosition(Location);
-	const FIntPoint ComponentMapKey(FMath::FloorToInt(LocalPoint.X / LandscapeInfo->ComponentSizeQuads), FMath::FloorToInt(LocalPoint.Y / LandscapeInfo->ComponentSizeQuads));
+		ALandscapeProxy* Landscape = LandscapePtr.Get();
+		if (!Landscape)
+		{
+			continue;
+		}
+
+		ULandscapeInfo* LandscapeInfo = Landscape->GetLandscapeInfo();
+		if (!LandscapeInfo)
+		{
+			FailureReason = TEXT("Unable to get landscape layer weights because the landscape info is not available (landscape not registered yet?");
+			continue;
+		}
+
+		const FVector LocalPoint = Landscape->GetTransform().InverseTransformPosition(Location);
+		const FIntPoint ComponentMapKey(FMath::FloorToInt(LocalPoint.X / LandscapeInfo->ComponentSizeQuads), FMath::FloorToInt(LocalPoint.Y / LandscapeInfo->ComponentSizeQuads));
+
+		#if WITH_EDITOR
+		ULandscapeComponent* LandscapeComponent = LandscapeInfo->XYtoComponentMap.FindRef(ComponentMapKey);
+		const FPCGLandscapeCacheEntry* CacheEntry = LandscapeCache->GetCacheEntry(LandscapeComponent, ComponentMapKey);
+		#else
+		const FPCGLandscapeCacheEntry* CacheEntry = LandscapeCache->GetCacheEntry(Landscape->GetOriginalLandscapeGuid(), ComponentMapKey);
+		#endif
+
+		if (!CacheEntry)
+		{
+			FailureReason = TEXT("Unable to get landscape layer weights because the cache entry is not available.");
+			continue;
+		}
+
+		const FVector2D ComponentLocalPoint(LocalPoint.X - ComponentMapKey.X * LandscapeInfo->ComponentSizeQuads, LocalPoint.Y - ComponentMapKey.Y * LandscapeInfo->ComponentSizeQuads);
+	
+		TArray<FPCGLandscapeLayerWeight> Result;
+		CacheEntry->GetInterpolatedLayerWeights(ComponentLocalPoint, Result);
 
-#if WITH_EDITOR
-	ULandscapeComponent* LandscapeComponent = LandscapeInfo->XYtoComponentMap.FindRef(ComponentMapKey);
-	const FPCGLandscapeCacheEntry* CacheEntry = LandscapeCache->GetCacheEntry(LandscapeComponent, ComponentMapKey);
-#else
-	const FPCGLandscapeCacheEntry* CacheEntry = LandscapeCache->GetCacheEntry(Landscape->GetOriginalLandscapeGuid(), ComponentMapKey);
-#endif
+		Result.Sort([](const FPCGLandscapeLayerWeight& Lhs, const FPCGLandscapeLayerWeight& Rhs) {
+			return Lhs.Weight > Rhs.Weight;
+		});
 
-	if (!CacheEntry)
-	{
-		return {};
+		return Result;
 	}
 
-	const FVector2D ComponentLocalPoint(LocalPoint.X - ComponentMapKey.X * LandscapeInfo->ComponentSizeQuads, LocalPoint.Y - ComponentMapKey.Y * LandscapeInfo->ComponentSizeQuads);
-	
-	TArray<FPCGLandscapeLayerWeight> Result;
-	CacheEntry->GetInterpolatedLayerWeights(ComponentLocalPoint, Result);
-
-	Result.Sort([](const FPCGLandscapeLayerWeight& Lhs, const FPCGLandscapeLayerWeight& Rhs) {
-		return Lhs.Weight > Rhs.Weight;
-	});
+	if(FailureReason.Len())
+	{
+		UE_LOG(LogPCG, Warning, TEXT("%s"), *FailureReason);
+	}
 
-	return Result;
+	return {};
 }
 
 int64 UPCGBlueprintHelpers::GetTaskId(FPCGContext& Context)
