[PCG] Fix ensure if one quits the editor while inspecting a graph.

#rnx

diff --git a/Engine/Plugins/Experimental/PCG/Source/PCGEditor/Private/PCGEditor.cpp b/Engine/Plugins/Experimental/PCG/Source/PCGEditor/Private/PCGEditor.cpp
--- a/Engine/Plugins/Experimental/PCG/Source/PCGEditor/Private/PCGEditor.cpp
+++ b/Engine/Plugins/Experimental/PCG/Source/PCGEditor/Private/PCGEditor.cpp
@@ -1950,9 +1950,12 @@ void FPCGEditor::OnClose()
 
 	if (PCGComponentBeingInspected.IsValid())
 	{
-		PCGComponentBeingInspected->DisableInspection();
+		if (PCGComponentBeingInspected->IsInspecting())
+		{
+			PCGComponentBeingInspected->DisableInspection();
+		}
 
-		if (PCGGraphBeingEdited)
+		if (PCGGraphBeingEdited && PCGGraphBeingEdited->IsInspecting())
 		{
 			PCGGraphBeingEdited->DisableInspection();
 		}
