Added a node to copy attributes when some specified value matches against the source attribute set. Supports weights so covers all three cases of by weight, by attribute and weighted by category.
Intended to be used with multi attribute sets as generated from a Load Data Table.

diff --git a/Engine/Plugins/Experimental/PCG/Source/PCG/Private/Elements/PCGMatchAndSetAttributes.cpp b/Engine/Plugins/Experimental/PCG/Source/PCG/Private/Elements/PCGMatchAndSetAttributes.cpp
new file mode 100644
--- /dev/null
+++ b/Engine/Plugins/Experimental/PCG/Source/PCG/Private/Elements/PCGMatchAndSetAttributes.cpp
@@ -0,0 +1,529 @@
+// Copyright Epic Games, Inc. All Rights Reserved.
+
+#include "Elements/PCGMatchAndSetAttributes.h"
+
+#include "PCGContext.h"
+#include "PCGParamData.h"
+#include "PCGPin.h"
+#include "Elements/Metadata/PCGMetadataElementCommon.h"
+#include "Helpers/PCGBlueprintHelpers.h"
+#include "Helpers/PCGPointDataPartition.h"
+#include "Metadata/PCGMetadataAttributeTpl.h"
+#include "Metadata/PCGMetadataAttributeTraits.h"
+#include "Metadata/Accessors/IPCGAttributeAccessor.h"
+#include "Metadata/Accessors/PCGAttributeAccessorHelpers.h"
+#include "Metadata/Accessors/PCGAttributeAccessorKeys.h"
+
+#include "Algo/AnyOf.h"
+
+#include UE_INLINE_GENERATED_CPP_BY_NAME(PCGMatchAndSetAttributes)
+
+#define LOCTEXT_NAMESPACE "PCGMatchAndSetAttributes"
+
+namespace PCGMatchAndSetAttributesConstants
+{
+	const FName MatchDataLabel = TEXT("Match Data");
+}
+
+#if WITH_EDITOR
+FName UPCGMatchAndSetAttributesSettings::GetDefaultNodeName() const
+{
+	return FName(TEXT("MatchAndSetAttributes"));
+}
+
+FText UPCGMatchAndSetAttributesSettings::GetDefaultNodeTitle() const
+{
+	return LOCTEXT("NodeTitle", "Match And Set Attributes");
+}
+
+FText UPCGMatchAndSetAttributesSettings::GetNodeTooltipText() const
+{
+	return LOCTEXT("NodeTooltip", "Matches or randomly assigns values from the Attribute Set to the input Point Data");
+}
+#endif // WITH_EDITOR
+
+UPCGMatchAndSetAttributesSettings::UPCGMatchAndSetAttributesSettings()
+{
+	bUseSeed = true;
+}
+
+TArray<FPCGPinProperties> UPCGMatchAndSetAttributesSettings::InputPinProperties() const
+{
+	TArray<FPCGPinProperties> PinProperties;
+	PinProperties.Emplace(PCGPinConstants::DefaultInputLabel, EPCGDataType::Point);
+	PinProperties.Emplace(PCGMatchAndSetAttributesConstants::MatchDataLabel, 
+		EPCGDataType::Param, 
+		/*bAllowMultipleConnection=*/false, 
+		/*bAllowMultipleData=*/false,
+		LOCTEXT("MatchDataTooltip", "Input containing the data to match to, then copy the accompanying attribute values")
+	);
+
+	return PinProperties;
+}
+
+TArray<FPCGPinProperties> UPCGMatchAndSetAttributesSettings::OutputPinProperties() const
+{
+	return Super::DefaultPointOutputPinProperties();
+}
+
+FPCGElementPtr UPCGMatchAndSetAttributesSettings::CreateElement() const
+{
+	return MakeShared<FPCGMatchAndSetAttributesElement>();
+}
+
+class FPCGAttributeSetPartition
+{
+public:
+	struct AttributeSetPartitionEntry
+	{
+		double TotalWeight = 0;
+		TArray<PCGMetadataEntryKey> Keys;
+		TArray<double> CumulativeWeight;
+
+		void AddEntry(PCGMetadataEntryKey Key, double Weight)
+		{
+			if (Weight > 0)
+			{
+				Keys.Add(Key);
+				TotalWeight += Weight;
+				CumulativeWeight.Add(TotalWeight);
+			}
+		}
+	};
+
+	FPCGAttributeSetPartition() = default;
+	FPCGAttributeSetPartition(FPCGContext* InContext, const UPCGParamData* InParamData, bool bPartitionByAttribute, FName AttributeName, bool bUseWeightAttribute, FName WeightAttributeName)
+	{
+		Initialize(InContext, InParamData, bPartitionByAttribute, AttributeName, bUseWeightAttribute, WeightAttributeName);
+	}
+
+	bool Initialize(FPCGContext* Context, const UPCGParamData* InParamData, bool bPartitionByAttribute, FName AttributeName, bool bUseWeightAttribute, FName WeightAttributeName)
+	{
+		ParamData = InParamData;
+
+		if (!ParamData || !ParamData->Metadata)
+		{
+			return false;
+		}
+
+		const UPCGMetadata* Metadata = ParamData->ConstMetadata();
+
+		if (bPartitionByAttribute)
+		{
+			Attribute = Metadata->GetConstAttribute(AttributeName);
+
+			if (!Attribute)
+			{
+				PCGE_LOG_C(Error, GraphAndLog, Context, FText::Format(LOCTEXT("CannotFindAttribute", "Cannot find attribute '{0}' in the source Attribute Set."), FText::FromName(AttributeName)));
+				return false;
+			}
+		}
+
+		const FPCGMetadataAttributeBase* WeightAttribute = nullptr;
+		if (bUseWeightAttribute)
+		{
+			WeightAttribute = Metadata->GetConstAttribute(WeightAttributeName);
+
+			if (!WeightAttribute)
+			{
+				PCGE_LOG_C(Error, GraphAndLog, Context, FText::Format(LOCTEXT("CannotFindWeightAttribute", "Cannot find weight attribute '{0}' in the source Attribute Set."), FText::FromName(WeightAttributeName)));
+				return false;
+			}
+
+			if(!PCG::Private::IsOfTypes<int32, int64, float, double>(WeightAttribute->GetTypeId()))
+			{
+				PCGE_LOG_C(Error, GraphAndLog, Context, FText::Format(LOCTEXT("InvalidWeightAttributeType", "Weight attribute '{0}' does not have the proper type (int32, int64, float or double)."), FText::FromName(WeightAttributeName)));
+				return false;
+			}
+		}
+
+		auto GetWeightFromAttribute = [WeightAttribute](PCGMetadataEntryKey EntryKey) -> double
+		{
+			double Weight = 1.0;
+
+			if (WeightAttribute)
+			{
+				auto GetValue = [WeightAttribute, EntryKey](auto Dummy) -> double
+				{
+					using ValueType = decltype(Dummy);
+					if constexpr (PCG::Private::IsOfTypes<ValueType, int32, int64, float, double>())
+					{
+						return static_cast<const FPCGMetadataAttribute<ValueType>*>(WeightAttribute)->GetValueFromItemKey(EntryKey);
+					}
+					else
+					{
+						return 1.0;
+					}
+				};
+
+				Weight = PCGMetadataAttribute::CallbackWithRightType(WeightAttribute->GetTypeId(), GetValue);
+			}
+
+			return Weight;
+		};
+		
+		// Note: since we don't have an accessor to the entries from the metadata,
+		// we're going to assume that they exist in a consecutive sequence, which should hold true for param data.
+		const int64 FirstKey = Metadata->GetItemKeyCountForParent();
+		const int64 KeyCount = Metadata->GetLocalItemCount();
+
+		for (int64 EntryKey = FirstKey; EntryKey < FirstKey + KeyCount; ++EntryKey)
+		{
+			PCGMetadataValueKey ValueKey = PCGDefaultValueKey;
+			if (Attribute)
+			{
+				ValueKey = Attribute->GetValueKey(EntryKey);
+			}
+
+			double Weight = GetWeightFromAttribute(EntryKey);
+			TPair<PCGMetadataValueKey, AttributeSetPartitionEntry>* MatchingVK = nullptr;
+
+			if (Attribute)
+			{
+				MatchingVK = PartitionData.FindByPredicate([this, ValueKey](const TPair<PCGMetadataValueKey, AttributeSetPartitionEntry>& Entry)
+				{
+					return Attribute->AreValuesEqual(Entry.Key, ValueKey);
+				});
+			}
+			else if(!PartitionData.IsEmpty())
+			{
+				MatchingVK = &PartitionData[0];
+			}
+
+			if (!MatchingVK)
+			{
+				MatchingVK = &PartitionData.Emplace_GetRef(ValueKey, AttributeSetPartitionEntry());
+			}
+
+			MatchingVK->Value.AddEntry(EntryKey, Weight);
+		}
+
+#ifdef WITH_EDITOR
+		// Check for empty entries
+		const bool bHasEmptyEntries = Algo::AnyOf(PartitionData, [](const TPair<PCGMetadataValueKey, AttributeSetPartitionEntry>& PartitionEntry) { return PartitionEntry.Value.TotalWeight == 0; });
+		if (bHasEmptyEntries)
+		{
+			PCGE_LOG_C(Warning, GraphAndLog, Context, FText::Format(LOCTEXT("EmptyCategory", "Some match entries on Attribute '{0}' in the Attribute Set do not have any associated valid weight."), FText::FromName(Attribute ? Attribute->Name : NAME_None)));
+		}
+#endif // WITH_EDITOR
+		
+		bIsValid = true;
+		return true;
+	}
+
+	bool IsValid() const { return bIsValid; }
+
+	TArray<int32> GetMatchingPartitionDataIndices(const TUniquePtr<const IPCGAttributeAccessor>& InputAttribute, const TUniquePtr<const IPCGAttributeAccessorKeys>& InputKeys, int32 PointsNum) const
+	{
+		TArray<int32> MatchingPartitionDataIndices;
+
+		if (InputAttribute && Attribute)
+		{
+			check(InputKeys.IsValid() && InputKeys->GetNum() == PointsNum);
+			auto FindMatchingValueKeyIndex = [this, &InputAttribute, &InputKeys, &MatchingPartitionDataIndices](auto AttributeDummyValue) -> bool
+			{
+				using AttributeType = decltype(AttributeDummyValue);
+
+				// Get the values to match against from the attribute
+				const FPCGMetadataAttribute<AttributeType>* TypedAttribute = static_cast<const FPCGMetadataAttribute<AttributeType>*>(Attribute);
+				TArray<AttributeType> AttributeValues;
+				AttributeValues.Reserve(PartitionData.Num());
+
+				for (const TPair<PCGMetadataValueKey, AttributeSetPartitionEntry>& PartitionEntry : PartitionData)
+				{
+					AttributeValues.Add(TypedAttribute->GetValue(PartitionEntry.Key));
+				}
+
+				return PCGMetadataElementCommon::ApplyOnAccessor<AttributeType>(*InputKeys, *InputAttribute, [&AttributeValues, &MatchingPartitionDataIndices](const AttributeType& InValue, int32)
+				{
+					int32 MatchingPartitionDataIndex = INDEX_NONE;
+					for (int32 AttributeValueIndex = 0; AttributeValueIndex < AttributeValues.Num(); ++AttributeValueIndex)
+					{
+						if (PCG::Private::MetadataTraits<AttributeType>::Equal(InValue, AttributeValues[AttributeValueIndex]))
+						{
+							MatchingPartitionDataIndex = AttributeValueIndex;
+							break;
+						}
+					}
+
+					MatchingPartitionDataIndices.Add(MatchingPartitionDataIndex);
+				}, EPCGAttributeAccessorFlags::AllowBroadcast | EPCGAttributeAccessorFlags::AllowConstructible);
+			};
+
+			if (!PCGMetadataAttribute::CallbackWithRightType(Attribute->GetTypeId(), FindMatchingValueKeyIndex))
+			{
+				// Attribute isn't able to retrieve & compare - reset values
+				MatchingPartitionDataIndices.Init(INDEX_NONE, InputKeys->GetNum());
+			}
+		}
+		else if (!PartitionData.IsEmpty())
+		{
+			// There is only one entry
+			MatchingPartitionDataIndices.Init(0, PointsNum);
+		}
+		else
+		{
+			// There are no entries in the partition data
+			MatchingPartitionDataIndices.Init(INDEX_NONE, PointsNum);
+		}
+
+		return MatchingPartitionDataIndices;
+	}
+
+	PCGMetadataEntryKey GetWeightedEntry(int32 PartitionDataIndex, const FRandomStream& InRandomStream) const
+	{
+		if (PartitionDataIndex == INDEX_NONE)
+		{
+			return PCGInvalidEntryKey;
+		}
+
+		// Second, resolve weight-based entries
+		const AttributeSetPartitionEntry& PartitionDataEntry = PartitionData[PartitionDataIndex].Value;
+		int RandomPick = INDEX_NONE;
+
+		if (PartitionDataEntry.Keys.Num() == 1)
+		{
+			RandomPick = 0;
+		}
+		else if (PartitionDataEntry.Keys.Num() > 1)
+		{
+			double RandomWeightedPick = InRandomStream.FRandRange(0, PartitionDataEntry.TotalWeight);
+			RandomPick = 0;
+			while (RandomPick < PartitionDataEntry.CumulativeWeight.Num() && PartitionDataEntry.CumulativeWeight[RandomPick] <= RandomWeightedPick)
+			{
+				++RandomPick;
+			}
+		}
+
+		if (RandomPick != INDEX_NONE)
+		{
+			return PartitionDataEntry.Keys[RandomPick];
+		}
+		else
+		{
+			// No entry in partition data, which is unexpected, but possible if all entries were <= 0.
+			return PCGInvalidEntryKey;
+		}
+	}
+
+private:
+	const UPCGParamData* ParamData = nullptr;
+	const FPCGMetadataAttributeBase* Attribute = nullptr;
+	bool bIsValid = false;
+
+	TArray<TPair<PCGMetadataValueKey, AttributeSetPartitionEntry>> PartitionData;
+};
+
+class FPCGMatchAndSetPartition : public FPCGPointDataPartitionBase<FPCGMatchAndSetPartition, PCGMetadataValueKey>
+{
+public:
+	FPCGMatchAndSetPartition(FPCGContext* InContext, const UPCGMatchAndSetAttributesSettings* InSettings, const UPCGComponent* InSourceComponent, const UPCGParamData* InParamData) 
+		: FPCGPointDataPartitionBase<FPCGMatchAndSetPartition, PCGMetadataValueKey>()
+		, Context(InContext)
+		, Settings(InSettings)
+		, ParamData(InParamData)
+		, SourceComponent(InSourceComponent)
+	{
+	}
+
+	bool Initialize()
+	{
+		if (!Settings)
+		{
+			return false;
+		}
+
+		AttributeSetPartition.Initialize(Context,
+			ParamData,
+			Settings->bMatchAttributes,
+			Settings->MatchAttribute,
+			Settings->bUseWeightAttribute,
+			Settings->WeightAttribute);
+
+		return AttributeSetPartition.IsValid();
+	}
+
+	bool InitializeForPointData(const UPCGPointData* PointData, UPCGPointData* OutPointData)
+	{
+		if (!PointData || !PointData->ConstMetadata() || !OutPointData)
+		{
+			return false;
+		}
+
+		if (Settings->bMatchAttributes)
+		{
+			FPCGAttributePropertyInputSelector InputAttributeSource = Settings->InputAttribute.CopyAndFixLast(PointData);
+			InputAttributeAccessor = PCGAttributeAccessorHelpers::CreateConstAccessor(PointData, InputAttributeSource);
+			InputAttributeKeys = PCGAttributeAccessorHelpers::CreateConstKeys(PointData, InputAttributeSource);
+
+			if (!InputAttributeAccessor.IsValid() || !InputAttributeKeys.IsValid())
+			{
+				PCGE_LOG_C(Error, GraphAndLog, Context, FText::Format(LOCTEXT("MissingAttribute", "Point data does not have the input attribute '{0}'."), InputAttributeSource.GetDisplayText()));
+				return false;
+			}
+		}
+
+		AttributesToSet.Reset();
+
+		// Prepare set of attributes to copy, i.e. create the attributes if we need to, make a 1:1 pair with the ones from the
+		// param data. Note that we don't need to copy over the matched attribute if any, nor the weight.
+		TArray<FName> ParamAttributeNames;
+		TArray<EPCGMetadataTypes> ParamAttributeTypes;
+		ParamData->ConstMetadata()->GetAttributes(ParamAttributeNames, ParamAttributeTypes);
+
+		for (int32 AttributeIndex = 0; AttributeIndex < ParamAttributeNames.Num(); ++AttributeIndex)
+		{
+			const FName AttributeName = ParamAttributeNames[AttributeIndex];
+
+			if ((Settings->bMatchAttributes && AttributeName == Settings->MatchAttribute) ||
+				(Settings->bUseWeightAttribute && AttributeName == Settings->WeightAttribute))
+			{
+				continue;
+			}
+
+			const FPCGMetadataAttributeBase* ParamAttribute = ParamData->ConstMetadata()->GetConstAttribute(AttributeName);
+			
+			FPCGMetadataAttributeBase* PointAttribute = OutPointData->Metadata->GetMutableAttribute(AttributeName);
+			if (PointAttribute && PointAttribute->GetTypeId() != ParamAttribute->GetTypeId())
+			{
+				OutPointData->Metadata->DeleteAttribute(AttributeName);
+				PointAttribute = nullptr;
+			}
+
+			if (!PointAttribute)
+			{
+				PointAttribute = OutPointData->Metadata->CopyAttribute(ParamAttribute, AttributeName, /*bKeepParent=*/false, /*bCopyEntries=*/false, /*bCopyValues=*/false);
+			}
+
+			if (!PointAttribute) // Failed to create attribute
+			{
+				PCGE_LOG_C(Warning, GraphAndLog, Context, FText::Format(LOCTEXT("UnableToCreateAttribute", "Unable to create attribute '{0}' on point data."), FText::FromName(AttributeName)));
+				return false;
+			}
+
+			AttributesToSet.Emplace(ParamAttribute, PointAttribute);
+		}
+
+		return true;
+	}
+
+	FPCGPointDataPartitionBase::Element* SelectPoint(const FPCGPoint& Point, int32 PointIndex)
+	{
+		return nullptr;
+	}
+
+	void Finalize(const UPCGPointData* InPointData, UPCGPointData* OutPointData)
+	{
+		check(InPointData && OutPointData);
+		const TArray<FPCGPoint>& Points = InPointData->GetPoints();
+		TArray<FPCGPoint>& OutPoints = OutPointData->GetMutablePoints();
+
+		// In this case, we haven't written any points yet to the OutPointData,
+		// as we need to find the matching index for each entry first.
+		TArray<int32> PartitionDataIndices = AttributeSetPartition.GetMatchingPartitionDataIndices(InputAttributeAccessor, InputAttributeKeys, Points.Num());
+
+		if (PartitionDataIndices.Num() != Points.Num())
+		{
+			return;
+		}
+
+		for (int32 PointIndex = 0; PointIndex < Points.Num(); ++PointIndex)
+		{
+			const FPCGPoint& Point = Points[PointIndex];
+			int32 PartitionDataIndex = PartitionDataIndices[PointIndex];
+			PCGMetadataEntryKey AttributeSetKey = PCGInvalidEntryKey;
+
+			if (PartitionDataIndex != INDEX_NONE)
+			{
+				AttributeSetKey = AttributeSetPartition.GetWeightedEntry(PartitionDataIndex, UPCGBlueprintHelpers::GetRandomStream(Point, Settings, SourceComponent));
+			}
+
+			if (Settings->bKeepUnmatched || AttributeSetKey != PCGInvalidEntryKey)
+			{
+				FPCGPoint& OutPoint = OutPoints.Add_GetRef(Point);
+
+				if (AttributeSetKey != PCGInvalidEntryKey)
+				{
+					// This is similar to UPCGMetadata::SetAttributes but for a subset of attributes
+					OutPoint.MetadataEntry = OutPointData->MutableMetadata()->AddEntry(Point.MetadataEntry);
+
+					for (const TPair<const FPCGMetadataAttributeBase*, FPCGMetadataAttributeBase*>& AttributePair : AttributesToSet)
+					{
+						const FPCGMetadataAttributeBase* ParamAttribute = AttributePair.Key;
+						FPCGMetadataAttributeBase* PointAttribute = AttributePair.Value;
+
+						PointAttribute->SetValue(OutPoint.MetadataEntry, ParamAttribute, AttributeSetKey);
+					}
+				}
+			}
+		}
+	}
+
+private:
+	FPCGContext* Context = nullptr;
+	const UPCGMatchAndSetAttributesSettings* Settings = nullptr;
+	const UPCGParamData* ParamData = nullptr;
+	FPCGAttributeSetPartition AttributeSetPartition;
+	const UPCGComponent* SourceComponent = nullptr;
+
+	// Per point data iteration data
+	TUniquePtr<const IPCGAttributeAccessor> InputAttributeAccessor;
+	TUniquePtr<const IPCGAttributeAccessorKeys> InputAttributeKeys;
+	TArray<TPair<const FPCGMetadataAttributeBase*, FPCGMetadataAttributeBase*>> AttributesToSet;
+};
+
+bool FPCGMatchAndSetAttributesElement::ExecuteInternal(FPCGContext* Context) const
+{
+	TRACE_CPUPROFILER_EVENT_SCOPE(FPCGMatchAndSetAttributesElement::Execute);
+
+	const UPCGMatchAndSetAttributesSettings* Settings = Context->GetInputSettings<UPCGMatchAndSetAttributesSettings>();
+	check(Settings);
+
+	TArray<FPCGTaggedData> Inputs = Context->InputData.GetInputsByPin(PCGPinConstants::DefaultInputLabel);
+	TArray<FPCGTaggedData> ParamDataInputs = Context->InputData.GetInputsByPin(PCGMatchAndSetAttributesConstants::MatchDataLabel);
+	TArray<FPCGTaggedData>& Outputs = Context->OutputData.TaggedData;
+	
+	const UPCGParamData* ParamData = nullptr;
+
+	if (ParamDataInputs.Num() == 1)
+	{
+		ParamData = Cast<const UPCGParamData>(ParamDataInputs[0].Data);
+	}
+
+	if (!ParamData)
+	{
+		PCGE_LOG(Warning, GraphAndLog, LOCTEXT("NoMatchData", "Must have exactly one Attribute Set to match against"));
+		return true;
+	}
+
+	FPCGMatchAndSetPartition Partition(Context, Settings, Context->SourceComponent.Get(), ParamData);
+
+	if (!Partition.Initialize())
+	{
+		return true;
+	}
+
+	for (const FPCGTaggedData& Input : Inputs)
+	{
+		FPCGTaggedData& Output = Outputs.Add_GetRef(Input);
+
+		const UPCGPointData* InPointData = Cast<UPCGPointData>(Input.Data);
+		if (!InPointData)
+		{
+			PCGE_LOG(Error, GraphAndLog, LOCTEXT("InvalidInputDataType", "Input data must be of type Point"));
+			continue;
+		}
+
+		UPCGPointData* OutPointData = NewObject<UPCGPointData>();
+		OutPointData->InitializeFromData(InPointData);
+		OutPointData->GetMutablePoints().Reserve(InPointData->GetPoints().Num());
+
+		int CurrentPointIndex = 0;
+		ensure(Partition.SelectPoints(*Context, InPointData, CurrentPointIndex, OutPointData));
+		Output.Data = OutPointData;
+	}
+
+	return true;
+}
+
+#undef LOCTEXT_NAMESPACE
\ No newline at end of file
diff --git a/Engine/Plugins/Experimental/PCG/Source/PCG/Private/Elements/PCGSpawnActor.cpp b/Engine/Plugins/Experimental/PCG/Source/PCG/Private/Elements/PCGSpawnActor.cpp
--- a/Engine/Plugins/Experimental/PCG/Source/PCG/Private/Elements/PCGSpawnActor.cpp
+++ b/Engine/Plugins/Experimental/PCG/Source/PCG/Private/Elements/PCGSpawnActor.cpp
@@ -140,7 +140,7 @@ public:
 	{
 	}
 
-	bool InitializeForPointData(const UPCGPointData* PointData)
+	bool InitializeForPointData(const UPCGPointData* PointData, UPCGPointData* OutPointData)
 	{
 		if (!PointData || !PointData->ConstMetadata())
 		{
diff --git a/Engine/Plugins/Experimental/PCG/Source/PCG/Public/Elements/Metadata/PCGMetadataElementCommon.h b/Engine/Plugins/Experimental/PCG/Source/PCG/Public/Elements/Metadata/PCGMetadataElementCommon.h
--- a/Engine/Plugins/Experimental/PCG/Source/PCG/Public/Elements/Metadata/PCGMetadataElementCommon.h
+++ b/Engine/Plugins/Experimental/PCG/Source/PCG/Public/Elements/Metadata/PCGMetadataElementCommon.h
@@ -51,7 +51,7 @@ namespace PCGMetadataElementCommon
 	* Return false if it process nothing.
 	*/
 	template <typename T, typename Func>
-	bool ApplyOnAccessor(const IPCGAttributeAccessorKeys& Keys, const IPCGAttributeAccessor& Accessor, Func&& Callback, const int32 ChunkSize = DefaultChunkSize)
+	bool ApplyOnAccessor(const IPCGAttributeAccessorKeys& Keys, const IPCGAttributeAccessor& Accessor, Func&& Callback, EPCGAttributeAccessorFlags Flags = EPCGAttributeAccessorFlags::StrictType, const int32 ChunkSize = DefaultChunkSize)
 	{
 		const int32 NumberOfEntries = Keys.GetNum();
 
@@ -71,7 +71,7 @@ namespace PCGMetadataElementCommon
 			const int32 Range = FMath::Min(NumberOfEntries - StartIndex, ChunkSize);
 			TArrayView<T> View(TempValues.GetData(), Range);
 
-			if (!Accessor.GetRange(View, StartIndex, Keys))
+			if (!Accessor.GetRange(View, StartIndex, Keys, Flags))
 			{
 				return false;
 			}
diff --git a/Engine/Plugins/Experimental/PCG/Source/PCG/Public/Elements/PCGMatchAndSetAttributes.h b/Engine/Plugins/Experimental/PCG/Source/PCG/Public/Elements/PCGMatchAndSetAttributes.h
new file mode 100644
--- /dev/null
+++ b/Engine/Plugins/Experimental/PCG/Source/PCG/Public/Elements/PCGMatchAndSetAttributes.h
@@ -0,0 +1,66 @@
+// Copyright Epic Games, Inc. All Rights Reserved.
+
+#pragma once
+
+#include "PCGSettings.h"
+
+#include "PCGMatchAndSetAttributes.generated.h"
+
+/** This class creates a PCG node that can match, select by weight or match & select by weight 
+* a 'matching' entry in a provided Attribute Set with multiple entries.
+* E.g. for a given point, if the point has the same specified attribute as the matching attribute in the attribute set,
+* then we will copy all the other non-selection attributes to the point.
+*/
+UCLASS(BlueprintType, ClassGroup = (Procedural))
+class PCG_API UPCGMatchAndSetAttributesSettings : public UPCGSettings
+{
+	GENERATED_BODY()
+
+public:
+	UPCGMatchAndSetAttributesSettings();
+
+	// ~Begin UPCGSettings interface
+#if WITH_EDITOR
+	virtual FName GetDefaultNodeName() const override;
+	virtual FText GetDefaultNodeTitle() const override;
+	virtual FText GetNodeTooltipText() const override;
+	virtual EPCGSettingsType GetType() const override { return EPCGSettingsType::Metadata; }
+#endif // WITH_EDITOR
+
+protected:
+	virtual TArray<FPCGPinProperties> InputPinProperties() const override;
+	virtual TArray<FPCGPinProperties> OutputPinProperties() const override;
+	virtual FPCGElementPtr CreateElement() const override;
+	// ~End UPCGSettings interface
+
+public:
+	/** Controls whether selection of the attribute set values to copy will be done by matching point-to-attribute set (true) or done randomly (false) */
+	UPROPERTY(BlueprintReadWrite, EditAnywhere, Category = Settings)
+	bool bMatchAttributes = false;
+
+	/** Attribute from the point data to select & maetch */
+	UPROPERTY(BlueprintReadWrite, EditAnywhere, Category = Settings, meta = (EditCondition = "bMatchAttributes"))
+	FPCGAttributePropertyInputSelector InputAttribute;
+
+	/** Attribute from the attribute set to match against */
+	UPROPERTY(BlueprintReadWrite, EditAnywhere, Category = Settings, meta = (EditCondition = "bMatchAttributes"))
+	FName MatchAttribute = NAME_None;
+
+	/** Controls whether points that have no valid match in the attribute set are kept as is (default values) or removed from the output */
+	UPROPERTY(BlueprintReadWrite, EditAnywhere, Category = Settings, meta = (EditCondition = "bMatchAttributes", PCG_Overridable))
+	bool bKeepUnmatched = true;
+
+	/** Controls whether we will consider the weights, as determined by the Weight Attribute values on the attribute set */
+	UPROPERTY(BlueprintReadWrite, EditAnywhere, Category = Settings, meta = (InlineEditConditionToggle, PCG_Overridable))
+	bool bUseWeightAttribute = false;
+
+	/** Attribute to weight more or less some entries from the attribute set */
+	UPROPERTY(BlueprintReadWrite, EditAnywhere, Category = Settings, meta = (EditCondition = "bUseWeightAttribute", PCG_Overridable))
+	FName WeightAttribute = NAME_None;
+};
+
+class FPCGMatchAndSetAttributesElement : public FSimplePCGElement
+{
+public:
+	virtual bool ExecuteInternal(FPCGContext* Context) const override;
+};
\ No newline at end of file
diff --git a/Engine/Plugins/Experimental/PCG/Source/PCG/Public/Helpers/PCGPointDataPartition.h b/Engine/Plugins/Experimental/PCG/Source/PCG/Public/Helpers/PCGPointDataPartition.h
--- a/Engine/Plugins/Experimental/PCG/Source/PCG/Public/Helpers/PCGPointDataPartition.h
+++ b/Engine/Plugins/Experimental/PCG/Source/PCG/Public/Helpers/PCGPointDataPartition.h
@@ -46,7 +46,8 @@ public:
 	bool Initialize() { return true; }
 	bool InitializeForPointData(const UPCGPointData* PointData) { return true; }
 	void AddToPartitionData(Element* SelectedElement, const UPCGPointData* ParentPointData, const FPCGPoint& Point) {}
-	void WriteToOutputData(UPCGPointData* OutPointData, Element* SelectedElement, const FPCGPoint& Point) {}
+	void WriteToOutputData(UPCGPointData* OutPointData, Element* SelectedElement, const FPCGPoint& Point, int32 PointIndex) {}
+	void Finalize(const UPCGPointData* InPointData, UPCGPointData* OutPointData) {}
 	Element* SelectPoint(const FPCGPoint& Point, int32 PointIndex) { return nullptr; }
 	int32 TimeSlicingCheckFrequency() const { return 1024; }
 
@@ -69,7 +70,7 @@ bool FPCGPointDataPartitionBase<Derived, KeyType>::SelectPoints(FPCGContext& Con
 
 	if (InCurrentPointIndex == 0)
 	{
-		if (!This()->InitializeForPointData(PointData))
+		if (!This()->InitializeForPointData(PointData, OutPointData))
 		{
 			return true;
 		}
@@ -92,7 +93,7 @@ bool FPCGPointDataPartitionBase<Derived, KeyType>::SelectPoints(FPCGContext& Con
 			This()->AddToPartitionData(SelectedElement, PointData, Point);
 		}
 
-		This()->WriteToOutputData(OutPointData, SelectedElement, Point);
+		This()->WriteToOutputData(OutPointData, SelectedElement, Point, CurrentPointIndex);
 
 		++CurrentPointIndex;
 
@@ -111,5 +112,10 @@ bool FPCGPointDataPartitionBase<Derived, KeyType>::SelectPoints(FPCGContext& Con
 
 	InCurrentPointIndex = CurrentPointIndex;
 
+	if (CurrentPointIndex == Points.Num())
+	{
+		This()->Finalize(PointData, OutPointData);
+	}
+
 	return (CurrentPointIndex == Points.Num());
 }
\ No newline at end of file
