[PCG] Adds a node that adds tags from the output pin to the next node.
--Added unit tests that check no tags, one tag, and multiple tags

diff --git a/Engine/Plugins/Experimental/PCG/Source/PCG/Private/Elements/PCGAddTag.cpp b/Engine/Plugins/Experimental/PCG/Source/PCG/Private/Elements/PCGAddTag.cpp
new file mode 100644
--- /dev/null
+++ b/Engine/Plugins/Experimental/PCG/Source/PCG/Private/Elements/PCGAddTag.cpp
@@ -0,0 +1,68 @@
+// Copyright Epic Games, Inc. All Rights Reserved.
+
+#include "Elements/PCGAddTag.h"
+
+#include "PCGContext.h"
+#include "PCGPin.h"
+#include "Helpers/PCGHelpers.h"
+
+#define LOCTEXT_NAMESPACE "PCGAddTagElement"
+
+EPCGDataType UPCGAddTagSettings::GetCurrentPinTypes(const UPCGPin* InPin) const
+{
+	check(InPin);
+	if (!InPin->IsOutputPin())
+	{
+		return Super::GetCurrentPinTypes(InPin);
+	}
+
+	// Output pin narrows to union of inputs on first pin
+	const EPCGDataType InputTypeUnion = GetTypeUnionOfIncidentEdges(PCGPinConstants::DefaultInputLabel);
+	return InputTypeUnion != EPCGDataType::None ? InputTypeUnion : EPCGDataType::Any;
+}
+
+TArray<FPCGPinProperties> UPCGAddTagSettings::InputPinProperties() const
+{
+	TArray<FPCGPinProperties> PinProperties;
+	PinProperties.Emplace(PCGPinConstants::DefaultInputLabel, EPCGDataType::Any);
+
+	return PinProperties;
+}
+
+TArray<FPCGPinProperties> UPCGAddTagSettings::OutputPinProperties() const
+{
+	TArray<FPCGPinProperties> PinProperties;
+	PinProperties.Emplace(PCGPinConstants::DefaultOutputLabel, EPCGDataType::Any);
+
+	return PinProperties;
+}
+
+FPCGElementPtr UPCGAddTagSettings::CreateElement() const
+{
+	return MakeShared<FPCGAddTagElement>();
+}
+
+bool FPCGAddTagElement::ExecuteInternal(FPCGContext* Context) const
+{
+	TRACE_CPUPROFILER_EVENT_SCOPE(FPCGAddTagElement::Execute);
+	
+	check(Context);
+
+	const UPCGAddTagSettings* Settings = Context->GetInputSettings<UPCGAddTagSettings>();
+	check(Settings);
+	
+	Context->OutputData = Context->InputData;
+	const TArray<FString> TagsArray = PCGHelpers::GetStringArrayFromCommaSeparatedString(Settings->TagsToAdd);
+	
+	for (const FString& Tag : TagsArray)
+	{
+		for (FPCGTaggedData& OutputTaggedData : Context->OutputData.TaggedData)
+		{
+				OutputTaggedData.Tags.Add(Tag);
+		}
+	}
+	
+	return true;
+}
+
+#undef LOCTEXT_NAMESPACE
diff --git a/Engine/Plugins/Experimental/PCG/Source/PCG/Private/Tests/Elements/PCGAddTagTest.cpp b/Engine/Plugins/Experimental/PCG/Source/PCG/Private/Tests/Elements/PCGAddTagTest.cpp
new file mode 100644
--- /dev/null
+++ b/Engine/Plugins/Experimental/PCG/Source/PCG/Private/Tests/Elements/PCGAddTagTest.cpp
@@ -0,0 +1,126 @@
+// Copyright Epic Games, Inc. All Rights Reserved.
+
+#include "Tests/Determinism/PCGDeterminismTestsCommon.h"
+
+#include "PCGComponent.h"
+#include "PCGContext.h"
+
+#include "Elements/PCGAddTag.h"
+
+IMPLEMENT_CUSTOM_SIMPLE_AUTOMATION_TEST(FPCGAddTagTest_ZeroTags, FPCGTestBaseClass, "pcg.tests.AddTag.ZeroTags", PCGTestsCommon::TestFlags)
+
+bool FPCGAddTagTest_ZeroTags::RunTest(const FString& Parameters)
+{
+	PCGTestsCommon::FTestData TestData(PCGDeterminismTests::Defaults::Seed);
+
+	//auto generate settings for this node
+	PCGDeterminismTests::GenerateSettings<UPCGAddTagSettings>(TestData);
+	UPCGAddTagSettings* Settings = CastChecked<UPCGAddTagSettings>(TestData.Settings);
+
+	FPCGElementPtr TestElement = TestData.Settings->GetElement();
+
+	TUniquePtr<FPCGContext> Context = TestData.InitializeTestContext();
+
+	while (!TestElement->Execute(Context.Get())) {}
+
+	//testing inputs and outputs
+	const TArray<FPCGTaggedData>& Inputs = Context->InputData.TaggedData;
+
+	UTEST_EQUAL("Number of Inputs", Inputs.Num(), 1);
+
+	const TArray<FPCGTaggedData>& Outputs = Context->OutputData.TaggedData;
+
+	UTEST_EQUAL("Number of Outputs", Outputs.Num(), 1);
+
+	//testing tag info
+	const TSet<FString> OutputTagSet = Context->OutputData.TaggedData[0].Tags;
+
+	UTEST_EQUAL("Tag count", OutputTagSet.Num(), 0);
+
+	return true;
+}
+
+IMPLEMENT_CUSTOM_SIMPLE_AUTOMATION_TEST(FPCGAddTagTest_OneTag, FPCGTestBaseClass, "pcg.tests.AddTag.OneTag", PCGTestsCommon::TestFlags)
+
+bool FPCGAddTagTest_OneTag::RunTest(const FString& Parameters)
+{
+	PCGTestsCommon::FTestData TestData(PCGDeterminismTests::Defaults::Seed);
+
+	//auto generate settings for this node
+	PCGDeterminismTests::GenerateSettings<UPCGAddTagSettings>(TestData);
+	UPCGAddTagSettings* Settings = CastChecked<UPCGAddTagSettings>(TestData.Settings);
+
+	FPCGElementPtr TestElement = TestData.Settings->GetElement();
+
+	Settings->TagsToAdd = "Feature";
+
+	TUniquePtr<FPCGContext> Context = TestData.InitializeTestContext();
+
+	while (!TestElement->Execute(Context.Get())) {}
+
+	//testing inputs and outputs
+	const TArray<FPCGTaggedData>& Inputs = Context->InputData.TaggedData;
+
+	UTEST_EQUAL("Number of Inputs", Inputs.Num(), 1);
+
+	const TArray<FPCGTaggedData>& Outputs = Context->OutputData.TaggedData;
+
+	UTEST_EQUAL("Number of Outputs", Outputs.Num(), 1);
+
+	//testing tag info
+	const TSet<FString>& OutputTagSet = Outputs[0].Tags;
+
+	UTEST_EQUAL("Tag count", OutputTagSet.Num(), 1);
+
+	if (!OutputTagSet.IsEmpty())
+	{
+		FString OutTag = OutputTagSet.Array()[0];
+
+		UTEST_TRUE("Tags contains Feature", OutputTagSet.Contains("Feature"));
+	}
+
+	return true;
+}
+
+IMPLEMENT_CUSTOM_SIMPLE_AUTOMATION_TEST(FPCGAddTagTest_MultipleTags, FPCGTestBaseClass, "pcg.tests.AddTag.MultipleTags", PCGTestsCommon::TestFlags)
+
+bool FPCGAddTagTest_MultipleTags::RunTest(const FString& Parameters)
+{
+	PCGTestsCommon::FTestData TestData(PCGDeterminismTests::Defaults::Seed);
+
+	//auto generate settings for this node
+	PCGDeterminismTests::GenerateSettings<UPCGAddTagSettings>(TestData);
+	UPCGAddTagSettings* Settings = CastChecked<UPCGAddTagSettings>(TestData.Settings);
+
+	FPCGElementPtr TestElement = TestData.Settings->GetElement();
+
+	Settings->TagsToAdd = "Feature, Cool, Purple";
+
+	TUniquePtr<FPCGContext> Context = TestData.InitializeTestContext();
+
+	while (!TestElement->Execute(Context.Get())) {}
+
+	//testing inputs and outputs
+	const TArray<FPCGTaggedData>& Inputs = Context->InputData.TaggedData;
+
+	UTEST_EQUAL("Number of Inputs", Inputs.Num(), 1);
+
+	const TArray<FPCGTaggedData>& Outputs = Context->OutputData.TaggedData;
+
+	UTEST_EQUAL("Number of Outputs", Outputs.Num(), 1);
+	
+	//testing tag info
+	const TSet<FString> OutputTagSet = Outputs[0].Tags;
+
+	UTEST_EQUAL("Tag count", OutputTagSet.Num(), 3);
+
+	TArray<FString, TInlineAllocator<3>> ExpectedTags = { "Feature", "Cool", "Purple" };
+
+	for (const FString& ExpectTag : ExpectedTags)
+	{
+		UTEST_TRUE(*FString::Printf(TEXT("Tags contains %s"), *ExpectTag), OutputTagSet.Contains(ExpectTag));
+	}
+	
+	return true;
+}
+
diff --git a/Engine/Plugins/Experimental/PCG/Source/PCG/Public/Elements/PCGAddTag.h b/Engine/Plugins/Experimental/PCG/Source/PCG/Public/Elements/PCGAddTag.h
new file mode 100644
--- /dev/null
+++ b/Engine/Plugins/Experimental/PCG/Source/PCG/Public/Elements/PCGAddTag.h
@@ -0,0 +1,44 @@
+// Copyright Epic Games, Inc. All Rights Reserved.
+
+#pragma once
+
+#include "PCGSettings.h"
+
+#include "PCGAddTag.generated.h"
+
+/**
+ * Applies the specified tags on the output data.
+ */
+UCLASS(BlueprintType, ClassGroup = (Procedural))
+class UPCGAddTagSettings : public UPCGSettings
+{
+	GENERATED_BODY()
+
+public:
+	//~Begin UPCGSettings interface
+#if WITH_EDITOR
+	virtual FName GetDefaultNodeName() const override { return FName(TEXT("AddTags")); }
+	virtual FText GetDefaultNodeTitle() const override { return NSLOCTEXT("PCGAddTagElement", "NodeTitle", "Add Tags"); }
+	virtual FText GetNodeTooltipText() const override { return NSLOCTEXT("PCGAddTagElement", "NodeTooltip", "Applies the specified tags on the output data."); }
+	virtual EPCGSettingsType GetType() const override { return EPCGSettingsType::Generic; }
+	virtual bool HasDynamicPins() const override { return true; }
+#endif
+	virtual EPCGDataType GetCurrentPinTypes(const UPCGPin* InPin) const override;
+
+protected:
+	virtual TArray<FPCGPinProperties> InputPinProperties() const override;
+	virtual TArray<FPCGPinProperties> OutputPinProperties() const override;
+	virtual FPCGElementPtr CreateElement() const override;
+	//~End UPCGSettings interface
+
+public:
+	/** Comma-separated list of tags to apply to the node */
+	UPROPERTY(EditAnywhere, BlueprintReadWrite, Category = Settings, meta = (PCG_Overridable))
+	FString TagsToAdd;
+};
+
+class FPCGAddTagElement : public FSimplePCGElement
+{
+protected:
+	virtual bool ExecuteInternal(FPCGContext* Context) const override;
+};
\ No newline at end of file
