Remove impossible cast in UPCGSplineData::CreatePointData
UPCGSplineProjectionData is not a UPCGSplineData
#rnx

diff --git a/Engine/Plugins/Experimental/PCG/Source/PCG/Private/Data/PCGSplineData.cpp b/Engine/Plugins/Experimental/PCG/Source/PCG/Private/Data/PCGSplineData.cpp
--- a/Engine/Plugins/Experimental/PCG/Source/PCG/Private/Data/PCGSplineData.cpp
+++ b/Engine/Plugins/Experimental/PCG/Source/PCG/Private/Data/PCGSplineData.cpp
@@ -119,11 +119,12 @@ const UPCGPointData* UPCGSplineData::CreatePointData(FPCGContext* Context) const
 
 	const UPCGSpatialData* ProjectionTarget = nullptr;
 	FPCGProjectionParams ProjectionParams;
+	/* Impossible cast, UPCGSplineProjectionData is not a UPCGSplineData
 	if (const UPCGSplineProjectionData* SplineProjection = Cast<const UPCGSplineProjectionData>(this))
 	{
 		ProjectionTarget = SplineProjection->GetSurface();
 		ProjectionParams = SplineProjection->GetProjectionParams();
-	}
+	}*/
 
 	PCGSplineSampler::SampleLineData(this, nullptr, ProjectionTarget, ProjectionParams, SamplerParams, Data);
 	UE_LOG(LogPCG, Verbose, TEXT("Spline generated %d points"), Data->GetPoints().Num());
