Added merge attribute set node

diff --git a/Engine/Plugins/Experimental/PCG/Source/PCG/Private/Elements/PCGMergeAttributes.cpp b/Engine/Plugins/Experimental/PCG/Source/PCG/Private/Elements/PCGMergeAttributes.cpp
new file mode 100644
--- /dev/null
+++ b/Engine/Plugins/Experimental/PCG/Source/PCG/Private/Elements/PCGMergeAttributes.cpp
@@ -0,0 +1,105 @@
+// Copyright Epic Games, Inc. All Rights Reserved.
+
+#include "Elements/PCGMergeAttributes.h"
+
+#include "PCGContext.h"
+#include "PCGParamData.h"
+#include "PCGPin.h"
+#include "Metadata/PCGMetadata.h"
+
+#define LOCTEXT_NAMESPACE "PCGMergeAttributesSettings"
+
+#if WITH_EDITOR
+FText UPCGMergeAttributesSettings::GetDefaultNodeTitle() const
+{
+	return LOCTEXT("NodeTitle", "Merge Attributes");
+}
+
+FText UPCGMergeAttributesSettings::GetNodeTooltipText() const
+{
+	return LOCTEXT("NodeTooltip", "Merges multiple attribute sets in a single attribute set with multiple entries and all the provided attributes");
+}
+#endif // WITH_EDITOR
+
+TArray<FPCGPinProperties> UPCGMergeAttributesSettings::InputPinProperties() const
+{
+	TArray<FPCGPinProperties> PinProperties;
+	PinProperties.Emplace(PCGPinConstants::DefaultInputLabel, EPCGDataType::Param, /*bAllowMultipleConnections=*/true);
+
+	return PinProperties;
+}
+
+TArray<FPCGPinProperties> UPCGMergeAttributesSettings::OutputPinProperties() const
+{
+	TArray<FPCGPinProperties> PinProperties;
+	PinProperties.Emplace(PCGPinConstants::DefaultOutputLabel, EPCGDataType::Param, /*bAllowMultipleConnections=*/true, /*bAllowMultipleData=*/false);
+
+	return PinProperties;
+}
+
+FPCGElementPtr UPCGMergeAttributesSettings::CreateElement() const
+{
+	return MakeShared<FPCGMergeAttributesElement>();
+}
+
+bool FPCGMergeAttributesElement::ExecuteInternal(FPCGContext* Context) const
+{
+	TRACE_CPUPROFILER_EVENT_SCOPE(FPCGMergeAttributesElement::Execute);
+	check(Context);
+
+	const UPCGMergeAttributesSettings* Settings = Context->GetInputSettings<UPCGMergeAttributesSettings>();
+	check(Settings);
+
+	TArray<FPCGTaggedData>& Outputs = Context->OutputData.TaggedData;
+
+	FPCGTaggedData* MergedOutput = nullptr;
+	UPCGParamData* MergedAttributeSet = nullptr;
+
+	for (const FPCGTaggedData& Source : Context->InputData.TaggedData)
+	{
+		const UPCGParamData* SourceData = Cast<const UPCGParamData>(Source.Data);
+
+		if (!SourceData)
+		{
+			PCGE_LOG(Error, GraphAndLog, LOCTEXT("UnsupportedDataType", "Unsupported data type in merge attributes"));
+			continue;
+		}
+
+		if (!MergedOutput)
+		{
+			MergedOutput = &Outputs.Add_GetRef(Source);
+			continue;
+		}
+
+		const UPCGMetadata* SourceMetadata = SourceData->Metadata;
+		
+		// When we're merging the 2nd element, create the actual merged attribute set
+		if (!MergedAttributeSet)
+		{
+			MergedAttributeSet = NewObject<UPCGParamData>();
+			check(MergedOutput);
+			MergedAttributeSet->Metadata->InitializeAsCopy(CastChecked<const UPCGParamData>(MergedOutput->Data)->Metadata);
+			MergedOutput->Data = MergedAttributeSet;
+		}
+
+		// For all entries starting from the second:
+		// - Add missing attributes
+		MergedAttributeSet->Metadata->AddAttributes(SourceMetadata);
+
+		// - Merge entries
+		const int64 ParamItemCount = SourceData->Metadata->GetLocalItemCount();
+
+		for (int64 LocalItemKey = 0; LocalItemKey < ParamItemCount; ++LocalItemKey)
+		{
+			PCGMetadataEntryKey MergedItemKey = MergedAttributeSet->Metadata->AddEntry();
+			MergedAttributeSet->Metadata->SetAttributes(LocalItemKey, SourceMetadata, MergedItemKey);
+		}
+
+		// - Merge tags too (to be in line with the Merge points node)
+		MergedOutput->Tags.Append(Source.Tags);
+	}
+
+	return true;
+}
+
+#undef LOCTEXT_NAMESPACE
\ No newline at end of file
diff --git a/Engine/Plugins/Experimental/PCG/Source/PCG/Public/Elements/PCGMergeAttributes.h b/Engine/Plugins/Experimental/PCG/Source/PCG/Public/Elements/PCGMergeAttributes.h
new file mode 100644
--- /dev/null
+++ b/Engine/Plugins/Experimental/PCG/Source/PCG/Public/Elements/PCGMergeAttributes.h
@@ -0,0 +1,35 @@
+// Copyright Epic Games, Inc. All Rights Reserved.
+
+#pragma once
+
+#include "PCGSettings.h"
+
+#include "PCGMergeAttributes.generated.h"
+
+/** Merges multiple attribute sets together into a single attribute set */
+UCLASS(BlueprintType, Classgroup = (Procedural))
+class PCG_API UPCGMergeAttributesSettings : public UPCGSettings
+{
+	GENERATED_BODY()
+
+public:
+	//~Begin UPCGSettings interface
+#if WITH_EDITOR
+	virtual FName GetDefaultNodeName() const override { return TEXT("MergeAttributes"); }
+	virtual FText GetDefaultNodeTitle() const override;
+	virtual FText GetNodeTooltipText() const override;
+	virtual EPCGSettingsType GetType() const override { return EPCGSettingsType::Metadata; }
+#endif // WITH_EDITOR
+
+protected:
+	virtual TArray<FPCGPinProperties> InputPinProperties() const override;
+	virtual TArray<FPCGPinProperties> OutputPinProperties() const override;
+	virtual FPCGElementPtr CreateElement() const override;
+	//~End UPCGSettings interface
+};
+
+class FPCGMergeAttributesElement : public FSimplePCGElement
+{
+protected:
+	virtual bool ExecuteInternal(FPCGContext* Context) const override;
+};
\ No newline at end of file
