[PCG] Push higen grid size resolve pass deeper, from schedule to compile.

This way higen considerations are made before caching and this removes the 'patching' that currently happens in schedule. A set of compiled tasks are cached per execution grid size.

A next step will be to cull tasks that are not relevant for a grid size, which will be the major benefit of caching per grid size. And static switches will hopefully also allow culling here too.

#rnx

diff --git a/Engine/Plugins/Experimental/PCG/Source/PCG/Private/Graph/PCGGraphCompiler.cpp b/Engine/Plugins/Experimental/PCG/Source/PCG/Private/Graph/PCGGraphCompiler.cpp
--- a/Engine/Plugins/Experimental/PCG/Source/PCG/Private/Graph/PCGGraphCompiler.cpp
+++ b/Engine/Plugins/Experimental/PCG/Source/PCG/Private/Graph/PCGGraphCompiler.cpp
@@ -67,7 +67,8 @@ TArray<FPCGGraphTask> FPCGGraphCompiler::CompileGraph(UPCGGraph* InGraph, FPCGTa
 			// 1. Compile the subgraph making sure we don't reuse the same ids
 			// Note that we will not consume the pre or post-execute tasks, ergo bIsTopGraph=false
 			FPCGStackContext SubgraphStackContext;
-			TArray<FPCGGraphTask> Subtasks = GetCompiledTasks(Subgraph, SubgraphStackContext, /*bIsTopGraph=*/false);
+			// Passed uninitialized grid size to get all tasks
+			TArray<FPCGGraphTask> Subtasks = GetCompiledTasks(Subgraph, PCGHiGenGrid::UninitializedGridSize(), SubgraphStackContext, /*bIsTopGraph=*/false);
 
 #if WITH_EDITOR
 			GraphDependenciesLock.Lock();
@@ -265,22 +266,31 @@ void FPCGGraphCompiler::Compile(UPCGGraph* InGraph)
 	// Store back the results in the cache if it's valid
 	if (!CompiledTasks.IsEmpty())
 	{
-		GraphToTaskMapLock.WriteLock();
+		FWriteScopeLock Lock(GraphToTaskMapLock);
 		if (!GraphToTaskMap.Contains(InGraph))
 		{
 			GraphToTaskMap.Add(InGraph, MoveTemp(CompiledTasks));
 			GraphToStackContext.Add(InGraph, StackContext);
 		}
-		GraphToTaskMapLock.WriteUnlock();
 	}
 }
 
-TArray<FPCGGraphTask> FPCGGraphCompiler::GetPrecompiledTasks(UPCGGraph* InGraph, FPCGStackContext& OutStackContext, bool bIsTopGraph) const
+TArray<FPCGGraphTask> FPCGGraphCompiler::GetPrecompiledTasks(UPCGGraph* InGraph, uint32 GenerationGridSize, FPCGStackContext& OutStackContext, bool bIsTopGraph) const
 {
 	// Get compiled tasks in a threadsafe way
 	FReadScopeLock ReadLock(GraphToTaskMapLock);
 
-	const TArray<FPCGGraphTask>* ExistingTasks = (bIsTopGraph ? TopGraphToTaskMap : GraphToTaskMap).Find(InGraph);
+	const TArray<FPCGGraphTask>* ExistingTasks = nullptr;
+	if (bIsTopGraph)
+	{
+		// Top graphs are optimized per grid size.
+		const TMap<uint32, TArray<FPCGGraphTask>>* GridSizeToCompiledGraph = TopGraphToTaskMap.Find(InGraph);
+		ExistingTasks = GridSizeToCompiledGraph ? GridSizeToCompiledGraph->Find(GenerationGridSize) : nullptr;
+	}
+	else
+	{
+		ExistingTasks = GraphToTaskMap.Find(InGraph);
+	}
 
 	const FPCGStackContext* StackContext = (bIsTopGraph ? TopGraphToStackContext : GraphToStackContext).Find(InGraph);
 	if (StackContext)
@@ -296,7 +306,7 @@ TArray<FPCGGraphTask> FPCGGraphCompiler::GetPrecompiledTasks(UPCGGraph* InGraph,
 	return ExistingTasks ? *ExistingTasks : TArray<FPCGGraphTask>();
 }
 
-void FPCGGraphCompiler::ResolveGridSizes(TArray<FPCGGraphTask>& InOutCompiledTasks, const FPCGStackContext& InStackContext) const
+void FPCGGraphCompiler::ResolveGridSizes(TArray<FPCGGraphTask>& InOutCompiledTasks, const FPCGStackContext& InStackContext, EPCGHiGenGrid GenerationDefaultGrid) const
 {
 	// The stack is used to form the ResourceKey - a string that provides a path to the data from top graph down to specific pin.
 	// This will be used by link tasks as store/retrieve keys to marshal data for edges that cross grid size boundaries.
@@ -309,7 +319,7 @@ void FPCGGraphCompiler::ResolveGridSizes(TArray<FPCGGraphTask>& InOutCompiledTas
 	// Calculate execution grid values for each task.
 	for (FPCGGraphTask& Task : InOutCompiledTasks)
 	{
-		CalculateGridRecursive(Task.NodeId, InStackContext, InOutCompiledTasks);
+		CalculateGridRecursive(Task.NodeId, GenerationDefaultGrid, InStackContext, InOutCompiledTasks);
 	}
 
 	// Now add link tasks - if a Grid256 task depends on data from a Grid512 task, inject a link
@@ -371,16 +381,19 @@ void FPCGGraphCompiler::ResolveGridSizes(TArray<FPCGGraphTask>& InOutCompiledTas
 	}
 }
 
-EPCGHiGenGrid FPCGGraphCompiler::CalculateGridRecursive(FPCGTaskId InTaskId, const FPCGStackContext& InStackContext, TArray<FPCGGraphTask>& InOutCompiledTasks) const
+EPCGHiGenGrid FPCGGraphCompiler::CalculateGridRecursive(
+	FPCGTaskId InTaskId,
+	EPCGHiGenGrid GenerationDefaultGrid,
+	const FPCGStackContext& InStackContext,
+	TArray<FPCGGraphTask>& InOutCompiledTasks) const
 {
 	if (InOutCompiledTasks[InTaskId].GraphGenerationGrid != EPCGHiGenGrid::Uninitialized)
 	{
 		return InOutCompiledTasks[InTaskId].GraphGenerationGrid;
 	}
 
-	// GenerationDefault as default - means we don't know at compilation-time. Any tasks with this unknown value
-	// will receive the graph default grid size at schedule-time.
-	EPCGHiGenGrid Grid = EPCGHiGenGrid::GenerationDefault;
+	// Default for outside any grid size range.
+	EPCGHiGenGrid Grid = GenerationDefaultGrid;
 
 	const UPCGNode* Node = InOutCompiledTasks[InTaskId].Node;
 	const UPCGSettings* Settings = Node ? Node->GetSettings() : nullptr;
@@ -394,7 +407,7 @@ EPCGHiGenGrid FPCGGraphCompiler::CalculateGridRecursive(FPCGTaskId InTaskId, con
 		// Grid of this task is minimum of all input grids. We can link in data from a larger grid, but not from a finer grid (this goes against hierarchy).
 		for (FPCGGraphTaskInput InputTask : InOutCompiledTasks[InTaskId].Inputs)
 		{
-			const EPCGHiGenGrid InputGrid = CalculateGridRecursive(InputTask.TaskId, InStackContext, InOutCompiledTasks);
+			const EPCGHiGenGrid InputGrid = CalculateGridRecursive(InputTask.TaskId, GenerationDefaultGrid, InStackContext, InOutCompiledTasks);
 			if (PCGHiGenGrid::IsValidGrid(InputGrid))
 			{
 				Grid = FMath::Min(InputGrid, Grid);
@@ -407,23 +420,24 @@ EPCGHiGenGrid FPCGGraphCompiler::CalculateGridRecursive(FPCGTaskId InTaskId, con
 	return Grid;
 }
 
-TArray<FPCGGraphTask> FPCGGraphCompiler::GetCompiledTasks(UPCGGraph* InGraph, FPCGStackContext& OutStackContext, bool bIsTopGraph)
+TArray<FPCGGraphTask> FPCGGraphCompiler::GetCompiledTasks(UPCGGraph* InGraph, uint32 GenerationGridSize, FPCGStackContext& OutStackContext, bool bIsTopGraph)
 {
 	TArray<FPCGGraphTask> CompiledTasks;
 
 	if (bIsTopGraph)
 	{
 		// Always try to compile
-		CompileTopGraph(InGraph);
+		CompileTopGraph(InGraph, GenerationGridSize);
 
 		// Get compiled tasks in a threadsafe way
-		GraphToTaskMapLock.ReadLock();
-		if (TopGraphToTaskMap.Contains(InGraph))
+		FReadScopeLock Lock(GraphToTaskMapLock);
+		const TMap<uint32, TArray<FPCGGraphTask>>* GridSizeToCompiledGraph = TopGraphToTaskMap.Find(InGraph);
+		const TArray<FPCGGraphTask>* Tasks = GridSizeToCompiledGraph ? GridSizeToCompiledGraph->Find(GenerationGridSize) : nullptr;
+		if (Tasks)
 		{
-			CompiledTasks = TopGraphToTaskMap[InGraph];
+			CompiledTasks = *Tasks;
 			OutStackContext = TopGraphToStackContext[InGraph];
 		}
-		GraphToTaskMapLock.ReadUnlock();
 	}
 	else
 	{
@@ -431,13 +445,12 @@ TArray<FPCGGraphTask> FPCGGraphCompiler::GetCompiledTasks(UPCGGraph* InGraph, FP
 		Compile(InGraph);
 
 		// Get compiled tasks in a threadsafe way
-		GraphToTaskMapLock.ReadLock();
-		if (GraphToTaskMap.Contains(InGraph))
+		FReadScopeLock Lock(GraphToTaskMapLock);
+		if (TArray<FPCGGraphTask>* Tasks = GraphToTaskMap.Find(InGraph))
 		{
-			CompiledTasks = GraphToTaskMap[InGraph];
+			CompiledTasks = *Tasks;
 			OutStackContext = GraphToStackContext[InGraph];
 		}
-		GraphToTaskMapLock.ReadUnlock();
 	}
 
 	return CompiledTasks;
@@ -465,10 +478,11 @@ void FPCGGraphCompiler::OffsetNodeIds(TArray<FPCGGraphTask>& Tasks, FPCGTaskId O
 	}
 }
 
-void FPCGGraphCompiler::CompileTopGraph(UPCGGraph* InGraph)
+void FPCGGraphCompiler::CompileTopGraph(UPCGGraph* InGraph, uint32 GenerationGridSize)
 {
 	GraphToTaskMapLock.ReadLock();
-	bool bAlreadyCached = TopGraphToTaskMap.Contains(InGraph);
+	const TMap<uint32, TArray<FPCGGraphTask>>* GridSizeToCompiledGraph = TopGraphToTaskMap.Find(InGraph);
+	const bool bAlreadyCached = GridSizeToCompiledGraph && GridSizeToCompiledGraph->Contains(GenerationGridSize);
 	GraphToTaskMapLock.ReadUnlock();
 
 	if (bAlreadyCached)
@@ -478,7 +492,7 @@ void FPCGGraphCompiler::CompileTopGraph(UPCGGraph* InGraph)
 
 	// Build from non-top tasks
 	FPCGStackContext StackContext;
-	TArray<FPCGGraphTask> CompiledTasks = GetCompiledTasks(InGraph, StackContext, /*bIsTopGraph=*/false);
+	TArray<FPCGGraphTask> CompiledTasks = GetCompiledTasks(InGraph, GenerationGridSize, StackContext, /*bIsTopGraph=*/false);
 
 	// Check that the compilation was valid
 	if (CompiledTasks.Num() == 0)
@@ -486,11 +500,20 @@ void FPCGGraphCompiler::CompileTopGraph(UPCGGraph* InGraph)
 		return;
 	}
 
-	if (InGraph->IsHierarchicalGenerationEnabled())
+	if (GenerationGridSize != PCGHiGenGrid::UninitializedGridSize())
 	{
-		// Operates on compiled tasks before pre/post task added, as it may add tasks for linkages
-		// and the pre/post task need to reside at the end of the compiled task array.
-		ResolveGridSizes(CompiledTasks, StackContext);
+		if (InGraph->IsHierarchicalGenerationEnabled())
+		{
+			const EPCGHiGenGrid DefaultGrid = PCGHiGenGrid::GridSizeToGrid(InGraph->GetDefaultGridSize());
+			ResolveGridSizes(CompiledTasks, StackContext, DefaultGrid);
+		}
+
+		const EPCGHiGenGrid GenerationGrid = PCGHiGenGrid::GridSizeToGrid(GenerationGridSize);
+		for (FPCGGraphTask& Task : CompiledTasks)
+		{
+			// Make a copy of this because many things can be scheduled over many frames
+			Task.GenerationGrid = GenerationGrid;
+		}
 	}
 
 	const int TaskNum = CompiledTasks.Num();
@@ -543,14 +566,24 @@ void FPCGGraphCompiler::CompileTopGraph(UPCGGraph* InGraph)
 
 	// Store back the results in the cache
 	GraphToTaskMapLock.WriteLock();
-	if (!TopGraphToTaskMap.Contains(InGraph))
+	TMap<uint32, TArray<FPCGGraphTask>>& TasksPerGenerationGrid = TopGraphToTaskMap.FindOrAdd(InGraph);
+	if (!TasksPerGenerationGrid.Contains(GenerationGridSize))
 	{
-		TopGraphToTaskMap.Add(InGraph, MoveTemp(CompiledTasks));
+		TasksPerGenerationGrid.Add(GenerationGridSize, MoveTemp(CompiledTasks));
 		TopGraphToStackContext.Add(InGraph, StackContext);
 	}
 	GraphToTaskMapLock.WriteUnlock();
 }
 
+void FPCGGraphCompiler::ClearCache()
+{
+	FWriteScopeLock Lock(GraphToTaskMapLock);
+	GraphToTaskMap.Reset();
+	GraphToStackContext.Reset();
+	TopGraphToTaskMap.Reset();
+	TopGraphToStackContext.Reset();
+}
+
 #if WITH_EDITOR
 void FPCGGraphCompiler::NotifyGraphChanged(UPCGGraph* InGraph)
 {
diff --git a/Engine/Plugins/Experimental/PCG/Source/PCG/Private/Graph/PCGGraphCompiler.h b/Engine/Plugins/Experimental/PCG/Source/PCG/Private/Graph/PCGGraphCompiler.h
--- a/Engine/Plugins/Experimental/PCG/Source/PCG/Private/Graph/PCGGraphCompiler.h
+++ b/Engine/Plugins/Experimental/PCG/Source/PCG/Private/Graph/PCGGraphCompiler.h
@@ -16,29 +16,33 @@ class FPCGGraphCompiler
 {
 public:
 	void Compile(UPCGGraph* InGraph);
-	TArray<FPCGGraphTask> GetCompiledTasks(UPCGGraph* InGraph, FPCGStackContext& OutStackContext, bool bIsTopGraph = true);
-	TArray<FPCGGraphTask> GetPrecompiledTasks(UPCGGraph* InGraph, FPCGStackContext& OutStackContext, bool bIsTopGraph = true) const;
+	TArray<FPCGGraphTask> GetCompiledTasks(UPCGGraph* InGraph, uint32 GenerationGridSize, FPCGStackContext& OutStackContext, bool bIsTopGraph = true);
+	TArray<FPCGGraphTask> GetPrecompiledTasks(UPCGGraph* InGraph, uint32 GenerationGridSize, FPCGStackContext& OutStackContext, bool bIsTopGraph = true) const;
 
 	static void OffsetNodeIds(TArray<FPCGGraphTask>& Tasks, FPCGTaskId Offset, FPCGTaskId ParentId);
 
 	/** Propagates grid sizes through a graph's compiled tasks. */
-	void ResolveGridSizes(TArray<FPCGGraphTask>& InOutCompiledTasks, const FPCGStackContext& InStackContext) const;
+	void ResolveGridSizes(TArray<FPCGGraphTask>& InOutCompiledTasks, const FPCGStackContext& InStackContext, EPCGHiGenGrid GenerationDefaultGrid) const;
 
 #if WITH_EDITOR
 	void NotifyGraphChanged(UPCGGraph* InGraph);
 #endif
 
+	/** Flush all cached compiled graphs. */
+	void ClearCache();
+
 private:
 	TArray<FPCGGraphTask> CompileGraph(UPCGGraph* InGraph, FPCGTaskId& NextId, FPCGStackContext& InOutStackContext);
-	void CompileTopGraph(UPCGGraph* InGraph);
+	void CompileTopGraph(UPCGGraph* InGraph, uint32 GenerationGridSize);
 
-	/** Returns the grid for InCompiledTaskId which will be a concrete grid size or GenerationDefault if it cannot be determined statically. */
-	EPCGHiGenGrid CalculateGridRecursive(FPCGTaskId InTaskId, const FPCGStackContext& InStackContext, TArray<FPCGGraphTask>& InOutCompiledTasks) const;
+	/** Returns the grid for InCompiledTaskId. */
+	EPCGHiGenGrid CalculateGridRecursive(FPCGTaskId InTaskId, EPCGHiGenGrid GenerationDefaultGrid, const FPCGStackContext& InStackContext, TArray<FPCGGraphTask>& InOutCompiledTasks) const;
 
 	mutable FRWLock GraphToTaskMapLock;
 	TMap<UPCGGraph*, TArray<FPCGGraphTask>> GraphToTaskMap;
 	TMap<UPCGGraph*, FPCGStackContext> GraphToStackContext;
-	TMap<UPCGGraph*, TArray<FPCGGraphTask>> TopGraphToTaskMap;
+	// Top graphs are optimized for execution grid and store one set of compiled tasks per grid size.
+	TMap<UPCGGraph*, TMap<uint32, TArray<FPCGGraphTask>>> TopGraphToTaskMap;
 	TMap<UPCGGraph*, FPCGStackContext> TopGraphToStackContext;
 
 #if WITH_EDITOR
diff --git a/Engine/Plugins/Experimental/PCG/Source/PCG/Private/Graph/PCGGraphExecutor.cpp b/Engine/Plugins/Experimental/PCG/Source/PCG/Private/Graph/PCGGraphExecutor.cpp
--- a/Engine/Plugins/Experimental/PCG/Source/PCG/Private/Graph/PCGGraphExecutor.cpp
+++ b/Engine/Plugins/Experimental/PCG/Source/PCG/Private/Graph/PCGGraphExecutor.cpp
@@ -99,14 +99,15 @@ FPCGTaskId FPCGGraphExecutor::Schedule(UPCGComponent* Component, const TArray<FP
 
 FPCGTaskId FPCGGraphExecutor::Schedule(UPCGGraph* Graph, UPCGComponent* SourceComponent, FPCGElementPtr PreGraphElement, FPCGElementPtr InputElement, const TArray<FPCGTaskId>& ExternalDependencies, const FPCGStack* InFromStack)
 {
-	if (SourceComponent && IsGraphCacheDebuggingEnabled())
+	check(SourceComponent);
+
+	if (IsGraphCacheDebuggingEnabled())
 	{
 		UE_LOG(LogPCG, Log, TEXT("[%s] --- SCHEDULE GRAPH ---"), *SourceComponent->GetOwner()->GetName());
 	}
 
-	UPCGSubsystem* Subsystem = SourceComponent ? UPCGSubsystem::GetInstance(SourceComponent->GetWorld()) : nullptr;
 #if WITH_EDITOR
-	if (Subsystem)
+	if (UPCGSubsystem* Subsystem = SourceComponent ? UPCGSubsystem::GetInstance(SourceComponent->GetWorld()) : nullptr)
 	{
 		for (const UPCGNode* Node : Graph->GetNodes())
 		{
@@ -118,9 +119,12 @@ FPCGTaskId FPCGGraphExecutor::Schedule(UPCGGraph* Graph, UPCGComponent* SourceCo
 
 	FPCGTaskId ScheduledId = InvalidPCGTaskId;
 
+	const bool bNonPartitionedComponent = !SourceComponent->IsLocalComponent() && !SourceComponent->IsPartitioned();
+	const uint32 GenerationGridSize = bNonPartitionedComponent ? PCGHiGenGrid::UninitializedGridSize() : SourceComponent->GetGenerationGridSize();
+
 	// Get compiled tasks from compiler
 	TSharedPtr<FPCGStackContext> StackContextPtr = MakeShared<FPCGStackContext>();
-	TArray<FPCGGraphTask> CompiledTasks = GraphCompiler->GetCompiledTasks(Graph, *StackContextPtr);
+	TArray<FPCGGraphTask> CompiledTasks = GraphCompiler->GetCompiledTasks(Graph, GenerationGridSize, *StackContextPtr);
 
 	// Create the final stack context by including the current stack frames
 	if (InFromStack)
@@ -135,38 +139,6 @@ FPCGTaskId FPCGGraphExecutor::Schedule(UPCGGraph* Graph, UPCGComponent* SourceCo
 		Task.StackContext = StackContextPtr;
 	}
 
-	// Finalize grid sizes on tasks if graph has hierarchical generation enabled. Some tasks do not have a concrete grid size
-	// assigned at compile time, such as nodes outside of any authored grid size range. Also if we are generating on a non-partitioned
-	// component, knock out the grid sizes.
-	if (SourceComponent && SourceComponent->GetGraph() && SourceComponent->GetGraph()->IsHierarchicalGenerationEnabled())
-	{
-		const bool bNonPartitionedComponent = !SourceComponent->IsLocalComponent() && !SourceComponent->IsPartitioned();
-		const EPCGHiGenGrid Grid = SourceComponent->GetGenerationGrid();
-		const EPCGHiGenGrid DefaultGrid = PCGHiGenGrid::GridSizeToGrid(SourceComponent->GetGraph()->GetDefaultGridSize());
-
-		for (FPCGGraphTask& Task : CompiledTasks)
-		{
-			// Make a copy of this because many things can be scheduled over many frames
-			Task.GenerationGrid = Grid;
-
-			// Set graph generation grid size for this task
-			if (bNonPartitionedComponent)
-			{
-				// Non-partitioned component - throw away grid from all tasks which will force them all to run
-				Task.GraphGenerationGrid = EPCGHiGenGrid::Uninitialized;
-			}
-			else if (!!(Task.GraphGenerationGrid & EPCGHiGenGrid::GenerationDefault))
-			{
-				// Remove the generation default flag (but leave any other grid sizes that might have been set - linkage tasks
-				// will store the To grid size in the lower bits).
-				Task.GraphGenerationGrid &= ~EPCGHiGenGrid::GenerationDefault;
-
-				// Add default grid size
-				Task.GraphGenerationGrid |= DefaultGrid;
-			}
-		}
-	}
-
 	// Prepare scheduled task that will be promoted in the next Execute call.
 	if (CompiledTasks.Num() > 0)
 	{ 
@@ -1270,10 +1242,14 @@ FPCGElementPtr FPCGGraphExecutor::GetFetchInputElement()
 
 FPCGTaskId FPCGGraphExecutor::ScheduleDebugWithTaskCallback(UPCGComponent* InComponent, TFunction<void(FPCGTaskId/* TaskId*/, const UPCGNode*/* Node*/, const FPCGDataCollection&/* TaskOutput*/)> TaskCompleteCallback)
 {
+	check(InComponent);
 	FPCGTaskId FinalTaskID = Schedule(InComponent, {});
 
+	const bool bNonPartitionedComponent = !InComponent->IsLocalComponent() && !InComponent->IsPartitioned();
+	const uint32 GenerationGridSize = bNonPartitionedComponent ? PCGHiGenGrid::UninitializedGridSize() : InComponent->GetGenerationGridSize();
+
 	FPCGStackContext DummyStackContext;
-	TArray<FPCGGraphTask> CompiledTasks = GraphCompiler->GetCompiledTasks(InComponent->GetGraph(), DummyStackContext, /*bIsTopGraph=*/true);
+	TArray<FPCGGraphTask> CompiledTasks = GraphCompiler->GetCompiledTasks(InComponent->GetGraph(), GenerationGridSize, DummyStackContext, /*bIsTopGraph=*/true);
 	CompiledTasks.Pop(); // Remove the final task
 
 	// Set up all final dependencies for the entire execution
@@ -1501,20 +1477,14 @@ namespace PCGGraphExecutor
 	bool ExecuteGridLinkage(EPCGHiGenGrid InFromGrid, EPCGHiGenGrid InToGrid, const FString& InResourceKey, const FName& InOutputPinLabel, const UPCGNode* InDownstreamNode, FPCGGridLinkageContext* InContext)
 	{
 		// Non-hierarchical generation - no linkage required - data should just pass through.
-		if (!InContext->SourceComponent->GetGraph()->IsHierarchicalGenerationEnabled())
+		if (!InContext->SourceComponent->GetGraph()->IsHierarchicalGenerationEnabled()
+			|| !ensure(PCGHiGenGrid::IsValidGrid(InFromGrid) || InFromGrid == EPCGHiGenGrid::Unbounded))
 		{
 			InContext->OutputData = InContext->InputData;
 			return true;
 		}
 
-		EPCGHiGenGrid FromGrid = InFromGrid;
-		// If grid was not determined, apply default generation grid size.
-		if (!PCGHiGenGrid::IsValidGrid(FromGrid) && ensure(FromGrid == EPCGHiGenGrid::GenerationDefault))
-		{
-			FromGrid = InContext->SourceComponent->GetGraph()->GetDefaultGrid();
-		}
-
-		const uint32 FromGridSize = PCGHiGenGrid::IsValidGrid(FromGrid) ? PCGHiGenGrid::GridToGridSize(FromGrid) : PCGHiGenGrid::UnboundedGridSize();
+		const uint32 FromGridSize = PCGHiGenGrid::IsValidGrid(InFromGrid) ? PCGHiGenGrid::GridToGridSize(InFromGrid) : PCGHiGenGrid::UnboundedGridSize();
 		const uint32 ToGridSize = PCGHiGenGrid::IsValidGrid(InToGrid) ? PCGHiGenGrid::GridToGridSize(InToGrid) : PCGHiGenGrid::UnboundedGridSize();
 
 		// Never allow a large grid to read data from small grid - this violates hierarchy.
@@ -1553,7 +1523,7 @@ namespace PCGGraphExecutor
 			return true;
 		}
 
-		if (!!(FromGrid & InContext->GenerationGrid) && FromGridSize != ToGridSize)
+		if (!!(InFromGrid & InContext->GenerationGrid) && FromGridSize != ToGridSize)
 		{
 			PCGGraphExecutionLogging::LogGridLinkageTaskExecuteStore(InContext, FromGridSize, ToGridSize, InResourceKey);
 
diff --git a/Engine/Plugins/Experimental/PCG/Source/PCG/Private/PCGSubsystem.cpp b/Engine/Plugins/Experimental/PCG/Source/PCG/Private/PCGSubsystem.cpp
--- a/Engine/Plugins/Experimental/PCG/Source/PCG/Private/PCGSubsystem.cpp
+++ b/Engine/Plugins/Experimental/PCG/Source/PCG/Private/PCGSubsystem.cpp
@@ -7,6 +7,7 @@
 #include "PCGModule.h"
 #include "PCGWorldActor.h"
 #include "Graph/PCGGraphCache.h"
+#include "Graph/PCGGraphCompiler.h"
 #include "Graph/PCGGraphExecutor.h"
 #include "Graph/PCGStackContext.h"
 #include "Grid/PCGLandscapeCache.h"
@@ -35,7 +36,7 @@ namespace PCGSubsystemConsole
 {
 	static FAutoConsoleCommand CommandFlushCache(
 		TEXT("pcg.FlushCache"),
-		TEXT("Clears the PCG results cache."),
+		TEXT("Clears the PCG results cache and compiled graph cache."),
 		FConsoleCommandDelegate::CreateLambda([]()
 			{
 				if (UPCGSubsystem* PCGSubsystem = UPCGSubsystem::GetSubsystemForCurrentWorld())
@@ -1329,6 +1330,11 @@ void UPCGSubsystem::FlushCache()
 	if (GraphExecutor)
 	{
 		GraphExecutor->GetCache().ClearCache();
+
+		if (FPCGGraphCompiler* Compiler = GraphExecutor->GetCompiler())
+		{
+			Compiler->ClearCache();
+		}
 	}
 
 #if WITH_EDITOR
diff --git a/Engine/Plugins/Experimental/PCG/Source/PCG/Private/Utils/PCGExtraCapture.cpp b/Engine/Plugins/Experimental/PCG/Source/PCG/Private/Utils/PCGExtraCapture.cpp
--- a/Engine/Plugins/Experimental/PCG/Source/PCG/Private/Utils/PCGExtraCapture.cpp
+++ b/Engine/Plugins/Experimental/PCG/Source/PCG/Private/Utils/PCGExtraCapture.cpp
@@ -170,7 +170,8 @@ PCGUtils::FCallTreeInfo PCGUtils::FExtraCapture::CalculateCallTreeInfo(const UPC
 	}
 
 	FPCGStackContext DummyStackContext;
-	TArray<FPCGGraphTask> CompiledTasks = Compiler->GetPrecompiledTasks(Component->GetGraph(), DummyStackContext);
+	// Passed uninitialized grid size to get all tasks
+	TArray<FPCGGraphTask> CompiledTasks = Compiler->GetPrecompiledTasks(Component->GetGraph(), PCGHiGenGrid::UninitializedGridSize(), DummyStackContext);
 	if (CompiledTasks.IsEmpty())
 	{
 		return {};
diff --git a/Engine/Plugins/Experimental/PCG/Source/PCG/Public/PCGCommon.h b/Engine/Plugins/Experimental/PCG/Source/PCG/Public/PCGCommon.h
--- a/Engine/Plugins/Experimental/PCG/Source/PCG/Public/PCGCommon.h
+++ b/Engine/Plugins/Experimental/PCG/Source/PCG/Public/PCGCommon.h
@@ -153,9 +153,6 @@ enum class EPCGHiGenGrid : uint32
 
 	// Should execute once rather than executing on any grid
 	Unbounded = 2 * GridMax,
-
-	// Flag for grid not known yet. Represents the default grid size (which can be changed at runtime on PCGWorldActor)
-	GenerationDefault = 2 * Unbounded UMETA(Hidden),
 };
 ENUM_CLASS_FLAGS(EPCGHiGenGrid);
 
