[PCG] Added column tooltips to the attribute list view that shows the type

diff --git a/Engine/Plugins/Experimental/PCG/Source/PCGEditor/Private/SPCGEditorGraphAttributeListView.cpp b/Engine/Plugins/Experimental/PCG/Source/PCGEditor/Private/SPCGEditorGraphAttributeListView.cpp
--- a/Engine/Plugins/Experimental/PCG/Source/PCGEditor/Private/SPCGEditorGraphAttributeListView.cpp
+++ b/Engine/Plugins/Experimental/PCG/Source/PCGEditor/Private/SPCGEditorGraphAttributeListView.cpp
@@ -1177,6 +1177,8 @@ void SPCGEditorGraphAttributeListView::OnFilterTextCommitted(const FText& NewTex
 
 void SPCGEditorGraphAttributeListView::AddColumn(const UPCGPointData* InPCGPointData, const FName& InColumnId, const FText& ColumnLabel, EHorizontalAlignment HeaderHAlign, EHorizontalAlignment CellHAlign)
 {
+	FText ToolTip;
+	
 	if (InPCGPointData)
 	{
 		const FString ColumnIdString = InColumnId.ToString();
@@ -1209,6 +1211,8 @@ void SPCGEditorGraphAttributeListView::AddColumn(const UPCGPointData* InPCGPoint
 			ColumnData.DataAccessor = PCGAttributeAccessorHelpers::CreateConstAccessor(InPCGPointData, TargetSelector);
 			ColumnData.DataKeys = PCGAttributeAccessorHelpers::CreateConstKeys(InPCGPointData, TargetSelector);
 		}
+		
+		ToolTip = FText::FromString(PCG::Private::GetTypeName(ColumnData.DataAccessor->GetUnderlyingType()));
 	}
 
 	const float ColumnWidth = PCGEditorGraphAttributeListView::CalculateColumnWidth(ColumnLabel);
@@ -1216,6 +1220,7 @@ void SPCGEditorGraphAttributeListView::AddColumn(const UPCGPointData* InPCGPoint
 	SHeaderRow::FColumn::FArguments Arguments;
 	Arguments.ColumnId(InColumnId);
 	Arguments.DefaultLabel(ColumnLabel);
+	Arguments.DefaultTooltip(ToolTip);
 	Arguments.ManualWidth(ColumnWidth);
 	Arguments.HAlignHeader(HeaderHAlign);
 	Arguments.HAlignCell(CellHAlign);
@@ -1275,6 +1280,7 @@ void SPCGEditorGraphAttributeListView::AddMetadataColumn(const UPCGData* InPCGDa
 	}
 
 	FText ColumnLabel;
+	FText ToolTip;
 	FString ColumnIdString(OriginalColumnIdString);
 	if (InColumnId == NAME_None)
 	{
@@ -1292,6 +1298,8 @@ void SPCGEditorGraphAttributeListView::AddMetadataColumn(const UPCGData* InPCGDa
 		FPCGColumnData& ColumnData = PCGColumnData.Add(ColumnId);
 		ColumnData.DataAccessor = PCGAttributeAccessorHelpers::CreateConstAccessor(InPCGData, TargetSelector);
 		ColumnData.DataKeys = PCGAttributeAccessorHelpers::CreateConstKeys(InPCGData, TargetSelector);
+
+		ToolTip = FText::FromString(PCG::Private::GetTypeName(ColumnData.DataAccessor->GetUnderlyingType()));
 	}
 
 	if (ColumnLabel.IsEmpty())
@@ -1316,6 +1324,7 @@ void SPCGEditorGraphAttributeListView::AddMetadataColumn(const UPCGData* InPCGDa
 	SHeaderRow::FColumn::FArguments ColumnArguments;
 	ColumnArguments.ColumnId(ColumnId);
 	ColumnArguments.DefaultLabel(ColumnLabel);
+	ColumnArguments.DefaultTooltip(ToolTip);
 	ColumnArguments.HAlignHeader(EHorizontalAlignment::HAlign_Center);
 	ColumnArguments.HAlignCell(CellAlignment);
 	ColumnArguments.ManualWidth(ColumnWidth);
