Minor fix to add bounds check in the property to param node, even though it needs a major refactoring.
Lowered gravity of not found actors to verbose, as this is a relatively common case.

diff --git a/Engine/Plugins/Experimental/PCG/Source/PCG/Private/Elements/PCGDataFromActor.cpp b/Engine/Plugins/Experimental/PCG/Source/PCG/Private/Elements/PCGDataFromActor.cpp
--- a/Engine/Plugins/Experimental/PCG/Source/PCG/Private/Elements/PCGDataFromActor.cpp
+++ b/Engine/Plugins/Experimental/PCG/Source/PCG/Private/Elements/PCGDataFromActor.cpp
@@ -192,7 +192,7 @@ bool FPCGDataFromActorElement::ExecuteInternal(FPCGContext* InContext) const
 
 		if (Context->FoundActors.IsEmpty())
 		{
-			PCGE_LOG(Warning, GraphAndLog, LOCTEXT("NoActorFound", "No matching actor was found"));
+			PCGE_LOG(Verbose, LogOnly, LOCTEXT("NoActorFound", "No matching actor was found"));
 			return true;
 		}
 
diff --git a/Engine/Plugins/Experimental/PCG/Source/PCG/Private/Elements/PCGPropertyToParamData.cpp b/Engine/Plugins/Experimental/PCG/Source/PCG/Private/Elements/PCGPropertyToParamData.cpp
--- a/Engine/Plugins/Experimental/PCG/Source/PCG/Private/Elements/PCGPropertyToParamData.cpp
+++ b/Engine/Plugins/Experimental/PCG/Source/PCG/Private/Elements/PCGPropertyToParamData.cpp
@@ -7,6 +7,7 @@
 #include "PCGComponent.h"
 #include "PCGParamData.h"
 #include "Helpers/PCGBlueprintHelpers.h"
+#include "Helpers/PCGHelpers.h"
 #include "Helpers/PCGSettingsHelpers.h"
 #include "Metadata/Accessors/PCGAttributeAccessorHelpers.h"
 
@@ -120,14 +121,24 @@ bool FPCGPropertyToParamDataElement::ExecuteInternal(FPCGContext* Context) const
 
 	// First find the actor depending on the selection
 	UPCGComponent* OriginalComponent = UPCGBlueprintHelpers::GetOriginalComponent(*Context);
-	auto NoBoundsCheck = [](const AActor*) -> bool { return true; };
+	TFunction<bool(const AActor*)> BoundsCheck = [](const AActor*) -> bool { return true; };
 	auto NoSelfIgnoreCheck = [](const AActor*) -> bool { return true; };
 
-	AActor* FoundActor = PCGActorSelector::FindActor(Settings->ActorSelector, OriginalComponent, NoBoundsCheck, NoSelfIgnoreCheck);
+	if (OriginalComponent && OriginalComponent->GetOwner() && Settings->ActorSelector.bMustOverlapSelf)
+	{
+		const FBox ActorBounds = PCGHelpers::GetActorBounds(OriginalComponent->GetOwner());
+		BoundsCheck = [Settings, ActorBounds, OriginalComponent](const AActor* OtherActor) -> bool
+		{
+			const FBox OtherActorBounds = OtherActor ? PCGHelpers::GetGridBounds(OtherActor, OriginalComponent) : FBox(EForceInit::ForceInit);
+			return ActorBounds.Intersect(OtherActorBounds);
+		};
+	}
+
+	AActor* FoundActor = PCGActorSelector::FindActor(Settings->ActorSelector, OriginalComponent, BoundsCheck, NoSelfIgnoreCheck);
 
 	if (!FoundActor)
 	{
-		PCGE_LOG(Error, GraphAndLog, LOCTEXT("NoActorFound", "No matching actor was found"));
+		PCGE_LOG(Verbose, LogOnly, LOCTEXT("NoActorFound", "No matching actor was found"));
 		return true;
 	}
 
