[PCG] Enable caching of 'getter' nodes (get actor data, get landscape, etc) which works after the tracking system improvements.

#rnx

diff --git a/Engine/Plugins/Experimental/PCG/Source/PCG/Public/Elements/PCGDataFromActor.h b/Engine/Plugins/Experimental/PCG/Source/PCG/Public/Elements/PCGDataFromActor.h
--- a/Engine/Plugins/Experimental/PCG/Source/PCG/Public/Elements/PCGDataFromActor.h
+++ b/Engine/Plugins/Experimental/PCG/Source/PCG/Public/Elements/PCGDataFromActor.h
@@ -99,7 +99,6 @@ class FPCGDataFromActorElement : public IPCGElement
 public:
 	virtual FPCGContext* Initialize(const FPCGDataCollection& InputData, TWeakObjectPtr<UPCGComponent> SourceComponent, const UPCGNode* Node) override;
 	virtual bool CanExecuteOnlyOnMainThread(FPCGContext* Context) const override { return true; }
-	virtual bool IsCacheable(const UPCGSettings* InSettings) const override { return false; }
 
 protected:
 	virtual bool ExecuteInternal(FPCGContext* Context) const;
