[PCG] Fixed an issue were OriginalSettings object wasn't set after pasting a settings instance node.

diff --git a/Engine/Plugins/Experimental/PCG/Source/PCG/Private/PCGSettings.cpp b/Engine/Plugins/Experimental/PCG/Source/PCG/Private/PCGSettings.cpp
--- a/Engine/Plugins/Experimental/PCG/Source/PCG/Private/PCGSettings.cpp
+++ b/Engine/Plugins/Experimental/PCG/Source/PCG/Private/PCGSettings.cpp
@@ -753,15 +753,7 @@ void UPCGSettingsInstance::PostLoad()
 	Super::PostLoad();
 
 #if WITH_EDITOR
-	if (Settings)
-	{
-		Settings->OnSettingsChangedDelegate.AddUObject(this, &UPCGSettingsInstance::OnSettingsChanged);
-		Settings->ConditionalPostLoad();
-	}
-#endif
-
-#if WITH_EDITOR
-	OriginalSettings = Settings;
+	PostSettingsChanged();
 #endif
 }
 
@@ -777,6 +769,15 @@ void UPCGSettingsInstance::BeginDestroy()
 	Super::BeginDestroy();
 }
 
+void UPCGSettingsInstance::PostEditImport()
+{
+	Super::PostEditImport();
+
+#if WITH_EDITOR
+	PostSettingsChanged();
+#endif
+}
+
 void UPCGSettingsInstance::SetSettings(UPCGSettings* InSettings)
 {
 #if WITH_EDITOR
@@ -787,15 +788,9 @@ void UPCGSettingsInstance::SetSettings(UPCGSettings* InSettings)
 #endif
 
 	Settings = InSettings;
-#if WITH_EDITOR
-	OriginalSettings = Settings;
-#endif
 
 #if WITH_EDITOR
-	if (Settings)
-	{
-		Settings->OnSettingsChangedDelegate.AddUObject(this, &UPCGSettingsInstance::OnSettingsChanged);
-	}
+	PostSettingsChanged();
 #endif
 }
 
@@ -815,6 +810,17 @@ void UPCGSettingsInstance::OnSettingsChanged(UPCGSettings* InSettings, EPCGChang
 		OnSettingsChangedDelegate.Broadcast(InSettings, ChangeType);
 	}
 }
+
+void UPCGSettingsInstance::PostSettingsChanged()
+{
+	OriginalSettings = Settings;
+
+	if (Settings)
+	{
+		Settings->OnSettingsChangedDelegate.AddUObject(this, &UPCGSettingsInstance::OnSettingsChanged);
+		Settings->ConditionalPostLoad();
+	}
+}
 #endif
 
 UPCGTrivialSettings::UPCGTrivialSettings()
diff --git a/Engine/Plugins/Experimental/PCG/Source/PCG/Public/PCGSettings.h b/Engine/Plugins/Experimental/PCG/Source/PCG/Public/PCGSettings.h
--- a/Engine/Plugins/Experimental/PCG/Source/PCG/Public/PCGSettings.h
+++ b/Engine/Plugins/Experimental/PCG/Source/PCG/Public/PCGSettings.h
@@ -436,6 +436,7 @@ public:
 	// ~Begin UObject interface
 	virtual void PostLoad() override;
 	virtual void BeginDestroy() override;
+	virtual void PostEditImport() override;
 
 #if WITH_EDITOR
 	virtual void PostEditChangeProperty(FPropertyChangedEvent& PropertyChangedEvent) override;
@@ -444,8 +445,11 @@ public:
 
 	void SetSettings(UPCGSettings* InSettings);
 
+protected:
 #if WITH_EDITOR
 	void OnSettingsChanged(UPCGSettings* InSettings, EPCGChangeType ChangeType);
+
+	void PostSettingsChanged();
 #endif
 
 #if WITH_EDITORONLY_DATA
@@ -453,6 +457,7 @@ public:
 	TObjectPtr<UPCGSettings> OriginalSettings = nullptr; // Transient just for display
 #endif
 
+public:
 	UPROPERTY(VisibleAnywhere, BlueprintReadOnly, Category = Instance, meta = (EditInline))
 	TObjectPtr<UPCGSettings> Settings = nullptr;
 };
