Removed limitation on object types in the property to param helpers

diff --git a/Engine/Plugins/Experimental/PCG/Source/PCG/Public/Helpers/PCGPropertyHelpers.h b/Engine/Plugins/Experimental/PCG/Source/PCG/Public/Helpers/PCGPropertyHelpers.h
--- a/Engine/Plugins/Experimental/PCG/Source/PCG/Public/Helpers/PCGPropertyHelpers.h
+++ b/Engine/Plugins/Experimental/PCG/Source/PCG/Public/Helpers/PCGPropertyHelpers.h
@@ -57,20 +57,6 @@ inline decltype(auto) PCGPropertyHelpers::GetPropertyValueWithCallback(const Obj
 		return false;
 	}
 
-	// Object are not yet supported as accessors
-	if (const FObjectPropertyBase* ObjectProperty = CastField<FObjectPropertyBase>(InProperty))
-	{
-		if (const UObject* Object = ObjectProperty->GetObjectPropertyValue(InProperty->ContainerPtrToValuePtr<void>(InObject)))
-		{
-			// Object are transformed into their soft path name (as a string attribute)
-			return InFunc(Object->GetPathName());
-		}
-		else
-		{
-			return false;
-		}
-	}
-
 	TUniquePtr<IPCGAttributeAccessor> PropertyAccessor = PCGAttributeAccessorHelpers::CreatePropertyAccessor(InProperty);
 
 	if (!PropertyAccessor.IsValid())
