[PCG] Fix broken callbacks on subgraph duplicates

#rnx

diff --git a/Engine/Plugins/Experimental/PCG/Source/PCG/Private/PCGGraph.cpp b/Engine/Plugins/Experimental/PCG/Source/PCG/Private/PCGGraph.cpp
--- a/Engine/Plugins/Experimental/PCG/Source/PCG/Private/PCGGraph.cpp
+++ b/Engine/Plugins/Experimental/PCG/Source/PCG/Private/PCGGraph.cpp
@@ -1481,7 +1481,7 @@ void UPCGGraphInstance::ResetPropertyToDefault(const FProperty* InProperty)
 
 bool UPCGGraphInstance::IsPropertyOverriddenAndNotDefault(const FProperty* InProperty) const
 {
-	return ParametersOverrides.IsPropertyOverriddenAndNotDefault(InProperty, Graph->GetUserParametersStruct());
+	return Graph ? ParametersOverrides.IsPropertyOverriddenAndNotDefault(InProperty, Graph->GetUserParametersStruct()) : false;
 }
 
 bool FPCGOverrideInstancedPropertyBag::RefreshParameters(const FInstancedPropertyBag* ParentUserParameters, EPCGGraphParameterEvent InChangeType, FName InChangedPropertyName)
diff --git a/Engine/Plugins/Experimental/PCG/Source/PCG/Private/PCGSubgraph.cpp b/Engine/Plugins/Experimental/PCG/Source/PCG/Private/PCGSubgraph.cpp
--- a/Engine/Plugins/Experimental/PCG/Source/PCG/Private/PCGSubgraph.cpp
+++ b/Engine/Plugins/Experimental/PCG/Source/PCG/Private/PCGSubgraph.cpp
@@ -33,6 +33,26 @@ namespace PCGSubgraphSettings
 	}
 }
 
+#if WITH_EDITOR
+void UPCGBaseSubgraphSettings::SetupCallbacks()
+{
+	UPCGGraphInterface* Subgraph = GetSubgraphInterface();
+
+	if (Subgraph && !Subgraph->OnGraphChangedDelegate.IsBoundToObject(this))
+	{
+		Subgraph->OnGraphChangedDelegate.AddUObject(this, &UPCGBaseSubgraphSettings::OnSubgraphChanged);
+	}
+}
+
+void UPCGBaseSubgraphSettings::TeardownCallbacks()
+{
+	if (UPCGGraphInterface* Subgraph = GetSubgraphInterface())
+	{
+		Subgraph->OnGraphChangedDelegate.RemoveAll(this);
+	}
+}
+#endif // WITH_EDITOR
+
 UPCGGraph* UPCGBaseSubgraphSettings::GetSubgraph() const
 {
 	UPCGGraphInterface* SubgraphInterface = GetSubgraphInterface();
@@ -44,10 +64,7 @@ void UPCGBaseSubgraphSettings::PostInitProperties()
 	Super::PostInitProperties();
 
 #if WITH_EDITOR
-	if (UPCGGraphInterface* Subgraph = GetSubgraphInterface())
-	{
-		Subgraph->OnGraphChangedDelegate.AddUObject(this, &UPCGBaseSubgraphSettings::OnSubgraphChanged);
-	}
+	SetupCallbacks();
 #endif
 }
 
@@ -56,32 +73,38 @@ void UPCGBaseSubgraphSettings::PostLoad()
 	Super::PostLoad();
 
 #if WITH_EDITOR
-	if (UPCGGraphInterface* Subgraph = GetSubgraphInterface())
-	{
-		// We might have already connected in PostInitProperties
-		// To be sure, remove it and re-add it.
-		Subgraph->OnGraphChangedDelegate.RemoveAll(this);
-		Subgraph->OnGraphChangedDelegate.AddUObject(this, &UPCGBaseSubgraphSettings::OnSubgraphChanged);
-	}
+	SetupCallbacks();
+#endif
+}
+
+void UPCGBaseSubgraphSettings::PostDuplicate(bool bDuplicateForPIE)
+{
+	Super::PostDuplicate(bDuplicateForPIE);
+
+#if WITH_EDITOR
+	SetupCallbacks();
+#endif
+}
+
+void UPCGBaseSubgraphSettings::PostEditImport()
+{
+	Super::PostEditImport();
+
+#if WITH_EDITOR
+	SetupCallbacks();
 #endif
 }
 
 void UPCGBaseSubgraphSettings::SetSubgraph(UPCGGraphInterface* InGraph)
 {
 #if WITH_EDITOR
-	if (UPCGGraphInterface* Subgraph = GetSubgraphInterface())
-	{
-		Subgraph->OnGraphChangedDelegate.RemoveAll(this);
-	}
+	TeardownCallbacks();
 #endif // WITH_EDITOR
 
 	SetSubgraphInternal(InGraph);
 
 #if WITH_EDITOR
-	if (UPCGGraphInterface* Subgraph = GetSubgraphInterface())
-	{
-		Subgraph->OnGraphChangedDelegate.AddUObject(this, &UPCGBaseSubgraphSettings::OnSubgraphChanged);
-	}
+	SetupCallbacks();
 #endif // WITH_EDITOR
 
 	// Also, reconstruct overrides
@@ -91,24 +114,32 @@ void UPCGBaseSubgraphSettings::SetSubgraph(UPCGGraphInterface* InGraph)
 void UPCGBaseSubgraphSettings::BeginDestroy()
 {
 #if WITH_EDITOR
-	if (UPCGGraphInterface* Subgraph = GetSubgraphInterface())
-	{
-		Subgraph->OnGraphChangedDelegate.RemoveAll(this);
-	}
+	TeardownCallbacks();
 #endif
 
 	Super::BeginDestroy();
 }
 
 #if WITH_EDITOR
+void UPCGBaseSubgraphSettings::PreEditUndo()
+{
+	Super::PreEditUndo();
+
+	TeardownCallbacks();
+}
+
+void UPCGBaseSubgraphSettings::PostEditUndo()
+{
+	Super::PostEditUndo();
+
+	SetupCallbacks();
+}
+
 void UPCGBaseSubgraphSettings::PreEditChange(FProperty* PropertyAboutToChange)
 {
 	if (PropertyAboutToChange && IsStructuralProperty(PropertyAboutToChange->GetFName()))
 	{
-		if (UPCGGraphInterface* Subgraph = GetSubgraphInterface())
-		{
-			Subgraph->OnGraphChangedDelegate.RemoveAll(this);
-		}
+		TeardownCallbacks();
 	}
 
 	Super::PreEditChange(PropertyAboutToChange);
@@ -120,10 +151,7 @@ void UPCGBaseSubgraphSettings::PostEditChangeProperty(struct FPropertyChangedEve
 
 	if (PropertyChangedEvent.Property && IsStructuralProperty(PropertyChangedEvent.Property->GetFName()))
 	{
-		if (UPCGGraphInterface* Subgraph = GetSubgraphInterface())
-		{
-			Subgraph->OnGraphChangedDelegate.AddUObject(this, &UPCGBaseSubgraphSettings::OnSubgraphChanged);
-		}
+		SetupCallbacks();
 	}
 }
 
diff --git a/Engine/Plugins/Experimental/PCG/Source/PCG/Public/PCGSubgraph.h b/Engine/Plugins/Experimental/PCG/Source/PCG/Public/PCGSubgraph.h
--- a/Engine/Plugins/Experimental/PCG/Source/PCG/Public/PCGSubgraph.h
+++ b/Engine/Plugins/Experimental/PCG/Source/PCG/Public/PCGSubgraph.h
@@ -33,8 +33,13 @@ protected:
 	//~Begin UObject interface implementation
 	virtual void PostLoad() override;
 	virtual void BeginDestroy() override;
+	virtual void PostDuplicate(bool bDuplicateForPIE) override;
+	virtual void PostEditImport() override;
 	virtual void PostInitProperties() override;
+
 #if WITH_EDITOR
+	virtual void PreEditUndo() override;
+	virtual void PostEditUndo() override;
 	virtual void PreEditChange(FProperty* PropertyAboutToChange) override;
 	virtual void PostEditChangeProperty(struct FPropertyChangedEvent& PropertyChangedEvent) override;
 	//~End UObject interface implementation
@@ -42,6 +47,9 @@ protected:
 	//~Begin UPCGSettings interface
 	virtual void GetTrackedActorKeys(FPCGActorSelectionKeyToSettingsMap& OutKeysToSettings, TArray<TObjectPtr<const UPCGGraph>>& OutVisitedGraphs) const override;
 	virtual bool IsStructuralProperty(const FName& InPropertyName) const override;
+
+	void SetupCallbacks();
+	void TeardownCallbacks();
 #endif
 
 	TArray<FPCGPinProperties> InputPinProperties() const override;
