[PCG] Deprecate GetParamsWithDeprecation since it was incorrectly used to detect overrides

* This function was there to support overrides using the In pin, whihc was the very old way to do overrides.
* But this flow is not supported anymore anyway, so deprecate it.

diff --git a/Engine/Plugins/Experimental/PCG/Source/PCG/Private/PCGContext.cpp b/Engine/Plugins/Experimental/PCG/Source/PCG/Private/PCGContext.cpp
--- a/Engine/Plugins/Experimental/PCG/Source/PCG/Private/PCGContext.cpp
+++ b/Engine/Plugins/Experimental/PCG/Source/PCG/Private/PCGContext.cpp
@@ -120,7 +120,7 @@ void FPCGContext::InitializeSettings()
 		const TArray<FPCGSettingsOverridableParam>& OverridableParams = NodeSettings->OverridableParams();
 		if (!OverridableParams.IsEmpty())
 		{
-			bool bHasParamConnected = InputData.GetParamsWithDeprecation(Node) != nullptr;
+			bool bHasParamConnected = !InputData.GetParamsByPin(PCGPinConstants::DefaultParamsLabel).IsEmpty();
 
 			int32 Index = 0;
 			while (!bHasParamConnected && Index < OverridableParams.Num())
diff --git a/Engine/Plugins/Experimental/PCG/Source/PCG/Private/PCGSubgraph.cpp b/Engine/Plugins/Experimental/PCG/Source/PCG/Private/PCGSubgraph.cpp
--- a/Engine/Plugins/Experimental/PCG/Source/PCG/Private/PCGSubgraph.cpp
+++ b/Engine/Plugins/Experimental/PCG/Source/PCG/Private/PCGSubgraph.cpp
@@ -410,7 +410,7 @@ FPCGContext* FPCGSubgraphElement::Initialize(const FPCGDataCollection& InputData
 
 	if (!OverridableParams.IsEmpty() && UserParametersView.IsValid())
 	{
-		bool bHasParamConnected = InputData.GetParamsWithDeprecation(Node) != nullptr;
+		bool bHasParamConnected = !InputData.GetParamsByPin(PCGPinConstants::DefaultParamsLabel).IsEmpty();
 
 		int32 Index = 0;
 		while (!bHasParamConnected && Index < OverridableParams.Num())
diff --git a/Engine/Plugins/Experimental/PCG/Source/PCG/Public/PCGData.h b/Engine/Plugins/Experimental/PCG/Source/PCG/Public/PCGData.h
--- a/Engine/Plugins/Experimental/PCG/Source/PCG/Public/PCGData.h
+++ b/Engine/Plugins/Experimental/PCG/Source/PCG/Public/PCGData.h
@@ -145,6 +145,7 @@ struct PCG_API FPCGDataCollection
 
 	// Only used as a temporary solution for old graph with nodes that didn't have params pins.
 	// Should NOT be used with new nodes.
+	UE_DEPRECATED(5.4, "Was not supposed to be used anyway, you should query the data per pin using GetParamsByPin")
 	UPCGParamData* GetParamsWithDeprecation(const UPCGNode* Node) const;
 
 	/** Returns the first/only param found on the default params pin */
