[PCG] Add support for override on Attribute Selectors

diff --git a/Engine/Plugins/Experimental/PCG/Source/PCG/Private/Metadata/Accessors/PCGAttributeAccessorHelpers.cpp b/Engine/Plugins/Experimental/PCG/Source/PCG/Private/Metadata/Accessors/PCGAttributeAccessorHelpers.cpp
--- a/Engine/Plugins/Experimental/PCG/Source/PCG/Private/Metadata/Accessors/PCGAttributeAccessorHelpers.cpp
+++ b/Engine/Plugins/Experimental/PCG/Source/PCG/Private/Metadata/Accessors/PCGAttributeAccessorHelpers.cpp
@@ -330,6 +330,10 @@ namespace PCGAttributeAccessorHelpers
 			{
 				return Functor(Signature<FPCGPropertyPathAccessor<FSoftClassPath>>{}, InProperty);
 			}
+			else if (StructProperty->Struct->IsA(FPCGAttributePropertySelector::StaticStruct()->GetClass()))
+			{
+				return Functor(Signature<FPCGAttributePropertySelectorAccessor>{}, StructProperty);
+			}
 		}
 
 		return ReturnType{};
diff --git a/Engine/Plugins/Experimental/PCG/Source/PCG/Public/Metadata/Accessors/PCGPropertyAccessor.h b/Engine/Plugins/Experimental/PCG/Source/PCG/Public/Metadata/Accessors/PCGPropertyAccessor.h
--- a/Engine/Plugins/Experimental/PCG/Source/PCG/Public/Metadata/Accessors/PCGPropertyAccessor.h
+++ b/Engine/Plugins/Experimental/PCG/Source/PCG/Public/Metadata/Accessors/PCGPropertyAccessor.h
@@ -4,6 +4,8 @@
 
 #include "IPCGAttributeAccessorTpl.h"
 
+#include "Metadata/PCGAttributePropertySelector.h"
+
 #include "Containers/UnrealString.h" // IWYU pragma: keep
 #include "UObject/SoftObjectPath.h" // IWYU pragma: keep
 #include "UObject/UnrealType.h" // IWYU pragma: keep
@@ -333,4 +335,42 @@ public:
 
 private:
 	const FObjectProperty* Property = nullptr;
-};
\ No newline at end of file
+};
+
+/**
+* Special accessor to support attribute selector overrides. Interface with a string.
+* Key supported: All
+*/
+class FPCGAttributePropertySelectorAccessor : public IPCGAttributeAccessorT<FPCGAttributePropertySelectorAccessor>
+{
+public:
+	using Type = FString;
+	using Super = IPCGAttributeAccessorT<FPCGAttributePropertySelectorAccessor>;
+
+	FPCGAttributePropertySelectorAccessor(const FStructProperty* InProperty)
+		: Super(/*bInReadOnly=*/ false)
+		, Property(InProperty)
+	{
+		ensure(InProperty && InProperty->Struct && InProperty->Struct->IsA(FPCGAttributePropertySelector::StaticStruct()->GetClass()));
+	}
+
+	bool GetRangeImpl(TArrayView<FString> OutValues, int32 Index, const IPCGAttributeAccessorKeys& Keys) const
+	{
+		return PCGPropertyAccessor::IterateGet(Property, OutValues, Index, Keys, [this](const void* PropertyAddressData) -> FString
+		{
+			return reinterpret_cast<const FPCGAttributePropertySelector*>(PropertyAddressData)->GetDisplayText().ToString();
+		});
+	}
+
+	bool SetRangeImpl(TArrayView<const FString> InValues, int32 Index, IPCGAttributeAccessorKeys& Keys, EPCGAttributeAccessorFlags Flags)
+	{
+		return PCGPropertyAccessor::IterateSet(Property, InValues, Index, Keys, [this](void* PropertyAddressData, const FString& Value) -> void
+		{
+			reinterpret_cast<FPCGAttributePropertySelector*>(PropertyAddressData)->Update(Value);
+		});
+	}
+
+private:
+	const FStructProperty* Property = nullptr;
+};
+
