[PCG] Include PCGPin.h in PCGSettings.h

* Because we are returning an array of PinProperties and it is always the responsibility of the child class to include it, which doesn't make a lot of sense
* Remove all includes from children

#rnx

diff --git a/Engine/Plugins/Experimental/PCG/Source/PCG/Private/Tests/GraphAuthoring/PCGGraphAuthoringTestHelperSettings.h b/Engine/Plugins/Experimental/PCG/Source/PCG/Private/Tests/GraphAuthoring/PCGGraphAuthoringTestHelperSettings.h
--- a/Engine/Plugins/Experimental/PCG/Source/PCG/Private/Tests/GraphAuthoring/PCGGraphAuthoringTestHelperSettings.h
+++ b/Engine/Plugins/Experimental/PCG/Source/PCG/Private/Tests/GraphAuthoring/PCGGraphAuthoringTestHelperSettings.h
@@ -3,8 +3,6 @@
 
 #include "PCGSettings.h"
 
-#include "PCGPin.h"
-
 #include "PCGGraphAuthoringTestHelperSettings.generated.h"
 
 /** Testing helper - generates a node with a single input and output pin of the stipulated type. */
diff --git a/Engine/Plugins/Experimental/PCG/Source/PCG/Public/Elements/IO/PCGExternalData.h b/Engine/Plugins/Experimental/PCG/Source/PCG/Public/Elements/IO/PCGExternalData.h
--- a/Engine/Plugins/Experimental/PCG/Source/PCG/Public/Elements/IO/PCGExternalData.h
+++ b/Engine/Plugins/Experimental/PCG/Source/PCG/Public/Elements/IO/PCGExternalData.h
@@ -4,7 +4,6 @@
 
 #include "PCGContext.h"
 #include "PCGElement.h"
-#include "PCGPin.h"
 #include "PCGSettings.h"
 #include "Metadata/PCGAttributePropertySelector.h"
 
diff --git a/Engine/Plugins/Experimental/PCG/Source/PCG/Public/Elements/Metadata/PCGMetadataElement.h b/Engine/Plugins/Experimental/PCG/Source/PCG/Public/Elements/Metadata/PCGMetadataElement.h
--- a/Engine/Plugins/Experimental/PCG/Source/PCG/Public/Elements/Metadata/PCGMetadataElement.h
+++ b/Engine/Plugins/Experimental/PCG/Source/PCG/Public/Elements/Metadata/PCGMetadataElement.h
@@ -2,7 +2,6 @@
 
 #pragma once
 
-#include "PCGPin.h"
 #include "PCGSettings.h"
 #include "Metadata/PCGAttributePropertySelector.h"
 
diff --git a/Engine/Plugins/Experimental/PCG/Source/PCG/Public/Elements/Metadata/PCGMetadataPartition.h b/Engine/Plugins/Experimental/PCG/Source/PCG/Public/Elements/Metadata/PCGMetadataPartition.h
--- a/Engine/Plugins/Experimental/PCG/Source/PCG/Public/Elements/Metadata/PCGMetadataPartition.h
+++ b/Engine/Plugins/Experimental/PCG/Source/PCG/Public/Elements/Metadata/PCGMetadataPartition.h
@@ -2,7 +2,6 @@
 
 #pragma once
 
-#include "PCGPin.h"
 #include "PCGSettings.h"
 
 #include "PCGMetadataPartition.generated.h"
diff --git a/Engine/Plugins/Experimental/PCG/Source/PCG/Public/Elements/PCGAttributeNoise.h b/Engine/Plugins/Experimental/PCG/Source/PCG/Public/Elements/PCGAttributeNoise.h
--- a/Engine/Plugins/Experimental/PCG/Source/PCG/Public/Elements/PCGAttributeNoise.h
+++ b/Engine/Plugins/Experimental/PCG/Source/PCG/Public/Elements/PCGAttributeNoise.h
@@ -4,7 +4,6 @@
 
 #include "PCGContext.h"
 #include "PCGElement.h"
-#include "PCGPin.h"
 #include "PCGSettings.h"
 #include "Metadata/PCGAttributePropertySelector.h"
 #include "Metadata/Accessors/IPCGAttributeAccessor.h"
diff --git a/Engine/Plugins/Experimental/PCG/Source/PCG/Public/Elements/PCGBoundsModifier.h b/Engine/Plugins/Experimental/PCG/Source/PCG/Public/Elements/PCGBoundsModifier.h
--- a/Engine/Plugins/Experimental/PCG/Source/PCG/Public/Elements/PCGBoundsModifier.h
+++ b/Engine/Plugins/Experimental/PCG/Source/PCG/Public/Elements/PCGBoundsModifier.h
@@ -2,7 +2,6 @@
 
 #pragma once
 
-#include "PCGPin.h"
 #include "PCGSettings.h"
 
 #include "Elements/PCGPointProcessingElementBase.h"
diff --git a/Engine/Plugins/Experimental/PCG/Source/PCG/Public/Elements/PCGCollapseElement.h b/Engine/Plugins/Experimental/PCG/Source/PCG/Public/Elements/PCGCollapseElement.h
--- a/Engine/Plugins/Experimental/PCG/Source/PCG/Public/Elements/PCGCollapseElement.h
+++ b/Engine/Plugins/Experimental/PCG/Source/PCG/Public/Elements/PCGCollapseElement.h
@@ -2,7 +2,6 @@
 
 #pragma once
 
-#include "PCGPin.h"
 #include "PCGSettings.h"
 
 #include "PCGCollapseElement.generated.h"
diff --git a/Engine/Plugins/Experimental/PCG/Source/PCG/Public/Elements/PCGCopyPoints.h b/Engine/Plugins/Experimental/PCG/Source/PCG/Public/Elements/PCGCopyPoints.h
--- a/Engine/Plugins/Experimental/PCG/Source/PCG/Public/Elements/PCGCopyPoints.h
+++ b/Engine/Plugins/Experimental/PCG/Source/PCG/Public/Elements/PCGCopyPoints.h
@@ -2,7 +2,6 @@
 
 #pragma once
 
-#include "PCGPin.h"
 #include "PCGSettings.h"
 
 #include "PCGCopyPoints.generated.h"
diff --git a/Engine/Plugins/Experimental/PCG/Source/PCG/Public/Elements/PCGCreateSpline.h b/Engine/Plugins/Experimental/PCG/Source/PCG/Public/Elements/PCGCreateSpline.h
--- a/Engine/Plugins/Experimental/PCG/Source/PCG/Public/Elements/PCGCreateSpline.h
+++ b/Engine/Plugins/Experimental/PCG/Source/PCG/Public/Elements/PCGCreateSpline.h
@@ -3,7 +3,6 @@
 #pragma once
 
 #include "PCGElement.h"
-#include "PCGPin.h"
 #include "PCGSettings.h"
 #include "Data/PCGSplineData.h"
 
diff --git a/Engine/Plugins/Experimental/PCG/Source/PCG/Public/Elements/PCGCreateTargetActor.h b/Engine/Plugins/Experimental/PCG/Source/PCG/Public/Elements/PCGCreateTargetActor.h
--- a/Engine/Plugins/Experimental/PCG/Source/PCG/Public/Elements/PCGCreateTargetActor.h
+++ b/Engine/Plugins/Experimental/PCG/Source/PCG/Public/Elements/PCGCreateTargetActor.h
@@ -2,7 +2,6 @@
 
 #pragma once
 
-#include "PCGPin.h"
 #include "PCGSettings.h"
 #include "Metadata/PCGAttributePropertySelector.h"
 
diff --git a/Engine/Plugins/Experimental/PCG/Source/PCG/Public/Elements/PCGDataFromActor.h b/Engine/Plugins/Experimental/PCG/Source/PCG/Public/Elements/PCGDataFromActor.h
--- a/Engine/Plugins/Experimental/PCG/Source/PCG/Public/Elements/PCGDataFromActor.h
+++ b/Engine/Plugins/Experimental/PCG/Source/PCG/Public/Elements/PCGDataFromActor.h
@@ -6,7 +6,6 @@
 #include "PCGSettings.h"
 #include "Elements/PCGActorSelector.h"
 
-#include "PCGPin.h"
 #include "PCGDataFromActor.generated.h"
 
 UENUM()
diff --git a/Engine/Plugins/Experimental/PCG/Source/PCG/Public/Elements/PCGDensityFilter.h b/Engine/Plugins/Experimental/PCG/Source/PCG/Public/Elements/PCGDensityFilter.h
--- a/Engine/Plugins/Experimental/PCG/Source/PCG/Public/Elements/PCGDensityFilter.h
+++ b/Engine/Plugins/Experimental/PCG/Source/PCG/Public/Elements/PCGDensityFilter.h
@@ -2,8 +2,6 @@
 
 #pragma once
 
-
-#include "PCGPin.h"
 #include "PCGSettings.h"
 
 #include "PCGDensityFilter.generated.h"
diff --git a/Engine/Plugins/Experimental/PCG/Source/PCG/Public/Elements/PCGDensityRemap.h b/Engine/Plugins/Experimental/PCG/Source/PCG/Public/Elements/PCGDensityRemap.h
--- a/Engine/Plugins/Experimental/PCG/Source/PCG/Public/Elements/PCGDensityRemap.h
+++ b/Engine/Plugins/Experimental/PCG/Source/PCG/Public/Elements/PCGDensityRemap.h
@@ -2,7 +2,6 @@
 
 #pragma once
 
-#include "PCGPin.h"
 #include "PCGSettings.h"
 
 #include "PCGDensityRemap.generated.h"
diff --git a/Engine/Plugins/Experimental/PCG/Source/PCG/Public/Elements/PCGDensityRemapElement.h b/Engine/Plugins/Experimental/PCG/Source/PCG/Public/Elements/PCGDensityRemapElement.h
--- a/Engine/Plugins/Experimental/PCG/Source/PCG/Public/Elements/PCGDensityRemapElement.h
+++ b/Engine/Plugins/Experimental/PCG/Source/PCG/Public/Elements/PCGDensityRemapElement.h
@@ -2,7 +2,6 @@
 
 #pragma once
 
-#include "PCGPin.h"
 #include "PCGSettings.h"
 
 #include "Elements/PCGPointProcessingElementBase.h"
diff --git a/Engine/Plugins/Experimental/PCG/Source/PCG/Public/Elements/PCGMakeConcreteElement.h b/Engine/Plugins/Experimental/PCG/Source/PCG/Public/Elements/PCGMakeConcreteElement.h
--- a/Engine/Plugins/Experimental/PCG/Source/PCG/Public/Elements/PCGMakeConcreteElement.h
+++ b/Engine/Plugins/Experimental/PCG/Source/PCG/Public/Elements/PCGMakeConcreteElement.h
@@ -4,8 +4,6 @@
 
 #include "PCGSettings.h"
 
-#include "PCGPin.h"
-
 #include "PCGMakeConcreteElement.generated.h"
 
 /** Makes data Concrete, collapsing to Point if necessary. Discards non-Spatial data. */
diff --git a/Engine/Plugins/Experimental/PCG/Source/PCG/Public/Elements/PCGMutateSeed.h b/Engine/Plugins/Experimental/PCG/Source/PCG/Public/Elements/PCGMutateSeed.h
--- a/Engine/Plugins/Experimental/PCG/Source/PCG/Public/Elements/PCGMutateSeed.h
+++ b/Engine/Plugins/Experimental/PCG/Source/PCG/Public/Elements/PCGMutateSeed.h
@@ -4,7 +4,6 @@
 
 #include "PCGSettings.h"
 
-#include "PCGPin.h"
 #include "Elements/PCGPointProcessingElementBase.h"
 
 #include "UObject/NameTypes.h"
diff --git a/Engine/Plugins/Experimental/PCG/Source/PCG/Public/Elements/PCGNormalToDensity.h b/Engine/Plugins/Experimental/PCG/Source/PCG/Public/Elements/PCGNormalToDensity.h
--- a/Engine/Plugins/Experimental/PCG/Source/PCG/Public/Elements/PCGNormalToDensity.h
+++ b/Engine/Plugins/Experimental/PCG/Source/PCG/Public/Elements/PCGNormalToDensity.h
@@ -1,7 +1,6 @@
 // Copyright Epic Games, Inc. All Rights Reserved.
 #pragma once
 
-#include "PCGPin.h"
 #include "PCGSettings.h"
 #include "Elements/PCGPointProcessingElementBase.h"
 
diff --git a/Engine/Plugins/Experimental/PCG/Source/PCG/Public/Elements/PCGPointExtentsModifier.h b/Engine/Plugins/Experimental/PCG/Source/PCG/Public/Elements/PCGPointExtentsModifier.h
--- a/Engine/Plugins/Experimental/PCG/Source/PCG/Public/Elements/PCGPointExtentsModifier.h
+++ b/Engine/Plugins/Experimental/PCG/Source/PCG/Public/Elements/PCGPointExtentsModifier.h
@@ -2,7 +2,6 @@
 
 #pragma once
 
-#include "PCGPin.h"
 #include "PCGSettings.h"
 
 #include "Elements/PCGPointProcessingElementBase.h"
diff --git a/Engine/Plugins/Experimental/PCG/Source/PCG/Public/Elements/PCGPointFromMeshElement.h b/Engine/Plugins/Experimental/PCG/Source/PCG/Public/Elements/PCGPointFromMeshElement.h
--- a/Engine/Plugins/Experimental/PCG/Source/PCG/Public/Elements/PCGPointFromMeshElement.h
+++ b/Engine/Plugins/Experimental/PCG/Source/PCG/Public/Elements/PCGPointFromMeshElement.h
@@ -2,7 +2,6 @@
 
 #pragma once
 
-#include "PCGPin.h"
 #include "PCGSettings.h"
 
 #include "UObject/SoftObjectPtr.h"
diff --git a/Engine/Plugins/Experimental/PCG/Source/PCG/Public/Elements/PCGPropertyToParamData.h b/Engine/Plugins/Experimental/PCG/Source/PCG/Public/Elements/PCGPropertyToParamData.h
--- a/Engine/Plugins/Experimental/PCG/Source/PCG/Public/Elements/PCGPropertyToParamData.h
+++ b/Engine/Plugins/Experimental/PCG/Source/PCG/Public/Elements/PCGPropertyToParamData.h
@@ -2,7 +2,6 @@
 
 #pragma once
 
-#include "PCGPin.h"
 #include "PCGSettings.h"
 #include "Elements/PCGActorSelector.h"
 #include "Metadata/PCGMetadataAttribute.h"
diff --git a/Engine/Plugins/Experimental/PCG/Source/PCG/Public/Elements/PCGSelectPoints.h b/Engine/Plugins/Experimental/PCG/Source/PCG/Public/Elements/PCGSelectPoints.h
--- a/Engine/Plugins/Experimental/PCG/Source/PCG/Public/Elements/PCGSelectPoints.h
+++ b/Engine/Plugins/Experimental/PCG/Source/PCG/Public/Elements/PCGSelectPoints.h
@@ -2,7 +2,6 @@
 
 #pragma once
 
-#include "PCGPin.h"
 #include "PCGSettings.h"
 
 #include "PCGSelectPoints.generated.h"
diff --git a/Engine/Plugins/Experimental/PCG/Source/PCG/Public/Elements/PCGSelfPruning.h b/Engine/Plugins/Experimental/PCG/Source/PCG/Public/Elements/PCGSelfPruning.h
--- a/Engine/Plugins/Experimental/PCG/Source/PCG/Public/Elements/PCGSelfPruning.h
+++ b/Engine/Plugins/Experimental/PCG/Source/PCG/Public/Elements/PCGSelfPruning.h
@@ -2,7 +2,6 @@
 
 #pragma once
 
-#include "PCGPin.h"
 #include "PCGSettings.h"
 
 #include "PCGSelfPruning.generated.h"
diff --git a/Engine/Plugins/Experimental/PCG/Source/PCG/Public/Elements/PCGSpawnActor.h b/Engine/Plugins/Experimental/PCG/Source/PCG/Public/Elements/PCGSpawnActor.h
--- a/Engine/Plugins/Experimental/PCG/Source/PCG/Public/Elements/PCGSpawnActor.h
+++ b/Engine/Plugins/Experimental/PCG/Source/PCG/Public/Elements/PCGSpawnActor.h
@@ -2,7 +2,6 @@
 
 #pragma once
 
-#include "PCGPin.h"
 #include "PCGSubgraph.h"
 #include "Metadata/PCGAttributePropertySelector.h"
 
diff --git a/Engine/Plugins/Experimental/PCG/Source/PCG/Public/Elements/PCGSplineSampler.h b/Engine/Plugins/Experimental/PCG/Source/PCG/Public/Elements/PCGSplineSampler.h
--- a/Engine/Plugins/Experimental/PCG/Source/PCG/Public/Elements/PCGSplineSampler.h
+++ b/Engine/Plugins/Experimental/PCG/Source/PCG/Public/Elements/PCGSplineSampler.h
@@ -2,7 +2,6 @@
 
 #pragma once
 
-#include "PCGPin.h"
 #include "PCGSettings.h"
 
 #include "Curves/CurveFloat.h"
diff --git a/Engine/Plugins/Experimental/PCG/Source/PCG/Public/Elements/PCGStaticMeshSpawner.h b/Engine/Plugins/Experimental/PCG/Source/PCG/Public/Elements/PCGStaticMeshSpawner.h
--- a/Engine/Plugins/Experimental/PCG/Source/PCG/Public/Elements/PCGStaticMeshSpawner.h
+++ b/Engine/Plugins/Experimental/PCG/Source/PCG/Public/Elements/PCGStaticMeshSpawner.h
@@ -6,7 +6,6 @@
 
 #include "Engine/CollisionProfile.h"
 
-#include "PCGPin.h"
 #include "PCGStaticMeshSpawner.generated.h"
 
 class UPCGInstanceDataPackerBase;
diff --git a/Engine/Plugins/Experimental/PCG/Source/PCG/Public/Elements/PCGTextureSampler.h b/Engine/Plugins/Experimental/PCG/Source/PCG/Public/Elements/PCGTextureSampler.h
--- a/Engine/Plugins/Experimental/PCG/Source/PCG/Public/Elements/PCGTextureSampler.h
+++ b/Engine/Plugins/Experimental/PCG/Source/PCG/Public/Elements/PCGTextureSampler.h
@@ -3,7 +3,6 @@
 #pragma once
 
 #include "PCGContext.h"
-#include "PCGPin.h"
 #include "PCGSettings.h"
 #include "Data/PCGTextureData.h"
 
diff --git a/Engine/Plugins/Experimental/PCG/Source/PCG/Public/Elements/PCGTransformPoints.h b/Engine/Plugins/Experimental/PCG/Source/PCG/Public/Elements/PCGTransformPoints.h
--- a/Engine/Plugins/Experimental/PCG/Source/PCG/Public/Elements/PCGTransformPoints.h
+++ b/Engine/Plugins/Experimental/PCG/Source/PCG/Public/Elements/PCGTransformPoints.h
@@ -2,7 +2,6 @@
 
 #pragma once
 
-#include "PCGPin.h"
 #include "PCGSettings.h"
 
 #include "PCGTransformPoints.generated.h"
diff --git a/Engine/Plugins/Experimental/PCG/Source/PCG/Public/Elements/PCGUserParameterGet.h b/Engine/Plugins/Experimental/PCG/Source/PCG/Public/Elements/PCGUserParameterGet.h
--- a/Engine/Plugins/Experimental/PCG/Source/PCG/Public/Elements/PCGUserParameterGet.h
+++ b/Engine/Plugins/Experimental/PCG/Source/PCG/Public/Elements/PCGUserParameterGet.h
@@ -5,7 +5,6 @@
 #include "CoreMinimal.h"
 
 #include "PCGElement.h"
-#include "PCGPin.h"
 #include "PCGSettings.h"
 
 #include "PCGUserParameterGet.generated.h"
diff --git a/Engine/Plugins/Experimental/PCG/Source/PCG/Public/Elements/PCGVolumeSampler.h b/Engine/Plugins/Experimental/PCG/Source/PCG/Public/Elements/PCGVolumeSampler.h
--- a/Engine/Plugins/Experimental/PCG/Source/PCG/Public/Elements/PCGVolumeSampler.h
+++ b/Engine/Plugins/Experimental/PCG/Source/PCG/Public/Elements/PCGVolumeSampler.h
@@ -5,7 +5,6 @@
 #include "Math/Box.h"
 #include "PCGSettings.h"
 
-#include "PCGPin.h"
 #include "PCGVolumeSampler.generated.h"
 
 class UPCGPointData;
diff --git a/Engine/Plugins/Experimental/PCG/Source/PCG/Public/Elements/PCGWorldQuery.h b/Engine/Plugins/Experimental/PCG/Source/PCG/Public/Elements/PCGWorldQuery.h
--- a/Engine/Plugins/Experimental/PCG/Source/PCG/Public/Elements/PCGWorldQuery.h
+++ b/Engine/Plugins/Experimental/PCG/Source/PCG/Public/Elements/PCGWorldQuery.h
@@ -2,7 +2,6 @@
 
 #pragma once
 
-#include "PCGPin.h"
 #include "PCGSettings.h"
 #include "Data/PCGWorldData.h"
 
diff --git a/Engine/Plugins/Experimental/PCG/Source/PCG/Public/PCGSettings.h b/Engine/Plugins/Experimental/PCG/Source/PCG/Public/PCGSettings.h
--- a/Engine/Plugins/Experimental/PCG/Source/PCG/Public/PCGSettings.h
+++ b/Engine/Plugins/Experimental/PCG/Source/PCG/Public/PCGSettings.h
@@ -5,6 +5,7 @@
 #include "PCGData.h"
 #include "PCGDebug.h"
 #include "PCGElement.h"
+#include "PCGPin.h"
 #include "Elements/PCGActorSelector.h"
 #include "Tests/Determinism/PCGDeterminismSettings.h"
 
diff --git a/Engine/Plugins/Experimental/PCGGeometryScriptInterop/Source/PCGGeometryScriptInterop/Public/Elements/PCGMeshSampler.h b/Engine/Plugins/Experimental/PCGGeometryScriptInterop/Source/PCGGeometryScriptInterop/Public/Elements/PCGMeshSampler.h
--- a/Engine/Plugins/Experimental/PCGGeometryScriptInterop/Source/PCGGeometryScriptInterop/Public/Elements/PCGMeshSampler.h
+++ b/Engine/Plugins/Experimental/PCGGeometryScriptInterop/Source/PCGGeometryScriptInterop/Public/Elements/PCGMeshSampler.h
@@ -4,7 +4,6 @@
 
 #include "PCGContext.h"
 #include "PCGElement.h"
-#include "PCGPin.h"
 #include "PCGSettings.h"
 
 #include "Async/Future.h"
