[PCG] Wrong condition to detect Attribute Selector structs for overrides

#rnx

diff --git a/Engine/Plugins/Experimental/PCG/Source/PCG/Private/Metadata/Accessors/PCGAttributeAccessorHelpers.cpp b/Engine/Plugins/Experimental/PCG/Source/PCG/Private/Metadata/Accessors/PCGAttributeAccessorHelpers.cpp
--- a/Engine/Plugins/Experimental/PCG/Source/PCG/Private/Metadata/Accessors/PCGAttributeAccessorHelpers.cpp
+++ b/Engine/Plugins/Experimental/PCG/Source/PCG/Private/Metadata/Accessors/PCGAttributeAccessorHelpers.cpp
@@ -330,7 +330,7 @@ namespace PCGAttributeAccessorHelpers
 			{
 				return Functor(Signature<FPCGPropertyPathAccessor<FSoftClassPath>>{}, InProperty);
 			}
-			else if (StructProperty->Struct->IsA(FPCGAttributePropertySelector::StaticStruct()->GetClass()))
+			else if (StructProperty->Struct && StructProperty->Struct->IsChildOf(FPCGAttributePropertySelector::StaticStruct()))
 			{
 				return Functor(Signature<FPCGAttributePropertySelectorAccessor>{}, StructProperty);
 			}
