Added option to export single point in conjunction with pcg data

diff --git a/Engine/Plugins/Experimental/PCG/Source/PCG/Private/Elements/PCGDataFromActor.cpp b/Engine/Plugins/Experimental/PCG/Source/PCG/Private/Elements/PCGDataFromActor.cpp
--- a/Engine/Plugins/Experimental/PCG/Source/PCG/Private/Elements/PCGDataFromActor.cpp
+++ b/Engine/Plugins/Experimental/PCG/Source/PCG/Private/Elements/PCGDataFromActor.cpp
@@ -19,6 +19,11 @@
 
 #define LOCTEXT_NAMESPACE "PCGDataFromActorElement"
 
+namespace PCGDataFromActorConstants
+{
+	static const FName SinglePointPinLabel = TEXT("Single Point");
+}
+
 #if WITH_EDITOR
 void UPCGDataFromActorSettings::GetTrackedActorKeys(FPCGActorSelectionKeyToSettingsMap& OutKeysToSettings, TArray<TObjectPtr<const UPCGGraph>>& OutVisitedGraphs) const
 {
@@ -113,6 +118,15 @@ TArray<FPCGPinProperties> UPCGDataFromActorSettings::OutputPinProperties() const
 		{
 			Pins.Emplace(Pin);
 		}
+
+		if (bAlsoOutputSinglePointData)
+		{
+			Pins.Emplace(PCGDataFromActorConstants::SinglePointPinLabel,
+				EPCGDataType::Point,
+				/*bAllowMultipleConnections=*/true,
+				/*bAllowMultiData=*/true,
+				LOCTEXT("SinglePointPinTooltip", "Matching single point associated to the actors from which data has been retrieved"));
+		}
 	}
 
 	return Pins;
@@ -446,6 +460,18 @@ void FPCGDataFromActorElement::ProcessActor(FPCGContext* Context, const UPCGData
 		FPCGDataCollection Collection = UPCGComponent::CreateActorPCGDataCollection(FoundActor, Context->SourceComponent.Get(), Settings->GetDataFilter(), bParseActor);
 		Outputs += Collection.TaggedData;
 	}
+
+	// Finally, if we're in a case where we need to output the single point data too, let's do it now.
+	if (Settings->bAlsoOutputSinglePointData && (Settings->Mode == EPCGGetDataFromActorMode::GetDataFromPCGComponent || Settings->Mode == EPCGGetDataFromActorMode::GetDataFromPCGComponentOrParseComponents))
+	{
+		const bool bParseActor = false;
+		FPCGDataCollection Collection = UPCGComponent::CreateActorPCGDataCollection(FoundActor, Context->SourceComponent.Get(), EPCGDataType::Any, bParseActor);
+		for (const FPCGTaggedData& SinglePointData : Collection.TaggedData)
+		{
+			FPCGTaggedData& OutSinglePoint = Outputs.Add_GetRef(SinglePointData);
+			OutSinglePoint.Pin = PCGDataFromActorConstants::SinglePointPinLabel;
+		}
+	}
 }
 
 #undef LOCTEXT_NAMESPACE
diff --git a/Engine/Plugins/Experimental/PCG/Source/PCG/Public/Elements/PCGDataFromActor.h b/Engine/Plugins/Experimental/PCG/Source/PCG/Public/Elements/PCGDataFromActor.h
--- a/Engine/Plugins/Experimental/PCG/Source/PCG/Public/Elements/PCGDataFromActor.h
+++ b/Engine/Plugins/Experimental/PCG/Source/PCG/Public/Elements/PCGDataFromActor.h
@@ -66,6 +66,9 @@ public:
 	UPROPERTY(BlueprintReadWrite, EditAnywhere, Category = Settings, meta = (EditCondition = bDisplayModeSettings, EditConditionHides, HideEditConditionToggle))
 	EPCGGetDataFromActorMode Mode = EPCGGetDataFromActorMode::ParseActorComponents;
 
+	UPROPERTY(BlueprintReadWrite, EditAnywhere, Category = Settings, meta = (EditCondition = "Mode == EPCGGetDataFromActorMode::GetDataFromPCGComponent || Mode == EPCGGetDataFromActorMode::GetDataFromPCGComponentOrParseComponents", EditConditionHides))
+	bool bAlsoOutputSinglePointData = false;
+
 	UPROPERTY(BlueprintReadWrite, EditAnywhere, Category = Settings, meta = (EditCondition = "Mode == EPCGGetDataFromActorMode::GetSinglePoint", EditConditionHides))
 	bool bMergeSinglePointData = false;
 
