[PCG] Add a Add/RemoveNodes function for bulk nodes adding/removing

* Use it when deleting nodes and pasting nodes in the editor.
* Use it when deleting nodes when collapsing in a subgraph.
* It will improve reactiveness in big graphs in already generated maps.

#rnx

diff --git a/Engine/Plugins/Experimental/PCG/Source/PCG/Private/Helpers/PCGSubgraphHelpers.cpp b/Engine/Plugins/Experimental/PCG/Source/PCG/Private/Helpers/PCGSubgraphHelpers.cpp
--- a/Engine/Plugins/Experimental/PCG/Source/PCG/Private/Helpers/PCGSubgraphHelpers.cpp
+++ b/Engine/Plugins/Experimental/PCG/Source/PCG/Private/Helpers/PCGSubgraphHelpers.cpp
@@ -442,9 +442,12 @@ UPCGGraph* FPCGSubgraphHelpers::CollapseIntoSubgraph(UPCGGraph* InOriginalGraph,
 	Algo::Transform(CollapseInfo.GraphParametersUsed, GraphParametersUsed, [](const FPropertyBagPropertyDesc* In) { check(In); return *In; });
 	NewPCGGraph->AddUserParameters(GraphParametersUsed, InOriginalGraph);
 
-	// 6. Create subgraph and delete old nodes and superfluous nodes
+	// 6. Create subgraph and delete old nodes, extra nodes and superfluous nodes
 	UPCGNode* SubgraphNode = nullptr;
 	{
+		TArray<UPCGNode*> NodesToRemove;
+		NodesToRemove.Reserve(CollapseInfo.ValidNodesToCollapse.Num());
+
 		for (UPCGNode* PCGNode : CollapseInfo.ValidNodesToCollapse)
 		{
 			const UPCGSettings* NodeSettings = PCGNode->GetSettings();
@@ -455,7 +458,7 @@ UPCGGraph* FPCGSubgraphHelpers::CollapseIntoSubgraph(UPCGGraph* InOriginalGraph,
 				continue;
 			}
 
-			InOriginalGraph->RemoveNode(PCGNode);
+			NodesToRemove.Add(PCGNode);
 		}
 
 #if WITH_EDITOR
@@ -474,10 +477,13 @@ UPCGGraph* FPCGSubgraphHelpers::CollapseIntoSubgraph(UPCGGraph* InOriginalGraph,
 
 			if (!bHasAnyEdges)
 			{
-				InOriginalGraph->RemoveNode(SuperfluousNode);
+				NodesToRemove.Add(SuperfluousNode);
 			}
 		}
 
+		// Finally remove them all in one go
+		InOriginalGraph->RemoveNodes(NodesToRemove);
+
 		UPCGSettings* DefaultNodeSettings = nullptr;
 		SubgraphNode = InOriginalGraph->AddNodeOfType(UPCGSubgraphSettings::StaticClass(), DefaultNodeSettings);
 		UPCGSubgraphSettings* DefaultSubgraphSettings = CastChecked<UPCGSubgraphSettings>(DefaultNodeSettings);
diff --git a/Engine/Plugins/Experimental/PCG/Source/PCG/Private/PCGGraph.cpp b/Engine/Plugins/Experimental/PCG/Source/PCG/Private/PCGGraph.cpp
--- a/Engine/Plugins/Experimental/PCG/Source/PCG/Private/PCGGraph.cpp
+++ b/Engine/Plugins/Experimental/PCG/Source/PCG/Private/PCGGraph.cpp
@@ -451,21 +451,42 @@ UPCGNode* UPCGGraph::AddNodeCopy(UPCGSettings* InSettings, UPCGSettings*& Defaul
 }
 
 void UPCGGraph::OnNodeAdded(UPCGNode* InNode)
+{
+	OnNodesAdded(MakeArrayView<UPCGNode*>(&InNode, 1));
+}
+
+void UPCGGraph::OnNodesAdded(TArrayView<UPCGNode*> InNodes)
 {
 #if WITH_EDITOR
-	InNode->OnNodeChangedDelegate.AddUObject(this, &UPCGGraph::OnNodeChanged);
+	for (UPCGNode* Node : InNodes)
+	{
+		if (Node)
+		{
+			Node->OnNodeChangedDelegate.AddUObject(this, &UPCGGraph::OnNodeChanged);
+		}
+	}
+
 	NotifyGraphChanged(EPCGChangeType::Structural);
 #endif
 }
 
 void UPCGGraph::OnNodeRemoved(UPCGNode* InNode)
+{
+	OnNodesRemoved(MakeArrayView<UPCGNode*>(&InNode, 1));
+}
+
+void UPCGGraph::OnNodesRemoved(TArrayView<UPCGNode*> InNodes)
 {
 #if WITH_EDITOR
-	if (InNode)
+	for (UPCGNode* Node : InNodes)
 	{
-		InNode->OnNodeChangedDelegate.RemoveAll(this);
-		NotifyGraphChanged(EPCGChangeType::Structural);
+		if (Node)
+		{
+			Node->OnNodeChangedDelegate.RemoveAll(this);
+		}
 	}
+
+	NotifyGraphChanged(EPCGChangeType::Structural); 
 #endif
 }
 
@@ -540,50 +561,87 @@ bool UPCGGraph::Contains(UPCGNode* Node) const
 
 void UPCGGraph::AddNode(UPCGNode* InNode)
 {
-	check(InNode);
+	AddNodes_Internal(MakeArrayView<UPCGNode*>(&InNode, 1));
+}
+
+void UPCGGraph::AddNodes(TArray<UPCGNode*>& InNodes)
+{
+	AddNodes_Internal(InNodes);
+}
+
+void UPCGGraph::AddNodes_Internal(TArrayView<UPCGNode*> InNodes)
+{
+	if (InNodes.IsEmpty())
+	{
+		return;
+	}
 
 	Modify();
 
-	InNode->Rename(nullptr, this);
+	for (UPCGNode* Node : InNodes)
+	{
+		check(Node);
+		Node->Rename(nullptr, this);
 
 #if WITH_EDITOR
-	const FName DefaultNodeName = InNode->GetSettings()->GetDefaultNodeName();
-	if (DefaultNodeName != NAME_None)
-	{
-		FName NodeName = MakeUniqueObjectName(this, UPCGNode::StaticClass(), DefaultNodeName);
-		InNode->Rename(*NodeName.ToString());
-	}
+		const FName DefaultNodeName = Node->GetSettings()->GetDefaultNodeName();
+		if (DefaultNodeName != NAME_None)
+		{
+			FName NodeName = MakeUniqueObjectName(this, UPCGNode::StaticClass(), DefaultNodeName);
+			Node->Rename(*NodeName.ToString());
+		}
 #endif
 
-	Nodes.Add(InNode);
-	OnNodeAdded(InNode);
+		Nodes.Add(Node);
+	}
+
+	OnNodesAdded(InNodes);
 }
 
 void UPCGGraph::RemoveNode(UPCGNode* InNode)
 {
-	check(InNode);
+	RemoveNodes_Internal(MakeArrayView<UPCGNode*>(&InNode, 1));
+}
+
+void UPCGGraph::RemoveNodes(TArray<UPCGNode*>& InNodes)
+{
+	RemoveNodes_Internal(InNodes);
+}
+
+void UPCGGraph::RemoveNodes_Internal(TArrayView<UPCGNode*> InNodes)
+{
+	if (InNodes.IsEmpty())
+	{
+		return;
+	}
 
 	Modify();
 	
 	TSet<UPCGNode*> TouchedNodes;
 
-	for (UPCGPin* InputPin : InNode->InputPins)
+	for (UPCGNode* Node : InNodes)
 	{
-		InputPin->BreakAllEdges(&TouchedNodes);
-	}
+		check(Node);
 
-	for (UPCGPin* OutputPin : InNode->OutputPins)
-	{
-		OutputPin->BreakAllEdges(&TouchedNodes);
-	}
+		for (UPCGPin* InputPin : Node->InputPins)
+		{
+			InputPin->BreakAllEdges(&TouchedNodes);
+		}
 
-	// We're about to remove InNode, so don't bother triggering updates
-	TouchedNodes.Remove(InNode);
+		for (UPCGPin* OutputPin : Node->OutputPins)
+		{
+			OutputPin->BreakAllEdges(&TouchedNodes);
+		}
+
+		// We're about to remove InNode, so don't bother triggering updates
+		TouchedNodes.Remove(Node);
+
+		Nodes.Remove(Node);
+	}
 
 	PCGGraphUtils::NotifyTouchedNodes(TouchedNodes);
 
-	Nodes.Remove(InNode);
-	OnNodeRemoved(InNode);
+	OnNodesRemoved(InNodes);
 }
 
 bool UPCGGraph::RemoveEdge(UPCGNode* From, const FName& FromLabel, UPCGNode* To, const FName& ToLabel)
diff --git a/Engine/Plugins/Experimental/PCG/Source/PCG/Public/PCGGraph.h b/Engine/Plugins/Experimental/PCG/Source/PCG/Public/PCGGraph.h
--- a/Engine/Plugins/Experimental/PCG/Source/PCG/Public/PCGGraph.h
+++ b/Engine/Plugins/Experimental/PCG/Source/PCG/Public/PCGGraph.h
@@ -185,6 +185,10 @@ public:
 	UFUNCTION(BlueprintCallable, Category = Graph)
 	void RemoveNode(UPCGNode* InNode);
 
+	/** Bulk removal of nodes, to avoid notifying the world everytime. */
+	UFUNCTION(BlueprintCallable, Category = Graph)
+	void RemoveNodes(TArray<UPCGNode*>& InNodes);
+
 	/** Adds a directed edge in the graph. Returns the "To" node for easy chaining */
 	UFUNCTION(BlueprintCallable, Category = Graph)
 	UPCGNode* AddEdge(UPCGNode* From, const FName& FromPinLabel, UPCGNode* To, const FName& ToPinLabel);
@@ -210,6 +214,7 @@ public:
 	bool Contains(UPCGNode* Node) const;
 	const TArray<UPCGNode*>& GetNodes() const { return Nodes; }
 	void AddNode(UPCGNode* InNode);
+	void AddNodes(TArray<UPCGNode*>& InNodes);
 
 	/** Calls the lambda on every node in graph. */
 	void ForEachNode(const TFunction<void(UPCGNode*)>& Action);
@@ -252,7 +257,12 @@ public:
 
 protected:
 	void OnNodeAdded(UPCGNode* InNode);
+	void OnNodesAdded(TArrayView<UPCGNode*> InNodes);
 	void OnNodeRemoved(UPCGNode* InNode);
+	void OnNodesRemoved(TArrayView<UPCGNode*> InNodes);
+
+	void RemoveNodes_Internal(TArrayView<UPCGNode*> InNodes);
+	void AddNodes_Internal(TArrayView<UPCGNode*> InNodes);
 
 	/** Calculates node grid size. Not thread safe, called within write lock. */
 	uint32 CalculateNodeGridSizeRecursive_Unsafe(const UPCGNode* InNode, uint32 InDefaultGridSize) const;
diff --git a/Engine/Plugins/Experimental/PCG/Source/PCGEditor/Private/PCGEditor.cpp b/Engine/Plugins/Experimental/PCG/Source/PCGEditor/Private/PCGEditor.cpp
--- a/Engine/Plugins/Experimental/PCG/Source/PCGEditor/Private/PCGEditor.cpp
+++ b/Engine/Plugins/Experimental/PCG/Source/PCGEditor/Private/PCGEditor.cpp
@@ -1490,6 +1490,7 @@ void FPCGEditor::DeleteSelectedNodes()
 		check(PCGEditorGraph && PCGGraph);
 
 		bool bChanged = false;
+		TArray<UPCGNode*> NodesToRemove;
 		{
 			const FScopedTransaction Transaction(*FPCGEditorCommon::ContextIdentifier, LOCTEXT("PCGEditorDeleteTransactionMessage", "PCG Editor: Delete"), nullptr);
 			PCGEditorGraph->Modify();
@@ -1503,7 +1504,8 @@ void FPCGEditor::DeleteSelectedNodes()
 						UPCGNode* PCGNode = PCGEditorGraphNode->GetPCGNode();
 						check(PCGNode);
 
-						PCGGraph->RemoveNode(PCGNode);
+						NodesToRemove.Add(PCGNode);
+
 						PCGEditorGraphNode->DestroyNode();
 						bChanged = true;
 					}
@@ -1523,9 +1525,9 @@ void FPCGEditor::DeleteSelectedNodes()
 
 		if (bChanged)
 		{
+			PCGGraph->RemoveNodes(NodesToRemove);
 			GraphEditorWidget->ClearSelectionSet();
 			GraphEditorWidget->NotifyGraphChanged();
-			PCGGraphBeingEdited->NotifyGraphChanged(EPCGChangeType::Structural);
 		}
 	}
 }
@@ -1657,6 +1659,8 @@ void FPCGEditor::PasteNodesHere(const FVector2D& Location)
 		AvgNodePosition.Y *= InvNumNodes;
 	}
 
+	TArray<UPCGNode*> NodesToPaste;
+
 	for (UEdGraphNode* PastedNode : PastedNodes)
 	{
 		GraphEditorWidget->SetNodeSelection(PastedNode, true);
@@ -1671,10 +1675,12 @@ void FPCGEditor::PasteNodesHere(const FVector2D& Location)
 		UPCGEditorGraphNodeBase* PastedPCGGraphNode = Cast<UPCGEditorGraphNodeBase>(PastedNode);
 		if (UPCGNode* PastedPCGNode = PastedPCGGraphNode ? PastedPCGGraphNode->GetPCGNode() : nullptr)
 		{
-			PCGGraphBeingEdited->AddNode(PastedPCGNode);
+			NodesToPaste.Add(PastedPCGNode);
 		}
 	}
 
+	PCGGraphBeingEdited->AddNodes(NodesToPaste);
+
 	for (UEdGraphNode* PastedNode : PastedNodes)
 	{
 		UPCGEditorGraphNodeBase* PastedPCGGraphNode = Cast<UPCGEditorGraphNodeBase>(PastedNode);
