[PCG] Fixed an issue where user created pcg components could not be deleted after undo

diff --git a/Engine/Plugins/Experimental/PCG/Source/PCG/Private/PCGComponent.cpp b/Engine/Plugins/Experimental/PCG/Source/PCG/Private/PCGComponent.cpp
--- a/Engine/Plugins/Experimental/PCG/Source/PCG/Private/PCGComponent.cpp
+++ b/Engine/Plugins/Experimental/PCG/Source/PCG/Private/PCGComponent.cpp
@@ -1347,6 +1347,8 @@ void UPCGComponent::PostEditImport()
 
 void UPCGComponent::PreEditUndo()
 {
+	Super::PreEditUndo();
+
 	// Here we will keep a copy of flags that we require to keep through the undo
 	// so we can have a consistent state
 	LastGeneratedBoundsPriorToUndo = LastGeneratedBounds;
@@ -1376,6 +1378,8 @@ void UPCGComponent::PostEditUndo()
 	{
 		Refresh(/*bIsStructural=*/true);
 	}
+
+	Super::PostEditUndo();
 }
 
 void UPCGComponent::UpdateTrackingCache()
