[PCG] Better error logging on Attribute operations

#rnx

diff --git a/Engine/Plugins/Experimental/PCG/Source/PCG/Private/Elements/Metadata/PCGMetadataOpElementBase.cpp b/Engine/Plugins/Experimental/PCG/Source/PCG/Private/Elements/Metadata/PCGMetadataOpElementBase.cpp
--- a/Engine/Plugins/Experimental/PCG/Source/PCG/Private/Elements/Metadata/PCGMetadataOpElementBase.cpp
+++ b/Engine/Plugins/Experimental/PCG/Source/PCG/Private/Elements/Metadata/PCGMetadataOpElementBase.cpp
@@ -358,13 +358,14 @@ bool FPCGMetadataElementBase::ExecuteInternal(FPCGContext* Context) const
 
 		OperationData.InputSources[Index] = Settings->GetInputSource(Index).CopyAndFixLast(InputTaggedData[Index].Data);
 		const FPCGAttributePropertyInputSelector& InputSource = OperationData.InputSources[Index];
+		const FText InputSourceText = InputSource.GetDisplayText();
 
 		OperationData.InputAccessors[Index] = PCGAttributeAccessorHelpers::CreateConstAccessor(InputTaggedData[Index].Data, InputSource);
 		OperationData.InputKeys[Index] = PCGAttributeAccessorHelpers::CreateConstKeys(InputTaggedData[Index].Data, InputSource);
 
 		if (!OperationData.InputAccessors[Index].IsValid() || !OperationData.InputKeys[Index].IsValid())
 		{
-			PCGE_LOG(Error, GraphAndLog, FText::Format(LOCTEXT("AttributeDoesNotExist", "Attribute/Property '{0}' from pin {1} does not exist"), FText::FromName(InputSource.GetName()), FText::FromName(InputTaggedData[Index].Pin)));
+			PCGE_LOG(Error, GraphAndLog, FText::Format(LOCTEXT("AttributeDoesNotExist", "Attribute/Property '{0}' from pin {1} does not exist"), InputSourceText, FText::FromName(InputTaggedData[Index].Pin)));
 			return false;
 		}
 
@@ -374,9 +375,9 @@ bool FPCGMetadataElementBase::ExecuteInternal(FPCGContext* Context) const
 		bool bHasSpecialRequirement = false;
 		if (!Settings->IsSupportedInputType(AttributeTypeId, Index, bHasSpecialRequirement))
 		{
-			FText AttributeTypeName = FText::FromString(PCG::Private::GetTypeName(AttributeTypeId));
+			const FText AttributeTypeName = PCG::Private::GetTypeNameText(AttributeTypeId);
 			PCGE_LOG(Error, GraphAndLog, FText::Format(LOCTEXT("UnsupportedAttributeType", "Attribute/Property '{0}' from pin {1} is not a supported type ('{2}')"), 
-				FText::FromName(InputSource.GetName()), 
+				InputSourceText,
 				FText::FromName(InputTaggedData[Index].Pin),
 				AttributeTypeName
 			));
@@ -392,10 +393,10 @@ bool FPCGMetadataElementBase::ExecuteInternal(FPCGContext* Context) const
 			}
 			else if (OperationData.MostComplexInputType != AttributeTypeId && !PCG::Private::IsBroadcastable(AttributeTypeId, OperationData.MostComplexInputType))
 			{
-				FText AttributeTypeName = FText::FromString(PCG::Private::GetTypeName(AttributeTypeId));
-				FText MostComplexTypeName = FText::FromString(PCG::Private::GetTypeName(OperationData.MostComplexInputType));
+				const FText AttributeTypeName = PCG::Private::GetTypeNameText(AttributeTypeId);
+				const FText MostComplexTypeName = PCG::Private::GetTypeNameText(OperationData.MostComplexInputType);
 				PCGE_LOG(Error, GraphAndLog, FText::Format(LOCTEXT("AttributeCannotBeBroadcasted", "Attribute '{0}' (from pin {1}) of type '{2}' cannot be used for operation with type '{3}'"), 
-					FText::FromName(InputSource.GetName()), 
+					InputSourceText,
 					FText::FromName(InputTaggedData[Index].Pin),
 					AttributeTypeName,
 					MostComplexTypeName));
@@ -482,7 +483,8 @@ bool FPCGMetadataElementBase::ExecuteInternal(FPCGContext* Context) const
 
 		UPCGMetadata* OutMetadata = nullptr;
 
-		FName OutputName = OutputTarget.GetName();
+		const FName OutputName = OutputTarget.GetName();
+		const FText OutputTargetText = OutputTarget.GetDisplayText();
 
 		if (OutputTarget.GetSelection() == EPCGAttributePropertySelection::Attribute && OutputTarget.GetExtraNames().IsEmpty())
 		{
@@ -517,7 +519,10 @@ bool FPCGMetadataElementBase::ExecuteInternal(FPCGContext* Context) const
 				// We matched an attribute/property, check if the output type is valid.
 				if (!PCG::Private::IsBroadcastable(PCG::Private::MetadataTypes<AttributeType>::Id, TempConstAccessor->GetUnderlyingType()))
 				{
-					PCGE_LOG(Error, GraphAndLog, FText::Format(LOCTEXT("AttributeTypeBroadcastFailed", "Attribute/Property '{0}' cannot be broadcasted to match types for input"), FText::FromName(OutputName)));
+					PCGE_LOG(Error, GraphAndLog, FText::Format(LOCTEXT("AttributeTypeBroadcastFailed_Updated", "Attribute/Property '{0}' ({1}) is not compatible with operation output type ({2})."), 
+						OutputTargetText,
+						PCG::Private::GetTypeNameText(TempConstAccessor->GetUnderlyingType()),
+						PCG::Private::GetTypeNameText<AttributeType>()));
 					return false;
 				}
 
@@ -535,7 +540,7 @@ bool FPCGMetadataElementBase::ExecuteInternal(FPCGContext* Context) const
 
 		if (OperationData.OutputAccessors[OutputIndex]->IsReadOnly())
 		{
-			PCGE_LOG(Error, GraphAndLog, FText::Format(LOCTEXT("OutputAccessorIsReadOnly", "Attribute/Property '{0}' is read only."), OutputTarget.GetDisplayText()));
+			PCGE_LOG(Error, GraphAndLog, FText::Format(LOCTEXT("OutputAccessorIsReadOnly", "Attribute/Property '{0}' is read only."), OutputTargetText));
 			return false;
 		}
 
diff --git a/Engine/Plugins/Experimental/PCG/Source/PCG/Public/Metadata/PCGMetadataAttributeTraits.h b/Engine/Plugins/Experimental/PCG/Source/PCG/Public/Metadata/PCGMetadataAttributeTraits.h
--- a/Engine/Plugins/Experimental/PCG/Source/PCG/Public/Metadata/PCGMetadataAttributeTraits.h
+++ b/Engine/Plugins/Experimental/PCG/Source/PCG/Public/Metadata/PCGMetadataAttributeTraits.h
@@ -107,6 +107,23 @@ namespace PCG
 			return PCGDataTypeEnum ? PCGDataTypeEnum->GetNameStringByValue(InType) : FString(TEXT("Unknown"));
 		}
 
+		template <typename T>
+		FString GetTypeName()
+		{
+			return GetTypeName(MetadataTypes<T>::Id);
+		}
+
+		inline FText GetTypeNameText(uint16 InType)
+		{
+			return FText::FromString(GetTypeName(InType));
+		}
+
+		template <typename T>
+		FText GetTypeNameText()
+		{
+			return FText::FromString(GetTypeName<T>());
+		}
+
 		// Wrapper around a standard 2-dimensional CArray that is constexpr, to know if a type is broadcastable to another.
 		// First index is the original type, second index is the wanted type. Returns true if we can broadcast first type into second type.
 		struct UBroadcastableTypes
