[PCG] Added 'Get Random Stream From Two Points' BP helper function to make random stream creation in pairwise point processing easy

Also rename the current function that gets a random stream from a single point as it name clashes with another function.

#rnx

diff --git a/Engine/Plugins/Experimental/PCG/Config/BasePCG.ini b/Engine/Plugins/Experimental/PCG/Config/BasePCG.ini
--- a/Engine/Plugins/Experimental/PCG/Config/BasePCG.ini
+++ b/Engine/Plugins/Experimental/PCG/Config/BasePCG.ini
@@ -7,6 +7,7 @@
 +FunctionRedirects=(OldName="PCGBlueprintElement.LoopNTimes",NewName="PCGBlueprintElement.IterationLoop")
 +FunctionRedirects=(OldName="PCGBlueprintElement.InputLabels",NewName="PCGBlueprintElement.CustomInputLabels")
 +FunctionRedirects=(OldName="PCGBlueprintElement.OutputLabels",NewName="PCGBlueprintElement.CustomOutputLabels")
++FunctionRedirects=(OldName="PCGBlueprintHelpers.GetRandomStream",NewName="PCGBlueprintHelpers.GetRandomStreamFromPoint")
 +FunctionRedirects=(OldName="PCGDataFunctionLibrary.GetInputsByPin",NewName="PCGDataFunctionLibrary.GetInputsByPinLabel")
 +FunctionRedirects=(OldName="PCGDataFunctionLibrary.GetTaggedInputs",NewName="PCGDataFunctionLibrary.GetInputsByTag")
 +FunctionRedirects=(OldName="PCGDataFunctionLibrary.GetParamsByPin",NewName="PCGDataFunctionLibrary.GetParamsByPinLabel")
diff --git a/Engine/Plugins/Experimental/PCG/Source/PCG/Private/Elements/PCGMatchAndSetAttributes.cpp b/Engine/Plugins/Experimental/PCG/Source/PCG/Private/Elements/PCGMatchAndSetAttributes.cpp
--- a/Engine/Plugins/Experimental/PCG/Source/PCG/Private/Elements/PCGMatchAndSetAttributes.cpp
+++ b/Engine/Plugins/Experimental/PCG/Source/PCG/Private/Elements/PCGMatchAndSetAttributes.cpp
@@ -435,7 +435,7 @@ public:
 
 			if (PartitionDataIndex != INDEX_NONE)
 			{
-				AttributeSetKey = AttributeSetPartition.GetWeightedEntry(PartitionDataIndex, UPCGBlueprintHelpers::GetRandomStream(Point, Settings, SourceComponent));
+				AttributeSetKey = AttributeSetPartition.GetWeightedEntry(PartitionDataIndex, UPCGBlueprintHelpers::GetRandomStreamFromPoint(Point, Settings, SourceComponent));
 			}
 
 			if (Settings->bKeepUnmatched || AttributeSetKey != PCGInvalidEntryKey)
diff --git a/Engine/Plugins/Experimental/PCG/Source/PCG/Private/Helpers/PCGBlueprintHelpers.cpp b/Engine/Plugins/Experimental/PCG/Source/PCG/Private/Helpers/PCGBlueprintHelpers.cpp
--- a/Engine/Plugins/Experimental/PCG/Source/PCG/Private/Helpers/PCGBlueprintHelpers.cpp
+++ b/Engine/Plugins/Experimental/PCG/Source/PCG/Private/Helpers/PCGBlueprintHelpers.cpp
@@ -28,21 +28,41 @@ void UPCGBlueprintHelpers::SetSeedFromPosition(FPCGPoint& InPoint)
 	InPoint.Seed = ComputeSeedFromPosition(InPoint.Transform.GetLocation());
 }
 
-FRandomStream UPCGBlueprintHelpers::GetRandomStream(const FPCGPoint& InPoint, const UPCGSettings* OptionalSettings, const UPCGComponent* OptionalComponent)
+FRandomStream UPCGBlueprintHelpers::GetRandomStreamFromPoint(const FPCGPoint& InPoint, const UPCGSettings* OptionalSettings, const UPCGComponent* OptionalComponent)
 {
 	int Seed = InPoint.Seed;
 
 	if (OptionalSettings && OptionalComponent)
 	{
-		Seed = PCGHelpers::ComputeSeed(InPoint.Seed, OptionalSettings->Seed, OptionalComponent->Seed);
+		Seed = PCGHelpers::ComputeSeed(Seed, OptionalSettings->Seed, OptionalComponent->Seed);
 	}
 	else if (OptionalSettings)
 	{
-		Seed = PCGHelpers::ComputeSeed(InPoint.Seed, OptionalSettings->Seed);
+		Seed = PCGHelpers::ComputeSeed(Seed, OptionalSettings->Seed);
 	}
 	else if (OptionalComponent)
 	{
-		Seed = PCGHelpers::ComputeSeed(InPoint.Seed, OptionalComponent->Seed);
+		Seed = PCGHelpers::ComputeSeed(Seed, OptionalComponent->Seed);
+	}
+
+	return FRandomStream(Seed);
+}
+
+FRandomStream UPCGBlueprintHelpers::GetRandomStreamFromTwoPoints(const FPCGPoint& InPointA, const FPCGPoint& InPointB, const UPCGSettings* OptionalSettings, const UPCGComponent* OptionalComponent)
+{
+	int Seed = PCGHelpers::ComputeSeed(InPointA.Seed, InPointB.Seed);
+
+	if (OptionalSettings && OptionalComponent)
+	{
+		Seed = PCGHelpers::ComputeSeed(Seed, OptionalSettings->Seed, OptionalComponent->Seed);
+	}
+	else if (OptionalSettings)
+	{
+		Seed = PCGHelpers::ComputeSeed(Seed, OptionalSettings->Seed);
+	}
+	else if (OptionalComponent)
+	{
+		Seed = PCGHelpers::ComputeSeed(Seed, OptionalComponent->Seed);
 	}
 
 	return FRandomStream(Seed);
diff --git a/Engine/Plugins/Experimental/PCG/Source/PCG/Private/MatchAndSet/PCGMatchAndSetWeighted.cpp b/Engine/Plugins/Experimental/PCG/Source/PCG/Private/MatchAndSet/PCGMatchAndSetWeighted.cpp
--- a/Engine/Plugins/Experimental/PCG/Source/PCG/Private/MatchAndSet/PCGMatchAndSetWeighted.cpp
+++ b/Engine/Plugins/Experimental/PCG/Source/PCG/Private/MatchAndSet/PCGMatchAndSetWeighted.cpp
@@ -121,7 +121,7 @@ void UPCGMatchAndSetWeighted::MatchAndSet_Implementation(
 	{
 		const FPCGPoint& Point = InPoints[Index];
 
-		FRandomStream RandomSource = UPCGBlueprintHelpers::GetRandomStream(Point, InSettings, SourceComponent);
+		FRandomStream RandomSource = UPCGBlueprintHelpers::GetRandomStreamFromPoint(Point, InSettings, SourceComponent);
 		int RandomWeightedPick = RandomSource.RandRange(0, TotalWeight - 1);
 
 		int RandomPick = 0;
diff --git a/Engine/Plugins/Experimental/PCG/Source/PCG/Private/MatchAndSet/PCGMatchAndSetWeightedByCategory.cpp b/Engine/Plugins/Experimental/PCG/Source/PCG/Private/MatchAndSet/PCGMatchAndSetWeightedByCategory.cpp
--- a/Engine/Plugins/Experimental/PCG/Source/PCG/Private/MatchAndSet/PCGMatchAndSetWeightedByCategory.cpp
+++ b/Engine/Plugins/Experimental/PCG/Source/PCG/Private/MatchAndSet/PCGMatchAndSetWeightedByCategory.cpp
@@ -208,7 +208,7 @@ void UPCGMatchAndSetWeightedByCategory::MatchAndSet_Implementation(
 
 		const FPCGPoint& Point = InPoints[PointIndex];
 
-		FRandomStream RandomSource = UPCGBlueprintHelpers::GetRandomStream(Point, InSettings, SourceComponent);
+		FRandomStream RandomSource = UPCGBlueprintHelpers::GetRandomStreamFromPoint(Point, InSettings, SourceComponent);
 		int RandomWeightedPick = RandomSource.RandRange(0, TotalWeight - 1);
 
 		int RandomPick = 0;
diff --git a/Engine/Plugins/Experimental/PCG/Source/PCG/Private/MeshSelectors/PCGMeshSelectorWeighted.cpp b/Engine/Plugins/Experimental/PCG/Source/PCG/Private/MeshSelectors/PCGMeshSelectorWeighted.cpp
--- a/Engine/Plugins/Experimental/PCG/Source/PCG/Private/MeshSelectors/PCGMeshSelectorWeighted.cpp
+++ b/Engine/Plugins/Experimental/PCG/Source/PCG/Private/MeshSelectors/PCGMeshSelectorWeighted.cpp
@@ -213,7 +213,7 @@ bool UPCGMeshSelectorWeighted::SelectInstances(
 		{
 			const FPCGPoint& Point = Points[CurrentPointIndex++];
 
-			FRandomStream RandomSource = UPCGBlueprintHelpers::GetRandomStream(Point, Settings, Context.SourceComponent.Get());
+			FRandomStream RandomSource = UPCGBlueprintHelpers::GetRandomStreamFromPoint(Point, Settings, Context.SourceComponent.Get());
 			const int RandomWeightedPick = RandomSource.RandRange(0, TotalWeight - 1);
 
 			int RandomPick = 0;
diff --git a/Engine/Plugins/Experimental/PCG/Source/PCG/Private/MeshSelectors/PCGMeshSelectorWeightedByCategory.cpp b/Engine/Plugins/Experimental/PCG/Source/PCG/Private/MeshSelectors/PCGMeshSelectorWeightedByCategory.cpp
--- a/Engine/Plugins/Experimental/PCG/Source/PCG/Private/MeshSelectors/PCGMeshSelectorWeightedByCategory.cpp
+++ b/Engine/Plugins/Experimental/PCG/Source/PCG/Private/MeshSelectors/PCGMeshSelectorWeightedByCategory.cpp
@@ -223,7 +223,7 @@ bool UPCGMeshSelectorWeightedByCategory::SelectInstances(
 
 		const int TotalWeight = InstancesAndWeights->CumulativeWeights.Last();
 
-		FRandomStream RandomSource = UPCGBlueprintHelpers::GetRandomStream(Point, Settings, Context.SourceComponent.Get());
+		FRandomStream RandomSource = UPCGBlueprintHelpers::GetRandomStreamFromPoint(Point, Settings, Context.SourceComponent.Get());
 		const int RandomWeightedPick = RandomSource.RandRange(0, TotalWeight - 1);
 
 		int RandomPick = 0;
diff --git a/Engine/Plugins/Experimental/PCG/Source/PCG/Public/Helpers/PCGBlueprintHelpers.h b/Engine/Plugins/Experimental/PCG/Source/PCG/Public/Helpers/PCGBlueprintHelpers.h
--- a/Engine/Plugins/Experimental/PCG/Source/PCG/Public/Helpers/PCGBlueprintHelpers.h
+++ b/Engine/Plugins/Experimental/PCG/Source/PCG/Public/Helpers/PCGBlueprintHelpers.h
@@ -30,7 +30,11 @@ public:
 
 	/** Creates a random stream from a point's seed and settings/component's seed (optional) */
 	UFUNCTION(BlueprintCallable, Category = "PCG|Helpers", meta = (ScriptMethod))
-	static FRandomStream GetRandomStream(const FPCGPoint& InPoint, const UPCGSettings* OptionalSettings = nullptr, const UPCGComponent* OptionalComponent = nullptr);
+	static FRandomStream GetRandomStreamFromPoint(const FPCGPoint& InPoint, const UPCGSettings* OptionalSettings = nullptr, const UPCGComponent* OptionalComponent = nullptr);
+
+	/** Creates a random stream from using the random seeds from two points, as well as settings/component's seed (optional) */
+	UFUNCTION(BlueprintCallable, Category = "PCG|Helpers", meta = (ScriptMethod))
+	static FRandomStream GetRandomStreamFromTwoPoints(const FPCGPoint& InPointA, const FPCGPoint& InPointB, const UPCGSettings* OptionalSettings = nullptr, const UPCGComponent* OptionalComponent = nullptr);
 
 	UFUNCTION(BlueprintCallable, Category = "PCG|Helpers", meta = (ScriptMethod))
 	static const UPCGSettings* GetSettings(UPARAM(ref) FPCGContext& Context);
