[PCG] Display mesh names in MeshSelector entries panel

diff --git a/Engine/Plugins/Experimental/PCG/Source/PCG/Private/MeshSelectors/PCGMeshSelectorWeighted.cpp b/Engine/Plugins/Experimental/PCG/Source/PCG/Private/MeshSelectors/PCGMeshSelectorWeighted.cpp
--- a/Engine/Plugins/Experimental/PCG/Source/PCG/Private/MeshSelectors/PCGMeshSelectorWeighted.cpp
+++ b/Engine/Plugins/Experimental/PCG/Source/PCG/Private/MeshSelectors/PCGMeshSelectorWeighted.cpp
@@ -56,6 +56,9 @@ FPCGMeshSelectorWeightedEntry::FPCGMeshSelectorWeightedEntry(TSoftObjectPtr<USta
 	: Weight(InWeight)
 {
 	Descriptor.StaticMesh = InMesh;
+#if WITH_EDITOR
+	DisplayName = InMesh.ToSoftObjectPath().GetAssetFName();
+#endif // WITH_EDITOR
 }
 
 #if WITH_EDITOR
@@ -103,18 +106,6 @@ void FPCGMeshSelectorWeightedEntry::ApplyDeprecation()
 }
 #endif
 
-void UPCGMeshSelectorWeighted::PostLoad()
-{
-	Super::PostLoad();
-
-#if WITH_EDITOR
-	for (FPCGMeshSelectorWeightedEntry& Entry : MeshEntries)
-	{
-		Entry.ApplyDeprecation();
-	}
-#endif
-}
-
 bool UPCGMeshSelectorWeighted::SelectInstances(
 	FPCGStaticMeshSpawnerContext& Context,
 	const UPCGStaticMeshSpawnerSettings* Settings, 
@@ -301,4 +292,56 @@ bool UPCGMeshSelectorWeighted::SelectInstances(
 	}
 }
 
+void UPCGMeshSelectorWeighted::PostLoad()
+{
+	Super::PostLoad();
+
+#if WITH_EDITOR
+	for (FPCGMeshSelectorWeightedEntry& Entry : MeshEntries)
+	{
+		Entry.ApplyDeprecation();
+	}
+
+	RefreshDisplayNames();
+#endif // WITH_EDITOR
+}
+
+void UPCGMeshSelectorWeighted::PostDuplicate(bool bDuplicateForPIE)
+{
+	Super::PostDuplicate(bDuplicateForPIE);
+
+#if WITH_EDITOR
+	RefreshDisplayNames();
+#endif // WITH_EDITOR
+}
+
+void UPCGMeshSelectorWeighted::PostEditImport()
+{
+	Super::PostEditImport();
+
+#if WITH_EDITOR
+	RefreshDisplayNames();
+#endif // WITH_EDITOR
+}
+
+#if WITH_EDITOR
+void UPCGMeshSelectorWeighted::PostEditChangeProperty(FPropertyChangedEvent& PropertyChangedEvent)
+{
+	Super::PostEditChangeProperty(PropertyChangedEvent);
+
+	if (PropertyChangedEvent.GetPropertyName() == GET_MEMBER_NAME_CHECKED(FSoftISMComponentDescriptor, StaticMesh))
+	{
+		RefreshDisplayNames();
+	}
+}
+
+void UPCGMeshSelectorWeighted::RefreshDisplayNames()
+{
+	for (FPCGMeshSelectorWeightedEntry& Entry : MeshEntries)
+	{
+		Entry.DisplayName = Entry.Descriptor.StaticMesh.ToSoftObjectPath().GetAssetFName();
+	}
+}
+#endif // WITH_EDITOR
+
 #undef LOCTEXT_NAMESPACE
diff --git a/Engine/Plugins/Experimental/PCG/Source/PCG/Private/MeshSelectors/PCGMeshSelectorWeightedByCategory.cpp b/Engine/Plugins/Experimental/PCG/Source/PCG/Private/MeshSelectors/PCGMeshSelectorWeightedByCategory.cpp
--- a/Engine/Plugins/Experimental/PCG/Source/PCG/Private/MeshSelectors/PCGMeshSelectorWeightedByCategory.cpp
+++ b/Engine/Plugins/Experimental/PCG/Source/PCG/Private/MeshSelectors/PCGMeshSelectorWeightedByCategory.cpp
@@ -26,18 +26,6 @@ void FPCGWeightedByCategoryEntryList::ApplyDeprecation()
 }
 #endif
 
-void UPCGMeshSelectorWeightedByCategory::PostLoad()
-{
-	Super::PostLoad();
-
-#if WITH_EDITOR
-	for (FPCGWeightedByCategoryEntryList& Entry : Entries)
-	{
-		Entry.ApplyDeprecation();
-	}
-#endif
-}
-
 bool UPCGMeshSelectorWeightedByCategory::SelectInstances(
 	FPCGStaticMeshSpawnerContext& Context,
 	const UPCGStaticMeshSpawnerSettings* Settings,
@@ -313,4 +301,59 @@ bool UPCGMeshSelectorWeightedByCategory::SelectInstances(
 	}
 }
 
+void UPCGMeshSelectorWeightedByCategory::PostLoad()
+{
+	Super::PostLoad();
+
+#if WITH_EDITOR
+	for (FPCGWeightedByCategoryEntryList& Entry : Entries)
+	{
+		Entry.ApplyDeprecation();
+	}
+
+	RefreshDisplayNames();
+#endif // WITH_EDITOR
+}
+
+void UPCGMeshSelectorWeightedByCategory::PostDuplicate(bool bDuplicateForPIE)
+{
+	Super::PostDuplicate(bDuplicateForPIE);
+
+#if WITH_EDITOR
+	RefreshDisplayNames();
+#endif // WITH_EDITOR
+}
+
+void UPCGMeshSelectorWeightedByCategory::PostEditImport()
+{
+	Super::PostEditImport();
+
+#if WITH_EDITOR
+	RefreshDisplayNames();
+#endif // WITH_EDITOR
+}
+
+#if WITH_EDITOR
+void UPCGMeshSelectorWeightedByCategory::PostEditChangeProperty(FPropertyChangedEvent& PropertyChangedEvent)
+{
+	Super::PostEditChangeProperty(PropertyChangedEvent);
+
+	if (PropertyChangedEvent.GetPropertyName() == GET_MEMBER_NAME_CHECKED(FSoftISMComponentDescriptor, StaticMesh))
+	{
+		RefreshDisplayNames();
+	}
+}
+
+void UPCGMeshSelectorWeightedByCategory::RefreshDisplayNames()
+{
+	for (FPCGWeightedByCategoryEntryList& Entry : Entries)
+	{
+		for (FPCGMeshSelectorWeightedEntry& WeightedEntry : Entry.WeightedMeshEntries)
+		{
+			WeightedEntry.DisplayName = WeightedEntry.Descriptor.StaticMesh.ToSoftObjectPath().GetAssetFName();
+		}
+	}
+}
+#endif // WITH_EDITOR
+
 #undef LOCTEXT_NAMESPACE
diff --git a/Engine/Plugins/Experimental/PCG/Source/PCG/Public/MeshSelectors/PCGMeshSelectorWeighted.h b/Engine/Plugins/Experimental/PCG/Source/PCG/Public/MeshSelectors/PCGMeshSelectorWeighted.h
--- a/Engine/Plugins/Experimental/PCG/Source/PCG/Public/MeshSelectors/PCGMeshSelectorWeighted.h
+++ b/Engine/Plugins/Experimental/PCG/Source/PCG/Public/MeshSelectors/PCGMeshSelectorWeighted.h
@@ -26,7 +26,7 @@ struct PCG_API FPCGMeshSelectorWeightedEntry
 
 #if WITH_EDITOR
 	void ApplyDeprecation();
-#endif
+#endif // WITH_EDITOR
 
 	UPROPERTY(EditAnywhere, Category = Settings)
 	FSoftISMComponentDescriptor Descriptor;
@@ -34,6 +34,11 @@ struct PCG_API FPCGMeshSelectorWeightedEntry
 	UPROPERTY(BlueprintReadWrite, EditAnywhere, Category = Settings, meta = (ClampMin = "0"))
 	int Weight = 1;
 
+#if WITH_EDITORONLY_DATA
+	UPROPERTY(Transient, VisibleAnywhere, Category = Settings)
+	FName DisplayName = NAME_None;
+#endif
+
 #if WITH_EDITORONLY_DATA
 	UPROPERTY()
 	TSoftObjectPtr<UStaticMesh> Mesh_DEPRECATED;
@@ -76,10 +81,21 @@ public:
 		TArray<FPCGMeshInstanceList>& OutMeshInstances,
 		UPCGPointData* OutPointData) const override;
 
-	void PostLoad();
+	// ~Begin UObject interface
+	virtual void PostLoad() override;
+	virtual void PostDuplicate(bool bDuplicateForPIE) override;
+	virtual void PostEditImport() override;
+#if WITH_EDITOR
+	virtual void PostEditChangeProperty(FPropertyChangedEvent& PropertyChangedEvent) override;
+	// ~End UObject interface
+
+protected:
+	/** Refresh MeshEntries display names */
+	void RefreshDisplayNames();
+#endif // WITH_EDITOR
 
 public:
-	UPROPERTY(BlueprintReadWrite, EditAnywhere, Category = Settings)
+	UPROPERTY(BlueprintReadWrite, EditAnywhere, Category = Settings, meta = (TitleProperty = "DisplayName"))
 	TArray<FPCGMeshSelectorWeightedEntry> MeshEntries;
 
 	UPROPERTY(BlueprintReadWrite, EditAnywhere, Category = Settings, meta = (InlineEditConditionToggle))
diff --git a/Engine/Plugins/Experimental/PCG/Source/PCG/Public/MeshSelectors/PCGMeshSelectorWeightedByCategory.h b/Engine/Plugins/Experimental/PCG/Source/PCG/Public/MeshSelectors/PCGMeshSelectorWeightedByCategory.h
--- a/Engine/Plugins/Experimental/PCG/Source/PCG/Public/MeshSelectors/PCGMeshSelectorWeightedByCategory.h
+++ b/Engine/Plugins/Experimental/PCG/Source/PCG/Public/MeshSelectors/PCGMeshSelectorWeightedByCategory.h
@@ -28,7 +28,7 @@ struct PCG_API FPCGWeightedByCategoryEntryList
 	UPROPERTY(BlueprintReadWrite, EditAnywhere, Category = Settings)
 	bool IsDefault = false;
 
-	UPROPERTY(BlueprintReadWrite, EditAnywhere, Category = Settings)
+	UPROPERTY(BlueprintReadWrite, EditAnywhere, Category = Settings, meta = (TitleProperty = "DisplayName"))
 	TArray<FPCGMeshSelectorWeightedEntry> WeightedMeshEntries;
 };
 
@@ -45,7 +45,18 @@ public:
 		TArray<FPCGMeshInstanceList>& OutMeshInstances,
 		UPCGPointData* OutPointData) const override;
 
-	void PostLoad();
+	// ~Begin UObject interface
+	virtual void PostLoad() override;
+	virtual void PostDuplicate(bool bDuplicateForPIE) override;
+	virtual void PostEditImport() override;
+#if WITH_EDITOR
+	virtual void PostEditChangeProperty(FPropertyChangedEvent& PropertyChangedEvent) override;
+	// ~End UObject interface
+
+protected:
+	// Refresh MeshEntries display names
+	void RefreshDisplayNames();
+#endif // WITH_EDITOR
 
 public:
 	UPROPERTY(BlueprintReadWrite, EditAnywhere, Category = Settings)
