[PCG] Moved scale initialization to PreActorSpawn in order to have the volume size set before dropping the PCG asset, for easier placement.

diff --git a/Engine/Plugins/Experimental/PCG/Source/PCGEditor/Private/PCGVolumeFactory.cpp b/Engine/Plugins/Experimental/PCG/Source/PCGEditor/Private/PCGVolumeFactory.cpp
--- a/Engine/Plugins/Experimental/PCG/Source/PCGEditor/Private/PCGVolumeFactory.cpp
+++ b/Engine/Plugins/Experimental/PCG/Source/PCGEditor/Private/PCGVolumeFactory.cpp
@@ -38,6 +38,21 @@ bool UPCGVolumeFactory::CanCreateActorFrom(const FAssetData& AssetData, FText& O
 	return true;
 }
 
+bool UPCGVolumeFactory::PreSpawnActor(UObject* Asset, FTransform& InOutLocation)
+{
+	if (Super::PreSpawnActor(Asset, InOutLocation))
+	{
+		if (const UPCGEngineSettings* Settings = GetDefault<UPCGEngineSettings>())
+		{
+			InOutLocation.SetScale3D(Settings->VolumeScale);
+		}
+
+		return true;
+	}
+	
+	return false;
+}
+
 void UPCGVolumeFactory::PostPlaceAsset(TArrayView<const FTypedElementHandle> InElementHandles, const FAssetPlacementInfo& InPlacementInfo, const FPlacementOptions& InPlacementOptions)
 {
 	Super::PostPlaceAsset(InElementHandles, InPlacementInfo, InPlacementOptions);
@@ -64,9 +79,7 @@ void UPCGVolumeFactory::PostPlaceAsset(TArrayView<const FTypedElementHandle> InE
 			continue;
 		}
 
-		APCGVolume* PCGVolume = CastChecked<APCGVolume>(NewActor);
-		PCGVolume->SetActorScale3D(Settings->VolumeScale);
-
+		const APCGVolume* PCGVolume = CastChecked<APCGVolume>(NewActor);
 		UPCGComponent* PCGComponent = CastChecked<UPCGComponent>(PCGVolume->GetComponentByClass(UPCGComponent::StaticClass()));
 		PCGComponent->SetGraph(PCGGraph);
 
diff --git a/Engine/Plugins/Experimental/PCG/Source/PCGEditor/Private/PCGVolumeFactory.h b/Engine/Plugins/Experimental/PCG/Source/PCGEditor/Private/PCGVolumeFactory.h
--- a/Engine/Plugins/Experimental/PCG/Source/PCGEditor/Private/PCGVolumeFactory.h
+++ b/Engine/Plugins/Experimental/PCG/Source/PCGEditor/Private/PCGVolumeFactory.h
@@ -19,6 +19,7 @@ public:
 
 	//~ Begin UActorFactory Interface
 	virtual bool CanCreateActorFrom(const FAssetData& AssetData, FText& OutErrorMsg) override;
+	virtual bool PreSpawnActor(UObject* Asset, FTransform& InOutLocation) override;
 	virtual void PostPlaceAsset(TArrayView<const FTypedElementHandle> InElementHandles, const FAssetPlacementInfo& InPlacementInfo, const FPlacementOptions& InPlacementOptions) override;
 	//~ End UActorFactory Interface
 };
