Fixed behavior for PCG in level instances.
Targeted target actor level instead of the default current level to prevent polluting level instances

diff --git a/Engine/Plugins/Experimental/PCG/Source/PCG/Private/Elements/PCGCreateSpline.cpp b/Engine/Plugins/Experimental/PCG/Source/PCG/Private/Elements/PCGCreateSpline.cpp
--- a/Engine/Plugins/Experimental/PCG/Source/PCG/Private/Elements/PCGCreateSpline.cpp
+++ b/Engine/Plugins/Experimental/PCG/Source/PCG/Private/Elements/PCGCreateSpline.cpp
@@ -9,6 +9,7 @@
 #include "PCGPin.h"
 #include "Data/PCGPointData.h"
 #include "Data/PCGSplineData.h"
+#include "Helpers/PCGActorHelpers.h"
 #include "Helpers/PCGHelpers.h"
 #include "Metadata/PCGMetadata.h"
 #include "Metadata/PCGMetadataAttributeTpl.h"
@@ -138,15 +139,10 @@ bool FPCGCreateSplineElement::ExecuteInternal(FPCGContext* Context) const
 		}
 		else if(Settings->Mode == EPCGCreateSplineMode::CreateNewActor)
 		{
-			FActorSpawnParameters SpawnParams;
-			if (PCGHelpers::IsRuntimeOrPIE())
-			{
-				SpawnParams.ObjectFlags |= RF_Transient;
-			}
-
 			// TODO: check ownership of target actor like in the spawn actor?
 			// TODO: allow template?
-			SplineActor = TargetActor->GetWorld()->SpawnActor(AActor::StaticClass(), &TargetActor->GetTransform(), SpawnParams);
+			FActorSpawnParameters ActorSpawnParams;
+			SplineActor = UPCGActorHelpers::SpawnDefaultActor(TargetActor->GetWorld(), TargetActor->GetLevel(), AActor::StaticClass(), TargetActor->GetTransform(), ActorSpawnParams);
 
 			if (!SplineActor)
 			{
diff --git a/Engine/Plugins/Experimental/PCG/Source/PCG/Private/Elements/PCGCreateTargetActor.cpp b/Engine/Plugins/Experimental/PCG/Source/PCG/Private/Elements/PCGCreateTargetActor.cpp
--- a/Engine/Plugins/Experimental/PCG/Source/PCG/Private/Elements/PCGCreateTargetActor.cpp
+++ b/Engine/Plugins/Experimental/PCG/Source/PCG/Private/Elements/PCGCreateTargetActor.cpp
@@ -263,7 +263,7 @@ bool FPCGCreateTargetActorElement::ExecuteInternal(FPCGContext* Context) const
 		Transform = Settings->ActorPivot;
 	}
 
-	AActor* GeneratedActor = UPCGActorHelpers::SpawnDefaultActor(TargetActor->GetWorld(), Settings->TemplateActorClass, Transform, SpawnParams);
+	AActor* GeneratedActor = UPCGActorHelpers::SpawnDefaultActor(TargetActor->GetWorld(), TargetActor->GetLevel(), Settings->TemplateActorClass, Transform, SpawnParams);
 
 	if (!GeneratedActor)
 	{
diff --git a/Engine/Plugins/Experimental/PCG/Source/PCG/Private/Elements/PCGSpawnActor.cpp b/Engine/Plugins/Experimental/PCG/Source/PCG/Private/Elements/PCGSpawnActor.cpp
--- a/Engine/Plugins/Experimental/PCG/Source/PCG/Private/Elements/PCGSpawnActor.cpp
+++ b/Engine/Plugins/Experimental/PCG/Source/PCG/Private/Elements/PCGSpawnActor.cpp
@@ -843,6 +843,7 @@ void FPCGSpawnActorElement::SpawnActors(FPCGSubgraphContext* Context, AActor* Ta
 	FActorSpawnParameters SpawnParams;
 	SpawnParams.Template = TemplateActor;
 	SpawnParams.SpawnCollisionHandlingOverride = ESpawnActorCollisionHandlingMethod::AlwaysSpawn;
+	SpawnParams.OverrideLevel = TargetActor->GetLevel();
 
 	if (PCGHelpers::IsRuntimeOrPIE())
 	{
diff --git a/Engine/Plugins/Experimental/PCG/Source/PCG/Private/Helpers/PCGActorHelpers.cpp b/Engine/Plugins/Experimental/PCG/Source/PCG/Private/Helpers/PCGActorHelpers.cpp
--- a/Engine/Plugins/Experimental/PCG/Source/PCG/Private/Helpers/PCGActorHelpers.cpp
+++ b/Engine/Plugins/Experimental/PCG/Source/PCG/Private/Helpers/PCGActorHelpers.cpp
@@ -286,7 +286,7 @@ void UPCGActorHelpers::ForEachActorInWorld(UWorld* World, TSubclassOf<AActor> Ac
 	}
 }
 
-AActor* UPCGActorHelpers::SpawnDefaultActor(UWorld* World, TSubclassOf<AActor> ActorClass, FName BaseName, const FTransform& Transform, AActor* Parent)
+AActor* UPCGActorHelpers::SpawnDefaultActor(UWorld* World, ULevel* Level, TSubclassOf<AActor> ActorClass, FName BaseName, const FTransform& Transform, AActor* Parent)
 {
 	if (!World || !ActorClass)
 	{
@@ -294,23 +294,29 @@ AActor* UPCGActorHelpers::SpawnDefaultActor(UWorld* World, TSubclassOf<AActor> A
 	}
 
 	FActorSpawnParameters SpawnParams;
-	SpawnParams.Name = MakeUniqueObjectName(World->GetCurrentLevel(), ActorClass, BaseName);
+	SpawnParams.Name = MakeUniqueObjectName(Level ? Level : World->GetCurrentLevel(), ActorClass, BaseName);
 
-	if (PCGHelpers::IsRuntimeOrPIE())
-	{
-		SpawnParams.ObjectFlags |= RF_Transient;
-	}
-
-	return SpawnDefaultActor(World, ActorClass, Transform, SpawnParams, Parent);
+	return SpawnDefaultActor(World, Level, ActorClass, Transform, SpawnParams, Parent);
 }
 
-AActor* UPCGActorHelpers::SpawnDefaultActor(UWorld* World, TSubclassOf<AActor> ActorClass, const FTransform& Transform, const FActorSpawnParameters& SpawnParams, AActor* Parent)
+AActor* UPCGActorHelpers::SpawnDefaultActor(UWorld* World, ULevel* Level, TSubclassOf<AActor> ActorClass, const FTransform& Transform, const FActorSpawnParameters& InSpawnParams, AActor* Parent)
 {
 	if (!World || !ActorClass)
 	{
 		return nullptr;
 	}
 
+	FActorSpawnParameters SpawnParams = InSpawnParams;
+	if (Level)
+	{
+		SpawnParams.OverrideLevel = Level;
+	}
+
+	if (PCGHelpers::IsRuntimeOrPIE())
+	{
+		SpawnParams.ObjectFlags |= RF_Transient;
+	}
+
 	AActor* NewActor = World->SpawnActor(*ActorClass, &Transform, SpawnParams);
 	
 	if (!NewActor)
diff --git a/Engine/Plugins/Experimental/PCG/Source/PCG/Private/PCGActorAndComponentMapping.cpp b/Engine/Plugins/Experimental/PCG/Source/PCG/Private/PCGActorAndComponentMapping.cpp
--- a/Engine/Plugins/Experimental/PCG/Source/PCG/Private/PCGActorAndComponentMapping.cpp
+++ b/Engine/Plugins/Experimental/PCG/Source/PCG/Private/PCGActorAndComponentMapping.cpp
@@ -15,8 +15,12 @@
 #include "StaticMeshCompiler.h"
 #include "TextureCompiler.h"
 #include "Engine/Engine.h"
+#include "Engine/Level.h"
 #include "Engine/StaticMesh.h"
 #include "Kismet/GameplayStatics.h"
+#include "LevelInstance/LevelInstanceInterface.h"
+#include "LevelInstance/LevelInstanceSubsystem.h"
+#include "LevelInstance/LevelInstanceEditorInstanceActor.h"
 
 namespace PCGActorAndComponentMapping
 {
@@ -48,6 +52,24 @@ namespace PCGActorAndComponentMapping
 		TEXT("pcg.DisableDelayedUnregister"),
 		false,
 		TEXT("If delayed unregister for all is introducing bad behavior, disables it, allowing people to continue working while we investigate."));
+
+#if WITH_EDITOR
+	void PropagateToLevelInstanceActors(AActor* InActor, UPCGSubsystem* PCGSubsystem, TFunctionRef<bool(AActor* LevelActor)> InFunc)
+	{
+		if (!PCGSubsystem || !PCGSubsystem->GetWorld())
+		{
+			return;
+		}
+
+		if (ILevelInstanceInterface* LevelInstance = Cast<ILevelInstanceInterface>(InActor))
+		{
+			if (const ULevelInstanceSubsystem* LevelInstanceSubsystem = PCGSubsystem->GetWorld()->GetSubsystem<ULevelInstanceSubsystem>())
+			{
+				LevelInstanceSubsystem->ForEachActorInLevelInstance(LevelInstance, InFunc);
+			}
+		}
+	}
+#endif
 }
 
 UPCGActorAndComponentMapping::UPCGActorAndComponentMapping(UPCGSubsystem* InPCGSubsystem)
@@ -712,7 +734,6 @@ void UPCGActorAndComponentMapping::RegisterOrUpdateTracking(UPCGComponent* InCom
 		InComponent->OnPCGGraphGeneratedDelegate.AddRaw(this, &UPCGActorAndComponentMapping::OnPCGGraphGeneratedOrCleaned);
 		InComponent->OnPCGGraphCleanedDelegate.AddRaw(this, &UPCGActorAndComponentMapping::OnPCGGraphGeneratedOrCleaned);
 	}
-
 }
 
 void UPCGActorAndComponentMapping::RemapTracking(const UPCGComponent* InOldComponent, UPCGComponent* InNewComponent)
@@ -847,7 +868,8 @@ void UPCGActorAndComponentMapping::AddDelayedActors()
 		}
 		else
 		{
-			OnActorAdded_Internal(Actor, ActorPtrAndShouldDirty.Get<1>());
+			// Implementation note: since the delayed actors list is built from top-level actors (e.g. directly in the level) then depth here is 0.
+			OnActorAdded_Internal(Actor, ActorPtrAndShouldDirty.Get<1>(), 0);
 		}
 	}
 
@@ -856,10 +878,11 @@ void UPCGActorAndComponentMapping::AddDelayedActors()
 
 void UPCGActorAndComponentMapping::OnActorAdded(AActor* InActor)
 {
-	OnActorAdded_Internal(InActor);
+	// Implementation note: since this is called only for actors directly in the current level, the depth here is 0.
+	OnActorAdded_Internal(InActor, true, 0);
 }
 
-void UPCGActorAndComponentMapping::OnActorAdded_Internal(AActor* InActor, bool bShouldDirty)
+void UPCGActorAndComponentMapping::OnActorAdded_Internal(AActor* InActor, bool bShouldDirty, int32 LevelInstanceDepth)
 {
 	TRACE_CPUPROFILER_EVENT_SCOPE(UPCGActorAndComponentMapping::OnActorAdded);
 
@@ -883,10 +906,42 @@ void UPCGActorAndComponentMapping::OnActorAdded_Internal(AActor* InActor, bool b
 		return;
 	}
 
+#if WITH_EDITOR
+	// When a level instance is added to the level (or loaded), the Level Instance Editor Instance will be spawned
+	// but when we get the callback, it is not initialized yet (and does not contain the actors) hence adding it to the delayed actors.
+	// Note that while conceptually we'd want a callback when a level instance is loaded, it's a bit tricky vs. what world we want to do our changes in.
+	if (ALevelInstanceEditorInstanceActor* LevelInstanceEditorInstance = Cast<ALevelInstanceEditorInstanceActor>(InActor))
+	{
+		const FLevelInstanceID& LevelInstanceId = LevelInstanceEditorInstance->GetLevelInstanceID();
+		const ULevelInstanceSubsystem* LevelInstanceSubsystem = PCGSubsystem->GetWorld()->GetSubsystem<ULevelInstanceSubsystem>();
+
+		if (LevelInstanceId.IsValid() && LevelInstanceSubsystem)
+		{
+			if (ILevelInstanceInterface* LevelInstance = LevelInstanceSubsystem->GetLevelInstance(LevelInstanceId))
+			{
+				LevelInstanceSubsystem->ForEachActorInLevelInstance(LevelInstance, [this, bShouldDirty, LevelInstanceDepth, LevelInstanceEditorInstance](AActor* LevelActor)
+				{
+					if (LevelActor != LevelInstanceEditorInstance)
+					{
+						OnActorAdded_Internal(LevelActor, bShouldDirty, LevelInstanceDepth + 1);
+					}
+					
+					return true;
+				});
+			}
+		}
+		else
+		{
+			DelayedAddedActors.Emplace({ InActor, bShouldDirty });
+			return;
+		}
+	}
+#endif
+
 	if (AddOrUpdateTrackedActor(InActor) && bShouldDirty)
 	{
 		// Finally notify them all
-		OnActorChanged(InActor, /*bInHasMoved=*/ false);
+		OnActorChanged(InActor, /*bInHasMoved=*/ false, nullptr, LevelInstanceDepth);
 	}
 }
 
@@ -1017,6 +1072,12 @@ bool UPCGActorAndComponentMapping::UnregisterActor(AActor* InActor)
 }
 
 void UPCGActorAndComponentMapping::OnActorDeleted(AActor* InActor)
+{
+	// Implementation note: since this is called only for actors directly in the current level, the depth here is 0.
+	OnActorDeleted_Internal(InActor, 0);
+}
+
+void UPCGActorAndComponentMapping::OnActorDeleted_Internal(AActor* InActor, int32 LevelInstanceDepth)
 {
 	TRACE_CPUPROFILER_EVENT_SCOPE(UPCGActorAndComponentMapping::OnActorDeleted);
 
@@ -1030,6 +1091,35 @@ void UPCGActorAndComponentMapping::OnActorDeleted(AActor* InActor)
 	{
 		return;
 	}
+
+	// When editing/deleting/replacing a level instance, this will allow us to first remove the previously tracked actors that were under the level instance
+	if (ALevelInstanceEditorInstanceActor* LevelInstanceEditorInstance = Cast<ALevelInstanceEditorInstanceActor>(InActor))
+	{
+		const FLevelInstanceID& LevelInstanceId = LevelInstanceEditorInstance->GetLevelInstanceID();
+		const ULevelInstanceSubsystem* LevelInstanceSubsystem = PCGSubsystem->GetWorld()->GetSubsystem<ULevelInstanceSubsystem>();
+
+		if (LevelInstanceId.IsValid() && LevelInstanceSubsystem)
+		{
+			if (ILevelInstanceInterface* LevelInstance = LevelInstanceSubsystem->GetLevelInstance(LevelInstanceId))
+			{
+				LevelInstanceSubsystem->ForEachActorInLevelInstance(LevelInstance, [this, LevelInstanceDepth, LevelInstanceEditorInstance](AActor* LevelActor)
+				{
+					if (LevelActor != LevelInstanceEditorInstance)
+					{
+						OnActorDeleted_Internal(LevelActor, LevelInstanceDepth + 1);
+					}
+
+					return true;
+				});
+			}
+		}
+	}
+
+	PCGActorAndComponentMapping::PropagateToLevelInstanceActors(InActor, PCGSubsystem, [this, LevelInstanceDepth](AActor* LevelActor)
+	{
+		OnActorDeleted_Internal(LevelActor, LevelInstanceDepth + 1);
+		return true;
+	});
 #endif
 
 	if (!IsActorTracked(InActor))
@@ -1038,17 +1128,22 @@ void UPCGActorAndComponentMapping::OnActorDeleted(AActor* InActor)
 	}
 
 	// Notify all components that the actor has changed (was removed), but the Refresh will only happen AFTER the actor was actually removed from the world (because of delayed refresh).
-	OnActorChanged(InActor, /*bInHasMoved=*/ false);
+	OnActorChanged(InActor, /*bInHasMoved=*/ false, nullptr, LevelInstanceDepth);
 
 	// And then delete everything
 	UnregisterActor(InActor);
 }
 
 void UPCGActorAndComponentMapping::OnActorMoved(AActor* InActor)
+{
+	OnActorMoved_Internal(InActor, /*LevelInstanceDepth=*/0);
+}
+
+void UPCGActorAndComponentMapping::OnActorMoved_Internal(AActor* InActor, int32 LevelInstanceDepth)
 {
 	TRACE_CPUPROFILER_EVENT_SCOPE(UPCGActorAndComponentMapping::OnActorMoved);
 
-	if (!InActor || (PCGSubsystem && InActor->GetWorld() != PCGSubsystem->GetWorld()) || !IsActorTracked(InActor))
+	if (!InActor || (PCGSubsystem && InActor->GetWorld() != PCGSubsystem->GetWorld()))
 	{
 		return;
 	}
@@ -1058,10 +1153,23 @@ void UPCGActorAndComponentMapping::OnActorMoved(AActor* InActor)
 	{
 		return;
 	}
+
+	// If we've moved a level instance actor, we instead "push" forward the notification while making sure this will
+	//  trigger changes at this "level" in the world hierarchy
+	PCGActorAndComponentMapping::PropagateToLevelInstanceActors(InActor, PCGSubsystem, [this, LevelInstanceDepth](AActor* LevelActor)
+	{
+		OnActorMoved_Internal(LevelActor, LevelInstanceDepth + 1);
+		return true;
+	});
 #endif
 
+	if (!IsActorTracked(InActor))
+	{
+		return;
+	}
+
 	// Notify all components
-	OnActorChanged(InActor, /*bInHasMoved=*/ true);
+	OnActorChanged(InActor, /*bInHasMoved=*/ true, nullptr, LevelInstanceDepth);
 
 	// Update Actor position
 	if (FBox* ActorBounds = TrackedActorToPositionMap.Find(InActor))
@@ -1173,7 +1281,7 @@ void UPCGActorAndComponentMapping::OnObjectPropertyChanged(UObject* InObject, FP
 	}
 }
 
-void UPCGActorAndComponentMapping::OnActorChanged(AActor* InActor, bool bInHasMoved, const UObject* InOriginatingChangeObject)
+void UPCGActorAndComponentMapping::OnActorChanged(AActor* InActor, bool bInHasMoved, const UObject* InOriginatingChangeObject, int32 LevelInstanceDepth)
 {
 	TRACE_CPUPROFILER_EVENT_SCOPE(UPCGActorAndComponentMapping::OnActorChanged);
 
@@ -1321,11 +1429,38 @@ void UPCGActorAndComponentMapping::OnActorChanged(AActor* InActor, bool bInHasMo
 		}
 	}
 
+	ULevelInstanceSubsystem* LevelInstanceSubsystem = (PCGSubsystem && PCGSubsystem->GetWorld()) ? PCGSubsystem->GetWorld()->GetSubsystem<ULevelInstanceSubsystem>() : nullptr;
+
 	// And refresh all dirtied components
 	for (UPCGComponent* Component : DirtyComponents)
 	{
 		if (Component)
 		{
+			// When an object changes, we need to make sure that we don't trigger a refresh on PCG components that are "higher" in the
+			// level hierarchy, otherwise we will end up generating in the Level Instance level, which is wrong.
+			// Note that in some instances (e.g. when something happens at the Level Instance level) we need to make sure that the
+			// PCG components higher-up are properly updated, hence the level instance depth
+			if (LevelInstanceSubsystem)
+			{
+				const ILevelInstanceInterface* ActorLevelInstance = LevelInstanceSubsystem->GetParentLevelInstance(InActor);
+				int LocalDepth = LevelInstanceDepth;
+				while (ActorLevelInstance && LocalDepth-- > 0)
+				{
+					ActorLevelInstance = LevelInstanceSubsystem->GetParentLevelInstance(CastChecked<AActor>(ActorLevelInstance));
+				}
+
+				const ILevelInstanceInterface* ComponentLevelInstance = LevelInstanceSubsystem->GetParentLevelInstance(Component->GetOwner());
+				while (ComponentLevelInstance && ComponentLevelInstance != ActorLevelInstance)
+				{
+					ComponentLevelInstance = LevelInstanceSubsystem->GetParentLevelInstance(CastChecked<AActor>(ComponentLevelInstance));
+				}
+
+				if (ActorLevelInstance != ComponentLevelInstance)
+				{
+					continue;
+				}
+			}
+			
 			Component->Refresh();
 		}
 	}
diff --git a/Engine/Plugins/Experimental/PCG/Source/PCG/Private/PCGComponent.cpp b/Engine/Plugins/Experimental/PCG/Source/PCG/Private/PCGComponent.cpp
--- a/Engine/Plugins/Experimental/PCG/Source/PCG/Private/PCGComponent.cpp
+++ b/Engine/Plugins/Experimental/PCG/Source/PCG/Private/PCGComponent.cpp
@@ -595,7 +595,7 @@ AActor* UPCGComponent::ClearPCGLink(UClass* TemplateActor)
 	UWorld* World = GetWorld();
 
 	// First create a new actor that will be the new owner of all the resources
-	AActor* NewActor = UPCGActorHelpers::SpawnDefaultActor(World, TemplateActor ? TemplateActor : AActor::StaticClass(), TEXT("PCGStamp"), GetOwner()->GetTransform());
+	AActor* NewActor = UPCGActorHelpers::SpawnDefaultActor(World, GetOwner()->GetLevel(), TemplateActor ? TemplateActor : AActor::StaticClass(), TEXT("PCGStamp"), GetOwner()->GetTransform(), GetOwner());
 
 	// Then move all resources linked to this component to this actor
 	bool bHasMovedResources = MoveResourcesToNewActor(NewActor, /*bCreateChild=*/false);
@@ -712,7 +712,7 @@ bool UPCGComponent::MoveResourcesToNewActor(AActor* InNewActor, bool bCreateChil
 
 	if (bCreateChild)
 	{
-		NewActor = UPCGActorHelpers::SpawnDefaultActor(GetWorld(), NewActor->GetClass(), TEXT("PCGStampChild"), Owner->GetTransform());
+		NewActor = UPCGActorHelpers::SpawnDefaultActor(GetWorld(), GetOwner()->GetLevel(), NewActor->GetClass(), TEXT("PCGStampChild"), Owner->GetTransform());
 		NewActor->AttachToActor(InNewActor, FAttachmentTransformRules::KeepWorldTransform);
 		check(NewActor);
 	}
@@ -2372,6 +2372,9 @@ void FPCGComponentInstanceData::ApplyToComponent(UActorComponent* Component, con
 #if WITH_EDITOR
 			// bDirtyGenerated is transient.
 			PCGComponent->bDirtyGenerated = SourceComponent->bDirtyGenerated; 
+
+			// While this is serialized, it is not properly copied because it is not visible. This is needed otherwise a refresh can retoggle from not generated to generated
+			PCGComponent->bForceGenerateOnBPAddedToWorld = SourceComponent->bForceGenerateOnBPAddedToWorld;
 #endif // WITH_EDITOR
 
 			// Non-critical but should be done: transient data, tracked actors cache, landscape tracking
diff --git a/Engine/Plugins/Experimental/PCG/Source/PCG/Public/Helpers/PCGActorHelpers.h b/Engine/Plugins/Experimental/PCG/Source/PCG/Public/Helpers/PCGActorHelpers.h
--- a/Engine/Plugins/Experimental/PCG/Source/PCG/Public/Helpers/PCGActorHelpers.h
+++ b/Engine/Plugins/Experimental/PCG/Source/PCG/Public/Helpers/PCGActorHelpers.h
@@ -74,35 +74,38 @@ public:
 	/**
 	* Spawn a new actor of type T and attach it to the parent (if not null)
 	* @param World The world
+	* @param Level The level to spawn into
 	* @param BaseName Base name for the actor, will have a unique name
 	* @param Transform The transform for the new actor
 	* @param Parent Optional parent to attach to.
 	*/
 	template <typename T = AActor, typename = typename std::enable_if_t<std::is_base_of_v<AActor, T>>>
-	inline static AActor* SpawnDefaultActor(UWorld* World, FName BaseName, const FTransform& Transform, AActor* Parent = nullptr)
+	inline static AActor* SpawnDefaultActor(UWorld* World, ULevel* Level, FName BaseName, const FTransform& Transform, AActor* Parent = nullptr)
 	{
-		return SpawnDefaultActor(World, T::StaticClass(), BaseName, Transform, Parent);
+		return SpawnDefaultActor(World, Level, T::StaticClass(), BaseName, Transform, Parent);
 	}
 
 	/**
 	* Spawn a new actor and attach it to the parent (if not null)
 	* @param World The world
+	* @param Level The level to spawn into
 	* @param ActorClass Class of the actor to spawn
 	* @param BaseName Base name for the actor, will have a unique name
 	* @param Transform The transform for the new actor
 	* @param Parent Optional parent to attach to.
 	*/
-	static AActor* SpawnDefaultActor(UWorld* World, TSubclassOf<AActor> ActorClass, FName BaseName, const FTransform& Transform, AActor* Parent = nullptr);
+	static AActor* SpawnDefaultActor(UWorld* World, ULevel* Level, TSubclassOf<AActor> ActorClass, FName BaseName, const FTransform& Transform, AActor* Parent = nullptr);
 
 	/**
 	* Spawn a new actor and attach it to the parent (if not null)
 	* @param World The world
+	* @param Level The level to spawn into
 	* @param ActorClass Class of the actor to spawn
 	* @param Transform The transform for the new actor
 	* @param SpawnParams The spawn parameters
 	* @param Parent Optional parent to attach to.
 	*/
-	static AActor* SpawnDefaultActor(UWorld* World, TSubclassOf<AActor> ActorClass, const FTransform& Transform, const FActorSpawnParameters& SpawnParams, AActor* Parent = nullptr);
+	static AActor* SpawnDefaultActor(UWorld* World, ULevel* Level, TSubclassOf<AActor> ActorClass, const FTransform& Transform, const FActorSpawnParameters& SpawnParams, AActor* Parent = nullptr);
 
 	/**
 	 * Return the grid cell coordinates on the PCG partition grid given a position and the grid size.
diff --git a/Engine/Plugins/Experimental/PCG/Source/PCG/Public/PCGActorAndComponentMapping.h b/Engine/Plugins/Experimental/PCG/Source/PCG/Public/PCGActorAndComponentMapping.h
--- a/Engine/Plugins/Experimental/PCG/Source/PCG/Public/PCGActorAndComponentMapping.h
+++ b/Engine/Plugins/Experimental/PCG/Source/PCG/Public/PCGActorAndComponentMapping.h
@@ -119,9 +119,11 @@ private:
 	bool UnregisterActor(AActor* InActor);
 
 	void OnActorAdded(AActor* InActor);
-	void OnActorAdded_Internal(AActor* InActor, bool bShouldDirty = true);
+	void OnActorAdded_Internal(AActor* InActor, bool bShouldDirty, int32 LevelInstanceDepth);
 	void OnActorDeleted(AActor* InActor);
+	void OnActorDeleted_Internal(AActor* InActor, int32 LevelInstanceDepth);
 	void OnActorMoved(AActor* InActor);
+	void OnActorMoved_Internal(AActor* InActor, int32 LevelInstanceDepth);
 	void OnLandscapeChanged(ALandscapeProxy* InLandscape, const FLandscapeProxyComponentDataChangedParams& InChangeParams);
 	void OnPreObjectPropertyChanged(UObject* InObject, const FEditPropertyChain& InEditPropertyChain);
 	void OnObjectPropertyChanged(UObject* InObject, FPropertyChangedEvent& InEvent);
@@ -137,7 +139,7 @@ private:
 	* Can specify if the actor has moved to also update components that were at its previous position.
 	* Can also specify an optional object, originating the change, to avoid re-dirtying a component if it was the origin.
 	*/
-	void OnActorChanged(AActor* InActor, bool bInHasMoved, const UObject* InOriginatingChangeObject = nullptr);
+	void OnActorChanged(AActor* InActor, bool bInHasMoved, const UObject* InOriginatingChangeObject = nullptr, int32 LevelInstanceDepth = 0);
 
 	/** Update dependencies for a given tracked actor. */
 	void UpdateActorDependencies(AActor* InActor);
