[PCG] Remove impossible cast from spline sampling

#rnx

diff --git a/Engine/Plugins/Experimental/PCG/Source/PCG/Private/Data/PCGSplineData.cpp b/Engine/Plugins/Experimental/PCG/Source/PCG/Private/Data/PCGSplineData.cpp
--- a/Engine/Plugins/Experimental/PCG/Source/PCG/Private/Data/PCGSplineData.cpp
+++ b/Engine/Plugins/Experimental/PCG/Source/PCG/Private/Data/PCGSplineData.cpp
@@ -117,16 +117,7 @@ const UPCGPointData* UPCGSplineData::CreatePointData(FPCGContext* Context) const
 	FPCGSplineSamplerParams SamplerParams;
 	SamplerParams.Mode = EPCGSplineSamplingMode::Distance;
 
-	const UPCGSpatialData* ProjectionTarget = nullptr;
-	FPCGProjectionParams ProjectionParams;
-	/* Impossible cast, UPCGSplineProjectionData is not a UPCGSplineData
-	if (const UPCGSplineProjectionData* SplineProjection = Cast<const UPCGSplineProjectionData>(this))
-	{
-		ProjectionTarget = SplineProjection->GetSurface();
-		ProjectionParams = SplineProjection->GetProjectionParams();
-	}*/
-
-	PCGSplineSampler::SampleLineData(this, nullptr, ProjectionTarget, ProjectionParams, SamplerParams, Data);
+	PCGSplineSampler::SampleLineData(this, /*InBoundingShape=*/nullptr, /*InProjectionTarget=*/nullptr, /*InProjectionParams=*/{}, SamplerParams, Data);
 	UE_LOG(LogPCG, Verbose, TEXT("Spline generated %d points"), Data->GetPoints().Num());
 
 	return Data;
