[PCG] Support CPU texture sampling

diff --git a/Engine/Plugins/Experimental/PCG/Source/PCG/Private/Data/PCGTextureData.cpp b/Engine/Plugins/Experimental/PCG/Source/PCG/Private/Data/PCGTextureData.cpp
--- a/Engine/Plugins/Experimental/PCG/Source/PCG/Private/Data/PCGTextureData.cpp
+++ b/Engine/Plugins/Experimental/PCG/Source/PCG/Private/Data/PCGTextureData.cpp
@@ -15,11 +15,11 @@
 namespace PCGTextureSampling
 {
 	template<typename ValueType>
-	bool Sample(const FVector2D& InPosition, 
-		const FBox2D& InSurface, 
-		const UPCGBaseTextureData* InTextureData, 
-		int32 Width, 
-		int32 Height, 
+	bool Sample(const FVector2D& InPosition,
+		const FBox2D& InSurface,
+		const UPCGBaseTextureData* InTextureData,
+		int32 Width,
+		int32 Height,
 		ValueType& SampledValue,
 		TFunctionRef<ValueType(int32 Index)> SamplingFunction)
 	{
@@ -52,7 +52,7 @@ namespace PCGTextureSampling
 			const FVector Translation = FVector(0.5 + InTextureData->CenterOffset.X, 0.5 + InTextureData->CenterOffset.Y, 0);
 
 			FTransform Transform = FTransform(Rotation, Translation, Scale);
-			
+
 			// Transform to tile-space
 			const FVector2D SamplePosition = FVector2D(Transform.InverseTransformPosition(FVector(LocalSpacePos, 0.f)));
 
@@ -288,12 +288,164 @@ void UPCGTextureData::Initialize(UTexture2D* InTexture, const FTransform& InTran
 		return;
 	}
 
-	if (!IsSupported(InTexture))
+	// Prioritize initializing from a CPU texture when the provided texture is marked as CPU accessible
+	if (InitializeFromCPUTexture())
 	{
-		UE_LOG(LogPCG, Error, TEXT("PCGTextureData does not support one or more settings on '%s'"), *Texture->GetFName().ToString());
-
 		PostInitializeCallback();
-		return;
+	}
+	else
+	{
+		if (!InitializeFromGPUTexture(PostInitializeCallback))
+		{
+			UE_LOG(LogPCG, Error, TEXT("PCGTextureData failed to initialize texture '%s'"), *Texture->GetFName().ToString());
+
+			PostInitializeCallback();
+		}
+	}
+}
+
+UPCGSpatialData* UPCGTextureData::CopyInternal() const
+{
+	UPCGTextureData* NewTextureData = NewObject<UPCGTextureData>();
+
+	CopyBaseTextureData(NewTextureData);
+
+	NewTextureData->Texture = Texture;
+
+	return NewTextureData;
+}
+
+bool UPCGTextureData::InitializeFromCPUTexture()
+{
+	if (!Texture.IsValid())
+	{
+		return false;
+	}
+
+	FSharedImageConstRef CPUTextureRef = Texture->GetCPUCopy();
+	if (!CPUTextureRef.IsValid())
+	{
+		return false;
+	}
+
+	Width = CPUTextureRef->SizeX;
+	Height = CPUTextureRef->SizeY;
+
+	const int32 PixelCount = Width * Height;
+	ColorData.SetNum(PixelCount);
+
+	if (CPUTextureRef->Format == ERawImageFormat::G8)
+	{
+		const TArrayView64<const uint8> DataView = CPUTextureRef->AsG8();
+
+		for (int32 D = 0; D < PixelCount; ++D)
+		{
+			ColorData[D] = FColor(DataView[D], DataView[D], DataView[D]).ReinterpretAsLinear();
+		}
+	}
+	else if (CPUTextureRef->Format == ERawImageFormat::BGRA8)
+	{
+		const TArrayView64<const FColor> DataView = CPUTextureRef->AsBGRA8();
+
+		// Memory representation of FColor is BGRA
+		for (int32 D = 0; D < PixelCount; ++D)
+		{
+			ColorData[D] = DataView[D].ReinterpretAsLinear();
+		}
+	}
+	else if (CPUTextureRef->Format == ERawImageFormat::BGRE8)
+	{
+		const TArrayView64<const FColor> DataView = CPUTextureRef->AsBGRE8();
+
+		// Memory representation of FColor is BGRA
+		for (int32 D = 0; D < PixelCount; ++D)
+		{
+			ColorData[D] = DataView[D].ReinterpretAsLinear();
+		}
+	}
+	else if (CPUTextureRef->Format == ERawImageFormat::RGBA16)
+	{
+		const TArrayView64<const uint16> DataView = CPUTextureRef->AsRGBA16();
+		check(PixelCount * 4 == DataView.Num());
+
+		for (int32 D = 0; D < PixelCount; ++D)
+		{
+			const uint32 Index = D * 4;
+
+			// To avoid swapping R and B to BGRA format, we can access in the correct order here
+			ColorData[D] = FColor(DataView[Index + 2], DataView[Index + 1], DataView[Index + 0], DataView[Index + 3]).ReinterpretAsLinear();
+		}
+	}
+	else if (CPUTextureRef->Format == ERawImageFormat::RGBA16F)
+	{
+		const TArrayView64<const FFloat16Color> DataView = CPUTextureRef->AsRGBA16F();
+
+		for (int32 D = 0; D < PixelCount; ++D)
+		{
+			// Swap R and B to achieve BGRA format
+			FLinearColor Temp = FLinearColor(DataView[D]);
+			Swap(Temp.R, Temp.B);
+			ColorData[D] = Temp;
+		}
+	}
+	else if (CPUTextureRef->Format == ERawImageFormat::RGBA32F)
+	{
+		const TArrayView64<const FLinearColor> DataView = CPUTextureRef->AsRGBA32F();
+
+		for (int32 D = 0; D < PixelCount; ++D)
+		{
+			// Swap R and B to achieve BGRA format
+			FLinearColor Temp = DataView[D];
+			Swap(Temp.R, Temp.B);
+			ColorData[D] = Temp;
+		}
+	}
+	else if (CPUTextureRef->Format == ERawImageFormat::G16)
+	{
+		const TArrayView64<const uint16> DataView = CPUTextureRef->AsG16();
+
+		for (int32 D = 0; D < PixelCount; ++D)
+		{
+			ColorData[D] = FColor(DataView[D], DataView[D], DataView[D]).ReinterpretAsLinear();
+		}
+	}
+	else if (CPUTextureRef->Format == ERawImageFormat::R16F)
+	{
+		const TArrayView64<const FFloat16> DataView = CPUTextureRef->AsR16F();
+
+		for (int32 D = 0; D < PixelCount; ++D)
+		{
+			ColorData[D] = FLinearColor(DataView[D], DataView[D], DataView[D]);
+		}
+	}
+	else if (CPUTextureRef->Format == ERawImageFormat::R32F)
+	{
+		const TArrayView64<const float> DataView = CPUTextureRef->AsR32F();
+
+		for (int32 D = 0; D < PixelCount; ++D)
+		{
+			ColorData[D] = FLinearColor(DataView[D], DataView[D], DataView[D]);
+		}
+	}
+	else
+	{
+		UE_LOG(LogPCG, Error, TEXT("PCGTextureReadback has an invalid format (%d) for CPU texture '%s'."), CPUTextureRef->Format, *Texture->GetFName().ToString());
+
+		Width = 0;
+		Height = 0;
+		ColorData.SetNum(0);
+
+		return false;
+	}
+
+	return true;
+}
+
+bool UPCGTextureData::InitializeFromGPUTexture(const TFunction<void()>& PostInitializeCallback)
+{
+	if (!Texture.IsValid())
+	{
+		return false;
 	}
 
 #if WITH_EDITOR
@@ -307,7 +459,7 @@ void UPCGTextureData::Initialize(UTexture2D* InTexture, const FTransform& InTran
 	Texture->WaitForPendingInitOrStreaming();
 
 	FTexturePlatformData* PlatformData = Texture->GetPlatformData();
-	FTextureResource* TextureResource = InTexture->GetResource();
+	FTextureResource* TextureResource = Texture->GetResource();
 
 	if (PlatformData && TextureResource && TextureResource->TextureRHI)
 	{
@@ -353,25 +505,8 @@ void UPCGTextureData::Initialize(UTexture2D* InTexture, const FTransform& InTran
 	else
 	{
 		UE_LOG(LogPCG, Error, TEXT("PCGTextureData failed to acquire texture resource for '%s'"), *Texture->GetFName().ToString());
-
-		PostInitializeCallback();
+		return false;
 	}
-}
 
-bool UPCGTextureData::IsSupported(UTexture2D* InTexture)
-{
-	// TODO: GPU readback seems to support all textures, but future implementations (such as UBitmap) may have
-	// some limitations, so we will avoid removing this trivial function from the API for now
 	return true;
 }
-
-UPCGSpatialData* UPCGTextureData::CopyInternal() const
-{
-	UPCGTextureData* NewTextureData = NewObject<UPCGTextureData>();
-
-	CopyBaseTextureData(NewTextureData);
-
-	NewTextureData->Texture = Texture;
-
-	return NewTextureData;
-}
diff --git a/Engine/Plugins/Experimental/PCG/Source/PCG/Private/Elements/PCGTextureSampler.cpp b/Engine/Plugins/Experimental/PCG/Source/PCG/Private/Elements/PCGTextureSampler.cpp
--- a/Engine/Plugins/Experimental/PCG/Source/PCG/Private/Elements/PCGTextureSampler.cpp
+++ b/Engine/Plugins/Experimental/PCG/Source/PCG/Private/Elements/PCGTextureSampler.cpp
@@ -16,6 +16,15 @@
 
 #define LOCTEXT_NAMESPACE "PCGTextureSamplerElement"
 
+#if WITH_EDITOR
+FText UPCGTextureSamplerSettings::GetNodeTooltipText() const
+{
+	return LOCTEXT("NodeTooltip", "Generates points by sampling the given texture.\n"
+		"If the texture is CPU-accessible, the sampler will prefer the CPU version of the texture.\n"
+		"Otherwise, the texture will be read back from the GPU if one is present.");
+}
+#endif // WITH_EDITOR
+
 TArray<FPCGPinProperties> UPCGTextureSamplerSettings::OutputPinProperties() const
 {
 	TArray<FPCGPinProperties> Properties;
@@ -61,12 +70,6 @@ bool FPCGTextureSamplerElement::ExecuteInternal(FPCGContext* InContext) const
 		return true;
 	}
 
-	if (!UPCGTextureData::IsSupported(Texture))
-	{
-		PCGE_LOG(Error, GraphAndLog, FText::Format(LOCTEXT("UnsupportedTextureFormat", "Texture '{0}' has unsupported settings"), FText::FromName(Texture->GetFName())));
-		return true;
-	}
-
 	const FTransform& Transform = Settings->Transform;
 	const bool bUseAbsoluteTransform = Settings->bUseAbsoluteTransform;
 	const EPCGTextureDensityFunction DensityFunction = Settings->DensityFunction;
diff --git a/Engine/Plugins/Experimental/PCG/Source/PCG/Public/Data/PCGTextureData.h b/Engine/Plugins/Experimental/PCG/Source/PCG/Public/Data/PCGTextureData.h
--- a/Engine/Plugins/Experimental/PCG/Source/PCG/Public/Data/PCGTextureData.h
+++ b/Engine/Plugins/Experimental/PCG/Source/PCG/Public/Data/PCGTextureData.h
@@ -106,14 +106,21 @@ public:
 
 	void Initialize(UTexture2D* InTexture, const FTransform& InTransform, const TFunction<void()>& PostInitializeCallback);
 
-	/** Returns true if the format of InTexture is compatible and can be loaded. Will load texture if not already loaded. */
-	static bool IsSupported(UTexture2D* InTexture);
-
 	//~Begin UPCGSpatialData interface
 protected:
 	virtual UPCGSpatialData* CopyInternal() const override;
 	//~End UPCGSpatialData interface
 
+private:
+	/** Attempts to initialize the UPCGTextureData from a CPU-accessible texture. Returns true if CPU initialization succeeds. */
+	bool InitializeFromCPUTexture();
+
+	/** 
+	* Attempts to initialize the UPCGTextureData from a GPU-accessible texture. Returns true if GPU texture readback can be dispatched. 
+	* The PostInitializeCallback is only executed if this function succeeds.
+	*/
+	bool InitializeFromGPUTexture(const TFunction<void()>& PostInitializeCallback);
+
 public:
 	UPROPERTY(BlueprintReadOnly, VisibleAnywhere, Category = Properties)
 	TWeakObjectPtr<UTexture2D> Texture = nullptr;
diff --git a/Engine/Plugins/Experimental/PCG/Source/PCG/Public/Elements/PCGTextureSampler.h b/Engine/Plugins/Experimental/PCG/Source/PCG/Public/Elements/PCGTextureSampler.h
--- a/Engine/Plugins/Experimental/PCG/Source/PCG/Public/Elements/PCGTextureSampler.h
+++ b/Engine/Plugins/Experimental/PCG/Source/PCG/Public/Elements/PCGTextureSampler.h
@@ -18,6 +18,7 @@ public:
 #if WITH_EDITOR
 	virtual FName GetDefaultNodeName() const override { return FName(TEXT("GetTextureData")); }
 	virtual FText GetDefaultNodeTitle() const override { return NSLOCTEXT("PCGTextureSamplerSettings", "NodeTitle", "Get Texture Data"); }
+	virtual FText GetNodeTooltipText() const override;
 	virtual EPCGSettingsType GetType() const override { return EPCGSettingsType::Spatial; }
 #endif
 
