[PCG] Provide more information on fail Get/Set Attribute

#rnx

diff --git a/Engine/Plugins/Experimental/PCG/Source/PCG/Private/Metadata/PCGMetadataAccessor.cpp b/Engine/Plugins/Experimental/PCG/Source/PCG/Private/Metadata/PCGMetadataAccessor.cpp
--- a/Engine/Plugins/Experimental/PCG/Source/PCG/Private/Metadata/PCGMetadataAccessor.cpp
+++ b/Engine/Plugins/Experimental/PCG/Source/PCG/Private/Metadata/PCGMetadataAccessor.cpp
@@ -6,13 +6,35 @@
 
 #include UE_INLINE_GENERATED_CPP_BY_NAME(PCGMetadataAccessor)
 
+#define LOCTEXT_NAMESPACE "PCGMetadataAccessorHelpers"
+
+namespace PCGMetadataAccessorHelpers
+{
+	static const FText InvalidAttributeNameFormat = LOCTEXT("InvalidAttributeNameFormat", "[PCG] Invalid attribute name({0})");
+	static const FText InvalidTypeFormat = LOCTEXT("InvalidTypeFormat", "[PCG] Attribute {0} is of type {1} and it is not the requested type {2}");
+	static const FText NoMetadata = LOCTEXT("NoMetadata", "[PCG] No metadata provided");
+
+	void OnException(FText ErrorMessage)
+	{
+		if (FFrame::GetThreadLocalTopStackFrame()->Object)
+		{
+			const FBlueprintExceptionInfo ExceptionInfo(EBlueprintExceptionType::FatalError, ErrorMessage);
+			FBlueprintCoreDelegates::ThrowScriptException(FFrame::GetThreadLocalTopStackFrame()->Object, *FFrame::GetThreadLocalTopStackFrame(), ExceptionInfo);
+		}
+		else
+		{
+			UE_LOG(LogPCG, Error, TEXT("%s"), *ErrorMessage.ToString());
+		}
+	}
+}
+
 /** Key-based implmentations */
 template<typename T>
 T UPCGMetadataAccessorHelpers::GetAttribute(PCGMetadataEntryKey Key, const UPCGMetadata* Metadata, FName AttributeName)
 {
 	if (!Metadata)
 	{
-		UE_LOG(LogPCG, Error, TEXT("Source data has no metadata"));
+		PCGMetadataAccessorHelpers::OnException(PCGMetadataAccessorHelpers::NoMetadata);
 		return T{};
 	}
 
@@ -23,12 +45,12 @@ T UPCGMetadataAccessorHelpers::GetAttribute(PCGMetadataEntryKey Key, const UPCGM
 	}
 	else if (Attribute)
 	{
-		UE_LOG(LogPCG, Error, TEXT("Attribute %s does not have the matching type"), *AttributeName.ToString());
+		PCGMetadataAccessorHelpers::OnException(FText::Format(PCGMetadataAccessorHelpers::InvalidTypeFormat, FText::FromName(AttributeName), FText::FromString(PCG::Private::GetTypeName(Attribute->GetTypeId())), FText::FromString(PCG::Private::GetTypeName(PCG::Private::MetadataTypes<T>::Id))));
 		return T{};
 	}
 	else
 	{
-		UE_LOG(LogPCG, Error, TEXT("Invalid attribute name (%s)"), *AttributeName.ToString());
+		PCGMetadataAccessorHelpers::OnException(FText::Format(PCGMetadataAccessorHelpers::InvalidAttributeNameFormat, FText::FromName(AttributeName)));
 		return T{};
 	}
 }
@@ -38,7 +60,7 @@ void UPCGMetadataAccessorHelpers::SetAttribute(PCGMetadataEntryKey& Key, UPCGMet
 {
 	if (!Metadata)
 	{
-		UE_LOG(LogPCG, Error, TEXT("Data has no metadata; cannot write value in attribute"));
+		PCGMetadataAccessorHelpers::OnException(PCGMetadataAccessorHelpers::NoMetadata);
 		return;
 	}
 
@@ -46,7 +68,7 @@ void UPCGMetadataAccessorHelpers::SetAttribute(PCGMetadataEntryKey& Key, UPCGMet
 
 	if (Key == PCGInvalidEntryKey)
 	{
-		UE_LOG(LogPCG, Error, TEXT("Metadata key has no entry, therefore can't set values"));
+		PCGMetadataAccessorHelpers::OnException(LOCTEXT("NoMetadataEntry", "[PCG] Metadata key has no entry, therefore can't set values"));
 		return;
 	}
 
@@ -57,11 +79,11 @@ void UPCGMetadataAccessorHelpers::SetAttribute(PCGMetadataEntryKey& Key, UPCGMet
 	}
 	else if (Attribute)
 	{
-		UE_LOG(LogPCG, Error, TEXT("Attribute %s does not have the matching type"), *AttributeName.ToString());
+		PCGMetadataAccessorHelpers::OnException(FText::Format(PCGMetadataAccessorHelpers::InvalidTypeFormat, FText::FromName(AttributeName), FText::FromString(PCG::Private::GetTypeName(Attribute->GetTypeId())), FText::FromString(PCG::Private::GetTypeName(PCG::Private::MetadataTypes<T>::Id))));
 	}
 	else
 	{
-		UE_LOG(LogPCG, Error, TEXT("Invalid attribute name (%s)"), *AttributeName.ToString());
+		PCGMetadataAccessorHelpers::OnException(FText::Format(PCGMetadataAccessorHelpers::InvalidAttributeNameFormat, FText::FromName(AttributeName)));
 	}
 }
 
@@ -69,7 +91,7 @@ bool UPCGMetadataAccessorHelpers::HasAttributeSetByMetadataKey(PCGMetadataEntryK
 {
 	if (!Metadata)
 	{
-		UE_LOG(LogPCG, Error, TEXT("Data has no metadata"));
+		PCGMetadataAccessorHelpers::OnException(PCGMetadataAccessorHelpers::NoMetadata);
 		return false;
 	}
 
@@ -85,7 +107,7 @@ bool UPCGMetadataAccessorHelpers::HasAttributeSetByMetadataKey(PCGMetadataEntryK
 	}
 	else
 	{
-		UE_LOG(LogPCG, Error, TEXT("Metadata does not have a %s attribute"), *AttributeName.ToString());
+		PCGMetadataAccessorHelpers::OnException(FText::Format(PCGMetadataAccessorHelpers::InvalidAttributeNameFormat, FText::FromName(AttributeName)));
 		return false;
 	}
 }
@@ -409,3 +431,5 @@ bool UPCGMetadataAccessorHelpers::HasAttributeSet(const FPCGPoint& Point, const
 {
 	return HasAttributeSetByMetadataKey(Point.MetadataEntry, Metadata, AttributeName);
 }
+
+#undef LOCTEXT_NAMESPACE
\ No newline at end of file
