[PCG] Adds a node that sorts points based on attribute
--New helper function SortByAttribute included
--Unit test added
--Added ability to randomize density in PCGTestCommon

diff --git a/Engine/Plugins/Experimental/PCG/Source/PCG/Private/Elements/PCGSortPoints.cpp b/Engine/Plugins/Experimental/PCG/Source/PCG/Private/Elements/PCGSortPoints.cpp
new file mode 100644
--- /dev/null
+++ b/Engine/Plugins/Experimental/PCG/Source/PCG/Private/Elements/PCGSortPoints.cpp
@@ -0,0 +1,68 @@
+// Copyright Epic Games, Inc. All Rights Reserved.
+
+#include "Elements/PCGSortPoints.h"
+
+#include "PCGContext.h"
+#include "Data/PCGPointData.h"
+#include "Metadata/Accessors/IPCGAttributeAccessor.h"
+#include "Metadata/Accessors/PCGAttributeAccessorHelpers.h"
+#include "Metadata/Accessors/PCGAttributeAccessorKeys.h"
+
+#define LOCTEXT_NAMESPACE "PCGSortPointsElement"
+
+FPCGElementPtr UPCGSortPointsSettings::CreateElement() const
+{
+	return MakeShared<FPCGSortPointsElement>();
+}
+
+bool FPCGSortPointsElement::ExecuteInternal(FPCGContext* Context) const
+{
+	TRACE_CPUPROFILER_EVENT_SCOPE(FPCGSortPointsElement::Execute);
+	
+	check(Context);
+
+	const UPCGSortPointsSettings* Settings = Context->GetInputSettings<UPCGSortPointsSettings>();
+	check(Settings);
+	
+	TArray<FPCGTaggedData> Inputs = Context->InputData.GetInputs();
+	TArray<FPCGTaggedData>& Outputs = Context->OutputData.TaggedData;
+
+	for (FPCGTaggedData& CurrentInput : Inputs)
+	{
+		TRACE_CPUPROFILER_EVENT_SCOPE(FPCGSortPointsElement::InputLoop);
+
+		const FPCGTaggedData& Input = CurrentInput;
+
+		const UPCGPointData* PointData = Cast<UPCGPointData>(Input.Data);
+
+		if (PointData) 
+		{
+			const TArray<FPCGPoint>& InputPoints = PointData->GetPoints();
+			FPCGAttributePropertyInputSelector InputSource = Settings->InputSource.CopyAndFixLast(PointData);
+
+			FPCGTaggedData& Output = Outputs.Add_GetRef(Input);
+			UPCGPointData* OutputData = NewObject<UPCGPointData>();
+			OutputData->InitializeFromData(PointData);
+
+			TArray<FPCGPoint>& OutputPoints = OutputData->GetMutablePoints();
+			OutputPoints = InputPoints;
+			Output.Data = OutputData;
+
+			TUniquePtr<const IPCGAttributeAccessor> Accessor = PCGAttributeAccessorHelpers::CreateConstAccessor(PointData, InputSource);
+			TUniquePtr<const IPCGAttributeAccessorKeys> Keys = PCGAttributeAccessorHelpers::CreateConstKeys(PointData, InputSource);
+			
+			if (Accessor.IsValid() && Keys.IsValid())
+			{
+				PCGAttributeAccessorHelpers::SortByAttribute(*Accessor, *Keys, OutputPoints, Settings->SortMethod == EPCGSortMethod::Ascending);
+			}
+			else
+			{
+				PCGE_LOG(Error, GraphAndLog, LOCTEXT("InvalidAccessor","Attribute does not exist or cannot be sorted."));
+			}
+		}
+	}
+
+	return true;
+}
+
+#undef LOCTEXT_NAMESPACE
diff --git a/Engine/Plugins/Experimental/PCG/Source/PCG/Private/Tests/Elements/PCGSortPointsTest.cpp b/Engine/Plugins/Experimental/PCG/Source/PCG/Private/Tests/Elements/PCGSortPointsTest.cpp
new file mode 100644
--- /dev/null
+++ b/Engine/Plugins/Experimental/PCG/Source/PCG/Private/Tests/Elements/PCGSortPointsTest.cpp
@@ -0,0 +1,142 @@
+// Copyright Epic Games, Inc. All Rights Reserved.
+
+#include "Tests/Determinism/PCGDeterminismTestsCommon.h"
+
+#include "PCGComponent.h"
+#include "PCGContext.h"
+#include "Data/PCGPointData.h"
+#include "Elements/PCGSortPoints.h"
+
+namespace SortCommonTestData
+{
+	TUniquePtr<FPCGContext> GenerateTestDataAndRunSort(EPCGSortMethod Method, bool bRandomDensity)
+	{
+		PCGTestsCommon::FTestData TestData(PCGDeterminismTests::Defaults::Seed);
+		PCGDeterminismTests::GenerateSettings<UPCGSortPointsSettings>(TestData);
+		UPCGSortPointsSettings* Settings = CastChecked<UPCGSortPointsSettings>(TestData.Settings);
+		Settings->InputSource.SetPointProperty(EPCGPointProperties::Density);
+		Settings->SortMethod = Method;
+		FPCGElementPtr TestElement = TestData.Settings->GetElement();
+
+		FPCGTaggedData& SourcePin = TestData.InputData.TaggedData.Emplace_GetRef();
+		SourcePin.Pin = PCGPinConstants::DefaultInputLabel;
+		SourcePin.Data = PCGTestsCommon::CreateRandomPointData(100, 42, bRandomDensity);
+
+		TUniquePtr<FPCGContext> Context = TestData.InitializeTestContext();
+
+		while (!TestElement->Execute(Context.Get())) {}
+	
+		return Context;
+	}
+}
+
+IMPLEMENT_CUSTOM_SIMPLE_AUTOMATION_TEST(FPCGSortPointsTest_Basic, FPCGTestBaseClass, "pcg.tests.SortPoints.Basic", PCGTestsCommon::TestFlags)
+
+bool FPCGSortPointsTest_Basic::RunTest(const FString& Parameters)
+{
+
+	TUniquePtr<FPCGContext> Context = SortCommonTestData::GenerateTestDataAndRunSort(EPCGSortMethod::Ascending, false);
+
+	// test our point data
+	const TArray<FPCGTaggedData>& Outputs = Context->OutputData.TaggedData;
+
+	UTEST_EQUAL("Output count", Outputs.Num(), 1);
+
+	UPCGPointData* OutPointData = Cast<UPCGPointData>(Outputs[0].Data);
+
+	UTEST_NOT_NULL("Output point data", OutPointData);
+
+	const TArray<FPCGPoint>& OutPoints = OutPointData->GetPoints();
+
+	UTEST_EQUAL("Output point count", OutPoints.Num(), 100);
+
+	return true;
+}
+
+IMPLEMENT_CUSTOM_SIMPLE_AUTOMATION_TEST(FPCGSortPointsTest_Ascending, FPCGTestBaseClass, "pcg.tests.SortPoints.Ascending", PCGTestsCommon::TestFlags)
+
+bool FPCGSortPointsTest_Ascending::RunTest(const FString& Parameters)
+{
+	TUniquePtr<FPCGContext> Context = SortCommonTestData::GenerateTestDataAndRunSort(EPCGSortMethod::Ascending, true);
+
+	const TArray<FPCGTaggedData>& Outputs = Context->OutputData.TaggedData;
+
+	UTEST_EQUAL("Output count", Outputs.Num(), 1);
+
+	// compare point data
+	const FPCGTaggedData& Output = Outputs[0];
+	const UPCGPointData* OutPointData = Cast<UPCGPointData>(Output.Data);
+	if (OutPointData) 
+	{
+		const TArray<FPCGPoint>& OutPoints = OutPointData->GetPoints();
+
+		for (int i = 0; i < OutPoints.Num() - 1; ++i)
+		{
+			UTEST_TRUE(FString::Format(TEXT("{0} is less than/equal to {1}"), { OutPoints[i].Density, OutPoints[i + 1].Density }), OutPoints[i].Density <= OutPoints[i + 1].Density);
+		}
+	}
+	
+	return true;
+}
+
+IMPLEMENT_CUSTOM_SIMPLE_AUTOMATION_TEST(FPCGSortPointsTest_Descending, FPCGTestBaseClass, "pcg.tests.SortPoints.Descending", PCGTestsCommon::TestFlags)
+
+bool FPCGSortPointsTest_Descending::RunTest(const FString& Parameters)
+{
+	TUniquePtr<FPCGContext> Context = SortCommonTestData::GenerateTestDataAndRunSort(EPCGSortMethod::Descending, true);
+
+	const TArray<FPCGTaggedData>& Outputs = Context->OutputData.TaggedData;
+
+	UTEST_EQUAL("Output count", Outputs.Num(), 1);
+
+	// compare point data
+	const FPCGTaggedData& Output = Outputs[0];
+	const UPCGPointData* OutPointData = Cast<UPCGPointData>(Output.Data);
+	if (OutPointData)
+	{
+		const TArray<FPCGPoint>& OutPoints = OutPointData->GetPoints();
+
+		for (int i = 0; i < OutPoints.Num() - 1; ++i)
+		{
+			UTEST_TRUE(FString::Format(TEXT("{0} is greater than/equal to {1}"), { OutPoints[i].Density, OutPoints[i + 1].Density }), OutPoints[i].Density >= OutPoints[i + 1].Density);
+		}
+	}
+	
+	return true;
+}
+
+IMPLEMENT_CUSTOM_SIMPLE_AUTOMATION_TEST(FPCGSortPointsTest_SameValues, FPCGTestBaseClass, "pcg.tests.SortPoints.SameValues", PCGTestsCommon::TestFlags)
+
+bool FPCGSortPointsTest_SameValues::RunTest(const FString& Parameters)
+{
+	TUniquePtr<FPCGContext> Context = SortCommonTestData::GenerateTestDataAndRunSort(EPCGSortMethod::Ascending, false);
+
+	const TArray<FPCGTaggedData>& Inputs = Context->InputData.GetInputsByPin(PCGPinConstants::DefaultInputLabel);
+	const TArray<FPCGTaggedData>& Outputs = Context->OutputData.TaggedData;
+
+	UTEST_EQUAL("Output count", Outputs.Num(), 1);
+
+	// compare point data
+	const FPCGTaggedData& Input = Inputs[0];
+	const FPCGTaggedData& Output = Outputs[0];
+
+	const UPCGPointData* InPointData = Cast<UPCGPointData>(Input.Data);
+	const UPCGPointData* OutPointData = Cast<UPCGPointData>(Output.Data);
+	if (OutPointData && InPointData)
+	{
+		const TArray<FPCGPoint>& InPoints = InPointData->GetPoints();
+		const TArray<FPCGPoint>& OutPoints = OutPointData->GetPoints();
+
+		UTEST_TRUE("Arrays have the same number of points:", InPoints.Num() == OutPoints.Num());
+
+		for (int i = 0; i < InPoints.Num(); ++i)
+		{
+			//if they're in the same spots after sorting, they should have exactly the same properties across the board
+			UTEST_EQUAL(FString::Format(TEXT("UnsortedArray[{0}].Seed is equal to SortedArray[{0}].Seed"), {i}), OutPoints[i].Seed, InPoints[i].Seed);
+			UTEST_EQUAL(FString::Format(TEXT("UnsortedArray[{0}].Density is equal to SortedArray[{0}].Density"), {i}), OutPoints[i].Density, InPoints[i].Density);
+			UTEST_EQUAL(FString::Format(TEXT("UnsortedArray[{0}].Transform is equal to SortedArray[{0}].Transform"), {i}), OutPoints[i].Transform, InPoints[i].Transform);
+		}
+	}
+
+	return true;
+}
\ No newline at end of file
diff --git a/Engine/Plugins/Experimental/PCG/Source/PCG/Private/Tests/PCGTestsCommon.cpp b/Engine/Plugins/Experimental/PCG/Source/PCG/Private/Tests/PCGTestsCommon.cpp
--- a/Engine/Plugins/Experimental/PCG/Source/PCG/Private/Tests/PCGTestsCommon.cpp
+++ b/Engine/Plugins/Experimental/PCG/Source/PCG/Private/Tests/PCGTestsCommon.cpp
@@ -131,8 +131,8 @@ namespace PCGTestsCommon
 		return SinglePointData;
 	}
 
-	/** Creates a point data with PointCount many points, and randomizes the Transform and Color */
-	UPCGPointData* CreateRandomPointData(int32 PointCount, int32 Seed)
+	/** Creates a point data with PointCount many points, and randomizes the Transform, Color, and Density */
+	UPCGPointData* CreateRandomPointData(int32 PointCount, int32 Seed, bool bRandomDensity)
 	{
 		TObjectPtr<UPCGPointData> PointData = PCGTestsCommon::CreateEmptyPointData();
 		TArray<FPCGPoint>& Points = PointData->GetMutablePoints();
@@ -148,7 +148,8 @@ namespace PCGTestsCommon
 
 			FPCGPoint& Point = Points.Emplace_GetRef(FTransform(Rotation, Location, Scale), 1.f, I);
 			Point.Color = RandomSource.VRand();
-			Point.Density = 1.f;
+			Point.Density = bRandomDensity ? RandomSource.GetFraction() : 1.0f;
+			
 			Point.Seed = I;
 		}
 
diff --git a/Engine/Plugins/Experimental/PCG/Source/PCG/Public/Elements/PCGSortPoints.h b/Engine/Plugins/Experimental/PCG/Source/PCG/Public/Elements/PCGSortPoints.h
new file mode 100644
--- /dev/null
+++ b/Engine/Plugins/Experimental/PCG/Source/PCG/Public/Elements/PCGSortPoints.h
@@ -0,0 +1,52 @@
+// Copyright Epic Games, Inc. All Rights Reserved.
+
+#pragma once
+
+#include "PCGSettings.h"
+
+#include "PCGSortPoints.generated.h"
+
+/**
+ * Sorts points based on an attribute.
+ */
+
+UENUM()
+enum class EPCGSortMethod : uint8
+{
+	Ascending,
+	Descending
+};
+
+UCLASS(BlueprintType, ClassGroup = (Procedural))
+class UPCGSortPointsSettings : public UPCGSettings
+{
+	GENERATED_BODY()
+
+public:
+	//~Begin UPCGSettings interface
+#if WITH_EDITOR
+	virtual FName GetDefaultNodeName() const override { return FName(TEXT("SortPoints")); }
+	virtual FText GetDefaultNodeTitle() const override { return NSLOCTEXT("PCGSortPointsElement", "NodeTitle", "Sort Points"); }
+	virtual FText GetNodeTooltipText() const override { return NSLOCTEXT("PCGSortPointsElement", "NodeTooltip", "Sorts points based on an attribute."); }
+	virtual EPCGSettingsType GetType() const override { return EPCGSettingsType::Generic; }
+#endif
+
+protected:
+	virtual TArray<FPCGPinProperties> InputPinProperties() const override { return Super::DefaultPointInputPinProperties(); }
+	virtual TArray<FPCGPinProperties> OutputPinProperties() const override { return Super::DefaultPointOutputPinProperties(); }
+	virtual FPCGElementPtr CreateElement() const override;
+	//~End UPCGSettings interface
+
+public:
+	UPROPERTY(BlueprintReadWrite, EditAnywhere, Category = Settings, meta = (PCG_Overridable))
+	FPCGAttributePropertyInputSelector InputSource;
+
+	UPROPERTY(BlueprintReadWrite, EditAnywhere, Category = Settings, meta = (PCG_Overridable))
+	EPCGSortMethod SortMethod = EPCGSortMethod::Ascending;
+};
+
+class FPCGSortPointsElement : public FSimplePCGElement
+{
+protected:
+	virtual bool ExecuteInternal(FPCGContext* Context) const override;
+};
\ No newline at end of file
diff --git a/Engine/Plugins/Experimental/PCG/Source/PCG/Public/Metadata/Accessors/PCGAttributeAccessorHelpers.h b/Engine/Plugins/Experimental/PCG/Source/PCG/Public/Metadata/Accessors/PCGAttributeAccessorHelpers.h
--- a/Engine/Plugins/Experimental/PCG/Source/PCG/Public/Metadata/Accessors/PCGAttributeAccessorHelpers.h
+++ b/Engine/Plugins/Experimental/PCG/Source/PCG/Public/Metadata/Accessors/PCGAttributeAccessorHelpers.h
@@ -2,6 +2,9 @@
 
 #pragma once
 
+#include "Metadata/PCGMetadataAttributeTpl.h"
+#include "Metadata/PCGMetadataAttributeTraits.h"
+#include "Metadata/Accessors/IPCGAttributeAccessor.h"
 #include "Templates/UniquePtr.h"
 #include "UObject/NameTypes.h"
 
@@ -66,4 +69,109 @@ namespace PCGAttributeAccessorHelpers
 
 	PCG_API TUniquePtr<const IPCGAttributeAccessorKeys> CreateConstKeys(const UPCGData* InData, const FPCGAttributePropertySelector& InSelector);
 	PCG_API TUniquePtr<IPCGAttributeAccessorKeys> CreateKeys(UPCGData* InData, const FPCGAttributePropertySelector& InSelector);
+
+	/**
+	* Sorts array given the accessors and keys of the array
+	*/
+	template <typename T>
+	void SortByAttribute(const IPCGAttributeAccessor& InAccessor, const IPCGAttributeAccessorKeys& InKeys, TArray<T>& InArray, bool bAscending)
+	{
+		SortByAttribute(InAccessor, InKeys, InArray, bAscending, [](int Index) { return Index; });
+	}
+
+	/**
+	* Sorts array given the accessors and keys of the array.
+	* A custom function is used to extract the index of the elements of the array,
+	* and sort by the values associated with that index.
+	*/
+	template <typename T, typename Func>
+	void SortByAttribute(const IPCGAttributeAccessor& InAccessor, const IPCGAttributeAccessorKeys& InKeys, TArray<T>& InArray, bool bAscending, Func&& CustomGetIndex)
+	{
+		if (InArray.IsEmpty())
+		{
+			return;
+		} 
+
+		auto Callback = [&InAccessor, &InKeys, &InArray, bAscending, &CustomGetIndex](auto Dummy)
+		{
+			using ValueType = decltype(Dummy);
+
+			if constexpr (PCG::Private::MetadataTraits<ValueType>::CanCompare)
+			{
+				TArray<ValueType> CachedValues;
+
+				if constexpr (std::is_trivially_copyable_v<ValueType>)
+				{
+					CachedValues.SetNumUninitialized(InArray.Num());
+				}
+				else
+				{
+					CachedValues.SetNum(InArray.Num());
+				}
+
+				InAccessor.GetRange(TArrayView<ValueType>(CachedValues), 0, InKeys);
+
+				auto CompareAscending = [&CachedValues, &CustomGetIndex](int LHS, int RHS)
+				{
+					const int32 LHSIndex = CustomGetIndex(LHS);
+					const int32 RHSIndex = CustomGetIndex(RHS);
+
+					const ValueType& LHSValue = CachedValues[LHSIndex];
+					const ValueType& RHSValue = CachedValues[RHSIndex];
+
+					if (PCG::Private::MetadataTraits<ValueType>::Equal(LHSValue, RHSValue))
+					{
+						return LHSIndex < RHSIndex;
+					}
+
+					return PCG::Private::MetadataTraits<ValueType>::Less(LHSValue, RHSValue);
+				};
+
+				auto CompareDescending = [&CachedValues, &CustomGetIndex](int LHS, int RHS)
+				{
+					const int32 LHSIndex = CustomGetIndex(LHS);
+					const int32 RHSIndex = CustomGetIndex(RHS);
+
+					const ValueType& LHSValue = CachedValues[LHSIndex];
+					const ValueType& RHSValue = CachedValues[RHSIndex];
+
+					if (PCG::Private::MetadataTraits<ValueType>::Equal(LHSValue, RHSValue))
+					{
+						return LHSIndex > RHSIndex;
+					}
+
+					return PCG::Private::MetadataTraits<ValueType>::Greater(LHSValue, RHSValue);
+				};
+
+				// Fill integer sequence
+				TArray<int32> ElementsIndexes;
+				ElementsIndexes.Reserve(CachedValues.Num());
+				for (int i = 0; i < CachedValues.Num(); ++i)
+				{
+					ElementsIndexes.Add(i);
+				}
+
+				if (bAscending)
+				{
+					ElementsIndexes.Sort(CompareAscending);
+				}
+				else
+				{
+					ElementsIndexes.Sort(CompareDescending);
+				}
+
+				TArray<T> SortedArray;
+				SortedArray.Reserve(InArray.Num());
+
+				for (int i = 0; i < InArray.Num(); ++i)
+				{
+					SortedArray.Add(MoveTemp(InArray[ElementsIndexes[i]]));
+				}
+
+				InArray = MoveTemp(SortedArray);
+			}
+		};
+
+		PCGMetadataAttribute::CallbackWithRightType(InAccessor.GetUnderlyingType(), Callback);
+	}
 }
\ No newline at end of file
diff --git a/Engine/Plugins/Experimental/PCG/Source/PCG/Public/Tests/PCGTestsCommon.h b/Engine/Plugins/Experimental/PCG/Source/PCG/Public/Tests/PCGTestsCommon.h
--- a/Engine/Plugins/Experimental/PCG/Source/PCG/Public/Tests/PCGTestsCommon.h
+++ b/Engine/Plugins/Experimental/PCG/Source/PCG/Public/Tests/PCGTestsCommon.h
@@ -56,8 +56,8 @@ namespace PCGTestsCommon
 	/** Creates a PointData with a single point at the provided location */
 	PCG_API UPCGPointData* CreatePointData(const FVector& InLocation);
 
-	/** Creates a PointData with PointCount many points, and randomizes the Transform and Color */
-	PCG_API UPCGPointData* CreateRandomPointData(int32 PointCount, int32 Seed);
+	/** Creates a PointData with PointCount many points, and randomizes the Transform, Color, and Density */
+	PCG_API UPCGPointData* CreateRandomPointData(int32 PointCount, int32 Seed, bool RandomDensity = false);
 
 	PCG_API UPCGPolyLineData* CreatePolyLineData();
 	PCG_API UPCGSurfaceData* CreateSurfaceData();
diff --git a/Engine/Plugins/Experimental/PCG/Source/PCGEditor/Private/SPCGEditorGraphAttributeListView.cpp b/Engine/Plugins/Experimental/PCG/Source/PCGEditor/Private/SPCGEditorGraphAttributeListView.cpp
--- a/Engine/Plugins/Experimental/PCG/Source/PCGEditor/Private/SPCGEditorGraphAttributeListView.cpp
+++ b/Engine/Plugins/Experimental/PCG/Source/PCGEditor/Private/SPCGEditorGraphAttributeListView.cpp
@@ -1094,63 +1094,8 @@ void SPCGEditorGraphAttributeListView::RefreshSorting()
 	{
 		if (ColumnData->DataAccessor.IsValid() && ColumnData->DataKeys.IsValid())
 		{
-			auto Callback = [this, &ColumnData](auto Dummy)
-			{
-				using ValueType = decltype(Dummy);
-
-				if constexpr (PCG::Private::MetadataTraits<ValueType>::CanCompare)
-				{
-					TArray<ValueType> CachedValues;
-					// Strings need initialized values for later assignment
-					if constexpr (PCG::Private::MetadataTraits<ValueType>::NeedsConstruction)
-					{
-						CachedValues.SetNum(ListViewItems.Num());
-					}
-					else
-					{
-						CachedValues.SetNumUninitialized(ListViewItems.Num());
-					}
-
-					ColumnData->DataAccessor->GetRange(TArrayView<ValueType>(CachedValues), 0, *ColumnData->DataKeys);
-
-					auto SortAscending = [&CachedValues] (const PCGListviewItemPtr& LHS, const PCGListviewItemPtr& RHS)
-					{
-						const ValueType& LHSValue = CachedValues[LHS->Index];
-						const ValueType& RHSValue = CachedValues[RHS->Index];
-
-						if (PCG::Private::MetadataTraits<ValueType>::Equal(LHSValue, RHSValue))
-						{
-							return LHS->Index < RHS->Index;
-						}
-
-						return PCG::Private::MetadataTraits<ValueType>::Less(LHSValue, RHSValue);
-					};
-
-					auto SortDescending = [&CachedValues] (const PCGListviewItemPtr& LHS, const PCGListviewItemPtr& RHS)
-					{
-						const ValueType& LHSValue = CachedValues[LHS->Index];
-						const ValueType& RHSValue = CachedValues[RHS->Index];
-
-						if (PCG::Private::MetadataTraits<ValueType>::Equal(LHSValue, RHSValue))
-						{
-							return LHS->Index > RHS->Index;
-						}
-
-						return PCG::Private::MetadataTraits<ValueType>::Greater(LHSValue, RHSValue);
-					};
-
-					if (SortMode == EColumnSortMode::Ascending)
-					{
-						ListViewItems.Sort(SortAscending);
-					}
-					else if (SortMode == EColumnSortMode::Descending)
-					{
-						ListViewItems.Sort(SortDescending);
-					}
-				}
-			};
-
-			PCGMetadataAttribute::CallbackWithRightType(ColumnData->DataAccessor->GetUnderlyingType(), Callback);
+			//lambda used here to get the index value of an item in the array for sorting
+			PCGAttributeAccessorHelpers::SortByAttribute(*ColumnData->DataAccessor, *ColumnData->DataKeys, ListViewItems, SortMode == EColumnSortMode::Ascending, [this](int Index) { return ListViewItems[Index]->Index; });
 		}
 	}
 
