[PCG] Reducing number of errors/warning in metadata ops

* No data as input is not a warning case anymore, it can happen, especially in subgraph.
* Data but no points will have the same treatment
* Fix Default value created if someting was connected but with no data. Default value should only be created if the pin is not connected

#rnx

diff --git a/Engine/Plugins/Experimental/PCG/Source/PCG/Private/Elements/Metadata/PCGMetadataOpElementBase.cpp b/Engine/Plugins/Experimental/PCG/Source/PCG/Private/Elements/Metadata/PCGMetadataOpElementBase.cpp
--- a/Engine/Plugins/Experimental/PCG/Source/PCG/Private/Elements/Metadata/PCGMetadataOpElementBase.cpp
+++ b/Engine/Plugins/Experimental/PCG/Source/PCG/Private/Elements/Metadata/PCGMetadataOpElementBase.cpp
@@ -284,10 +284,11 @@ bool FPCGMetadataElementBase::ExecuteInternal(FPCGContext* Context) const
 	for (uint32 InputPinIndex = 0; InputPinIndex < NumberOfInputs; ++InputPinIndex)
 	{
 		FName PinLabel = Settings->GetInputPinLabel(InputPinIndex);
+		const bool bIsInputConnected = Context->Node ? Context->Node->IsInputPinConnected(PinLabel) : false;
 		TArray<FPCGTaggedData> InputData = Context->InputData.GetInputsByPin(PinLabel);
 
-		// If input data is empty but we can have default values, add a "fake" param data input.
-		if (InputData.IsEmpty() && Settings->DoesInputSupportDefaultValue(InputPinIndex))
+		// If input data is empty and not connected but we can have default values, add a "fake" param data input.
+		if (InputData.IsEmpty() && !bIsInputConnected && Settings->DoesInputSupportDefaultValue(InputPinIndex))
 		{
 			FPCGTaggedData& DummyData = InputData.Emplace_GetRef();
 			DummyData.Pin = PinLabel;
@@ -302,7 +303,7 @@ bool FPCGMetadataElementBase::ExecuteInternal(FPCGContext* Context) const
 		if (InputData.IsEmpty())
 		{
 			// Visually warn the user, since this is causing execution to be aborted
-			PCGE_LOG(Warning, GraphAndLog, FText::Format(LOCTEXT("MissingInputDataForPin", "No data provided on pin {0}"), FText::FromName(PinLabel)));
+			PCGE_LOG(Verbose, LogOnly, FText::Format(LOCTEXT("MissingInputDataForPin", "No data provided on pin {0}"), FText::FromName(PinLabel)));
 			return true;
 		}
 		else if (InputData.Num() > 1)
@@ -314,6 +315,16 @@ bool FPCGMetadataElementBase::ExecuteInternal(FPCGContext* Context) const
 		// By construction, there can only be one of then(hence the 0 index)
 		InputTaggedData[InputPinIndex] = MoveTemp(InputData[0]);
 
+		// Check if we have any points
+		if (const UPCGPointData* PointInput = Cast<const UPCGPointData>(InputTaggedData[InputPinIndex].Data))
+		{
+			if (PointInput->GetPoints().IsEmpty())
+			{
+				PCGE_LOG(Verbose, LogOnly, FText::Format(LOCTEXT("NoPointsForPin", "No points in point data provided on pin {0}"), FText::FromName(PinLabel)));
+				return true;
+			}
+		}
+
 		// Only gather Spacial and Params input. 
 		if (const UPCGSpatialData* SpatialInput = Cast<const UPCGSpatialData>(InputTaggedData[InputPinIndex].Data))
 		{
