[PCG] Profiling view polish pass

* Expanded columns so contents are (less) clipped by default
* Add header tooltips
* Normalized capitalisation and time units
* Numerical columns sort Descending by default (biggest costs at top)
* Removed partial implementation of PostExecute column which elements don't really use right now
* Add ellipsis overflow

#rnx

diff --git a/Engine/Plugins/Experimental/PCG/Source/PCGEditor/Private/SPCGEditorGraphProfilingView.cpp b/Engine/Plugins/Experimental/PCG/Source/PCGEditor/Private/SPCGEditorGraphProfilingView.cpp
--- a/Engine/Plugins/Experimental/PCG/Source/PCGEditor/Private/SPCGEditorGraphProfilingView.cpp
+++ b/Engine/Plugins/Experimental/PCG/Source/PCGEditor/Private/SPCGEditorGraphProfilingView.cpp
@@ -26,16 +26,21 @@ namespace PCGEditorGraphProfilingView
 	const FName NAME_MaxExecutionFrameTime = FName(TEXT("MaxFrameTime"));
 	const FName NAME_TotalExecutionTime = FName(TEXT("TotalExecutionTime"));
 	const FName NAME_NbExecutionFrames = FName(TEXT("NbExecutionFrames"));
-	const FName NAME_PostExecuteTime = FName(TEXT("PostExecuteTime"));
 
 	/** Labels of the columns */
 	const FText TEXT_NodeLabel = LOCTEXT("NodeLabel", "Node");
-	const FText TEXT_PrepareDataTimeLabel = LOCTEXT("PrepareDataTimeLabel", "PrepareData (ms)");
-	const FText TEXT_PostExecuteTimeLabel = LOCTEXT("PostExecuteTimeLabel", "PostExecute (ms)");
-	const FText TEXT_MinExecutionFrameTimeLabel = LOCTEXT("MinExecutionFrameTimeLabel", "Min Frame Time(ms)");
-	const FText TEXT_MaxExecutionFrameTimeLabel = LOCTEXT("MaxExecutionFrameTimeLabel", "Max Frame Time(ms)");
-	const FText TEXT_TotalExecutionTimeLabel = LOCTEXT("TotalExecutionTimeLabel", "Total time(s)");
-	const FText TEXT_NbExecutionFramesLabel = LOCTEXT("NbExecutionFramesLabel", "Exec frames");
+	const FText TEXT_PrepareDataTimeLabel = LOCTEXT("PrepareDataTimeLabel", "PrepareData Phase (ms)");
+	const FText TEXT_MinExecutionFrameTimeLabel = LOCTEXT("MinExecutionFrameTimeLabel", "Min Frame Time (ms)");
+	const FText TEXT_MaxExecutionFrameTimeLabel = LOCTEXT("MaxExecutionFrameTimeLabel", "Max Frame Time (ms)");
+	const FText TEXT_TotalExecutionTimeLabel = LOCTEXT("TotalExecutionTimeLabel", "Total Time (ms)");
+	const FText TEXT_NbExecutionFramesLabel = LOCTEXT("NbExecutionFramesLabel", "Exec Frames");
+
+	/** Tooltips */
+	const FText TEXT_PrepareDataTimeTooltip = LOCTEXT("PrepareDataTimeTooltip", "Cost of the PrepareData execution phase which some nodes use to process the incoming data.");
+	const FText TEXT_MinExecutionFrameTimeTooltip = LOCTEXT("MinExecutionFrameTimeTooltip", "The minimum time spent of all execution frames.");
+	const FText TEXT_MaxExecutionFrameTimeTooltip = LOCTEXT("MaxExecutionFrameTimeTooltip", "The maximum time spent of all execution frames.");
+	const FText TEXT_TotalExecutionTimeTooltip = LOCTEXT("TotalExecutionTimeTooltip", "The total time spent, summed over all execution frames and phases.");
+	const FText TEXT_NbExecutionFramesTooltip = LOCTEXT("NbExecutionFramesTooltip", "The number of frames in which one or more execution phases were executed.");
 }
 
 void SPCGProfilingListViewItemRow::Construct(const FArguments& InArgs, const TSharedRef<STableViewBase>& InOwnerTableView, const PCGProfilingListViewItemPtr& Item)
@@ -61,7 +66,7 @@ TSharedRef<SWidget> SPCGProfilingListViewItemRow::GenerateWidgetForColumn(const
 		{
 			ColumnData = ((InternalItem->CallTime.ExecutionFrameCount >= 0) ? FText::AsNumber(InternalItem->CallTime.ExecutionFrameCount) : FText());
 		}
-		else if (!InternalItem->HasData)
+		else if (!InternalItem->bHasData)
 		{
 			ColumnData = PCGEditorGraphProfilingView::NoDataAvailableText;
 		}
@@ -77,20 +82,18 @@ TSharedRef<SWidget> SPCGProfilingListViewItemRow::GenerateWidgetForColumn(const
 		}
 		else if (ColumnId == PCGEditorGraphProfilingView::NAME_TotalExecutionTime)
 		{
-			// In s
-			ColumnData = FText::AsNumber(InternalItem->CallTime.ExecutionTime);
+			// In ms
+			ColumnData = FText::AsNumber(InternalItem->CallTime.ExecutionTime * 1000.0);
 		}
 		else if (ColumnId == PCGEditorGraphProfilingView::NAME_PrepareDataTime)
 		{
 			ColumnData = FText::AsNumber(InternalItem->CallTime.PrepareDataTime * 1000.0);
 		}
-		else if (ColumnId == PCGEditorGraphProfilingView::NAME_PostExecuteTime)
-		{
-			ColumnData = FText::AsNumber(InternalItem->CallTime.PostExecuteTime * 1000.0);
-		}
 	}
 
-	return SNew(STextBlock).Text(ColumnData);
+	return SNew(STextBlock)
+		.Text(ColumnData)
+		.OverflowPolicy(ETextOverflowPolicy::Ellipsis);
 }
 
 void SPCGEditorGraphProfilingView::OnItemDoubleClicked(PCGProfilingListViewItemPtr Item)
@@ -160,6 +163,7 @@ void SPCGEditorGraphProfilingView::Construct(const FArguments& InArgs, TSharedPt
 			SNew(SHorizontalBox)
 			+SHorizontalBox::Slot()
 			.AutoWidth()
+			.Padding(1.0f, 1.0f)
 			[
 				SNew(SButton)
 				.Text(LOCTEXT("ResetButton", "Reset"))
@@ -167,25 +171,31 @@ void SPCGEditorGraphProfilingView::Construct(const FArguments& InArgs, TSharedPt
 			]
 			+SHorizontalBox::Slot()
 			.AutoWidth()
+			.VAlign(VAlign_Center)
+			.Padding(1.0f, 0.0f)
 			[
 				SNew(STextBlock)
 				.Text(LOCTEXT("ExpandSubgraph", "Expand Subgraph"))
 			]
 			+SHorizontalBox::Slot()
 			.AutoWidth()
+			.VAlign(VAlign_Center)
+			.Padding(1.0f, 0.0f)
 			[
 				SNew(SCheckBox)
 				.IsChecked(this, &SPCGEditorGraphProfilingView::IsSubgraphExpanded)
 				.OnCheckStateChanged(this, &SPCGEditorGraphProfilingView::OnSubgraphExpandedChanged)
 			]
 			+SHorizontalBox::Slot()
-			.Padding(FMargin(30.f, 0.f, 0.f, 5.f))
+			.Padding(FMargin(30.0f, 0.0f, 2.0f, 0.0f))
+			.VAlign(EVerticalAlignment::VAlign_Center)
 			.AutoWidth()
 			[
 				SNew(STextBlock)
 				.Text(LOCTEXT("TotalExecutionTime", "Total Time:"))
 			]
 			+SHorizontalBox::Slot()
+			.VAlign(EVerticalAlignment::VAlign_Center)
 			.AutoWidth()
 			[
 				SNew(STextBlock)
@@ -234,47 +244,57 @@ TSharedRef<SHeaderRow> SPCGEditorGraphProfilingView::CreateHeaderRowWidget()
 		.ResizeMode(ESplitterResizeMode::FixedPosition)
 		.CanSelectGeneratedColumn(true)
 		+ SHeaderRow::Column(PCGEditorGraphProfilingView::NAME_Node)
-		.ManualWidth(150)
+		.ManualWidth(350.0f)
 		.DefaultLabel(PCGEditorGraphProfilingView::TEXT_NodeLabel)
 		.HAlignHeader(HAlign_Center)
 		.HAlignCell(HAlign_Left)
 		.SortMode(this, &SPCGEditorGraphProfilingView::GetColumnSortMode, PCGEditorGraphProfilingView::NAME_Node)
 		.OnSort(this, &SPCGEditorGraphProfilingView::OnSortColumnHeader)
 		+ SHeaderRow::Column(PCGEditorGraphProfilingView::NAME_PrepareDataTime)
-		.ManualWidth(125)
+		.ManualWidth(170.0f)
 		.DefaultLabel(PCGEditorGraphProfilingView::TEXT_PrepareDataTimeLabel)
 		.HAlignHeader(HAlign_Center)
 		.HAlignCell(HAlign_Right)
 		.SortMode(this, &SPCGEditorGraphProfilingView::GetColumnSortMode, PCGEditorGraphProfilingView::NAME_PrepareDataTime)
 		.OnSort(this, &SPCGEditorGraphProfilingView::OnSortColumnHeader)
+		.InitialSortMode(EColumnSortMode::Descending)
+		.DefaultTooltip(PCGEditorGraphProfilingView::TEXT_PrepareDataTimeTooltip)
 		+ SHeaderRow::Column(PCGEditorGraphProfilingView::NAME_NbExecutionFrames)
-		.ManualWidth(80)
+		.ManualWidth(105.0f)
 		.DefaultLabel(PCGEditorGraphProfilingView::TEXT_NbExecutionFramesLabel)
 		.HAlignHeader(HAlign_Center)
 		.HAlignCell(HAlign_Right)
 		.SortMode(this, &SPCGEditorGraphProfilingView::GetColumnSortMode, PCGEditorGraphProfilingView::NAME_NbExecutionFrames)
 		.OnSort(this, &SPCGEditorGraphProfilingView::OnSortColumnHeader)
+		.InitialSortMode(EColumnSortMode::Descending)
+		.DefaultTooltip(PCGEditorGraphProfilingView::TEXT_NbExecutionFramesTooltip)
 		+ SHeaderRow::Column(PCGEditorGraphProfilingView::NAME_MinExecutionFrameTime)
-		.ManualWidth(130)
+		.ManualWidth(155.0f)
 		.DefaultLabel(PCGEditorGraphProfilingView::TEXT_MinExecutionFrameTimeLabel)
 		.HAlignHeader(HAlign_Center)
 		.HAlignCell(HAlign_Right)
 		.SortMode(this, &SPCGEditorGraphProfilingView::GetColumnSortMode, PCGEditorGraphProfilingView::NAME_MinExecutionFrameTime)
 		.OnSort(this, &SPCGEditorGraphProfilingView::OnSortColumnHeader)
+		.InitialSortMode(EColumnSortMode::Descending)
+		.DefaultTooltip(PCGEditorGraphProfilingView::TEXT_MinExecutionFrameTimeTooltip)
 		+ SHeaderRow::Column(PCGEditorGraphProfilingView::NAME_MaxExecutionFrameTime)
-		.ManualWidth(130)
+		.ManualWidth(155.0f)
 		.DefaultLabel(PCGEditorGraphProfilingView::TEXT_MaxExecutionFrameTimeLabel)
 		.HAlignHeader(HAlign_Center)
 		.HAlignCell(HAlign_Right)
 		.SortMode(this, &SPCGEditorGraphProfilingView::GetColumnSortMode, PCGEditorGraphProfilingView::NAME_MaxExecutionFrameTime)
 		.OnSort(this, &SPCGEditorGraphProfilingView::OnSortColumnHeader)
+		.InitialSortMode(EColumnSortMode::Descending)
+		.DefaultTooltip(PCGEditorGraphProfilingView::TEXT_MaxExecutionFrameTimeTooltip)
 		+ SHeaderRow::Column(PCGEditorGraphProfilingView::NAME_TotalExecutionTime)
-		.ManualWidth(100)
+		.ManualWidth(120.0f)
 		.DefaultLabel(PCGEditorGraphProfilingView::TEXT_TotalExecutionTimeLabel)
 		.HAlignHeader(HAlign_Center)
 		.HAlignCell(HAlign_Right)
 		.SortMode(this, &SPCGEditorGraphProfilingView::GetColumnSortMode, PCGEditorGraphProfilingView::NAME_TotalExecutionTime)
-		.OnSort(this, &SPCGEditorGraphProfilingView::OnSortColumnHeader);
+		.OnSort(this, &SPCGEditorGraphProfilingView::OnSortColumnHeader)
+		.InitialSortMode(EColumnSortMode::Descending)
+		.DefaultTooltip(PCGEditorGraphProfilingView::TEXT_TotalExecutionTimeTooltip);
 }
 
 FText SPCGEditorGraphProfilingView::GetTotalTimeLabel() const
@@ -287,7 +307,7 @@ void SPCGEditorGraphProfilingView::OnSortColumnHeader(const EColumnSortPriority:
 	if (SortingColumn == ColumnId)
 	{
 		// Circling
-		SortMode = EColumnSortMode::Type((SortMode + 1) % 3);
+		SortMode = static_cast<EColumnSortMode::Type>((SortMode + 1) % 3);
 	}
 	else
 	{
@@ -348,7 +368,7 @@ namespace PCGEditorGraphProfilingView
 		const TArray<PCGUtils::FCallTreeInfo>& TreeInfo,
 		const TMap<const UPCGNode*, UPCGEditorGraphNode*>& EditorNodeLookup,
 		bool bExpandSubgraph,
-		FString FolderName,
+		const FString& FolderName,
 		UPCGEditorGraphNode* CurrentEditorNode = nullptr)
 	{
 		for (const PCGUtils::FCallTreeInfo& Info : TreeInfo)
@@ -370,7 +390,7 @@ namespace PCGEditorGraphProfilingView
 				ListViewItem->PCGNode = Info.Node;
 				ListViewItem->EditorNode = EditorNode;
 				ListViewItem->Name = Fullname;
-				ListViewItem->HasData = true;
+				ListViewItem->bHasData = true;
 
 				ListViewItem->CallTime = Info.CallTime;
 
@@ -383,7 +403,7 @@ namespace PCGEditorGraphProfilingView
 			}
 		}
 	}
-} 
+}
 
 FReply SPCGEditorGraphProfilingView::Refresh()
 {
diff --git a/Engine/Plugins/Experimental/PCG/Source/PCGEditor/Private/SPCGEditorGraphProfilingView.h b/Engine/Plugins/Experimental/PCG/Source/PCGEditor/Private/SPCGEditorGraphProfilingView.h
--- a/Engine/Plugins/Experimental/PCG/Source/PCGEditor/Private/SPCGEditorGraphProfilingView.h
+++ b/Engine/Plugins/Experimental/PCG/Source/PCGEditor/Private/SPCGEditorGraphProfilingView.h
@@ -23,7 +23,7 @@ struct FPCGProfilingListViewItem
 
 	PCGUtils::FCallTime CallTime;
 
-	bool HasData = false;
+	bool bHasData = false;
 };
 
 typedef TSharedPtr<FPCGProfilingListViewItem> PCGProfilingListViewItemPtr;
