PCG: do not use Cast with TObjectPtr as a const_cast
#rnx

diff --git a/Engine/Plugins/Experimental/PCG/Source/PCG/Private/Elements/Metadata/PCGMetadataOpElementBase.cpp b/Engine/Plugins/Experimental/PCG/Source/PCG/Private/Elements/Metadata/PCGMetadataOpElementBase.cpp
--- a/Engine/Plugins/Experimental/PCG/Source/PCG/Private/Elements/Metadata/PCGMetadataOpElementBase.cpp
+++ b/Engine/Plugins/Experimental/PCG/Source/PCG/Private/Elements/Metadata/PCGMetadataOpElementBase.cpp
@@ -516,14 +516,14 @@ bool FPCGMetadataElementBase::ExecuteInternal(FPCGContext* Context) const
 				PCGMetadataElementCommon::CopyEntryToValueKeyMap(SourceMetadata[InputPinToForward], SourceAttribute[InputPinToForward], OutputAttribute);
 			}
 
-			OperationData.OutputAccessors[OutputIndex] = PCGAttributeAccessorHelpers::CreateAccessor(Cast<UPCGData>(OutputData.Data), OutputTarget);
+			OperationData.OutputAccessors[OutputIndex] = PCGAttributeAccessorHelpers::CreateAccessor( const_cast<UPCGData*>(OutputData.Data.Get()), OutputTarget);
 		}
 		else
 		{
 			// In case of property or attribute with extra accessor, we need to validate that the property/attribute can accept the output type.
 			// Verify this before duplicating, because an extra allocation is certainly less costly than duplicating the data.
 			// Do it with a const accessor, since OutputData.Data is still pointing on the const input data.
-			TUniquePtr<const IPCGAttributeAccessor> TempConstAccessor = PCGAttributeAccessorHelpers::CreateConstAccessor(Cast<const UPCGData>(OutputData.Data), OutputTarget);
+			TUniquePtr<const IPCGAttributeAccessor> TempConstAccessor = PCGAttributeAccessorHelpers::CreateConstAccessor(OutputData.Data.Get(), OutputTarget);
 
 			if (TempConstAccessor.IsValid())
 			{
@@ -540,7 +540,7 @@ bool FPCGMetadataElementBase::ExecuteInternal(FPCGContext* Context) const
 				PCGMetadataElementCommon::DuplicateTaggedData(InputTaggedData[InputPinToForward], OutputData, OutMetadata);
 
 				// Re-create the accessor to point to the right data (since we just duplicated the data)
-				OperationData.OutputAccessors[OutputIndex] = PCGAttributeAccessorHelpers::CreateAccessor(Cast<UPCGData>(OutputData.Data), OutputTarget);
+				OperationData.OutputAccessors[OutputIndex] = PCGAttributeAccessorHelpers::CreateAccessor(const_cast<UPCGData*>(OutputData.Data.Get()), OutputTarget);
 			}
 		}
 
@@ -555,7 +555,7 @@ bool FPCGMetadataElementBase::ExecuteInternal(FPCGContext* Context) const
 			return false;
 		}
 
-		OperationData.OutputKeys[OutputIndex] = PCGAttributeAccessorHelpers::CreateKeys(Cast<UPCGData>(OutputData.Data), OutputTarget);
+		OperationData.OutputKeys[OutputIndex] = PCGAttributeAccessorHelpers::CreateKeys(const_cast<UPCGData*>(OutputData.Data.Get()), OutputTarget);
 
 		return OperationData.OutputKeys[OutputIndex].IsValid();
 	};
diff --git a/Engine/Plugins/Experimental/PCG/Source/PCG/Private/Elements/PCGAttributeNoise.cpp b/Engine/Plugins/Experimental/PCG/Source/PCG/Private/Elements/PCGAttributeNoise.cpp
--- a/Engine/Plugins/Experimental/PCG/Source/PCG/Private/Elements/PCGAttributeNoise.cpp
+++ b/Engine/Plugins/Experimental/PCG/Source/PCG/Private/Elements/PCGAttributeNoise.cpp
@@ -296,7 +296,7 @@ bool FPCGAttributeNoiseElement::ExecuteInternal(FPCGContext* InContext) const
 		}
 		else
 		{
-			OutputPoints = &CastChecked<UPCGPointData>(Outputs[CurrentInput].Data)->GetMutablePoints();
+			OutputPoints = &(const_cast<UPCGPointData*>(CastChecked<UPCGPointData>(Outputs[CurrentInput].Data)))->GetMutablePoints();
 			// Note: for deprecation purposes, we consider that the input here is a spatial data (even though the pin is typed to be a Point Data)
 			// hence the need to call ToPointData, as otherwise a direct case to UPCGPointData would fail.
 			InputPoints = &(CastChecked<const UPCGSpatialData>(Inputs[CurrentInput].Data)->ToPointData(Context))->GetPoints();
diff --git a/Engine/Plugins/Experimental/PCG/Source/PCG/Private/Elements/PCGCreateAttribute.cpp b/Engine/Plugins/Experimental/PCG/Source/PCG/Private/Elements/PCGCreateAttribute.cpp
--- a/Engine/Plugins/Experimental/PCG/Source/PCG/Private/Elements/PCGCreateAttribute.cpp
+++ b/Engine/Plugins/Experimental/PCG/Source/PCG/Private/Elements/PCGCreateAttribute.cpp
@@ -265,7 +265,7 @@ bool FPCGCreateAttributeElement::ExecuteInternal(FPCGContext* Context) const
 	check(Settings);
 
 	TArray<FPCGTaggedData> SourceParams = Context->InputData.GetInputsByPin(PCGCreateAttributeConstants::AttributesLabel);
-	UPCGParamData* SourceParamData = nullptr;
+	const UPCGParamData* SourceParamData = nullptr;
 	FName SourceParamAttributeName = NAME_None;
 	FName OutputAttributeName = NAME_None;
 
diff --git a/Engine/Plugins/Experimental/PCG/Source/PCG/Private/Elements/PCGOuterIntersectionElement.cpp b/Engine/Plugins/Experimental/PCG/Source/PCG/Private/Elements/PCGOuterIntersectionElement.cpp
--- a/Engine/Plugins/Experimental/PCG/Source/PCG/Private/Elements/PCGOuterIntersectionElement.cpp
+++ b/Engine/Plugins/Experimental/PCG/Source/PCG/Private/Elements/PCGOuterIntersectionElement.cpp
@@ -140,7 +140,7 @@ bool FPCGOuterIntersectionElement::ExecuteInternal(FPCGContext* Context) const
 	// Iterate through primary inputs, and create intersection data from the unioned
 	for (const FPCGTaggedData& Input : Context->InputData.GetInputsByPin(PCGIntersectionConstants::PrimaryLabel))
 	{
-		UPCGSpatialData* PrimarySpatialData = Cast<UPCGSpatialData>(Input.Data);
+		const UPCGSpatialData* PrimarySpatialData = Cast<UPCGSpatialData>(Input.Data);
 
 		// If this data isn't spatial, passthrough to the output
 		if (!PrimarySpatialData)
@@ -149,7 +149,7 @@ bool FPCGOuterIntersectionElement::ExecuteInternal(FPCGContext* Context) const
 			continue;
 		}
 		
-		UPCGSpatialData* IntersectionData = PrimarySpatialData;
+		const UPCGSpatialData* IntersectionData = PrimarySpatialData;
 		for (const UPCGSpatialData* SecondarySourceUnionData : SecondarySourceUnionArray)
 		{
 			UPCGIntersectionData* TempIntersectionData = IntersectionData->IntersectWith(SecondarySourceUnionData);
diff --git a/Engine/Plugins/Experimental/PCG/Source/PCG/Private/Elements/PCGProjectionElement.cpp b/Engine/Plugins/Experimental/PCG/Source/PCG/Private/Elements/PCGProjectionElement.cpp
--- a/Engine/Plugins/Experimental/PCG/Source/PCG/Private/Elements/PCGProjectionElement.cpp
+++ b/Engine/Plugins/Experimental/PCG/Source/PCG/Private/Elements/PCGProjectionElement.cpp
@@ -52,7 +52,7 @@ bool FPCGProjectionElement::ExecuteInternal(FPCGContext* Context) const
 	}
 
 	// Ensure we have spatial data to project onto
-	UPCGSpatialData* ProjectionTarget = Cast<UPCGSpatialData>(Targets[0].Data);
+	const UPCGSpatialData* ProjectionTarget = Cast<UPCGSpatialData>(Targets[0].Data);
 	if (!ProjectionTarget)
 	{
 		return true;
@@ -70,7 +70,7 @@ bool FPCGProjectionElement::ExecuteInternal(FPCGContext* Context) const
 
 	for (FPCGTaggedData& Source : Sources)
 	{
-		UPCGSpatialData* ProjectionSource = Cast<UPCGSpatialData>(Source.Data);
+		const UPCGSpatialData* ProjectionSource = Cast<UPCGSpatialData>(Source.Data);
 
 		if (!ProjectionSource)
 		{
diff --git a/Engine/Plugins/Experimental/PCG/Source/PCG/Private/Elements/PCGSurfaceSampler.cpp b/Engine/Plugins/Experimental/PCG/Source/PCG/Private/Elements/PCGSurfaceSampler.cpp
--- a/Engine/Plugins/Experimental/PCG/Source/PCG/Private/Elements/PCGSurfaceSampler.cpp
+++ b/Engine/Plugins/Experimental/PCG/Source/PCG/Private/Elements/PCGSurfaceSampler.cpp
@@ -521,7 +521,7 @@ bool FPCGSurfaceSamplerElement::ExecuteInternal(FPCGContext* InContext) const
 
 	while (Context->CurrentGeneratingShape < Context->GeneratingShapes.Num())
 	{
-		UPCGPointData* PointData = Cast<UPCGPointData>(Outputs[Context->CurrentGeneratingShape].Data);
+		UPCGPointData* PointData = const_cast<UPCGPointData*>(Cast<UPCGPointData>(Outputs[Context->CurrentGeneratingShape].Data));
 
 		if (PointData)
 		{
diff --git a/Engine/Plugins/Experimental/PCG/Source/PCG/Private/Elements/PCGVolumeSampler.cpp b/Engine/Plugins/Experimental/PCG/Source/PCG/Private/Elements/PCGVolumeSampler.cpp
--- a/Engine/Plugins/Experimental/PCG/Source/PCG/Private/Elements/PCGVolumeSampler.cpp
+++ b/Engine/Plugins/Experimental/PCG/Source/PCG/Private/Elements/PCGVolumeSampler.cpp
@@ -218,7 +218,7 @@ bool FPCGVolumeSamplerElement::ExecuteInternal(FPCGContext* Context) const
 	TArray<const UPCGSpatialData*, TInlineAllocator<16>> GeneratingShapes;
 	for (FPCGTaggedData& TaggedData : VolumeInputs)
 	{
-		if (UPCGSpatialData* SpatialData = Cast<UPCGSpatialData>(TaggedData.Data))
+		if (const UPCGSpatialData* SpatialData = Cast<UPCGSpatialData>(TaggedData.Data))
 		{
 			GeneratingShapes.Add(SpatialData);
 			Outputs.Add(TaggedData);
diff --git a/Engine/Plugins/Experimental/PCG/Source/PCG/Private/Metadata/Accessors/PCGAttributeAccessorHelpers.cpp b/Engine/Plugins/Experimental/PCG/Source/PCG/Private/Metadata/Accessors/PCGAttributeAccessorHelpers.cpp
--- a/Engine/Plugins/Experimental/PCG/Source/PCG/Private/Metadata/Accessors/PCGAttributeAccessorHelpers.cpp
+++ b/Engine/Plugins/Experimental/PCG/Source/PCG/Private/Metadata/Accessors/PCGAttributeAccessorHelpers.cpp
@@ -436,7 +436,7 @@ TUniquePtr<const IPCGAttributeAccessor> PCGAttributeAccessorHelpers::CreateConst
 {
 	bool bFromGlobalParamsPin = false;
 	TArray<FPCGTaggedData> InputParamData = InInputData.GetParamsByPin(InParam.Label);
-	UPCGParamData* ParamData = !InputParamData.IsEmpty() ? CastChecked<UPCGParamData>(InputParamData[0].Data) : nullptr;
+	const UPCGParamData* ParamData = !InputParamData.IsEmpty() ? CastChecked<UPCGParamData>(InputParamData[0].Data) : nullptr;
 
 	// If it is empty, try with the Overrides pin (Global Params)
 	if (!ParamData)
diff --git a/Engine/Plugins/Experimental/PCG/Source/PCG/Private/PCGComponent.cpp b/Engine/Plugins/Experimental/PCG/Source/PCG/Private/PCGComponent.cpp
--- a/Engine/Plugins/Experimental/PCG/Source/PCG/Private/PCGComponent.cpp
+++ b/Engine/Plugins/Experimental/PCG/Source/PCG/Private/PCGComponent.cpp
@@ -381,11 +381,11 @@ void UPCGComponent::PostProcessGraph(const FBox& InNewBounds, bool bInGenerated,
 				DuplicatedTaggedData.Data = Cast<UPCGData>(StaticDuplicateObject(TaggedData.Data, this));
 
 				UPCGMetadata* DuplicatedMetadata = nullptr;
-				if (UPCGSpatialData* DuplicatedSpatialData = Cast<UPCGSpatialData>(DuplicatedTaggedData.Data))
+				if (const UPCGSpatialData* DuplicatedSpatialData = Cast<UPCGSpatialData>(DuplicatedTaggedData.Data))
 				{
 					DuplicatedMetadata = DuplicatedSpatialData->Metadata;
 				}
-				else if (UPCGParamData* DuplicatedParamData = Cast<UPCGParamData>(DuplicatedTaggedData.Data))
+				else if (const UPCGParamData* DuplicatedParamData = Cast<UPCGParamData>(DuplicatedTaggedData.Data))
 				{
 					DuplicatedMetadata = DuplicatedParamData->Metadata;
 				}
@@ -1855,7 +1855,7 @@ UPCGData* UPCGComponent::CreateActorPCGData(AActor* Actor, const UPCGComponent*
 	}
 	else if(Collection.TaggedData.Num() == 1)
 	{
-		return Cast<UPCGData>(Collection.TaggedData[0].Data);
+		return const_cast<UPCGData*>(Collection.TaggedData[0].Data.Get());
 	}
 	else
 	{
diff --git a/Engine/Plugins/Experimental/PCG/Source/PCG/Private/PCGData.cpp b/Engine/Plugins/Experimental/PCG/Source/PCG/Private/PCGData.cpp
--- a/Engine/Plugins/Experimental/PCG/Source/PCG/Private/PCGData.cpp
+++ b/Engine/Plugins/Experimental/PCG/Source/PCG/Private/PCGData.cpp
@@ -303,10 +303,10 @@ UPCGParamData* FPCGDataCollection::GetParamsWithDeprecation(const UPCGNode* Node
 
 			for (const FPCGTaggedData& TaggedDatum : GetParamsByPin(InputPin->Properties.Label))
 			{
-				if (UPCGParamData* Params = Cast<UPCGParamData>(TaggedDatum.Data))
+				if (const UPCGParamData* Params = Cast<UPCGParamData>(TaggedDatum.Data))
 				{
 					UE_LOG(LogPCG, Warning, TEXT("[%s] Found an Attribute Set data on an input pin that should not accept attributes. Make sure to re-wire it to the Overrides pin if it is used for overrides."), *Node->GetNodeTitle().ToString());
-					return Params;
+					return const_cast<UPCGParamData*>(Params);
 				}
 			}
 		}
@@ -319,9 +319,9 @@ UPCGParamData* FPCGDataCollection::GetParams() const
 {
 	for (const FPCGTaggedData& TaggedDatum : TaggedData)
 	{
-		if (UPCGParamData* Params = Cast<UPCGParamData>(TaggedDatum.Data))
+		if (const UPCGParamData* Params = Cast<UPCGParamData>(TaggedDatum.Data))
 		{
-			return Params; 
+			return const_cast<UPCGParamData*>(Params); 
 		}
 	}
 
@@ -331,7 +331,7 @@ UPCGParamData* FPCGDataCollection::GetParams() const
 UPCGParamData* FPCGDataCollection::GetFirstParamsOnParamsPin() const
 {
 	TArray<FPCGTaggedData> ParamsOnDefaultPin = GetParamsByPin(PCGPinConstants::DefaultParamsLabel);
-	return (ParamsOnDefaultPin.IsEmpty() ? nullptr : Cast<UPCGParamData>(ParamsOnDefaultPin[0].Data));
+	return (ParamsOnDefaultPin.IsEmpty() ? nullptr : const_cast<UPCGParamData*>(Cast<UPCGParamData>(ParamsOnDefaultPin[0].Data)));
 }
 
 void FPCGDataCollection::GetResourceSizeEx(FResourceSizeEx& CumulativeResourceSize) const
@@ -382,7 +382,7 @@ const UPCGSettingsInterface* FPCGDataCollection::GetSettingsInterface(const UPCG
 	else
 	{
 		const FPCGTaggedData* MatchingData = TaggedData.FindByPredicate([InDefaultSettingsInterface](const FPCGTaggedData& Data) {
-			if (UPCGSettingsInterface* DataSettingsInterface = Cast<UPCGSettingsInterface>(Data.Data))
+			if (const UPCGSettingsInterface* DataSettingsInterface = Cast<UPCGSettingsInterface>(Data.Data))
 			{
 				// Compare settings classes
 				return DataSettingsInterface->GetSettings()->GetClass() == InDefaultSettingsInterface->GetSettings()->GetClass() ||
@@ -426,8 +426,7 @@ void FPCGDataCollection::AddToRootSet(FPCGRootSet& RootSet) const
 	{
 		if (Data.Data)
 		{
-			// This is technically a const_cast
-			RootSet.Add(Cast<UObject>(Data.Data));
+			RootSet.Add(const_cast<UPCGData*>(Data.Data.Get()));
 		}
 	}
 }
@@ -438,8 +437,7 @@ void FPCGDataCollection::RemoveFromRootSet(FPCGRootSet& RootSet) const
 	{
 		if (Data.Data)
 		{
-			// This is technically a const_cast
-			RootSet.Remove(Cast<UObject>(Data.Data));
+			RootSet.Remove(const_cast<UPCGData*>(Data.Data.Get()));
 		}
 	}
 }
diff --git a/Engine/Plugins/Experimental/PCG/Source/PCG/Private/PCGElement.cpp b/Engine/Plugins/Experimental/PCG/Source/PCG/Private/PCGElement.cpp
--- a/Engine/Plugins/Experimental/PCG/Source/PCG/Private/PCGElement.cpp
+++ b/Engine/Plugins/Experimental/PCG/Source/PCG/Private/PCGElement.cpp
@@ -456,7 +456,7 @@ void IPCGElement::CleanupAndValidateOutput(FPCGContext* Context) const
 
 				if (CVarPCGValidatePointMetadata.GetValueOnAnyThread())
 				{
-					if (UPCGPointData* PointData = Cast<UPCGPointData>(TaggedData.Data))
+					if (const UPCGPointData* PointData = Cast<UPCGPointData>(TaggedData.Data))
 					{
 						const TArray<FPCGPoint>& NewPoints = PointData->GetPoints();
 						const int32 MaxMetadataEntry = PointData->Metadata ? PointData->Metadata->GetItemCountForChild() : 0;
diff --git a/Engine/Plugins/Experimental/PCG/Source/PCG/Private/Tests/Determinism/PCGDeterminismTestsCommon.cpp b/Engine/Plugins/Experimental/PCG/Source/PCG/Private/Tests/Determinism/PCGDeterminismTestsCommon.cpp
--- a/Engine/Plugins/Experimental/PCG/Source/PCG/Private/Tests/Determinism/PCGDeterminismTestsCommon.cpp
+++ b/Engine/Plugins/Experimental/PCG/Source/PCG/Private/Tests/Determinism/PCGDeterminismTestsCommon.cpp
@@ -1354,7 +1354,7 @@ namespace PCGDeterminismTests
 		{
 			if (DataCanBeShuffled(TaggedData.Data))
 			{
-				UPCGPointData* PointData = CastChecked<UPCGPointData>(TaggedData.Data);
+				UPCGPointData* PointData = const_cast<UPCGPointData*>(CastChecked<UPCGPointData>(TaggedData.Data));
 				ShuffleArray<FPCGPoint>(PointData->GetMutablePoints(), TestData.RandomStream);
 			}
 		}
diff --git a/Engine/Plugins/Experimental/PCG/Source/PCG/Private/Tests/Elements/PCGDistanceTest.cpp b/Engine/Plugins/Experimental/PCG/Source/PCG/Private/Tests/Elements/PCGDistanceTest.cpp
--- a/Engine/Plugins/Experimental/PCG/Source/PCG/Private/Tests/Elements/PCGDistanceTest.cpp
+++ b/Engine/Plugins/Experimental/PCG/Source/PCG/Private/Tests/Elements/PCGDistanceTest.cpp
@@ -48,7 +48,7 @@ bool FPCGDistanceTest_PointToPoint::RunTest(const FString& Parameters)
 
 	UTEST_EQUAL("Output count", Outputs.Num(), 1);
 
-	UPCGPointData *OutPointData = Cast<UPCGPointData>(Outputs[0].Data);
+	const UPCGPointData *OutPointData = Cast<UPCGPointData>(Outputs[0].Data);
 
 	UTEST_NOT_NULL("Output point data", OutPointData);
 
@@ -108,7 +108,7 @@ bool FPCGDistanceTest_SetDensity::RunTest(const FString& Parameters)
 
 	UTEST_EQUAL("Output count", Outputs.Num(), 1);
 
-	UPCGPointData *OutPointData = Cast<UPCGPointData>(Outputs[0].Data);
+	const UPCGPointData *OutPointData = Cast<UPCGPointData>(Outputs[0].Data);
 
 	UTEST_NOT_NULL("Output point data", OutPointData);
 
@@ -170,7 +170,7 @@ bool FPCGDistanceTest_PointToSphere::RunTest(const FString& Parameters)
 
 	UTEST_EQUAL("Output count", Outputs.Num(), 1);
 
-	UPCGPointData *OutPointData = Cast<UPCGPointData>(Outputs[0].Data);
+	const UPCGPointData *OutPointData = Cast<UPCGPointData>(Outputs[0].Data);
 
 	UTEST_NOT_NULL("Output point data", OutPointData);
 
@@ -235,7 +235,7 @@ bool FPCGDistanceTest_PointToBox::RunTest(const FString& Parameters)
 
 	UTEST_EQUAL("Output count", Outputs.Num(), 1);
 
-	UPCGPointData *OutPointData = Cast<UPCGPointData>(Outputs[0].Data);
+	const UPCGPointData *OutPointData = Cast<UPCGPointData>(Outputs[0].Data);
 
 	UTEST_NOT_NULL("Output point data", OutPointData);
 
diff --git a/Engine/Plugins/Experimental/PCG/Source/PCG/Private/Tests/Elements/PCGGatherTest.cpp b/Engine/Plugins/Experimental/PCG/Source/PCG/Private/Tests/Elements/PCGGatherTest.cpp
--- a/Engine/Plugins/Experimental/PCG/Source/PCG/Private/Tests/Elements/PCGGatherTest.cpp
+++ b/Engine/Plugins/Experimental/PCG/Source/PCG/Private/Tests/Elements/PCGGatherTest.cpp
@@ -33,7 +33,7 @@ bool FPCGGatherTest_Basic::RunTest(const FString& Parameters)
 
 	UTEST_EQUAL("Output count", Outputs.Num(), 1);
 
-	UPCGPointData *OutPointData = Cast<UPCGPointData>(Outputs[0].Data);
+	const UPCGPointData *OutPointData = Cast<UPCGPointData>(Outputs[0].Data);
 
 	UTEST_NOT_NULL("Output point data", OutPointData);
 
diff --git a/Engine/Plugins/Experimental/PCG/Source/PCG/Private/Tests/Elements/PCGNormalToDensityTest.cpp b/Engine/Plugins/Experimental/PCG/Source/PCG/Private/Tests/Elements/PCGNormalToDensityTest.cpp
--- a/Engine/Plugins/Experimental/PCG/Source/PCG/Private/Tests/Elements/PCGNormalToDensityTest.cpp
+++ b/Engine/Plugins/Experimental/PCG/Source/PCG/Private/Tests/Elements/PCGNormalToDensityTest.cpp
@@ -41,7 +41,7 @@ bool FPCGNormalToDensityTest_Set::RunTest(const FString& Parameters)
 
 	UTEST_EQUAL("Output count", Outputs.Num(), 1);
 
-	UPCGSpatialData* OutputData = Cast<UPCGSpatialData>(Outputs[0].Data);
+	const UPCGSpatialData* OutputData = Cast<UPCGSpatialData>(Outputs[0].Data);
 
 	UTEST_NOT_NULL("Output data", OutputData);
 
@@ -88,7 +88,7 @@ bool FPCGNormalToDensityTest_CustomNormal::RunTest(const FString& Parameters)
 
 	UTEST_EQUAL("Output count", Outputs.Num(), 1);
 
-	UPCGSpatialData* OutputData = Cast<UPCGSpatialData>(Outputs[0].Data);
+	const UPCGSpatialData* OutputData = Cast<UPCGSpatialData>(Outputs[0].Data);
 
 	UTEST_NOT_NULL("Output data", OutputData);
 
@@ -135,7 +135,7 @@ bool FPCGNormalToDensityTest_Strength::RunTest(const FString& Parameters)
 
 	UTEST_EQUAL("Output count", Outputs.Num(), 1);
 
-	UPCGSpatialData* OutputData = Cast<UPCGSpatialData>(Outputs[0].Data);
+	const UPCGSpatialData* OutputData = Cast<UPCGSpatialData>(Outputs[0].Data);
 
 	UTEST_NOT_NULL("Output data", OutputData);
 
@@ -183,7 +183,7 @@ bool FPCGNormalToDensityTest_Minimum::RunTest(const FString& Parameters)
 
 	UTEST_EQUAL("Output count", Outputs.Num(), 1);
 
-	UPCGSpatialData* OutputData = Cast<UPCGSpatialData>(Outputs[0].Data);
+	const UPCGSpatialData* OutputData = Cast<UPCGSpatialData>(Outputs[0].Data);
 
 	UTEST_NOT_NULL("Output data", OutputData);
 
@@ -231,7 +231,7 @@ bool FPCGNormalToDensityTest_Maximum::RunTest(const FString& Parameters)
 
 	UTEST_EQUAL("Output count", Outputs.Num(), 1);
 
-	UPCGSpatialData* OutputData = Cast<UPCGSpatialData>(Outputs[0].Data);
+	const UPCGSpatialData* OutputData = Cast<UPCGSpatialData>(Outputs[0].Data);
 
 	UTEST_NOT_NULL("Output data", OutputData);
 
@@ -280,7 +280,7 @@ bool FPCGNormalToDensityTest_Add::RunTest(const FString& Parameters)
 
 	UTEST_EQUAL("Output count", Outputs.Num(), 1);
 
-	UPCGSpatialData* OutputData = Cast<UPCGSpatialData>(Outputs[0].Data);
+	const UPCGSpatialData* OutputData = Cast<UPCGSpatialData>(Outputs[0].Data);
 
 	UTEST_NOT_NULL("Output data", OutputData);
 
@@ -328,7 +328,7 @@ bool FPCGNormalToDensityTest_Subtract::RunTest(const FString& Parameters)
 
 	UTEST_EQUAL("Output count", Outputs.Num(), 1);
 
-	UPCGSpatialData* OutputData = Cast<UPCGSpatialData>(Outputs[0].Data);
+	const UPCGSpatialData* OutputData = Cast<UPCGSpatialData>(Outputs[0].Data);
 
 	UTEST_NOT_NULL("Output data", OutputData);
 
@@ -376,7 +376,7 @@ bool FPCGNormalToDensityTest_Multiply::RunTest(const FString& Parameters)
 
 	UTEST_EQUAL("Output count", Outputs.Num(), 1);
 
-	UPCGSpatialData* OutputData = Cast<UPCGSpatialData>(Outputs[0].Data);
+	const UPCGSpatialData* OutputData = Cast<UPCGSpatialData>(Outputs[0].Data);
 
 	UTEST_NOT_NULL("Output data", OutputData);
 
@@ -424,7 +424,7 @@ bool FPCGNormalToDensityTest_Divide::RunTest(const FString& Parameters)
 
 	UTEST_EQUAL("Output count", Outputs.Num(), 1);
 
-	UPCGSpatialData* OutputData = Cast<UPCGSpatialData>(Outputs[0].Data);
+	const UPCGSpatialData* OutputData = Cast<UPCGSpatialData>(Outputs[0].Data);
 
 	UTEST_NOT_NULL("Output data", OutputData);
 
diff --git a/Engine/Plugins/Experimental/PCG/Source/PCG/Private/Tests/Elements/PCGPropertyToParamDataElementTest.cpp b/Engine/Plugins/Experimental/PCG/Source/PCG/Private/Tests/Elements/PCGPropertyToParamDataElementTest.cpp
--- a/Engine/Plugins/Experimental/PCG/Source/PCG/Private/Tests/Elements/PCGPropertyToParamDataElementTest.cpp
+++ b/Engine/Plugins/Experimental/PCG/Source/PCG/Private/Tests/Elements/PCGPropertyToParamDataElementTest.cpp
@@ -70,7 +70,7 @@ bool VerifyAttributeValue(FPCGTestBaseClass* TestInstance, PCGTestsCommon::FTest
 
 		if (bSuccess)
 		{
-			UPCGParamData* ParamData = Cast<UPCGParamData>(Context->OutputData.GetAllParams()[0].Data);
+			const UPCGParamData* ParamData = Cast<UPCGParamData>(Context->OutputData.GetAllParams()[0].Data);
 			const FPCGMetadataAttributeBase* Attribute = ParamData->ConstMetadata()->GetConstAttribute(Settings->PropertyName);
 
 			bSuccess = TestInstance->TestEqual(FormatWithPropertyName("AttributeType"), Attribute->GetTypeId(), PCG::Private::MetadataTypes<RawAttributeType>::Id);
diff --git a/Engine/Plugins/Experimental/PCG/Source/PCG/Private/Tests/Elements/PCGSpatialNoiseTest.cpp b/Engine/Plugins/Experimental/PCG/Source/PCG/Private/Tests/Elements/PCGSpatialNoiseTest.cpp
--- a/Engine/Plugins/Experimental/PCG/Source/PCG/Private/Tests/Elements/PCGSpatialNoiseTest.cpp
+++ b/Engine/Plugins/Experimental/PCG/Source/PCG/Private/Tests/Elements/PCGSpatialNoiseTest.cpp
@@ -66,7 +66,7 @@ bool FPCGSpatialNoise_Perlin2D::RunTest(const FString& Parameters)
 
 	UTEST_EQUAL("Output count", Outputs.Num(), 1);
 
-	UPCGPointData* OutputData = Cast<UPCGPointData>(Outputs[0].Data);
+	const UPCGPointData* OutputData = Cast<UPCGPointData>(Outputs[0].Data);
 
 	UTEST_NOT_NULL("Output data", OutputData);
 	UTEST_EQUAL("Output point data count", OutputData->GetPoints().Num(), 100);
@@ -104,7 +104,7 @@ bool FPCGSpatialNoise_Caustic2D::RunTest(const FString& Parameters)
 
 	UTEST_EQUAL("Output count", Outputs.Num(), 1);
 
-	UPCGPointData* OutputData = Cast<UPCGPointData>(Outputs[0].Data);
+	const UPCGPointData* OutputData = Cast<UPCGPointData>(Outputs[0].Data);
 
 	UTEST_NOT_NULL("Output data", OutputData);
 	UTEST_EQUAL("Output point data count", OutputData->GetPoints().Num(), 100);
@@ -143,7 +143,7 @@ bool FPCGSpatialNoise_Voronoi2D::RunTest(const FString& Parameters)
 
 	UTEST_EQUAL("Output count", Outputs.Num(), 1);
 
-	UPCGPointData* OutputData = Cast<UPCGPointData>(Outputs[0].Data);
+	const UPCGPointData* OutputData = Cast<UPCGPointData>(Outputs[0].Data);
 
 	UTEST_NOT_NULL("Output data", OutputData);
 	UTEST_EQUAL("Output point data count", OutputData->GetPoints().Num(), 100);
@@ -185,7 +185,7 @@ bool FPCGSpatialNoise_FractionalBrownian2D::RunTest(const FString& Parameters)
 
 	UTEST_EQUAL("Output count", Outputs.Num(), 1);
 
-	UPCGPointData* OutputData = Cast<UPCGPointData>(Outputs[0].Data);
+	const UPCGPointData* OutputData = Cast<UPCGPointData>(Outputs[0].Data);
 
 	UTEST_NOT_NULL("Output data", OutputData);
 	UTEST_EQUAL("Output point data count", OutputData->GetPoints().Num(), 100);
@@ -223,7 +223,7 @@ bool FPCGSpatialNoise_EdgeMask2D::RunTest(const FString& Parameters)
 
 	UTEST_EQUAL("Output count", Outputs.Num(), 1);
 
-	UPCGPointData* OutputData = Cast<UPCGPointData>(Outputs[0].Data);
+	const UPCGPointData* OutputData = Cast<UPCGPointData>(Outputs[0].Data);
 
 	UTEST_NOT_NULL("Output data", OutputData);
 	UTEST_EQUAL("Output point data count", OutputData->GetPoints().Num(), 100);
@@ -261,7 +261,7 @@ bool FPCGSpatialNoise_TilingPerlin2D::RunTest(const FString& Parameters)
 
 	UTEST_EQUAL("Output count", Outputs.Num(), 1);
 
-	UPCGPointData* OutputData = Cast<UPCGPointData>(Outputs[0].Data);
+	const UPCGPointData* OutputData = Cast<UPCGPointData>(Outputs[0].Data);
 
 	UTEST_NOT_NULL("Output data", OutputData);
 	UTEST_EQUAL("Output point data count", OutputData->GetPoints().Num(), 100);
@@ -300,7 +300,7 @@ bool FPCGSpatialNoise_TilingVoronoi2D::RunTest(const FString& Parameters)
 
 	UTEST_EQUAL("Output count", Outputs.Num(), 1);
 
-	UPCGPointData* OutputData = Cast<UPCGPointData>(Outputs[0].Data);
+	const UPCGPointData* OutputData = Cast<UPCGPointData>(Outputs[0].Data);
 
 	UTEST_NOT_NULL("Output data", OutputData);
 	UTEST_EQUAL("Output point data count", OutputData->GetPoints().Num(), 100);
diff --git a/Engine/Plugins/Experimental/PCG/Source/PCG/Private/Tests/Elements/PCGStaticMeshSpawnerTest.cpp b/Engine/Plugins/Experimental/PCG/Source/PCG/Private/Tests/Elements/PCGStaticMeshSpawnerTest.cpp
--- a/Engine/Plugins/Experimental/PCG/Source/PCG/Private/Tests/Elements/PCGStaticMeshSpawnerTest.cpp
+++ b/Engine/Plugins/Experimental/PCG/Source/PCG/Private/Tests/Elements/PCGStaticMeshSpawnerTest.cpp
@@ -48,13 +48,13 @@ namespace
 }
 
 /** Finds the output UPCGPointData and builds a mapping from world-space position to each FPCGPoint */
-void GetOutputDataAndBuildPositionMap(FPCGTestBaseClass* Test, const TArray<FPCGTaggedData>& OutputData, TMap<FVector, FPCGPoint>& OutLocationToPoint, UPCGPointData*& OutPointData)
+void GetOutputDataAndBuildPositionMap(FPCGTestBaseClass* Test, const TArray<FPCGTaggedData>& OutputData, TMap<FVector, FPCGPoint>& OutLocationToPoint, const UPCGPointData*& OutPointData)
 {
 	check(Test);
 
 	if (Test->TestEqual("Valid number of data elements in output", OutputData.Num(), 2))
 	{
-		TObjectPtr<UPCGSettings> OutSettings = Cast<UPCGSettings>(OutputData[0].Data);
+		TObjectPtr<const UPCGSettings> OutSettings = Cast<UPCGSettings>(OutputData[0].Data);
 		Test->TestTrue("Valid output settings", OutSettings != nullptr);
 
 		OutPointData = Cast<UPCGPointData>(OutputData[1].Data);
@@ -145,7 +145,7 @@ void TestMeshSelectorByAttribute(
 	{}
 
 	TMap<FVector, FPCGPoint> LocationToPoint;
-	UPCGPointData* OutPointData = nullptr;
+	const UPCGPointData* OutPointData = nullptr;
 
 	if (bValidateOutput)
 	{
@@ -268,7 +268,7 @@ void TestMeshSelectorWeighted(
 	{}
 
 	TMap<FVector, FPCGPoint> LocationToPoint;
-	UPCGPointData* OutPointData = nullptr;
+	const UPCGPointData* OutPointData = nullptr;
 
 	if (bValidateOutput)
 	{
@@ -356,7 +356,7 @@ void TestMeshSelectorWeightedByCategory(
 	}
 
 	TMap<FVector, FPCGPoint> LocationToPoint;
-	UPCGPointData* OutPointData = nullptr;
+	const UPCGPointData* OutPointData = nullptr;
 
 	if (bValidateOutput)
 	{
