[PCG] Cull tasks that won't generate on the target execution grid.

This allowed to remove all higen considerations from the graph executor. Execution code is restored to original version.

Task culling routine implemented in a generic way which will support culling base on other factors such as static branches.

#rnx

diff --git a/Engine/Plugins/Experimental/PCG/Source/PCG/Private/Elements/PCGIndirectionElement.cpp b/Engine/Plugins/Experimental/PCG/Source/PCG/Private/Elements/PCGIndirectionElement.cpp
--- a/Engine/Plugins/Experimental/PCG/Source/PCG/Private/Elements/PCGIndirectionElement.cpp
+++ b/Engine/Plugins/Experimental/PCG/Source/PCG/Private/Elements/PCGIndirectionElement.cpp
@@ -121,7 +121,6 @@ bool FPCGIndirectionElement::PrepareDataInternal(FPCGContext* InContext) const
 		Context->InnerContext->TaskId = Context->TaskId;
 		Context->InnerContext->CompiledTaskId = Context->CompiledTaskId;
 		Context->InnerContext->DependenciesCrc = Context->DependenciesCrc;
-		Context->InnerContext->GenerationGrid = Context->GenerationGrid;
 
 		Context->InnerContext->AsyncState = Context->AsyncState;
 	}
diff --git a/Engine/Plugins/Experimental/PCG/Source/PCG/Private/Graph/PCGGraphCompiler.cpp b/Engine/Plugins/Experimental/PCG/Source/PCG/Private/Graph/PCGGraphCompiler.cpp
--- a/Engine/Plugins/Experimental/PCG/Source/PCG/Private/Graph/PCGGraphCompiler.cpp
+++ b/Engine/Plugins/Experimental/PCG/Source/PCG/Private/Graph/PCGGraphCompiler.cpp
@@ -306,29 +306,53 @@ TArray<FPCGGraphTask> FPCGGraphCompiler::GetPrecompiledTasks(UPCGGraph* InGraph,
 	return ExistingTasks ? *ExistingTasks : TArray<FPCGGraphTask>();
 }
 
-void FPCGGraphCompiler::ResolveGridSizes(TArray<FPCGGraphTask>& InOutCompiledTasks, const FPCGStackContext& InStackContext, EPCGHiGenGrid GenerationDefaultGrid) const
+void FPCGGraphCompiler::ResolveGridSizes(
+	EPCGHiGenGrid GenerationGrid,
+	const TArray<FPCGGraphTask>& CompiledTasks,
+	const FPCGStackContext& StackContext,
+	EPCGHiGenGrid GenerationDefaultGrid,
+	TArray<EPCGHiGenGrid>& InOutTaskGenerationGrid)
 {
-	// The stack is used to form the ResourceKey - a string that provides a path to the data from top graph down to specific pin.
-	// This will be used by link tasks as store/retrieve keys to marshal data for edges that cross grid size boundaries.
-	const FPCGStack* CurrentStack = InStackContext.GetStack(InStackContext.GetCurrentStackIndex());
-	if (InOutCompiledTasks.Num() == 0 || !ensure(CurrentStack))
+	TRACE_CPUPROFILER_EVENT_SCOPE(FPCGGraphCompiler::ResolveGridSizes);
+
+	if (CompiledTasks.IsEmpty())
 	{
 		return;
 	}
+	check(!InOutTaskGenerationGrid.IsEmpty());
 
-	// Calculate execution grid values for each task.
-	for (FPCGGraphTask& Task : InOutCompiledTasks)
+	// Special case - input node must always be present. Execute it on the current grid.
+	InOutTaskGenerationGrid[0] = GenerationGrid;
+
+	// Calculate execution grid values for subsequent tasks.
+	for (int32 i = 1; i < CompiledTasks.Num(); ++i)
 	{
-		CalculateGridRecursive(Task.NodeId, GenerationDefaultGrid, InStackContext, InOutCompiledTasks);
+		CalculateGridRecursive(CompiledTasks[i].NodeId, GenerationDefaultGrid, StackContext, CompiledTasks, InOutTaskGenerationGrid);
 	}
+}
 
+void FPCGGraphCompiler::CreateGridLinkages(
+	EPCGHiGenGrid InGenerationGrid,
+	TArray<EPCGHiGenGrid>& InOutTaskGenerationGrid,
+	TArray<FPCGGraphTask>& InOutCompiledTasks,
+	const FPCGStackContext& InStackContext)
+{
 	// Now add link tasks - if a Grid256 task depends on data from a Grid512 task, inject a link
 	// task that looks up the Grid512 component, schedules its execution if it does not have data, and
 	// then uses its output data.
+	// 
+	// The stack is used to form the ResourceKey - a string that provides a path to the data from top graph down to specific pin.
+	// This will be used by link tasks as store/retrieve keys to marshal data for edges that cross grid size boundaries.
+	const FPCGStack* CurrentStack = InStackContext.GetStack(InStackContext.GetCurrentStackIndex());
+	if (InOutCompiledTasks.IsEmpty() || !ensure(CurrentStack))
+	{
+		return;
+	}
+
 	const int32 NumCompiledTasksBefore = InOutCompiledTasks.Num();
 	for (FPCGTaskId TaskId = 0; TaskId < NumCompiledTasksBefore; ++TaskId)
 	{
-		const EPCGHiGenGrid GraphGenerationGrid = InOutCompiledTasks[TaskId].GraphGenerationGrid;
+		const EPCGHiGenGrid GraphGenerationGrid = InOutTaskGenerationGrid[TaskId];
 		for (FPCGGraphTaskInput& TaskInput : InOutCompiledTasks[TaskId].Inputs)
 		{
 			if (!TaskInput.InPin)
@@ -337,7 +361,7 @@ void FPCGGraphCompiler::ResolveGridSizes(TArray<FPCGGraphTask>& InOutCompiledTas
 				continue;
 			}
 
-			const EPCGHiGenGrid InputGraphGenerationGrid = InOutCompiledTasks[TaskInput.TaskId].GraphGenerationGrid;
+			const EPCGHiGenGrid InputGraphGenerationGrid = InOutTaskGenerationGrid[TaskInput.TaskId];
 			// Register linkage task if grid sizes don't match - either way! This allows us to generate an execution-time error if going from
 			// small grid to large grid.
 			if (InputGraphGenerationGrid != EPCGHiGenGrid::Uninitialized && InputGraphGenerationGrid > GraphGenerationGrid)
@@ -356,14 +380,21 @@ void FPCGGraphCompiler::ResolveGridSizes(TArray<FPCGGraphTask>& InOutCompiledTas
 
 				LinkTask.Inputs.Emplace(TaskInput.TaskId, TaskInput.InPin, nullptr, /*bConsumeInputData=*/true);
 
-				const EPCGHiGenGrid FromGrid = InOutCompiledTasks[TaskInput.TaskId].GraphGenerationGrid;
-				const EPCGHiGenGrid ToGrid = InOutCompiledTasks[TaskId].GraphGenerationGrid;
+				const EPCGHiGenGrid FromGrid = InOutTaskGenerationGrid[TaskInput.TaskId];
+				const EPCGHiGenGrid ToGrid = InOutTaskGenerationGrid[TaskId];
 
 				// This lambda runs at execution time and attempts to retrieve the data from a larger grid. Capture by value is intentional.
 				auto GridLinkageOperation = [FromGrid, ToGrid, ResourceKey, OutputPinLabel = TaskInput.InPin->Properties.Label,
-					DownstreamNode = InOutCompiledTasks[TaskId].Node](FPCGContext* InContext)
+					DownstreamNode = InOutCompiledTasks[TaskId].Node, InGenerationGrid](FPCGContext* InContext)
 				{
-					return PCGGraphExecutor::ExecuteGridLinkage(FromGrid, ToGrid, ResourceKey, OutputPinLabel, DownstreamNode, static_cast<FPCGGridLinkageContext*>(InContext));
+					return PCGGraphExecutor::ExecuteGridLinkage(
+						InGenerationGrid,
+						FromGrid,
+						ToGrid,
+						ResourceKey,
+						OutputPinLabel,
+						DownstreamNode,
+						static_cast<FPCGGridLinkageContext*>(InContext));
 				};
 				FPCGGenericElement::FContextAllocator ContextAllocator = [](const FPCGDataCollection&, TWeakObjectPtr<UPCGComponent>, const UPCGNode*)
 				{
@@ -375,7 +406,7 @@ void FPCGGraphCompiler::ResolveGridSizes(TArray<FPCGGraphTask>& InOutCompiledTas
 				TaskInput.TaskId = LinkTask.NodeId;
 
 				// The link needs to execute at both FROM grid size (store) and TO grid size (retrieve).
-				LinkTask.GraphGenerationGrid = FromGrid | ToGrid;
+				InOutTaskGenerationGrid.Add(FromGrid | ToGrid);
 			}
 		}
 	}
@@ -385,17 +416,18 @@ EPCGHiGenGrid FPCGGraphCompiler::CalculateGridRecursive(
 	FPCGTaskId InTaskId,
 	EPCGHiGenGrid GenerationDefaultGrid,
 	const FPCGStackContext& InStackContext,
-	TArray<FPCGGraphTask>& InOutCompiledTasks) const
+	const TArray<FPCGGraphTask>& InCompiledTasks,
+	TArray<EPCGHiGenGrid>& InOutTaskGenerationGrid)
 {
-	if (InOutCompiledTasks[InTaskId].GraphGenerationGrid != EPCGHiGenGrid::Uninitialized)
+	if (InOutTaskGenerationGrid[InTaskId] != EPCGHiGenGrid::Uninitialized)
 	{
-		return InOutCompiledTasks[InTaskId].GraphGenerationGrid;
+		return InOutTaskGenerationGrid[InTaskId];
 	}
 
 	// Default for outside any grid size range.
 	EPCGHiGenGrid Grid = GenerationDefaultGrid;
 
-	const UPCGNode* Node = InOutCompiledTasks[InTaskId].Node;
+	const UPCGNode* Node = InCompiledTasks[InTaskId].Node;
 	const UPCGSettings* Settings = Node ? Node->GetSettings() : nullptr;
 	const UPCGHiGenGridSizeSettings* Gate = Cast<UPCGHiGenGridSizeSettings>(Settings);
 	if (Gate && Gate->bEnabled)
@@ -405,9 +437,9 @@ EPCGHiGenGrid FPCGGraphCompiler::CalculateGridRecursive(
 	else
 	{
 		// Grid of this task is minimum of all input grids. We can link in data from a larger grid, but not from a finer grid (this goes against hierarchy).
-		for (FPCGGraphTaskInput InputTask : InOutCompiledTasks[InTaskId].Inputs)
+		for (FPCGGraphTaskInput InputTask : InCompiledTasks[InTaskId].Inputs)
 		{
-			const EPCGHiGenGrid InputGrid = CalculateGridRecursive(InputTask.TaskId, GenerationDefaultGrid, InStackContext, InOutCompiledTasks);
+			const EPCGHiGenGrid InputGrid = CalculateGridRecursive(InputTask.TaskId, GenerationDefaultGrid, InStackContext, InCompiledTasks, InOutTaskGenerationGrid);
 			if (PCGHiGenGrid::IsValidGrid(InputGrid))
 			{
 				Grid = FMath::Min(InputGrid, Grid);
@@ -415,11 +447,71 @@ EPCGHiGenGrid FPCGGraphCompiler::CalculateGridRecursive(
 		}
 	}
 
-	InOutCompiledTasks[InTaskId].GraphGenerationGrid = Grid;
+	InOutTaskGenerationGrid[InTaskId] = Grid;
 
 	return Grid;
 }
 
+void FPCGGraphCompiler::CullTasks(TArray<FPCGGraphTask>& InOutCompiledTasks, TFunctionRef<bool(const FPCGGraphTask&)> CullTask)
+{
+	TRACE_CPUPROFILER_EVENT_SCOPE(FPCGGraphCompiler::CullTasks);
+
+	// Check that we have more than just the input task to work with.
+	if (InOutCompiledTasks.Num() < 2)
+	{
+		return;
+	}
+
+	TArray<int32> TaskRemapping;
+	TaskRemapping.SetNumUninitialized(InOutCompiledTasks.Num());
+
+	// Never cull first (Input) task.
+	int32 WriteIndex = 1;
+	int32 ReadIndex = 1;
+	TaskRemapping[0] = 0;
+
+	while (ReadIndex < InOutCompiledTasks.Num())
+	{
+		if (!CullTask(InOutCompiledTasks[ReadIndex]))
+		{
+			if (WriteIndex != ReadIndex)
+			{
+				InOutCompiledTasks[WriteIndex] = MoveTemp(InOutCompiledTasks[ReadIndex]);
+				InOutCompiledTasks[WriteIndex].NodeId = WriteIndex;
+			}
+
+			TaskRemapping[ReadIndex] = WriteIndex;
+			++WriteIndex;
+		}
+		else
+		{
+			// Flag as culled.
+			TaskRemapping[ReadIndex] = INDEX_NONE;
+		}
+
+		++ReadIndex;
+	}
+
+	InOutCompiledTasks.SetNum(WriteIndex);
+
+	for (FPCGGraphTask& Task : InOutCompiledTasks)
+	{
+		for (int32 InputIndex = Task.Inputs.Num() - 1; InputIndex >= 0; --InputIndex)
+		{
+			const int32 InputTaskId = Task.Inputs[InputIndex].TaskId;
+			const int32 Remap = TaskRemapping[InputTaskId];
+			if (Remap != INDEX_NONE)
+			{
+				Task.Inputs[InputIndex].TaskId = Remap;
+			}
+			else
+			{
+				Task.Inputs.RemoveAt(InputIndex);
+			}
+		}
+	}
+}
+
 TArray<FPCGGraphTask> FPCGGraphCompiler::GetCompiledTasks(UPCGGraph* InGraph, uint32 GenerationGridSize, FPCGStackContext& OutStackContext, bool bIsTopGraph)
 {
 	TArray<FPCGGraphTask> CompiledTasks;
@@ -500,20 +592,26 @@ void FPCGGraphCompiler::CompileTopGraph(UPCGGraph* InGraph, uint32 GenerationGri
 		return;
 	}
 
-	if (GenerationGridSize != PCGHiGenGrid::UninitializedGridSize())
+	// For hierarchical generation resolve the execution grid for each task and cull any tasks that won't execute.
+	if (InGraph->IsHierarchicalGenerationEnabled() && GenerationGridSize != PCGHiGenGrid::UninitializedGridSize())
 	{
-		if (InGraph->IsHierarchicalGenerationEnabled())
-		{
-			const EPCGHiGenGrid DefaultGrid = PCGHiGenGrid::GridSizeToGrid(InGraph->GetDefaultGridSize());
-			ResolveGridSizes(CompiledTasks, StackContext, DefaultGrid);
-		}
-
 		const EPCGHiGenGrid GenerationGrid = PCGHiGenGrid::GridSizeToGrid(GenerationGridSize);
-		for (FPCGGraphTask& Task : CompiledTasks)
+		const EPCGHiGenGrid DefaultGrid = PCGHiGenGrid::GridSizeToGrid(InGraph->GetDefaultGridSize());
+
+		// Propagate grid size nodes through the graph to determine which grid size each task should execute on.
+		TArray<EPCGHiGenGrid> TaskGenerationGrid;
+		TaskGenerationGrid.SetNumZeroed(CompiledTasks.Num());
+		ResolveGridSizes(GenerationGrid, CompiledTasks, StackContext, DefaultGrid, TaskGenerationGrid);
+
+		// Create linkage tasks for edges that cross from large grid to small grid tasks.
+		CreateGridLinkages(GenerationGrid, TaskGenerationGrid, CompiledTasks, StackContext);
+
+		// Cull any task that should not execute on the current grid.
+		CullTasks(CompiledTasks, [GenerationGrid, &TaskGenerationGrid](const FPCGGraphTask& InTask)
 		{
-			// Make a copy of this because many things can be scheduled over many frames
-			Task.GenerationGrid = GenerationGrid;
-		}
+			const EPCGHiGenGrid TaskGrid = TaskGenerationGrid[InTask.NodeId];
+			return TaskGrid != EPCGHiGenGrid::Uninitialized && !(TaskGrid & GenerationGrid);
+		});
 	}
 
 	const int TaskNum = CompiledTasks.Num();
diff --git a/Engine/Plugins/Experimental/PCG/Source/PCG/Private/Graph/PCGGraphCompiler.h b/Engine/Plugins/Experimental/PCG/Source/PCG/Private/Graph/PCGGraphCompiler.h
--- a/Engine/Plugins/Experimental/PCG/Source/PCG/Private/Graph/PCGGraphCompiler.h
+++ b/Engine/Plugins/Experimental/PCG/Source/PCG/Private/Graph/PCGGraphCompiler.h
@@ -21,9 +21,6 @@ public:
 
 	static void OffsetNodeIds(TArray<FPCGGraphTask>& Tasks, FPCGTaskId Offset, FPCGTaskId ParentId);
 
-	/** Propagates grid sizes through a graph's compiled tasks. */
-	void ResolveGridSizes(TArray<FPCGGraphTask>& InOutCompiledTasks, const FPCGStackContext& InStackContext, EPCGHiGenGrid GenerationDefaultGrid) const;
-
 #if WITH_EDITOR
 	void NotifyGraphChanged(UPCGGraph* InGraph);
 #endif
@@ -35,8 +32,31 @@ private:
 	TArray<FPCGGraphTask> CompileGraph(UPCGGraph* InGraph, FPCGTaskId& NextId, FPCGStackContext& InOutStackContext);
 	void CompileTopGraph(UPCGGraph* InGraph, uint32 GenerationGridSize);
 
-	/** Returns the grid for InCompiledTaskId. */
-	EPCGHiGenGrid CalculateGridRecursive(FPCGTaskId InTaskId, EPCGHiGenGrid GenerationDefaultGrid, const FPCGStackContext& InStackContext, TArray<FPCGGraphTask>& InOutCompiledTasks) const;
+	/** Propagates grid sizes through a graph's compiled tasks. */
+	static void ResolveGridSizes(
+		EPCGHiGenGrid GenerationGrid,
+		const TArray<FPCGGraphTask>& CompiledTasks,
+		const FPCGStackContext& StackContext,
+		EPCGHiGenGrid GenerationDefaultGrid,
+		TArray<EPCGHiGenGrid>& InOutTaskGenerationGrid);
+
+	/** Returns the execution grid for the given task. */
+	static EPCGHiGenGrid CalculateGridRecursive(
+		FPCGTaskId InTaskId,
+		EPCGHiGenGrid GenerationDefaultGrid,
+		const FPCGStackContext& InStackContext,
+		const TArray<FPCGGraphTask>& InCompiledTasks,
+		TArray<EPCGHiGenGrid>& InOutTaskGenerationGrid);
+
+	/** Create linkage tasks for edges that cross from large grid to small grid tasks. */
+	static void CreateGridLinkages(
+		EPCGHiGenGrid InGenerationGrid,
+		TArray<EPCGHiGenGrid>& TaskGenerationGrid,
+		TArray<FPCGGraphTask>& InOutCompiledTasks,
+		const FPCGStackContext& InStackContext);
+
+	/** Culls tasks based on a given lambda. Never culls the first (input) task in the array. */
+	static void CullTasks(TArray<FPCGGraphTask>& InOutCompiledTasks, TFunctionRef<bool(const FPCGGraphTask&)> CullTask);
 
 	mutable FRWLock GraphToTaskMapLock;
 	TMap<UPCGGraph*, TArray<FPCGGraphTask>> GraphToTaskMap;
diff --git a/Engine/Plugins/Experimental/PCG/Source/PCG/Private/Graph/PCGGraphExecutor.cpp b/Engine/Plugins/Experimental/PCG/Source/PCG/Private/Graph/PCGGraphExecutor.cpp
--- a/Engine/Plugins/Experimental/PCG/Source/PCG/Private/Graph/PCGGraphExecutor.cpp
+++ b/Engine/Plugins/Experimental/PCG/Source/PCG/Private/Graph/PCGGraphExecutor.cpp
@@ -605,49 +605,31 @@ void FPCGGraphExecutor::Execute()
 					continue;
 				}
 
-				bool bNeedsToCreateActiveTask = true;
-				if (Task.GraphGenerationGrid != EPCGHiGenGrid::Uninitialized)
-				{
-					// Graph generation sizes should be resolved by now
-					ensure(PCGHiGenGrid::IsValidGrid(Task.GraphGenerationGrid & ~EPCGHiGenGrid::Unbounded) || Task.GraphGenerationGrid == EPCGHiGenGrid::Unbounded);
-
-					bNeedsToCreateActiveTask = !!(Task.GenerationGrid & Task.GraphGenerationGrid);
-
-					if (bGraphCacheDebuggingEnabled && !bNeedsToCreateActiveTask && Task.SourceComponent.Get() && Task.SourceComponent->GetOwner() && Task.Node)
-					{
-						UE_LOG(LogPCG, Warning, TEXT("[%s] %s\t\tGRID SIZE %d SKIPPED"), *Task.SourceComponent->GetOwner()->GetName(), *Task.Node->GetNodeTitle().ToString(), PCGHiGenGrid::GridToGridSize(Task.GraphGenerationGrid));
-					}
-				}
+				// If a task is cacheable and has been cached, then we don't need to create an active task for it unless
+				// there is an execution mode that would prevent us from doing so.
+				const UPCGSettingsInterface* TaskSettingsInterface = TaskInput.GetSettingsInterface(Task.Node ? Task.Node->GetSettingsInterface() : nullptr);
+				const UPCGSettings* TaskSettings = TaskSettingsInterface ? TaskSettingsInterface->GetSettings() : nullptr;
+				const bool bCacheable = Task.Element->IsCacheableInstance(TaskSettingsInterface);
 
+				// Calculate Crc of dependencies (input data Crcs, settings) and use this as the key in the cache lookup
 				FPCGCrc DependenciesCrc;
-				FPCGDataCollection CachedOutput;
-				bool bResultAlreadyInCache = false;
-				if (bNeedsToCreateActiveTask)
+				if (TaskSettings && bCacheable)
 				{
-					// If a task is cacheable and has been cached, then we don't need to create an active task for it unless
-					// there is an execution mode that would prevent us from doing so.
-					const UPCGSettingsInterface* TaskSettingsInterface = TaskInput.GetSettingsInterface(Task.Node ? Task.Node->GetSettingsInterface() : nullptr);
-					const UPCGSettings* TaskSettings = TaskSettingsInterface ? TaskSettingsInterface->GetSettings() : nullptr;
-					const bool bCacheable = Task.Element->IsCacheableInstance(TaskSettingsInterface);
-
-					// Calculate Crc of dependencies (input data Crcs, settings) and use this as the key in the cache lookup
-					if (TaskSettings && bCacheable)
-					{
-						Task.Element->GetDependenciesCrc(TaskInput, TaskSettings, Task.SourceComponent.Get(), DependenciesCrc);
-					}
+					Task.Element->GetDependenciesCrc(TaskInput, TaskSettings, Task.SourceComponent.Get(), DependenciesCrc);
+				}
 
-					if (bGraphCacheDebuggingEnabled && !bCacheable && Task.SourceComponent.Get() && Task.Node)
-					{
-						UE_LOG(LogPCG, Warning, TEXT("[%s] %s\t\tCACHING DISABLED"), *Task.SourceComponent->GetOwner()->GetName(), *Task.Node->GetNodeTitle().ToString());
-					}
+				if (bGraphCacheDebuggingEnabled && !bCacheable && Task.SourceComponent.Get() && Task.Node)
+				{
+					UE_LOG(LogPCG, Warning, TEXT("[%s] %s\t\tCACHING DISABLED"), *Task.SourceComponent->GetOwner()->GetName(), *Task.Node->GetNodeTitle().ToString());
+				}
 
-					bResultAlreadyInCache = bCacheable && DependenciesCrc.IsValid() && GraphCache.GetFromCache(Task.Node, Task.Element.Get(), DependenciesCrc, TaskInput, TaskSettings, Task.SourceComponent.Get(), CachedOutput);
+				FPCGDataCollection CachedOutput;
+				const bool bResultAlreadyInCache = bCacheable && DependenciesCrc.IsValid() && GraphCache.GetFromCache(Task.Node, Task.Element.Get(), DependenciesCrc, TaskInput, TaskSettings, Task.SourceComponent.Get(), CachedOutput);
 #if WITH_EDITOR
-					bNeedsToCreateActiveTask = !bResultAlreadyInCache || TaskSettingsInterface->bDebug;
+				const bool bNeedsToCreateActiveTask = !bResultAlreadyInCache || TaskSettingsInterface->bDebug;
 #else
-					bNeedsToCreateActiveTask = !bResultAlreadyInCache;
+				const bool bNeedsToCreateActiveTask = !bResultAlreadyInCache;
 #endif
-				}
 
 				if (!bNeedsToCreateActiveTask)
 				{
@@ -683,7 +665,6 @@ void FPCGGraphExecutor::Execute()
 					Task.Context->CompiledTaskId = Task.CompiledTaskId;
 					Task.Context->InputData.Crc = TaskInput.Crc;
 					Task.Context->DependenciesCrc = DependenciesCrc;
-					Task.Context->GenerationGrid = Task.GenerationGrid;
 					Task.Context->Stack = (Task.StackContext && Task.StackIndex != INDEX_NONE) ? Task.StackContext->GetStack(Task.StackIndex) : nullptr;
 				}
 
@@ -1474,8 +1455,17 @@ bool FPCGGenericElement::ExecuteInternal(FPCGContext* Context) const
 
 namespace PCGGraphExecutor
 {
-	bool ExecuteGridLinkage(EPCGHiGenGrid InFromGrid, EPCGHiGenGrid InToGrid, const FString& InResourceKey, const FName& InOutputPinLabel, const UPCGNode* InDownstreamNode, FPCGGridLinkageContext* InContext)
+	bool ExecuteGridLinkage(
+		EPCGHiGenGrid InGenerationGrid,
+		EPCGHiGenGrid InFromGrid,
+		EPCGHiGenGrid InToGrid,
+		const FString& InResourceKey,
+		const FName& InOutputPinLabel,
+		const UPCGNode* InDownstreamNode,
+		FPCGGridLinkageContext* InContext)
 	{
+		TRACE_CPUPROFILER_EVENT_SCOPE(PCGGraphExecutor::ExecuteGridLinkage);
+
 		// Non-hierarchical generation - no linkage required - data should just pass through.
 		if (!InContext->SourceComponent->GetGraph()->IsHierarchicalGenerationEnabled()
 			|| !ensure(PCGHiGenGrid::IsValidGrid(InFromGrid) || InFromGrid == EPCGHiGenGrid::Unbounded))
@@ -1523,17 +1513,17 @@ namespace PCGGraphExecutor
 			return true;
 		}
 
-		if (!!(InFromGrid & InContext->GenerationGrid) && FromGridSize != ToGridSize)
+		if (!!(InFromGrid & InGenerationGrid) && FromGridSize != ToGridSize)
 		{
-			PCGGraphExecutionLogging::LogGridLinkageTaskExecuteStore(InContext, FromGridSize, ToGridSize, InResourceKey);
+			PCGGraphExecutionLogging::LogGridLinkageTaskExecuteStore(InContext, InGenerationGrid, FromGridSize, ToGridSize, InResourceKey);
 
 			FPCGDataCollection Data;
 			Data.TaggedData = InContext->InputData.GetInputsByPin(InOutputPinLabel);
 			InContext->SourceComponent->StoreOutputDataForPin(InResourceKey, Data);
 		}
-		else if (InToGrid == InContext->GenerationGrid && FromGridSize != ToGridSize)
+		else if (InToGrid == InGenerationGrid && FromGridSize != ToGridSize)
 		{
-			PCGGraphExecutionLogging::LogGridLinkageTaskExecuteRetrieve(InContext, FromGridSize, ToGridSize, InResourceKey);
+			PCGGraphExecutionLogging::LogGridLinkageTaskExecuteRetrieve(InContext, InGenerationGrid, FromGridSize, ToGridSize, InResourceKey);
 
 			UPCGSubsystem* Subsystem = UPCGSubsystem::GetInstance(InContext->SourceComponent->GetWorld());
 			if (!ensure(Subsystem))
@@ -1550,6 +1540,8 @@ namespace PCGGraphExecutor
 			}
 			else
 			{
+				TRACE_CPUPROFILER_EVENT_SCOPE(PCGGraphExecutor::ExecuteGridLinkage::FindComponent);
+
 				Subsystem->ForAllOverlappingComponentsInHierarchy(InContext->SourceComponent.Get(), [FromGridSize, &ComponentWithData](UPCGComponent* InLocalComponent)
 				{
 					if (InLocalComponent->GetGenerationGridSize() == FromGridSize)
diff --git a/Engine/Plugins/Experimental/PCG/Source/PCG/Private/Graph/PCGGraphExecutor.h b/Engine/Plugins/Experimental/PCG/Source/PCG/Private/Graph/PCGGraphExecutor.h
--- a/Engine/Plugins/Experimental/PCG/Source/PCG/Private/Graph/PCGGraphExecutor.h
+++ b/Engine/Plugins/Experimental/PCG/Source/PCG/Private/Graph/PCGGraphExecutor.h
@@ -56,13 +56,6 @@ struct FPCGGraphTask
 	FPCGTaskId CompiledTaskId = InvalidPCGTaskId; // the task id as it exists when compiled
 	FPCGTaskId ParentId = InvalidPCGTaskId; // represents the parent sub object graph task, if we were called from one
 
-	// Required purely to cache across multiple frames and allow different grids to be scheduled simultaneously. Set up
-	// at schedule-time, used during execution.
-	EPCGHiGenGrid GenerationGrid = EPCGHiGenGrid::Uninitialized;
-	
-	// The generation level for this node within the graph.
-	EPCGHiGenGrid GraphGenerationGrid = EPCGHiGenGrid::Uninitialized;
-
 	int32 StackIndex = INDEX_NONE;
 	TSharedPtr<const FPCGStackContext> StackContext;
 };
@@ -272,6 +265,7 @@ namespace PCGGraphExecutor
 {
 	/** Compares InFromGrid and InToGrid and performs data storage/retrieval as necessary to marshal data across execution grids. */
 	bool ExecuteGridLinkage(
+		EPCGHiGenGrid InGenerationGrid,
 		EPCGHiGenGrid InFromGrid,
 		EPCGHiGenGrid InToGrid,
 		const FString& InResourceKey,
diff --git a/Engine/Plugins/Experimental/PCG/Source/PCG/Private/Utils/PCGGraphExecutionLogging.cpp b/Engine/Plugins/Experimental/PCG/Source/PCG/Private/Utils/PCGGraphExecutionLogging.cpp
--- a/Engine/Plugins/Experimental/PCG/Source/PCG/Private/Utils/PCGGraphExecutionLogging.cpp
+++ b/Engine/Plugins/Experimental/PCG/Source/PCG/Private/Utils/PCGGraphExecutionLogging.cpp
@@ -27,7 +27,7 @@ namespace PCGGraphExecutionLogging
 #endif
 	}
 
-	void LogGridLinkageTaskExecuteStore(const FPCGContext* InContext, int32 InFromGridSize, int32 InToGridSize, const FString& InResourcePath)
+	void LogGridLinkageTaskExecuteStore(const FPCGContext* InContext, EPCGHiGenGrid InGenerationGrid, int32 InFromGridSize, int32 InToGridSize, const FString& InResourcePath)
 	{
 #if WITH_EDITOR
 		if (!LogEnabled())
@@ -37,11 +37,11 @@ namespace PCGGraphExecutionLogging
 		check(InContext);
 
 		const FString OwnerName = InContext->SourceComponent.Get() ? (InContext->SourceComponent->GetOwner() ? InContext->SourceComponent->GetOwner()->GetActorLabel() : FString()) : FString();
-		UE_LOG(LogPCG, Log, TEXT("[GRIDLINKING] [%s] STORE. GenerationGridSize=%d, FromGridSize=%d, ToGridSize=%d, Path=%s"), *OwnerName, PCGHiGenGrid::GridToGridSize(InContext->GenerationGrid), InFromGridSize, InToGridSize, *InResourcePath);
+		UE_LOG(LogPCG, Log, TEXT("[GRIDLINKING] [%s] STORE. GenerationGridSize=%d, FromGridSize=%d, ToGridSize=%d, Path=%s"), *OwnerName, PCGHiGenGrid::GridToGridSize(InGenerationGrid), InFromGridSize, InToGridSize, *InResourcePath);
 #endif
 	}
 
-	void LogGridLinkageTaskExecuteRetrieve(const FPCGContext* InContext, int32 InFromGridSize, int32 InToGridSize, const FString& InResourcePath)
+	void LogGridLinkageTaskExecuteRetrieve(const FPCGContext* InContext, EPCGHiGenGrid InGenerationGrid, int32 InFromGridSize, int32 InToGridSize, const FString& InResourcePath)
 	{
 #if WITH_EDITOR
 		if (!LogEnabled())
@@ -51,7 +51,7 @@ namespace PCGGraphExecutionLogging
 		check(InContext);
 
 		const FString OwnerName = InContext->SourceComponent.Get() ? (InContext->SourceComponent->GetOwner() ? InContext->SourceComponent->GetOwner()->GetActorLabel() : FString()) : FString();
-		UE_LOG(LogPCG, Log, TEXT("[GRIDLINKING] [%s] RETRIEVE. GenerationGridSize=%d, FromGridSize=%d, ToGridSize=%d, Path=%s"), *OwnerName, PCGHiGenGrid::GridToGridSize(InContext->GenerationGrid), InFromGridSize, InToGridSize, *InResourcePath);
+		UE_LOG(LogPCG, Log, TEXT("[GRIDLINKING] [%s] RETRIEVE. GenerationGridSize=%d, FromGridSize=%d, ToGridSize=%d, Path=%s"), *OwnerName, PCGHiGenGrid::GridToGridSize(InGenerationGrid), InFromGridSize, InToGridSize, *InResourcePath);
 #endif
 	}
 
diff --git a/Engine/Plugins/Experimental/PCG/Source/PCG/Private/Utils/PCGGraphExecutionLogging.h b/Engine/Plugins/Experimental/PCG/Source/PCG/Private/Utils/PCGGraphExecutionLogging.h
--- a/Engine/Plugins/Experimental/PCG/Source/PCG/Private/Utils/PCGGraphExecutionLogging.h
+++ b/Engine/Plugins/Experimental/PCG/Source/PCG/Private/Utils/PCGGraphExecutionLogging.h
@@ -5,13 +5,14 @@
 #include "Containers/UnrealString.h"
 #include "HAL/Platform.h"
 
+enum class EPCGHiGenGrid : uint32;
 struct FPCGContext;
 class UPCGComponent;
 
 namespace PCGGraphExecutionLogging
 {
-	void LogGridLinkageTaskExecuteStore(const FPCGContext* InContext, int32 InFromLOD, int32 InToLOD, const FString& InResourcePath);
-	void LogGridLinkageTaskExecuteRetrieve(const FPCGContext* InContext, int32 InFromLOD, int32 InToLOD, const FString& InResourcePath);
+	void LogGridLinkageTaskExecuteStore(const FPCGContext* InContext, EPCGHiGenGrid InGenerationGrid, int32 InFromLOD, int32 InToLOD, const FString& InResourcePath);
+	void LogGridLinkageTaskExecuteRetrieve(const FPCGContext* InContext, EPCGHiGenGrid InGenerationGrid, int32 InFromLOD, int32 InToLOD, const FString& InResourcePath);
 	void LogGridLinkageTaskExecuteRetrieveSuccess(const FPCGContext* InContext, const UPCGComponent* InComponent, const FString& InResourcePath, int32 InDataItemCount);
 	void LogGridLinkageTaskExecuteRetrieveScheduleGraph(const FPCGContext* InContext, const UPCGComponent* InScheduledComponent, const FString& InResourcePath);
 	void LogGridLinkageTaskExecuteRetrieveWaitOnScheduledGraph(const FPCGContext* InContext, const UPCGComponent* InWaitOnComponent, const FString& InResourcePath);
diff --git a/Engine/Plugins/Experimental/PCG/Source/PCG/Public/PCGContext.h b/Engine/Plugins/Experimental/PCG/Source/PCG/Public/PCGContext.h
--- a/Engine/Plugins/Experimental/PCG/Source/PCG/Public/PCGContext.h
+++ b/Engine/Plugins/Experimental/PCG/Source/PCG/Public/PCGContext.h
@@ -68,9 +68,6 @@ struct PCG_API FPCGContext
 	EPCGExecutionPhase CurrentPhase = EPCGExecutionPhase::NotExecuted;
 	int32 BypassedOutputCount = 0;
 
-	/** Used by grid size data linkages to store/retrieve data. */
-	EPCGHiGenGrid GenerationGrid = EPCGHiGenGrid::Uninitialized;
-
 	/** The current call stack. */
 	const FPCGStack* Stack = nullptr;
 
