Added an input weight attribute accessor so we can easily drive selection from other attributes

#rnx

diff --git a/Engine/Plugins/Experimental/PCG/Source/PCG/Private/Elements/PCGMatchAndSetAttributes.cpp b/Engine/Plugins/Experimental/PCG/Source/PCG/Private/Elements/PCGMatchAndSetAttributes.cpp
--- a/Engine/Plugins/Experimental/PCG/Source/PCG/Private/Elements/PCGMatchAndSetAttributes.cpp
+++ b/Engine/Plugins/Experimental/PCG/Source/PCG/Private/Elements/PCGMatchAndSetAttributes.cpp
@@ -44,6 +44,7 @@ FText UPCGMatchAndSetAttributesSettings::GetNodeTooltipText() const
 
 UPCGMatchAndSetAttributesSettings::UPCGMatchAndSetAttributesSettings()
 {
+	// Minor TODO: could mark use seed true only if we don't use the input weight attribute
 	bUseSeed = true;
 }
 
@@ -206,6 +207,21 @@ public:
 			PCGE_LOG_C(Warning, GraphAndLog, Context, FText::Format(LOCTEXT("EmptyCategory", "Some match entries on Attribute '{0}' in the Attribute Set do not have any associated valid weight."), FText::FromName(Attribute ? Attribute->Name : NAME_None)));
 		}
 #endif // WITH_EDITOR
+
+		// Normalize weights
+		for (TPair<PCGMetadataValueKey, AttributeSetPartitionEntry>& PartitionEntry : PartitionData)
+		{
+			AttributeSetPartitionEntry& Entry = PartitionEntry.Value;
+			if (Entry.TotalWeight > 0)
+			{
+				for (double& Weight : Entry.CumulativeWeight)
+				{
+					Weight /= Entry.TotalWeight;
+				}
+
+				Entry.TotalWeight = 1.0;
+			}
+		}
 		
 		bIsValid = true;
 		return true;
@@ -270,7 +286,7 @@ public:
 		return MatchingPartitionDataIndices;
 	}
 
-	PCGMetadataEntryKey GetWeightedEntry(int32 PartitionDataIndex, const FRandomStream& InRandomStream) const
+	PCGMetadataEntryKey GetWeightedEntry(int32 PartitionDataIndex, double RandomWeightedPick) const
 	{
 		if (PartitionDataIndex == INDEX_NONE)
 		{
@@ -287,12 +303,14 @@ public:
 		}
 		else if (PartitionDataEntry.Keys.Num() > 1)
 		{
-			double RandomWeightedPick = InRandomStream.FRandRange(0, PartitionDataEntry.TotalWeight);
 			RandomPick = 0;
 			while (RandomPick < PartitionDataEntry.CumulativeWeight.Num() && PartitionDataEntry.CumulativeWeight[RandomPick] <= RandomWeightedPick)
 			{
 				++RandomPick;
 			}
+
+			// If weight is outside of the unit range, then we can still take the last entry
+			RandomPick = FMath::Min(RandomPick, PartitionDataEntry.CumulativeWeight.Num() - 1);
 		}
 
 		if (RandomPick != INDEX_NONE)
@@ -352,7 +370,7 @@ public:
 
 		if (Settings->bMatchAttributes)
 		{
-			FPCGAttributePropertyInputSelector InputAttributeSource = Settings->InputAttribute.CopyAndFixLast(PointData);
+			const FPCGAttributePropertyInputSelector InputAttributeSource = Settings->InputAttribute.CopyAndFixLast(PointData);
 			InputAttributeAccessor = PCGAttributeAccessorHelpers::CreateConstAccessor(PointData, InputAttributeSource);
 			InputAttributeKeys = PCGAttributeAccessorHelpers::CreateConstKeys(PointData, InputAttributeSource);
 
@@ -363,6 +381,29 @@ public:
 			}
 		}
 
+		if (Settings->bUseInputWeightAttribute)
+		{
+			const FPCGAttributePropertyInputSelector InputWeightAttributeSource = Settings->InputWeightAttribute.CopyAndFixLast(PointData);
+			InputWeightAccessor = PCGAttributeAccessorHelpers::CreateConstAccessor(PointData, InputWeightAttributeSource);
+
+			if (!InputAttributeKeys)
+			{
+				InputAttributeKeys = PCGAttributeAccessorHelpers::CreateConstKeys(PointData, InputWeightAttributeSource);
+			}
+
+			if (!InputWeightAccessor.IsValid() || !InputAttributeKeys.IsValid())
+			{
+				PCGE_LOG_C(Error, GraphAndLog, Context, FText::Format(LOCTEXT("MissingWeightAttribute", "Point data does not have the input weight attribute '{0}'."), InputWeightAttributeSource.GetDisplayText()));
+				return false;
+			}
+
+			if (!PCG::Private::IsOfTypes<float, double>(InputWeightAccessor->GetUnderlyingType()))
+			{
+				PCGE_LOG_C(Error, GraphAndLog, Context, FText::Format(LOCTEXT("InvalidInputWeightAttributeType", "Input weight attribute '{0}' does not have the proper type (float or double)."), InputWeightAttributeSource.GetDisplayText()));
+				return false;
+			}
+		}
+
 		AttributesToSet.Reset();
 
 		// Prepare set of attributes to copy, i.e. create the attributes if we need to, make a 1:1 pair with the ones from the
@@ -407,6 +448,29 @@ public:
 		return true;
 	}
 
+	TArray<double> GetWeights(const TArray<FPCGPoint>& Points) 
+	{
+		TArray<double> Weights;
+
+		if (InputWeightAccessor.IsValid() && InputAttributeKeys.IsValid())
+		{
+			Weights.SetNumUninitialized(Points.Num());
+			InputWeightAccessor->GetRange<double>(Weights, 0, *InputAttributeKeys, EPCGAttributeAccessorFlags::AllowConstructible);
+		}
+		else
+		{
+			Weights.Reserve(Points.Num());
+
+			// Generate a random value from the seed
+			for (const FPCGPoint& Point : Points)
+			{
+				Weights.Add(UPCGBlueprintHelpers::GetRandomStreamFromPoint(Point, Settings, SourceComponent).FRand());
+			}
+		}
+
+		return Weights;
+	}
+
 	FPCGPointDataPartitionBase::Element* SelectPoint(const FPCGPoint& Point, int32 PointIndex)
 	{
 		return nullptr;
@@ -427,6 +491,13 @@ public:
 			return;
 		}
 
+		const TArray<double> PointWeights = GetWeights(Points);
+
+		if (PointWeights.Num() != Points.Num())
+		{
+			return;
+		}
+
 		for (int32 PointIndex = 0; PointIndex < Points.Num(); ++PointIndex)
 		{
 			const FPCGPoint& Point = Points[PointIndex];
@@ -435,7 +506,7 @@ public:
 
 			if (PartitionDataIndex != INDEX_NONE)
 			{
-				AttributeSetKey = AttributeSetPartition.GetWeightedEntry(PartitionDataIndex, UPCGBlueprintHelpers::GetRandomStreamFromPoint(Point, Settings, SourceComponent));
+				AttributeSetKey = AttributeSetPartition.GetWeightedEntry(PartitionDataIndex, PointWeights[PointIndex]);
 			}
 
 			if (Settings->bKeepUnmatched || AttributeSetKey != PCGInvalidEntryKey)
@@ -468,6 +539,7 @@ private:
 
 	// Per point data iteration data
 	TUniquePtr<const IPCGAttributeAccessor> InputAttributeAccessor;
+	TUniquePtr<const IPCGAttributeAccessor> InputWeightAccessor;
 	TUniquePtr<const IPCGAttributeAccessorKeys> InputAttributeKeys;
 	TArray<TPair<const FPCGMetadataAttributeBase*, FPCGMetadataAttributeBase*>> AttributesToSet;
 };
diff --git a/Engine/Plugins/Experimental/PCG/Source/PCG/Private/Metadata/PCGMetadata.cpp b/Engine/Plugins/Experimental/PCG/Source/PCG/Private/Metadata/PCGMetadata.cpp
--- a/Engine/Plugins/Experimental/PCG/Source/PCG/Private/Metadata/PCGMetadata.cpp
+++ b/Engine/Plugins/Experimental/PCG/Source/PCG/Private/Metadata/PCGMetadata.cpp
@@ -30,7 +30,7 @@ void UPCGMetadata::Serialize(FArchive& InArchive)
 			InArchive << AttributeTypeId;
 
 			FPCGMetadataAttributeBase* SerializedAttribute = PCGMetadataAttribute::AllocateEmptyAttributeFromType(static_cast<int16>(AttributeTypeId));
-			if (SerializedAttribute)
+			if (ensure(SerializedAttribute))
 			{
 				SerializedAttribute->Name = AttributeName;
 				SerializedAttribute->Serialize(this, InArchive);
diff --git a/Engine/Plugins/Experimental/PCG/Source/PCG/Public/Elements/PCGMatchAndSetAttributes.h b/Engine/Plugins/Experimental/PCG/Source/PCG/Public/Elements/PCGMatchAndSetAttributes.h
--- a/Engine/Plugins/Experimental/PCG/Source/PCG/Public/Elements/PCGMatchAndSetAttributes.h
+++ b/Engine/Plugins/Experimental/PCG/Source/PCG/Public/Elements/PCGMatchAndSetAttributes.h
@@ -50,12 +50,19 @@ public:
 	UPROPERTY(BlueprintReadWrite, EditAnywhere, Category = Settings, meta = (EditCondition = "bMatchAttributes", PCG_Overridable))
 	bool bKeepUnmatched = true;
 
+	UPROPERTY(BlueprintReadWrite, EditAnywhere, Category = Settings)
+	bool bUseInputWeightAttribute = false;
+
+	/** Input weight from the points, assumed to be in the [0, 1] range */
+	UPROPERTY(BlueprintReadWrite, EditAnywhere, Category = Settings, meta = (EditCondition = "bUseInputWeightAttribute"))
+	FPCGAttributePropertyInputSelector InputWeightAttribute;
+
 	/** Controls whether we will consider the weights, as determined by the Weight Attribute values on the attribute set */
-	UPROPERTY(BlueprintReadWrite, EditAnywhere, Category = Settings, meta = (InlineEditConditionToggle, PCG_Overridable))
+	UPROPERTY(BlueprintReadWrite, EditAnywhere, Category = Settings, meta = (DisplayName = "Use Match Weight", InlineEditConditionToggle, PCG_Overridable))
 	bool bUseWeightAttribute = false;
 
 	/** Attribute to weight more or less some entries from the attribute set */
-	UPROPERTY(BlueprintReadWrite, EditAnywhere, Category = Settings, meta = (EditCondition = "bUseWeightAttribute", PCG_Overridable))
+	UPROPERTY(BlueprintReadWrite, EditAnywhere, Category = Settings, meta = (DisplayName="Match Weight Attribute", EditCondition = "bUseWeightAttribute", PCG_Overridable))
 	FName WeightAttribute = NAME_None;
 };
 
