[PCG] Fix crash when breaking PCG link - move component properly

diff --git a/Engine/Plugins/Experimental/PCG/Source/PCG/Private/PCGManagedResource.cpp b/Engine/Plugins/Experimental/PCG/Source/PCG/Private/PCGManagedResource.cpp
--- a/Engine/Plugins/Experimental/PCG/Source/PCG/Private/PCGManagedResource.cpp
+++ b/Engine/Plugins/Experimental/PCG/Source/PCG/Private/PCGManagedResource.cpp
@@ -257,6 +257,8 @@ bool UPCGManagedComponent::MoveResourceToNewActor(AActor* NewActor, const AActor
 	bool bDetached = false;
 	bool bAttached = false;
 
+	GeneratedComponent->UnregisterComponent();
+
 	// Need to change owner first to avoid that the PCG Component will react to this component changes.
 	GeneratedComponent->Rename(nullptr, NewActor);
 
@@ -279,6 +281,8 @@ bool UPCGManagedComponent::MoveResourceToNewActor(AActor* NewActor, const AActor
 		NewActor->AddInstanceComponent(GeneratedComponent.Get());
 	}
 
+	GeneratedComponent->RegisterComponent();
+
 	ForgetComponent();
 
 	return true;
