[PCG] Refactored Set Graph Parameters to operate with PCG Graphs (not just instances) and properly utilize callbacks

+ C++ API for getting/setting graph parameters through code
+ Updated enums to work with internal enums as well as objects

diff --git a/Engine/Plugins/Experimental/PCG/Config/BasePCG.ini b/Engine/Plugins/Experimental/PCG/Config/BasePCG.ini
--- a/Engine/Plugins/Experimental/PCG/Config/BasePCG.ini
+++ b/Engine/Plugins/Experimental/PCG/Config/BasePCG.ini
@@ -26,6 +26,34 @@
 +PropertyRedirects=(OldName="PCGStaticMeshSpawnerSettings.MeshSelectorInstance",NewName="PCGStaticMeshSpawnerSettings.MeshSelectorParameters")
 +PropertyRedirects=(OldName="PCGStaticMeshSpawnerSettings.InstancePackerType",NewName="PCGStaticMeshSpawnerSettings.InstanceDataPackerType")
 +PropertyRedirects=(OldName="PCGStaticMeshSpawnerSettings.InstancePackerInstance",NewName="PCGStaticMeshSpawnerSettings.InstanceDataPackerParameters")
++PropertyRedirects=(OldName="PCGGraphParametersHelpers.IsOverridden.GraphInstance",NewName="PCGGraphParametersHelpers.IsOverridden.GraphInterface")
++PropertyRedirects=(OldName="PCGGraphParametersHelpers.GetFloatParameter.GraphInstance",NewName="PCGGraphParametersHelpers.GetFloatParameter.GraphInterface")
++PropertyRedirects=(OldName="PCGGraphParametersHelpers.GetDoubleParameter.GraphInstance",NewName="PCGGraphParametersHelpers.GetDoubleParameter.GraphInterface")
++PropertyRedirects=(OldName="PCGGraphParametersHelpers.GetBoolParameter.GraphInstance",NewName="PCGGraphParametersHelpers.GetBoolParameter.GraphInterface")
++PropertyRedirects=(OldName="PCGGraphParametersHelpers.GetByteParameter.GraphInstance",NewName="PCGGraphParametersHelpers.GetByteParameter.GraphInterface")
++PropertyRedirects=(OldName="PCGGraphParametersHelpers.GetInt32Parameter.GraphInstance",NewName="PCGGraphParametersHelpers.GetInt32Parameter.GraphInterface")
++PropertyRedirects=(OldName="PCGGraphParametersHelpers.GetInt64Parameter.GraphInstance",NewName="PCGGraphParametersHelpers.GetInt64Parameter.GraphInterface")
++PropertyRedirects=(OldName="PCGGraphParametersHelpers.GetNameParameter.GraphInstance",NewName="PCGGraphParametersHelpers.GetNameParameter.GraphInterface")
++PropertyRedirects=(OldName="PCGGraphParametersHelpers.GetStringParameter.GraphInstance",NewName="PCGGraphParametersHelpers.GetStringParameter.GraphInterface")
++PropertyRedirects=(OldName="PCGGraphParametersHelpers.GetSoftObjectParameter.GraphInstance",NewName="PCGGraphParametersHelpers.GetSoftObjectParameter.GraphInterface")
++PropertyRedirects=(OldName="PCGGraphParametersHelpers.GetSoftClassParameter.GraphInstance",NewName="PCGGraphParametersHelpers.GetSoftClassParameter.GraphInterface")
++PropertyRedirects=(OldName="PCGGraphParametersHelpers.GetVectorParameter.GraphInstance",NewName="PCGGraphParametersHelpers.GetVectorParameter.GraphInterface")
++PropertyRedirects=(OldName="PCGGraphParametersHelpers.GetRotatorParameter.GraphInstance",NewName="PCGGraphParametersHelpers.GetRotatorParameter.GraphInterface")
++PropertyRedirects=(OldName="PCGGraphParametersHelpers.GetTransformParameter.GraphInstance",NewName="PCGGraphParametersHelpers.GetTransformParameter.GraphInterface")
++PropertyRedirects=(OldName="PCGGraphParametersHelpers.SetFloatParameter.GraphInstance",NewName="PCGGraphParametersHelpers.SetFloatParameter.GraphInterface")
++PropertyRedirects=(OldName="PCGGraphParametersHelpers.SetDoubleParameter.GraphInstance",NewName="PCGGraphParametersHelpers.SetDoubleParameter.GraphInterface")
++PropertyRedirects=(OldName="PCGGraphParametersHelpers.SetBoolParameter.GraphInstance",NewName="PCGGraphParametersHelpers.SetBoolParameter.GraphInterface")
++PropertyRedirects=(OldName="PCGGraphParametersHelpers.SetByteParameter.GraphInstance",NewName="PCGGraphParametersHelpers.SetByteParameter.GraphInterface")
++PropertyRedirects=(OldName="PCGGraphParametersHelpers.SetInt32Parameter.GraphInstance",NewName="PCGGraphParametersHelpers.SetInt32Parameter.GraphInterface")
++PropertyRedirects=(OldName="PCGGraphParametersHelpers.SetInt64Parameter.GraphInstance",NewName="PCGGraphParametersHelpers.SetInt64Parameter.GraphInterface")
++PropertyRedirects=(OldName="PCGGraphParametersHelpers.SetNameParameter.GraphInstance",NewName="PCGGraphParametersHelpers.SetNameParameter.GraphInterface")
++PropertyRedirects=(OldName="PCGGraphParametersHelpers.SetStringParameter.GraphInstance",NewName="PCGGraphParametersHelpers.SetStringParameter.GraphInterface")
++PropertyRedirects=(OldName="PCGGraphParametersHelpers.SetEnumParameter.GraphInstance",NewName="PCGGraphParametersHelpers.SetEnumParameter.GraphInterface")
++PropertyRedirects=(OldName="PCGGraphParametersHelpers.SetSoftObjectParameter.GraphInstance",NewName="PCGGraphParametersHelpers.SetSoftObjectParameter.GraphInterface")
++PropertyRedirects=(OldName="PCGGraphParametersHelpers.SetSoftClassParameter.GraphInstance",NewName="PCGGraphParametersHelpers.SetSoftClassParameter.GraphInterface")
++PropertyRedirects=(OldName="PCGGraphParametersHelpers.SetVectorParameter.GraphInstance",NewName="PCGGraphParametersHelpers.SetVectorParameter.GraphInterface")
++PropertyRedirects=(OldName="PCGGraphParametersHelpers.SetRotatorParameter.GraphInstance",NewName="PCGGraphParametersHelpers.SetRotatorParameter.GraphInterface")
++PropertyRedirects=(OldName="PCGGraphParametersHelpers.SetTransformParameter.GraphInstance",NewName="PCGGraphParametersHelpers.SetTransformParameter.GraphInterface")
 
 +ClassRedirects=(OldName="PCGPointSamplerSettings",NewName="/Script/PCG.PCGSelectPointsSettings")
 +ClassRedirects=(OldName="PCGDensityNoiseSettings",NewName="/Script/PCG.PCGAttributeNoiseSettings")
diff --git a/Engine/Plugins/Experimental/PCG/Source/PCG/Private/Helpers/PCGGraphParameterExtension.cpp b/Engine/Plugins/Experimental/PCG/Source/PCG/Private/Helpers/PCGGraphParameterExtension.cpp
new file mode 100644
--- /dev/null
+++ b/Engine/Plugins/Experimental/PCG/Source/PCG/Private/Helpers/PCGGraphParameterExtension.cpp
@@ -0,0 +1,181 @@
+﻿// Copyright Epic Games, Inc. All Rights Reserved.
+
+#pragma once
+
+#include "Helpers/PCGGraphParameterExtension.h"
+
+#include "PCGGraph.h"
+
+////////////
+// Getters
+////////////
+template<> TValueOrError<float, EPropertyBagResult> FPCGGraphParameterExtension::GetGraphParameter(const FInstancedPropertyBag& PropertyBag, const FName PropertyName)
+{
+	return PropertyBag.GetValueFloat(PropertyName);
+}
+template<> TValueOrError<double, EPropertyBagResult> FPCGGraphParameterExtension::GetGraphParameter(const FInstancedPropertyBag& PropertyBag, const FName PropertyName)
+{
+	return PropertyBag.GetValueDouble(PropertyName);
+}
+template<> TValueOrError<bool, EPropertyBagResult> FPCGGraphParameterExtension::GetGraphParameter(const FInstancedPropertyBag& PropertyBag, const FName PropertyName)
+{
+	return PropertyBag.GetValueBool(PropertyName);
+}
+template<> TValueOrError<uint8, EPropertyBagResult> FPCGGraphParameterExtension::GetGraphParameter(const FInstancedPropertyBag& PropertyBag, const FName PropertyName)
+{
+	return PropertyBag.GetValueByte(PropertyName);
+}
+template<> TValueOrError<int32, EPropertyBagResult> FPCGGraphParameterExtension::GetGraphParameter(const FInstancedPropertyBag& PropertyBag, const FName PropertyName)
+{
+	return PropertyBag.GetValueInt32(PropertyName);
+}
+template<> TValueOrError<int64, EPropertyBagResult> FPCGGraphParameterExtension::GetGraphParameter(const FInstancedPropertyBag& PropertyBag, const FName PropertyName)
+{
+	return PropertyBag.GetValueInt64(PropertyName);
+}
+template<> TValueOrError<FName, EPropertyBagResult> FPCGGraphParameterExtension::GetGraphParameter(const FInstancedPropertyBag& PropertyBag, const FName PropertyName)
+{
+	return PropertyBag.GetValueName(PropertyName);
+}
+template<> TValueOrError<FString, EPropertyBagResult> FPCGGraphParameterExtension::GetGraphParameter(const FInstancedPropertyBag& PropertyBag, const FName PropertyName)
+{
+	return PropertyBag.GetValueString(PropertyName);
+}
+template<> TValueOrError<FSoftObjectPath, EPropertyBagResult> FPCGGraphParameterExtension::GetGraphParameter(const FInstancedPropertyBag& PropertyBag, const FName PropertyName)
+{
+	return PropertyBag.GetValueSoftPath(PropertyName);
+}
+template<> TValueOrError<TSoftObjectPtr<UObject>, EPropertyBagResult> FPCGGraphParameterExtension::GetGraphParameter(const FInstancedPropertyBag& PropertyBag, const FName PropertyName)
+{
+	TValueOrError<FSoftObjectPath, EPropertyBagResult> InitialResult = PropertyBag.GetValueSoftPath(PropertyName);
+	if (InitialResult.HasError())
+	{
+		return MakeError(InitialResult.GetError());
+	}
+	TValueOrError<TSoftObjectPtr<UObject>, EPropertyBagResult> FinalResult(MakeValue(TSoftObjectPtr<UObject>(InitialResult.GetValue())));
+	return FinalResult;
+}
+template<> TValueOrError<TSoftClassPtr<UObject>, EPropertyBagResult> FPCGGraphParameterExtension::GetGraphParameter(const FInstancedPropertyBag& PropertyBag, const FName PropertyName)
+{
+	TValueOrError<FSoftObjectPath, EPropertyBagResult> InitialResult = PropertyBag.GetValueSoftPath(PropertyName);
+	if (InitialResult.HasError())
+	{
+		return MakeError(InitialResult.GetError());
+	}
+	TValueOrError<TSoftClassPtr<UObject>, EPropertyBagResult> FinalResult(MakeValue(TSoftClassPtr<UObject>(InitialResult.GetValue())));
+	return FinalResult;
+}
+template<> TValueOrError<FVector*, EPropertyBagResult> FPCGGraphParameterExtension::GetGraphParameter(const FInstancedPropertyBag& PropertyBag, const FName PropertyName)
+{
+	return PropertyBag.GetValueStruct<FVector>(PropertyName);
+}
+template<> TValueOrError<FRotator*, EPropertyBagResult> FPCGGraphParameterExtension::GetGraphParameter(const FInstancedPropertyBag& PropertyBag, const FName PropertyName)
+{
+	return PropertyBag.GetValueStruct<FRotator>(PropertyName);
+}
+template<> TValueOrError<FTransform*, EPropertyBagResult> FPCGGraphParameterExtension::GetGraphParameter(const FInstancedPropertyBag& PropertyBag, const FName PropertyName)
+{
+	return PropertyBag.GetValueStruct<FTransform>(PropertyName);
+}
+template<> TValueOrError<FVector4*, EPropertyBagResult> FPCGGraphParameterExtension::GetGraphParameter(const FInstancedPropertyBag& PropertyBag, const FName PropertyName)
+{
+	return PropertyBag.GetValueStruct<FVector4>(PropertyName);
+}
+template<> TValueOrError<FVector2D*, EPropertyBagResult> FPCGGraphParameterExtension::GetGraphParameter(const FInstancedPropertyBag& PropertyBag, const FName PropertyName)
+{
+	return PropertyBag.GetValueStruct<FVector2D>(PropertyName);
+}
+template<> TValueOrError<FQuat*, EPropertyBagResult> FPCGGraphParameterExtension::GetGraphParameter(const FInstancedPropertyBag& PropertyBag, const FName PropertyName)
+{
+	return PropertyBag.GetValueStruct<FQuat>(PropertyName);
+}
+
+TValueOrError<uint64, EPropertyBagResult> FPCGGraphParameterExtension::GetGraphParameter(const FInstancedPropertyBag& PropertyBag, const FName PropertyName, const UEnum* Enum)
+{
+	// TODO: To stay congruent with BP Enums and Property Bag, Value converted from uint8. Update this if/when this is improved externally
+	TValueOrError<uint8, EPropertyBagResult> Result =  Enum ? PropertyBag.GetValueEnum(PropertyName, Enum) : PropertyBag.GetValueByte(PropertyName);
+	if (Result.HasError())
+	{
+		return MakeError(Result.GetError());
+	}
+
+	return MakeValue(static_cast<uint64>(Result.GetValue()));
+}
+
+////////////
+// Setters
+////////////
+template<> EPropertyBagResult FPCGGraphParameterExtension::SetGraphParameter(FInstancedPropertyBag& PropertyBag, const FName PropertyName, const float& Value)
+{
+	return PropertyBag.SetValueFloat(PropertyName, Value);
+}
+template<> EPropertyBagResult FPCGGraphParameterExtension::SetGraphParameter(FInstancedPropertyBag& PropertyBag, const FName PropertyName, const double& Value)
+{
+	return PropertyBag.SetValueDouble(PropertyName, Value);
+}
+template<> EPropertyBagResult FPCGGraphParameterExtension::SetGraphParameter(FInstancedPropertyBag& PropertyBag, const FName PropertyName, const bool& bValue)
+{
+	return PropertyBag.SetValueBool(PropertyName, bValue);
+}
+template<> EPropertyBagResult FPCGGraphParameterExtension::SetGraphParameter(FInstancedPropertyBag& PropertyBag, const FName PropertyName, const uint8& Value)
+{
+	return PropertyBag.SetValueByte(PropertyName, Value);
+}
+template<> EPropertyBagResult FPCGGraphParameterExtension::SetGraphParameter(FInstancedPropertyBag& PropertyBag, const FName PropertyName, const int32& Value)
+{
+	return PropertyBag.SetValueInt32(PropertyName, Value);
+}
+template<> EPropertyBagResult FPCGGraphParameterExtension::SetGraphParameter(FInstancedPropertyBag& PropertyBag, const FName PropertyName, const int64& Value)
+{
+	return PropertyBag.SetValueInt64(PropertyName, Value);
+}
+template<> EPropertyBagResult FPCGGraphParameterExtension::SetGraphParameter(FInstancedPropertyBag& PropertyBag, const FName PropertyName, const FName& Value)
+{
+	return PropertyBag.SetValueName(PropertyName, Value);
+}
+template<> EPropertyBagResult FPCGGraphParameterExtension::SetGraphParameter(FInstancedPropertyBag& PropertyBag, const FName PropertyName, const FString& Value)
+{
+	return PropertyBag.SetValueString(PropertyName, Value);
+}
+template<> EPropertyBagResult FPCGGraphParameterExtension::SetGraphParameter(FInstancedPropertyBag& PropertyBag, const FName PropertyName, const FSoftObjectPath& Value)
+{
+	return PropertyBag.SetValueSoftPath(PropertyName, Value);
+}
+template<> EPropertyBagResult FPCGGraphParameterExtension::SetGraphParameter(FInstancedPropertyBag& PropertyBag, const FName PropertyName, const TSoftObjectPtr<UObject>& Value)
+{
+	return PropertyBag.SetValueSoftPath(PropertyName, Value.ToSoftObjectPath());
+}
+template<> EPropertyBagResult FPCGGraphParameterExtension::SetGraphParameter(FInstancedPropertyBag& PropertyBag, const FName PropertyName, const TSoftClassPtr<UObject>& Value)
+{
+	return PropertyBag.SetValueSoftPath(PropertyName, Value.ToSoftObjectPath());
+}
+template<> EPropertyBagResult FPCGGraphParameterExtension::SetGraphParameter(FInstancedPropertyBag& PropertyBag, const FName PropertyName, const FVector& Value)
+{
+	return PropertyBag.SetValueStruct(PropertyName, Value);
+}
+template<> EPropertyBagResult FPCGGraphParameterExtension::SetGraphParameter(FInstancedPropertyBag& PropertyBag, const FName PropertyName, const FRotator& Value)
+{
+	return PropertyBag.SetValueStruct(PropertyName, Value);
+}
+template<> EPropertyBagResult FPCGGraphParameterExtension::SetGraphParameter(FInstancedPropertyBag& PropertyBag, const FName PropertyName, const FTransform& Value)
+{
+	return PropertyBag.SetValueStruct(PropertyName, Value);
+}
+template<> EPropertyBagResult FPCGGraphParameterExtension::SetGraphParameter(FInstancedPropertyBag& PropertyBag, const FName PropertyName, const FVector4& Value)
+{
+	return PropertyBag.SetValueStruct(PropertyName, Value);
+}
+template<> EPropertyBagResult FPCGGraphParameterExtension::SetGraphParameter(FInstancedPropertyBag& PropertyBag, const FName PropertyName, const FVector2D& Value)
+{
+	return PropertyBag.SetValueStruct(PropertyName, Value);
+}
+template<> EPropertyBagResult FPCGGraphParameterExtension::SetGraphParameter(FInstancedPropertyBag& PropertyBag, const FName PropertyName, const FQuat& Value)
+{
+	return PropertyBag.SetValueStruct(PropertyName, Value);
+}
+
+EPropertyBagResult FPCGGraphParameterExtension::SetGraphParameter(FInstancedPropertyBag& PropertyBag, const FName PropertyName, const uint64& Value, const UEnum* Enum)
+{
+	// TODO: To stay congruent with BP Enums and Property Bag, Value converted to uint8. Update this if/when this is improved externally
+	return Enum ? PropertyBag.SetValueEnum(PropertyName, static_cast<uint8>(Value), Enum) : PropertyBag.SetValueByte(PropertyName, static_cast<uint8>(Value));
+}
\ No newline at end of file
diff --git a/Engine/Plugins/Experimental/PCG/Source/PCG/Private/Helpers/PCGGraphParametersHelpers.cpp b/Engine/Plugins/Experimental/PCG/Source/PCG/Private/Helpers/PCGGraphParametersHelpers.cpp
--- a/Engine/Plugins/Experimental/PCG/Source/PCG/Private/Helpers/PCGGraphParametersHelpers.cpp
+++ b/Engine/Plugins/Experimental/PCG/Source/PCG/Private/Helpers/PCGGraphParametersHelpers.cpp
@@ -11,27 +11,38 @@
 
 #define LOCTEXT_NAMESPACE "PCGGraphParametersHelpers"
 
-namespace PCGGraphParametersHelpers
+// Helper functions for this blueprint library
+namespace PCGGraphParametersHelpersPrivate
 {
-	static const FText InvalidGraphInstanceError = LOCTEXT("InvalidGraphInstance", "Invalid graph instance.");
+	static const FText InvalidGraphInterfaceError = LOCTEXT("InvalidGraphInterface", "Invalid graph interface.");
 
-	void OnException(EPropertyBagResult InResult, FName InPropertyName)
+	void ThrowBlueprintException(const FText& ErrorMessage)
+	{
+		const FBlueprintExceptionInfo ExceptionInfo(EBlueprintExceptionType::FatalError, ErrorMessage);
+		FBlueprintCoreDelegates::ThrowScriptException(FFrame::GetThreadLocalTopStackFrame()->Object, *FFrame::GetThreadLocalTopStackFrame(), ExceptionInfo);
+	}
+
+	void OnException(const EPropertyBagResult Result, const FName PropertyName)
 	{
 		FText ErrorMessage;
-		switch (InResult)
+		switch (Result)
 		{
+		case EPropertyBagResult::Success:
+			return;
 		case EPropertyBagResult::TypeMismatch:
-			ErrorMessage = FText::Format(LOCTEXT("TypeMismatch", "Parameter {0} is not of the right type."), FText::FromName(InPropertyName));
+			ErrorMessage = FText::Format(LOCTEXT("TypeMismatch", "Parameter {0} is not of the right type."), FText::FromName(PropertyName));
+			break;
+		case EPropertyBagResult::OutOfBounds:
+			ErrorMessage = FText::Format(LOCTEXT("OutOfBounds", "Parameter {0} is out of bounds."), FText::FromName(PropertyName));
 			break;
 		case EPropertyBagResult::PropertyNotFound:
-			ErrorMessage = FText::Format(LOCTEXT("PropertyNotFound", "Parameter {0} does not exist."), FText::FromName(InPropertyName));
+			ErrorMessage = FText::Format(LOCTEXT("PropertyNotFound", "Parameter {0} does not exist."), FText::FromName(PropertyName));
 			break;
 		default:
-			break;
+			checkNoEntry();
 		}
 
-		const FBlueprintExceptionInfo ExceptionInfo(EBlueprintExceptionType::FatalError, ErrorMessage);
-		FBlueprintCoreDelegates::ThrowScriptException(FFrame::GetThreadLocalTopStackFrame()->Object, *FFrame::GetThreadLocalTopStackFrame(), ExceptionInfo);
+		ThrowBlueprintException(ErrorMessage);
 
 		if (FPlatformMisc::IsDebuggerPresent())
 		{
@@ -39,324 +50,265 @@ namespace PCGGraphParametersHelpers
 		}
 	}
 
-	void UpdateOverrides(UPCGGraphInstance* GraphInstance, FName InPropertyName)
+	// Usually returns T, but in some cases (like soft objects or struct), it returns a ptr or something else.
+	// That's what U is for, and there are strict conditions on the relation between T and U.
+	template <typename T, typename U = T>
+	T ValidateAndReturnResult(const UPCGGraphInterface* GraphInterface, const FName PropertyName)
 	{
-		if (!GraphInstance)
+		if (!GraphInterface)
 		{
-			return;
+			return T{};
 		}
 
-		if (const FPropertyBagPropertyDesc* Desc = GraphInstance->ParametersOverrides.Parameters.FindPropertyDescByName(InPropertyName))
+		const TValueOrError<U, EPropertyBagResult> Result = GraphInterface->GetGraphParameter<U>(PropertyName);
+		if (Result.HasError())
 		{
-			GraphInstance->ParametersOverrides.PropertiesIDsOverridden.Add(Desc->ID);
+			OnException(Result.GetError(), PropertyName);
+			return T{};
 		}
-	}
 
-	// Get a parameter from the Instanced Property Bag using its getters. It usually returns T, but in some cases (like soft objects or struct), it returns a ptr or something
-	// else. That's what U is for, and there are strict conditions on the relation between T and U.
-	template <typename T, typename U = T>
-	T GetParameter(UPCGGraphInstance* GraphInstance, FName PropertyName, TFunctionRef<TValueOrError<U, EPropertyBagResult>(FInstancedPropertyBag&)> Getter)
-	{
-		if (!GraphInstance)
+		if constexpr (std::is_pointer_v<U>)
 		{
-			const FBlueprintExceptionInfo ExceptionInfo(EBlueprintExceptionType::FatalError, InvalidGraphInstanceError);
-			FBlueprintCoreDelegates::ThrowScriptException(FFrame::GetThreadLocalTopStackFrame()->Object, *FFrame::GetThreadLocalTopStackFrame(), ExceptionInfo);
-
-			return T{};
+			static_assert(std::is_same_v<T, std::remove_pointer_t<U>>);
+			return Result.GetValue() ? *Result.GetValue() : T{};
 		}
-
-		TValueOrError<U, EPropertyBagResult> Result = Getter(GraphInstance->ParametersOverrides.Parameters);
-
-		if (Result.HasError())
+		else if constexpr (!std::is_same_v<T, U>)
 		{
-			PCGGraphParametersHelpers::OnException(Result.GetError(), PropertyName);
-			return T{};
+			return T(Result.GetValue());
 		}
 		else
 		{
-			if constexpr (std::is_pointer_v<U>)
-			{
-				static_assert(std::is_same_v<T, std::remove_pointer_t<U>>);
-				return Result.GetValue() ? *Result.GetValue() : T{};
-			}
-			else if constexpr (!std::is_same_v<T, U>)
-			{
-				return T(Result.GetValue());
-			}
-			else
-			{
-				return Result.GetValue();
-			}
+			return Result.GetValue();
 		}
 	}
 
-	// Convenient alias when T and U are the same.
 	template <typename T>
-	T GetParameter(UPCGGraphInstance* GraphInstance, FName PropertyName, TFunctionRef<TValueOrError<T, EPropertyBagResult>(FInstancedPropertyBag&)> Getter)
+	void ValidateAndSetValue(UPCGGraphInterface* GraphInterface, const FName PropertyName, const T& Value)
 	{
-		return GetParameter<T, T>(GraphInstance, PropertyName, Getter);
+		if (!GraphInterface)
+		{
+			ThrowBlueprintException(InvalidGraphInterfaceError);
+			return;
+		}
+
+		const EPropertyBagResult Result = GraphInterface->SetGraphParameter<T>(PropertyName, Value);
+		if (Result != EPropertyBagResult::Success)
+		{
+			OnException(Result, PropertyName);
+		}
 	}
 
-	// Get a parameter from the Instanced Property Bag using its setters. Value to set must be captured in the Setter.
-	// If the set succeeded, it will force this parameter to be overridden.
-	void SetParameter(UPCGGraphInstance* GraphInstance, FName PropertyName, TFunctionRef<EPropertyBagResult(FInstancedPropertyBag&)> Setter)
+	void ValidateAndSetValue(UPCGGraphInterface* GraphInterface, const FName PropertyName, const uint64& Value, const UEnum* Enum)
 	{
-		if (!GraphInstance)
+		if (!GraphInterface)
 		{
-			const FBlueprintExceptionInfo ExceptionInfo(EBlueprintExceptionType::FatalError, InvalidGraphInstanceError);
-			FBlueprintCoreDelegates::ThrowScriptException(FFrame::GetThreadLocalTopStackFrame()->Object, *FFrame::GetThreadLocalTopStackFrame(), ExceptionInfo);
-
+			ThrowBlueprintException(InvalidGraphInterfaceError);
 			return;
 		}
 
-		EPropertyBagResult Result = Setter(GraphInstance->ParametersOverrides.Parameters);
-
+		const EPropertyBagResult Result = GraphInterface->SetGraphParameter(PropertyName, Value, Enum);
 		if (Result != EPropertyBagResult::Success)
 		{
-			PCGGraphParametersHelpers::OnException(Result, PropertyName);
-		}
-		else
-		{
-			PCGGraphParametersHelpers::UpdateOverrides(GraphInstance, PropertyName);
+			OnException(Result, PropertyName);
 		}
 	}
 }
 
-bool UPCGGraphParametersHelpers::IsOverridden(UPCGGraphInstance* GraphInstance, FName InPropertyName)
-{
-	if (!GraphInstance)
-	{
-		return false;
-	}
+//////////////////////////
+// Begin Blueprint Library
+//////////////////////////
 
-	const FPropertyBagPropertyDesc* Desc = GraphInstance->ParametersOverrides.Parameters.FindPropertyDescByName(InPropertyName);
-	return Desc ? GraphInstance->ParametersOverrides.PropertiesIDsOverridden.Contains(Desc->ID) : false;
+bool UPCGGraphParametersHelpers::IsOverridden(const UPCGGraphInterface* GraphInterface, const FName Name)
+{
+	return GraphInterface && GraphInterface->IsGraphParameterOverridden(Name);
 }
 
 ////////////
 // Getters
 ////////////
 
-float UPCGGraphParametersHelpers::GetFloatParameter(UPCGGraphInstance* GraphInstance, const FName Name)
+float UPCGGraphParametersHelpers::GetFloatParameter(const UPCGGraphInterface* GraphInterface, const FName Name)
 {
-	return PCGGraphParametersHelpers::GetParameter<float>(GraphInstance, Name, [Name](FInstancedPropertyBag& Bag) { return Bag.GetValueFloat(Name); });
+	return PCGGraphParametersHelpersPrivate::ValidateAndReturnResult<float>(GraphInterface, Name);
 }
 
-double UPCGGraphParametersHelpers::GetDoubleParameter(UPCGGraphInstance* GraphInstance, const FName Name)
+double UPCGGraphParametersHelpers::GetDoubleParameter(const UPCGGraphInterface* GraphInterface, const FName Name)
 {
-	return PCGGraphParametersHelpers::GetParameter<double>(GraphInstance, Name, [Name](FInstancedPropertyBag& Bag) { return Bag.GetValueDouble(Name); });
+	return PCGGraphParametersHelpersPrivate::ValidateAndReturnResult<double>(GraphInterface, Name);
 }
 
-bool UPCGGraphParametersHelpers::GetBoolParameter(UPCGGraphInstance* GraphInstance, const FName Name)
+bool UPCGGraphParametersHelpers::GetBoolParameter(const UPCGGraphInterface* GraphInterface, const FName Name)
 {
-	return PCGGraphParametersHelpers::GetParameter<bool>(GraphInstance, Name, [Name](FInstancedPropertyBag& Bag) { return Bag.GetValueBool(Name); });
+	return PCGGraphParametersHelpersPrivate::ValidateAndReturnResult<bool>(GraphInterface, Name);
 }
 
-uint8 UPCGGraphParametersHelpers::GetByteParameter(UPCGGraphInstance* GraphInstance, const FName Name)
+uint8 UPCGGraphParametersHelpers::GetByteParameter(const UPCGGraphInterface* GraphInterface, const FName Name)
 {
-	return PCGGraphParametersHelpers::GetParameter<uint8>(GraphInstance, Name, [Name](FInstancedPropertyBag& Bag) { return Bag.GetValueByte(Name); });
+	return PCGGraphParametersHelpersPrivate::ValidateAndReturnResult<uint8>(GraphInterface, Name);
 }
 
-int32 UPCGGraphParametersHelpers::GetInt32Parameter(UPCGGraphInstance* GraphInstance, const FName Name)
+int32 UPCGGraphParametersHelpers::GetInt32Parameter(const UPCGGraphInterface* GraphInterface, const FName Name)
 {
-	return PCGGraphParametersHelpers::GetParameter<int32>(GraphInstance, Name, [Name](FInstancedPropertyBag& Bag) { return Bag.GetValueInt32(Name); });
+	return PCGGraphParametersHelpersPrivate::ValidateAndReturnResult<int32>(GraphInterface, Name);
 }
 
-int64 UPCGGraphParametersHelpers::GetInt64Parameter(UPCGGraphInstance* GraphInstance, const FName Name)
+int64 UPCGGraphParametersHelpers::GetInt64Parameter(const UPCGGraphInterface* GraphInterface, const FName Name)
 {
-	return PCGGraphParametersHelpers::GetParameter<int64>(GraphInstance, Name, [Name](FInstancedPropertyBag& Bag) { return Bag.GetValueInt64(Name); });
+	return PCGGraphParametersHelpersPrivate::ValidateAndReturnResult<int64>(GraphInterface, Name);
 }
 
-FName UPCGGraphParametersHelpers::GetNameParameter(UPCGGraphInstance* GraphInstance, const FName Name)
+FName UPCGGraphParametersHelpers::GetNameParameter(const UPCGGraphInterface* GraphInterface, const FName Name)
 {
-	return PCGGraphParametersHelpers::GetParameter<FName>(GraphInstance, Name, [Name](FInstancedPropertyBag& Bag) { return Bag.GetValueName(Name); });
+	return PCGGraphParametersHelpersPrivate::ValidateAndReturnResult<FName>(GraphInterface, Name);
 }
 
-FString UPCGGraphParametersHelpers::GetStringParameter(UPCGGraphInstance* GraphInstance, const FName Name)
+FString UPCGGraphParametersHelpers::GetStringParameter(const UPCGGraphInterface* GraphInterface, const FName Name)
 {
-	return PCGGraphParametersHelpers::GetParameter<FString>(GraphInstance, Name, [Name](FInstancedPropertyBag& Bag) { return Bag.GetValueString(Name); });
+	return PCGGraphParametersHelpersPrivate::ValidateAndReturnResult<FString>(GraphInterface, Name);
 }
 
-TSoftObjectPtr<UObject> UPCGGraphParametersHelpers::GetSoftObjectParameter(UPCGGraphInstance* GraphInstance, const FName Name)
+uint8 UPCGGraphParametersHelpers::GetEnumParameter(const UPCGGraphInterface* GraphInterface, const FName Name)
 {
-	// TODO: Uncomment when StructUtils is updated
-	// return PCGGraphParametersHelpers::GetParameter<TSoftObjectPtr<UObject>, FSoftObjectPath>(GraphInstance, Name, [Name](FInstancedPropertyBag& Bag) { return Bag.GetValueSoftPath(Name); });
-
-	// Workaround
-	return PCGGraphParametersHelpers::GetParameter<TSoftObjectPtr<UObject>>(GraphInstance, Name, [Name](FInstancedPropertyBag& Bag) -> TValueOrError<TSoftObjectPtr<UObject>, EPropertyBagResult>
-	{
-		const FPropertyBagPropertyDesc* Desc = Bag.FindPropertyDescByName(Name);
-		if (!Desc)
-		{
-			return MakeError(EPropertyBagResult::PropertyNotFound);
-		}
+	return PCGGraphParametersHelpersPrivate::ValidateAndReturnResult<uint8>(GraphInterface, Name);
+}
 
-		const FSoftObjectProperty* Property = CastField<FSoftObjectProperty>(Desc->CachedProperty);
-		if (!Property)
-		{
-			return MakeError(EPropertyBagResult::TypeMismatch);
-		}
+FSoftObjectPath UPCGGraphParametersHelpers::GetSoftObjectPathParameter(const UPCGGraphInterface* GraphInterface, const FName Name)
+{
+	return PCGGraphParametersHelpersPrivate::ValidateAndReturnResult<FSoftObjectPath>(GraphInterface, Name);
+}
 
-		const FSoftObjectPtr& SoftObjectPtr = Property->GetPropertyValue_InContainer(Bag.GetValue().GetMemory());
-		return MakeValue(TSoftObjectPtr<UObject>(SoftObjectPtr.ToSoftObjectPath()));
-	});
+TSoftObjectPtr<UObject> UPCGGraphParametersHelpers::GetSoftObjectParameter(const UPCGGraphInterface* GraphInterface, const FName Name)
+{
+	return PCGGraphParametersHelpersPrivate::ValidateAndReturnResult<TSoftObjectPtr<UObject>>(GraphInterface, Name);
 }
 
-TSoftClassPtr<UObject> UPCGGraphParametersHelpers::GetSoftClassParameter(UPCGGraphInstance* GraphInstance, const FName Name)
+TSoftClassPtr<UObject> UPCGGraphParametersHelpers::GetSoftClassParameter(const UPCGGraphInterface* GraphInterface, const FName Name)
 {
-	// TODO: Uncomment when StructUtils is updated
-	// return PCGGraphParametersHelpers::GetParameter<TSoftClassPtr<UObject>, FSoftObjectPath>(GraphInstance, Name, [Name](FInstancedPropertyBag& Bag) { return Bag.GetValueSoftPath(Name); });
+	return PCGGraphParametersHelpersPrivate::ValidateAndReturnResult<TSoftClassPtr<UObject>>(GraphInterface, Name);
+}
 
-	// Workaround
-	return PCGGraphParametersHelpers::GetParameter<TSoftClassPtr<UObject>>(GraphInstance, Name, [Name](FInstancedPropertyBag& Bag) -> TValueOrError<TSoftClassPtr<UObject>, EPropertyBagResult>
-	{
-		const FPropertyBagPropertyDesc* Desc = Bag.FindPropertyDescByName(Name);
-		if (!Desc)
-		{
-			return MakeError(EPropertyBagResult::PropertyNotFound);
-		}
+FVector UPCGGraphParametersHelpers::GetVectorParameter(const UPCGGraphInterface* GraphInterface, const FName Name)
+{
+	return PCGGraphParametersHelpersPrivate::ValidateAndReturnResult<FVector, FVector*>(GraphInterface, Name);
+}
 
-		const FSoftClassProperty* Property = CastField<FSoftClassProperty>(Desc->CachedProperty);
-		if (!Property)
-		{
-			return MakeError(EPropertyBagResult::TypeMismatch);
-		}
+FRotator UPCGGraphParametersHelpers::GetRotatorParameter(const UPCGGraphInterface* GraphInterface, const FName Name)
+{
+	return PCGGraphParametersHelpersPrivate::ValidateAndReturnResult<FRotator, FRotator*>(GraphInterface, Name);
+}
 
-		const FSoftObjectPtr& SoftObjectPtr = Property->GetPropertyValue_InContainer(Bag.GetValue().GetMemory());
-		return MakeValue(TSoftClassPtr<UObject>(SoftObjectPtr.ToSoftObjectPath()));
-	});
+FTransform UPCGGraphParametersHelpers::GetTransformParameter(const UPCGGraphInterface* GraphInterface, const FName Name)
+{
+	return PCGGraphParametersHelpersPrivate::ValidateAndReturnResult<FTransform, FTransform*>(GraphInterface, Name);
 }
 
-FVector UPCGGraphParametersHelpers::GetVectorParameter(UPCGGraphInstance* GraphInstance, const FName Name)
+FVector4 UPCGGraphParametersHelpers::GetVector4Parameter(const UPCGGraphInterface* GraphInterface, const FName Name)
 {
-	return PCGGraphParametersHelpers::GetParameter<FVector, FVector*>(GraphInstance, Name, [Name](FInstancedPropertyBag& Bag) { return Bag.GetValueStruct<FVector>(Name); });
+	return PCGGraphParametersHelpersPrivate::ValidateAndReturnResult<FVector4, FVector4*>(GraphInterface, Name);
 }
 
-FRotator UPCGGraphParametersHelpers::GetRotatorParameter(UPCGGraphInstance* GraphInstance, const FName Name)
+FVector2D UPCGGraphParametersHelpers::GetVector2DParameter(const UPCGGraphInterface* GraphInterface, const FName Name)
 {
-	return PCGGraphParametersHelpers::GetParameter<FRotator, FRotator*>(GraphInstance, Name, [Name](FInstancedPropertyBag& Bag) { return Bag.GetValueStruct<FRotator>(Name); });
+	return PCGGraphParametersHelpersPrivate::ValidateAndReturnResult<FVector2D, FVector2D*>(GraphInterface, Name);
 }
 
-FTransform UPCGGraphParametersHelpers::GetTransformParameter(UPCGGraphInstance* GraphInstance, const FName Name)
+FQuat UPCGGraphParametersHelpers::GetQuaternionParameter(const UPCGGraphInterface* GraphInterface, const FName Name)
 {
-	return PCGGraphParametersHelpers::GetParameter<FTransform, FTransform*>(GraphInstance, Name, [Name](FInstancedPropertyBag& Bag) { return Bag.GetValueStruct<FTransform>(Name); });
+	return PCGGraphParametersHelpersPrivate::ValidateAndReturnResult<FQuat, FQuat*>(GraphInterface, Name);
 }
 
 
 ////////////
 // Setters
 ////////////
-
-void UPCGGraphParametersHelpers::SetFloatParameter(UPCGGraphInstance* GraphInstance, const FName Name, const float Value)
+void UPCGGraphParametersHelpers::SetFloatParameter(UPCGGraphInterface* GraphInterface, const FName Name, const float Value)
 {
-	PCGGraphParametersHelpers::SetParameter(GraphInstance, Name, [Name, Value](FInstancedPropertyBag& Bag) { return Bag.SetValueFloat(Name, Value); });
+	PCGGraphParametersHelpersPrivate::ValidateAndSetValue(GraphInterface, Name, Value);
 }
 
-void UPCGGraphParametersHelpers::SetDoubleParameter(UPCGGraphInstance* GraphInstance, const FName Name, const double Value)
+void UPCGGraphParametersHelpers::SetDoubleParameter(UPCGGraphInterface* GraphInterface, const FName Name, const double Value)
 {
-	PCGGraphParametersHelpers::SetParameter(GraphInstance, Name, [Name, Value](FInstancedPropertyBag& Bag) { return Bag.SetValueDouble(Name, Value); });
+	PCGGraphParametersHelpersPrivate::ValidateAndSetValue(GraphInterface, Name, Value);
 }
 
-void UPCGGraphParametersHelpers::SetBoolParameter(UPCGGraphInstance* GraphInstance, const FName Name, const bool bValue)
+void UPCGGraphParametersHelpers::SetBoolParameter(UPCGGraphInterface* GraphInterface, const FName Name, const bool bValue)
 {
-	PCGGraphParametersHelpers::SetParameter(GraphInstance, Name, [Name, bValue](FInstancedPropertyBag& Bag) { return Bag.SetValueBool(Name, bValue); });
+	PCGGraphParametersHelpersPrivate::ValidateAndSetValue(GraphInterface, Name, bValue);
 }
 
-void UPCGGraphParametersHelpers::SetByteParameter(UPCGGraphInstance* GraphInstance, const FName Name, const uint8 Value)
+void UPCGGraphParametersHelpers::SetByteParameter(UPCGGraphInterface* GraphInterface, const FName Name, const uint8 Value)
 {
-	PCGGraphParametersHelpers::SetParameter(GraphInstance, Name, [Name, Value](FInstancedPropertyBag& Bag) { return Bag.SetValueByte(Name, Value); });
+	PCGGraphParametersHelpersPrivate::ValidateAndSetValue(GraphInterface, Name, Value);
 }
 
-void UPCGGraphParametersHelpers::SetInt32Parameter(UPCGGraphInstance* GraphInstance, const FName Name, const int32 Value)
+void UPCGGraphParametersHelpers::SetInt32Parameter(UPCGGraphInterface* GraphInterface, const FName Name, const int32 Value)
 {
-	PCGGraphParametersHelpers::SetParameter(GraphInstance, Name, [Name, Value](FInstancedPropertyBag& Bag) { return Bag.SetValueInt32(Name, Value); });
+	PCGGraphParametersHelpersPrivate::ValidateAndSetValue(GraphInterface, Name, Value);
 }
 
-void UPCGGraphParametersHelpers::SetInt64Parameter(UPCGGraphInstance* GraphInstance, const FName Name, const int64 Value)
+void UPCGGraphParametersHelpers::SetInt64Parameter(UPCGGraphInterface* GraphInterface, const FName Name, const int64 Value)
 {
-	PCGGraphParametersHelpers::SetParameter(GraphInstance, Name, [Name, Value](FInstancedPropertyBag& Bag) { return Bag.SetValueInt64(Name, Value); });
+	PCGGraphParametersHelpersPrivate::ValidateAndSetValue(GraphInterface, Name, Value);
 }
 
-void UPCGGraphParametersHelpers::SetNameParameter(UPCGGraphInstance* GraphInstance, const FName Name, const FName Value)
+void UPCGGraphParametersHelpers::SetNameParameter(UPCGGraphInterface* GraphInterface, const FName Name, const FName Value)
 {
-	PCGGraphParametersHelpers::SetParameter(GraphInstance, Name, [Name, Value](FInstancedPropertyBag& Bag) { return Bag.SetValueName(Name, Value); });
+	PCGGraphParametersHelpersPrivate::ValidateAndSetValue(GraphInterface, Name, Value);
 }
 
-void UPCGGraphParametersHelpers::SetStringParameter(UPCGGraphInstance* GraphInstance, const FName Name, const FString& Value)
+void UPCGGraphParametersHelpers::SetStringParameter(UPCGGraphInterface* GraphInterface, const FName Name, const FString Value)
 {
-	PCGGraphParametersHelpers::SetParameter(GraphInstance, Name, [Name, Value](FInstancedPropertyBag& Bag) { return Bag.SetValueString(Name, Value); });
+	PCGGraphParametersHelpersPrivate::ValidateAndSetValue(GraphInterface, Name, Value);
 }
 
-void UPCGGraphParametersHelpers::SetEnumParameter(UPCGGraphInstance* GraphInstance, const FName Name, const UEnum* Enum, const uint8 Value)
+void UPCGGraphParametersHelpers::SetEnumParameter(UPCGGraphInterface* GraphInterface, const FName Name, const uint8 Value, const UEnum* Enum)
 {
-	PCGGraphParametersHelpers::SetParameter(GraphInstance, Name, [Name, Value, Enum](FInstancedPropertyBag& Bag) { return Bag.SetValueEnum(Name, Value, Enum); });
+	PCGGraphParametersHelpersPrivate::ValidateAndSetValue(GraphInterface, Name, Value, Enum);
 }
 
-void UPCGGraphParametersHelpers::SetSoftObjectParameter(UPCGGraphInstance* GraphInstance, const FName Name, const TSoftObjectPtr<UObject>& Value)
+void UPCGGraphParametersHelpers::SetSoftObjectPathParameter(UPCGGraphInterface* GraphInterface, const FName Name, const FSoftObjectPath& Value)
 {
-	// TODO: Uncomment when StructUtils is updated
-	// PCGGraphParametersHelpers::SetParameter(GraphInstance, Name, [Name, Value](FInstancedPropertyBag& Bag) { return Bag.SetValueSoftPath(Name, Value.ToSoftObjectPath()); });
-
-	// Workaround
-	PCGGraphParametersHelpers::SetParameter(GraphInstance, Name, [Name, Value](FInstancedPropertyBag& Bag)
-	{
-		const FPropertyBagPropertyDesc* Desc = Bag.FindPropertyDescByName(Name);
-		if (!Desc)
-		{
-			return EPropertyBagResult::PropertyNotFound;
-		}
-
-		const FSoftObjectProperty* Property = CastField<FSoftObjectProperty>(Desc->CachedProperty);
-		if (!Property)
-		{
-			return EPropertyBagResult::TypeMismatch;
-		}
+	PCGGraphParametersHelpersPrivate::ValidateAndSetValue(GraphInterface, Name, Value);
+}
 
-		const FSoftObjectPtr SoftObjectPtr(Value.ToSoftObjectPath());
-		Property->SetPropertyValue_InContainer(Bag.GetMutableValue().GetMemory(), SoftObjectPtr);
-		return EPropertyBagResult::Success;
-	});
+void UPCGGraphParametersHelpers::SetSoftObjectParameter(UPCGGraphInterface* GraphInterface, const FName Name, const TSoftObjectPtr<UObject>& Value)
+{
+	PCGGraphParametersHelpersPrivate::ValidateAndSetValue(GraphInterface, Name, Value);
 }
 
-void UPCGGraphParametersHelpers::SetSoftClassParameter(UPCGGraphInstance* GraphInstance, const FName Name, const TSoftClassPtr<UObject>& Value)
+void UPCGGraphParametersHelpers::SetSoftClassParameter(UPCGGraphInterface* GraphInterface, const FName Name, const TSoftClassPtr<UObject>& Value)
 {
-	// TODO: Uncomment when StructUtils is updated
-	// PCGGraphParametersHelpers::SetParameter(GraphInstance, Name, [Name, Value](FInstancedPropertyBag& Bag) { return Bag.SetValueSoftPath(Name, Value.ToSoftObjectPath()); });
+	PCGGraphParametersHelpersPrivate::ValidateAndSetValue(GraphInterface, Name, Value);
+}
 
-	// Workaround
-	PCGGraphParametersHelpers::SetParameter(GraphInstance, Name, [Name, Value](FInstancedPropertyBag& Bag)
-	{
-		const FPropertyBagPropertyDesc* Desc = Bag.FindPropertyDescByName(Name);
-		if (!Desc)
-		{
-			return EPropertyBagResult::PropertyNotFound;
-		}
+void UPCGGraphParametersHelpers::SetVectorParameter(UPCGGraphInterface* GraphInterface, const FName Name, const FVector& Value)
+{
+	PCGGraphParametersHelpersPrivate::ValidateAndSetValue(GraphInterface, Name, Value);
+}
 
-		const FSoftClassProperty* Property = CastField<FSoftClassProperty>(Desc->CachedProperty);
-		if (!Property)
-		{
-			return EPropertyBagResult::TypeMismatch;
-		}
+void UPCGGraphParametersHelpers::SetRotatorParameter(UPCGGraphInterface* GraphInterface, const FName Name, const FRotator& Value)
+{
+	PCGGraphParametersHelpersPrivate::ValidateAndSetValue(GraphInterface, Name, Value);
+}
 
-		const FSoftObjectPtr SoftObjectPtr(Value.ToSoftObjectPath());
-		Property->SetPropertyValue_InContainer(Bag.GetMutableValue().GetMemory(), SoftObjectPtr);
-		return EPropertyBagResult::Success;
-	});
+void UPCGGraphParametersHelpers::SetTransformParameter(UPCGGraphInterface* GraphInterface, const FName Name, const FTransform& Value)
+{
+	PCGGraphParametersHelpersPrivate::ValidateAndSetValue(GraphInterface, Name, Value);
 }
 
-void UPCGGraphParametersHelpers::SetVectorParameter(UPCGGraphInstance* GraphInstance, const FName Name, const FVector& Value)
+void UPCGGraphParametersHelpers::SetVector4Parameter(UPCGGraphInterface* GraphInterface, const FName Name, const FVector4& Value)
 {
-	PCGGraphParametersHelpers::SetParameter(GraphInstance, Name, [Name, Value](FInstancedPropertyBag& Bag) { return Bag.SetValueStruct<FVector>(Name, Value); });
+	PCGGraphParametersHelpersPrivate::ValidateAndSetValue(GraphInterface, Name, Value);
 }
 
-void UPCGGraphParametersHelpers::SetRotatorParameter(UPCGGraphInstance* GraphInstance, const FName Name, const FRotator& Value)
+void UPCGGraphParametersHelpers::SetVector2DParameter(UPCGGraphInterface* GraphInterface, const FName Name, const FVector2D& Value)
 {
-	PCGGraphParametersHelpers::SetParameter(GraphInstance, Name, [Name, Value](FInstancedPropertyBag& Bag) { return Bag.SetValueStruct<FRotator>(Name, Value); });
+	PCGGraphParametersHelpersPrivate::ValidateAndSetValue(GraphInterface, Name, Value);
 }
 
-void UPCGGraphParametersHelpers::SetTransformParameter(UPCGGraphInstance* GraphInstance, const FName Name, const FTransform& Value)
+void UPCGGraphParametersHelpers::SetQuaternionParameter(UPCGGraphInterface* GraphInterface, const FName Name, const FQuat& Value)
 {
-	PCGGraphParametersHelpers::SetParameter(GraphInstance, Name, [Name, Value](FInstancedPropertyBag& Bag) { return Bag.SetValueStruct<FTransform>(Name, Value); });
+	PCGGraphParametersHelpersPrivate::ValidateAndSetValue(GraphInterface, Name, Value);
 }
 
 #undef LOCTEXT_NAMESPACE
\ No newline at end of file
diff --git a/Engine/Plugins/Experimental/PCG/Source/PCG/Private/PCGGraph.cpp b/Engine/Plugins/Experimental/PCG/Source/PCG/Private/PCGGraph.cpp
--- a/Engine/Plugins/Experimental/PCG/Source/PCG/Private/PCGGraph.cpp
+++ b/Engine/Plugins/Experimental/PCG/Source/PCG/Private/PCGGraph.cpp
@@ -97,6 +97,19 @@ namespace PCGGraphUtils
 * UPCGGraphInterface
 ****************************/
 
+EPropertyBagResult UPCGGraphInterface::SetGraphParameter(const FName PropertyName, const uint64 Value, const UEnum* Enum)
+{
+	FInstancedPropertyBag* UserParameters = GetMutableUserParametersStruct();
+	check(UserParameters);
+
+	const EPropertyBagResult Result = FPCGGraphParameterExtension::SetGraphParameter(*UserParameters, PropertyName, Value, Enum);
+	if (Result == EPropertyBagResult::Success)
+	{
+		OnGraphParametersChanged(EPCGGraphParameterEvent::ValueModifiedLocally, PropertyName);
+	}
+	return Result;
+}
+
 bool UPCGGraphInterface::IsInstance() const
 {
 	return this != GetGraph();
@@ -1029,6 +1042,37 @@ void UPCGGraph::PostEditChangeProperty(FPropertyChangedEvent& PropertyChangedEve
 	NumberOfUserParametersPreEdit = 0;
 }
 
+bool UPCGGraph::UserParametersCanRemoveProperty(FGuid InPropertyID, FName InPropertyName)
+{
+	// Check if the property has some getters in the graph
+	for (const UPCGNode* Node : Nodes)
+	{
+		if (!Node)
+		{
+			continue;
+		}
+
+		if (const UPCGUserParameterGetSettings* Settings = Cast<UPCGUserParameterGetSettings>(Node->GetSettings()))
+		{
+			if (Settings->PropertyGuid == InPropertyID)
+			{
+				// We found a getter. Ask the user if he is OK with that
+				FText RemoveCheckMessage = FText::Format(LOCTEXT("UserParametersRemoveCheck", "Property {0} is in use in the graph. Are you sure you want to remove it?"), FText::FromName(InPropertyName));
+				FSuppressableWarningDialog::FSetupInfo Info(RemoveCheckMessage, LOCTEXT("UserParametersRemoveCheck_Message", "Remove property"), "UserParametersRemove");
+				Info.ConfirmText = FCoreTexts::Get().Yes;
+				Info.CancelText = FCoreTexts::Get().No;
+				FSuppressableWarningDialog AddLevelWarning(Info);
+				if (AddLevelWarning.ShowModal() == FSuppressableWarningDialog::Cancel)
+				{
+					return false;
+				}
+			}
+		}
+	}
+
+	return true;
+}
+
 void UPCGGraph::OnGraphParametersChanged(EPCGGraphParameterEvent InChangeType, FName InChangedPropertyName)
 {
 	bool bWasModified = false;
@@ -1122,40 +1166,13 @@ bool UPCGGraph::UserParametersIsPinTypeAccepted(FEdGraphPinType InPinType, bool
 		return true;
 	}
 }
+#endif // WITH_EDITOR
 
-bool UPCGGraph::UserParametersCanRemoveProperty(FGuid InPropertyID, FName InPropertyName)
+FInstancedPropertyBag* UPCGGraph::GetMutableUserParametersStruct()
 {
-	// Check if the property has some getters in the graph
-	for (const UPCGNode* Node : Nodes)
-	{
-		if (!Node)
-		{
-			continue;
-		}
-
-		if (const UPCGUserParameterGetSettings* Settings = Cast<UPCGUserParameterGetSettings>(Node->GetSettings()))
-		{
-			if (Settings->PropertyGuid == InPropertyID)
-			{
-				// We found a getter. Ask the user if he is OK with that
-				FText RemoveCheckMessage = FText::Format(LOCTEXT("UserParametersRemoveCheck", "Property {0} is in use in the graph. Are you sure you want to remove it?"), FText::FromName(InPropertyName));
-				FSuppressableWarningDialog::FSetupInfo Info(RemoveCheckMessage, LOCTEXT("UserParametersRemoveCheck_Message", "Remove property"), "UserParametersRemove");
-				Info.ConfirmText = FCoreTexts::Get().Yes;
-				Info.CancelText = FCoreTexts::Get().No;
-				FSuppressableWarningDialog AddLevelWarning(Info);
-				if (AddLevelWarning.ShowModal() == FSuppressableWarningDialog::Cancel)
-				{
-					return false;
-				}
-			}
-		}
-	}
-
-	return true;
+	return &UserParameters;
 }
 
-#endif // WITH_EDITOR
-
 uint32 UPCGGraph::GetNodeGenerationGridSize(const UPCGNode* InNode, uint32 InDefaultGridSize) const
 {
 	{
@@ -1231,9 +1248,7 @@ void UPCGGraph::AddUserParameters(const TArray<FPropertyBagPropertyDesc>& InDesc
 		}
 	}
 
-#if WITH_EDITOR
 	OnGraphParametersChanged(EPCGGraphParameterEvent::MultiplePropertiesAdded, NAME_None);
-#endif // WITH_EDITOR
 }
 
 /****************************
@@ -1451,6 +1466,7 @@ void UPCGGraphInstance::NotifyGraphParametersChanged(EPCGGraphParameterEvent InC
 	// Also propagates the changes
 	OnGraphChanged(Graph, EPCGChangeType::Settings);
 }
+#endif // WITH_EDITOR
 
 void UPCGGraphInstance::OnGraphParametersChanged(UPCGGraphInterface* InGraph, EPCGGraphParameterEvent InChangeType, FName InChangedPropertyName)
 {
@@ -1467,9 +1483,20 @@ void UPCGGraphInstance::OnGraphParametersChanged(UPCGGraphInterface* InGraph, EP
 	}
 
 	RefreshParameters(ChangeType, InChangedPropertyName);
+#if WITH_EDITOR
 	NotifyGraphParametersChanged(ChangeType, InChangedPropertyName);
-}
 #endif // WITH_EDITOR
+}
+
+void UPCGGraphInstance::OnGraphParametersChanged(EPCGGraphParameterEvent InChangeType, FName InChangedPropertyName)
+{
+	OnGraphParametersChanged(this, InChangeType, InChangedPropertyName);
+}
+
+FInstancedPropertyBag* UPCGGraphInstance::GetMutableUserParametersStruct()
+{
+	return &ParametersOverrides.Parameters;
+}
 
 void UPCGGraphInstance::RefreshParameters(EPCGGraphParameterEvent InChangeType, FName InChangedPropertyName)
 {
@@ -1556,6 +1583,11 @@ bool UPCGGraphInstance::IsPropertyOverriddenAndNotDefault(const FProperty* InPro
 	return Graph ? ParametersOverrides.IsPropertyOverriddenAndNotDefault(InProperty, Graph->GetUserParametersStruct()) : false;
 }
 
+bool UPCGGraphInstance::IsGraphParameterOverridden(const FName PropertyName) const
+{
+	return (ParametersOverrides.Parameters.FindPropertyDescByName(PropertyName) != nullptr); 
+}
+
 bool FPCGOverrideInstancedPropertyBag::RefreshParameters(const FInstancedPropertyBag* ParentUserParameters, EPCGGraphParameterEvent InChangeType, FName InChangedPropertyName)
 {
 	check(ParentUserParameters);
diff --git a/Engine/Plugins/Experimental/PCG/Source/PCG/Public/Helpers/PCGGraphParameterExtension.h b/Engine/Plugins/Experimental/PCG/Source/PCG/Public/Helpers/PCGGraphParameterExtension.h
new file mode 100644
--- /dev/null
+++ b/Engine/Plugins/Experimental/PCG/Source/PCG/Public/Helpers/PCGGraphParameterExtension.h
@@ -0,0 +1,59 @@
+﻿// Copyright Epic Games, Inc. All Rights Reserved.
+
+#pragma once
+
+#include "PropertyBag.h"
+
+/** Extension helper class to allow templatization without dirtying the PCGGraph.h file */
+class PCG_API FPCGGraphParameterExtension
+{
+	template <typename T> static bool constexpr StaticFail{false};
+
+public:
+	// Catch all implicit templates for unsupported types
+	template <typename T>
+	static TValueOrError<T, EPropertyBagResult> GetGraphParameter(const FInstancedPropertyBag& PropertyBag, const FName PropertyName) { static_assert(StaticFail<T>, "Invalid graph parameter type"); return MakeError(EPropertyBagResult::PropertyNotFound); }
+	template <typename T>
+	static EPropertyBagResult SetGraphParameter(FInstancedPropertyBag& PropertyBag, const FName PropertyName, const T& Value) { static_assert(StaticFail<T>, "Invalid graph parameter type"); return EPropertyBagResult::PropertyNotFound; }
+
+	// Enum Objects are a special signature
+	static TValueOrError<uint64, EPropertyBagResult> GetGraphParameter(const FInstancedPropertyBag& PropertyBag, const FName PropertyName, const UEnum* Enum);
+	static EPropertyBagResult SetGraphParameter(FInstancedPropertyBag& PropertyBag, const FName PropertyName, const uint64& Value, const UEnum* Enum);
+};
+
+// Template declarations necessary for DLL API
+template<> PCG_API TValueOrError<float, EPropertyBagResult> FPCGGraphParameterExtension::GetGraphParameter(const FInstancedPropertyBag& PropertyBag, const FName PropertyName);
+template<> PCG_API TValueOrError<double, EPropertyBagResult> FPCGGraphParameterExtension::GetGraphParameter(const FInstancedPropertyBag& PropertyBag, const FName PropertyName);
+template<> PCG_API TValueOrError<bool, EPropertyBagResult> FPCGGraphParameterExtension::GetGraphParameter(const FInstancedPropertyBag& PropertyBag, const FName PropertyName);
+template<> PCG_API TValueOrError<uint8, EPropertyBagResult> FPCGGraphParameterExtension::GetGraphParameter(const FInstancedPropertyBag& PropertyBag, const FName PropertyName);
+template<> PCG_API TValueOrError<int32, EPropertyBagResult> FPCGGraphParameterExtension::GetGraphParameter(const FInstancedPropertyBag& PropertyBag, const FName PropertyName);
+template<> PCG_API TValueOrError<int64, EPropertyBagResult> FPCGGraphParameterExtension::GetGraphParameter(const FInstancedPropertyBag& PropertyBag, const FName PropertyName);
+template<> PCG_API TValueOrError<FName, EPropertyBagResult> FPCGGraphParameterExtension::GetGraphParameter(const FInstancedPropertyBag& PropertyBag, const FName PropertyName);
+template<> PCG_API TValueOrError<FString, EPropertyBagResult> FPCGGraphParameterExtension::GetGraphParameter(const FInstancedPropertyBag& PropertyBag, const FName PropertyName);
+template<> PCG_API TValueOrError<FSoftObjectPath, EPropertyBagResult> FPCGGraphParameterExtension::GetGraphParameter(const FInstancedPropertyBag& PropertyBag, const FName PropertyName);
+template<> PCG_API TValueOrError<TSoftObjectPtr<UObject>, EPropertyBagResult> FPCGGraphParameterExtension::GetGraphParameter(const FInstancedPropertyBag& PropertyBag, const FName PropertyName);
+template<> PCG_API TValueOrError<TSoftClassPtr<UObject>, EPropertyBagResult> FPCGGraphParameterExtension::GetGraphParameter(const FInstancedPropertyBag& PropertyBag, const FName PropertyName);
+template<> PCG_API TValueOrError<FVector*, EPropertyBagResult> FPCGGraphParameterExtension::GetGraphParameter(const FInstancedPropertyBag& PropertyBag, const FName PropertyName);
+template<> PCG_API TValueOrError<FRotator*, EPropertyBagResult> FPCGGraphParameterExtension::GetGraphParameter(const FInstancedPropertyBag& PropertyBag, const FName PropertyName);
+template<> PCG_API TValueOrError<FTransform*, EPropertyBagResult> FPCGGraphParameterExtension::GetGraphParameter(const FInstancedPropertyBag& PropertyBag, const FName PropertyName);
+template<> PCG_API TValueOrError<FVector4*, EPropertyBagResult> FPCGGraphParameterExtension::GetGraphParameter(const FInstancedPropertyBag& PropertyBag, const FName PropertyName);
+template<> PCG_API TValueOrError<FVector2D*, EPropertyBagResult> FPCGGraphParameterExtension::GetGraphParameter(const FInstancedPropertyBag& PropertyBag, const FName PropertyName);
+template<> PCG_API TValueOrError<FQuat*, EPropertyBagResult> FPCGGraphParameterExtension::GetGraphParameter(const FInstancedPropertyBag& PropertyBag, const FName PropertyName);
+
+template<> PCG_API EPropertyBagResult FPCGGraphParameterExtension::SetGraphParameter(FInstancedPropertyBag& PropertyBag, const FName PropertyName, const float& Value);
+template<> PCG_API EPropertyBagResult FPCGGraphParameterExtension::SetGraphParameter(FInstancedPropertyBag& PropertyBag, const FName PropertyName, const double& Value);
+template<> PCG_API EPropertyBagResult FPCGGraphParameterExtension::SetGraphParameter(FInstancedPropertyBag& PropertyBag, const FName PropertyName, const bool& bValue);
+template<> PCG_API EPropertyBagResult FPCGGraphParameterExtension::SetGraphParameter(FInstancedPropertyBag& PropertyBag, const FName PropertyName, const uint8& Value);
+template<> PCG_API EPropertyBagResult FPCGGraphParameterExtension::SetGraphParameter(FInstancedPropertyBag& PropertyBag, const FName PropertyName, const int32& Value);
+template<> PCG_API EPropertyBagResult FPCGGraphParameterExtension::SetGraphParameter(FInstancedPropertyBag& PropertyBag, const FName PropertyName, const int64& Value);
+template<> PCG_API EPropertyBagResult FPCGGraphParameterExtension::SetGraphParameter(FInstancedPropertyBag& PropertyBag, const FName PropertyName, const FName& Value);
+template<> PCG_API EPropertyBagResult FPCGGraphParameterExtension::SetGraphParameter(FInstancedPropertyBag& PropertyBag, const FName PropertyName, const FString& Value);
+template<> PCG_API EPropertyBagResult FPCGGraphParameterExtension::SetGraphParameter(FInstancedPropertyBag& PropertyBag, const FName PropertyName, const FSoftObjectPath& Value);
+template<> PCG_API EPropertyBagResult FPCGGraphParameterExtension::SetGraphParameter(FInstancedPropertyBag& PropertyBag, const FName PropertyName, const TSoftObjectPtr<UObject>& Value);
+template<> PCG_API EPropertyBagResult FPCGGraphParameterExtension::SetGraphParameter(FInstancedPropertyBag& PropertyBag, const FName PropertyName, const TSoftClassPtr<UObject>& Value);
+template<> PCG_API EPropertyBagResult FPCGGraphParameterExtension::SetGraphParameter(FInstancedPropertyBag& PropertyBag, const FName PropertyName, const FVector& Value);
+template<> PCG_API EPropertyBagResult FPCGGraphParameterExtension::SetGraphParameter(FInstancedPropertyBag& PropertyBag, const FName PropertyName, const FRotator& Value);
+template<> PCG_API EPropertyBagResult FPCGGraphParameterExtension::SetGraphParameter(FInstancedPropertyBag& PropertyBag, const FName PropertyName, const FTransform& Value);
+template<> PCG_API EPropertyBagResult FPCGGraphParameterExtension::SetGraphParameter(FInstancedPropertyBag& PropertyBag, const FName PropertyName, const FVector4& Value);
+template<> PCG_API EPropertyBagResult FPCGGraphParameterExtension::SetGraphParameter(FInstancedPropertyBag& PropertyBag, const FName PropertyName, const FVector2D& Value);
+template<> PCG_API EPropertyBagResult FPCGGraphParameterExtension::SetGraphParameter(FInstancedPropertyBag& PropertyBag, const FName PropertyName, const FQuat& Value);
\ No newline at end of file
diff --git a/Engine/Plugins/Experimental/PCG/Source/PCG/Public/Helpers/PCGGraphParametersHelpers.h b/Engine/Plugins/Experimental/PCG/Source/PCG/Public/Helpers/PCGGraphParametersHelpers.h
--- a/Engine/Plugins/Experimental/PCG/Source/PCG/Public/Helpers/PCGGraphParametersHelpers.h
+++ b/Engine/Plugins/Experimental/PCG/Source/PCG/Public/Helpers/PCGGraphParametersHelpers.h
@@ -6,10 +6,10 @@
 
 #include "PCGGraphParametersHelpers.generated.h"
 
-class UPCGGraphInstance;
+class UPCGGraphInterface;
 
 /**
-* Helpers to dynamically set parameters on graph instances.
+* Blueprint Library to get or set graph parameters on graphs and graph instances
 */
 UCLASS()
 class PCG_API UPCGGraphParametersHelpers : public UBlueprintFunctionLibrary
@@ -18,94 +18,121 @@ class PCG_API UPCGGraphParametersHelpers : public UBlueprintFunctionLibrary
 
 public:
 	UFUNCTION(BlueprintCallable, Category = "PCG|Graph Parameters")
-	static bool IsOverridden(UPCGGraphInstance* GraphInstance, const FName Name);
+	static bool IsOverridden(const UPCGGraphInterface* GraphInterface, const FName Name);
 
 	////////////
 	// Getters
 	////////////
 
 	UFUNCTION(BlueprintCallable, Category = "PCG|Graph Parameters")
-	static float GetFloatParameter(UPCGGraphInstance* GraphInstance, const FName Name);
+	static float GetFloatParameter(const UPCGGraphInterface* GraphInterface, const FName Name);
 
 	UFUNCTION(BlueprintCallable, Category = "PCG|Graph Parameters")
-	static double GetDoubleParameter(UPCGGraphInstance* GraphInstance, const FName Name);
+	static double GetDoubleParameter(const UPCGGraphInterface* GraphInterface, const FName Name);
 
 	UFUNCTION(BlueprintCallable, Category = "PCG|Graph Parameters")
-	static bool GetBoolParameter(UPCGGraphInstance* GraphInstance, const FName Name);
+	static bool GetBoolParameter(const UPCGGraphInterface* GraphInterface, const FName Name);
 
 	UFUNCTION(BlueprintCallable, Category = "PCG|Graph Parameters")
-	static uint8 GetByteParameter(UPCGGraphInstance* GraphInstance, const FName Name);
+	static uint8 GetByteParameter(const UPCGGraphInterface* GraphInterface, const FName Name);
 
 	UFUNCTION(BlueprintCallable, Category = "PCG|Graph Parameters")
-	static int32 GetInt32Parameter(UPCGGraphInstance* GraphInstance, const FName Name);
+	static int32 GetInt32Parameter(const UPCGGraphInterface* GraphInterface, const FName Name);
 
 	UFUNCTION(BlueprintCallable, Category = "PCG|Graph Parameters")
-	static int64 GetInt64Parameter(UPCGGraphInstance* GraphInstance, const FName Name);
+	static int64 GetInt64Parameter(const UPCGGraphInterface* GraphInterface, const FName Name);
 
 	UFUNCTION(BlueprintCallable, Category = "PCG|Graph Parameters")
-	static FName GetNameParameter(UPCGGraphInstance* GraphInstance, const FName Name);
+	static FName GetNameParameter(const UPCGGraphInterface* GraphInterface, const FName Name);
 
 	UFUNCTION(BlueprintCallable, Category = "PCG|Graph Parameters")
-	static FString GetStringParameter(UPCGGraphInstance* GraphInstance, const FName Name);
+	static FString GetStringParameter(const UPCGGraphInterface* GraphInterface, const FName Name);
 
 	UFUNCTION(BlueprintCallable, Category = "PCG|Graph Parameters")
-	static TSoftObjectPtr<UObject> GetSoftObjectParameter(UPCGGraphInstance* GraphInstance, const FName Name);
+	static uint8 GetEnumParameter(const UPCGGraphInterface* GraphInterface, const FName Name);
 
 	UFUNCTION(BlueprintCallable, Category = "PCG|Graph Parameters")
-	static TSoftClassPtr<UObject> GetSoftClassParameter(UPCGGraphInstance* GraphInstance, const FName Name);
+	static TSoftObjectPtr<UObject> GetSoftObjectParameter(const UPCGGraphInterface* GraphInterface, const FName Name);
 
 	UFUNCTION(BlueprintCallable, Category = "PCG|Graph Parameters")
-	static FVector GetVectorParameter(UPCGGraphInstance* GraphInstance, const FName Name);
+	static TSoftClassPtr<UObject> GetSoftClassParameter(const UPCGGraphInterface* GraphInterface, const FName Name);
 
 	UFUNCTION(BlueprintCallable, Category = "PCG|Graph Parameters")
-	static FRotator GetRotatorParameter(UPCGGraphInstance* GraphInstance, const FName Name);
+	static FSoftObjectPath GetSoftObjectPathParameter(const UPCGGraphInterface* GraphInterface, const FName Name);
 
 	UFUNCTION(BlueprintCallable, Category = "PCG|Graph Parameters")
-	static FTransform GetTransformParameter(UPCGGraphInstance* GraphInstance, const FName Name);
+	static FVector GetVectorParameter(const UPCGGraphInterface* GraphInterface, const FName Name);
+
+	UFUNCTION(BlueprintCallable, Category = "PCG|Graph Parameters")
+	static FRotator GetRotatorParameter(const UPCGGraphInterface* GraphInterface, const FName Name);
+
+	UFUNCTION(BlueprintCallable, Category = "PCG|Graph Parameters")
+	static FTransform GetTransformParameter(const UPCGGraphInterface* GraphInterface, const FName Name);
+
+	UFUNCTION(BlueprintCallable, Category = "PCG|Graph Parameters")
+	static FVector4 GetVector4Parameter(const UPCGGraphInterface* GraphInterface, const FName Name);
+
+	UFUNCTION(BlueprintCallable, Category = "PCG|Graph Parameters")
+	static FVector2D GetVector2DParameter(const UPCGGraphInterface* GraphInterface, const FName Name);
+
+	UFUNCTION(BlueprintCallable, Category = "PCG|Graph Parameters")
+	static FQuat GetQuaternionParameter(const UPCGGraphInterface* GraphInterface, const FName Name);
 
 	////////////
 	// Setters
 	////////////
 
 	UFUNCTION(BlueprintCallable, Category = "PCG|Graph Parameters")
-	static void SetFloatParameter(UPCGGraphInstance* GraphInstance, const FName Name, const float Value);
+	static void SetFloatParameter(UPCGGraphInterface* GraphInterface, const FName Name, const float Value);
+
+	UFUNCTION(BlueprintCallable, Category = "PCG|Graph Parameters")
+	static void SetDoubleParameter(UPCGGraphInterface* GraphInterface, const FName Name, const double Value);
+
+	UFUNCTION(BlueprintCallable, Category = "PCG|Graph Parameters")
+	static void SetBoolParameter(UPCGGraphInterface* GraphInterface, const FName Name, const bool bValue);
+
+	UFUNCTION(BlueprintCallable, Category = "PCG|Graph Parameters")
+	static void SetByteParameter(UPCGGraphInterface* GraphInterface, const FName Name, const uint8 Value);
+
+	UFUNCTION(BlueprintCallable, Category = "PCG|Graph Parameters")
+	static void SetInt32Parameter(UPCGGraphInterface* GraphInterface, const FName Name, const int32 Value);
 
 	UFUNCTION(BlueprintCallable, Category = "PCG|Graph Parameters")
-	static void SetDoubleParameter(UPCGGraphInstance* GraphInstance, const FName Name, const double Value);
+	static void SetInt64Parameter(UPCGGraphInterface* GraphInterface, const FName Name, const int64 Value);
 
 	UFUNCTION(BlueprintCallable, Category = "PCG|Graph Parameters")
-	static void SetBoolParameter(UPCGGraphInstance* GraphInstance, const FName Name, const bool bValue);
+	static void SetNameParameter(UPCGGraphInterface* GraphInterface, const FName Name, const FName Value);
 
 	UFUNCTION(BlueprintCallable, Category = "PCG|Graph Parameters")
-	static void SetByteParameter(UPCGGraphInstance* GraphInstance, const FName Name, const uint8 Value);
+	static void SetStringParameter(UPCGGraphInterface* GraphInterface, const FName Name, const FString Value);
 
 	UFUNCTION(BlueprintCallable, Category = "PCG|Graph Parameters")
-	static void SetInt32Parameter(UPCGGraphInstance* GraphInstance, const FName Name, const int32 Value);
+	static void SetEnumParameter(UPCGGraphInterface* GraphInterface, const FName Name, const uint8 Value, const UEnum* Enum = nullptr);
 
 	UFUNCTION(BlueprintCallable, Category = "PCG|Graph Parameters")
-	static void SetInt64Parameter(UPCGGraphInstance* GraphInstance, const FName Name, const int64 Value);
+	static void SetSoftObjectPathParameter(UPCGGraphInterface* GraphInterface, const FName Name, UPARAM(ref) const FSoftObjectPath& Value);
 
 	UFUNCTION(BlueprintCallable, Category = "PCG|Graph Parameters")
-	static void SetNameParameter(UPCGGraphInstance* GraphInstance, const FName Name, const FName Value);
+	static void SetSoftObjectParameter(UPCGGraphInterface* GraphInterface, const FName Name, UPARAM(ref) const TSoftObjectPtr<UObject>& Value);
 
 	UFUNCTION(BlueprintCallable, Category = "PCG|Graph Parameters")
-	static void SetStringParameter(UPCGGraphInstance* GraphInstance, const FName Name, const FString& Value);
+	static void SetSoftClassParameter(UPCGGraphInterface* GraphInterface, const FName Name, UPARAM(ref) const TSoftClassPtr<UObject>& Value);
 
 	UFUNCTION(BlueprintCallable, Category = "PCG|Graph Parameters")
-	static void SetEnumParameter(UPCGGraphInstance* GraphInstance, const FName Name, const UEnum* Enum, const uint8 Value);
+	static void SetVectorParameter(UPCGGraphInterface* GraphInterface, const FName Name, UPARAM(ref) const FVector& Value);
 
 	UFUNCTION(BlueprintCallable, Category = "PCG|Graph Parameters")
-	static void SetSoftObjectParameter(UPCGGraphInstance* GraphInstance, const FName Name, UPARAM(ref) const TSoftObjectPtr<UObject>& Value);
+	static void SetRotatorParameter(UPCGGraphInterface* GraphInterface, const FName Name, UPARAM(ref) const FRotator& Value);
 
 	UFUNCTION(BlueprintCallable, Category = "PCG|Graph Parameters")
-	static void SetSoftClassParameter(UPCGGraphInstance* GraphInstance, const FName Name, UPARAM(ref) const TSoftClassPtr<UObject>& Value);
+	static void SetTransformParameter(UPCGGraphInterface* GraphInterface, const FName Name, UPARAM(ref) const FTransform& Value);
 
 	UFUNCTION(BlueprintCallable, Category = "PCG|Graph Parameters")
-	static void SetVectorParameter(UPCGGraphInstance* GraphInstance, const FName Name, UPARAM(ref) const FVector& Value);
+	static void SetVector4Parameter(UPCGGraphInterface* GraphInterface, const FName Name, UPARAM(ref) const FVector4& Value);
 
 	UFUNCTION(BlueprintCallable, Category = "PCG|Graph Parameters")
-	static void SetRotatorParameter(UPCGGraphInstance* GraphInstance, const FName Name, UPARAM(ref) const FRotator& Value);
+	static void SetVector2DParameter(UPCGGraphInterface* GraphInterface, const FName Name, UPARAM(ref) const FVector2D& Value);
 
 	UFUNCTION(BlueprintCallable, Category = "PCG|Graph Parameters")
-	static void SetTransformParameter(UPCGGraphInstance* GraphInstance, const FName Name, UPARAM(ref) const FTransform& Value);
+	static void SetQuaternionParameter(UPCGGraphInterface* GraphInterface, const FName Name, UPARAM(ref) const FQuat& Value);
 };
diff --git a/Engine/Plugins/Experimental/PCG/Source/PCG/Public/PCGGraph.h b/Engine/Plugins/Experimental/PCG/Source/PCG/Public/PCGGraph.h
--- a/Engine/Plugins/Experimental/PCG/Source/PCG/Public/PCGGraph.h
+++ b/Engine/Plugins/Experimental/PCG/Source/PCG/Public/PCGGraph.h
@@ -6,12 +6,14 @@
 #include "PCGNode.h"
 #include "PCGSettings.h"
 #include "Graph/PCGStackContext.h"
+#include "Helpers/PCGGraphParameterExtension.h"
 
 #include "PropertyBag.h"
 #include "UObject/ObjectPtr.h"
 
 #include "PCGGraph.generated.h"
 
+class UPCGGraphInterface;
 #if WITH_EDITOR
 struct FEdGraphPinType;
 #endif // WITH_EDITOR
@@ -111,6 +113,54 @@ public:
 	FOnPCGGraphParametersChanged OnGraphParametersChangedDelegate;
 	FOnPCGGraphDynamicallyExecuted OnGraphDynamicallyExecutedDelegate;
 #endif // WITH_EDITOR
+
+	template <typename T>
+	TValueOrError<T, EPropertyBagResult> GetGraphParameter(const FName PropertyName) const
+	{
+		const FInstancedPropertyBag* UserParameters = GetUserParametersStruct();
+		check(UserParameters);
+
+		if constexpr (std::is_enum_v<T> && StaticEnum<T>())
+		{
+			return FPCGGraphParameterExtension::GetGraphParameter<T>(*UserParameters, PropertyName, StaticEnum<T>());
+		}
+		else
+		{
+			return FPCGGraphParameterExtension::GetGraphParameter<T>(*UserParameters, PropertyName);
+		}
+	}
+
+	virtual bool IsGraphParameterOverridden(const FName PropertyName) const PURE_VIRTUAL(UPCGGraphInterface::IsGraphParameterOverridden, return false;)
+
+	template <typename T>
+	EPropertyBagResult SetGraphParameter(const FName PropertyName, const T& Value)
+	{
+		FInstancedPropertyBag* UserParameters = GetMutableUserParametersStruct();
+		check(UserParameters);
+
+		EPropertyBagResult Result;
+		if constexpr (std::is_enum_v<T> && StaticEnum<T>())
+		{
+			Result = FPCGGraphParameterExtension::SetGraphParameter(*UserParameters, PropertyName, Value, StaticEnum<T>());
+		}
+		else
+		{
+			Result = FPCGGraphParameterExtension::SetGraphParameter<T>(*UserParameters, PropertyName, Value);
+		}
+
+		if (Result == EPropertyBagResult::Success)
+		{
+			OnGraphParametersChanged(EPCGGraphParameterEvent::ValueModifiedLocally, PropertyName);
+		}
+
+		return Result;
+	}
+
+	EPropertyBagResult SetGraphParameter(const FName PropertyName, const uint64 Value, const UEnum* Enum);
+
+protected:
+	virtual void OnGraphParametersChanged(EPCGGraphParameterEvent InChangeType, FName InChangedPropertyName) PURE_VIRTUAL(UPCGGraphInterface::OnGraphParametersChanged, )
+	virtual FInstancedPropertyBag* GetMutableUserParametersStruct() PURE_VIRTUAL(UPCGGraphInterface::GetMutableUserParametersStruct, return nullptr;)
 };
 
 UCLASS(BlueprintType, ClassGroup = (Procedural), hidecategories=(Object))
@@ -295,6 +345,8 @@ protected:
 	bool UserParametersCanRemoveProperty(FGuid InPropertyID, FName InPropertyName);
 #endif // WITH_EDITOR
 
+	virtual FInstancedPropertyBag* GetMutableUserParametersStruct() override;
+
 	UPROPERTY(EditAnywhere, Category = Settings)
 	bool bUseHierarchicalGeneration = false;
 
@@ -323,12 +375,13 @@ public:
 	void AddUserParameters(const TArray<FPropertyBagPropertyDesc>& InDescs, const UPCGGraph* InOptionalOriginalGraph = nullptr);
 
 #if WITH_EDITOR
+protected:
+	virtual void OnGraphParametersChanged(EPCGGraphParameterEvent InChangeType, FName InChangedPropertyName) override;
 private:
 	void NotifyGraphChanged(EPCGChangeType ChangeType);
 	void OnNodeChanged(UPCGNode* InNode, EPCGChangeType ChangeType);
 
 	void NotifyGraphParametersChanged(EPCGGraphParameterEvent InChangeType, FName InChangedPropertyName);
-	void OnGraphParametersChanged(EPCGGraphParameterEvent InChangeType, FName InChangedPropertyName);
 
 	int32 GraphChangeNotificationsDisableCounter = 0;
 	bool bDelayedChangeNotification = false;
@@ -373,9 +426,11 @@ protected:
 	void OnGraphChanged(UPCGGraphInterface* InGraph, EPCGChangeType ChangeType);
 	void OnGraphGridSizesChanged(UPCGGraphInterface* InGraph);
 	void NotifyGraphParametersChanged(EPCGGraphParameterEvent InChangeType, FName InChangedPropertyName);
-	void OnGraphParametersChanged(UPCGGraphInterface* InGraph, EPCGGraphParameterEvent InChangeType, FName InChangedPropertyName);
 #endif
+	void OnGraphParametersChanged(UPCGGraphInterface* InGraph, EPCGGraphParameterEvent InChangeType, FName InChangedPropertyName);
+	virtual void OnGraphParametersChanged(EPCGGraphParameterEvent InChangeType, FName InChangedPropertyName) override;
 	void RefreshParameters(EPCGGraphParameterEvent InChangeType, FName InChangedPropertyName = NAME_None);
+	virtual FInstancedPropertyBag* GetMutableUserParametersStruct() override;
 
 public:
 	static TObjectPtr<UPCGGraphInterface> CreateInstance(UObject* InOwner, UPCGGraphInterface* InGraph);
@@ -394,6 +449,7 @@ public:
 	FPCGOverrideInstancedPropertyBag ParametersOverrides;
 
 	virtual const FInstancedPropertyBag* GetUserParametersStruct() const override { return &ParametersOverrides.Parameters; }
+	virtual bool IsGraphParameterOverridden(const FName PropertyName) const override;
 
 private:
 #if WITH_EDITORONLY_DATA
