[PCG] Enforce bAllowMultipleConnections to false if AllowMultipleData is false on pins.

* Also rename AllowMultipleConnections function to AllowsMultipleConnections.
* Also rename AllowMultipleData function to AllowsMultipleData, for consistency.

#rnx

diff --git a/Engine/Plugins/Experimental/PCG/Source/PCG/Private/Elements/PCGDifferenceElement.cpp b/Engine/Plugins/Experimental/PCG/Source/PCG/Private/Elements/PCGDifferenceElement.cpp
--- a/Engine/Plugins/Experimental/PCG/Source/PCG/Private/Elements/PCGDifferenceElement.cpp
+++ b/Engine/Plugins/Experimental/PCG/Source/PCG/Private/Elements/PCGDifferenceElement.cpp
@@ -31,7 +31,7 @@ void UPCGDifferenceSettings::ApplyStructuralDeprecation(UPCGNode* InOutNode)
 
 		// Check if single connection AND single data, to avoid adding a union if not necessary
 		const int32 SourceEdgeCount = SourcePin->EdgeCount();
-		if (SourceEdgeCount > 1 || (SourceEdgeCount == 1 && !SourcePin->Edges.IsEmpty() && SourcePin->Edges[0]->InputPin->AllowMultipleData()))
+		if (SourceEdgeCount > 1 || (SourceEdgeCount == 1 && !SourcePin->Edges.IsEmpty() && SourcePin->Edges[0]->InputPin->AllowsMultipleData()))
 		{
 			// To match previous default behavior, we'll add a union node preceding the source inputs
 			UPCGGraph* PCGGraph = InOutNode->GetGraph();
diff --git a/Engine/Plugins/Experimental/PCG/Source/PCG/Private/Elements/PCGExecuteBlueprint.cpp b/Engine/Plugins/Experimental/PCG/Source/PCG/Private/Elements/PCGExecuteBlueprint.cpp
--- a/Engine/Plugins/Experimental/PCG/Source/PCG/Private/Elements/PCGExecuteBlueprint.cpp
+++ b/Engine/Plugins/Experimental/PCG/Source/PCG/Private/Elements/PCGExecuteBlueprint.cpp
@@ -120,18 +120,6 @@ void UPCGBlueprintElement::PostEditChangeProperty(FPropertyChangedEvent& Propert
 {
 	Super::PostEditChangeProperty(PropertyChangedEvent);
 
-	// Disable multiple connections on single data pins
-	if (PropertyChangedEvent.Property && PropertyChangedEvent.Property->GetFName() == GET_MEMBER_NAME_CHECKED(FPCGPinProperties, bAllowMultipleData))
-	{
-		for (FPCGPinProperties& Pin : CustomInputPins)
-		{
-			if (!Pin.bAllowMultipleData)
-			{
-				Pin.bAllowMultipleConnections = false;
-			}
-		}
-	}
-
 	// Since we don't really know what changed, let's just rebuild our data dependencies
 	DataDependencies = PCGBlueprintHelper::GetDataDependencies(this, DependencyParsingDepth);
 
diff --git a/Engine/Plugins/Experimental/PCG/Source/PCG/Private/Elements/PCGReroute.cpp b/Engine/Plugins/Experimental/PCG/Source/PCG/Private/Elements/PCGReroute.cpp
--- a/Engine/Plugins/Experimental/PCG/Source/PCG/Private/Elements/PCGReroute.cpp
+++ b/Engine/Plugins/Experimental/PCG/Source/PCG/Private/Elements/PCGReroute.cpp
@@ -42,7 +42,7 @@ TArray<FPCGPinProperties> UPCGRerouteSettings::InputPinProperties() const
 {
 	FPCGPinProperties PinProperties;
 	PinProperties.Label = PCGPinConstants::DefaultInputLabel;
-	PinProperties.bAllowMultipleConnections = false;
+	PinProperties.SetAllowMultipleConnections(false);
 	PinProperties.AllowedTypes = EPCGDataType::Any;
 
 	return { PinProperties };
diff --git a/Engine/Plugins/Experimental/PCG/Source/PCG/Private/Helpers/PCGSettingsHelpers.cpp b/Engine/Plugins/Experimental/PCG/Source/PCG/Private/Helpers/PCGSettingsHelpers.cpp
--- a/Engine/Plugins/Experimental/PCG/Source/PCG/Private/Helpers/PCGSettingsHelpers.cpp
+++ b/Engine/Plugins/Experimental/PCG/Source/PCG/Private/Helpers/PCGSettingsHelpers.cpp
@@ -246,8 +246,8 @@ namespace PCGSettingsHelpers
 		NewParamsPin->Node = InOutNode;
 		NewParamsPin->Properties.AllowedTypes = EPCGDataType::Param;
 		NewParamsPin->Properties.Label = PCGPinConstants::DefaultParamsLabel;
-		NewParamsPin->Properties.bAllowMultipleConnections = true;
 		NewParamsPin->Properties.bAllowMultipleData = true;
+		NewParamsPin->Properties.SetAllowMultipleConnections(true);
 		InputPins.Add(NewParamsPin);
 
 		// Make list of param pins that In pin is currently connected to.
diff --git a/Engine/Plugins/Experimental/PCG/Source/PCG/Private/PCGGraph.cpp b/Engine/Plugins/Experimental/PCG/Source/PCG/Private/PCGGraph.cpp
--- a/Engine/Plugins/Experimental/PCG/Source/PCG/Private/PCGGraph.cpp
+++ b/Engine/Plugins/Experimental/PCG/Source/PCG/Private/PCGGraph.cpp
@@ -541,7 +541,7 @@ bool UPCGGraph::AddLabeledEdge(UPCGNode* From, const FName& FromPinLabel, UPCGNo
 	bool bToPinBrokeOtherEdges = false;
 	
 	// Add an edge to a pin that doesn't allow multiple connections requires to do some cleanup
-	if (!ToPin->AllowMultipleConnections())
+	if (!ToPin->AllowsMultipleConnections())
 	{
 		bToPinBrokeOtherEdges = ToPin->BreakAllIncompatibleEdges(&TouchedNodes);
 	}
diff --git a/Engine/Plugins/Experimental/PCG/Source/PCG/Private/PCGPin.cpp b/Engine/Plugins/Experimental/PCG/Source/PCG/Private/PCGPin.cpp
--- a/Engine/Plugins/Experimental/PCG/Source/PCG/Private/PCGPin.cpp
+++ b/Engine/Plugins/Experimental/PCG/Source/PCG/Private/PCGPin.cpp
@@ -33,11 +33,26 @@ namespace PCGPin
 }
 
 FPCGPinProperties::FPCGPinProperties(const FName& InLabel, EPCGDataType InAllowedTypes, bool bInAllowMultipleConnections, bool bInAllowMultipleData, const FText& InTooltip)
-	: Label(InLabel), AllowedTypes(InAllowedTypes), bAllowMultipleData(bInAllowMultipleData), bAllowMultipleConnections(bInAllowMultipleConnections)
+	: Label(InLabel), AllowedTypes(InAllowedTypes), bAllowMultipleData(bInAllowMultipleData)
 #if WITH_EDITORONLY_DATA
 	, Tooltip(InTooltip)
 #endif
-{}
+{
+	SetAllowMultipleConnections(bInAllowMultipleConnections);
+}
+
+void FPCGPinProperties::SetAllowMultipleConnections(bool bInAllowMultipleConnectons)
+{
+	if (bInAllowMultipleConnectons)
+	{
+		bAllowMultipleConnections = true;
+		bAllowMultipleData = true;
+	}
+	else
+	{
+		bAllowMultipleConnections = false;
+	}
+}
 
 bool FPCGPinProperties::operator==(const FPCGPinProperties& Other) const
 {
@@ -199,7 +214,7 @@ bool UPCGPin::BreakAllIncompatibleEdges(TSet<UPCGNode*>* InTouchedNodes/*= nullp
 		UPCGEdge* Edge = Edges[EdgeIndex];
 		UPCGPin* OtherPin = Edge->GetOtherPin(this);
 
-		bool bRemoveEdge = !IsCompatible(OtherPin) || (!AllowMultipleConnections() && bHasAValidEdge);
+		bool bRemoveEdge = !IsCompatible(OtherPin) || (!AllowsMultipleConnections() && bHasAValidEdge);
 
 		if (bRemoveEdge)
 		{
@@ -314,20 +329,20 @@ bool UPCGPin::IsCompatible(const UPCGPin* OtherPin) const
 	return !!(UpstreamTypes & DownstreamTypes);
 }
 
-bool UPCGPin::AllowMultipleConnections() const
+bool UPCGPin::AllowsMultipleConnections() const
 {
 	// Always allow multiple connection on output pin
-	return IsOutputPin() || Properties.bAllowMultipleConnections;
+	return IsOutputPin() || Properties.AllowsMultipleConnections();
 }
 
-bool UPCGPin::AllowMultipleData() const
+bool UPCGPin::AllowsMultipleData() const
 {
 	return Properties.bAllowMultipleData;
 }
 
 bool UPCGPin::CanConnect(const UPCGPin* OtherPin) const
 {
-	return OtherPin && (Edges.IsEmpty() || AllowMultipleConnections());
+	return OtherPin && (Edges.IsEmpty() || AllowsMultipleConnections());
 }
 
 EPCGTypeConversion UPCGPin::GetRequiredTypeConversion(const UPCGPin* InOtherPin) const
diff --git a/Engine/Plugins/Experimental/PCG/Source/PCG/Private/Tests/PCGTestsCommon.cpp b/Engine/Plugins/Experimental/PCG/Source/PCG/Private/Tests/PCGTestsCommon.cpp
--- a/Engine/Plugins/Experimental/PCG/Source/PCG/Private/Tests/PCGTestsCommon.cpp
+++ b/Engine/Plugins/Experimental/PCG/Source/PCG/Private/Tests/PCGTestsCommon.cpp
@@ -214,7 +214,7 @@ namespace PCGTestsCommon
 			SingleTaggedData.Data = TypeToData.Value();
 			SingleTaggedData.Pin = PinProperties.Label;
 
-			if (!PinProperties.bAllowMultipleConnections)
+			if (!PinProperties.AllowsMultipleConnections())
 			{
 				continue;
 			}
diff --git a/Engine/Plugins/Experimental/PCG/Source/PCG/Public/PCGPin.h b/Engine/Plugins/Experimental/PCG/Source/PCG/Public/PCGPin.h
--- a/Engine/Plugins/Experimental/PCG/Source/PCG/Public/PCGPin.h
+++ b/Engine/Plugins/Experimental/PCG/Source/PCG/Public/PCGPin.h
@@ -30,9 +30,6 @@ struct PCG_API FPCGPinProperties
 	UPROPERTY(BlueprintReadWrite, EditAnywhere, Category = Settings)
 	bool bAllowMultipleData = true;
 
-	UPROPERTY(BlueprintReadWrite, EditAnywhere, Category = Settings, meta = (EditCondition = "bAllowMultipleData"))
-	bool bAllowMultipleConnections = true;
-
 	/* Advanced pin will be hidden by default in the UI and will be shown only if the user extend the node (in the UI) to see advanced pins. */
 	UPROPERTY(BlueprintReadWrite, EditAnywhere, Category = Settings)
 	bool bAdvancedPin = false;
@@ -41,6 +38,16 @@ struct PCG_API FPCGPinProperties
 	UPROPERTY(EditAnywhere, Category = Settings)
 	FText Tooltip;
 #endif
+
+	// Multiple connections are only possible if we support multi data.
+	bool AllowsMultipleConnections() const { return bAllowMultipleData && bAllowMultipleConnections; }
+
+	// Allowing multiple connections will automatically enable multi data.
+	void SetAllowMultipleConnections(bool bInAllowMultipleConnectons);
+
+private:
+	UPROPERTY(BlueprintReadWrite, EditAnywhere, Category = Settings, meta = (AllowPrivateAccess = "true", EditCondition = "bAllowMultipleData", DisplayAfter = "bAllowMultipleData"))
+	bool bAllowMultipleConnections = true;
 };
 
 UENUM()
@@ -84,7 +91,7 @@ public:
 		}
 
 		bAllowMultipleData = PinProperty.bAllowMultipleData;
-		bAllowMultipleConnections = PinProperty.bAllowMultipleConnections;
+		bAllowMultipleConnections = PinProperty.AllowsMultipleConnections();
 		bAdvancedPin = PinProperty.bAdvancedPin;
 	}
 
@@ -146,8 +153,8 @@ public:
 	UFUNCTION(BlueprintCallable, Category = Settings)
 	void SetTooltip(const FText& InTooltip);
 
-	bool AllowMultipleConnections() const;
-	bool AllowMultipleData() const;
+	bool AllowsMultipleConnections() const;
+	bool AllowsMultipleData() const;
 	bool IsCompatible(const UPCGPin* OtherPin) const;
 	bool CanConnect(const UPCGPin* OtherPin) const;
 
diff --git a/Engine/Plugins/Experimental/PCG/Source/PCGEditor/Private/PCGEditorGraphNodeBase.cpp b/Engine/Plugins/Experimental/PCG/Source/PCGEditor/Private/PCGEditorGraphNodeBase.cpp
--- a/Engine/Plugins/Experimental/PCG/Source/PCGEditor/Private/PCGEditorGraphNodeBase.cpp
+++ b/Engine/Plugins/Experimental/PCG/Source/PCGEditor/Private/PCGEditorGraphNodeBase.cpp
@@ -698,7 +698,7 @@ void UPCGEditorGraphNodeBase::GetPinHoverText(const UEdGraphPin& Pin, FString& H
 		if (bIsInputPin)
 		{
 			MultiDataSupport = MatchingPin->Properties.bAllowMultipleData ? FText(LOCTEXT("InputSupportsMultiData", "Supports multiple data in input(s). ")) : FText(LOCTEXT("InputSingleDataOnly", "Supports only single data in input(s). "));
-			MultiConnectionSupport = MatchingPin->Properties.bAllowMultipleConnections ? FText(LOCTEXT("SupportsMultiInput", "Supports multiple inputs.")) : FText(LOCTEXT("SingleInputOnly", "Supports only one input."));
+			MultiConnectionSupport = MatchingPin->Properties.AllowsMultipleConnections() ? FText(LOCTEXT("SupportsMultiInput", "Supports multiple inputs.")) : FText(LOCTEXT("SingleInputOnly", "Supports only one input."));
 		}
 		else
 		{
diff --git a/Engine/Plugins/Experimental/PCG/Source/PCGEditor/Private/PCGEditorGraphSchema.cpp b/Engine/Plugins/Experimental/PCG/Source/PCGEditor/Private/PCGEditorGraphSchema.cpp
--- a/Engine/Plugins/Experimental/PCG/Source/PCGEditor/Private/PCGEditorGraphSchema.cpp
+++ b/Engine/Plugins/Experimental/PCG/Source/PCGEditor/Private/PCGEditorGraphSchema.cpp
@@ -133,7 +133,7 @@ const FPinConnectionResponse UPCGEditorGraphSchema::CanCreateConnection(const UE
 		return FPinConnectionResponse(CONNECT_RESPONSE_MAKE_WITH_CONVERSION_NODE, LOCTEXT("ConnectionUsingMakeConcrete", "Make data concrete"));
 	}
 
-	if (!InputPin->AllowMultipleConnections() && InputPin->EdgeCount() > 0)
+	if (!InputPin->AllowsMultipleConnections() && InputPin->EdgeCount() > 0)
 	{
 		return FPinConnectionResponse((A->Direction == EGPD_Output) ? CONNECT_RESPONSE_BREAK_OTHERS_B : CONNECT_RESPONSE_BREAK_OTHERS_A, LOCTEXT("ConnectionBreakExisting", "Break existing connection?"));
 	}
diff --git a/Engine/Plugins/Experimental/PCG/Source/PCGEditor/Private/SPCGEditorGraphNode.cpp b/Engine/Plugins/Experimental/PCG/Source/PCGEditor/Private/SPCGEditorGraphNode.cpp
--- a/Engine/Plugins/Experimental/PCG/Source/PCGEditor/Private/SPCGEditorGraphNode.cpp
+++ b/Engine/Plugins/Experimental/PCG/Source/PCGEditor/Private/SPCGEditorGraphNode.cpp
@@ -466,7 +466,7 @@ void SPCGEditorGraphNode::AddPin(const TSharedRef<SGraphPin>& PinToAdd)
 		if(UPCGPin* Pin = (bIsInPin ? PCGNode->GetInputPin(PinName) : PCGNode->GetOutputPin(PinName)))
 		{
 			const bool bIsMultiData = Pin->Properties.bAllowMultipleData;
-			const bool bIsMultiConnections = Pin->AllowMultipleConnections();
+			const bool bIsMultiConnections = Pin->AllowsMultipleConnections();
 
 			// Check for special types
 			if (Pin->GetCurrentTypes() == EPCGDataType::Param)
