PCG: Add missing const
#rnx

diff --git a/Engine/Plugins/Experimental/PCG/Source/PCG/Private/Elements/PCGGetBounds.cpp b/Engine/Plugins/Experimental/PCG/Source/PCG/Private/Elements/PCGGetBounds.cpp
--- a/Engine/Plugins/Experimental/PCG/Source/PCG/Private/Elements/PCGGetBounds.cpp
+++ b/Engine/Plugins/Experimental/PCG/Source/PCG/Private/Elements/PCGGetBounds.cpp
@@ -50,7 +50,7 @@ bool FPCGGetBoundsElement::ExecuteInternal(FPCGContext* Context) const
 
 	for (const FPCGTaggedData& Input : Inputs)
 	{
-		UPCGSpatialData* InputSpatialData = Cast<UPCGSpatialData>(Input.Data);
+		const UPCGSpatialData* InputSpatialData = Cast<UPCGSpatialData>(Input.Data);
 		if (!InputSpatialData)
 		{
 			continue;
diff --git a/Engine/Plugins/Experimental/PCG/Source/PCG/Private/Tests/Elements/PCGSortPointsTest.cpp b/Engine/Plugins/Experimental/PCG/Source/PCG/Private/Tests/Elements/PCGSortPointsTest.cpp
--- a/Engine/Plugins/Experimental/PCG/Source/PCG/Private/Tests/Elements/PCGSortPointsTest.cpp
+++ b/Engine/Plugins/Experimental/PCG/Source/PCG/Private/Tests/Elements/PCGSortPointsTest.cpp
@@ -42,7 +42,7 @@ bool FPCGSortPointsTest_Basic::RunTest(const FString& Parameters)
 
 	UTEST_EQUAL("Output count", Outputs.Num(), 1);
 
-	UPCGPointData* OutPointData = Cast<UPCGPointData>(Outputs[0].Data);
+	const UPCGPointData* OutPointData = Cast<UPCGPointData>(Outputs[0].Data);
 
 	UTEST_NOT_NULL("Output point data", OutPointData);
 
