[PCG] Fix Profiling window for partitioned graphs

#rnx

diff --git a/Engine/Plugins/Experimental/PCG/Source/PCG/Private/Graph/PCGGraphCompiler.cpp b/Engine/Plugins/Experimental/PCG/Source/PCG/Private/Graph/PCGGraphCompiler.cpp
--- a/Engine/Plugins/Experimental/PCG/Source/PCG/Private/Graph/PCGGraphCompiler.cpp
+++ b/Engine/Plugins/Experimental/PCG/Source/PCG/Private/Graph/PCGGraphCompiler.cpp
@@ -275,7 +275,7 @@ void FPCGGraphCompiler::Compile(UPCGGraph* InGraph)
 	}
 }
 
-TArray<FPCGGraphTask> FPCGGraphCompiler::GetPrecompiledTasks(UPCGGraph* InGraph, uint32 GenerationGridSize, FPCGStackContext& OutStackContext, bool bIsTopGraph) const
+TArray<FPCGGraphTask> FPCGGraphCompiler::GetPrecompiledTasks(const UPCGGraph* InGraph, uint32 GenerationGridSize, FPCGStackContext& OutStackContext, bool bIsTopGraph) const
 {
 	// Get compiled tasks in a threadsafe way
 	FReadScopeLock ReadLock(GraphToTaskMapLock);
diff --git a/Engine/Plugins/Experimental/PCG/Source/PCG/Private/Graph/PCGGraphCompiler.h b/Engine/Plugins/Experimental/PCG/Source/PCG/Private/Graph/PCGGraphCompiler.h
--- a/Engine/Plugins/Experimental/PCG/Source/PCG/Private/Graph/PCGGraphCompiler.h
+++ b/Engine/Plugins/Experimental/PCG/Source/PCG/Private/Graph/PCGGraphCompiler.h
@@ -17,7 +17,7 @@ class FPCGGraphCompiler
 public:
 	void Compile(UPCGGraph* InGraph);
 	TArray<FPCGGraphTask> GetCompiledTasks(UPCGGraph* InGraph, uint32 GenerationGridSize, FPCGStackContext& OutStackContext, bool bIsTopGraph = true);
-	TArray<FPCGGraphTask> GetPrecompiledTasks(UPCGGraph* InGraph, uint32 GenerationGridSize, FPCGStackContext& OutStackContext, bool bIsTopGraph = true) const;
+	TArray<FPCGGraphTask> GetPrecompiledTasks(const UPCGGraph* InGraph, uint32 GenerationGridSize, FPCGStackContext& OutStackContext, bool bIsTopGraph = true) const;
 
 	static void OffsetNodeIds(TArray<FPCGGraphTask>& Tasks, FPCGTaskId Offset, FPCGTaskId ParentId);
 
diff --git a/Engine/Plugins/Experimental/PCG/Source/PCG/Private/Utils/PCGExtraCapture.cpp b/Engine/Plugins/Experimental/PCG/Source/PCG/Private/Utils/PCGExtraCapture.cpp
--- a/Engine/Plugins/Experimental/PCG/Source/PCG/Private/Utils/PCGExtraCapture.cpp
+++ b/Engine/Plugins/Experimental/PCG/Source/PCG/Private/Utils/PCGExtraCapture.cpp
@@ -4,6 +4,7 @@
 
 #include "PCGComponent.h"
 #include "PCGContext.h"
+#include "PCGGraph.h"
 #include "Graph/PCGGraphCompiler.h"
 #include "Graph/PCGGraphExecutor.h"
 
@@ -157,8 +158,8 @@ namespace PCGUtils
 PCGUtils::FCallTreeInfo PCGUtils::FExtraCapture::CalculateCallTreeInfo(const UPCGComponent* Component) const
 {
 	UPCGSubsystem* PCGSubsystem = Component ? Component->GetSubsystem() : nullptr;
-
-	if (!PCGSubsystem)
+	const UPCGGraph* PCGGraph = Component ? Component->GetGraph() : nullptr;
+	if (!PCGSubsystem || !PCGGraph)
 	{
 		return {};
 	}
@@ -169,9 +170,10 @@ PCGUtils::FCallTreeInfo PCGUtils::FExtraCapture::CalculateCallTreeInfo(const UPC
 		return {};
 	}
 
+	const uint32 GridSize = PCGGraph->IsHierarchicalGenerationEnabled() ? Component->GetGenerationGridSize() : PCGHiGenGrid::UninitializedGridSize();
+
 	FPCGStackContext DummyStackContext;
-	// Passed uninitialized grid size to get all tasks
-	TArray<FPCGGraphTask> CompiledTasks = Compiler->GetPrecompiledTasks(Component->GetGraph(), PCGHiGenGrid::UninitializedGridSize(), DummyStackContext);
+	TArray<FPCGGraphTask> CompiledTasks = Compiler->GetPrecompiledTasks(PCGGraph, GridSize, DummyStackContext);
 	if (CompiledTasks.IsEmpty())
 	{
 		return {};
