[PCG] Added support for multiple pcg components within the same actor when setting debug object from selection

#rnx

diff --git a/Engine/Plugins/Experimental/PCG/Source/PCGEditor/Private/SPCGEditorGraphDebugObjectWidget.cpp b/Engine/Plugins/Experimental/PCG/Source/PCGEditor/Private/SPCGEditorGraphDebugObjectWidget.cpp
--- a/Engine/Plugins/Experimental/PCG/Source/PCGEditor/Private/SPCGEditorGraphDebugObjectWidget.cpp
+++ b/Engine/Plugins/Experimental/PCG/Source/PCGEditor/Private/SPCGEditorGraphDebugObjectWidget.cpp
@@ -256,56 +256,29 @@ bool SPCGEditorGraphDebugObjectWidget::IsSelectDebugObjectButtonEnabled() const
 
 void SPCGEditorGraphDebugObjectWidget::SetDebugObjectFromSelection_OnClicked()
 {
-	const UPCGGraph* PCGGraph = GetPCGGraph();
-	if (!PCGGraph)
+	ForEachStackInSelection([this](const FPCGStack& InStack, UPCGComponent* InPCGComponent)
 	{
-		return;
-	}
+		DebugObjects.Empty();
+		DebugObjectsComboBox->RefreshOptions();
 
-	USelection* SelectedActors = GEditor->GetSelectedActors();
-	if (!IsValid(SelectedActors))
-	{
-		return;
-	}
+		const TSharedPtr<FPCGEditorGraphDebugObjectInstance> DebugInstance = MakeShared<FPCGEditorGraphDebugObjectInstance>(InPCGComponent, InStack);
+		DebugObjects.Add(DebugInstance);
+		DebugObjectsComboBox->SetSelectedItem(DebugInstance);
+		PCGEditorPtr.Pin()->SetComponentAndStackBeingInspected(InPCGComponent, InStack);
+		
+		return true;
+	});
+}
 
-	for (FSelectionIterator It(*SelectedActors); It; ++It)
+bool SPCGEditorGraphDebugObjectWidget::IsSetDebugObjectFromSelectionButtonEnabled() const
+{	
+	return ForEachStackInSelection([](const FPCGStack& /*InStack*/, const UPCGComponent* /*InPCGComponent*/)
 	{
-		const AActor* SelectedActor = Cast<AActor>(*It);
-		if (!IsValid(SelectedActor))
-		{
-			continue;
-		}
-
-		UPCGComponent* PCGComponent = SelectedActor->GetComponentByClass<UPCGComponent>();
-		if (!IsValid(PCGComponent))
-		{
-			continue;
-		}
-
-		FPCGStackContext StackContext = FPCGStackContext::CreateStackContextFromGraph(PCGComponent->GetGraph());
-
-		for (const FPCGStack& Stack : StackContext.GetStacks())
-		{
-			const FPCGStackFrame& StackFrame = Stack.GetStackFrames().Top();
-			if (const UPCGGraph* StackGraph = Cast<const UPCGGraph>(StackFrame.Object))
-			{
-				if (StackGraph == PCGGraph)
-				{
-					DebugObjects.Empty();
-					DebugObjectsComboBox->RefreshOptions();
-
-					const TSharedPtr<FPCGEditorGraphDebugObjectInstance> DebugInstance = MakeShared<FPCGEditorGraphDebugObjectInstance>(PCGComponent, Stack);
-					DebugObjects.Add(DebugInstance);
-					DebugObjectsComboBox->SetSelectedItem(DebugInstance);
-					PCGEditorPtr.Pin()->SetComponentAndStackBeingInspected(PCGComponent, Stack);
-					break;
-				}
-			}
-		}
-	}
+		return true;
+	});
 }
 
-bool SPCGEditorGraphDebugObjectWidget::IsSetDebugObjectFromSelectionButtonEnabled() const
+bool SPCGEditorGraphDebugObjectWidget::ForEachStackInSelection(const TFunctionRef<bool(const FPCGStack&, UPCGComponent*)>& InOperation) const
 {
 	const UPCGGraph* PCGGraph = GetPCGGraph();
 	if (!PCGGraph)
@@ -327,21 +300,29 @@ bool SPCGEditorGraphDebugObjectWidget::IsSetDebugObjectFromSelectionButtonEnable
 			continue;
 		}
 		
-		const UPCGComponent* PCGComponent = SelectedActor->GetComponentByClass<UPCGComponent>();
-		if (!IsValid(PCGComponent))
-		{
-			continue;
-		}
+		TArray<UPCGComponent*> PCGComponents;
+		SelectedActor->GetComponents<UPCGComponent>(PCGComponents, /*bIncludeFromChildActors=*/true);
 
-		FPCGStackContext StackContext = FPCGStackContext::CreateStackContextFromGraph(PCGComponent->GetGraph());
-		for (const FPCGStack& Stack : StackContext.GetStacks())
+		for (UPCGComponent* PCGComponent : PCGComponents)
 		{
-			const FPCGStackFrame& StackFrame = Stack.GetStackFrames().Top();
-			if (const UPCGGraph* StackGraph = Cast<const UPCGGraph>(StackFrame.Object))
+			if (!IsValid(PCGComponent))
 			{
-				if (StackGraph == PCGGraph)
+				continue;
+			}
+
+			FPCGStackContext StackContext = FPCGStackContext::CreateStackContextFromGraph(PCGComponent->GetGraph());
+			for (const FPCGStack& Stack : StackContext.GetStacks())
+			{
+				const FPCGStackFrame& StackFrame = Stack.GetStackFrames().Top();
+				if (const UPCGGraph* StackGraph = Cast<const UPCGGraph>(StackFrame.Object))
 				{
-					return true;
+					if (StackGraph == PCGGraph)
+					{
+						if (InOperation(Stack, PCGComponent))
+						{
+							return true;
+						}
+					}
 				}
 			}
 		}
diff --git a/Engine/Plugins/Experimental/PCG/Source/PCGEditor/Private/SPCGEditorGraphDebugObjectWidget.h b/Engine/Plugins/Experimental/PCG/Source/PCGEditor/Private/SPCGEditorGraphDebugObjectWidget.h
--- a/Engine/Plugins/Experimental/PCG/Source/PCGEditor/Private/SPCGEditorGraphDebugObjectWidget.h
+++ b/Engine/Plugins/Experimental/PCG/Source/PCGEditor/Private/SPCGEditorGraphDebugObjectWidget.h
@@ -61,6 +61,8 @@ private:
 	
 	void SetDebugObjectFromSelection_OnClicked();
 	bool IsSetDebugObjectFromSelectionButtonEnabled() const;
+
+	bool ForEachStackInSelection(const TFunctionRef<bool(const FPCGStack&, UPCGComponent*)>& InOperation) const;
 	
 	/** Pointer back to the PCG editor that owns us */
 	TWeakPtr<FPCGEditor> PCGEditorPtr;
