[PCG] Fix broken build

* Override of pure virtual function was within #if WITH_EDITOR, while the pure function was not.

#rnx

diff --git a/Engine/Plugins/Experimental/PCG/Source/PCG/Private/PCGGraph.cpp b/Engine/Plugins/Experimental/PCG/Source/PCG/Private/PCGGraph.cpp
--- a/Engine/Plugins/Experimental/PCG/Source/PCG/Private/PCGGraph.cpp
+++ b/Engine/Plugins/Experimental/PCG/Source/PCG/Private/PCGGraph.cpp
@@ -1073,6 +1073,41 @@ bool UPCGGraph::UserParametersCanRemoveProperty(FGuid InPropertyID, FName InProp
 	return true;
 }
 
+bool UPCGGraph::UserParametersIsPinTypeAccepted(FEdGraphPinType InPinType, bool bIsChild)
+{
+	// Text and interface not supported
+	if (InPinType.PinCategory == TEXT("text") || InPinType.PinCategory == TEXT("interface"))
+	{
+		return false;
+	}
+	else if (InPinType.PinCategory == TEXT("struct"))
+	{
+		// Vector/Transform/Rotator are "explicitly" defined. Also return true for a "null" PinSubCategoryObject
+		// to mark it as "valid" so the system will call this function back on all available structures.
+		if (!bIsChild)
+		{
+			return InPinType.PinSubCategoryObject.IsExplicitlyNull()
+				|| InPinType.PinSubCategoryObject == TBaseStructure<FVector>::Get()
+				|| InPinType.PinSubCategoryObject == TBaseStructure<FTransform>::Get()
+				|| InPinType.PinSubCategoryObject == TBaseStructure<FRotator>::Get();
+		}
+		else
+		{
+			// Since we are not supporting all the structures by default, only show the ones that we are supporting with PCG.
+			return InPinType.PinSubCategoryObject == TBaseStructure<FVector2D>::Get()
+				|| InPinType.PinSubCategoryObject == TBaseStructure<FVector4>::Get()
+				|| InPinType.PinSubCategoryObject == TBaseStructure<FQuat>::Get()
+				|| InPinType.PinSubCategoryObject == TBaseStructure<FSoftObjectPath>::Get()
+				|| InPinType.PinSubCategoryObject == TBaseStructure<FSoftClassPath>::Get();
+		}
+	}
+	else
+	{
+		return true;
+	}
+}
+#endif // WITH_EDITOR
+
 void UPCGGraph::OnGraphParametersChanged(EPCGGraphParameterEvent InChangeType, FName InChangedPropertyName)
 {
 	bool bWasModified = false;
@@ -1130,43 +1165,10 @@ void UPCGGraph::OnGraphParametersChanged(EPCGGraphParameterEvent InChangeType, F
 		UserParameters.Reset();
 	}
 
+#if WITH_EDITOR
 	NotifyGraphParametersChanged(InChangeType, InChangedPropertyName);
-}
-
-bool UPCGGraph::UserParametersIsPinTypeAccepted(FEdGraphPinType InPinType, bool bIsChild)
-{
-	// Text and interface not supported
-	if (InPinType.PinCategory == TEXT("text") || InPinType.PinCategory == TEXT("interface"))
-	{
-		return false;
-	}
-	else if (InPinType.PinCategory == TEXT("struct"))
-	{
-		// Vector/Transform/Rotator are "explicitly" defined. Also return true for a "null" PinSubCategoryObject
-		// to mark it as "valid" so the system will call this function back on all available structures.
-		if (!bIsChild)
-		{
-			return InPinType.PinSubCategoryObject.IsExplicitlyNull()
-				|| InPinType.PinSubCategoryObject == TBaseStructure<FVector>::Get()
-				|| InPinType.PinSubCategoryObject == TBaseStructure<FTransform>::Get()
-				|| InPinType.PinSubCategoryObject == TBaseStructure<FRotator>::Get();
-		}
-		else
-		{
-			// Since we are not supporting all the structures by default, only show the ones that we are supporting with PCG.
-			return InPinType.PinSubCategoryObject == TBaseStructure<FVector2D>::Get()
-				|| InPinType.PinSubCategoryObject == TBaseStructure<FVector4>::Get()
-				|| InPinType.PinSubCategoryObject == TBaseStructure<FQuat>::Get()
-				|| InPinType.PinSubCategoryObject == TBaseStructure<FSoftObjectPath>::Get()
-				|| InPinType.PinSubCategoryObject == TBaseStructure<FSoftClassPath>::Get();
-		}
-	}
-	else
-	{
-		return true;
-	}
-}
 #endif // WITH_EDITOR
+}
 
 FInstancedPropertyBag* UPCGGraph::GetMutableUserParametersStruct()
 {
