Added some flexibility around how we serialize and build the landscape cache (never, always, only on cook)

diff --git a/Engine/Plugins/Experimental/PCG/Source/PCG/Private/Grid/PCGLandscapeCache.cpp b/Engine/Plugins/Experimental/PCG/Source/PCG/Private/Grid/PCGLandscapeCache.cpp
--- a/Engine/Plugins/Experimental/PCG/Source/PCG/Private/Grid/PCGLandscapeCache.cpp
+++ b/Engine/Plugins/Experimental/PCG/Source/PCG/Private/Grid/PCGLandscapeCache.cpp
@@ -2,6 +2,7 @@
 
 #include "Grid/PCGLandscapeCache.h"
 
+#include "PCGCustomVersion.h"
 #include "PCGPoint.h"
 #include "Helpers/PCGBlueprintHelpers.h"
 #include "Metadata/PCGMetadata.h"
@@ -84,7 +85,6 @@ FPCGLandscapeCacheEntry* FPCGLandscapeCacheEntry::CreateCacheEntry(ULandscapeInf
 
 	FPCGLandscapeCacheEntry *Result = new FPCGLandscapeCacheEntry();
 
- 	Result->Component = InComponent;
 	Result->PointHalfSize = InComponent->GetComponentTransform().GetScale3D() * 0.5;
 	Result->Stride = Stride;
 
@@ -318,10 +318,6 @@ bool FPCGLandscapeCacheEntry::TouchAndLoad(int32 InTouch) const
 {
 	Touch = InTouch; // technically, this could be an atomic, but we don't need this to be very precise
 
-#if WITH_EDITOR
-	check(bDataLoaded);
-	return false; // In editor, we're "always" loaded since it's created on the fly
-#else
 	if (!bDataLoaded)
 	{
 		FScopeLock ScopeDataLock(&DataLock);
@@ -333,7 +329,6 @@ bool FPCGLandscapeCacheEntry::TouchAndLoad(int32 InTouch) const
 	}
 
 	return false;
-#endif
 }
 
 void FPCGLandscapeCacheEntry::Unload()
@@ -420,12 +415,15 @@ void FPCGLandscapeCacheEntry::SerializeFromBulkData() const
 
 void FPCGLandscapeCacheEntry::Serialize(FArchive& Archive, UObject* Owner, int32 Index)
 {
-	if (bDataLoaded && Archive.IsSaving() && Archive.IsCooking())
+	// Important implementation note:
+	// If the serialization here or in the cache entries change, we still need to load data from previous versions.
+	// While that's not really needed in non-editor builds, it is very much important when loading from the editor,
+	// at least in the "AlwaysSerialize" case.
+	if (bDataLoaded && Archive.IsSaving())
 	{
 		SerializeToBulkData();
 	}
-
-	Archive << Component;
+	
 	Archive << PointHalfSize;
 	Archive << Stride;
 	Archive << LayerDataNames;
@@ -450,13 +448,24 @@ void UPCGLandscapeCache::Serialize(FArchive& Archive)
 {
 	Super::Serialize(Archive);
 
-	if (Archive.IsSaving() && !Archive.IsCooking())
+	const bool bShouldSerializeEntries = (Archive.IsSaving() &&
+		(SerializationMode == EPCGLandscapeCacheSerializationMode::AlwaysSerialize || 
+			(SerializationMode == EPCGLandscapeCacheSerializationMode::SerializeOnlyAtCook && Archive.IsCooking())));
+
+	// Important implementation note:
+	// If the serialization here or in the cache entries change, we still need to load data from previous versions.
+	// While that's not really needed in non-editor builds, it is very much important when loading from the editor,
+	// at least in the "AlwaysSerialize" case.
+	Archive.UsingCustomVersion(FPCGCustomVersion::GUID);
+
+	int32 DataVersion = FPCGCustomVersion::LatestVersion;
+	if (Archive.IsLoading())
 	{
-		ClearCache();
+		DataVersion = Archive.CustomVer(FPCGCustomVersion::GUID);
 	}
 
 	// Serialize cache entries
-	int32 NumEntries = (Archive.IsLoading() ? 0 : CachedData.Num());
+	int32 NumEntries = ((Archive.IsLoading() || !bShouldSerializeEntries) ? 0 : CachedData.Num());
 	Archive << NumEntries;
 
 	if (Archive.IsLoading())
@@ -473,8 +482,12 @@ void UPCGLandscapeCache::Serialize(FArchive& Archive)
 
 			CachedData.Add(Key, Entry);
 		}
+
+#if WITH_EDITOR
+		CacheEntryCount = CachedData.Num();
+#endif
 	}
-	else
+	else if(bShouldSerializeEntries)
 	{
 		int32 EntryIndex = 0;
 		for (auto& CacheEntry : CachedData)
@@ -610,6 +623,11 @@ void UPCGLandscapeCache::PrimeCache()
 		NewEntries[EntryIndex] = FPCGLandscapeCacheEntry::CreateCacheEntry(CacheEntryInfo.Get<0>(), CacheEntryInfo.Get<1>());
 	});
 
+	if (SerializationMode != EPCGLandscapeCacheSerializationMode::NeverSerialize)
+	{
+		Modify();
+	}
+
 	// Finally, write them back to the cache
 	for (int32 EntryIndex = 0; EntryIndex < CacheEntriesToBuild.Num(); ++EntryIndex)
 	{
@@ -619,12 +637,18 @@ void UPCGLandscapeCache::PrimeCache()
 		}
 	}
 
+	CacheEntryCount = CachedData.Num();
 	CacheLayerNames();
 #endif
 }
 
 void UPCGLandscapeCache::ClearCache()
 {
+	if (SerializationMode != EPCGLandscapeCacheSerializationMode::NeverSerialize)
+	{
+		Modify();
+	}
+
 	for (TPair<TPair<FGuid, FIntPoint>, FPCGLandscapeCacheEntry*>& CacheEntry : CachedData)
 	{
 		delete CacheEntry.Value;
@@ -632,6 +656,10 @@ void UPCGLandscapeCache::ClearCache()
 	}
 
 	CachedData.Reset();
+	CachedLayerNames.Reset();
+#if WITH_EDITOR
+	CacheEntryCount = 0;
+#endif
 }
 
 #if WITH_EDITOR
@@ -655,6 +683,7 @@ const FPCGLandscapeCacheEntry* UPCGLandscapeCache::GetCacheEntry(ULandscapeCompo
 			{
 				CacheEntry = NewEntry;
 				CachedData.Add(ComponentKey, NewEntry);
+				++CacheEntryCount;
 			}
 		}
 
@@ -809,6 +838,7 @@ void UPCGLandscapeCache::OnLandscapeChanged(ALandscapeProxy* InLandscape, const
 
 			if (CachedData.RemoveAndCopyValue(ComponentKey, EntryToDelete))
 			{
+				--CacheEntryCount;
 				delete EntryToDelete;
 			}
 		}
@@ -891,6 +921,7 @@ void UPCGLandscapeCache::RemoveComponentFromCache(const ALandscapeProxy* Landsca
 			if (FPCGLandscapeCacheEntry** FoundEntry = CachedData.Find(ComponentKey))
 			{
 				CachedData.Remove(ComponentKey);
+				--CacheEntryCount;
 				delete* FoundEntry;
 			}
 		}
diff --git a/Engine/Plugins/Experimental/PCG/Source/PCG/Private/PCGSubsystem.cpp b/Engine/Plugins/Experimental/PCG/Source/PCG/Private/PCGSubsystem.cpp
--- a/Engine/Plugins/Experimental/PCG/Source/PCG/Private/PCGSubsystem.cpp
+++ b/Engine/Plugins/Experimental/PCG/Source/PCG/Private/PCGSubsystem.cpp
@@ -1291,7 +1291,6 @@ void UPCGSubsystem::BuildLandscapeCache(bool bQuiet)
 	TRACE_CPUPROFILER_EVENT_SCOPE(UPCGSubsystem::BuildLandscapeCache);
 	if (UPCGLandscapeCache* LandscapeCache = GetLandscapeCache())
 	{
-		PCGWorldActor->Modify();
 		LandscapeCache->PrimeCache();
 	}
 	else if(!bQuiet)
diff --git a/Engine/Plugins/Experimental/PCG/Source/PCG/Public/Grid/PCGLandscapeCache.h b/Engine/Plugins/Experimental/PCG/Source/PCG/Public/Grid/PCGLandscapeCache.h
--- a/Engine/Plugins/Experimental/PCG/Source/PCG/Public/Grid/PCGLandscapeCache.h
+++ b/Engine/Plugins/Experimental/PCG/Source/PCG/Public/Grid/PCGLandscapeCache.h
@@ -14,6 +14,14 @@ class UPCGPointData;
 class UPCGMetadata;
 struct FPCGPoint;
 
+UENUM()
+enum class EPCGLandscapeCacheSerializationMode : uint8
+{
+	SerializeOnlyAtCook,
+	NeverSerialize,
+	AlwaysSerialize
+};
+
 USTRUCT(BlueprintType)
 struct FPCGLandscapeLayerWeight
 {
@@ -77,7 +85,6 @@ private:
 	void SerializeFromBulkData() const;
 
 	// Serialized data
-	TWeakObjectPtr<const ULandscapeComponent> Component = nullptr;
 	TArray<FName> LayerDataNames;
 	FVector PointHalfSize = FVector::One();
 	int32 Stride = 0;
@@ -108,8 +115,12 @@ public:
 	/** Initialize cache. Can be safely called multiple times. */
 	void Initialize();
 
+	UFUNCTION(CallInEditor, Category = "Cache")
 	void PrimeCache();
+
+	UFUNCTION(CallInEditor, Category = "Cache")
 	void ClearCache();
+
 	void Tick(float DeltaSeconds);
 
 #if WITH_EDITOR
@@ -125,6 +136,14 @@ public:
 	/** Convenience method to get metadata from the landscape for a given pair of landscape and position */
 	void SampleMetadataOnPoint(ALandscapeProxy* Landscape, FPCGPoint& OutPoint, UPCGMetadata* OutMetadata);
 
+	UPROPERTY(BlueprintReadWrite, EditAnywhere, Category = "Cache")
+	EPCGLandscapeCacheSerializationMode SerializationMode = EPCGLandscapeCacheSerializationMode::SerializeOnlyAtCook;
+
+#if WITH_EDITORONLY_DATA
+	UPROPERTY(Transient, BlueprintReadOnly, VisibleAnywhere, Category = "Cache")
+	int32 CacheEntryCount = 0;
+#endif
+
 private:
 #if WITH_EDITOR
 	void SetupLandscapeCallbacks();
@@ -138,6 +157,7 @@ private:
 	void RemoveComponentFromCache(const ALandscapeProxy* LandscapeProxy);
 #endif
 
+	// Mapping of landscape guid + coordinates to entries. This is manually serialized as needed (depends on the serialize options).
 	TMap<TPair<FGuid, FIntPoint>, FPCGLandscapeCacheEntry*> CachedData;
 
 	//TODO: separate by landscape
diff --git a/Engine/Plugins/Experimental/PCG/Source/PCG/Public/PCGWorldActor.h b/Engine/Plugins/Experimental/PCG/Source/PCG/Public/PCGWorldActor.h
--- a/Engine/Plugins/Experimental/PCG/Source/PCG/Public/PCGWorldActor.h
+++ b/Engine/Plugins/Experimental/PCG/Source/PCG/Public/PCGWorldActor.h
@@ -54,7 +54,7 @@ public:
 	uint32 PartitionGridSize;
 
 	/** Contains all the PCG data required to query the landscape complete. Serialized in cooked builds only */
-	UPROPERTY(VisibleAnywhere, Category = CachedData, meta = (DisplayName="Landscape Cache"))
+	UPROPERTY(BlueprintReadOnly, VisibleAnywhere, Category = CachedData, meta = (NoResetToDefault, DisplayName="Landscape Cache"))
 	TObjectPtr<UPCGLandscapeCache> LandscapeCacheObject = nullptr;
 
 	/** Disable creation of Partition Actors on the Z axis. Can improve performances if 3D partitioning is not needed. */
