[PCG] Fixed an issue and optimized our sorting helper function to support passing in arrays smaller than the original dataset, i.e. prefiltered list of points

#rnx

diff --git a/Engine/Plugins/Experimental/PCG/Source/PCG/Public/Metadata/Accessors/PCGAttributeAccessorHelpers.h b/Engine/Plugins/Experimental/PCG/Source/PCG/Public/Metadata/Accessors/PCGAttributeAccessorHelpers.h
--- a/Engine/Plugins/Experimental/PCG/Source/PCG/Public/Metadata/Accessors/PCGAttributeAccessorHelpers.h
+++ b/Engine/Plugins/Experimental/PCG/Source/PCG/Public/Metadata/Accessors/PCGAttributeAccessorHelpers.h
@@ -5,6 +5,7 @@
 #include "Metadata/PCGMetadataAttributeTpl.h"
 #include "Metadata/PCGMetadataAttributeTraits.h"
 #include "Metadata/Accessors/IPCGAttributeAccessor.h"
+#include "Metadata/Accessors/PCGAttributeAccessorKeys.h"
 #include "Templates/UniquePtr.h"
 #include "UObject/NameTypes.h"
 
@@ -51,14 +52,14 @@ namespace PCGAttributeAccessorHelpers
 	*/
 	PCG_API TUniquePtr<const IPCGAttributeAccessor> CreateConstAccessorForOverrideParamWithResult(const FPCGDataCollection& InInputData, const FPCGSettingsOverridableParam& InParam, AccessorParamResult* OutResult = nullptr);
 
-	/** 
+	/**
 	* Creates a const accessor to the property or attribute pointed at by the InSelector.
 	* Note that InData must not be null if the selector points to an attribute,
 	* but in the case of properties, it either has to be the appropriate type or null.
 	* Make sure to update your selector before-hand if you want to support "@Last"
 	*/
 	PCG_API TUniquePtr<const IPCGAttributeAccessor> CreateConstAccessor(const UPCGData* InData, const FPCGAttributePropertySelector& InSelector);
-	
+
 	/**
 	* Creates a accessor to the property or attribute pointed at by the InSelector.
 	* Note that InData must not be null if the selector points to an attribute,
@@ -87,10 +88,12 @@ namespace PCGAttributeAccessorHelpers
 	template <typename T, typename Func>
 	void SortByAttribute(const IPCGAttributeAccessor& InAccessor, const IPCGAttributeAccessorKeys& InKeys, TArray<T>& InArray, bool bAscending, Func&& CustomGetIndex)
 	{
+		check(InArray.Num() <= InKeys.GetNum())
+
 		if (InArray.IsEmpty())
 		{
 			return;
-		} 
+		}
 
 		auto Callback = [&InAccessor, &InKeys, &InArray, bAscending, &CustomGetIndex](auto Dummy)
 		{
@@ -102,11 +105,11 @@ namespace PCGAttributeAccessorHelpers
 
 				if constexpr (std::is_trivially_copyable_v<ValueType>)
 				{
-					CachedValues.SetNumUninitialized(InArray.Num());
+					CachedValues.SetNumUninitialized(InKeys.GetNum());
 				}
 				else
 				{
-					CachedValues.SetNum(InArray.Num());
+					CachedValues.SetNum(InKeys.GetNum());
 				}
 
 				InAccessor.GetRange(TArrayView<ValueType>(CachedValues), 0, InKeys);
@@ -145,8 +148,8 @@ namespace PCGAttributeAccessorHelpers
 
 				// Fill integer sequence
 				TArray<int32> ElementsIndexes;
-				ElementsIndexes.Reserve(CachedValues.Num());
-				for (int i = 0; i < CachedValues.Num(); ++i)
+				ElementsIndexes.Reserve(InArray.Num());
+				for (int i = 0; i < InArray.Num(); ++i)
 				{
 					ElementsIndexes.Add(i);
 				}
@@ -174,4 +177,4 @@ namespace PCGAttributeAccessorHelpers
 
 		PCGMetadataAttribute::CallbackWithRightType(InAccessor.GetUnderlyingType(), Callback);
 	}
-}
\ No newline at end of file
+}
