[PCG] Fix parameters override on PCG Component when both graphs has the same base graph

#rnx

diff --git a/Engine/Plugins/Experimental/PCG/Source/PCG/Private/PCGGraph.cpp b/Engine/Plugins/Experimental/PCG/Source/PCG/Private/PCGGraph.cpp
--- a/Engine/Plugins/Experimental/PCG/Source/PCG/Private/PCGGraph.cpp
+++ b/Engine/Plugins/Experimental/PCG/Source/PCG/Private/PCGGraph.cpp
@@ -1606,14 +1606,11 @@ bool FPCGOverrideInstancedPropertyBag::RefreshParameters(const FInstancedPropert
 	{
 	case EPCGGraphParameterEvent::GraphChanged:
 	{
-		// If the parameters property bags match, nothing to do.
-		if (ParentUserParameters->GetPropertyBagStruct() != Parameters.GetPropertyBagStruct())
-		{
-			bWasModified = true;
-			// Copy the parent parameters and reset overriddes
-			Parameters = *ParentUserParameters;
-			PropertiesIDsOverridden.Reset();
-		}
+		// We should always copy the parents parameters and reset overrides when the graph changes. Even if it is the same struct, values might be different.
+		bWasModified = true;
+		// Copy the parent parameters and reset overriddes
+		Parameters = *ParentUserParameters;
+		PropertiesIDsOverridden.Reset();
 		break;
 	}
 	case EPCGGraphParameterEvent::Added:
