[PCG]
* PCGGraphParameterExtension.cpp - Removed #pragma once

diff --git a/Engine/Plugins/Experimental/PCG/Source/PCG/Private/Helpers/PCGGraphParameterExtension.cpp b/Engine/Plugins/Experimental/PCG/Source/PCG/Private/Helpers/PCGGraphParameterExtension.cpp
--- a/Engine/Plugins/Experimental/PCG/Source/PCG/Private/Helpers/PCGGraphParameterExtension.cpp
+++ b/Engine/Plugins/Experimental/PCG/Source/PCG/Private/Helpers/PCGGraphParameterExtension.cpp
@@ -1,6 +1,4 @@
-﻿// Copyright Epic Games, Inc. All Rights Reserved.
-
-#pragma once
+// Copyright Epic Games, Inc. All Rights Reserved.
 
 #include "Helpers/PCGGraphParameterExtension.h"
 
