[PCG] Moved attribute list view sorting and filtering to separate thread in order to not freeze the editor on large data sets

diff --git a/Engine/Plugins/Experimental/PCG/Source/PCGEditor/Private/SPCGEditorGraphAttributeListView.cpp b/Engine/Plugins/Experimental/PCG/Source/PCGEditor/Private/SPCGEditorGraphAttributeListView.cpp
--- a/Engine/Plugins/Experimental/PCG/Source/PCGEditor/Private/SPCGEditorGraphAttributeListView.cpp
+++ b/Engine/Plugins/Experimental/PCG/Source/PCGEditor/Private/SPCGEditorGraphAttributeListView.cpp
@@ -23,6 +23,7 @@
 #include "Fonts/FontMeasure.h"
 #include "Framework/MultiBox/MultiBoxBuilder.h"
 #include "Widgets/Images/SLayeredImage.h"
+#include "Widgets/Images/SThrobber.h"
 #include "Widgets/Input/SComboBox.h"
 #include "Widgets/Input/SSearchBox.h"
 #include "Widgets/Layout/SScrollBox.h"
@@ -42,7 +43,7 @@ namespace PCGEditorGraphAttributeListView
 	const FText NoDataAvailableText = LOCTEXT("NoDataAvailableText", "No data available");
 	const FText NoNodeInspectedText = LOCTEXT("NoNodeInspectedText", "No node being inspected");
 	const FText NoNodeInspectedToolTip = LOCTEXT("NoNodeInspectedToolTip", "Inspect a node using the right click menu");
-	
+
 	/** Names of the columns in the attribute list */
 	const FName NAME_IndexColumn = FName(TEXT("IndexColumn"));
 	const FName NAME_PointPositionX = FName(TEXT("$Position.X"));
@@ -111,7 +112,7 @@ namespace PCGEditorGraphAttributeListView
 		check(FSlateApplication::Get().GetRenderer());
 		const TSharedRef<FSlateFontMeasure> FontMeasure = FSlateApplication::Get().GetRenderer()->GetFontMeasureService();
 		const FSlateFontInfo FontInfo = FAppStyle::GetFontStyle(TEXT("NormalText"));
-		
+
 		const float TextWidth = FontMeasure->Measure(InText, FontInfo).X;
 		constexpr float ColumnPadding = 22.0f; // TODO: Grab padding from header style
 		const float ColumnWidth = TextWidth + ColumnPadding;
@@ -119,6 +120,67 @@ namespace PCGEditorGraphAttributeListView
 	}
 }
 
+bool FPCGListViewUpdater::IsCompleted() const
+{
+	return UpdateTask.IsCompleted();
+}
+
+void FPCGListViewUpdater::Launch()
+{
+	// Passing a shared pointer to this in order for the task to keep the object alive even if we discard it in the attribute list view
+	UpdateTask = UE::Tasks::Launch(UE_SOURCE_LOCATION, [SharedContext = SharedThis(this)]()
+	{
+		SharedContext->AsyncFilter();
+		SharedContext->AsyncSort();
+	});
+}
+
+void FPCGListViewUpdater::AsyncSort()
+{
+	if (SortingColumn == PCGEditorGraphAttributeListView::NAME_IndexColumn || SortMode == EColumnSortMode::None)
+	{
+		if (SortMode == EColumnSortMode::Ascending || SortMode == EColumnSortMode::None)
+		{
+			ListViewItems.Sort([](const PCGListviewItemPtr& LHS, const PCGListviewItemPtr& RHS)
+			{
+				return  LHS->Index < RHS->Index;
+			});
+		}
+		else
+		{
+			ListViewItems.Sort([](const PCGListviewItemPtr& LHS, const PCGListviewItemPtr& RHS)
+			{
+				return  LHS->Index > RHS->Index;
+			});
+		}
+	}
+	else if (const FPCGColumnData* Data = ColumnData.Find(SortingColumn))
+	{
+		if (Data->DataAccessor.IsValid() && Data->DataKeys.IsValid())
+		{
+			//lambda used here to get the index value of an item in the array for sorting
+			PCGAttributeAccessorHelpers::SortByAttribute(*Data->DataAccessor, *Data->DataKeys, ListViewItems, SortMode == EColumnSortMode::Ascending, [this](int Index) { return ListViewItems[Index]->Index; });
+		}
+	}
+}
+
+void FPCGListViewUpdater::AsyncFilter()
+{
+	TArray<PCGListviewItemPtr> FilteredListViewItems;
+	FilteredListViewItems.Reserve(ListViewItems.Num());
+
+	for (const PCGListviewItemPtr& ListViewItem : ListViewItems)
+	{
+		const FPCGPointFilterExpressionContext PointFilterContext(ListViewItem.Get(), &ColumnData);
+		if (TextFilter->TestTextFilter(PointFilterContext))
+		{
+			FilteredListViewItems.Add(ListViewItem);
+		}
+	}
+
+	ListViewItems = MoveTemp(FilteredListViewItems);
+}
+
 void SPCGListViewItemRow::Construct(const FArguments& InArgs, const TSharedRef<STableViewBase>& InOwnerTableView)
 {
 	check(InArgs._ListViewItem && InArgs._AttributeListView);
@@ -178,11 +240,11 @@ TSharedRef<SWidget> SPCGListViewItemRow::GenerateWidgetForColumn(const FName& Co
 					}
 				}
 			};
-			
+
 			PCGMetadataAttribute::CallbackWithRightType(PCGColumnData->DataAccessor->GetUnderlyingType(), Callback);
 		}
 	}
-	
+
 	return SNew(STextBlock)
 		.Text(RowText)
 		.OverflowPolicy(ETextOverflowPolicy::Ellipsis)
@@ -220,7 +282,7 @@ bool FPCGPointFilterExpressionContext::TestComplexExpression(const FName& InKey,
 				using ValueType = decltype(Dummy);
 				ValueType Value{};
 				if (PCGColumnInfo->DataAccessor->Get<ValueType>(Value, Index, *PCGColumnInfo->DataKeys))
-				{					
+				{
 					FText TextValue;
 					if constexpr (PCG::Private::IsOfTypes<ValueType, bool>())
 					{
@@ -243,7 +305,7 @@ bool FPCGPointFilterExpressionContext::TestComplexExpression(const FName& InKey,
 						ensureMsgf(false, TEXT("Unsupported Data Type"));
 						return false;
 					}
-					
+
 					const FTextFilterString PointValue(TextValue.ToString());
 					return TextFilterUtils::TestComplexExpression(PointValue, InValue, InComparisonOperation, InTextComparisonMode);
 				}
@@ -325,7 +387,7 @@ void SPCGEditorGraphAttributeListView::Construct(const FArguments& InArgs, TShar
 		.ColorAndOpacity(FSlateColor::UseForeground());
 
 	FilterImage->AddLayer(TAttribute<const FSlateBrush*>(this, &SPCGEditorGraphAttributeListView::GetFilterBadgeIcon));
-	
+
 	SAssignNew(FilterButton, SComboButton)
 		.ForegroundColor(FSlateColor::UseStyle())
 		.HasDownArrow(false)
@@ -403,12 +465,24 @@ void SPCGEditorGraphAttributeListView::Construct(const FArguments& InArgs, TShar
 			+SHorizontalBox::Slot()
 			.FillWidth(1.0f)
 			[
-				SNew(SScrollBox)
-				.Orientation(Orient_Horizontal)
-				.ExternalScrollbar(HorizontalScrollBar)
-				+SScrollBox::Slot()
+				SNew(SOverlay)
+				+SOverlay::Slot()
 				[
-					ListView->AsShared()
+					SNew(SScrollBox)
+					.Orientation(Orient_Horizontal)
+					.ExternalScrollbar(HorizontalScrollBar)
+					+SScrollBox::Slot()
+					[
+						ListView->AsShared()
+					]
+				]
+				+SOverlay::Slot()
+				.HAlign(HAlign_Center)
+				.VAlign(VAlign_Center)
+				[
+					SNew(SCircularThrobber)
+					.Radius(12.0f)
+					.Visibility_Lambda([this](){return CurrentUpdateTask.IsValid() && !CurrentUpdateTask->IsCompleted() ? EVisibility::Visible : EVisibility::Hidden; })
 				]
 			]
 			+SHorizontalBox::Slot()
@@ -441,6 +515,18 @@ void SPCGEditorGraphAttributeListView::Tick(const FGeometry& AllottedGeometry, c
 		RefreshDataComboBox();
 		RefreshAttributeList();
 	}
+
+	if (CurrentUpdateTask.IsValid() && CurrentUpdateTask->IsCompleted())
+	{
+		FilteredListViewItems = MoveTemp(CurrentUpdateTask->ListViewItems);
+		if (ListView.IsValid())
+		{
+			ListView->SetItemsSource(&FilteredListViewItems);
+			ListView->RequestListRefresh();
+		}
+
+		CurrentUpdateTask.Reset();
+	}
 }
 
 TSharedRef<SHeaderRow> SPCGEditorGraphAttributeListView::CreateHeaderRowWidget() const
@@ -487,9 +573,9 @@ void SPCGEditorGraphAttributeListView::OnInspectedNodeChanged(UPCGEditorGraphNod
 	{
 		return;
 	}
-	
+
 	PCGEditorGraphNode = InPCGEditorGraphNode;
-	
+
 	if (PCGEditorGraphNode.IsValid())
 	{
 		NodeNameTextBlock->SetText(PCGEditorGraphNode->GetNodeTitle(ENodeTitleType::FullTitle));
@@ -630,7 +716,7 @@ void SPCGEditorGraphAttributeListView::RefreshAttributeList()
 				ListViewItem->PCGPoint = &PCGPoint;
 				ListViewItems.Add(ListViewItem);
 			}
-			
+
 			if (!PCGEditorGraphAttributeListView::IsGraphCacheDebuggingEnabled())
 			{
 				InfoTextBlock->SetText(FText::Format(LOCTEXT("PointInfoTextBlockFmt", "Number of points: {0}"), NumPoints));
@@ -649,7 +735,10 @@ void SPCGEditorGraphAttributeListView::RefreshAttributeList()
 		InfoTextBlock->SetText(FText::Format(PCGEditorGraphAttributeListView::LastLabelFormat, LastSelector, InfoTextBlock->GetText()));
 	}
 
-	RefreshSorting();
+	ListView->SetItemsSource(&ListViewItems);
+	ListView->RequestListRefresh();
+
+	LaunchUpdateTask();
 }
 
 void SPCGEditorGraphAttributeListView::RefreshPinComboBox()
@@ -734,32 +823,17 @@ void SPCGEditorGraphAttributeListView::RefreshDataComboBox()
 	}
 }
 
-void SPCGEditorGraphAttributeListView::ApplyRowFilter()
+void SPCGEditorGraphAttributeListView::LaunchUpdateTask()
 {
-	TRACE_CPUPROFILER_EVENT_SCOPE(SPCGEditorGraphAttributeListView::ApplyRowFilter);
-
-	FilteredListViewItems.Empty();
-
-	const FString& FilterString = ActiveFilterText.ToString();
-	if (!FilterString.IsEmpty())
-	{
-		for (const PCGListviewItemPtr& ListViewItem : ListViewItems)
-		{
-			const FPCGPointFilterExpressionContext PointFilterContext(ListViewItem.Get(), &PCGColumnData);
-			if (TextFilter->TestTextFilter(PointFilterContext))
-			{
-				FilteredListViewItems.Add(ListViewItem);
-			}
-		}
-		
-		ListView->SetItemsSource(&FilteredListViewItems);
-	}
-	else
-	{
-		ListView->SetItemsSource(&ListViewItems);
-	}
-
-	ListView->RequestListRefresh();
+	// Discarding any currently running updater, task will still run and keep the old object alive but we wont care about the result.
+	// This is done because we cant afford to wait for task completion before starting a new task.
+	CurrentUpdateTask.Reset();
+	CurrentUpdateTask = MakeShared<FPCGListViewUpdater>(ListViewItems,
+		PCGColumnData,
+		SortMode,
+		SortingColumn,
+		TextFilter);
+	CurrentUpdateTask->Launch();
 }
 
 const FSlateBrush* SPCGEditorGraphAttributeListView::GetFilterBadgeIcon() const
@@ -768,7 +842,7 @@ const FSlateBrush* SPCGEditorGraphAttributeListView::GetFilterBadgeIcon() const
 	{
 		if (!Column.bIsVisible)
 		{
-			return FAppStyle::Get().GetBrush("Icons.BadgeModified"); 
+			return FAppStyle::Get().GetBrush("Icons.BadgeModified");
 		}
 	}
 
@@ -790,7 +864,7 @@ TSharedRef<SWidget> SPCGEditorGraphAttributeListView::OnGenerateFilterMenu()
 		EUserInterfaceActionType::ToggleButton);
 
 	MenuBuilder.AddSeparator();
-	
+
 	const TIndirectArray<SHeaderRow::FColumn>& Columns = ListViewHeader->GetColumns();
 	TArray<FName> HiddenColumns = ListViewHeader->GetHiddenColumnIds();
 
@@ -811,7 +885,7 @@ TSharedRef<SWidget> SPCGEditorGraphAttributeListView::OnGenerateFilterMenu()
 			NAME_None,
 			EUserInterfaceActionType::ToggleButton);
 	}
-	
+
 	return MenuBuilder.MakeWidget();
 }
 
@@ -900,7 +974,7 @@ void SPCGEditorGraphAttributeListView::GenerateColumnsFromMetadata(const UPCGDat
 		const FName& AttributeName = AttributeNames[I];
 		const EPCGMetadataTypes AttributeType = AttributeTypes[I];
 		FName ColumnName = AttributeName;
-		
+
 		switch (AttributeType)
 		{
 		case EPCGMetadataTypes::Float:
@@ -994,7 +1068,7 @@ ECheckBoxState SPCGEditorGraphAttributeListView::GetAnyAttributeEnabledState() c
 {
 	bool bAllEnabled = true;
 	bool bAnyEnabled = false;
-	
+
 	for (const SHeaderRow::FColumn& Column : ListViewHeader->GetColumns())
 	{
 		if (Column.ColumnId == PCGEditorGraphAttributeListView::NAME_IndexColumn)
@@ -1012,7 +1086,7 @@ ECheckBoxState SPCGEditorGraphAttributeListView::GetAnyAttributeEnabledState() c
 	}
 	else if (bAnyEnabled)
 	{
-		return ECheckBoxState::Undetermined; 
+		return ECheckBoxState::Undetermined;
 	}
 	else
 	{
@@ -1039,7 +1113,7 @@ void SPCGEditorGraphAttributeListView::OnItemDoubleClicked(PCGListviewItemPtr It
 	if (const FPCGPoint* Point = Item->PCGPoint)
 	{
 		const FBox BoundingBox = Point->GetLocalBounds().TransformBy(Point->Transform.ToMatrixWithScale());
-		GEditor->MoveViewportCamerasToBox(BoundingBox, true, 2.5f);	
+		GEditor->MoveViewportCamerasToBox(BoundingBox, true, 2.5f);
 	}
 }
 
@@ -1056,7 +1130,7 @@ void SPCGEditorGraphAttributeListView::OnColumnSortModeChanged(const EColumnSort
 		SortMode = InNewSortMode;
 	}
 
-	RefreshSorting();
+	LaunchUpdateTask();
 }
 
 EColumnSortMode::Type SPCGEditorGraphAttributeListView::GetColumnSortMode(const FName InColumnId) const
@@ -1069,46 +1143,18 @@ EColumnSortMode::Type SPCGEditorGraphAttributeListView::GetColumnSortMode(const
 	return SortMode;
 }
 
-void SPCGEditorGraphAttributeListView::RefreshSorting()
-{
-	TRACE_CPUPROFILER_EVENT_SCOPE(SPCGEditorGraphAttributeListView::ColumnSort);
-	
-	if (SortingColumn == PCGEditorGraphAttributeListView::NAME_IndexColumn || SortMode == EColumnSortMode::None)
-	{
-		if (SortMode == EColumnSortMode::Ascending || SortMode == EColumnSortMode::None)
-		{
-			ListViewItems.Sort([](const PCGListviewItemPtr& LHS, const PCGListviewItemPtr& RHS)
-			{
-				return  LHS->Index < RHS->Index;
-			});
-		}
-		else
-		{
-			ListViewItems.Sort([](const PCGListviewItemPtr& LHS, const PCGListviewItemPtr& RHS)
-			{
-				return  LHS->Index > RHS->Index;
-			});
-		}
-	}
-	else if (FPCGColumnData* ColumnData = PCGColumnData.Find(SortingColumn))
-	{
-		if (ColumnData->DataAccessor.IsValid() && ColumnData->DataKeys.IsValid())
-		{
-			//lambda used here to get the index value of an item in the array for sorting
-			PCGAttributeAccessorHelpers::SortByAttribute(*ColumnData->DataAccessor, *ColumnData->DataKeys, ListViewItems, SortMode == EColumnSortMode::Ascending, [this](int Index) { return ListViewItems[Index]->Index; });
-		}
-	}
-
-	ApplyRowFilter();
-}
-
 void SPCGEditorGraphAttributeListView::OnFilterTextChanged(const FText& InFilterText)
 {
 	ActiveFilterText = InFilterText;
 	TextFilter->SetFilterText(InFilterText);
 
-	ApplyRowFilter();
-	SearchBoxWidget->SetError(TextFilter->GetFilterErrorText());
+	const FText ErrorText = TextFilter->GetFilterErrorText();
+	if(ErrorText.IsEmpty())
+	{
+		LaunchUpdateTask();
+	}
+
+	SearchBoxWidget->SetError(ErrorText);
 }
 
 void SPCGEditorGraphAttributeListView::OnFilterTextCommitted(const FText& NewText, ETextCommit::Type CommitInfo)
@@ -1123,24 +1169,25 @@ void SPCGEditorGraphAttributeListView::OnFilterTextCommitted(const FText& NewTex
 void SPCGEditorGraphAttributeListView::AddColumn(const UPCGPointData* InPCGPointData, const FName& InColumnId, const FText& ColumnLabel, EHorizontalAlignment HeaderHAlign, EHorizontalAlignment CellHAlign)
 {
 	FText ToolTip;
-	
+
 	if (InPCGPointData)
 	{
 		const FString ColumnIdString = InColumnId.ToString();
-		
+
 		FPCGAttributePropertyInputSelector TargetSelector;
 		TargetSelector.Update(ColumnIdString);
-	
+
 		FPCGColumnData& ColumnData = PCGColumnData.Add(InColumnId);
-		
+
 		if (InColumnId == PCGEditorGraphAttributeListView::NAME_PointMetadataEntry)
 		{
-			ColumnData.DataAccessor = PCGAttributeAccessorHelpers::CreatePropertyAccessor(GET_MEMBER_NAME_CHECKED(FPCGPoint, MetadataEntry), FPCGPoint::StaticStruct());
-			ColumnData.DataKeys = MakeUnique<FPCGAttributeAccessorKeysPoints>(InPCGPointData->GetPoints());
+			TUniquePtr<const IPCGAttributeAccessor> DataAccessor = PCGAttributeAccessorHelpers::CreatePropertyAccessor(GET_MEMBER_NAME_CHECKED(FPCGPoint, MetadataEntry), FPCGPoint::StaticStruct());
+			ColumnData.DataAccessor = TSharedPtr<const IPCGAttributeAccessor>(DataAccessor.Release());
+			ColumnData.DataKeys = MakeShared<FPCGAttributeAccessorKeysPoints>(InPCGPointData->GetPoints());
 		}
 		else if (InColumnId == PCGEditorGraphAttributeListView::NAME_PointMetadataEntryParent)
 		{
-			ColumnData.DataAccessor = MakeUnique<FPCGCustomPointAccessor<int64>>([InPCGPointData](const FPCGPoint& Point, void* OutValue)
+			ColumnData.DataAccessor = MakeShared<FPCGCustomPointAccessor<int64>>([InPCGPointData](const FPCGPoint& Point, void* OutValue)
 			{
 				if (const UPCGMetadata* Metadata = InPCGPointData->Metadata)
 				{
@@ -1149,19 +1196,19 @@ void SPCGEditorGraphAttributeListView::AddColumn(const UPCGPointData* InPCGPoint
 				}
 				return false;
 			}, nullptr);
-			ColumnData.DataKeys = MakeUnique<FPCGAttributeAccessorKeysPoints>(InPCGPointData->GetPoints());
+			ColumnData.DataKeys = MakeShared<FPCGAttributeAccessorKeysPoints>(InPCGPointData->GetPoints());
 		}
 		else
-		{	
-			ColumnData.DataAccessor = PCGAttributeAccessorHelpers::CreateConstAccessor(InPCGPointData, TargetSelector);
-			ColumnData.DataKeys = PCGAttributeAccessorHelpers::CreateConstKeys(InPCGPointData, TargetSelector);
+		{
+			ColumnData.DataAccessor = TSharedPtr<const IPCGAttributeAccessor>(PCGAttributeAccessorHelpers::CreateConstAccessor(InPCGPointData, TargetSelector).Release());
+			ColumnData.DataKeys = TSharedPtr<const IPCGAttributeAccessorKeys>(PCGAttributeAccessorHelpers::CreateConstKeys(InPCGPointData, TargetSelector).Release());
 		}
-		
+
 		ToolTip = FText::FromString(PCG::Private::GetTypeName(ColumnData.DataAccessor->GetUnderlyingType()));
 	}
 
 	const float ColumnWidth = PCGEditorGraphAttributeListView::CalculateColumnWidth(ColumnLabel);
-	
+
 	SHeaderRow::FColumn::FArguments Arguments;
 	Arguments.ColumnId(InColumnId);
 	Arguments.DefaultLabel(ColumnLabel);
@@ -1174,7 +1221,7 @@ void SPCGEditorGraphAttributeListView::AddColumn(const UPCGPointData* InPCGPoint
 	Arguments.OverflowPolicy(ETextOverflowPolicy::Ellipsis);
 
 	SHeaderRow::FColumn* NewColumn = new SHeaderRow::FColumn(Arguments);
-	NewColumn->bIsVisible = !HiddenAttributes.Contains(InColumnId); 
+	NewColumn->bIsVisible = !HiddenAttributes.Contains(InColumnId);
 	ListViewHeader->AddColumn(*NewColumn);
 }
 
@@ -1232,17 +1279,17 @@ void SPCGEditorGraphAttributeListView::AddMetadataColumn(const UPCGData* InPCGDa
 		ColumnIdString = FString::Printf(TEXT("%s%s"), *PCGEditorGraphAttributeListView::NoneAttributeId, PostFix);
 		ColumnLabel = FText::Format(LOCTEXT("NoneLabelFormat", "{0}{1}"), UEnum::GetDisplayValueAsText(InMetadataType), FText::FromString(PostFix));
 	}
-	
+
 	const FName ColumnId(ColumnIdString);
 
 	if (InPCGData)
 	{
 		FPCGAttributePropertyInputSelector TargetSelector;
 		TargetSelector.Update(OriginalColumnIdString);
-	
+
 		FPCGColumnData& ColumnData = PCGColumnData.Add(ColumnId);
-		ColumnData.DataAccessor = PCGAttributeAccessorHelpers::CreateConstAccessor(InPCGData, TargetSelector);
-		ColumnData.DataKeys = PCGAttributeAccessorHelpers::CreateConstKeys(InPCGData, TargetSelector);
+		ColumnData.DataAccessor = TSharedPtr<const IPCGAttributeAccessor>(PCGAttributeAccessorHelpers::CreateConstAccessor(InPCGData, TargetSelector).Release());
+		ColumnData.DataKeys = TSharedPtr<const IPCGAttributeAccessorKeys>(PCGAttributeAccessorHelpers::CreateConstKeys(InPCGData, TargetSelector).Release());
 
 		ToolTip = FText::FromString(PCG::Private::GetTypeName(ColumnData.DataAccessor->GetUnderlyingType()));
 	}
@@ -1251,7 +1298,7 @@ void SPCGEditorGraphAttributeListView::AddMetadataColumn(const UPCGData* InPCGDa
 	{
 		ColumnLabel = FText::FromName(ColumnId);
 	}
-	
+
 	float ColumnWidth = 0.0f;
 
 	EHorizontalAlignment CellAlignment = EHorizontalAlignment::HAlign_Right;
diff --git a/Engine/Plugins/Experimental/PCG/Source/PCGEditor/Private/SPCGEditorGraphAttributeListView.h b/Engine/Plugins/Experimental/PCG/Source/PCGEditor/Private/SPCGEditorGraphAttributeListView.h
--- a/Engine/Plugins/Experimental/PCG/Source/PCGEditor/Private/SPCGEditorGraphAttributeListView.h
+++ b/Engine/Plugins/Experimental/PCG/Source/PCGEditor/Private/SPCGEditorGraphAttributeListView.h
@@ -7,6 +7,7 @@
 #include "Metadata/Accessors/PCGAttributeAccessorKeys.h"
 
 #include "Misc/TextFilterExpressionEvaluator.h"
+#include "Tasks/Task.h"
 #include "Widgets/Views/SListView.h"
 #include "Widgets/Views/STableRow.h"
 
@@ -40,8 +41,8 @@ struct FPCGListViewItem
 
 struct FPCGColumnData
 {
-	TUniquePtr<const IPCGAttributeAccessor> DataAccessor;
-	TUniquePtr<const IPCGAttributeAccessorKeys> DataKeys;
+	TSharedPtr<const IPCGAttributeAccessor> DataAccessor;
+	TSharedPtr<const IPCGAttributeAccessorKeys> DataKeys;
 };
 
 typedef TSharedPtr<FPCGListViewItem> PCGListviewItemPtr;
@@ -53,6 +54,42 @@ struct FTextAsNumberIsValid : std::false_type {};
 template <typename T>
 struct FTextAsNumberIsValid<T, std::void_t<decltype(FText::AsNumber(std::declval<T>()))>> : std::true_type {};
 
+/** Class used for threaded filtering and sorting of list view items */
+class FPCGListViewUpdater : public TSharedFromThis<FPCGListViewUpdater>
+{
+public:
+	FPCGListViewUpdater(
+		const TArray<PCGListviewItemPtr>& InListViewItems,
+		const TMap<FName, FPCGColumnData>& InColumnData,
+		const EColumnSortMode::Type InSortMode,
+		const FName InSortingColumn,
+		const TSharedPtr<FTextFilterExpressionEvaluator>& InTextFilter)
+	: ListViewItems(InListViewItems)
+	, ColumnData(InColumnData)
+	, SortMode(InSortMode)
+	, SortingColumn(InSortingColumn)
+	, TextFilter(InTextFilter)
+	{}
+
+	bool IsCompleted() const;
+	void Launch();
+
+	TArray<PCGListviewItemPtr> ListViewItems;
+
+private:
+	void AsyncSort();
+	void AsyncFilter();
+
+	TMap<FName, FPCGColumnData> ColumnData;
+
+	EColumnSortMode::Type SortMode = EColumnSortMode::Type::Ascending;
+	FName SortingColumn = NAME_None;
+
+	TSharedPtr<FTextFilterExpressionEvaluator> TextFilter;
+
+	UE::Tasks::FTask UpdateTask;
+};
+
 class SPCGListViewItemRow : public SMultiColumnTableRow<PCGListviewItemPtr>
 {
 public:
@@ -64,7 +101,7 @@ public:
 	void Construct(const FArguments& InArgs, const TSharedRef<STableViewBase>& InOwnerTableView);
 
 	virtual TSharedRef<SWidget> GenerateWidgetForColumn(const FName& ColumnId) override;
-	
+
 private:
 	TWeakPtr<SPCGEditorGraphAttributeListView> AttributeListView;
 	PCGListviewItemPtr InternalItem;
@@ -88,7 +125,7 @@ private:
 class SPCGEditorGraphAttributeListView : public SCompoundWidget
 {
 	friend SPCGListViewItemRow;
-	
+
 public:
 	SLATE_BEGIN_ARGS(SPCGEditorGraphAttributeListView) {}
 	SLATE_END_ARGS()
@@ -100,7 +137,7 @@ public:
 	virtual void Tick(const FGeometry& AllottedGeometry, const double InCurrentTime, const float InDeltaTime) override;
 
 	void RequestRefresh() { bNeedsRefresh = true; }
-	
+
 private:
 	TSharedRef<SHeaderRow> CreateHeaderRowWidget() const;
 
@@ -111,12 +148,12 @@ private:
 	void OnGenerateUpdated(UPCGComponent* InPCGComponent);
 
 	const FPCGDataCollection* GetInspectionData();
-	
+
 	void RefreshAttributeList();
 	void RefreshPinComboBox();
 	void RefreshDataComboBox();
 
-	void ApplyRowFilter();
+	void LaunchUpdateTask();
 
 	/** Only connected input pins are added to combo box, so keep track of the node pin index for each item. */
 	struct FPinComboBoxItem
@@ -150,7 +187,6 @@ private:
 
 	void OnColumnSortModeChanged(const EColumnSortPriority::Type SortPriority, const FName& ColumnId, const EColumnSortMode::Type InSortMode);
 	EColumnSortMode::Type GetColumnSortMode(const FName InColumnId) const;
-	void RefreshSorting();
 
 	void OnFilterTextChanged(const FText& InFilterText);
 	void OnFilterTextCommitted(const FText& NewText, ETextCommit::Type CommitInfo);
@@ -171,7 +207,7 @@ private:
 
 	TSharedPtr<FTextFilterExpressionEvaluator> TextFilter;
 
-	TSharedPtr<SSearchBox> SearchBoxWidget;	
+	TSharedPtr<SSearchBox> SearchBoxWidget;
 	TSharedPtr<SHeaderRow> ListViewHeader;
 	TSharedPtr<SListView<PCGListviewItemPtr>> ListView;
 	TArray<PCGListviewItemPtr> ListViewItems;
@@ -182,7 +218,7 @@ private:
 
 	TSharedPtr<SComboBox<TSharedPtr<FName>>> DataComboBox;
 	TArray<TSharedPtr<FName>> DataComboBoxItems;
-	
+
 	TSharedPtr<STextBlock> NodeNameTextBlock;
 	TSharedPtr<STextBlock> InfoTextBlock;
 	TSharedPtr<SComboButton> FilterButton;
@@ -197,4 +233,6 @@ private:
 	EColumnSortMode::Type SortMode = EColumnSortMode::Type::Ascending;
 
 	bool bNeedsRefresh = false;
+
+	TSharedPtr<FPCGListViewUpdater> CurrentUpdateTask = nullptr;
 };
