[PCG] Updated the tooltip in the Point Filter element for bUseSpatialQuery to be a little more clear.

#rnx

diff --git a/Engine/Plugins/Experimental/PCG/Source/PCG/Public/Elements/PCGPointFilter.h b/Engine/Plugins/Experimental/PCG/Source/PCG/Public/Elements/PCGPointFilter.h
--- a/Engine/Plugins/Experimental/PCG/Source/PCG/Public/Elements/PCGPointFilter.h
+++ b/Engine/Plugins/Experimental/PCG/Source/PCG/Public/Elements/PCGPointFilter.h
@@ -64,7 +64,7 @@ struct PCG_API FPCGPointFilterThresholdSettings
 	UPROPERTY(BlueprintReadWrite, EditAnywhere, Category = Settings, meta = (EditCondition = "!bUseConstantThreshold", EditConditionHides, PCG_NotOverridable))
 	FPCGAttributePropertyInputSelector ThresholdAttribute;
 
-	/** If the threshold data is Point data, it will sample input points in threshold data. Always true with Spatial data.*/
+	/** For Point Data, enabling this option will use sampling rather than comparing points 1 to 1 directly. For other spatial data, this is always true. */
 	UPROPERTY(BlueprintReadWrite, EditAnywhere, Category = Settings, meta = (EditCondition = "!bUseConstantThreshold", EditConditionHides, PCG_Overridable))
 	bool bUseSpatialQuery = true;
 
