[PCG] Tentative fix for Editor test, testing PCG volume factory

* Can't repro locally and seems to be random. Add more checks to be sure to not use null pointers.
* Ran the preflight until the faulty job, failed but not on this test.

#rnx

diff --git a/Engine/Plugins/Experimental/PCG/Source/PCGEditor/Private/Tests/Display/PCGComponentInNewPCGVolumes.cpp b/Engine/Plugins/Experimental/PCG/Source/PCGEditor/Private/Tests/Display/PCGComponentInNewPCGVolumes.cpp
--- a/Engine/Plugins/Experimental/PCG/Source/PCGEditor/Private/Tests/Display/PCGComponentInNewPCGVolumes.cpp
+++ b/Engine/Plugins/Experimental/PCG/Source/PCGEditor/Private/Tests/Display/PCGComponentInNewPCGVolumes.cpp
@@ -20,20 +20,22 @@ IMPLEMENT_CUSTOM_SIMPLE_AUTOMATION_TEST(FPCGComponentInNewPCGVolumes, FPCGTestBa
 
 bool FPCGComponentInNewPCGVolumes::RunTest(const FString& Parameters)
 {
+	UTEST_NOT_NULL(TEXT("Failed to get editor!"), GEditor);
+	// Condition for static analysis.
+	check(GEditor);
+
 	// Get current world
 	UWorld* World = GEditor->GetEditorWorldContext().World();
 	UTEST_NOT_NULL(TEXT("Failed to get editor world context!"), World);
 
 	// Condition for static analysis.
-	if (!World)
-	{
-		return false;
-	}
+	check(World);
 
 	TSubclassOf<APCGVolume> PCGVolumeClass = APCGVolume::StaticClass();
 	UActorFactory* PCGVolumeFactory = GEditor->FindActorFactoryForActorClass(PCGVolumeClass);
 
 	UTEST_NOT_NULL(TEXT("Failed to find PCGVolume actor factory."), PCGVolumeFactory);
+	check(PCGVolumeFactory);
 
 	APCGVolume* VolumeActor = nullptr;
 
@@ -42,37 +44,35 @@ bool FPCGComponentInNewPCGVolumes::RunTest(const FString& Parameters)
 		FTransform ActorTransform;
 		VolumeActor = Cast<APCGVolume>(GEditor->UseActorFactory(PCGVolumeFactory, FAssetData(PCGVolumeClass), &ActorTransform));
 	}
+
 	UTEST_NOT_NULL(TEXT("Failed to add PCGVolume actor."), VolumeActor);
 	// Condition for static analysis.
-	if (VolumeActor)
+	check(VolumeActor);
+
+	VolumeActor->SetFlags(RF_Transient);
+
+	UPCGComponent* PCGComponent = VolumeActor->FindComponentByClass<UPCGComponent>();
+	UTEST_NOT_NULL(TEXT("PCGVolume actor does not contain a PCGComponent!"), PCGComponent);
+	// Condition for static analysis.
+	check(PCGComponent);
+
+	PCGComponent->ClearPCGLink(APCGVolume::StaticClass());
+
+	// Cleanup the VolumeActor
+	FAssetRegistryModule& AssetRegistryModule = FModuleManager::LoadModuleChecked<FAssetRegistryModule>("AssetRegistry");
+	IAssetRegistry& AssetRegistry = AssetRegistryModule.Get();
+	FAssetData PCGVolumeAssetData(VolumeActor);
+	AssetRegistry.AssetDeleted(PCGVolumeAssetData.GetAsset());
+
+	bool bSuccessful = ObjectTools::DeleteSingleObject(VolumeActor, true);
+	if (!bSuccessful)
 	{
-		UPackage* Package = VolumeActor->GetOutermost();
-		Package->SetFlags(RF_Transient);
-		TWeakObjectPtr<UPCGComponent> PCGComponent = VolumeActor->FindComponentByClass<UPCGComponent>();
-		TestNotNull(TEXT("PCGVolume actor does not contain a PCGComponent!"), PCGComponent.Get());
-
-		if (PCGComponent.IsValid())
-		{
-			PCGComponent->ClearPCGLink(APCGVolume::StaticClass());
-			PCGComponent = nullptr;
-
-		}
-
-		// Cleanup the VolumeActor
-		FAssetRegistryModule& AssetRegistryModule = FModuleManager::LoadModuleChecked<FAssetRegistryModule>("AssetRegistry");
-		IAssetRegistry& AssetRegistry = AssetRegistryModule.Get();
-		FAssetData PCGVolumeAssetData(VolumeActor);
-		AssetRegistry.AssetDeleted(PCGVolumeAssetData.GetAsset());
-
-		bool bSuccessful = ObjectTools::DeleteSingleObject(VolumeActor, true);
-		if (!bSuccessful)
-		{
-			//Clear references to the object so we can delete it
-			FAutomationEditorCommonUtils::NullReferencesToObject(VolumeActor);
-			bSuccessful = ObjectTools::DeleteSingleObject(VolumeActor, false);
-			TestTrue(TEXT("References to PCGVolume component could not be cleaned"), bSuccessful);
-		}
+		//Clear references to the object so we can delete it
+		FAutomationEditorCommonUtils::NullReferencesToObject(VolumeActor);
+		bSuccessful = ObjectTools::DeleteSingleObject(VolumeActor, false);
+		TestTrue(TEXT("References to PCGVolume component could not be cleaned"), bSuccessful);
 	}
+
 	return true;
 };
 
