[PCG] Enable Arrays and Arrays of Structs in Get Actor Property node

* It is limited to Arrays of basic types, or arrays of structs of basic types. Arays of arrays, arrays in structs are not supported.
* Also add a new node to get the number of entries in an attribute set
* Since we support arrays now, make Get Number of Points support multiple inputs.

diff --git a/Engine/Plugins/Experimental/PCG/Source/PCG/Private/Elements/PCGNumberOfElements.cpp b/Engine/Plugins/Experimental/PCG/Source/PCG/Private/Elements/PCGNumberOfElements.cpp
new file mode 100644
--- /dev/null
+++ b/Engine/Plugins/Experimental/PCG/Source/PCG/Private/Elements/PCGNumberOfElements.cpp
@@ -0,0 +1,137 @@
+// Copyright Epic Games, Inc. All Rights Reserved.
+
+#include "Elements/PCGNumberOfElements.h"
+
+#include "PCGContext.h"
+#include "PCGPin.h"
+#include "PCGParamData.h"
+#include "Data/PCGPointData.h"
+#include "Metadata/PCGMetadata.h"
+
+#define LOCTEXT_NAMESPACE "PCGNumberOfElementsSettings"
+
+////////////////////////////////////
+// UPCGNumberOfElementsBaseSettings
+////////////////////////////////////
+
+TArray<FPCGPinProperties> UPCGNumberOfElementsBaseSettings::OutputPinProperties() const
+{
+	TArray<FPCGPinProperties> PinProperties;
+	PinProperties.Emplace(PCGPinConstants::DefaultOutputLabel, EPCGDataType::Param);
+
+	return PinProperties;
+}
+
+
+////////////////////////////////////
+// UPCGNumberOfPointsSettings
+////////////////////////////////////
+
+TArray<FPCGPinProperties> UPCGNumberOfPointsSettings::InputPinProperties() const
+{
+	TArray<FPCGPinProperties> PinProperties;
+	PinProperties.Emplace(PCGPinConstants::DefaultInputLabel, EPCGDataType::Point, /*bInAllowMultipleConnections=*/true, /*bAllowMultipleData=*/true);
+
+	return PinProperties;
+}
+
+FPCGElementPtr UPCGNumberOfPointsSettings::CreateElement() const
+{
+	return MakeShared<FPCGNumberOfPointsElement>();
+}
+
+
+////////////////////////////////////
+// UPCGNumberOfEntriesSettings
+////////////////////////////////////
+
+TArray<FPCGPinProperties> UPCGNumberOfEntriesSettings::InputPinProperties() const
+{
+	TArray<FPCGPinProperties> PinProperties;
+	PinProperties.Emplace(PCGPinConstants::DefaultInputLabel, EPCGDataType::Param, /*bInAllowMultipleConnections=*/true, /*bAllowMultipleData=*/true);
+
+	return PinProperties;
+}
+
+FPCGElementPtr UPCGNumberOfEntriesSettings::CreateElement() const
+{
+	return MakeShared<FPCGNumberOfEntriesElement>();
+}
+
+
+////////////////////////////////////
+// FPCGNumberOfElementsBaseElement
+////////////////////////////////////
+
+template <typename DataType>
+bool FPCGNumberOfElementsBaseElement<DataType>::ExecuteInternal(FPCGContext* Context) const
+{
+	TRACE_CPUPROFILER_EVENT_SCOPE(FPCGNumberOfElementsBaseElement::Execute);
+
+	check(Context);
+
+	const UPCGNumberOfElementsBaseSettings* Settings = Context->GetInputSettings<UPCGNumberOfElementsBaseSettings>();
+	check(Settings);
+
+	TArray<FPCGTaggedData> Inputs = Context->InputData.GetInputsByPin(PCGPinConstants::DefaultInputLabel);
+
+	UPCGParamData* OutputParamData = nullptr;
+	FPCGMetadataAttribute<int32>* NewAttribute = nullptr;
+
+	for (int32 i = 0; i < Inputs.Num(); ++i)
+	{
+		const DataType* InputData = Cast<DataType>(Inputs[i].Data);
+
+		if (!InputData)
+		{
+			PCGE_LOG(Error, GraphAndLog, FText::Format(LOCTEXT("InputNotRightType", "Input {0} is not of the right input type, discarded"), FText::AsNumber(i)));
+			continue;
+		}
+
+		if (!OutputParamData)
+		{
+			OutputParamData = NewObject<UPCGParamData>();
+			NewAttribute = OutputParamData->Metadata->CreateAttribute<int32>(Settings->OutputAttributeName, 0, /*bAllowInterpolation=*/true, /*bOverrideParent=*/false);
+
+			if (!NewAttribute)
+			{
+				PCGE_LOG(Error, GraphAndLog, FText::Format(LOCTEXT("AttributeFailedToBeCreated", "New Attribute {0} failed to be created."), FText::FromName(Settings->OutputAttributeName)));
+				return true;
+			}
+		}
+
+		check(OutputParamData && NewAttribute);
+		NewAttribute->SetValue(OutputParamData->Metadata->AddEntry(), GetNum(InputData));
+	}
+
+	if (OutputParamData)
+	{
+		FPCGTaggedData& Output = Context->OutputData.TaggedData.Emplace_GetRef();
+		Output.Data = OutputParamData;
+		Output.Pin = PCGPinConstants::DefaultOutputLabel;
+	}
+
+	return true;
+}
+
+
+////////////////////////////////////
+// FPCGNumberOfPointsElement
+////////////////////////////////////
+
+int32 FPCGNumberOfPointsElement::GetNum(const UPCGPointData* InData) const
+{
+	return InData ? InData->GetPoints().Num() : 0;
+}
+
+
+////////////////////////////////////
+// FPCGNumberOfEntriesElement
+////////////////////////////////////
+
+int32 FPCGNumberOfEntriesElement::GetNum(const UPCGParamData* InData) const
+{
+	return (InData && InData->Metadata) ? InData->Metadata->GetLocalItemCount() : 0;
+}
+
+#undef LOCTEXT_NAMESPACE
diff --git a/Engine/Plugins/Experimental/PCG/Source/PCG/Private/Elements/PCGNumberOfPoints.cpp b/Engine/Plugins/Experimental/PCG/Source/PCG/Private/Elements/PCGNumberOfPoints.cpp
deleted file mode 100644
diff --git a/Engine/Plugins/Experimental/PCG/Source/PCG/Private/Elements/PCGPropertyToParamData.cpp b/Engine/Plugins/Experimental/PCG/Source/PCG/Private/Elements/PCGPropertyToParamData.cpp
--- a/Engine/Plugins/Experimental/PCG/Source/PCG/Private/Elements/PCGPropertyToParamData.cpp
+++ b/Engine/Plugins/Experimental/PCG/Source/PCG/Private/Elements/PCGPropertyToParamData.cpp
@@ -171,29 +171,38 @@ bool FPCGPropertyToParamDataElement::ExecuteInternal(FPCGContext* Context) const
 		return true;
 	}
 
-	using ExtractablePropertyTuple = TTuple<FName, const void*, const FProperty*>;
+	// If the property is an array, we will work on the underlying property, and extract each element as an entry in the param data
+	FArrayProperty* ArrayProperty = CastField<FArrayProperty>(Property);
+	if (ArrayProperty)
+	{
+		Property = ArrayProperty->Inner;
+	}
+
+	using GetAddressFunc = TFunction<const void* (const void*)>;
+	using ExtractablePropertyTuple = TTuple<FName, const FProperty*>;
 	TArray<ExtractablePropertyTuple> ExtractableProperties;
 
+	GetAddressFunc AddressFunc;
+
 	// Special case where the property is a struct/object, that is not supported by our metadata, we will try to break it down to multiple attributes in the resulting param data, if asked.
-	if (!PCGAttributeAccessorHelpers::IsPropertyAccessorSupported(Property) && (Property->IsA<FStructProperty>() || Property->IsA<FObjectProperty>()) && Settings->bExtractObjectAndStruct)
+	if ((Property->IsA<FStructProperty>() || Property->IsA<FObjectProperty>()) && Settings->bExtractObjectAndStruct)
 	{
 		UScriptStruct* UnderlyingStruct = nullptr;
 		UClass* UnderlyingClass = nullptr;
-		const void* ObjectAddress = nullptr;
 
 		if (FStructProperty* StructProperty = CastField<FStructProperty>(Property))
 		{
 			UnderlyingStruct = StructProperty->Struct;
-			ObjectAddress = StructProperty->ContainerPtrToValuePtr<void>(ObjectToInspect);
+			AddressFunc = [StructProperty](const void* InAddress) { return StructProperty->ContainerPtrToValuePtr<void>(InAddress); };
 		}
 		else if (FObjectProperty* ObjectProperty = CastField<FObjectProperty>(Property))
 		{
 			UnderlyingClass = ObjectProperty->PropertyClass;
-			ObjectAddress = ObjectProperty->GetObjectPropertyValue_InContainer(ObjectToInspect);
+			AddressFunc = [ObjectProperty](const void* InAddress) { return ObjectProperty->GetObjectPropertyValue_InContainer(InAddress); };
 		}
-		
+
 		check(UnderlyingStruct || UnderlyingClass);
-		check(ObjectAddress);
+		check(!!AddressFunc);
 
 		// Re-use code from overridable params
 		// Limit ourselves to not recurse into more structs.
@@ -216,7 +225,7 @@ bool FPCGPropertyToParamDataElement::ExecuteInternal(FPCGContext* Context) const
 					// We use authored name as attribute name to avoid issue with noisy property names, like in UUserDefinedStructs, where some random number is appended to the property name.
 					// By default, it will just return the property name anyway.
 					const FString AuthoredName = UnderlyingStruct ? UnderlyingStruct->GetAuthoredNameForField(ChildProperty) : UnderlyingClass->GetAuthoredNameForField(ChildProperty);
-					ExtractableProperties.Emplace(FName(AuthoredName), ObjectAddress, ChildProperty);
+					ExtractableProperties.Emplace(FName(AuthoredName), ChildProperty);
 				}
 			}
 		}
@@ -224,7 +233,9 @@ bool FPCGPropertyToParamDataElement::ExecuteInternal(FPCGContext* Context) const
 	else
 	{
 		const FName AttributeName = (Settings->OutputAttributeName == PCGMetadataAttributeConstants::SourceNameAttributeName) ? Property->GetFName() : Settings->OutputAttributeName;
-		ExtractableProperties.Emplace(AttributeName, ObjectToInspect, Property);
+		ExtractableProperties.Emplace(AttributeName, Property);
+		// Identity
+		AddressFunc = [](const void* InAddress) { return InAddress; };
 	}
 
 	if (ExtractableProperties.IsEmpty())
@@ -233,26 +244,54 @@ bool FPCGPropertyToParamDataElement::ExecuteInternal(FPCGContext* Context) const
 		return true;
 	}
 
+	// Before we need to compute all the addresses for each entry in our array (or just a single entry if there is no array)
+	TArray<const void*, TInlineAllocator<16>> ElementAddresses;
+	if (ArrayProperty)
+	{
+		FScriptArrayHelper_InContainer Helper(ArrayProperty, ObjectToInspect);
+		ElementAddresses.Reserve(Helper.Num());
+		for (int32 DynamicIndex = 0; DynamicIndex < Helper.Num(); ++DynamicIndex)
+		{
+			ElementAddresses.Add(Helper.GetRawPtr(DynamicIndex));
+		}
+	}
+	else
+	{
+		ElementAddresses.Add(ObjectToInspect);
+	}
+
 	// From there, we should be able to create the data.
 	UPCGParamData* ParamData = NewObject<UPCGParamData>();
 	UPCGMetadata* Metadata = ParamData->MutableMetadata();
 	check(Metadata);
-	PCGMetadataEntryKey EntryKey = Metadata->AddEntry();
-	bool bValidOperation = false;
 
-	for (ExtractablePropertyTuple& ExtractableProperty : ExtractableProperties)
-	{
-		const FName AttributeName = ExtractableProperty.Get<0>();
-		const void* ContainerPtr = ExtractableProperty.Get<1>();
-		const FProperty* FinalProperty = ExtractableProperty.Get<2>();
+	bool bValidOperation = true;
 
-		if (!Metadata->SetAttributeFromDataProperty(AttributeName, EntryKey, ContainerPtr, FinalProperty, /*bCreate=*/ true))
+	for (const void* ElementAddress : ElementAddresses)
+	{
+		if (!bValidOperation)
 		{
-			PCGE_LOG(Error, GraphAndLog, FText::Format(LOCTEXT("ErrorCreatingAttribute", "Error while creating an attribute for property '{0}'. Either the property type is not supported by PCG or attribute creation failed."), FText::FromString(FinalProperty->GetName())));
-			continue;
+			break;
 		}
 
-		bValidOperation = true;
+		// Add a new entry for all elements
+		PCGMetadataEntryKey EntryKey = Metadata->AddEntry();
+
+		for (ExtractablePropertyTuple& ExtractableProperty : ExtractableProperties)
+		{
+			const FName AttributeName = ExtractableProperty.Get<0>();
+			const FProperty* FinalProperty = ExtractableProperty.Get<1>();
+
+			// Offset the address if needed
+			const void* ContainerPtr = AddressFunc(ElementAddress);
+
+			if (!Metadata->SetAttributeFromDataProperty(AttributeName, EntryKey, ContainerPtr, FinalProperty, /*bCreate=*/ true))
+			{
+				PCGE_LOG(Error, GraphAndLog, FText::Format(LOCTEXT("ErrorCreatingAttribute", "Error while creating an attribute for property '{0}'. Either the property type is not supported by PCG or attribute creation failed."), FText::FromString(FinalProperty->GetName())));
+				bValidOperation = false;
+				break;
+			}
+		}
 	}
 
 	if (bValidOperation)
diff --git a/Engine/Plugins/Experimental/PCG/Source/PCG/Public/Elements/PCGNumberOfElements.h b/Engine/Plugins/Experimental/PCG/Source/PCG/Public/Elements/PCGNumberOfElements.h
new file mode 100644
--- /dev/null
+++ b/Engine/Plugins/Experimental/PCG/Source/PCG/Public/Elements/PCGNumberOfElements.h
@@ -0,0 +1,99 @@
+// Copyright Epic Games, Inc. All Rights Reserved.
+
+#pragma once
+
+#include "PCGSettings.h"
+
+#include "PCGNumberOfElements.generated.h"
+
+class UPCGParamData;
+class UPCGPointData;
+
+/**
+* Elements for getting the number of elements in a point data or a param data. Since the whole logic is identical
+* except for getting the number of elements, it is factorized in a base class.
+*/
+
+// Base class for common elements
+UCLASS(Abstract)
+class UPCGNumberOfElementsBaseSettings : public UPCGSettings
+{
+	GENERATED_BODY()
+
+protected:
+	virtual TArray<FPCGPinProperties> OutputPinProperties() const override;
+
+public:
+	UPROPERTY(BlueprintReadWrite, EditAnywhere, Category = Settings, meta = (PCG_Overridable))
+	FName OutputAttributeName = TEXT("NumEntries");
+};
+
+/**
+* Return the number of points in the input point data.
+*/
+
+UCLASS(BlueprintType, ClassGroup = (Procedural))
+class UPCGNumberOfPointsSettings : public UPCGNumberOfElementsBaseSettings
+{
+	GENERATED_BODY()
+
+public:
+	//~Begin UPCGSettings interface
+#if WITH_EDITOR
+	virtual FName GetDefaultNodeName() const override { return FName(TEXT("GetPointsCount")); }
+	virtual FText GetDefaultNodeTitle() const override { return NSLOCTEXT("PCGNumberOfElementsSettings", "NodeTitle", "Get Points Count"); }
+	virtual FText GetNodeTooltipText() const override { return NSLOCTEXT("PCGNumberOfElementsSettings", "NodeTooltip", "Return the number of points in the input point data."); }
+	virtual EPCGSettingsType GetType() const override { return EPCGSettingsType::Spatial; }
+#endif
+
+protected:
+	virtual TArray<FPCGPinProperties> InputPinProperties() const override;
+	virtual FPCGElementPtr CreateElement() const override;
+	//~End UPCGSettings interface
+};
+
+/**
+* Return the number of entries in the input attribute sets.
+*/
+
+UCLASS(BlueprintType, ClassGroup = (Procedural))
+class UPCGNumberOfEntriesSettings : public UPCGNumberOfElementsBaseSettings
+{
+	GENERATED_BODY()
+
+public:
+	//~Begin UPCGSettings interface
+#if WITH_EDITOR
+	virtual FName GetDefaultNodeName() const override { return FName(TEXT("GetEntriesCount")); }
+	virtual FText GetDefaultNodeTitle() const override { return NSLOCTEXT("PCGNumberOfElementsSettings", "NodeTitle", "Get Entries Count"); }
+	virtual FText GetNodeTooltipText() const override { return NSLOCTEXT("PCGNumberOfElementsSettings", "NodeTooltip", "Return the number of entries in the input attribute sets."); }
+	virtual EPCGSettingsType GetType() const override { return EPCGSettingsType::Generic; }
+#endif
+
+protected:
+	virtual TArray<FPCGPinProperties> InputPinProperties() const override;
+	virtual FPCGElementPtr CreateElement() const override;
+	//~End UPCGSettings interface
+};
+
+
+template <typename DataType>
+class FPCGNumberOfElementsBaseElement : public FSimplePCGElement
+{
+protected:
+	virtual bool ExecuteInternal(FPCGContext* Context) const override;
+
+	virtual int32 GetNum(const DataType* InData) const = 0;
+};
+
+class FPCGNumberOfPointsElement : public FPCGNumberOfElementsBaseElement<UPCGPointData>
+{
+protected:
+	virtual int32 GetNum(const UPCGPointData* InData) const override;
+};
+
+class FPCGNumberOfEntriesElement : public FPCGNumberOfElementsBaseElement<UPCGParamData>
+{
+protected:
+	virtual int32 GetNum(const UPCGParamData* InData) const override;
+};
\ No newline at end of file
diff --git a/Engine/Plugins/Experimental/PCG/Source/PCG/Public/Elements/PCGNumberOfPoints.h b/Engine/Plugins/Experimental/PCG/Source/PCG/Public/Elements/PCGNumberOfPoints.h
deleted file mode 100644
