Allow for specifying the landscape cache output data when cooked.  Options allow for removing layer data and/or position data.

diff --git a/Engine/Plugins/Experimental/PCG/Source/PCG/Private/Grid/PCGLandscapeCache.cpp b/Engine/Plugins/Experimental/PCG/Source/PCG/Private/Grid/PCGLandscapeCache.cpp
--- a/Engine/Plugins/Experimental/PCG/Source/PCG/Private/Grid/PCGLandscapeCache.cpp
+++ b/Engine/Plugins/Experimental/PCG/Source/PCG/Private/Grid/PCGLandscapeCache.cpp
@@ -353,18 +353,29 @@ int32 FPCGLandscapeCacheEntry::GetMemorySize() const
 	return MemSize;
 }
 
-void FPCGLandscapeCacheEntry::SerializeToBulkData()
+void FPCGLandscapeCacheEntry::SerializeToBulkData(EPCGLandscapeCacheSerializationContents SerializationContents)
 {
+	const bool bSerializeEverything = SerializationContents == EPCGLandscapeCacheSerializationContents::SerializeAll;
+	const bool bShouldSerializePositionAndNormals = (bSerializeEverything || SerializationContents == EPCGLandscapeCacheSerializationContents::SerializeOnlyPositionsAndNormals);
+	const bool bShouldSerializeLayerData = (bSerializeEverything || SerializationContents == EPCGLandscapeCacheSerializationContents::SerializeOnlyLayerData);
+
+	// Use empty buffers if the contents should not be serialized
+	TArray<FVector> PositionsAndNormalsTemp;
+	TArray<TArray<uint8>> LayerDataTemp;
+	TArray<FVector>& SerializePositionsAndNormals = bShouldSerializePositionAndNormals ? PositionsAndNormals : PositionsAndNormalsTemp;
+	TArray<TArray<uint8>>& SerializeLayerData = bShouldSerializeLayerData ? LayerData : LayerDataTemp;
+
+
 	// Move data from local arrays to the bulk data
 	BulkData.Lock(LOCK_READ_WRITE);
 
 	int32 NumBytes = 0;
 	// Number of entries in the array + size of the array
-	NumBytes += sizeof(int32) + PositionsAndNormals.Num() * PositionsAndNormals.GetTypeSize();
+	NumBytes += sizeof(int32) + SerializePositionsAndNormals.Num() * SerializePositionsAndNormals.GetTypeSize();
 
 	// Number of layers
 	NumBytes += sizeof(int32); 
-	for (TArray<uint8>& CurrentLayerData : LayerData)
+	for (TArray<uint8>& CurrentLayerData : SerializeLayerData)
 	{
 		// Number of entries in the layer data array + size of the array
 		NumBytes += sizeof(int32) + CurrentLayerData.Num() * CurrentLayerData.GetTypeSize();
@@ -374,12 +385,12 @@ void FPCGLandscapeCacheEntry::SerializeToBulkData()
 	FBufferWriter Ar(Dest, NumBytes);
 	Ar.SetIsPersistent(true);
 
-	Ar << PositionsAndNormals;
+	Ar << SerializePositionsAndNormals;
 	
-	int32 LayerDataCount = LayerData.Num();
+	int32 LayerDataCount = SerializeLayerData.Num();
 	Ar << LayerDataCount;
 
-	for (TArray<uint8>& CurrentLayerData : LayerData)
+	for (TArray<uint8>& CurrentLayerData : SerializeLayerData)
 	{
 		Ar << CurrentLayerData;
 	}
@@ -413,15 +424,21 @@ void FPCGLandscapeCacheEntry::SerializeFromBulkData() const
 	bDataLoaded = true;
 }
 
-void FPCGLandscapeCacheEntry::Serialize(FArchive& Archive, UObject* Owner, int32 Index)
+void FPCGLandscapeCacheEntry::Serialize(FArchive& Archive, UObject* Owner, int32 Index, EPCGLandscapeCacheSerializationContents SerializeContents)
 {
+	// If the current bulk output data does not match the desired save data, make sure to load from the bulk data and update the bulked save data
+	if(Archive.IsCooking() && SerializeContents != EPCGLandscapeCacheSerializationContents::SerializeAll && !bDataLoaded)
+	{
+		SerializeFromBulkData();
+	}
+	
 	// Important implementation note:
 	// If the serialization here or in the cache entries change, we still need to load data from previous versions.
 	// While that's not really needed in non-editor builds, it is very much important when loading from the editor,
 	// at least in the "AlwaysSerialize" case.
 	if (bDataLoaded && Archive.IsSaving())
 	{
-		SerializeToBulkData();
+		SerializeToBulkData(SerializeContents);
 	}
 	
 	Archive << PointHalfSize;
@@ -451,7 +468,7 @@ void UPCGLandscapeCache::Serialize(FArchive& Archive)
 	const bool bShouldSerializeEntries = (Archive.IsSaving() &&
 		(SerializationMode == EPCGLandscapeCacheSerializationMode::AlwaysSerialize || 
 			(SerializationMode == EPCGLandscapeCacheSerializationMode::SerializeOnlyAtCook && Archive.IsCooking())));
-
+	EPCGLandscapeCacheSerializationContents SerializedContents = Archive.IsCooking() ? CookedSerializedContents : EPCGLandscapeCacheSerializationContents::SerializeAll;
 	// Important implementation note:
 	// If the serialization here or in the cache entries change, we still need to load data from previous versions.
 	// While that's not really needed in non-editor builds, it is very much important when loading from the editor,
@@ -478,7 +495,7 @@ void UPCGLandscapeCache::Serialize(FArchive& Archive)
 			Archive << Key;
 
 			FPCGLandscapeCacheEntry* Entry = new FPCGLandscapeCacheEntry();
-			Entry->Serialize(Archive, this, EntryIndex);
+			Entry->Serialize(Archive, this, EntryIndex, SerializedContents);
 
 			CachedData.Add(Key, Entry);
 		}
@@ -493,7 +510,7 @@ void UPCGLandscapeCache::Serialize(FArchive& Archive)
 		for (auto& CacheEntry : CachedData)
 		{
 			Archive << CacheEntry.Key;
-			CacheEntry.Value->Serialize(Archive, this, EntryIndex++);
+			CacheEntry.Value->Serialize(Archive, this, EntryIndex++, SerializedContents);
 		}
 	}
 }
diff --git a/Engine/Plugins/Experimental/PCG/Source/PCG/Public/Grid/PCGLandscapeCache.h b/Engine/Plugins/Experimental/PCG/Source/PCG/Public/Grid/PCGLandscapeCache.h
--- a/Engine/Plugins/Experimental/PCG/Source/PCG/Public/Grid/PCGLandscapeCache.h
+++ b/Engine/Plugins/Experimental/PCG/Source/PCG/Public/Grid/PCGLandscapeCache.h
@@ -22,6 +22,14 @@ enum class EPCGLandscapeCacheSerializationMode : uint8
 	AlwaysSerialize
 };
 
+UENUM()
+enum class EPCGLandscapeCacheSerializationContents : uint8
+{
+	SerializeOnlyPositionsAndNormals,
+	SerializeOnlyLayerData,
+	SerializeAll
+};
+
 USTRUCT(BlueprintType)
 struct FPCGLandscapeLayerWeight
 {
@@ -78,10 +86,10 @@ private:
 #endif
 
 	// Serialize called from the landscape cache
-	void Serialize(FArchive& Ar, UObject* Owner, int32 BulkIndex);
+	void Serialize(FArchive& Ar, UObject* Owner, int32 BulkIndex, EPCGLandscapeCacheSerializationContents SerializationContents);
 
 	// Internal usage methods
-	void SerializeToBulkData();
+	void SerializeToBulkData(EPCGLandscapeCacheSerializationContents SerializationContents);
 	void SerializeFromBulkData() const;
 
 	// Serialized data
@@ -139,6 +147,9 @@ public:
 	UPROPERTY(BlueprintReadWrite, EditAnywhere, Category = "Cache")
 	EPCGLandscapeCacheSerializationMode SerializationMode = EPCGLandscapeCacheSerializationMode::SerializeOnlyAtCook;
 
+	UPROPERTY(BlueprintReadWrite, EditAnywhere, Category = "Cache")
+	EPCGLandscapeCacheSerializationContents CookedSerializedContents = EPCGLandscapeCacheSerializationContents::SerializeAll;
+
 #if WITH_EDITORONLY_DATA
 	UPROPERTY(Transient, BlueprintReadOnly, VisibleAnywhere, Category = "Cache")
 	int32 CacheEntryCount = 0;
