[PCG] Fixed issue were cleaning a component would generate a world actor. We now avoid doing a cleanup if the component is not generated yet

diff --git a/Engine/Plugins/Experimental/PCG/Source/PCG/Private/PCGComponent.cpp b/Engine/Plugins/Experimental/PCG/Source/PCG/Private/PCGComponent.cpp
--- a/Engine/Plugins/Experimental/PCG/Source/PCG/Private/PCGComponent.cpp
+++ b/Engine/Plugins/Experimental/PCG/Source/PCG/Private/PCGComponent.cpp
@@ -521,7 +521,7 @@ void UPCGComponent::OnProcessGraphAborted(bool bQuiet)
 
 void UPCGComponent::Cleanup()
 {
-	if (!GetSubsystem() || IsCleaningUp())
+	if ((!bGenerated && !IsGenerating()) || !GetSubsystem() || IsCleaningUp())
 	{
 		return;
 	}
@@ -545,7 +545,7 @@ void UPCGComponent::CleanupLocal(bool bRemoveComponents, bool bSave)
 
 FPCGTaskId UPCGComponent::CleanupInternal(bool bRemoveComponents, bool bSave, const TArray<FPCGTaskId>& Dependencies)
 {
-	if (!GetSubsystem() || IsCleaningUp())
+	if ((!bGenerated && !IsGenerating()) || !GetSubsystem() || IsCleaningUp())
 	{
 		return InvalidPCGTaskId;
 	}
