[PCG] Soft object path and soft class path added as new metadata types.

Previously this was emulated somewhat by storing the paths as FStrings but this did not result in proper references that are visible to the cook.

Notes:
* Producers/consumers of asset references moved to the new types.
* Broadcast to string enabled to allow string operations on paths.
* Broadcast from string always allowed when getting values, so that legacy usages which grabbed path as string will continue to work.
* Deprecation added to migrate old path handling (string mode).
* Serialize method added to traits because FSoftObjectPath can't be directly serialized to archive.
* Fixed some project-specific static casting to FString that coudl fail when applied to SOP/SCP attributes.

#rnx

diff --git a/Engine/Plugins/Experimental/PCG/Source/PCG/Private/Data/PCGPointData.cpp b/Engine/Plugins/Experimental/PCG/Source/PCG/Private/Data/PCGPointData.cpp
--- a/Engine/Plugins/Experimental/PCG/Source/PCG/Private/Data/PCGPointData.cpp
+++ b/Engine/Plugins/Experimental/PCG/Source/PCG/Private/Data/PCGPointData.cpp
@@ -326,7 +326,7 @@ void UPCGPointData::AddToCrc(FArchiveCrc32& Ar, bool bFullDataCrc) const
 						if (const FPCGMetadataAttribute<AttributeType>* TypedAttribute = static_cast<const FPCGMetadataAttribute<AttributeType>*>(Attribute))
 						{
 							ValueWithType = TypedAttribute->GetValueFromItemKey(Point.MetadataEntry);
-							Ar << ValueWithType;
+							PCG::Private::Serialize(Ar, ValueWithType);
 						}
 					};
 
@@ -413,10 +413,10 @@ void UPCGPointData::AddSinglePointFromActor(AActor* InActor)
 
 	Point.MetadataEntry = Metadata->AddEntry();
 
-	FPCGMetadataAttribute<FString>* ActorReferenceAttribute = Metadata->FindOrCreateAttribute(PCGPointDataConstants::ActorReferenceAttribute, FString(), /*bAllowsInterpolation=*/false, /*bOverrideParent=*/false, /*bOverwriteIfTypeMismatch=*/false);
+	FPCGMetadataAttribute<FSoftObjectPath>* ActorReferenceAttribute = Metadata->FindOrCreateAttribute(PCGPointDataConstants::ActorReferenceAttribute, FSoftObjectPath(), /*bAllowsInterpolation=*/false, /*bOverrideParent=*/false, /*bOverwriteIfTypeMismatch=*/false);
 	if (ActorReferenceAttribute)
 	{
-		ActorReferenceAttribute->SetValue(Point.MetadataEntry, FSoftObjectPath(InActor).ToString());
+		ActorReferenceAttribute->SetValue(Point.MetadataEntry, FSoftObjectPath(InActor));
 	}
 
 	// Parse tags as well
diff --git a/Engine/Plugins/Experimental/PCG/Source/PCG/Private/Data/PCGWorldData.cpp b/Engine/Plugins/Experimental/PCG/Source/PCG/Private/Data/PCGWorldData.cpp
--- a/Engine/Plugins/Experimental/PCG/Source/PCG/Private/Data/PCGWorldData.cpp
+++ b/Engine/Plugins/Experimental/PCG/Source/PCG/Private/Data/PCGWorldData.cpp
@@ -60,7 +60,7 @@ bool UPCGWorldVolumetricData::SamplePoint(const FTransform& InTransform, const F
 
 	check(World.IsValid());
 
-	FPCGMetadataAttribute<FString>* ActorOverlappedAttribute = ((OutMetadata && QueryParams.bGetReferenceToActorHit) ? OutMetadata->GetMutableTypedAttribute<FString>(PCGPointDataConstants::ActorReferenceAttribute) : nullptr);
+	FPCGMetadataAttribute<FSoftObjectPath>* ActorOverlappedAttribute = ((OutMetadata && QueryParams.bGetReferenceToActorHit) ? OutMetadata->GetMutableTypedAttribute<FSoftObjectPath>(PCGPointDataConstants::ActorReferenceAttribute) : nullptr);
 
 	FCollisionObjectQueryParams ObjectQueryParams(QueryParams.CollisionChannel);
 	FCollisionShape CollisionShape = FCollisionShape::MakeBox(InBounds.GetExtent() * InTransform.GetScale3D());
@@ -143,7 +143,7 @@ bool UPCGWorldVolumetricData::SamplePoint(const FTransform& InTransform, const F
 			if (ActorOverlappedAttribute && OverlappedComponent->GetOwner())
 			{
 				OutMetadata->InitializeOnSet(OutPoint.MetadataEntry);
-				ActorOverlappedAttribute->SetValue(OutPoint.MetadataEntry, FSoftObjectPath(OverlappedComponent->GetOwner()).ToString());
+				ActorOverlappedAttribute->SetValue(OutPoint.MetadataEntry, FSoftObjectPath(OverlappedComponent->GetOwner()));
 			}
 
 			return true;
@@ -235,8 +235,8 @@ bool UPCGWorldRayHitData::SamplePoint(const FTransform& InTransform, const FBox&
 	// TODO: This seems to be a projection - along a direction. I suspect that UPCGWorldVolumetricData is the SamplePoint(), and this is the ProjectPoint() (in a direction)?
 	check(World.IsValid());
 
-	FPCGMetadataAttribute<FString>* ActorHitAttribute = ((OutMetadata && QueryParams.bGetReferenceToActorHit) ? OutMetadata->GetMutableTypedAttribute<FString>(PCGPointDataConstants::ActorReferenceAttribute) : nullptr);
-	FPCGMetadataAttribute<FString>* PhysicalMaterialAttribute = ((OutMetadata && QueryParams.bGetReferenceToPhysicalMaterial) ? OutMetadata->GetMutableTypedAttribute<FString>(PCGWorldRayHitConstants::PhysicalMaterialReferenceAttribute) : nullptr);
+	FPCGMetadataAttribute<FSoftObjectPath>* ActorHitAttribute = ((OutMetadata && QueryParams.bGetReferenceToActorHit) ? OutMetadata->GetMutableTypedAttribute<FSoftObjectPath>(PCGPointDataConstants::ActorReferenceAttribute) : nullptr);
+	FPCGMetadataAttribute<FSoftObjectPath>* PhysicalMaterialAttribute = ((OutMetadata && QueryParams.bGetReferenceToPhysicalMaterial) ? OutMetadata->GetMutableTypedAttribute<FSoftObjectPath>(PCGWorldRayHitConstants::PhysicalMaterialReferenceAttribute) : nullptr);
 
 	// Todo: consider prebuilding this
 	FCollisionObjectQueryParams ObjectQueryParams(QueryParams.CollisionChannel);
@@ -347,13 +347,13 @@ bool UPCGWorldRayHitData::SamplePoint(const FTransform& InTransform, const FBox&
 		if (ActorHitAttribute && HitComponent->GetOwner())
 		{
 			OutMetadata->InitializeOnSet(OutPoint.MetadataEntry);
-			ActorHitAttribute->SetValue(OutPoint.MetadataEntry, FSoftObjectPath(HitComponent->GetOwner()).ToString());
+			ActorHitAttribute->SetValue(OutPoint.MetadataEntry, FSoftObjectPath(HitComponent->GetOwner()));
 		}
 
 		if (PhysicalMaterialAttribute && Hit.PhysMaterial.Get())
 		{
 			OutMetadata->InitializeOnSet(OutPoint.MetadataEntry);
-			PhysicalMaterialAttribute->SetValue(OutPoint.MetadataEntry, FSoftObjectPath(Hit.PhysMaterial.Get()).ToString());
+			PhysicalMaterialAttribute->SetValue(OutPoint.MetadataEntry, FSoftObjectPath(Hit.PhysMaterial.Get()));
 		}
 
 		return true;
diff --git a/Engine/Plugins/Experimental/PCG/Source/PCG/Private/Elements/PCGCreateAttribute.cpp b/Engine/Plugins/Experimental/PCG/Source/PCG/Private/Elements/PCGCreateAttribute.cpp
--- a/Engine/Plugins/Experimental/PCG/Source/PCG/Private/Elements/PCGCreateAttribute.cpp
+++ b/Engine/Plugins/Experimental/PCG/Source/PCG/Private/Elements/PCGCreateAttribute.cpp
@@ -55,6 +55,8 @@ void UPCGCreateAttributeBaseSettings::PostLoad()
 		InputSource.SetAttributeName(SourceParamAttributeName_DEPRECATED);
 		SourceParamAttributeName_DEPRECATED = NAME_None;
 	}
+
+	AttributeTypes.OnPostLoad();
 #endif // WITH_EDITOR
 }
 
diff --git a/Engine/Plugins/Experimental/PCG/Source/PCG/Private/Elements/PCGCreateTargetActor.cpp b/Engine/Plugins/Experimental/PCG/Source/PCG/Private/Elements/PCGCreateTargetActor.cpp
--- a/Engine/Plugins/Experimental/PCG/Source/PCG/Private/Elements/PCGCreateTargetActor.cpp
+++ b/Engine/Plugins/Experimental/PCG/Source/PCG/Private/Elements/PCGCreateTargetActor.cpp
@@ -293,7 +293,7 @@ bool FPCGCreateTargetActorElement::ExecuteInternal(FPCGContext* Context) const
 	// Create param data output with reference to actor
 	UPCGParamData* ParamData = NewObject<UPCGParamData>();
 	check(ParamData && ParamData->Metadata);
-	ParamData->Metadata->CreateAttribute<FString>(NAME_None, FSoftObjectPath(GeneratedActor).ToString(), /*bAllowsInterpolation=*/false, /*bOverrideParent=*/false);
+	ParamData->Metadata->CreateAttribute<FSoftObjectPath>(NAME_None, FSoftObjectPath(GeneratedActor), /*bAllowsInterpolation=*/false, /*bOverrideParent=*/false);
 
 	// Add param data to output and we're done
 	Context->OutputData.TaggedData.Emplace_GetRef().Data = ParamData;
diff --git a/Engine/Plugins/Experimental/PCG/Source/PCG/Private/Elements/PCGPointFilter.cpp b/Engine/Plugins/Experimental/PCG/Source/PCG/Private/Elements/PCGPointFilter.cpp
--- a/Engine/Plugins/Experimental/PCG/Source/PCG/Private/Elements/PCGPointFilter.cpp
+++ b/Engine/Plugins/Experimental/PCG/Source/PCG/Private/Elements/PCGPointFilter.cpp
@@ -251,6 +251,13 @@ namespace PCGPointFilterHelpers
 	}
 }
 
+#if WITH_EDITOR
+void FPCGPointFilterThresholdSettings::OnPostLoad()
+{
+	AttributeTypes.OnPostLoad();
+}
+#endif
+
 TArray<FPCGPinProperties> UPCGPointFilterSettings::InputPinProperties() const
 {
 	TArray<FPCGPinProperties> PinProperties;
@@ -365,9 +372,21 @@ PRAGMA_DISABLE_DEPRECATION_WARNINGS
 	ThresholdConstantType_DEPRECATED = EPCGPointFilterConstantType::Float;
 	FloatConstant_DEPRECATED = 0.0f;
 PRAGMA_ENABLE_DEPRECATION_WARNINGS
+	
+	AttributeTypes.OnPostLoad();
 #endif // WITH_EDITOR
 }
 
+#if WITH_EDITOR
+void UPCGPointFilterRangeSettings::PostLoad()
+{
+	Super::PostLoad();
+
+	MinThreshold.OnPostLoad();
+	MaxThreshold.OnPostLoad();
+}
+#endif
+
 FPCGElementPtr UPCGPointFilterRangeSettings::CreateElement() const
 {
 	return MakeShared<FPCGPointFilterRangeElement>();
diff --git a/Engine/Plugins/Experimental/PCG/Source/PCG/Private/Elements/PCGPointFromMeshElement.cpp b/Engine/Plugins/Experimental/PCG/Source/PCG/Private/Elements/PCGPointFromMeshElement.cpp
--- a/Engine/Plugins/Experimental/PCG/Source/PCG/Private/Elements/PCGPointFromMeshElement.cpp
+++ b/Engine/Plugins/Experimental/PCG/Source/PCG/Private/Elements/PCGPointFromMeshElement.cpp
@@ -60,7 +60,7 @@ bool FPCGPointFromMeshElement::ExecuteInternal(FPCGContext* Context) const
 
 	// Write StaticMesh path to MeshPathAttribute
 	check(OutPointData->Metadata);
-	OutPointData->Metadata->CreateStringAttribute(Settings->MeshPathAttributeName, Settings->StaticMesh.ToString(), /*bAllowsInterpolation=*/false);
+	OutPointData->Metadata->CreateSoftObjectPathAttribute(Settings->MeshPathAttributeName, Settings->StaticMesh.ToSoftObjectPath(), /*bAllowsInterpolation=*/false);
 
 	return true;
 }
diff --git a/Engine/Plugins/Experimental/PCG/Source/PCG/Private/Elements/PCGPointMatchAndSet.cpp b/Engine/Plugins/Experimental/PCG/Source/PCG/Private/Elements/PCGPointMatchAndSet.cpp
--- a/Engine/Plugins/Experimental/PCG/Source/PCG/Private/Elements/PCGPointMatchAndSet.cpp
+++ b/Engine/Plugins/Experimental/PCG/Source/PCG/Private/Elements/PCGPointMatchAndSet.cpp
@@ -87,6 +87,20 @@ void UPCGPointMatchAndSetSettings::PostLoad()
 		MatchAndSetInstance->SetFlags(Flags);
 		bUseSeed = MatchAndSetInstance->UsesRandomProcess();
 	}
+
+#if WITH_EDITOR
+	if (SetTargetType == EPCGMetadataTypes::String)
+	{
+		if (SetTargetStringMode_DEPRECATED == EPCGMetadataTypesConstantStructStringMode::SoftObjectPath)
+		{
+			SetTargetType = EPCGMetadataTypes::SoftObjectPath;
+		}
+		else if (SetTargetStringMode_DEPRECATED == EPCGMetadataTypesConstantStructStringMode::SoftClassPath)
+		{
+			SetTargetType = EPCGMetadataTypes::SoftClassPath;
+		}
+	}
+#endif
 }
 
 #if WITH_EDITOR
@@ -115,17 +129,16 @@ void UPCGPointMatchAndSetSettings::PostEditChangeProperty(FPropertyChangedEvent&
 
 					if (MatchAndSetInstance)
 					{
-						MatchAndSetInstance->SetType(SetTargetType, SetTargetStringMode);
+						MatchAndSetInstance->SetType(SetTargetType);
 					}
 				}
 			}
 		}
-		else if (PropertyName == GET_MEMBER_NAME_CHECKED(UPCGPointMatchAndSetSettings, SetTargetType) ||
-			PropertyName == GET_MEMBER_NAME_CHECKED(UPCGPointMatchAndSetSettings, SetTargetStringMode))
+		else if (PropertyName == GET_MEMBER_NAME_CHECKED(UPCGPointMatchAndSetSettings, SetTargetType))
 		{
 			if (MatchAndSetInstance)
 			{
-				MatchAndSetInstance->SetType(SetTargetType, SetTargetStringMode);
+				MatchAndSetInstance->SetType(SetTargetType);
 			}
 		}
 	}
@@ -159,7 +172,7 @@ void UPCGPointMatchAndSetSettings::RefreshMatchAndSet()
 		const EObjectFlags Flags = GetMaskedFlags(RF_PropagateToSubObjects);
 		MatchAndSetInstance = NewObject<UPCGMatchAndSetBase>(this, MatchAndSetType, NAME_None, Flags);
 		check(MatchAndSetInstance);
-		MatchAndSetInstance->SetType(SetTargetType, SetTargetStringMode);
+		MatchAndSetInstance->SetType(SetTargetType);
 	}
 	else
 	{
diff --git a/Engine/Plugins/Experimental/PCG/Source/PCG/Private/Elements/PCGSpawnActor.cpp b/Engine/Plugins/Experimental/PCG/Source/PCG/Private/Elements/PCGSpawnActor.cpp
--- a/Engine/Plugins/Experimental/PCG/Source/PCG/Private/Elements/PCGSpawnActor.cpp
+++ b/Engine/Plugins/Experimental/PCG/Source/PCG/Private/Elements/PCGSpawnActor.cpp
@@ -171,14 +171,27 @@ public:
 
 	FPCGPointDataPartitionBase::Element* SelectPoint(const FPCGPoint& Point, int32 PointIndex)
 	{
-		FString ActorPathString;
-		if (!SpawnAttributeAccessor->Get<FString>(ActorPathString, PointIndex, *SpawnAttributeKeys))
+		FSoftClassPath ActorPath;
+		TSoftClassPtr<AActor> ActorClassSoftPtr;
+
+		if (SpawnAttributeAccessor->Get<FSoftClassPath>(ActorPath, PointIndex, *SpawnAttributeKeys))
 		{
-			return nullptr;
+			ActorClassSoftPtr = TSoftClassPtr<AActor>(ActorPath);
+		}
+		else
+		{
+			FString ActorPathString;
+			if (SpawnAttributeAccessor->Get<FString>(ActorPathString, PointIndex, *SpawnAttributeKeys))
+			{
+				ActorPath = FSoftClassPath(ActorPathString);
+				ActorClassSoftPtr = TSoftClassPtr<AActor>(ActorPath);
+			}
 		}
 
-		FSoftObjectPath ActorPath(ActorPathString);
-		TSoftClassPtr<AActor> ActorClassSoftPtr(ActorPath);
+		if (!ActorPath.IsValid())
+		{
+			return nullptr;
+		}
 
 		UClass* ActorClass = ActorClassSoftPtr.LoadSynchronous();
 
diff --git a/Engine/Plugins/Experimental/PCG/Source/PCG/Private/Helpers/PCGMetadataHelpers.cpp b/Engine/Plugins/Experimental/PCG/Source/PCG/Private/Helpers/PCGMetadataHelpers.cpp
--- a/Engine/Plugins/Experimental/PCG/Source/PCG/Private/Helpers/PCGMetadataHelpers.cpp
+++ b/Engine/Plugins/Experimental/PCG/Source/PCG/Private/Helpers/PCGMetadataHelpers.cpp
@@ -55,4 +55,73 @@ namespace PCGMetadataHelpers
 			return nullptr;
 		}
 	}
-}
\ No newline at end of file
+
+	bool CreateObjectPathGetter(const FPCGMetadataAttributeBase* InAttributeBase, TFunction<void(int64, FSoftObjectPath&)>& OutGetter)
+	{
+		if (!InAttributeBase)
+		{
+			return false;
+		}
+
+		if (InAttributeBase->GetTypeId() == PCG::Private::MetadataTypes<FString>::Id)
+		{
+			OutGetter = [InAttributeBase](int64 InMetadataKey, FSoftObjectPath& OutSoftObjectPath)
+			{
+				FString Path = static_cast<const FPCGMetadataAttribute<FString>*>(InAttributeBase)->GetValueFromItemKey(InMetadataKey);
+				OutSoftObjectPath = FSoftObjectPath(Path);
+			};
+
+			return true;
+		}
+		else if (InAttributeBase->GetTypeId() == PCG::Private::MetadataTypes<FSoftObjectPath>::Id)
+		{
+			OutGetter = [InAttributeBase](int64 InMetadataKey, FSoftObjectPath& OutSoftObjectPath)
+			{
+				OutSoftObjectPath = static_cast<const FPCGMetadataAttribute<FSoftObjectPath>*>(InAttributeBase)->GetValueFromItemKey(InMetadataKey);
+			};
+
+			return true;
+		}
+
+		return false;
+	}
+
+	bool CreateObjectOrClassPathGetter(const FPCGMetadataAttributeBase* InAttributeBase, TFunction<void(int64, FSoftObjectPath&)>& OutGetter)
+	{
+		if (!InAttributeBase)
+		{
+			return false;
+		}
+
+		if (InAttributeBase->GetTypeId() == PCG::Private::MetadataTypes<FString>::Id)
+		{
+			OutGetter = [InAttributeBase](int64 InMetadataKey, FSoftObjectPath& OutSoftObjectPath)
+			{
+				FString Path = static_cast<const FPCGMetadataAttribute<FString>*>(InAttributeBase)->GetValueFromItemKey(InMetadataKey);
+				OutSoftObjectPath = FSoftObjectPath(Path);
+			};
+
+			return true;
+		}
+		else if (InAttributeBase->GetTypeId() == PCG::Private::MetadataTypes<FSoftObjectPath>::Id)
+		{
+			OutGetter = [InAttributeBase](int64 InMetadataKey, FSoftObjectPath& OutSoftObjectPath)
+			{
+				OutSoftObjectPath = static_cast<const FPCGMetadataAttribute<FSoftObjectPath>*>(InAttributeBase)->GetValueFromItemKey(InMetadataKey);
+			};
+
+			return true;
+		}
+		else if (InAttributeBase->GetTypeId() == PCG::Private::MetadataTypes<FSoftClassPath>::Id)
+		{
+			OutGetter = [InAttributeBase](int64 InMetadataKey, FSoftObjectPath& OutSoftObjectPath)
+			{
+				OutSoftObjectPath = static_cast<const FPCGMetadataAttribute<FSoftClassPath>*>(InAttributeBase)->GetValueFromItemKey(InMetadataKey);
+			};
+
+			return true;
+		}
+
+		return false;
+	}
+}
diff --git a/Engine/Plugins/Experimental/PCG/Source/PCG/Private/MatchAndSet/PCGMatchAndSetBase.cpp b/Engine/Plugins/Experimental/PCG/Source/PCG/Private/MatchAndSet/PCGMatchAndSetBase.cpp
--- a/Engine/Plugins/Experimental/PCG/Source/PCG/Private/MatchAndSet/PCGMatchAndSetBase.cpp
+++ b/Engine/Plugins/Experimental/PCG/Source/PCG/Private/MatchAndSet/PCGMatchAndSetBase.cpp
@@ -11,10 +11,28 @@
 
 #define LOCTEXT_NAMESPACE "PCGMatchAndSetBase"
 
-void UPCGMatchAndSetBase::SetType(EPCGMetadataTypes InType, EPCGMetadataTypesConstantStructStringMode InStringMode)
+#if WITH_EDITOR
+void UPCGMatchAndSetBase::PostLoad()
+{
+	Super::PostLoad();
+
+	if (Type == EPCGMetadataTypes::String)
+	{
+		if (StringMode_DEPRECATED == EPCGMetadataTypesConstantStructStringMode::SoftObjectPath)
+		{
+			Type = EPCGMetadataTypes::SoftObjectPath;
+		}
+		else if (StringMode_DEPRECATED == EPCGMetadataTypesConstantStructStringMode::SoftClassPath)
+		{
+			Type = EPCGMetadataTypes::SoftClassPath;
+		}
+	}
+}
+#endif
+
+void UPCGMatchAndSetBase::SetType(EPCGMetadataTypes InType)
 {
 	Type = InType;
-	StringMode = InStringMode;
 }
 
 void UPCGMatchAndSetBase::MatchAndSet_Implementation(
diff --git a/Engine/Plugins/Experimental/PCG/Source/PCG/Private/MatchAndSet/PCGMatchAndSetByAttribute.cpp b/Engine/Plugins/Experimental/PCG/Source/PCG/Private/MatchAndSet/PCGMatchAndSetByAttribute.cpp
--- a/Engine/Plugins/Experimental/PCG/Source/PCG/Private/MatchAndSet/PCGMatchAndSetByAttribute.cpp
+++ b/Engine/Plugins/Experimental/PCG/Source/PCG/Private/MatchAndSet/PCGMatchAndSetByAttribute.cpp
@@ -19,23 +19,29 @@ FPCGMatchAndSetByAttributeEntry::FPCGMatchAndSetByAttributeEntry()
 	Value.bAllowsTypeChange = false;
 }
 
-void UPCGMatchAndSetByAttribute::SetType(EPCGMetadataTypes InType, EPCGMetadataTypesConstantStructStringMode InStringMode)
+#if WITH_EDITOR
+void FPCGMatchAndSetByAttributeEntry::OnPostLoad()
+{
+	ValueToMatch.OnPostLoad();
+	Value.OnPostLoad();
+}
+#endif
+
+void UPCGMatchAndSetByAttribute::SetType(EPCGMetadataTypes InType)
 {
 	for (FPCGMatchAndSetByAttributeEntry& Entry : Entries)
 	{
 		Entry.Value.Type = InType;
-		Entry.Value.StringMode = InStringMode;
 	}
 
-	Super::SetType(InType, InStringMode);
+	Super::SetType(InType);
 }
 
-void UPCGMatchAndSetByAttribute::SetSourceType(EPCGMetadataTypes InType, EPCGMetadataTypesConstantStructStringMode InStringMode)
+void UPCGMatchAndSetByAttribute::SetSourceType(EPCGMetadataTypes InType)
 {
 	for (FPCGMatchAndSetByAttributeEntry& Entry : Entries)
 	{
 		Entry.ValueToMatch.Type = InType;
-		Entry.ValueToMatch.StringMode = InStringMode;
 	}
 }
 
@@ -46,21 +52,42 @@ void UPCGMatchAndSetByAttribute::PostEditChangeProperty(FPropertyChangedEvent& P
 	{
 		const FName& PropertyName = PropertyChangedEvent.Property->GetFName();
 
-		if (PropertyName == GET_MEMBER_NAME_CHECKED(UPCGMatchAndSetByAttribute, MatchSourceType) ||
-			PropertyName == GET_MEMBER_NAME_CHECKED(UPCGMatchAndSetByAttribute, MatchSourceStringMode))
+		if (PropertyName == GET_MEMBER_NAME_CHECKED(UPCGMatchAndSetByAttribute, MatchSourceType))
 		{
-			SetSourceType(MatchSourceType, MatchSourceStringMode);
+			SetSourceType(MatchSourceType);
 		}
 		else if (PropertyName == GET_MEMBER_NAME_CHECKED(UPCGMatchAndSetByAttribute, Entries))
 		{
 			// Some changes in the entries (add, insert) might require us to re-set the type
-			SetType(Type, StringMode);
-			SetSourceType(MatchSourceType, MatchSourceStringMode);
+			SetType(Type);
+			SetSourceType(MatchSourceType);
 		}
 	}
 
 	Super::PostEditChangeProperty(PropertyChangedEvent);
 }
+
+void UPCGMatchAndSetByAttribute::PostLoad()
+{
+	Super::PostLoad();
+
+	if (MatchSourceType == EPCGMetadataTypes::String)
+	{
+		if (MatchSourceStringMode_DEPRECATED == EPCGMetadataTypesConstantStructStringMode::SoftObjectPath)
+		{
+			MatchSourceType = EPCGMetadataTypes::SoftObjectPath;
+		}
+		else if (MatchSourceStringMode_DEPRECATED == EPCGMetadataTypesConstantStructStringMode::SoftClassPath)
+		{
+			MatchSourceType = EPCGMetadataTypes::SoftClassPath;
+		}
+	}
+
+	for (FPCGMatchAndSetByAttributeEntry& Entry : Entries)
+	{
+		Entry.OnPostLoad();
+	}
+}
 #endif // WITH_EDITOR
 
 void UPCGMatchAndSetByAttribute::MatchAndSet_Implementation(
diff --git a/Engine/Plugins/Experimental/PCG/Source/PCG/Private/MatchAndSet/PCGMatchAndSetWeighted.cpp b/Engine/Plugins/Experimental/PCG/Source/PCG/Private/MatchAndSet/PCGMatchAndSetWeighted.cpp
--- a/Engine/Plugins/Experimental/PCG/Source/PCG/Private/MatchAndSet/PCGMatchAndSetWeighted.cpp
+++ b/Engine/Plugins/Experimental/PCG/Source/PCG/Private/MatchAndSet/PCGMatchAndSetWeighted.cpp
@@ -20,15 +20,21 @@ FPCGMatchAndSetWeightedEntry::FPCGMatchAndSetWeightedEntry()
 	Value.bAllowsTypeChange = false;
 }
 
-void UPCGMatchAndSetWeighted::SetType(EPCGMetadataTypes InType, EPCGMetadataTypesConstantStructStringMode InStringMode)
+#if WITH_EDITOR
+void FPCGMatchAndSetWeightedEntry::OnPostLoad()
+{
+	Value.OnPostLoad();
+}
+#endif
+
+void UPCGMatchAndSetWeighted::SetType(EPCGMetadataTypes InType)
 {
 	for (FPCGMatchAndSetWeightedEntry& Entry : Entries)
 	{
 		Entry.Value.Type = InType;
-		Entry.Value.StringMode = InStringMode;
 	}
 
-	Super::SetType(InType, InStringMode);
+	Super::SetType(InType);
 }
 
 #if WITH_EDITOR
@@ -41,12 +47,22 @@ void UPCGMatchAndSetWeighted::PostEditChangeProperty(FPropertyChangedEvent& Prop
 		if (PropertyName == GET_MEMBER_NAME_CHECKED(UPCGMatchAndSetWeighted, Entries))
 		{
 			// Some changes in the array might (such as insert or new) might require us to re-set the type
-			SetType(Type, StringMode);
+			SetType(Type);
 		}
 	}
 
 	Super::PostEditChangeProperty(PropertyChangedEvent);
 }
+
+void UPCGMatchAndSetWeighted::PostLoad()
+{
+	Super::PostLoad();
+
+	for (FPCGMatchAndSetWeightedEntry& Entry : Entries)
+	{
+		Entry.OnPostLoad();
+	}
+}
 #endif // WITH_EDITOR
 
 void UPCGMatchAndSetWeighted::MatchAndSet_Implementation(
diff --git a/Engine/Plugins/Experimental/PCG/Source/PCG/Private/MatchAndSet/PCGMatchAndSetWeightedByCategory.cpp b/Engine/Plugins/Experimental/PCG/Source/PCG/Private/MatchAndSet/PCGMatchAndSetWeightedByCategory.cpp
--- a/Engine/Plugins/Experimental/PCG/Source/PCG/Private/MatchAndSet/PCGMatchAndSetWeightedByCategory.cpp
+++ b/Engine/Plugins/Experimental/PCG/Source/PCG/Private/MatchAndSet/PCGMatchAndSetWeightedByCategory.cpp
@@ -21,12 +21,23 @@ FPCGMatchAndSetWeightedByCategoryEntryList::FPCGMatchAndSetWeightedByCategoryEnt
 	CategoryValue.bAllowsTypeChange = false;
 }
 
-void FPCGMatchAndSetWeightedByCategoryEntryList::SetType(EPCGMetadataTypes InType, EPCGMetadataTypesConstantStructStringMode InStringMode)
+#if WITH_EDITOR
+void FPCGMatchAndSetWeightedByCategoryEntryList::OnPostLoad()
+{
+	CategoryValue.OnPostLoad();
+
+	for (FPCGMatchAndSetWeightedEntry& Entry : WeightedEntries)
+	{
+		Entry.OnPostLoad();
+	}
+} 
+#endif
+
+void FPCGMatchAndSetWeightedByCategoryEntryList::SetType(EPCGMetadataTypes InType)
 {
 	for (FPCGMatchAndSetWeightedEntry& Entry : WeightedEntries)
 	{
 		Entry.Value.Type = InType;
-		Entry.Value.StringMode = InStringMode;
 	}
 }
 
@@ -42,22 +53,21 @@ int FPCGMatchAndSetWeightedByCategoryEntryList::GetTotalWeight() const
 	return TotalWeight;
 }
 
-void UPCGMatchAndSetWeightedByCategory::SetType(EPCGMetadataTypes InType, EPCGMetadataTypesConstantStructStringMode InStringMode)
+void UPCGMatchAndSetWeightedByCategory::SetType(EPCGMetadataTypes InType)
 {
 	for (FPCGMatchAndSetWeightedByCategoryEntryList& Category : Categories)
 	{
-		Category.SetType(InType, InStringMode);
+		Category.SetType(InType);
 	}
 
-	Super::SetType(InType, InStringMode);
+	Super::SetType(InType);
 }
 
-void UPCGMatchAndSetWeightedByCategory::SetCategoryType(EPCGMetadataTypes InType, EPCGMetadataTypesConstantStructStringMode InStringMode)
+void UPCGMatchAndSetWeightedByCategory::SetCategoryType(EPCGMetadataTypes InType)
 {
 	for (FPCGMatchAndSetWeightedByCategoryEntryList& Category : Categories)
 	{
 		Category.CategoryValue.Type = InType;
-		Category.CategoryValue.StringMode = InStringMode;
 	}
 }
 
@@ -68,25 +78,46 @@ void UPCGMatchAndSetWeightedByCategory::PostEditChangeProperty(FPropertyChangedE
 	{
 		const FName& PropertyName = PropertyChangedEvent.Property->GetFName();
 
-		if (PropertyName == GET_MEMBER_NAME_CHECKED(UPCGMatchAndSetWeightedByCategory, CategoryType) ||
-			PropertyName == GET_MEMBER_NAME_CHECKED(UPCGMatchAndSetWeightedByCategory, CategoryStringMode))
+		if (PropertyName == GET_MEMBER_NAME_CHECKED(UPCGMatchAndSetWeightedByCategory, CategoryType))
 		{
-			SetCategoryType(CategoryType, CategoryStringMode);
+			SetCategoryType(CategoryType);
 		}
 		else if (PropertyName == GET_MEMBER_NAME_CHECKED(FPCGMatchAndSetWeightedByCategoryEntryList, WeightedEntries))
 		{
 			// Changes in any weighted entries (which store the end values) - due to add or insert - need to update type
-			SetType(Type, StringMode);
+			SetType(Type);
 		}
 		else if (PropertyName == GET_MEMBER_NAME_CHECKED(UPCGMatchAndSetWeightedByCategory, Categories))
 		{
-			SetCategoryType(CategoryType, CategoryStringMode);
-			SetType(Type, StringMode);
+			SetCategoryType(CategoryType);
+			SetType(Type);
 		}
 	}
 
 	Super::PostEditChangeProperty(PropertyChangedEvent);
 }
+
+void UPCGMatchAndSetWeightedByCategory::PostLoad()
+{
+	Super::PostLoad();
+
+	if (CategoryType == EPCGMetadataTypes::String)
+	{
+		if (CategoryStringMode_DEPRECATED == EPCGMetadataTypesConstantStructStringMode::SoftObjectPath)
+		{
+			CategoryType = EPCGMetadataTypes::SoftObjectPath;
+		}
+		else if (CategoryStringMode_DEPRECATED == EPCGMetadataTypesConstantStructStringMode::SoftClassPath)
+		{
+			CategoryType = EPCGMetadataTypes::SoftClassPath;
+		}
+	}
+
+	for (FPCGMatchAndSetWeightedByCategoryEntryList& Category : Categories)
+	{
+		Category.OnPostLoad();
+	}
+}
 #endif // WITH_EDITOR
 
 void UPCGMatchAndSetWeightedByCategory::MatchAndSet_Implementation(
diff --git a/Engine/Plugins/Experimental/PCG/Source/PCG/Private/MeshSelectors/PCGMeshSelectorBase.cpp b/Engine/Plugins/Experimental/PCG/Source/PCG/Private/MeshSelectors/PCGMeshSelectorBase.cpp
--- a/Engine/Plugins/Experimental/PCG/Source/PCG/Private/MeshSelectors/PCGMeshSelectorBase.cpp
+++ b/Engine/Plugins/Experimental/PCG/Source/PCG/Private/MeshSelectors/PCGMeshSelectorBase.cpp
@@ -68,14 +68,14 @@ void FPCGMeshMaterialOverrideHelper::Initialize(FPCGContext& InContext)
 
 			const FPCGMetadataAttributeBase* MaterialAttributeBase = Metadata->GetConstAttribute(MaterialOverrideAttributeName);
 			check(MaterialAttributeBase);
-
-			if (MaterialAttributeBase->GetTypeId() != PCG::Private::MetadataTypes<FString>::Id)
+			
+			if (!PCG::Private::IsOfTypes<FSoftObjectPath, FString>(MaterialAttributeBase->GetTypeId()))
 			{
-				PCGE_LOG_C(Error, GraphAndLog, &InContext, LOCTEXT("AttributeInvalidType", "Material override attribute is not of valid type"));
+				PCGE_LOG_C(Error, GraphAndLog, &InContext, LOCTEXT("AttributeInvalidType", "Material override attribute is not of valid type (FSoftObjectPath or FString)."));
 				return;
 			}
 
-			MaterialAttributes.Add(static_cast<const FPCGMetadataAttribute<FString>*>(MaterialAttributeBase));
+			MaterialAttributes.Add(MaterialAttributeBase);
 		}
 
 		ValueKeyToOverrideMaterials.SetNum(MaterialOverrideAttributeNames.Num());
@@ -107,14 +107,28 @@ const TArray<TSoftObjectPtr<UMaterialInterface>>& FPCGMeshMaterialOverrideHelper
 
 		for (int32 MaterialIndex = 0; MaterialIndex < MaterialAttributes.Num(); ++MaterialIndex)
 		{
-			const FPCGMetadataAttribute<FString>* MaterialAttribute = MaterialAttributes[MaterialIndex];
+			const FPCGMetadataAttributeBase* MaterialAttribute = MaterialAttributes[MaterialIndex];
 			PCGMetadataValueKey MaterialValueKey = MaterialAttribute->GetValueKey(EntryKey);
 			TSoftObjectPtr<UMaterialInterface>* NewMaterial = ValueKeyToOverrideMaterials[MaterialIndex].Find(MaterialValueKey);
 			TSoftObjectPtr<UMaterialInterface> Material = nullptr;
 
 			if (!NewMaterial)
 			{
-				FSoftObjectPath MaterialPath(MaterialAttribute->GetValue(MaterialValueKey));
+				FSoftObjectPath MaterialPath;
+				if (MaterialAttribute->GetTypeId() == PCG::Private::MetadataTypes<FSoftObjectPath>::Id)
+				{
+					MaterialPath = static_cast<const FPCGMetadataAttribute<FSoftObjectPath>*>(MaterialAttribute)->GetValue(MaterialValueKey);
+				}
+				else if (MaterialAttribute->GetTypeId() == PCG::Private::MetadataTypes<FString>::Id)
+				{
+					MaterialPath = FSoftObjectPath(static_cast<const FPCGMetadataAttribute<FString>*>(MaterialAttribute)->GetValue(MaterialValueKey));
+				}
+				else
+				{
+					ensure(false);
+					continue;
+				}
+
 				Material = TSoftObjectPtr<UMaterialInterface>(MaterialPath);
 				ValueKeyToOverrideMaterials[MaterialIndex].Add(MaterialValueKey, Material);
 			}
diff --git a/Engine/Plugins/Experimental/PCG/Source/PCG/Private/MeshSelectors/PCGMeshSelectorByAttribute.cpp b/Engine/Plugins/Experimental/PCG/Source/PCG/Private/MeshSelectors/PCGMeshSelectorByAttribute.cpp
--- a/Engine/Plugins/Experimental/PCG/Source/PCG/Private/MeshSelectors/PCGMeshSelectorByAttribute.cpp
+++ b/Engine/Plugins/Experimental/PCG/Source/PCG/Private/MeshSelectors/PCGMeshSelectorByAttribute.cpp
@@ -124,13 +124,34 @@ bool UPCGMeshSelectorByAttribute::SelectInstances(
 	const FPCGMetadataAttributeBase* AttributeBase = InPointData->Metadata->GetConstAttribute(AttributeName);
 	check(AttributeBase);
 
-	if (AttributeBase->GetTypeId() != PCG::Private::MetadataTypes<FString>::Id)
+	// Set up a getter lambda to retrieve the mesh asset in the selection loop below. Returns true if path was valid.
+	TFunction<void(PCGMetadataValueKey, FSoftObjectPath&)> MeshPathGetter;
+	if (PCG::Private::IsOfTypes<FSoftObjectPath>(AttributeBase->GetTypeId()))
 	{
-		PCGE_LOG_C(Error, GraphAndLog, &Context, FText::Format(LOCTEXT("AttributeInvalidType", "Attribute '{0}' is not of valid type FString"), FText::FromName(AttributeName)));
-		return true;
+		MeshPathGetter = [AttributeBase](PCGMetadataValueKey InValueKey, FSoftObjectPath& OutMeshPath)
+		{
+			const FPCGMetadataAttribute<FSoftObjectPath>* Attribute = static_cast<const FPCGMetadataAttribute<FSoftObjectPath>*>(AttributeBase);
+			OutMeshPath = Attribute->GetValue(InValueKey);
+		};
 	}
+	else if (PCG::Private::IsOfTypes<FString>(AttributeBase->GetTypeId()))
+	{
+		MeshPathGetter = [AttributeBase, &Context](PCGMetadataValueKey InValueKey, FSoftObjectPath& OutMeshPath)
+		{
+			const FPCGMetadataAttribute<FString>* Attribute = static_cast<const FPCGMetadataAttribute<FString>*>(AttributeBase);
 
-	const FPCGMetadataAttribute<FString>* Attribute = static_cast<const FPCGMetadataAttribute<FString>*>(AttributeBase);
+			const FString Path = Attribute->GetValue(InValueKey);
+			if (!Path.IsEmpty() && Path != TEXT("None"))
+			{
+				OutMeshPath = FSoftObjectPath(Path);
+			}
+		};
+	}
+	else
+	{
+		PCGE_LOG_C(Error, GraphAndLog, &Context, FText::Format(LOCTEXT("AttributeInvalidType", "Attribute '{0}' is not of valid type (must be FString or FSoftObjectPath)"), FText::FromName(AttributeName)));
+		return true;
+	}
 
 	FPCGMeshMaterialOverrideHelper& MaterialOverrideHelper = Context.MaterialOverrideHelper;
 	if (!MaterialOverrideHelper.IsInitialized())
@@ -168,28 +189,28 @@ bool UPCGMeshSelectorByAttribute::SelectInstances(
 	{
 		const FPCGPoint& Point = Points[CurrentPointIndex++];
 		
-		PCGMetadataValueKey ValueKey = Attribute->GetValueKey(Point.MetadataEntry);
+		const PCGMetadataValueKey ValueKey = AttributeBase->GetValueKey(Point.MetadataEntry);
 		TSoftObjectPtr<UStaticMesh>* NewMesh = ValueKeyToMesh.Find(ValueKey);
 		TSoftObjectPtr<UStaticMesh> Mesh = nullptr;
 
-		// if this ValueKey has not been seen before, let's cache it for the future
+		// If this ValueKey has not been seen before, let's cache it for the future
 		if (!NewMesh)
 		{
-			FString MeshSoftObjectPath = Attribute->GetValue(ValueKey);
+			FSoftObjectPath MeshPath;
+			MeshPathGetter(ValueKey, MeshPath);
 
-			if (!MeshSoftObjectPath.IsEmpty() && MeshSoftObjectPath != TEXT("None"))
+			if (!MeshPath.IsNull())
 			{
-				const FSoftObjectPath MeshPath(MeshSoftObjectPath);
 				Mesh = TSoftObjectPtr<UStaticMesh>(MeshPath);
 
 				if (Mesh.IsNull())
 				{
-					PCGE_LOG_C(Error, GraphAndLog, &Context, FText::Format(LOCTEXT("InvalidMeshPath", "Invalid mesh path: '{0}'."), FText::FromString(MeshSoftObjectPath)));
+					PCGE_LOG_C(Error, GraphAndLog, &Context, FText::Format(LOCTEXT("InvalidMeshPath", "Invalid mesh path: '{0}'."), FText::FromString(MeshPath.ToString())));
 				}
 			}
 			else
 			{
-				PCGE_LOG_C(Warning, LogOnly, &Context, FText::Format(LOCTEXT("TrivialInvalidMeshPath", "Trivially invalid mesh path used: '{0}'"), FText::FromString(MeshSoftObjectPath)));
+				PCGE_LOG_C(Warning, LogOnly, &Context, LOCTEXT("TrivialInvalidMeshPath", "Trivially invalid mesh path used."));
 			}
 
 			ValueKeyToMesh.Add(ValueKey, Mesh);
diff --git a/Engine/Plugins/Experimental/PCG/Source/PCG/Private/Metadata/Accessors/PCGAttributeAccessorHelpers.cpp b/Engine/Plugins/Experimental/PCG/Source/PCG/Private/Metadata/Accessors/PCGAttributeAccessorHelpers.cpp
--- a/Engine/Plugins/Experimental/PCG/Source/PCG/Private/Metadata/Accessors/PCGAttributeAccessorHelpers.cpp
+++ b/Engine/Plugins/Experimental/PCG/Source/PCG/Private/Metadata/Accessors/PCGAttributeAccessorHelpers.cpp
@@ -288,9 +288,13 @@ namespace PCGAttributeAccessorHelpers
 		{
 			return Functor(Signature<FPCGEnumPropertyAccessor>{}, EnumProperty);
 		}
+		else if (const FSoftClassProperty* SoftClassProperty = CastField<FSoftClassProperty>(InProperty))
+		{
+			return Functor(Signature<FPCGPropertySoftClassPathAccessor>{}, SoftClassProperty);
+		}
 		else if (const FSoftObjectProperty* SoftObjectProperty = CastField<FSoftObjectProperty>(InProperty))
 		{
-			return Functor(Signature<FPCGPropertySoftPtrAccessor>{}, SoftObjectProperty);
+			return Functor(Signature<FPCGPropertySoftObjectPathAccessor>{}, SoftObjectProperty);
 		}
 		else if (const FObjectProperty* ObjectProperty = CastField<FObjectProperty>(InProperty))
 		{
@@ -324,11 +328,11 @@ namespace PCGAttributeAccessorHelpers
 			}
 			else if (StructProperty->Struct == TBaseStructure<FSoftObjectPath>::Get())
 			{
-				return Functor(Signature<FPCGPropertyPathAccessor<FSoftObjectPath>>{}, InProperty);
+				return Functor(Signature<FPCGPropertyStructAccessor<FSoftObjectPath>>{}, StructProperty);
 			}
 			else if (StructProperty->Struct == TBaseStructure<FSoftClassPath>::Get())
 			{
-				return Functor(Signature<FPCGPropertyPathAccessor<FSoftClassPath>>{}, InProperty);
+				return Functor(Signature<FPCGPropertyStructAccessor<FSoftClassPath>>{}, StructProperty);
 			}
 			else if (StructProperty->Struct && StructProperty->Struct->IsChildOf(FPCGAttributePropertySelector::StaticStruct()))
 			{
diff --git a/Engine/Plugins/Experimental/PCG/Source/PCG/Private/Metadata/PCGMetadata.cpp b/Engine/Plugins/Experimental/PCG/Source/PCG/Private/Metadata/PCGMetadata.cpp
--- a/Engine/Plugins/Experimental/PCG/Source/PCG/Private/Metadata/PCGMetadata.cpp
+++ b/Engine/Plugins/Experimental/PCG/Source/PCG/Private/Metadata/PCGMetadata.cpp
@@ -518,6 +518,16 @@ void UPCGMetadata::CreateBoolAttribute(FName AttributeName, bool DefaultValue, b
 	CreateAttribute<bool>(AttributeName, DefaultValue, bAllowsInterpolation, bOverrideParent);
 }
 
+void UPCGMetadata::CreateSoftObjectPathAttribute(FName AttributeName, const FSoftObjectPath& DefaultValue, bool bAllowsInterpolation, bool bOverrideParent)
+{
+	CreateAttribute<FSoftObjectPath>(AttributeName, DefaultValue, bAllowsInterpolation, bOverrideParent);
+}
+
+void UPCGMetadata::CreateSoftClassPathAttribute(FName AttributeName, const FSoftClassPath& DefaultValue, bool bAllowsInterpolation, bool bOverrideParent)
+{
+	CreateAttribute<FSoftClassPath>(AttributeName, DefaultValue, bAllowsInterpolation, bOverrideParent);
+}
+
 namespace PCGMetadata
 {
 	template<typename DataType>
@@ -576,13 +586,16 @@ namespace PCGMetadata
 					using AttributeType = decltype(AttributeValue);
 					FPCGMetadataAttribute<AttributeType>* Attribute = static_cast<FPCGMetadataAttribute<AttributeType>*>(BaseAttribute);
 
+					// Special cased because FSoftObjectPath currently has a deprecated constructor from FName which generates compile warnings.
+					constexpr bool bAssigningNameToSoftObjectPath = std::is_same_v<AttributeType, FSoftObjectPath> && std::is_same_v<PropertyType, FName>;
+
 					if constexpr (std::is_same_v<AttributeType, PropertyType>)
 					{
 						Metadata->InitializeOnSet(EntryKey);
 						Attribute->SetValue(EntryKey, PropertyValue);
 						return true;
 					}
-					else if constexpr (std::is_constructible_v<AttributeType, PropertyType>)
+					else if constexpr (std::is_constructible_v<AttributeType, PropertyType> && !bAssigningNameToSoftObjectPath)
 					{
 						Metadata->InitializeOnSet(EntryKey);
 						Attribute->SetValue(EntryKey, AttributeType(PropertyValue));
diff --git a/Engine/Plugins/Experimental/PCG/Source/PCG/Private/Metadata/PCGMetadataAccessor.cpp b/Engine/Plugins/Experimental/PCG/Source/PCG/Private/Metadata/PCGMetadataAccessor.cpp
--- a/Engine/Plugins/Experimental/PCG/Source/PCG/Private/Metadata/PCGMetadataAccessor.cpp
+++ b/Engine/Plugins/Experimental/PCG/Source/PCG/Private/Metadata/PCGMetadataAccessor.cpp
@@ -28,7 +28,7 @@ namespace PCGMetadataAccessorHelpers
 	}
 }
 
-/** Key-based implmentations */
+/** Key-based implementations */
 template<typename T>
 T UPCGMetadataAccessorHelpers::GetAttribute(PCGMetadataEntryKey Key, const UPCGMetadata* Metadata, FName AttributeName)
 {
@@ -38,21 +38,37 @@ T UPCGMetadataAccessorHelpers::GetAttribute(PCGMetadataEntryKey Key, const UPCGM
 		return T{};
 	}
 
-	const FPCGMetadataAttribute<T>* Attribute = static_cast<const FPCGMetadataAttribute<T>*>(Metadata->GetConstAttribute(AttributeName));
-	if (Attribute && Attribute->GetTypeId() == PCG::Private::MetadataTypes<T>::Id)
+	const FPCGMetadataAttributeBase* AttributeBase = Metadata->GetConstAttribute(AttributeName);
+	if (!AttributeBase)
 	{
-		return Attribute->GetValueFromItemKey(Key);
+		PCGMetadataAccessorHelpers::OnException(FText::Format(PCGMetadataAccessorHelpers::InvalidAttributeNameFormat, FText::FromName(AttributeName)));
+		return T{};
 	}
-	else if (Attribute)
+
+	if (AttributeBase->GetTypeId() == PCG::Private::MetadataTypes<T>::Id)
 	{
-		PCGMetadataAccessorHelpers::OnException(FText::Format(PCGMetadataAccessorHelpers::InvalidTypeFormat, FText::FromName(AttributeName), FText::FromString(PCG::Private::GetTypeName(Attribute->GetTypeId())), FText::FromString(PCG::Private::GetTypeName(PCG::Private::MetadataTypes<T>::Id))));
-		return T{};
+		return static_cast<const FPCGMetadataAttribute<T>*>(AttributeBase)->GetValueFromItemKey(Key);
 	}
-	else
+	else if constexpr (std::is_same_v<T, FString>)
 	{
-		PCGMetadataAccessorHelpers::OnException(FText::Format(PCGMetadataAccessorHelpers::InvalidAttributeNameFormat, FText::FromName(AttributeName)));
-		return T{};
+		// Legacy path - allow reading soft object/class paths using string accessor.
+		if (AttributeBase->GetTypeId() == PCG::Private::MetadataTypes<FSoftObjectPath>::Id)
+		{
+			return static_cast<const FPCGMetadataAttribute<FSoftObjectPath>*>(AttributeBase)->GetValueFromItemKey(Key).ToString();
+		}
+		else if (AttributeBase->GetTypeId() == PCG::Private::MetadataTypes<FSoftClassPath>::Id)
+		{
+			return static_cast<const FPCGMetadataAttribute<FSoftClassPath>*>(AttributeBase)->GetValueFromItemKey(Key).ToString();
+		}
 	}
+
+	PCGMetadataAccessorHelpers::OnException(FText::Format(
+		PCGMetadataAccessorHelpers::InvalidTypeFormat,
+		FText::FromName(AttributeName),
+		FText::FromString(PCG::Private::GetTypeName(AttributeBase->GetTypeId())),
+		FText::FromString(PCG::Private::GetTypeName(PCG::Private::MetadataTypes<T>::Id))));
+
+	return T{};
 }
 
 template<typename T>
@@ -232,6 +248,26 @@ void UPCGMetadataAccessorHelpers::SetBoolAttributeByMetadataKey(int64& Key, UPCG
 	SetAttribute(Key, Metadata, AttributeName, Value);
 }
 
+FSoftObjectPath UPCGMetadataAccessorHelpers::GetSoftObjectPathAttributeByMetadataKey(int64 Key, const UPCGMetadata* Metadata, FName AttributeName)
+{
+	return GetAttribute<FSoftObjectPath>(Key, Metadata, AttributeName);
+}
+
+void UPCGMetadataAccessorHelpers::SetSoftObjectPathAttributeByMetadataKey(int64& Key, UPCGMetadata* Metadata, FName AttributeName, const FSoftObjectPath& Value)
+{
+	SetAttribute(Key, Metadata, AttributeName, Value);
+}
+
+FSoftClassPath UPCGMetadataAccessorHelpers::GetSoftClassPathAttributeByMetadataKey(int64 Key, const UPCGMetadata* Metadata, FName AttributeName)
+{
+	return GetAttribute<FSoftClassPath>(Key, Metadata, AttributeName);
+}
+
+void UPCGMetadataAccessorHelpers::SetSoftClassPathAttributeByMetadataKey(int64& Key, UPCGMetadata* Metadata, FName AttributeName, const FSoftClassPath& Value)
+{
+	SetAttribute(Key, Metadata, AttributeName, Value);
+}
+
 bool UPCGMetadataAccessorHelpers::SetAttributeFromPropertyByMetadataKey(int64& Key, UPCGMetadata* Metadata, FName AttributeName, const UObject* Object, FName PropertyName)
 {
 	if (!Object)
@@ -427,6 +463,26 @@ void UPCGMetadataAccessorHelpers::SetBoolAttribute(FPCGPoint& Point, UPCGMetadat
 	SetAttribute(Point.MetadataEntry, Metadata, AttributeName, Value);
 }
 
+FSoftObjectPath UPCGMetadataAccessorHelpers::GetSoftObjectPathAttribute(const FPCGPoint& Point, const UPCGMetadata* Metadata, FName AttributeName)
+{
+	return GetAttribute<FSoftObjectPath>(Point.MetadataEntry, Metadata, AttributeName);
+}
+
+void UPCGMetadataAccessorHelpers::SetSoftObjectPathAttribute(FPCGPoint& Point, UPCGMetadata* Metadata, FName AttributeName, const FSoftObjectPath& Value)
+{
+	SetAttribute(Point.MetadataEntry, Metadata, AttributeName, Value);
+}
+
+FSoftClassPath UPCGMetadataAccessorHelpers::GetSoftClassPathAttribute(const FPCGPoint& Point, const UPCGMetadata* Metadata, FName AttributeName)
+{
+	return GetAttribute<FSoftClassPath>(Point.MetadataEntry, Metadata, AttributeName);
+}
+
+void UPCGMetadataAccessorHelpers::SetSoftClassPathAttribute(FPCGPoint& Point, UPCGMetadata* Metadata, FName AttributeName, const FSoftClassPath& Value)
+{
+	SetAttribute(Point.MetadataEntry, Metadata, AttributeName, Value);
+}
+
 bool UPCGMetadataAccessorHelpers::HasAttributeSet(const FPCGPoint& Point, const UPCGMetadata* Metadata, FName AttributeName)
 {
 	return HasAttributeSetByMetadataKey(Point.MetadataEntry, Metadata, AttributeName);
diff --git a/Engine/Plugins/Experimental/PCG/Source/PCG/Private/Metadata/PCGMetadataTypesConstantStruct.cpp b/Engine/Plugins/Experimental/PCG/Source/PCG/Private/Metadata/PCGMetadataTypesConstantStruct.cpp
--- a/Engine/Plugins/Experimental/PCG/Source/PCG/Private/Metadata/PCGMetadataTypesConstantStruct.cpp
+++ b/Engine/Plugins/Experimental/PCG/Source/PCG/Private/Metadata/PCGMetadataTypesConstantStruct.cpp
@@ -15,19 +15,7 @@ FString FPCGMetadataTypesConstantStruct::ToString() const
 	case EPCGMetadataTypes::Double:
 		return FString::Printf(TEXT("%.2f"), DoubleValue);
 	case EPCGMetadataTypes::String:
-	{
-		switch (StringMode)
-		{
-		case EPCGMetadataTypesConstantStructStringMode::String:
-			return FString::Printf(TEXT("\"%s\""), *StringValue);
-		case EPCGMetadataTypesConstantStructStringMode::SoftObjectPath:
-			return FString::Printf(TEXT("\"%s\""), *(SoftObjectPathValue.GetAssetName()));
-		case EPCGMetadataTypesConstantStructStringMode::SoftClassPath:
-			return FString::Printf(TEXT("\"%s\""), *(SoftClassPathValue.GetAssetName()));
-		default:
-			break;
-		}
-	}
+		return FString::Printf(TEXT("\"%s\""), *StringValue);
 	case EPCGMetadataTypes::Name:
 		return FString::Printf(TEXT("N(\"%s\")"), *NameValue.ToString());
 	case EPCGMetadataTypes::Vector2:
@@ -44,9 +32,31 @@ FString FPCGMetadataTypesConstantStruct::ToString() const
 		return FString::Printf(TEXT("Transform"));
 	case EPCGMetadataTypes::Boolean:
 		return FString::Printf(TEXT("%s"), (BoolValue ? TEXT("True") : TEXT("False")));
+	case EPCGMetadataTypes::SoftObjectPath:
+		return FString::Printf(TEXT("%s"), *SoftObjectPathValue.ToString());
+	case EPCGMetadataTypes::SoftClassPath:
+		return FString::Printf(TEXT("%s"), *SoftClassPathValue.ToString());
 	default:
 		break;
 	}
 
 	return FString();
-}
\ No newline at end of file
+}
+
+#if WITH_EDITOR
+void FPCGMetadataTypesConstantStruct::OnPostLoad()
+{
+	// We used to represent soft object/class paths as strings, but now we support them natively.
+	if (Type == EPCGMetadataTypes::String)
+	{
+		if (StringMode_DEPRECATED == EPCGMetadataTypesConstantStructStringMode::SoftObjectPath)
+		{
+			Type = EPCGMetadataTypes::SoftObjectPath;
+		}
+		if (StringMode_DEPRECATED == EPCGMetadataTypesConstantStructStringMode::SoftClassPath)
+		{
+			Type = EPCGMetadataTypes::SoftClassPath;
+		}
+	}
+}
+#endif
diff --git a/Engine/Plugins/Experimental/PCG/Source/PCG/Private/PCGParamData.cpp b/Engine/Plugins/Experimental/PCG/Source/PCG/Private/PCGParamData.cpp
--- a/Engine/Plugins/Experimental/PCG/Source/PCG/Private/PCGParamData.cpp
+++ b/Engine/Plugins/Experimental/PCG/Source/PCG/Private/PCGParamData.cpp
@@ -54,7 +54,7 @@ void UPCGParamData::AddToCrc(FArchiveCrc32& Ar, bool bFullDataCrc) const
 			for (AttributeType Value : Values)
 			{
 				// Add value to Crc
-				Ar << Value;
+				PCG::Private::Serialize(Ar, Value);
 			}
 		};
 
diff --git a/Engine/Plugins/Experimental/PCG/Source/PCG/Public/Elements/PCGPointFilter.h b/Engine/Plugins/Experimental/PCG/Source/PCG/Public/Elements/PCGPointFilter.h
--- a/Engine/Plugins/Experimental/PCG/Source/PCG/Public/Elements/PCGPointFilter.h
+++ b/Engine/Plugins/Experimental/PCG/Source/PCG/Public/Elements/PCGPointFilter.h
@@ -54,6 +54,10 @@ struct PCG_API FPCGPointFilterThresholdSettings
 {
 	GENERATED_BODY()
 
+#if WITH_EDITOR
+	void OnPostLoad();
+#endif
+
 	/** If the threshold in included or excluded from the range. */
 	UPROPERTY(BlueprintReadWrite, EditAnywhere, Category = Settings, meta = (PCG_Overridable))
 	bool bInclusive = true;
@@ -193,6 +197,11 @@ class PCG_API UPCGPointFilterRangeSettings : public UPCGSettings
 public:
 	UPCGPointFilterRangeSettings();
 
+	//~Begin UObject interface
+#if WITH_EDITOR
+	virtual void PostLoad() override;
+#endif
+	//~End UObject interface
 
 	//~Begin UPCGSettings interface
 #if WITH_EDITOR
diff --git a/Engine/Plugins/Experimental/PCG/Source/PCG/Public/Elements/PCGPointMatchAndSet.h b/Engine/Plugins/Experimental/PCG/Source/PCG/Public/Elements/PCGPointMatchAndSet.h
--- a/Engine/Plugins/Experimental/PCG/Source/PCG/Public/Elements/PCGPointMatchAndSet.h
+++ b/Engine/Plugins/Experimental/PCG/Source/PCG/Public/Elements/PCGPointMatchAndSet.h
@@ -71,8 +71,8 @@ public:
 	EPCGMetadataTypes SetTargetType = EPCGMetadataTypes::Double;
 
 	/** For string types, the subtype is used to cleanup the UI. */
-	UPROPERTY(BlueprintReadWrite, EditAnywhere, Category = Settings, meta = (EditCondition = "bSetTargetIsAttribute && SetTargetType == EPCGMetadataTypes::String", EditConditionHides))
-	EPCGMetadataTypesConstantStructStringMode SetTargetStringMode = EPCGMetadataTypesConstantStructStringMode::String;
+	UPROPERTY()
+	EPCGMetadataTypesConstantStructStringMode SetTargetStringMode_DEPRECATED;
 
 private:
 #if WITH_EDITORONLY_DATA
diff --git a/Engine/Plugins/Experimental/PCG/Source/PCG/Public/Helpers/PCGMetadataHelpers.h b/Engine/Plugins/Experimental/PCG/Source/PCG/Public/Helpers/PCGMetadataHelpers.h
--- a/Engine/Plugins/Experimental/PCG/Source/PCG/Public/Helpers/PCGMetadataHelpers.h
+++ b/Engine/Plugins/Experimental/PCG/Source/PCG/Public/Helpers/PCGMetadataHelpers.h
@@ -4,8 +4,11 @@
 
 #include "HAL/Platform.h"
 
+template <typename FuncType> class TFunction;
 class UPCGData;
 class UPCGMetadata;
+class FPCGMetadataAttributeBase;
+struct FSoftObjectPath;
 
 namespace PCGMetadataHelpers
 {
@@ -15,6 +18,12 @@ namespace PCGMetadataHelpers
 	// Helpers functions to cast in spatial/param and return metadata. Nullptr if data doesn't have metadata
 	PCG_API const UPCGMetadata* GetConstMetadata(const UPCGData* InData);
 	PCG_API UPCGMetadata* GetMutableMetadata(UPCGData* InData);
+	
+	/** Create a lambda that will construct a soft object path from an underlying attribute of type FSoftObjectPath or FString. Returns true if successful. */
+	PCG_API UE_NODISCARD bool CreateObjectPathGetter(const FPCGMetadataAttributeBase* InAttributeBase, TFunction<void(int64, FSoftObjectPath&)>& OutGetter);
+
+	/** Create a lambda that will construct a soft object path from an underlying attribute of type FSoftObjectPath or FString. Returns true if successful. */
+	PCG_API UE_NODISCARD bool CreateObjectOrClassPathGetter(const FPCGMetadataAttributeBase* InAttributeBase, TFunction<void(int64, FSoftObjectPath&)>& OutGetter);
 }
 
 #if UE_ENABLE_INCLUDE_ORDER_DEPRECATED_IN_5_2
diff --git a/Engine/Plugins/Experimental/PCG/Source/PCG/Public/MatchAndSet/PCGMatchAndSetBase.h b/Engine/Plugins/Experimental/PCG/Source/PCG/Public/MatchAndSet/PCGMatchAndSetBase.h
--- a/Engine/Plugins/Experimental/PCG/Source/PCG/Public/MatchAndSet/PCGMatchAndSetBase.h
+++ b/Engine/Plugins/Experimental/PCG/Source/PCG/Public/MatchAndSet/PCGMatchAndSetBase.h
@@ -26,6 +26,12 @@ class PCG_API UPCGMatchAndSetBase : public UObject
 	GENERATED_BODY()
 
 public:
+	//~Begin UObject interface
+#if WITH_EDITOR
+	virtual void PostLoad() override;
+#endif
+	//~End UObject interface
+
 	/** Queries whether this object uses a random process - note that this is expected to be static through the life of the object */
 	virtual bool UsesRandomProcess() const { return false; }
 
@@ -33,7 +39,7 @@ public:
 	virtual bool ShouldMutateSeed() const { return false; }
 
 	/** Sets & propagates type change from the owner settings object */
-	virtual void SetType(EPCGMetadataTypes InType, EPCGMetadataTypesConstantStructStringMode InStringMode);
+	virtual void SetType(EPCGMetadataTypes InType);
 
 	/** Early check to prevent issues when the data does not contain the required information to perform the operation */
 	UFUNCTION(BlueprintNativeEvent, Category = Selection)
@@ -68,7 +74,7 @@ protected:
 	EPCGMetadataTypes Type = EPCGMetadataTypes::Double;
 
 	UPROPERTY()
-	EPCGMetadataTypesConstantStructStringMode StringMode = EPCGMetadataTypesConstantStructStringMode::String;
+	EPCGMetadataTypesConstantStructStringMode StringMode_DEPRECATED;
 };
 
 #if UE_ENABLE_INCLUDE_ORDER_DEPRECATED_IN_5_2
diff --git a/Engine/Plugins/Experimental/PCG/Source/PCG/Public/MatchAndSet/PCGMatchAndSetByAttribute.h b/Engine/Plugins/Experimental/PCG/Source/PCG/Public/MatchAndSet/PCGMatchAndSetByAttribute.h
--- a/Engine/Plugins/Experimental/PCG/Source/PCG/Public/MatchAndSet/PCGMatchAndSetByAttribute.h
+++ b/Engine/Plugins/Experimental/PCG/Source/PCG/Public/MatchAndSet/PCGMatchAndSetByAttribute.h
@@ -15,6 +15,10 @@ struct PCG_API FPCGMatchAndSetByAttributeEntry
 
 	FPCGMatchAndSetByAttributeEntry();
 
+#if WITH_EDITOR
+	void OnPostLoad();
+#endif
+
 	UPROPERTY(BlueprintReadWrite, EditAnywhere, Category = Settings, meta = (ShowOnlyInnerProperties))
 	FPCGMetadataTypesConstantStruct ValueToMatch;
 
@@ -32,10 +36,10 @@ class PCG_API UPCGMatchAndSetByAttribute : public UPCGMatchAndSetBase
 	GENERATED_BODY()
 
 public:
-	virtual void SetType(EPCGMetadataTypes InType, EPCGMetadataTypesConstantStructStringMode InStringMode) override;
+	virtual void SetType(EPCGMetadataTypes InType) override;
 
 	/** Propagates (does not set) the Match type to the entries */
-	virtual void SetSourceType(EPCGMetadataTypes InType, EPCGMetadataTypesConstantStructStringMode InStringMode);
+	virtual void SetSourceType(EPCGMetadataTypes InType);
 
 	virtual void MatchAndSet_Implementation(
 		FPCGContext& Context,
@@ -45,9 +49,12 @@ public:
 
 	virtual bool ValidatePreconditions_Implementation(const UPCGPointData* InPointData) const override;
 
+	//~Begin UObject interface
 #if WITH_EDITOR
 	virtual void PostEditChangeProperty(FPropertyChangedEvent& PropertyChangedEvent) override;
+	virtual void PostLoad() override;
 #endif
+	//~End UObject interface
 
 public:
 	/** Attribute to match on the data */
@@ -59,8 +66,8 @@ public:
 	EPCGMetadataTypes MatchSourceType = EPCGMetadataTypes::Double;
 
 	/** String type of the attribute to match against (if required). */
-	UPROPERTY(BlueprintReadWrite, EditAnywhere, Category = Settings, meta = (EditCondition = "MatchSourceType == EPCGMetadataTypes::String", EditConditionHides))
-	EPCGMetadataTypesConstantStructStringMode MatchSourceStringMode = EPCGMetadataTypesConstantStructStringMode::String;
+	UPROPERTY()
+	EPCGMetadataTypesConstantStructStringMode MatchSourceStringMode_DEPRECATED;
 
 	/** Lookup entries (key-value pairs) */
 	UPROPERTY(BlueprintReadWrite, EditAnywhere, Category = Settings)
diff --git a/Engine/Plugins/Experimental/PCG/Source/PCG/Public/MatchAndSet/PCGMatchAndSetWeighted.h b/Engine/Plugins/Experimental/PCG/Source/PCG/Public/MatchAndSet/PCGMatchAndSetWeighted.h
--- a/Engine/Plugins/Experimental/PCG/Source/PCG/Public/MatchAndSet/PCGMatchAndSetWeighted.h
+++ b/Engine/Plugins/Experimental/PCG/Source/PCG/Public/MatchAndSet/PCGMatchAndSetWeighted.h
@@ -14,6 +14,10 @@ struct PCG_API FPCGMatchAndSetWeightedEntry
 {
 	GENERATED_BODY()
 
+#if WITH_EDITOR
+	void OnPostLoad();
+#endif
+
 	FPCGMatchAndSetWeightedEntry();
 
 	UPROPERTY(BlueprintReadWrite, EditAnywhere, Category = Settings, meta = (ShowOnlyInnerProperties))
@@ -36,7 +40,7 @@ class PCG_API UPCGMatchAndSetWeighted : public UPCGMatchAndSetBase
 public:
 	virtual bool UsesRandomProcess() const { return true; }
 	virtual bool ShouldMutateSeed() const { return bShouldMutateSeed; }
-	virtual void SetType(EPCGMetadataTypes InType, EPCGMetadataTypesConstantStructStringMode InStringMode) override;
+	virtual void SetType(EPCGMetadataTypes InType) override;
 
 	virtual void MatchAndSet_Implementation(
 		FPCGContext& Context,
@@ -44,9 +48,12 @@ public:
 		const UPCGPointData* InPointData,
 		UPCGPointData* OutPointData) const override;
 
+	//~Begin UObject interface
 #if WITH_EDITOR
 	virtual void PostEditChangeProperty(FPropertyChangedEvent& PropertyChangedEvent) override;
+	virtual void PostLoad() override;
 #endif
+	//~End UObject interface
 
 public:
 	/** Values and their respective weights */
diff --git a/Engine/Plugins/Experimental/PCG/Source/PCG/Public/MatchAndSet/PCGMatchAndSetWeightedByCategory.h b/Engine/Plugins/Experimental/PCG/Source/PCG/Public/MatchAndSet/PCGMatchAndSetWeightedByCategory.h
--- a/Engine/Plugins/Experimental/PCG/Source/PCG/Public/MatchAndSet/PCGMatchAndSetWeightedByCategory.h
+++ b/Engine/Plugins/Experimental/PCG/Source/PCG/Public/MatchAndSet/PCGMatchAndSetWeightedByCategory.h
@@ -3,6 +3,7 @@
 #pragma once
 
 #include "MatchAndSet/PCGMatchAndSetBase.h"
+#include "MatchAndSet/PCGMatchAndSetWeighted.h"
 
 #include "PCGMatchAndSetWeightedByCategory.generated.h"
 
@@ -16,7 +17,11 @@ struct PCG_API FPCGMatchAndSetWeightedByCategoryEntryList
 
 	FPCGMatchAndSetWeightedByCategoryEntryList();
 
-	void SetType(EPCGMetadataTypes InType, EPCGMetadataTypesConstantStructStringMode InStringMode);
+#if WITH_EDITOR
+	void OnPostLoad();
+#endif
+
+	void SetType(EPCGMetadataTypes InType);
 	int GetTotalWeight() const;
 
 	UPROPERTY(BlueprintReadWrite, EditAnywhere, Category = Settings, meta = (ShowOnlyInnerProperties))
@@ -39,9 +44,9 @@ class PCG_API UPCGMatchAndSetWeightedByCategory : public UPCGMatchAndSetBase
 public:
 	virtual bool UsesRandomProcess() const { return true; }
 	virtual bool ShouldMutateSeed() const { return bShouldMutateSeed; }
-	virtual void SetType(EPCGMetadataTypes InType, EPCGMetadataTypesConstantStructStringMode InStringMode) override;
+	virtual void SetType(EPCGMetadataTypes InType) override;
 	/** Propagates (does not set) category (e.g. Match) type to entries. */
-	virtual void SetCategoryType(EPCGMetadataTypes InType, EPCGMetadataTypesConstantStructStringMode InStringMode);
+	virtual void SetCategoryType(EPCGMetadataTypes InType);
 
 	virtual void MatchAndSet_Implementation(
 		FPCGContext& Context,
@@ -51,9 +56,12 @@ public:
 
 	virtual bool ValidatePreconditions_Implementation(const UPCGPointData* InPointData) const override;
 
+	//~Begin UObject interface
 #if WITH_EDITOR
 	virtual void PostEditChangeProperty(FPropertyChangedEvent& PropertyChangedEvent) override;
+	virtual void PostLoad() override;
 #endif
+	//~End UObject interface
 
 public:
 	/** Attribute to match against */
@@ -64,9 +72,8 @@ public:
 	UPROPERTY(BlueprintReadWrite, EditAnywhere, Category = Settings, meta = (ValidEnumValues = "Float, Double, Integer32, Integer64, Vector2, Vector, Vector4, Quaternion, Transform, String, Boolean, Rotator, Name"))
 	EPCGMetadataTypes CategoryType = EPCGMetadataTypes::Double;
 
-	/** String subtype of the attribute to match against (if required). */
-	UPROPERTY(BlueprintReadWrite, EditAnywhere, Category = Settings, meta = (EditCondition = "CategoryType == EPCGMetadataTypes::String", EditConditionHides))
-	EPCGMetadataTypesConstantStructStringMode CategoryStringMode = EPCGMetadataTypesConstantStructStringMode::String;
+	UPROPERTY()
+	EPCGMetadataTypesConstantStructStringMode CategoryStringMode_DEPRECATED;
 
 	/** Lookup entries (key -> weighted list) */
 	UPROPERTY(BlueprintReadWrite, EditAnywhere, Category = Settings)
diff --git a/Engine/Plugins/Experimental/PCG/Source/PCG/Public/MeshSelectors/PCGMeshSelectorBase.h b/Engine/Plugins/Experimental/PCG/Source/PCG/Public/MeshSelectors/PCGMeshSelectorBase.h
--- a/Engine/Plugins/Experimental/PCG/Source/PCG/Public/MeshSelectors/PCGMeshSelectorBase.h
+++ b/Engine/Plugins/Experimental/PCG/Source/PCG/Public/MeshSelectors/PCGMeshSelectorBase.h
@@ -77,7 +77,7 @@ struct FPCGMeshMaterialOverrideHelper
 
 private:
 	// Cached data
-	TArray<const FPCGMetadataAttribute<FString>*> MaterialAttributes;
+	TArray<const FPCGMetadataAttributeBase*> MaterialAttributes;
 	TArray<TMap<PCGMetadataValueKey, TSoftObjectPtr<UMaterialInterface>>> ValueKeyToOverrideMaterials;
 	TArray<TSoftObjectPtr<UMaterialInterface>> WorkingMaterialOverrides;
 
diff --git a/Engine/Plugins/Experimental/PCG/Source/PCG/Public/Metadata/Accessors/IPCGAttributeAccessorTpl.h b/Engine/Plugins/Experimental/PCG/Source/PCG/Public/Metadata/Accessors/IPCGAttributeAccessorTpl.h
--- a/Engine/Plugins/Experimental/PCG/Source/PCG/Public/Metadata/Accessors/IPCGAttributeAccessorTpl.h
+++ b/Engine/Plugins/Experimental/PCG/Source/PCG/Public/Metadata/Accessors/IPCGAttributeAccessorTpl.h
@@ -45,9 +45,15 @@ private:
 		}
 		else
 		{
+			// Special case'd because FSoftObjectPath currently has a deprecated constructor from FName which generates compile warnings.
+			constexpr bool bNameToSoftObjectPath = std::is_same_v<U, FSoftObjectPath> && std::is_same_v<T, FName>;
+
 			if constexpr (PCG::Private::IsBroadcastable<T, U>())
 			{
-				if (!!(Flags & EPCGAttributeAccessorFlags::AllowBroadcast))
+				// Special case - always allow broadcasting from soft object path to string, so that legacy code that grabbed soft path attributes
+				// as FStrings will still work.
+				constexpr bool bSoftReferenceToString = (std::is_same_v<T, FSoftObjectPath> || std::is_same_v<T, FSoftClassPath>) && std::is_same_v<U, FString>;
+				if (!!(Flags & EPCGAttributeAccessorFlags::AllowBroadcast) || bSoftReferenceToString)
 				{
 					TArray<T, TInlineAllocator<4>> InValues;
 					InValues.SetNum(OutValues.Num());
@@ -62,7 +68,7 @@ private:
 					}
 				}
 			}
-			else if constexpr (std::is_constructible_v<U, T>)
+			else if constexpr (std::is_constructible_v<U, T> && !bNameToSoftObjectPath)
 			{
 				if (!!(Flags & EPCGAttributeAccessorFlags::AllowConstructible))
 				{
@@ -104,6 +110,9 @@ private:
 		}
 		else
 		{
+			// Special case'd because FSoftObjectPath currently has a deprecated constructor from FName which generates compile warnings.
+			constexpr bool bNameToSoftObjectPath = std::is_same_v<T, FSoftObjectPath> && std::is_same_v<U, FName>;
+
 			if constexpr (PCG::Private::IsBroadcastable<U, T>())
 			{
 				if (!!(Flags & EPCGAttributeAccessorFlags::AllowBroadcast))
@@ -118,7 +127,7 @@ private:
 					return This->SetRangeImpl(OutValues, Index, Keys, Flags);
 				}
 			}
-			else if constexpr (std::is_constructible_v<T, U>)
+			else if constexpr (std::is_constructible_v<T, U> && !bNameToSoftObjectPath)
 			{
 				if (!!(Flags & EPCGAttributeAccessorFlags::AllowConstructible))
 				{
diff --git a/Engine/Plugins/Experimental/PCG/Source/PCG/Public/Metadata/Accessors/PCGPropertyAccessor.h b/Engine/Plugins/Experimental/PCG/Source/PCG/Public/Metadata/Accessors/PCGPropertyAccessor.h
--- a/Engine/Plugins/Experimental/PCG/Source/PCG/Public/Metadata/Accessors/PCGPropertyAccessor.h
+++ b/Engine/Plugins/Experimental/PCG/Source/PCG/Public/Metadata/Accessors/PCGPropertyAccessor.h
@@ -260,18 +260,17 @@ private:
 };
 
 /**
-* Templated accessor class for soft object/class ptr properties.
+* Templated accessor class for soft object ptr properties - produces soft object path.
 * Do not instantiate it manually, use PCGAttributeAccessorHelpers::CreatePropertyAccessor.
-* Will always convert to FString for PCG
 * Key supported: Generic object
 */
-class FPCGPropertySoftPtrAccessor : public IPCGAttributeAccessorT<FPCGPropertySoftPtrAccessor>
+class FPCGPropertySoftObjectPathAccessor : public IPCGAttributeAccessorT<FPCGPropertySoftObjectPathAccessor>
 {
 public:
-	using Type = FString;
-	using Super = IPCGAttributeAccessorT<FPCGPropertySoftPtrAccessor>;
+	using Type = FSoftObjectPath;
+	using Super = IPCGAttributeAccessorT<FPCGPropertySoftObjectPathAccessor>;
 
-	FPCGPropertySoftPtrAccessor(const FSoftObjectProperty* InProperty)
+	FPCGPropertySoftObjectPathAccessor(const FSoftObjectProperty* InProperty)
 		: Super(/*bInReadOnly=*/ false)
 		, Property(InProperty)
 	{
@@ -282,7 +281,7 @@ public:
 	{
 		return PCGPropertyAccessor::IterateGet(Property, OutValues, Index, Keys, [this](const void* PropertyAddressData) -> Type
 		{
-			return Property->GetPropertyValue(PropertyAddressData).ToString();
+			return Property->GetPropertyValue(PropertyAddressData).ToSoftObjectPath();
 		});
 	}
 
@@ -290,7 +289,7 @@ public:
 	{
 		return PCGPropertyAccessor::IterateSet(Property, InValues, Index, Keys, [this](void* PropertyAddressData, const Type& Value) -> void
 		{
-			Property->SetPropertyValue(PropertyAddressData, FSoftObjectPtr(FSoftObjectPath(Value)));
+			Property->SetPropertyValue(PropertyAddressData, FSoftObjectPtr(Value));
 		});
 	}
 
@@ -298,6 +297,44 @@ private:
 	const FSoftObjectProperty* Property = nullptr;
 };
 
+/**
+* Templated accessor class for soft class ptr properties - produces soft class path.
+* Do not instantiate it manually, use PCGAttributeAccessorHelpers::CreatePropertyAccessor.
+* Key supported: Generic object
+*/
+class FPCGPropertySoftClassPathAccessor : public IPCGAttributeAccessorT<FPCGPropertySoftClassPathAccessor>
+{
+public:
+	using Type = FSoftClassPath;
+	using Super = IPCGAttributeAccessorT<FPCGPropertySoftClassPathAccessor>;
+
+	FPCGPropertySoftClassPathAccessor(const FSoftClassProperty* InProperty)
+		: Super(/*bInReadOnly=*/ false)
+		, Property(InProperty)
+	{
+		check(Property);
+	}
+
+	bool GetRangeImpl(TArrayView<Type> OutValues, int32 Index, const IPCGAttributeAccessorKeys& Keys) const
+	{
+		return PCGPropertyAccessor::IterateGet(Property, OutValues, Index, Keys, [this](const void* PropertyAddressData) -> Type
+		{
+			return FSoftClassPath(Property->GetPropertyValue(PropertyAddressData).ToString());
+		});
+	}
+
+	bool SetRangeImpl(TArrayView<const Type> InValues, int32 Index, IPCGAttributeAccessorKeys& Keys, EPCGAttributeAccessorFlags)
+	{
+		return PCGPropertyAccessor::IterateSet(Property, InValues, Index, Keys, [this](void* PropertyAddressData, const Type& Value) -> void
+		{
+			Property->SetPropertyValue(PropertyAddressData, FSoftObjectPtr(Value));
+		});
+	}
+
+private:
+	const FSoftClassProperty* Property = nullptr;
+};
+
 /**
 * Templated accessor class for object/class ptr properties.
 * Do not instantiate it manually, use PCGAttributeAccessorHelpers::CreatePropertyAccessor.
diff --git a/Engine/Plugins/Experimental/PCG/Source/PCG/Public/Metadata/PCGMetadata.h b/Engine/Plugins/Experimental/PCG/Source/PCG/Public/Metadata/PCGMetadata.h
--- a/Engine/Plugins/Experimental/PCG/Source/PCG/Public/Metadata/PCGMetadata.h
+++ b/Engine/Plugins/Experimental/PCG/Source/PCG/Public/Metadata/PCGMetadata.h
@@ -125,6 +125,12 @@ public:
 	UFUNCTION(BlueprintCallable, Category = "PCG|Metadata")
 	void CreateBoolAttribute(FName AttributeName, bool DefaultValue, bool bAllowsInterpolation, bool bOverrideParent = true);
 
+	UFUNCTION(BlueprintCallable, Category = "PCG|Metadata")
+	void CreateSoftObjectPathAttribute(FName AttributeName, const FSoftObjectPath& DefaultValue, bool bAllowsInterpolation, bool bOverrideParent = true);
+
+	UFUNCTION(BlueprintCallable, Category = "PCG|Metadata")
+	void CreateSoftClassPathAttribute(FName AttributeName, const FSoftClassPath& DefaultValue, bool bAllowsInterpolation, bool bOverrideParent = true);
+
 	/** Creates an attribute given a property.
 	* @param AttributeName: Target attribute to create
 	* @param Object: Object to get the property value from
diff --git a/Engine/Plugins/Experimental/PCG/Source/PCG/Public/Metadata/PCGMetadataAccessor.h b/Engine/Plugins/Experimental/PCG/Source/PCG/Public/Metadata/PCGMetadataAccessor.h
--- a/Engine/Plugins/Experimental/PCG/Source/PCG/Public/Metadata/PCGMetadataAccessor.h
+++ b/Engine/Plugins/Experimental/PCG/Source/PCG/Public/Metadata/PCGMetadataAccessor.h
@@ -90,6 +90,18 @@ public:
 	UFUNCTION(BlueprintCallable, Category = "PCG|Metadata")
 	static void SetBoolAttributeByMetadataKey(UPARAM(ref) int64& Key, UPCGMetadata* Metadata, FName AttributeName, bool Value);
 
+	UFUNCTION(BlueprintCallable, Category = "PCG|Metadata")
+	static FSoftObjectPath GetSoftObjectPathAttributeByMetadataKey(int64 Key, const UPCGMetadata* Metadata, FName AttributeName);
+
+	UFUNCTION(BlueprintCallable, Category = "PCG|Metadata")
+	static void SetSoftObjectPathAttributeByMetadataKey(UPARAM(ref) int64& Key, UPCGMetadata* Metadata, FName AttributeName, const FSoftObjectPath& Value);
+
+	UFUNCTION(BlueprintCallable, Category = "PCG|Metadata")
+	static FSoftClassPath GetSoftClassPathAttributeByMetadataKey(int64 Key, const UPCGMetadata* Metadata, FName AttributeName);
+
+	UFUNCTION(BlueprintCallable, Category = "PCG|Metadata")
+	static void SetSoftClassPathAttributeByMetadataKey(UPARAM(ref) int64& Key, UPCGMetadata* Metadata, FName AttributeName, const FSoftClassPath& Value);
+
 	UFUNCTION(BlueprintCallable, Category = "PCG|Metadata")
 	static bool SetAttributeFromPropertyByMetadataKey(UPARAM(ref) int64& Key, UPCGMetadata* Metadata, FName AttributeName, const UObject* Object, FName PropertyName);
 
@@ -186,6 +198,18 @@ public:
 	UFUNCTION(BlueprintCallable, Category = "PCG|Metadata", meta = (ScriptMethod))
 	static void SetBoolAttribute(UPARAM(ref) FPCGPoint& Point, UPCGMetadata* Metadata, FName AttributeName, bool Value);
 
+	UFUNCTION(BlueprintCallable, Category = "PCG|Metadata", meta = (ScriptMethod))
+	static FSoftObjectPath GetSoftObjectPathAttribute(const FPCGPoint& Point, const UPCGMetadata* Metadata, FName AttributeName);
+
+	UFUNCTION(BlueprintCallable, Category = "PCG|Metadata", meta = (ScriptMethod))
+	static void SetSoftObjectPathAttribute(UPARAM(ref) FPCGPoint& Point, UPCGMetadata* Metadata, FName AttributeName, const FSoftObjectPath& Value);
+
+	UFUNCTION(BlueprintCallable, Category = "PCG|Metadata", meta = (ScriptMethod))
+	static FSoftClassPath GetSoftClassPathAttribute(const FPCGPoint& Point, const UPCGMetadata* Metadata, FName AttributeName);
+
+	UFUNCTION(BlueprintCallable, Category = "PCG|Metadata", meta = (ScriptMethod))
+	static void SetSoftClassPathAttribute(UPARAM(ref) FPCGPoint& Point, UPCGMetadata* Metadata, FName AttributeName, const FSoftClassPath& Value);
+
 	UFUNCTION(BlueprintCallable, Category = "PCG|Metadata", meta = (ScriptMethod))
 	static bool HasAttributeSet(const FPCGPoint& Point, const UPCGMetadata* Metadata, FName AttributeName);
 
diff --git a/Engine/Plugins/Experimental/PCG/Source/PCG/Public/Metadata/PCGMetadataAttributeTraits.h b/Engine/Plugins/Experimental/PCG/Source/PCG/Public/Metadata/PCGMetadataAttributeTraits.h
--- a/Engine/Plugins/Experimental/PCG/Source/PCG/Public/Metadata/PCGMetadataAttributeTraits.h
+++ b/Engine/Plugins/Experimental/PCG/Source/PCG/Public/Metadata/PCGMetadataAttributeTraits.h
@@ -4,6 +4,7 @@
 
 #include "Math/Transform.h" // IWYU pragma: keep
 #include "UObject/Class.h"
+#include "UObject/SoftObjectPath.h"
 
 #include "PCGMetadataAttributeTraits.generated.h"
 
@@ -23,6 +24,8 @@ enum class EPCGMetadataTypes : uint8
 	Boolean,
 	Rotator,
 	Name,
+	SoftObjectPath,
+	SoftClassPath,
 
 	Count UMETA(Hidden),
 
@@ -31,19 +34,21 @@ enum class EPCGMetadataTypes : uint8
 
 // Convenient macro to avoid duplicating a lot of code with all our supported types.
 #define PCG_FOREACH_SUPPORTEDTYPES(MACRO) \
-	MACRO(int32)      \
-	MACRO(int64)      \
-	MACRO(float)      \
-	MACRO(double)     \
-	MACRO(FVector2D)  \
-	MACRO(FVector)    \
-	MACRO(FVector4)   \
-	MACRO(FQuat)      \
-	MACRO(FTransform) \
-	MACRO(FString)    \
-	MACRO(bool)       \
-	MACRO(FRotator)   \
-	MACRO(FName)
+	MACRO(int32)           \
+	MACRO(int64)           \
+	MACRO(float)           \
+	MACRO(double)          \
+	MACRO(FVector2D)       \
+	MACRO(FVector)         \
+	MACRO(FVector4)        \
+	MACRO(FQuat)           \
+	MACRO(FTransform)      \
+	MACRO(FString)         \
+	MACRO(bool)            \
+	MACRO(FRotator)        \
+	MACRO(FName)           \
+	MACRO(FSoftObjectPath) \
+	MACRO(FSoftClassPath)
 
 namespace PCG
 {
@@ -70,6 +75,8 @@ namespace PCG
 		PCGMetadataGenerateDataTypes(bool, Boolean);
 		PCGMetadataGenerateDataTypes(FRotator, Rotator);
 		PCGMetadataGenerateDataTypes(FName, Name);
+		PCGMetadataGenerateDataTypes(FSoftObjectPath, SoftObjectPath);
+		PCGMetadataGenerateDataTypes(FSoftClassPath, SoftClassPath);
 
 #undef PCGMetadataGenerateDataTypes
 
@@ -124,6 +131,26 @@ namespace PCG
 			return FText::FromString(GetTypeName<T>());
 		}
 
+		template <typename T>
+		void Serialize(FArchive& Ar, const T& A)
+		{
+			ensure(!Ar.IsLoading());
+			Serialize(Ar, const_cast<T&>(A));
+		}
+
+		template <typename T>
+		void Serialize(FArchive& Ar, T& A)
+		{
+			Ar << A;
+		}
+
+		// Serialize must be called on FSoftObjectPath, and this also covers FSoftClassPath.
+		template <>
+		inline void Serialize<FSoftObjectPath>(FArchive& Ar, FSoftObjectPath& A)
+		{
+			A.Serialize(Ar);
+		}
+
 		// Wrapper around a standard 2-dimensional CArray that is constexpr, to know if a type is broadcastable to another.
 		// First index is the original type, second index is the wanted type. Returns true if we can broadcast first type into second type.
 		struct UBroadcastableTypes
@@ -182,6 +209,9 @@ namespace PCG
 				PCGMetadataBroadcastable(Name, String);
 				PCGMetadataBroadcastable(String, Name);
 
+				PCGMetadataBroadcastable(SoftObjectPath, String);
+				PCGMetadataBroadcastable(SoftClassPath, String);
+
 #undef PCGMetadataBroadcastable
 			}
 
@@ -228,6 +258,14 @@ namespace PCG
 			return IsMoreComplexType(MetadataTypes<FirstType>::Id, MetadataTypes<SecondType>::Id);
 		}
 
+		/**
+		* Traits
+		* - CanCompare: Supports less than, greater than etc for ordering.
+		* - CanSearchString: Can do string search operations like matching substrings.
+		* - CompressData: Uses value keys, so if two values are identical, we only store one.
+		* - NeedsConstruction: Whether type requires constructor for valid state (i.e. can't simply be zero initialized).
+		*/
+
 		template<typename T>
 		struct DefaultStringTraits
 		{
@@ -812,6 +850,72 @@ namespace PCG
 			}
 		};
 
+		template<typename T>
+		struct SoftObjectPathTraits
+		{
+			enum { CompressData = true };
+			enum { CanMinMax = false };
+			enum { CanSubAdd = false };
+			enum { CanMulDiv = false };
+			enum { CanInterpolate = false };
+			enum { CanSearchString = true };
+			// Contains FString objects
+			enum { NeedsConstruction = true };
+			enum { CanCompare = true };
+
+			static bool Equal(const T& A, const T& B)
+			{
+				return A == B;
+			}
+
+			static T ZeroValue()
+			{
+				return T();
+			}
+
+			static bool Substring(const T& A, const T& B)
+			{
+				return A.ToString().Contains(B.ToString());
+			}
+
+			static bool Matches(const T& A, const T& B)
+			{
+				return A.ToString().MatchesWildcard(B.ToString());
+			}
+
+			static bool Less(const T& A, const T& B)
+			{
+				return A.ToString() < B.ToString();
+			}
+
+			static bool Greater(const T& A, const T& B)
+			{
+				return A.ToString() > B.ToString();
+			}
+
+			static bool LessOrEqual(const T& A, const T& B)
+			{
+				return A.ToString() <= B.ToString();
+			}
+
+			static bool GreaterOrEqual(const FSoftObjectPath& A, const FSoftObjectPath& B)
+			{
+				return A.ToString() >= B.ToString();
+			}
+		};
+
+		// Soft object paths
+		template<>
+		struct MetadataTraits<FSoftObjectPath> : SoftObjectPathTraits<FSoftObjectPath>, DefaultStringTraits<FSoftObjectPath>
+		{
+		};
+
+		// Soft class paths
+		template<>
+		struct MetadataTraits<FSoftClassPath> : SoftObjectPathTraits<FSoftClassPath>, DefaultStringTraits<FSoftClassPath>
+		{
+		};
+
 		/**
 		* Generic function to broadcast an InType to an OutType.
 		* Supports only PCG types
diff --git a/Engine/Plugins/Experimental/PCG/Source/PCG/Public/Metadata/PCGMetadataTypesConstantStruct.h b/Engine/Plugins/Experimental/PCG/Source/PCG/Public/Metadata/PCGMetadataTypesConstantStruct.h
--- a/Engine/Plugins/Experimental/PCG/Source/PCG/Public/Metadata/PCGMetadataTypesConstantStruct.h
+++ b/Engine/Plugins/Experimental/PCG/Source/PCG/Public/Metadata/PCGMetadataTypesConstantStruct.h
@@ -4,16 +4,19 @@
 
 #include "Metadata/PCGMetadataAttributeTraits.h"
 
+#include "UObject/SoftObjectPath.h"
+
 #include "PCGMetadataTypesConstantStruct.generated.h"
 
 class UPCGParamData;
 
+// Deprecated: we support the path types natively now.
 UENUM()
 enum class EPCGMetadataTypesConstantStructStringMode
 {
 	String,
-	SoftObjectPath,
-	SoftClassPath,
+	SoftObjectPath UMETA(Hidden),
+	SoftClassPath UMETA(Hidden),
 };
 
 /**
@@ -30,11 +33,16 @@ public:
 
 	FString ToString() const;
 
-	UPROPERTY(BlueprintReadWrite, EditAnywhere, Category = Settings, meta = (EditCondition = "bAllowsTypeChange", EditConditionHides, HideEditConditionToggle, ValidEnumValues = "Float, Double, Integer32, Integer64, Vector2, Vector, Vector4, Quaternion, Transform, String, Boolean, Rotator, Name"))
+#if WITH_EDITOR
+	/** This is not called by the engine because this is a USTRUCT, but it can be called from usages to perform deprecation. */
+	void OnPostLoad();
+#endif
+
+	UPROPERTY(BlueprintReadWrite, EditAnywhere, Category = Settings, meta = (EditCondition = "bAllowsTypeChange", EditConditionHides, HideEditConditionToggle, ValidEnumValues = "Float, Double, Integer32, Integer64, Vector2, Vector, Vector4, Quaternion, Transform, String, Boolean, Rotator, Name, SoftObjectPath, SoftClassPath"))
 	EPCGMetadataTypes Type = EPCGMetadataTypes::Double;
 
-	UPROPERTY(BlueprintReadWrite, EditAnywhere, Category = Settings, meta = (EditCondition = "bAllowsTypeChange && Type == EPCGMetadataTypes::String", EditConditionHides))
-	EPCGMetadataTypesConstantStructStringMode StringMode = EPCGMetadataTypesConstantStructStringMode::String;
+	UPROPERTY()
+	EPCGMetadataTypesConstantStructStringMode StringMode_DEPRECATED;
 
 	// All different types
 	UPROPERTY(BlueprintReadWrite, EditAnywhere, Category = Settings, meta = (EditCondition = "Type == EPCGMetadataTypes::Float", EditConditionHides))
@@ -64,7 +72,7 @@ public:
 	UPROPERTY(BlueprintReadWrite, EditAnywhere, Category = Settings, meta = (EditCondition = "Type == EPCGMetadataTypes::Transform", EditConditionHides))
 	FTransform TransformValue = FTransform::Identity;
 
-	UPROPERTY(BlueprintReadWrite, EditAnywhere, Category = Settings, meta = (EditCondition = "Type == EPCGMetadataTypes::String && StringMode == EPCGMetadataTypesConstantStructStringMode::String", EditConditionHides))
+	UPROPERTY(BlueprintReadWrite, EditAnywhere, Category = Settings, meta = (EditCondition = "Type == EPCGMetadataTypes::String", EditConditionHides))
 	FString StringValue = "";
 
 	UPROPERTY(BlueprintReadWrite, EditAnywhere, Category = Settings, meta = (EditCondition = "Type == EPCGMetadataTypes::Boolean", EditConditionHides))
@@ -76,10 +84,10 @@ public:
 	UPROPERTY(BlueprintReadWrite, EditAnywhere, Category = Settings, meta = (EditCondition = "Type == EPCGMetadataTypes::Name", EditConditionHides))
 	FName NameValue = NAME_None;
 
-	UPROPERTY(BlueprintReadWrite, EditAnywhere, Category = Settings, meta = (EditCondition = "Type == EPCGMetadataTypes::String && StringMode == EPCGMetadataTypesConstantStructStringMode::SoftClassPath", EditConditionHides))
+	UPROPERTY(BlueprintReadWrite, EditAnywhere, Category = Settings, meta = (EditCondition = "Type == EPCGMetadataTypes::SoftClassPath", EditConditionHides))
 	FSoftClassPath SoftClassPathValue;
 
-	UPROPERTY(BlueprintReadWrite, EditAnywhere, Category = Settings, meta = (EditCondition = "Type == EPCGMetadataTypes::String && StringMode == EPCGMetadataTypesConstantStructStringMode::SoftObjectPath", EditConditionHides))
+	UPROPERTY(BlueprintReadWrite, EditAnywhere, Category = Settings, meta = (EditCondition = "Type == EPCGMetadataTypes::SoftObjectPath", EditConditionHides))
 	FSoftObjectPath SoftObjectPathValue;
 
 	UPROPERTY()
@@ -112,25 +120,17 @@ decltype(auto) FPCGMetadataTypesConstantStruct::Dispatcher(Func Callback) const
 	case EPCGMetadataTypes::Transform:
 		return Callback(TransformValue);
 	case EPCGMetadataTypes::String:
-	{
-		switch (StringMode)
-		{
-		case EPCGMetadataTypesConstantStructStringMode::String:
-			return Callback(StringValue);
-		case EPCGMetadataTypesConstantStructStringMode::SoftObjectPath:
-			return Callback(SoftObjectPathValue.ToString());
-		case EPCGMetadataTypesConstantStructStringMode::SoftClassPath:
-			return Callback(SoftClassPathValue.ToString());
-		default:
-			break;
-		}
-	}
+		return Callback(StringValue);
 	case EPCGMetadataTypes::Boolean:
 		return Callback(BoolValue);
 	case EPCGMetadataTypes::Rotator:
 		return Callback(RotatorValue);
 	case EPCGMetadataTypes::Name:
 		return Callback(NameValue);
+	case EPCGMetadataTypes::SoftObjectPath:
+		return Callback(SoftObjectPathValue);
+	case EPCGMetadataTypes::SoftClassPath:
+		return Callback(SoftClassPathValue);
 	default:
 		break;
 	}
diff --git a/Engine/Plugins/Experimental/PCG/Source/PCGEditor/Private/SPCGEditorGraphAttributeListView.cpp b/Engine/Plugins/Experimental/PCG/Source/PCGEditor/Private/SPCGEditorGraphAttributeListView.cpp
--- a/Engine/Plugins/Experimental/PCG/Source/PCGEditor/Private/SPCGEditorGraphAttributeListView.cpp
+++ b/Engine/Plugins/Experimental/PCG/Source/PCGEditor/Private/SPCGEditorGraphAttributeListView.cpp
@@ -233,6 +233,14 @@ TSharedRef<SWidget> SPCGListViewItemRow::GenerateWidgetForColumn(const FName& Co
 					{
 						RowText = FText::AsNumber(Value);
 					}
+					else if constexpr (PCG::Private::IsOfTypes<ValueType, FSoftObjectPath>())
+					{
+						RowText = FText::FromString(Value.ToString());
+					}
+					else if constexpr (PCG::Private::IsOfTypes<ValueType, FSoftClassPath>())
+					{
+						RowText = FText::FromString(Value.ToString());
+					}
 					else
 					{
 						ensureMsgf(false, TEXT("Unsupported Data Type"));
@@ -984,6 +992,8 @@ void SPCGEditorGraphAttributeListView::GenerateColumnsFromMetadata(const UPCGDat
 		case EPCGMetadataTypes::Boolean:
 		case EPCGMetadataTypes::String:
 		case EPCGMetadataTypes::Name:
+		case EPCGMetadataTypes::SoftObjectPath:
+		case EPCGMetadataTypes::SoftClassPath:
 			{
 				AddMetadataColumn(InPCGData, ColumnName, AttributeType);
 				break;
