[PCG] Fix uinitialized variable warning

diff --git a/Engine/Plugins/Experimental/PCG/Source/PCG/Public/Metadata/PCGMetadataTypesConstantStruct.h b/Engine/Plugins/Experimental/PCG/Source/PCG/Public/Metadata/PCGMetadataTypesConstantStruct.h
--- a/Engine/Plugins/Experimental/PCG/Source/PCG/Public/Metadata/PCGMetadataTypesConstantStruct.h
+++ b/Engine/Plugins/Experimental/PCG/Source/PCG/Public/Metadata/PCGMetadataTypesConstantStruct.h
@@ -42,7 +42,7 @@ public:
 	EPCGMetadataTypes Type = EPCGMetadataTypes::Double;
 
 	UPROPERTY()
-	EPCGMetadataTypesConstantStructStringMode StringMode_DEPRECATED;
+	EPCGMetadataTypesConstantStructStringMode StringMode_DEPRECATED = EPCGMetadataTypesConstantStructStringMode::String;
 
 	// All different types
 	UPROPERTY(BlueprintReadWrite, EditAnywhere, Category = Settings, meta = (EditCondition = "Type == EPCGMetadataTypes::Float", EditConditionHides))
