[PCG] Fix localization issues

#rnx

diff --git a/Engine/Plugins/Experimental/PCG/Source/PCG/Public/Elements/PCGNumberOfElements.h b/Engine/Plugins/Experimental/PCG/Source/PCG/Public/Elements/PCGNumberOfElements.h
--- a/Engine/Plugins/Experimental/PCG/Source/PCG/Public/Elements/PCGNumberOfElements.h
+++ b/Engine/Plugins/Experimental/PCG/Source/PCG/Public/Elements/PCGNumberOfElements.h
@@ -41,8 +41,8 @@ public:
 	//~Begin UPCGSettings interface
 #if WITH_EDITOR
 	virtual FName GetDefaultNodeName() const override { return FName(TEXT("GetPointsCount")); }
-	virtual FText GetDefaultNodeTitle() const override { return NSLOCTEXT("PCGNumberOfElementsSettings", "NodeTitle", "Get Points Count"); }
-	virtual FText GetNodeTooltipText() const override { return NSLOCTEXT("PCGNumberOfElementsSettings", "NodeTooltip", "Return the number of points in the input point data."); }
+	virtual FText GetDefaultNodeTitle() const override { return NSLOCTEXT("PCGNumberOfElementsSettings", "NodeTitlePoint", "Get Points Count"); }
+	virtual FText GetNodeTooltipText() const override { return NSLOCTEXT("PCGNumberOfElementsSettings", "NodeTooltipPoint", "Return the number of points in the input point data."); }
 	virtual EPCGSettingsType GetType() const override { return EPCGSettingsType::Spatial; }
 #endif
 
@@ -65,8 +65,8 @@ public:
 	//~Begin UPCGSettings interface
 #if WITH_EDITOR
 	virtual FName GetDefaultNodeName() const override { return FName(TEXT("GetEntriesCount")); }
-	virtual FText GetDefaultNodeTitle() const override { return NSLOCTEXT("PCGNumberOfElementsSettings", "NodeTitle", "Get Entries Count"); }
-	virtual FText GetNodeTooltipText() const override { return NSLOCTEXT("PCGNumberOfElementsSettings", "NodeTooltip", "Return the number of entries in the input attribute sets."); }
+	virtual FText GetDefaultNodeTitle() const override { return NSLOCTEXT("PCGNumberOfElementsSettings", "NodeTitleEntry", "Get Entries Count"); }
+	virtual FText GetNodeTooltipText() const override { return NSLOCTEXT("PCGNumberOfElementsSettings", "NodeTooltipEntry", "Return the number of entries in the input attribute sets."); }
 	virtual EPCGSettingsType GetType() const override { return EPCGSettingsType::Generic; }
 #endif
 
