[PCG] Fixes and improvements handling corrupt graphs.

Log errors in more cases and remove edges to non-existent nodes which was hit on one of the projects.

#rnx

diff --git a/Engine/Plugins/Experimental/PCG/Source/PCG/Private/PCGGraph.cpp b/Engine/Plugins/Experimental/PCG/Source/PCG/Private/PCGGraph.cpp
--- a/Engine/Plugins/Experimental/PCG/Source/PCG/Private/PCGGraph.cpp
+++ b/Engine/Plugins/Experimental/PCG/Source/PCG/Private/PCGGraph.cpp
@@ -26,6 +26,14 @@
 
 #define LOCTEXT_NAMESPACE "PCGGraph"
 
+namespace PCGGraph
+{
+	static TAutoConsoleVariable<bool> CVarFixInvalidEdgesOnPostLoad(
+		TEXT("pcg.Graph.FixInvalidEdgesOnPostLoad"),
+		true,
+		TEXT("Validates all edges are connected to valid pins/nodes and removes any invalid edges"));
+}
+
 namespace PCGGraphUtils
 {
 	/** Returns true if the two descriptors are valid and compatible */
@@ -316,6 +324,11 @@ void UPCGGraph::PostLoad()
 			OnNodeAdded(Nodes[i]);
 		}
 	}
+
+	if (PCGGraph::CVarFixInvalidEdgesOnPostLoad.GetValueOnAnyThread())
+	{
+		FixInvalidEdges();
+	}
 #endif
 }
 
@@ -1042,6 +1055,49 @@ void UPCGGraph::PostEditChangeProperty(FPropertyChangedEvent& PropertyChangedEve
 	NumberOfUserParametersPreEdit = 0;
 }
 
+void UPCGGraph::FixInvalidEdges()
+{
+	auto ValidatePins = [this](const TArray<TObjectPtr<UPCGPin>>& Pins, bool bPinsAreInputs)
+	{
+		for (UPCGPin* Pin : Pins)
+		{
+			if (!Pin)
+			{
+				continue;
+			}
+
+			for (int32 i = Pin->Edges.Num() - 1; i >= 0; --i)
+			{
+				UPCGPin* OtherPin = Pin->Edges[i] ? (bPinsAreInputs ? Pin->Edges[i]->InputPin : Pin->Edges[i]->OutputPin) : nullptr;
+
+				// Remove trivially invalid edges.
+				if (!OtherPin || !OtherPin->Node)
+				{
+					UE_LOG(LogPCG, Error, TEXT("Removed edge to a missing pin or pin that has no node."));
+					ensure(false);
+					Pin->Edges.RemoveAt(i);
+					continue;
+				}
+
+				// Remove edges to nodes that are not present in the graph.
+				UPCGNode* ConnectedNode = OtherPin->Node;
+				if (!ConnectedNode || (GetInputNode() != ConnectedNode && GetOutputNode() != ConnectedNode && !Nodes.Contains(ConnectedNode)))
+				{
+					UE_LOG(LogPCG, Error, TEXT("Removed edge to a node '%s' that is not registered in the graph."), ConnectedNode ? *ConnectedNode->GetFName().ToString() : TEXT("NULL"));
+					ensure(false);
+					Pin->Edges.RemoveAt(i);
+				}
+			}
+		}
+	};
+
+	ForEachNode([&ValidatePins](UPCGNode* InNode)
+	{
+		ValidatePins(InNode->GetInputPins(), /*bPinsAreInputs=*/true);
+		ValidatePins(InNode->GetOutputPins(), /*bPinsAreInputs=*/false);
+	});
+}
+
 bool UPCGGraph::UserParametersCanRemoveProperty(FGuid InPropertyID, FName InPropertyName)
 {
 	// Check if the property has some getters in the graph
diff --git a/Engine/Plugins/Experimental/PCG/Source/PCG/Public/PCGGraph.h b/Engine/Plugins/Experimental/PCG/Source/PCG/Public/PCGGraph.h
--- a/Engine/Plugins/Experimental/PCG/Source/PCG/Public/PCGGraph.h
+++ b/Engine/Plugins/Experimental/PCG/Source/PCG/Public/PCGGraph.h
@@ -384,6 +384,9 @@ private:
 
 	void NotifyGraphParametersChanged(EPCGGraphParameterEvent InChangeType, FName InChangedPropertyName);
 
+	/** Remove invalid edges and edges to nodes that are not present in the node array. */
+	void FixInvalidEdges();
+
 	int32 GraphChangeNotificationsDisableCounter = 0;
 	bool bDelayedChangeNotification = false;
 	EPCGChangeType DelayedChangeType = EPCGChangeType::None;
diff --git a/Engine/Plugins/Experimental/PCG/Source/PCGEditor/Private/PCGEditorGraph.cpp b/Engine/Plugins/Experimental/PCG/Source/PCGEditor/Private/PCGEditorGraph.cpp
--- a/Engine/Plugins/Experimental/PCG/Source/PCGEditor/Private/PCGEditorGraph.cpp
+++ b/Engine/Plugins/Experimental/PCG/Source/PCGEditor/Private/PCGEditorGraph.cpp
@@ -119,19 +119,19 @@ void UPCGEditorGraph::OnClose()
 void UPCGEditorGraph::CreateLinks(UPCGEditorGraphNodeBase* GraphNode, bool bCreateInbound, bool bCreateOutbound)
 {
 	check(GraphNode);
-	// Build graph node to pcg node map
-	TMap<UPCGNode*, UPCGEditorGraphNodeBase*> GraphNodeToPCGNodeMap;
 
+	// Build pcg node to pcg editor graph node map
+	TMap<UPCGNode*, UPCGEditorGraphNodeBase*> PCGNodeToPCGEditorNodeMap;
 	for (const TObjectPtr<UEdGraphNode>& EdGraphNode : Nodes)
 	{
 		if (UPCGEditorGraphNodeBase* SomeGraphNode = Cast<UPCGEditorGraphNodeBase>(EdGraphNode))
 		{
-			GraphNodeToPCGNodeMap.Add(SomeGraphNode->GetPCGNode(), SomeGraphNode);
+			PCGNodeToPCGEditorNodeMap.Add(SomeGraphNode->GetPCGNode(), SomeGraphNode);
 		}
 	}
 
 	// Forward the call
-	CreateLinks(GraphNode, bCreateInbound, bCreateOutbound, GraphNodeToPCGNodeMap);
+	CreateLinks(GraphNode, bCreateInbound, bCreateOutbound, PCGNodeToPCGEditorNodeMap);
 }
 
 void UPCGEditorGraph::UpdateGridSizeVisualization(UPCGComponent* InPCGComponentBeingInspected)
@@ -146,7 +146,7 @@ void UPCGEditorGraph::UpdateGridSizeVisualization(UPCGComponent* InPCGComponentB
 	}
 }
 
-void UPCGEditorGraph::CreateLinks(UPCGEditorGraphNodeBase* GraphNode, bool bCreateInbound, bool bCreateOutbound, const TMap<UPCGNode*, UPCGEditorGraphNodeBase*>& GraphNodeToPCGNodeMap)
+void UPCGEditorGraph::CreateLinks(UPCGEditorGraphNodeBase* GraphNode, bool bCreateInbound, bool bCreateOutbound, const TMap<UPCGNode*, UPCGEditorGraphNodeBase*>& InPCGNodeToPCGEditorNodeMap)
 {
 	check(GraphNode);
 	const UPCGNode* PCGNode = GraphNode->GetPCGNode();
@@ -154,15 +154,15 @@ void UPCGEditorGraph::CreateLinks(UPCGEditorGraphNodeBase* GraphNode, bool bCrea
 
 	if (bCreateInbound)
 	{
-		for (const UPCGPin* InputPin : PCGNode->GetInputPins())
+		for (UPCGPin* InputPin : PCGNode->GetInputPins())
 		{
 			UEdGraphPin* InPin = GraphNode->FindPin(InputPin->Properties.Label, EEdGraphPinDirection::EGPD_Input);
-
 			if (!InPin)
 			{
 				if (!Cast<UPCGEditorGraphNodeInput>(GraphNode))
 				{
 					UE_LOG(LogPCGEditor, Error, TEXT("Invalid InputPin for %s"), *InputPin->Properties.Label.ToString());
+					ensure(false);
 				}
 
 				continue;
@@ -170,24 +170,24 @@ void UPCGEditorGraph::CreateLinks(UPCGEditorGraphNodeBase* GraphNode, bool bCrea
 
 			for (const UPCGEdge* InboundEdge : InputPin->Edges)
 			{
-				if (!InboundEdge->IsValid())
+				if (!InboundEdge || !InboundEdge->IsValid())
 				{
 					UE_LOG(LogPCGEditor, Error, TEXT("Invalid inbound edge for %s"), *InputPin->Properties.Label.ToString());
+					ensure(false);
 					continue;
 				}
 
-				const UPCGNode* InboundNode = InboundEdge->InputPin->Node;
-				if (UPCGEditorGraphNodeBase* const* ConnectedGraphNode = GraphNodeToPCGNodeMap.Find(InboundNode))
+				const UPCGNode* InboundNode = InboundEdge->InputPin ? InboundEdge->InputPin->Node : nullptr;
+				UPCGEditorGraphNodeBase* const* ConnectedGraphNode = InboundNode ? InPCGNodeToPCGEditorNodeMap.Find(InboundNode) : nullptr;
+				UEdGraphPin* OutPin = ConnectedGraphNode ? (*ConnectedGraphNode)->FindPin(InboundEdge->InputPin->Properties.Label, EEdGraphPinDirection::EGPD_Output) : nullptr;
+				if (OutPin)
 				{
-					if (UEdGraphPin* OutPin = (*ConnectedGraphNode)->FindPin(InboundEdge->InputPin->Properties.Label, EEdGraphPinDirection::EGPD_Output))
-					{
-						OutPin->MakeLinkTo(InPin);
-					}
-					else
-					{
-						UE_LOG(LogPCGEditor, Error, TEXT("Could not create link to InputPin %s from Node %s"), *InputPin->Properties.Label.ToString(), *InboundNode->GetFName().ToString());
-						continue;
-					}
+					OutPin->MakeLinkTo(InPin);
+				}
+				else
+				{
+					UE_LOG(LogPCGEditor, Error, TEXT("Could not create link to InputPin %s from Node %s"), *InputPin->Properties.Label.ToString(), *InboundNode->GetFName().ToString());
+					ensure(false);
 				}
 			}
 		}
@@ -195,10 +195,9 @@ void UPCGEditorGraph::CreateLinks(UPCGEditorGraphNodeBase* GraphNode, bool bCrea
 
 	if (bCreateOutbound)
 	{
-		for (const UPCGPin* OutputPin : PCGNode->GetOutputPins())
+		for (UPCGPin* OutputPin : PCGNode->GetOutputPins())
 		{
 			UEdGraphPin* OutPin = GraphNode->FindPin(OutputPin->Properties.Label, EEdGraphPinDirection::EGPD_Output);
-
 			if (!OutPin)
 			{
 				if (!Cast<UPCGEditorGraphNodeOutput>(GraphNode))
@@ -211,24 +210,24 @@ void UPCGEditorGraph::CreateLinks(UPCGEditorGraphNodeBase* GraphNode, bool bCrea
 
 			for (const UPCGEdge* OutboundEdge : OutputPin->Edges)
 			{
-				if (!OutboundEdge->IsValid())
+				if (!OutboundEdge || !OutboundEdge->IsValid())
 				{
 					UE_LOG(LogPCGEditor, Error, TEXT("Invalid outbound edge for %s"), *OutputPin->Properties.Label.ToString());
+					ensure(false);
 					continue;
 				}
 
-				const UPCGNode* OutboundNode = OutboundEdge->OutputPin->Node;
-				if (UPCGEditorGraphNodeBase* const* ConnectedGraphNode = GraphNodeToPCGNodeMap.Find(OutboundNode))
+				const UPCGNode* OutboundNode = OutboundEdge->OutputPin ? OutboundEdge->OutputPin->Node : nullptr;
+				UPCGEditorGraphNodeBase* const* ConnectedGraphNode = OutboundNode ? InPCGNodeToPCGEditorNodeMap.Find(OutboundNode) : nullptr;
+				UEdGraphPin* InPin = ConnectedGraphNode ? (*ConnectedGraphNode)->FindPin(OutboundEdge->OutputPin->Properties.Label, EEdGraphPinDirection::EGPD_Input) : nullptr;
+				if (InPin)
 				{
-					if (UEdGraphPin* InPin = (*ConnectedGraphNode)->FindPin(OutboundEdge->OutputPin->Properties.Label, EEdGraphPinDirection::EGPD_Input))
-					{
-						OutPin->MakeLinkTo(InPin);
-					}
-					else
-					{
-						UE_LOG(LogPCGEditor, Error, TEXT("Could not create link from OutputPin %s to Node %s"), *OutputPin->Properties.Label.ToString(), *OutboundNode->GetFName().ToString());
-						continue;
-					}
+					OutPin->MakeLinkTo(InPin);
+				}
+				else
+				{
+					UE_LOG(LogPCGEditor, Error, TEXT("Could not create link from OutputPin %s to Node %s"), *OutputPin->Properties.Label.ToString(), *OutboundNode->GetFName().ToString());
+					ensure(false);
 				}
 			}
 		}
diff --git a/Engine/Plugins/Experimental/PCG/Source/PCGEditor/Private/PCGEditorGraph.h b/Engine/Plugins/Experimental/PCG/Source/PCGEditor/Private/PCGEditorGraph.h
--- a/Engine/Plugins/Experimental/PCG/Source/PCGEditor/Private/PCGEditorGraph.h
+++ b/Engine/Plugins/Experimental/PCG/Source/PCGEditor/Private/PCGEditorGraph.h
@@ -44,7 +44,7 @@ public:
 	void UpdateGridSizeVisualization(UPCGComponent* InPCGComponentBeingInspected);
 
 protected:
-	void CreateLinks(UPCGEditorGraphNodeBase* InGraphNode, bool bCreateInbound, bool bCreateOutbound, const TMap<UPCGNode*, UPCGEditorGraphNodeBase*>& InGraphNodeToPCGNodeMap);
+	void CreateLinks(UPCGEditorGraphNodeBase* InGraphNode, bool bCreateInbound, bool bCreateOutbound, const TMap<UPCGNode*, UPCGEditorGraphNodeBase*>& InPCGNodeToPCGEditorNodeMap);
 
 	void OnGraphUserParametersChanged(UPCGGraphInterface* InGraph, EPCGGraphParameterEvent ChangeType, FName ChangedPropertyName);
 
