[PCG] Adds a node that creates point data
--Added Basic UnitTest

diff --git a/Engine/Plugins/Experimental/PCG/Source/PCG/Private/Elements/PCGCreatePoints.cpp b/Engine/Plugins/Experimental/PCG/Source/PCG/Private/Elements/PCGCreatePoints.cpp
new file mode 100644
--- /dev/null
+++ b/Engine/Plugins/Experimental/PCG/Source/PCG/Private/Elements/PCGCreatePoints.cpp
@@ -0,0 +1,99 @@
+// Copyright Epic Games, Inc. All Rights Reserved.
+
+#include "Elements/PCGCreatePoints.h"
+
+#include "PCGComponent.h"
+#include "PCGContext.h"
+#include "Helpers/PCGAsync.h"
+#include "Helpers/PCGBlueprintHelpers.h"
+#include "Helpers/PCGSettingsHelpers.h"
+
+#define LOCTEXT_NAMESPACE "PCGCreatePointsElement"
+
+UPCGCreatePointsSettings::UPCGCreatePointsSettings()
+{
+	// Add one default point in the array
+	PointsToCreate.Add(FPCGPoint());
+}
+
+TArray<FPCGPinProperties> UPCGCreatePointsSettings::InputPinProperties() const
+{
+	return TArray<FPCGPinProperties>();
+}
+
+FPCGElementPtr UPCGCreatePointsSettings::CreateElement() const
+{
+	return MakeShared<FPCGCreatePointsElement>();
+}
+
+bool FPCGCreatePointsElement::ExecuteInternal(FPCGContext* Context) const
+{
+	TRACE_CPUPROFILER_EVENT_SCOPE(FPCGCreatePointsElement::Execute);
+	
+	check(Context);
+
+	const UPCGCreatePointsSettings* Settings = Context->GetInputSettings<UPCGCreatePointsSettings>();
+	check(Settings);
+
+	UPCGComponent* OriginalComponentContext = Context->SourceComponent.Get();
+	check(OriginalComponentContext);
+	UPCGSpatialData* Target = Cast<UPCGSpatialData>(OriginalComponentContext->GetActorPCGData());
+
+	const FTransform OriginalComponentTransform = OriginalComponentContext->GetOwner()->GetActorTransform();
+	const FTransform ComponentTransformScaleOne = FTransform(OriginalComponentTransform.Rotator(), OriginalComponentTransform.GetLocation(), FVector::One());
+
+	TArray<FPCGPoint> PointsToLoopOn = Settings->PointsToCreate;
+	
+	TArray<FPCGTaggedData>& Outputs = Context->OutputData.TaggedData;
+	FPCGTaggedData& Output = Outputs.Emplace_GetRef();
+
+	UPCGPointData* PtData = NewObject<UPCGPointData>();
+	check(PtData); 
+
+	TArray<FPCGPoint>& OutputPoints = PtData->GetMutablePoints();
+	Output.Data = PtData;
+	
+
+	if (!Settings->bLocal && !Settings->bCullPointsOutsideVolume)
+	{
+		for (auto& Points : PointsToLoopOn)
+		{
+			if (Points.Seed == 0)
+			{
+				// If the seed is the default value, generate a new seed based on the its transform
+				Points.Seed = UPCGBlueprintHelpers::ComputeSeedFromPosition(Points.Transform.GetLocation());
+			}
+		}
+
+		PtData->SetPoints(PointsToLoopOn);
+	}
+	else
+	{
+		FPCGAsync::AsyncPointProcessing(Context, PointsToLoopOn.Num(), OutputPoints, [&PointsToLoopOn, Settings, Target, &ComponentTransformScaleOne](int32 Index, FPCGPoint& OutPoint)
+		{
+			const FPCGPoint& InPoint = PointsToLoopOn[Index];
+			OutPoint = InPoint;
+
+			if (Settings->bLocal)
+			{
+				OutPoint.Transform *= ComponentTransformScaleOne;
+			}
+
+			OutPoint.Seed = UPCGBlueprintHelpers::ComputeSeedFromPosition(OutPoint.Transform.GetLocation());
+
+			// Discards all points that are outside the volume
+			return !Settings->bCullPointsOutsideVolume || (Target->GetDensityAtPosition(OutPoint.Transform.GetLocation()) > 0.0f);
+		});
+	}
+
+	return true;
+}
+
+bool FPCGCreatePointsElement::IsCacheable(const UPCGSettings* InSettings) const
+{
+	const UPCGCreatePointsSettings* Settings = Cast<const UPCGCreatePointsSettings>(InSettings);
+
+	return Settings && !Settings->bLocal;
+}
+
+#undef LOCTEXT_NAMESPACE
diff --git a/Engine/Plugins/Experimental/PCG/Source/PCG/Private/Tests/Elements/PCGCreatePointsTest.cpp b/Engine/Plugins/Experimental/PCG/Source/PCG/Private/Tests/Elements/PCGCreatePointsTest.cpp
new file mode 100644
--- /dev/null
+++ b/Engine/Plugins/Experimental/PCG/Source/PCG/Private/Tests/Elements/PCGCreatePointsTest.cpp
@@ -0,0 +1,49 @@
+// Copyright Epic Games, Inc. All Rights Reserved.
+
+#include "Tests/Determinism/PCGDeterminismTestsCommon.h"
+
+#include "PCGContext.h"
+#include "Data/PCGPointData.h"
+
+#include "Elements/PCGCreatePoints.h"
+
+IMPLEMENT_CUSTOM_SIMPLE_AUTOMATION_TEST(FPCGCreatePointsTest_Basic, FPCGTestBaseClass, "pcg.tests.CreatePoints.Basic", PCGTestsCommon::TestFlags)
+
+bool FPCGCreatePointsTest_Basic::RunTest(const FString& Parameters)
+{
+	PCGTestsCommon::FTestData TestData(PCGDeterminismTests::Defaults::Seed);
+	PCGDeterminismTests::GenerateSettings<UPCGCreatePointsSettings>(TestData);
+	UPCGCreatePointsSettings* Settings = CastChecked<UPCGCreatePointsSettings>(TestData.Settings);
+	
+	TObjectPtr<UPCGPointData> PointData = PCGTestsCommon::CreateRandomPointData(100, 42, false);
+	Settings->PointsToCreate = PointData->GetMutablePoints();
+
+	FPCGElementPtr TestElement = TestData.Settings->GetElement();
+
+	TUniquePtr<FPCGContext> Context = TestData.InitializeTestContext();
+
+	while (!TestElement->Execute(Context.Get())) {}
+
+	// test our point data
+	const TArray<FPCGTaggedData>& Outputs = Context->OutputData.TaggedData;
+
+	UTEST_EQUAL("Output count", Outputs.Num(), 1);
+
+	const UPCGPointData* OutPointData = Cast<UPCGPointData>(Outputs[0].Data);
+
+	UTEST_NOT_NULL("Output point data", OutPointData);
+
+	const TArray<FPCGPoint>& OutPoints = OutPointData->GetPoints();
+
+	UTEST_EQUAL("Output point count", OutPoints.Num(), 100);
+
+	for (int i = 0; i < OutPoints.Num(); ++i)
+	{
+		UTEST_EQUAL(FString::Format(TEXT("InArray[{0}].Seed is equal to OutArray[{0}].Seed"), { i }), OutPoints[i].Seed, Settings->PointsToCreate[i].Seed);
+		UTEST_EQUAL(FString::Format(TEXT("InArray[{0}].Density is equal to OutArray[{0}].Density"), { i }), OutPoints[i].Density, Settings->PointsToCreate[i].Density);
+		UTEST_EQUAL(FString::Format(TEXT("InArray[{0}].Transform is equal to OutArray[{0}].Transform"), { i }), OutPoints[i].Transform, Settings->PointsToCreate[i].Transform);
+	}
+
+	return true;
+}
+
diff --git a/Engine/Plugins/Experimental/PCG/Source/PCG/Public/Elements/PCGCreatePoints.h b/Engine/Plugins/Experimental/PCG/Source/PCG/Public/Elements/PCGCreatePoints.h
new file mode 100644
--- /dev/null
+++ b/Engine/Plugins/Experimental/PCG/Source/PCG/Public/Elements/PCGCreatePoints.h
@@ -0,0 +1,49 @@
+// Copyright Epic Games, Inc. All Rights Reserved.
+#pragma once
+
+#include "PCGSettings.h"
+
+#include "Data/PCGPointData.h"
+#include "PCGCreatePoints.generated.h"
+
+/**
+ * Creates point data from a provided list of points.
+ */
+UCLASS(BlueprintType, ClassGroup = (Procedural))
+class UPCGCreatePointsSettings : public UPCGSettings
+{
+	GENERATED_BODY()
+
+public:
+	UPCGCreatePointsSettings();
+	//~Begin UPCGSettings interface
+#if WITH_EDITOR
+	virtual FName GetDefaultNodeName() const override { return FName(TEXT("CreatePoints")); }
+	virtual FText GetDefaultNodeTitle() const override { return NSLOCTEXT("PCGCreatePointsElement", "NodeTitle", "Create Points"); }
+	virtual FText GetNodeTooltipText() const override { return NSLOCTEXT("PCGCreatePointsElement", "NodeTooltip", "Creates point data from a provided list of points."); }
+	virtual EPCGSettingsType GetType() const override { return EPCGSettingsType::Spatial; }
+#endif
+
+protected:
+	virtual TArray<FPCGPinProperties> InputPinProperties() const override;
+	virtual TArray<FPCGPinProperties> OutputPinProperties() const override { return Super::DefaultPointOutputPinProperties(); }
+	virtual FPCGElementPtr CreateElement() const override;
+	//~End UPCGSettings interface
+
+public:
+	UPROPERTY(BlueprintReadWrite, EditAnywhere, Category = Settings, meta = (PCG_Overridable))
+	TArray<FPCGPoint> PointsToCreate;
+
+	UPROPERTY(BlueprintReadWrite, EditAnywhere, Category = Settings, meta = (PCG_Overridable))
+	bool bLocal = false;
+
+	UPROPERTY(BlueprintReadWrite, EditAnywhere, Category = Settings, meta = (PCG_Overridable))
+	bool bCullPointsOutsideVolume = false;
+};
+
+class FPCGCreatePointsElement : public FSimplePCGElement
+{
+protected:
+	virtual bool ExecuteInternal(FPCGContext* Context) const override;
+	virtual bool IsCacheable(const UPCGSettings* InSettings) const override;
+};
\ No newline at end of file
