Tentative fix for Non-unity Build

diff --git a/Engine/Plugins/Experimental/PCG/Source/PCG/Private/Elements/PCGCreatePoints.cpp b/Engine/Plugins/Experimental/PCG/Source/PCG/Private/Elements/PCGCreatePoints.cpp
--- a/Engine/Plugins/Experimental/PCG/Source/PCG/Private/Elements/PCGCreatePoints.cpp
+++ b/Engine/Plugins/Experimental/PCG/Source/PCG/Private/Elements/PCGCreatePoints.cpp
@@ -8,6 +8,8 @@
 #include "Helpers/PCGBlueprintHelpers.h"
 #include "Helpers/PCGSettingsHelpers.h"
 
+#include "GameFramework/Actor.h"
+
 #define LOCTEXT_NAMESPACE "PCGCreatePointsElement"
 
 UPCGCreatePointsSettings::UPCGCreatePointsSettings()
