[PCG] Fix and add unit tests for Get Actor Property

* Properties need to be visible to be extracted
* Extend validation method to support multiple values and attributes
* Add test for struct/object extract and arrays

#rnx

diff --git a/Engine/Plugins/Experimental/PCG/Source/PCG/Private/Tests/Elements/PCGPropertyToParamDataElementTest.cpp b/Engine/Plugins/Experimental/PCG/Source/PCG/Private/Tests/Elements/PCGPropertyToParamDataElementTest.cpp
--- a/Engine/Plugins/Experimental/PCG/Source/PCG/Private/Tests/Elements/PCGPropertyToParamDataElementTest.cpp
+++ b/Engine/Plugins/Experimental/PCG/Source/PCG/Private/Tests/Elements/PCGPropertyToParamDataElementTest.cpp
@@ -18,14 +18,24 @@
 IMPLEMENT_CUSTOM_SIMPLE_AUTOMATION_TEST(FPCGPropertyToParamDataPropertyTypeTest, FPCGTestBaseClass, "pcg.tests.PropertyToParamData.PropertyType", PCGTestsCommon::TestFlags)
 IMPLEMENT_CUSTOM_SIMPLE_AUTOMATION_TEST(FPCGPropertyToParamDataActorFindTest, FPCGTestBaseClass, "pcg.tests.PropertyToParamData.ActorFind", PCGTestsCommon::TestFlags)
 
-// bShouldFail is a constexpr bool used to cut the "valid" branch
-// Because invalid types in expected failed tests could not compile.
-template <typename AttributeType, bool bShouldFail>
-bool VerifyAttributeValue(FPCGTestBaseClass* TestInstance, PCGTestsCommon::FTestData& TestData, FName PropertyName, AttributeType&& ExpectedValue, FString ExtraTestWhat)
-{
-	// Use universal reference to either pass an rvalue or lvalue. Use this type below to know the raw type behind.
-	using RawAttributeType = std::remove_const_t<std::remove_reference_t<AttributeType>>;
 
+/**
+* Generic method to test the extraction and then verify that the values matches the one passed as input.
+* Use variadic template for expected values, allowing to test the extraction of structs/objects that doesn't have the same type for all their members.
+* bShouldFail is used as a constexpr bool to cut the branch where we try to compare the values, if it is expected to fail.
+* Be careful, if you extract structures, all int properties are int64 (even if they were like u8) and all float will be double. Take that into account in the ExpectedValues types.
+* 
+* @param TestInstance - Instance of the test, used for our test conditions
+* @param TestData - Test data used to run the element
+* @param PropertyName - Property name to extract
+* @param AttributeNames - List of all the attributes to test. AttributeNames count MUST divide the count of expected values. If there is more than one attribute name, it will be marked as extracting structures
+* @param ExtraTestWhat - Simple string to append to all test strings, to know which test we are currently running.
+* @param ExpectedValues - List of all values to test. If there is more than one, it can be because it is an extraction, or testing an array. If it is an array of struct, and attribute names is {'X', 'Y'}, it must follow this order: A[0].X, A[0].Y, A[1].X, A[1].Y, ...
+* @return if the test succeeded or not.
+*/
+template <bool bShouldFail, typename ...AttributeTypes>
+bool VerifyAttributeValue(FPCGTestBaseClass* TestInstance, PCGTestsCommon::FTestData& TestData, FName PropertyName, TArray<FName> AttributeNames, FString ExtraTestWhat, AttributeTypes&& ...ExpectedValues)
+{
 	// Use TestData settings to set the property name
 	UPCGPropertyToParamDataSettings* Settings = Cast<UPCGPropertyToParamDataSettings>(TestData.Settings);
 	if (!TestInstance->TestNotNull(TEXT("CastToUPCGPropertyToParamDataSettings"), Settings))
@@ -33,8 +43,16 @@ bool VerifyAttributeValue(FPCGTestBaseClass* TestInstance, PCGTestsCommon::FTest
 		return false;
 	}
 
+	// Make sure this call is not ill-formed
+	check(!AttributeNames.IsEmpty());
+	if (!ensureMsgf(sizeof...(ExpectedValues) % AttributeNames.Num() == 0, TEXT("%s"), TEXT("Call is ill-formed, there is a mismatch between the number of attribute names and the number of expected values.")))
+	{
+		return false;
+	}
+
 	Settings->PropertyName = PropertyName;
-	Settings->OutputAttributeName = Settings->PropertyName;
+	Settings->OutputAttributeName = AttributeNames[0];
+	Settings->bExtractObjectAndStruct = AttributeNames.Num() > 1;
 
 	// Add 2 nodes, PropertyToParamDataNode and a Trivial node (just there for the connection)
 	UPCGNode* TestNode = TestData.TestPCGComponent->GetGraph()->AddNode(Settings);
@@ -49,13 +67,19 @@ bool VerifyAttributeValue(FPCGTestBaseClass* TestInstance, PCGTestsCommon::FTest
 
 	// Run the element, can throw an error.
 	while (!Element->Execute(Context.Get()))
-	{};
+	{
+	};
 
 	auto FormatWithPropertyName = [&ExtraTestWhat, PropertyName](const FString& InText) -> FString
 	{
 		return FString::Format(TEXT("{0}_{1}_{2}"), { ExtraTestWhat, PropertyName.ToString(), InText });
 	};
 
+	auto FormatWithPropertyNameAndSubNames = [&ExtraTestWhat, PropertyName](const FString& InText, const FName InSubName, const int InItemKey) -> FString
+	{
+		return FString::Format(TEXT("{0}_{1}_SubName:{2}_ItemKey:{3} - {4}"), { ExtraTestWhat, PropertyName.ToString(), InSubName.ToString(), FString::FromInt(InItemKey), InText });
+	};
+
 	bool bSuccess = true;
 
 	if constexpr (bShouldFail)
@@ -65,21 +89,51 @@ bool VerifyAttributeValue(FPCGTestBaseClass* TestInstance, PCGTestsCommon::FTest
 	}
 	else
 	{
-		// If it should succeed, we check the number of output and verify that the type and value matches the expected one.
+		// If it should succeed, we check the number of output.
 		bSuccess = TestInstance->TestEqual(FormatWithPropertyName("NumOutput"), Context->OutputData.GetAllParams().Num(), 1);
 
 		if (bSuccess)
 		{
 			const UPCGParamData* ParamData = Cast<UPCGParamData>(Context->OutputData.GetAllParams()[0].Data);
-			const FPCGMetadataAttributeBase* Attribute = ParamData->ConstMetadata()->GetConstAttribute(Settings->PropertyName);
+			int32 Index = 0;
+			int32 ItemKey = 0;
 
-			bSuccess = TestInstance->TestEqual(FormatWithPropertyName("AttributeType"), Attribute->GetTypeId(), PCG::Private::MetadataTypes<RawAttributeType>::Id);
-
-			if (bSuccess)
+			// Fold expression that will iterate on all Expected Values, to check their types and their values.
+			// We keep track of the current Index, to know which attribute to check, and an ItemKey to know where to look at in the value array of the attribute.
+			// cf. comment at the top to know in which order ExpectedValues should be.
+			([&Index, &AttributeNames, ParamData, &FormatWithPropertyNameAndSubNames, TestInstance, &bSuccess, &ItemKey, &ExpectedValues]()
 			{
-				RawAttributeType Value = static_cast<const FPCGMetadataAttribute<RawAttributeType>*>(Attribute)->GetValueFromItemKey(0);
-				TestInstance->TestEqual(FormatWithPropertyName("AttributeValue"), Value, ExpectedValue);
-			}
+				if (bSuccess)
+				{
+					// In a fold expression, ExpectedValues will be a single value. We know its underlying type by decaying the type of the value (using decltype)
+					using AttributeType = typename std::decay_t<decltype(ExpectedValues)>;
+
+					// We use Index to track the current attribute to check
+					const FName AttributeName = AttributeNames[Index];
+					const FPCGMetadataAttributeBase* Attribute = ParamData->ConstMetadata()->GetConstAttribute(AttributeName);
+
+					bSuccess = TestInstance->TestNotNull(FormatWithPropertyNameAndSubNames("Attribute should not be null", AttributeName, ItemKey), Attribute);
+
+					if (bSuccess)
+					{
+						bSuccess = TestInstance->TestEqual(FormatWithPropertyNameAndSubNames("Attribute type is matching", AttributeName, ItemKey), Attribute->GetTypeId(), PCG::Private::MetadataTypes<AttributeType>::Id);
+					}
+
+					if (bSuccess)
+					{
+						AttributeType Value = static_cast<const FPCGMetadataAttribute<AttributeType>*>(Attribute)->GetValueFromItemKey(ItemKey);
+						bSuccess = TestInstance->TestEqual(FormatWithPropertyNameAndSubNames("Attribute value is the same as expected", AttributeName, ItemKey), Value, ExpectedValues);
+					}
+
+					// At each iteration on the ExpectedValues, we increment the index and if we reached the end of the Attribute Names we reset it to 0
+					// and we increment ItemKey, to check the next value in the attribute value array.
+					if (++Index >= AttributeNames.Num())
+					{
+						Index = 0;
+						ItemKey++;
+					}
+				}
+			}(), ...);
 		}
 	}
 
@@ -91,16 +145,26 @@ bool VerifyAttributeValue(FPCGTestBaseClass* TestInstance, PCGTestsCommon::FTest
 }
 
 // Use aliases for convenience
+
+// Single valid attribute, Property Name will be attribute name
 template <typename AttributeType>
 bool VerifyAttributeValueValid(FPCGTestBaseClass* TestInstance, PCGTestsCommon::FTestData& TestData, FName PropertyName, AttributeType&& ExpectedValue, FString ExtraTestWhat)
 {
-	return VerifyAttributeValue<AttributeType, false>(TestInstance, TestData, PropertyName, std::forward<AttributeType&&>(ExpectedValue), ExtraTestWhat);
+	return VerifyAttributeValue</*bShouldFail=*/false>(TestInstance, TestData, PropertyName, { PropertyName }, ExtraTestWhat, std::forward<AttributeType>(ExpectedValue));
 }
 
+// Single invalid attribute, Property Name will be attribute name
 template <typename AttributeType>
 bool VerifyAttributeValueInvalid(FPCGTestBaseClass* TestInstance, PCGTestsCommon::FTestData& TestData, FName PropertyName, AttributeType&& ExpectedValue, FString ExtraTestWhat)
 {
-	return VerifyAttributeValue<AttributeType, true>(TestInstance, TestData, PropertyName, std::forward<AttributeType&&>(ExpectedValue), ExtraTestWhat);
+	return VerifyAttributeValue</*bShouldFail=*/true>(TestInstance, TestData, PropertyName, { PropertyName }, ExtraTestWhat, std::forward<AttributeType>(ExpectedValue));
+}
+
+// Valid Single or multi attributes with single or multiple values.
+template <typename ...AttributeTypes>
+bool VerifyAttributeValuesValid(FPCGTestBaseClass* TestInstance, PCGTestsCommon::FTestData& TestData, FName PropertyName, const TArray<FName>& AttributeNames, FString ExtraTestWhat, AttributeTypes&& ...ExpectedValues)
+{
+	return VerifyAttributeValue</*bShouldFail=*/false>(TestInstance, TestData, PropertyName, AttributeNames, ExtraTestWhat, std::forward<AttributeTypes>(ExpectedValues)...);
 }
 
 bool FPCGPropertyToParamDataPropertyTypeTest::RunTest(const FString& Parameters)
@@ -122,18 +186,28 @@ bool FPCGPropertyToParamDataPropertyTypeTest::RunTest(const FString& Parameters)
 	const FString StringValue = TEXT("HelloWorld");
 
 	const FVector VectorValue{ 1.0, 2.0, 3.0 };
+	const FVector SecondVectorValue{ 4.0, 5.0, 6.0 };
 	const FVector4 Vector4Value{ 1.0, 2.0, 3.0, 4.0 };
 	const FRotator RotatorValue{ 45.0, 45.0, 45.0 };
 	const FQuat QuatValue = RotatorValue.Quaternion();
 	const FTransform TransformValue{ QuatValue, VectorValue, VectorValue };
 
-	APCGVolume* PCGVolume = NewObject<APCGVolume>();
+	UPCGDummyGetPropertyTest* ObjectValue = NewObject<UPCGDummyGetPropertyTest>();
+	ObjectValue->SetFlags(RF_Transient);
+	ObjectValue->Int64Property = 42ll;
+	ObjectValue->DoubleProperty = 1.0;
+
+	UPCGDummyGetPropertyTest* SecondObjectValue = NewObject<UPCGDummyGetPropertyTest>();
+	SecondObjectValue->SetFlags(RF_Transient);
+	SecondObjectValue->Int64Property = 43ll;
+	SecondObjectValue->DoubleProperty = 2.0;
 
-	const FSoftObjectPath SoftObjectPathValue{ PCGVolume };
-	const FSoftClassPath SoftClassPathValue{ APCGVolume::StaticClass() };
+	const FSoftObjectPath SoftObjectPathValue{ ObjectValue };
+	const FSoftClassPath SoftClassPathValue{ UPCGDummyGetPropertyTest::StaticClass() };
 
 	const FVector2D Vector2Value = { 1.0, 2.0 };
 	const FColor ColorValue = FColor::Cyan;
+	const FColor SecondColorValue = FColor::Yellow;
 
 	APCGUnitTestDummyActor* Actor = Cast<APCGUnitTestDummyActor>(TestData.TestActor);
 	Actor->IntProperty = 42;
@@ -151,10 +225,14 @@ bool FPCGPropertyToParamDataPropertyTypeTest::RunTest(const FString& Parameters)
 	Actor->TransformProperty = TransformValue;
 	Actor->SoftObjectPathProperty = SoftObjectPathValue;
 	Actor->SoftClassPathProperty = SoftClassPathValue;
-	Actor->ClassProperty = APCGVolume::StaticClass();
-	Actor->ObjectProperty = PCGVolume;
+	Actor->ClassProperty = UPCGDummyGetPropertyTest::StaticClass();
+	Actor->ObjectProperty = ObjectValue;
 	Actor->Vector2Property = Vector2Value;
 	Actor->ColorProperty = ColorValue;
+	Actor->ArrayOfIntsProperty = { 42, 43, 44 };
+	Actor->ArrayOfVectorsProperty = { VectorValue, SecondVectorValue };
+	Actor->ArrayOfStructsProperty = { ColorValue, SecondColorValue };
+	Actor->ArrayOfObjectsProperty = { ObjectValue, SecondObjectValue };
 
 	// Basic properties
 	bSuccess &= VerifyAttributeValueValid(this, TestData, GET_MEMBER_NAME_CHECKED(APCGUnitTestDummyActor, IntProperty), 42ll, ExtraTestWhat);
@@ -170,32 +248,54 @@ bool FPCGPropertyToParamDataPropertyTypeTest::RunTest(const FString& Parameters)
 	// Enum Property
 	bSuccess &= VerifyAttributeValueValid(this, TestData, GET_MEMBER_NAME_CHECKED(APCGUnitTestDummyActor, EnumProperty), (int64)EPCGUnitTestDummyEnum::Three, ExtraTestWhat);
 
-	// Struct Properties
+	// Supported struct Properties
 	bSuccess &= VerifyAttributeValueValid(this, TestData, GET_MEMBER_NAME_CHECKED(APCGUnitTestDummyActor, Vector2Property), Vector2Value, ExtraTestWhat);
 	bSuccess &= VerifyAttributeValueValid(this, TestData, GET_MEMBER_NAME_CHECKED(APCGUnitTestDummyActor, VectorProperty), VectorValue, ExtraTestWhat);
 	bSuccess &= VerifyAttributeValueValid(this, TestData, GET_MEMBER_NAME_CHECKED(APCGUnitTestDummyActor, Vector4Property), Vector4Value, ExtraTestWhat);
 	bSuccess &= VerifyAttributeValueValid(this, TestData, GET_MEMBER_NAME_CHECKED(APCGUnitTestDummyActor, RotatorProperty), RotatorValue, ExtraTestWhat);
 	bSuccess &= VerifyAttributeValueValid(this, TestData, GET_MEMBER_NAME_CHECKED(APCGUnitTestDummyActor, QuatProperty), QuatValue, ExtraTestWhat);
 	bSuccess &= VerifyAttributeValueValid(this, TestData, GET_MEMBER_NAME_CHECKED(APCGUnitTestDummyActor, TransformProperty), TransformValue, ExtraTestWhat);
-	bSuccess &= VerifyAttributeValueValid(this, TestData, GET_MEMBER_NAME_CHECKED(APCGUnitTestDummyActor, SoftObjectPathProperty), SoftObjectPathValue.ToString(), ExtraTestWhat);
-	bSuccess &= VerifyAttributeValueValid(this, TestData, GET_MEMBER_NAME_CHECKED(APCGUnitTestDummyActor, SoftClassPathProperty), SoftClassPathValue.ToString(), ExtraTestWhat);
+	bSuccess &= VerifyAttributeValueValid(this, TestData, GET_MEMBER_NAME_CHECKED(APCGUnitTestDummyActor, SoftObjectPathProperty), SoftObjectPathValue, ExtraTestWhat);
+	bSuccess &= VerifyAttributeValueValid(this, TestData, GET_MEMBER_NAME_CHECKED(APCGUnitTestDummyActor, SoftClassPathProperty), SoftClassPathValue, ExtraTestWhat);
 
-	// Objects properties
-	bSuccess &= VerifyAttributeValueValid(this, TestData, GET_MEMBER_NAME_CHECKED(APCGUnitTestDummyActor, ClassProperty), APCGVolume::StaticClass()->GetPathName(), ExtraTestWhat);
-	bSuccess &= VerifyAttributeValueValid(this, TestData, GET_MEMBER_NAME_CHECKED(APCGUnitTestDummyActor, ObjectProperty), PCGVolume->GetPathName(), ExtraTestWhat);
+	// Objects properties as String
+	bSuccess &= VerifyAttributeValueValid(this, TestData, GET_MEMBER_NAME_CHECKED(APCGUnitTestDummyActor, ClassProperty), UPCGDummyGetPropertyTest::StaticClass()->GetPathName(), ExtraTestWhat);
+	bSuccess &= VerifyAttributeValueValid(this, TestData, GET_MEMBER_NAME_CHECKED(APCGUnitTestDummyActor, ObjectProperty), ObjectValue->GetPathName(), ExtraTestWhat);
 
-	// Unsupported properties
+	// Struct Property Extracted - Colors
+	// Extracting int will always yield a int64 and extracting floats with yield doubles. Here color is u8, so cast all of them to int64
+	const TArray<FName> ColorPropertyNames = { TEXT("R"), TEXT("G"), TEXT("B"), TEXT("A") };
+	bSuccess &= VerifyAttributeValuesValid(this, TestData, GET_MEMBER_NAME_CHECKED(APCGUnitTestDummyActor, ColorProperty), ColorPropertyNames, ExtraTestWhat, (int64)ColorValue.R, (int64)ColorValue.G, (int64)ColorValue.B, (int64)ColorValue.A);
+
+	// Object Property Extracted - UPCGDummyGetPropertyTest
+	const TArray<FName> ObjectPropertyNames = { GET_MEMBER_NAME_CHECKED(UPCGDummyGetPropertyTest, Int64Property), GET_MEMBER_NAME_CHECKED(UPCGDummyGetPropertyTest, DoubleProperty) };
+	bSuccess &= VerifyAttributeValuesValid(this, TestData, GET_MEMBER_NAME_CHECKED(APCGUnitTestDummyActor, ObjectProperty), ObjectPropertyNames, ExtraTestWhat, 42ll, 1.0);
+
+	// Arrays of supported properties
+	bSuccess &= VerifyAttributeValuesValid(this, TestData, GET_MEMBER_NAME_CHECKED(APCGUnitTestDummyActor, ArrayOfIntsProperty), { GET_MEMBER_NAME_CHECKED(APCGUnitTestDummyActor, ArrayOfIntsProperty) }, ExtraTestWhat, 42ll, 43ll, 44ll);
+	bSuccess &= VerifyAttributeValuesValid(this, TestData, GET_MEMBER_NAME_CHECKED(APCGUnitTestDummyActor, ArrayOfVectorsProperty), { GET_MEMBER_NAME_CHECKED(APCGUnitTestDummyActor, ArrayOfVectorsProperty) }, ExtraTestWhat, VectorValue, SecondVectorValue);
+	bSuccess &= VerifyAttributeValuesValid(this, TestData, GET_MEMBER_NAME_CHECKED(APCGUnitTestDummyActor, ArrayOfObjectsProperty), { GET_MEMBER_NAME_CHECKED(APCGUnitTestDummyActor, ArrayOfObjectsProperty) }, ExtraTestWhat, ObjectValue->GetPathName(), SecondObjectValue->GetPathName());
+
+	// Arrays of extracted properties
+	bSuccess &= VerifyAttributeValuesValid(this, TestData, GET_MEMBER_NAME_CHECKED(APCGUnitTestDummyActor, ArrayOfStructsProperty), ColorPropertyNames, ExtraTestWhat, 
+		(int64)ColorValue.R, (int64)ColorValue.G, (int64)ColorValue.B, (int64)ColorValue.A, (int64)SecondColorValue.R, (int64)SecondColorValue.G, (int64)SecondColorValue.B, (int64)SecondColorValue.A);
+	bSuccess &= VerifyAttributeValuesValid(this, TestData, GET_MEMBER_NAME_CHECKED(APCGUnitTestDummyActor, ArrayOfObjectsProperty), ObjectPropertyNames, ExtraTestWhat, ObjectValue->Int64Property, ObjectValue->DoubleProperty, SecondObjectValue->Int64Property, SecondObjectValue->DoubleProperty);
+
+	// Unsupported struct properties
 	AddExpectedError(TEXT("Error while creating an attribute for property"), EAutomationExpectedErrorFlags::Contains, 1);
 	bSuccess &= VerifyAttributeValueInvalid(this, TestData, GET_MEMBER_NAME_CHECKED(APCGUnitTestDummyActor, ColorProperty), ColorValue, ExtraTestWhat);
 
 	// Unknown property
-	AddExpectedError(TEXT("Property doesn't exist in the found actor."), EAutomationExpectedErrorFlags::Contains, 1);
+	AddExpectedError(TEXT("Property 'DummyMissingProperty' does not exist in the found actor"), EAutomationExpectedErrorFlags::Contains, 1);
 	bSuccess &= VerifyAttributeValueInvalid(this, TestData, TEXT("DummyMissingProperty"), 42, ExtraTestWhat);
 
 	// Missing property
 	AddExpectedError(TEXT("Some parameters are missing, abort."), EAutomationExpectedErrorFlags::Contains, 1);
 	bSuccess &= VerifyAttributeValueInvalid(this, TestData, NAME_None, 42, ExtraTestWhat);
 
+	ObjectValue->MarkAsGarbage();
+	SecondObjectValue->MarkAsGarbage();
+
 	return bSuccess;
 }
 
@@ -287,7 +387,7 @@ bool FPCGPropertyToParamDataActorFindTest::RunTest(const FString& Parameters)
 		
 		PCGTestsCommon::FTestData TestData(Seed, Settings, AActor::StaticClass());
 		TestData.TestActor->Tags.Add(Tag);
-		UPCGUnitTestDummyComponent* Component = Cast< UPCGUnitTestDummyComponent>(TestData.TestActor->AddComponentByClass(Settings->ComponentClass, false, FTransform::Identity, false));
+		UPCGUnitTestDummyComponent* Component = Cast<UPCGUnitTestDummyComponent>(TestData.TestActor->AddComponentByClass(Settings->ComponentClass, false, FTransform::Identity, false));
 		Component->IntProperty = 42;
 
 		bSuccess &= VerifyAttributeValueValid(this, TestData, GET_MEMBER_NAME_CHECKED(UPCGUnitTestDummyComponent, IntProperty), 42ll, "PropertyToParamDataActorFindTest_Self_Tag_Component");
diff --git a/Engine/Plugins/Experimental/PCG/Source/PCG/Private/Tests/Elements/PCGPropertyToParamDataElementTest.h b/Engine/Plugins/Experimental/PCG/Source/PCG/Private/Tests/Elements/PCGPropertyToParamDataElementTest.h
--- a/Engine/Plugins/Experimental/PCG/Source/PCG/Private/Tests/Elements/PCGPropertyToParamDataElementTest.h
+++ b/Engine/Plugins/Experimental/PCG/Source/PCG/Private/Tests/Elements/PCGPropertyToParamDataElementTest.h
@@ -14,68 +14,93 @@ enum class EPCGUnitTestDummyEnum : int64
 	Three
 };
 
+UCLASS(MinimalAPI, NotBlueprintable, NotPlaceable, NotBlueprintType, Transient, HideDropdown, meta = (Hidden))
+class UPCGDummyGetPropertyTest : public UObject
+{
+	GENERATED_BODY()
+
+public:
+	UPROPERTY(BlueprintReadOnly, Category = "")
+	int64 Int64Property = 0;
+
+	UPROPERTY(BlueprintReadOnly, Category = "")
+	double DoubleProperty = 0.0;
+};
+
 UCLASS(MinimalAPI, NotBlueprintable, NotPlaceable, NotBlueprintType, Transient, HideDropdown, meta = (Hidden))
 class APCGUnitTestDummyActor : public AActor
 {
 	GENERATED_BODY()
 
 public:
-	UPROPERTY()
+	UPROPERTY(BlueprintReadOnly, Category="")
 	int IntProperty = 0;
 
-	UPROPERTY()
+	UPROPERTY(BlueprintReadOnly, Category = "")
 	float FloatProperty = 0.0f;
 
-	UPROPERTY()
+	UPROPERTY(BlueprintReadOnly, Category = "")
 	int64 Int64Property = 0;
 
-	UPROPERTY()
+	UPROPERTY(BlueprintReadOnly, Category = "")
 	double DoubleProperty = 0.0;
 
-	UPROPERTY()
+	UPROPERTY(BlueprintReadOnly, Category = "")
 	bool BoolProperty = false;
 
-	UPROPERTY()
+	UPROPERTY(BlueprintReadOnly, Category = "")
 	FName NameProperty = NAME_None;
 
-	UPROPERTY()
+	UPROPERTY(BlueprintReadOnly, Category = "")
 	FString StringProperty = "";
 
-	UPROPERTY()
+	UPROPERTY(BlueprintReadOnly, Category = "")
 	EPCGUnitTestDummyEnum EnumProperty = EPCGUnitTestDummyEnum::One;
 
-	UPROPERTY()
+	UPROPERTY(BlueprintReadOnly, Category = "")
 	FVector VectorProperty;
 
-	UPROPERTY()
+	UPROPERTY(BlueprintReadOnly, Category = "")
 	FVector4 Vector4Property;
 
-	UPROPERTY()
+	UPROPERTY(BlueprintReadOnly, Category = "")
 	FTransform TransformProperty;
 
-	UPROPERTY()
+	UPROPERTY(BlueprintReadOnly, Category = "")
 	FRotator RotatorProperty;
 
-	UPROPERTY()
+	UPROPERTY(BlueprintReadOnly, Category = "")
 	FQuat QuatProperty;
 
-	UPROPERTY()
+	UPROPERTY(BlueprintReadOnly, Category = "")
 	FSoftObjectPath SoftObjectPathProperty;
 
-	UPROPERTY()
+	UPROPERTY(BlueprintReadOnly, Category = "")
 	FSoftClassPath SoftClassPathProperty;
 
-	UPROPERTY()
+	UPROPERTY(BlueprintReadOnly, Category = "")
 	TSubclassOf<AActor> ClassProperty = AActor::StaticClass();
 
-	UPROPERTY()
-	TObjectPtr<UObject> ObjectProperty = nullptr;
+	UPROPERTY(BlueprintReadOnly, Category = "")
+	TObjectPtr<UPCGDummyGetPropertyTest> ObjectProperty = nullptr;
 
-	UPROPERTY()
+	UPROPERTY(BlueprintReadOnly, Category = "")
 	FVector2D Vector2Property;
 
-	UPROPERTY()
+	UPROPERTY(BlueprintReadOnly, Category = "")
 	FColor ColorProperty;
+
+	UPROPERTY(BlueprintReadOnly, Category = "")
+	TArray<int32> ArrayOfIntsProperty;
+
+	UPROPERTY(BlueprintReadOnly, Category = "")
+	TArray<FVector> ArrayOfVectorsProperty;
+
+	UPROPERTY(BlueprintReadOnly, Category = "")
+	TArray<FColor> ArrayOfStructsProperty;
+
+	UPROPERTY(BlueprintReadOnly, Category = "")
+	TArray<TObjectPtr<UPCGDummyGetPropertyTest>> ArrayOfObjectsProperty;
 };
 
 UCLASS(MinimalAPI, NotBlueprintable, NotPlaceable, NotBlueprintType, Transient, HideDropdown, meta = (Hidden))
@@ -84,6 +109,6 @@ class UPCGUnitTestDummyComponent : public UActorComponent
 	GENERATED_BODY()
 
 public:
-	UPROPERTY()
+	UPROPERTY(BlueprintReadOnly, Category = "")
 	int IntProperty = 0;
 };
