[PCG] Make Metadata nodes to passthrough the input if the other operands are empty (either no data or point data with no points)

#rnx

diff --git a/Engine/Plugins/Experimental/PCG/Source/PCG/Private/Elements/Metadata/PCGMetadataOpElementBase.cpp b/Engine/Plugins/Experimental/PCG/Source/PCG/Private/Elements/Metadata/PCGMetadataOpElementBase.cpp
--- a/Engine/Plugins/Experimental/PCG/Source/PCG/Private/Elements/Metadata/PCGMetadataOpElementBase.cpp
+++ b/Engine/Plugins/Experimental/PCG/Source/PCG/Private/Elements/Metadata/PCGMetadataOpElementBase.cpp
@@ -255,6 +255,29 @@ TArray<FPCGPinProperties> UPCGMetadataSettingsBase::OutputPinProperties() const
 	return PinProperties;
 }
 
+void FPCGMetadataElementBase::PassthroughInputs(FPCGContext* InContext, const UPCGMetadataSettingsBase* InSettings) const
+{
+	check(InContext && InSettings);
+
+	const uint32 NumberOfOutputs = InSettings->GetOutputPinNum();
+	uint32 InputPinToForward = InSettings->GetInputPinToForward();
+
+	TArray<FPCGTaggedData> InputsToForward = InContext->InputData.GetInputsByPin(InSettings->GetInputPinLabel(InputPinToForward));
+	TArray<FPCGTaggedData>& Outputs = InContext->OutputData.TaggedData;
+
+	// Empty and reserve memory at the same time.
+	Outputs.Empty(InputsToForward.Num() * NumberOfOutputs);
+
+	for (uint32 i = 0; i < NumberOfOutputs; ++i)
+	{
+		const FName OutputPin = InSettings->GetOutputPinLabel(i);
+		for (const FPCGTaggedData& Input : InputsToForward)
+		{
+			Outputs.Emplace_GetRef(Input).Pin = OutputPin;
+		}
+	}
+}
+
 bool FPCGMetadataElementBase::ExecuteInternal(FPCGContext* Context) const
 {
 	TRACE_CPUPROFILER_EVENT_SCOPE(FPCGMetadataElementBase::Execute);
@@ -302,8 +325,9 @@ bool FPCGMetadataElementBase::ExecuteInternal(FPCGContext* Context) const
 
 		if (InputData.IsEmpty())
 		{
-			// Visually warn the user, since this is causing execution to be aborted
+			// If we have no data, just passthrough the input and return
 			PCGE_LOG(Verbose, LogOnly, FText::Format(LOCTEXT("MissingInputDataForPin", "No data provided on pin {0}"), FText::FromName(PinLabel)));
+			PassthroughInputs(Context, Settings);
 			return true;
 		}
 		else if (InputData.Num() > 1)
@@ -320,7 +344,9 @@ bool FPCGMetadataElementBase::ExecuteInternal(FPCGContext* Context) const
 		{
 			if (PointInput->GetPoints().IsEmpty())
 			{
+				// If we have no points, just passthrough the input and return
 				PCGE_LOG(Verbose, LogOnly, FText::Format(LOCTEXT("NoPointsForPin", "No points in point data provided on pin {0}"), FText::FromName(PinLabel)));
+				PassthroughInputs(Context, Settings);
 				return true;
 			}
 		}
@@ -392,7 +418,7 @@ bool FPCGMetadataElementBase::ExecuteInternal(FPCGContext* Context) const
 				FText::FromName(InputTaggedData[Index].Pin),
 				AttributeTypeName
 			));
-			return true;
+			return false;
 		}
 
 		if (!bHasSpecialRequirement)
@@ -411,7 +437,7 @@ bool FPCGMetadataElementBase::ExecuteInternal(FPCGContext* Context) const
 					FText::FromName(InputTaggedData[Index].Pin),
 					AttributeTypeName,
 					MostComplexTypeName));
-				return true;
+				return false;
 			}
 		}
 
@@ -467,16 +493,10 @@ bool FPCGMetadataElementBase::ExecuteInternal(FPCGContext* Context) const
 	// If no operation is needed, just forward input 
 	if (bNoOperationNeeded)
 	{
-		for (uint32 OutputIndex = 0; OutputIndex < Settings->GetOutputPinNum(); ++OutputIndex)
-		{
-			FPCGTaggedData& OutputData = Outputs.Add_GetRef(InputTaggedData[InputPinToForward]);
-			OutputData.Pin = Settings->GetOutputPinLabel(OutputIndex);
-		}
-
+		PassthroughInputs(Context, Settings);
 		return true;
 	}
 
-
 	// At this point, we verified everything, so we can go forward with the computation, depending on the most complex type
 	// So first forward outputs and create the attribute
 	OperationData.OutputAccessors.SetNum(Settings->GetOutputPinNum());
@@ -516,7 +536,7 @@ bool FPCGMetadataElementBase::ExecuteInternal(FPCGContext* Context) const
 				PCGMetadataElementCommon::CopyEntryToValueKeyMap(SourceMetadata[InputPinToForward], SourceAttribute[InputPinToForward], OutputAttribute);
 			}
 
-			OperationData.OutputAccessors[OutputIndex] = PCGAttributeAccessorHelpers::CreateAccessor( const_cast<UPCGData*>(OutputData.Data.Get()), OutputTarget);
+			OperationData.OutputAccessors[OutputIndex] = PCGAttributeAccessorHelpers::CreateAccessor(const_cast<UPCGData*>(OutputData.Data.Get()), OutputTarget);
 		}
 		else
 		{
diff --git a/Engine/Plugins/Experimental/PCG/Source/PCG/Public/Elements/Metadata/PCGMetadataOpElementBase.h b/Engine/Plugins/Experimental/PCG/Source/PCG/Public/Elements/Metadata/PCGMetadataOpElementBase.h
--- a/Engine/Plugins/Experimental/PCG/Source/PCG/Public/Elements/Metadata/PCGMetadataOpElementBase.h
+++ b/Engine/Plugins/Experimental/PCG/Source/PCG/Public/Elements/Metadata/PCGMetadataOpElementBase.h
@@ -207,6 +207,9 @@ protected:
 
 	template <typename InType1, typename InType2, typename InType3, typename InType4, typename... Callbacks>
 	bool DoQuaternaryOp(FOperationData& InOperationData, Callbacks&& ...InCallbacks) const;
+
+	/** To be called if we have no data to perform any operation, it will passthrough the input. */
+	void PassthroughInputs(FPCGContext* InContext, const UPCGMetadataSettingsBase* InSettings) const;
 };
 
 template <int32 NbInputs, int32 NbOutputs>
