Fixed landscape transform retrieval at runtime which led to invalid queries into the landscape cache

diff --git a/Engine/Plugins/Experimental/PCG/Source/PCG/Private/Data/PCGLandscapeData.cpp b/Engine/Plugins/Experimental/PCG/Source/PCG/Private/Data/PCGLandscapeData.cpp
--- a/Engine/Plugins/Experimental/PCG/Source/PCG/Private/Data/PCGLandscapeData.cpp
+++ b/Engine/Plugins/Experimental/PCG/Source/PCG/Private/Data/PCGLandscapeData.cpp
@@ -157,7 +157,11 @@ bool UPCGLandscapeData::ProjectPoint(const FTransform& InTransform, const FBox&
 		return false;
 	}
 
+#if WITH_EDITOR
 	const FTransform& LandscapeTransform = LandscapeInfo->GetLandscapeProxy()->GetTransform();
+#else
+	const FTransform LandscapeTransform = LandscapeInfo->GetLandscapeProxy()->LandscapeActorToWorld();
+#endif
 
 	// TODO: compute full transform when we want to support bounds
 	const FVector LocalPoint = LandscapeTransform.InverseTransformPosition(InTransform.GetLocation());
@@ -255,7 +259,11 @@ const UPCGPointData* UPCGLandscapeData::CreatePointData(FPCGContext* Context, co
 			continue;
 		}
 
+#if WITH_EDITOR
 		const FTransform& LandscapeTransform = LandscapeInfo->GetLandscapeProxy()->GetTransform();
+#else
+		const FTransform LandscapeTransform = LandscapeInfo->GetLandscapeProxy()->LandscapeActorToWorld();
+#endif
 		const int32 ComponentSizeQuads = LandscapeInfo->ComponentSizeQuads;
 
 		// TODO: add offset to nearest edge, will have an impact if the grid size doesn't match the landscape size
diff --git a/Engine/Plugins/Experimental/PCG/Source/PCG/Private/Helpers/PCGBlueprintHelpers.cpp b/Engine/Plugins/Experimental/PCG/Source/PCG/Private/Helpers/PCGBlueprintHelpers.cpp
--- a/Engine/Plugins/Experimental/PCG/Source/PCG/Private/Helpers/PCGBlueprintHelpers.cpp
+++ b/Engine/Plugins/Experimental/PCG/Source/PCG/Private/Helpers/PCGBlueprintHelpers.cpp
@@ -177,10 +177,8 @@ TArray<FPCGLandscapeLayerWeight> UPCGBlueprintHelpers::GetInterpolatedPCGLandsca
 
 	FString FailureReason;
 
-
 	for (TWeakObjectPtr<ALandscapeProxy> LandscapePtr : Landscapes)
 	{
-
 		ALandscapeProxy* Landscape = LandscapePtr.Get();
 		if (!Landscape)
 		{
@@ -194,7 +192,11 @@ TArray<FPCGLandscapeLayerWeight> UPCGBlueprintHelpers::GetInterpolatedPCGLandsca
 			continue;
 		}
 
+#if WITH_EDITOR
 		const FVector LocalPoint = Landscape->GetTransform().InverseTransformPosition(Location);
+#else
+		const FVector LocalPoint = Landscape->LandscapeActorToWorld().InverseTransformPosition(Location);
+#endif
 		const FIntPoint ComponentMapKey(FMath::FloorToInt(LocalPoint.X / LandscapeInfo->ComponentSizeQuads), FMath::FloorToInt(LocalPoint.Y / LandscapeInfo->ComponentSizeQuads));
 
 		#if WITH_EDITOR
