[PCG] Fixed an issue where clear pcg link would parent the stamp actor under the original actor

#rnx

diff --git a/Engine/Plugins/Experimental/PCG/Source/PCG/Private/PCGComponent.cpp b/Engine/Plugins/Experimental/PCG/Source/PCG/Private/PCGComponent.cpp
--- a/Engine/Plugins/Experimental/PCG/Source/PCG/Private/PCGComponent.cpp
+++ b/Engine/Plugins/Experimental/PCG/Source/PCG/Private/PCGComponent.cpp
@@ -595,7 +595,7 @@ AActor* UPCGComponent::ClearPCGLink(UClass* TemplateActor)
 	UWorld* World = GetWorld();
 
 	// First create a new actor that will be the new owner of all the resources
-	AActor* NewActor = UPCGActorHelpers::SpawnDefaultActor(World, GetOwner()->GetLevel(), TemplateActor ? TemplateActor : AActor::StaticClass(), TEXT("PCGStamp"), GetOwner()->GetTransform(), GetOwner());
+	AActor* NewActor = UPCGActorHelpers::SpawnDefaultActor(World, GetOwner()->GetLevel(), TemplateActor ? TemplateActor : AActor::StaticClass(), TEXT("PCGStamp"), GetOwner()->GetTransform());
 
 	// Then move all resources linked to this component to this actor
 	bool bHasMovedResources = MoveResourcesToNewActor(NewActor, /*bCreateChild=*/false);
