[PCG] Disabled reset arrow for pcg components generated and dirty properties that are intended for visibility only

diff --git a/Engine/Plugins/Experimental/PCG/Source/PCG/Public/PCGComponent.h b/Engine/Plugins/Experimental/PCG/Source/PCG/Public/PCGComponent.h
--- a/Engine/Plugins/Experimental/PCG/Source/PCG/Public/PCGComponent.h
+++ b/Engine/Plugins/Experimental/PCG/Source/PCG/Public/PCGComponent.h
@@ -199,7 +199,7 @@ public:
 	EPCGComponentGenerationTrigger GenerationTrigger = EPCGComponentGenerationTrigger::GenerateOnLoad;
 
 	/** Flag to indicate whether this component has run in the editor. Note that for partitionable actors, this will always be false. */
-	UPROPERTY(BlueprintReadOnly, VisibleAnywhere, AdvancedDisplay, Category = Properties)
+	UPROPERTY(BlueprintReadOnly, VisibleAnywhere, AdvancedDisplay, Category = Properties, meta = (NoResetToDefault))
 	bool bGenerated = false;
 
 	UPROPERTY(NonPIEDuplicateTransient)
@@ -209,7 +209,7 @@ public:
 	UPROPERTY(BlueprintReadWrite, EditAnywhere, AdvancedDisplay, Category = Properties, meta = (DisplayName = "Regenerate PCG volume in editor"))
 	bool bRegenerateInEditor = true;
 
-	UPROPERTY(BlueprintReadOnly, VisibleAnywhere, Transient, AdvancedDisplay, Category = Properties)
+	UPROPERTY(BlueprintReadOnly, VisibleAnywhere, Transient, AdvancedDisplay, Category = Properties, meta = (NoResetToDefault))
 	bool bDirtyGenerated = false;
 
 	// Property that will automatically be set on BP templates, to allow for "Generate on add to world" in editor.
@@ -350,7 +350,7 @@ private:
 	void SetupCallbacksOnCreation();
 
 	void UpdateTrackingCache();
-	
+
 	// InOriginatingChangeObject can be null
 	bool DirtyTrackedActor(AActor* InActor, bool bIntersect, const TSet<FName>& InRemovedTags, const UObject* InOriginatingChangeObject);
 
