CreatePoints: Reduced the cost of target calls, added ShouldComputeFullOutputDataCrc() for performance

diff --git a/Engine/Plugins/Experimental/PCG/Source/PCG/Private/Elements/PCGCreatePoints.cpp b/Engine/Plugins/Experimental/PCG/Source/PCG/Private/Elements/PCGCreatePoints.cpp
--- a/Engine/Plugins/Experimental/PCG/Source/PCG/Private/Elements/PCGCreatePoints.cpp
+++ b/Engine/Plugins/Experimental/PCG/Source/PCG/Private/Elements/PCGCreatePoints.cpp
@@ -39,7 +39,6 @@ bool FPCGCreatePointsElement::ExecuteInternal(FPCGContext* Context) const
 
 	UPCGComponent* OriginalComponentContext = Context->SourceComponent.Get();
 	check(OriginalComponentContext);
-	UPCGSpatialData* Target = Cast<UPCGSpatialData>(OriginalComponentContext->GetActorPCGData());
 
 	const FTransform OriginalComponentTransform = OriginalComponentContext->GetOwner()->GetActorTransform();
 	const FTransform ComponentTransformScaleOne = FTransform(OriginalComponentTransform.Rotator(), OriginalComponentTransform.GetLocation(), FVector::One());
@@ -71,7 +70,9 @@ bool FPCGCreatePointsElement::ExecuteInternal(FPCGContext* Context) const
 	}
 	else
 	{
-		FPCGAsync::AsyncPointProcessing(Context, PointsToLoopOn.Num(), OutputPoints, [&PointsToLoopOn, Settings, Target, &ComponentTransformScaleOne](int32 Index, FPCGPoint& OutPoint)
+		const UPCGSpatialData* Target = Settings->bCullPointsOutsideVolume ? Cast<UPCGSpatialData>(OriginalComponentContext->GetActorPCGData()) : nullptr;
+
+		FPCGAsync::AsyncPointProcessing(Context, PointsToLoopOn.Num(), OutputPoints, [&PointsToLoopOn, Settings, &ComponentTransformScaleOne, Target](int32 Index, FPCGPoint& OutPoint)
 		{
 			const FPCGPoint& InPoint = PointsToLoopOn[Index];
 			OutPoint = InPoint;
@@ -84,7 +85,7 @@ bool FPCGCreatePointsElement::ExecuteInternal(FPCGContext* Context) const
 			OutPoint.Seed = UPCGBlueprintHelpers::ComputeSeedFromPosition(OutPoint.Transform.GetLocation());
 
 			// Discards all points that are outside the volume
-			return !Settings->bCullPointsOutsideVolume || (Target->GetDensityAtPosition(OutPoint.Transform.GetLocation()) > 0.0f);
+			return !Target || (Target->GetDensityAtPosition(OutPoint.Transform.GetLocation()) > 0.0f);
 		});
 	}
 
diff --git a/Engine/Plugins/Experimental/PCG/Source/PCG/Public/Elements/PCGCreatePoints.h b/Engine/Plugins/Experimental/PCG/Source/PCG/Public/Elements/PCGCreatePoints.h
--- a/Engine/Plugins/Experimental/PCG/Source/PCG/Public/Elements/PCGCreatePoints.h
+++ b/Engine/Plugins/Experimental/PCG/Source/PCG/Public/Elements/PCGCreatePoints.h
@@ -34,6 +34,7 @@ public:
 	UPROPERTY(BlueprintReadWrite, EditAnywhere, Category = Settings, meta = (PCG_Overridable))
 	TArray<FPCGPoint> PointsToCreate;
 
+	/** If true, points are transformed to world space using the PCG component transform */
 	UPROPERTY(BlueprintReadWrite, EditAnywhere, Category = Settings, meta = (PCG_Overridable))
 	bool bLocal = false;
 
@@ -46,4 +47,5 @@ class FPCGCreatePointsElement : public FSimplePCGElement
 protected:
 	virtual bool ExecuteInternal(FPCGContext* Context) const override;
 	virtual bool IsCacheable(const UPCGSettings* InSettings) const override;
+	virtual bool ShouldComputeFullOutputDataCrc(FPCGContext* Context) const { return true; }
 };
\ No newline at end of file
