Fixed the serialize-on-cook code path and added world actor merging with some caveats

diff --git a/Engine/Plugins/Experimental/PCG/Source/PCG/Private/Grid/PCGLandscapeCache.cpp b/Engine/Plugins/Experimental/PCG/Source/PCG/Private/Grid/PCGLandscapeCache.cpp
--- a/Engine/Plugins/Experimental/PCG/Source/PCG/Private/Grid/PCGLandscapeCache.cpp
+++ b/Engine/Plugins/Experimental/PCG/Source/PCG/Private/Grid/PCGLandscapeCache.cpp
@@ -679,6 +679,44 @@ void UPCGLandscapeCache::ClearCache()
 #endif
 }
 
+void UPCGLandscapeCache::TakeOwnership(UPCGLandscapeCache* InLandscapeCache)
+{
+	if (SerializationMode != EPCGLandscapeCacheSerializationMode::NeverSerialize)
+	{
+		Modify(/*bAlwaysMarkDirty=*/false);
+	}
+
+	bool bShouldDirty = false;
+
+	for (TPair<TPair<FGuid, FIntPoint>, FPCGLandscapeCacheEntry*>& CacheEntryPair : InLandscapeCache->CachedData)
+	{
+		FPCGLandscapeCacheEntry* CacheEntry = CachedData.FindOrAdd(CacheEntryPair.Key, CacheEntryPair.Value);
+		if (CacheEntry == CacheEntryPair.Value)
+		{
+			CacheEntryPair.Value = nullptr;
+			bShouldDirty = true;
+		}
+	}
+
+	for (FName LayerName : InLandscapeCache->CachedLayerNames)
+	{
+		bool bLayerWasAlreadyPresent = false;
+		CachedLayerNames.Add(LayerName, &bLayerWasAlreadyPresent);
+		bShouldDirty |= !bLayerWasAlreadyPresent;
+	}
+
+	InLandscapeCache->ClearCache();
+
+#if WITH_EDITOR
+	CacheEntryCount = CachedData.Num();
+#endif
+
+	if (bShouldDirty && SerializationMode != EPCGLandscapeCacheSerializationMode::NeverSerialize)
+	{
+		MarkPackageDirty();
+	}
+}
+
 #if WITH_EDITOR
 const FPCGLandscapeCacheEntry* UPCGLandscapeCache::GetCacheEntry(ULandscapeComponent* LandscapeComponent, const FIntPoint& ComponentCoordinate)
 {
diff --git a/Engine/Plugins/Experimental/PCG/Source/PCG/Private/PCGSubsystem.cpp b/Engine/Plugins/Experimental/PCG/Source/PCG/Private/PCGSubsystem.cpp
--- a/Engine/Plugins/Experimental/PCG/Source/PCG/Private/PCGSubsystem.cpp
+++ b/Engine/Plugins/Experimental/PCG/Source/PCG/Private/PCGSubsystem.cpp
@@ -254,6 +254,10 @@ void UPCGSubsystem::RegisterPCGWorldActor(APCGWorldActor* InActor)
 	{
 		PCGWorldActor = InActor;
 	}
+	else if (InActor != PCGWorldActor)
+	{
+		PCGWorldActor->MergeFrom(InActor);
+	}
 }
 
 void UPCGSubsystem::UnregisterPCGWorldActor(APCGWorldActor* InActor)
diff --git a/Engine/Plugins/Experimental/PCG/Source/PCG/Private/PCGWorldActor.cpp b/Engine/Plugins/Experimental/PCG/Source/PCG/Private/PCGWorldActor.cpp
--- a/Engine/Plugins/Experimental/PCG/Source/PCG/Private/PCGWorldActor.cpp
+++ b/Engine/Plugins/Experimental/PCG/Source/PCG/Private/PCGWorldActor.cpp
@@ -9,7 +9,16 @@
 #include "Grid/PCGPartitionActor.h"
 #include "Helpers/PCGActorHelpers.h"
 
+#include "LandscapeProxy.h"
 #include "Engine/World.h"
+#include "HAL/IConsoleManager.h"
+
+#if WITH_EDITOR
+#include "UObject/UObjectHash.h"
+#include "WorldPartition/WorldPartition.h"
+#include "WorldPartition/WorldPartitionHelpers.h"
+#include "WorldPartition/ActorPartition/PartitionActorDesc.h"
+#endif
 
 #include UE_INLINE_GENERATED_CPP_BY_NAME(PCGWorldActor)
 
@@ -30,11 +39,65 @@ void APCGWorldActor::BeginCacheForCookedPlatformData(const ITargetPlatform* Targ
 {
 	Super::BeginCacheForCookedPlatformData(TargetPlatform);
 	check(LandscapeCacheObject);
-	LandscapeCacheObject->PrimeCache();
+	
+	UWorld* World = GetWorld();
+
+	if (World && LandscapeCacheObject->SerializationMode == EPCGLandscapeCacheSerializationMode::SerializeOnlyAtCook)
+	{
+		// Implementation note: actor references gathered from the world partition helpers will register on creation and unregister on deletion
+		// which is why we need to manage this only in the non-WP case.
+		TSet<FWorldPartitionReference> ActorRefs;
+		TArray<ALandscapeProxy*> ProxiesToRegisterAndUnregister;
+
+		if (UWorldPartition* WorldPartition = World->GetWorldPartition())
+		{
+			FWorldPartitionHelpers::ForEachActorDesc<ALandscapeProxy>(WorldPartition, [WorldPartition, &ActorRefs](const FWorldPartitionActorDesc* ActorDesc)
+			{
+				check(ActorDesc);
+				ActorRefs.Add(FWorldPartitionReference(WorldPartition, ActorDesc->GetGuid()));
+				return true;
+			});
+		}
+		else
+		{
+			// Since we're not in a WP map, the proxies should be outered to this world.
+			// Important note: registering the landscape proxies can create objects, which can and will cause issues with the ForEachWithOuter, hence the second loop in which we do the register
+			ForEachObjectWithOuter(GetWorld(), [&ProxiesToRegisterAndUnregister](UObject* Object)
+			{
+				if (ALandscapeProxy* LandscapeProxy = Cast<ALandscapeProxy>(Object))
+				{
+					bool bHasUnregisteredComponents = false;
+					LandscapeProxy->ForEachComponent(/*bIncludeFromChildActors=*/false, [&bHasUnregisteredComponents](const UActorComponent* Component)
+					{
+						if (Component && !Component->IsRegistered())
+						{
+							bHasUnregisteredComponents = true;
+						}
+					});
+
+					if (bHasUnregisteredComponents)
+					{
+						ProxiesToRegisterAndUnregister.Add(LandscapeProxy);
+					}
+				}
+			});
+
+			for (ALandscapeProxy* ProxyToRegister : ProxiesToRegisterAndUnregister)
+			{
+				ProxyToRegister->RegisterAllComponents();
+			}
+		}
+
+		LandscapeCacheObject->PrimeCache();
+
+		for (ALandscapeProxy* ProxyToUnregister : ProxiesToRegisterAndUnregister)
+		{
+			ProxyToUnregister->UnregisterAllComponents();
+		}
+	}
 }
 #endif
 
-
 void APCGWorldActor::PostInitProperties()
 {
 	Super::PostInitProperties();
@@ -121,6 +184,13 @@ void APCGWorldActor::GetGridGuids(PCGHiGenGrid::FSizeToGuidMap& OutSizeToGuidMap
 	}
 }
 
+void APCGWorldActor::MergeFrom(APCGWorldActor* OtherWorldActor)
+{
+	check(OtherWorldActor && this != OtherWorldActor);
+	ensure(PartitionGridSize == OtherWorldActor->PartitionGridSize && bUse2DGrid == OtherWorldActor->bUse2DGrid && GridGuids.OrderIndependentCompareEqual(OtherWorldActor->GridGuids));
+	LandscapeCacheObject->TakeOwnership(OtherWorldActor->LandscapeCacheObject);
+}
+
 #if WITH_EDITOR
 APCGWorldActor* APCGWorldActor::CreatePCGWorldActor(UWorld* InWorld)
 {
diff --git a/Engine/Plugins/Experimental/PCG/Source/PCG/Public/Grid/PCGLandscapeCache.h b/Engine/Plugins/Experimental/PCG/Source/PCG/Public/Grid/PCGLandscapeCache.h
--- a/Engine/Plugins/Experimental/PCG/Source/PCG/Public/Grid/PCGLandscapeCache.h
+++ b/Engine/Plugins/Experimental/PCG/Source/PCG/Public/Grid/PCGLandscapeCache.h
@@ -129,6 +129,8 @@ public:
 	UFUNCTION(CallInEditor, Category = "Cache")
 	void ClearCache();
 
+	void TakeOwnership(UPCGLandscapeCache* InLandscapeCache);
+
 	void Tick(float DeltaSeconds);
 
 #if WITH_EDITOR
diff --git a/Engine/Plugins/Experimental/PCG/Source/PCG/Public/PCGWorldActor.h b/Engine/Plugins/Experimental/PCG/Source/PCG/Public/PCGWorldActor.h
--- a/Engine/Plugins/Experimental/PCG/Source/PCG/Public/PCGWorldActor.h
+++ b/Engine/Plugins/Experimental/PCG/Source/PCG/Public/PCGWorldActor.h
@@ -29,6 +29,8 @@ public:
 	void CreateGridGuidsIfNecessary(const PCGHiGenGrid::FSizeArray& InGridSizes);
 	void GetGridGuids(PCGHiGenGrid::FSizeToGuidMap& OutSizeToGuidMap) const;
 
+	void MergeFrom(APCGWorldActor* OtherWorldActor);
+
 #if WITH_EDITOR	
 	virtual bool CanChangeIsSpatiallyLoadedFlag() const { return false; }
 	virtual bool IsUserManaged() const override { return false; }
