[PCG] Add support for Indexes for Attribute Sets

#rnx

diff --git a/Engine/Plugins/Experimental/PCG/Source/PCG/Private/Metadata/Accessors/PCGAttributeAccessorHelpers.cpp b/Engine/Plugins/Experimental/PCG/Source/PCG/Private/Metadata/Accessors/PCGAttributeAccessorHelpers.cpp
--- a/Engine/Plugins/Experimental/PCG/Source/PCG/Private/Metadata/Accessors/PCGAttributeAccessorHelpers.cpp
+++ b/Engine/Plugins/Experimental/PCG/Source/PCG/Private/Metadata/Accessors/PCGAttributeAccessorHelpers.cpp
@@ -198,6 +198,11 @@ namespace PCGAttributeAccessorHelpers
 				return MakeUnique<FPCGAttributeAccessorKeysEntries>(Attribute);
 			}
 		}
+		else if (Metadata && InSelector.GetSelection() == EPCGAttributePropertySelection::ExtraProperty && InSelector.GetExtraProperty() == EPCGExtraProperties::Index)
+		{
+			// Special case for Indexes, we will use the metadata to get the keys
+			return MakeUnique<FPCGAttributeAccessorKeysEntries>(Metadata);
+		}
 		else
 		{
 			return TUniquePtr<KeysType>();
diff --git a/Engine/Plugins/Experimental/PCG/Source/PCG/Private/Metadata/Accessors/PCGAttributeAccessorKeys.cpp b/Engine/Plugins/Experimental/PCG/Source/PCG/Private/Metadata/Accessors/PCGAttributeAccessorKeys.cpp
--- a/Engine/Plugins/Experimental/PCG/Source/PCG/Private/Metadata/Accessors/PCGAttributeAccessorKeys.cpp
+++ b/Engine/Plugins/Experimental/PCG/Source/PCG/Private/Metadata/Accessors/PCGAttributeAccessorKeys.cpp
@@ -1,6 +1,7 @@
 // Copyright Epic Games, Inc. All Rights Reserved.
 
 #include "Metadata/Accessors/PCGAttributeAccessorKeys.h"
+#include "Metadata/PCGMetadata.h"
 #include "Metadata/PCGMetadataAttribute.h"
 
 #include "PCGPoint.h"
@@ -32,6 +33,29 @@ FPCGAttributeAccessorKeysEntries::FPCGAttributeAccessorKeysEntries(PCGMetadataEn
 	Entries.Add(EntryKey);
 }
 
+FPCGAttributeAccessorKeysEntries::FPCGAttributeAccessorKeysEntries(const UPCGMetadata* Metadata)
+	: IPCGAttributeAccessorKeys(/*bInReadOnly=*/ true)
+{
+	if (!Metadata)
+	{
+		return;
+	}
+
+	const PCGMetadataEntryKey ItemKeyLowerBound = Metadata->GetItemKeyCountForParent();
+	const PCGMetadataEntryKey ItemKeyUpperBound = Metadata->GetItemCountForChild();
+	const int64 Count = ItemKeyUpperBound - ItemKeyLowerBound;
+
+	if (Count > 0)
+	{
+		Entries.Reserve(Count);
+
+		for (PCGMetadataEntryKey Entry = ItemKeyLowerBound; Entry < ItemKeyUpperBound; ++Entry)
+		{
+			Entries.Add(Entry);
+		}
+	}
+}
+
 bool FPCGAttributeAccessorKeysEntries::GetMetadataEntryKeys(int32 InStart, TArrayView<PCGMetadataEntryKey*>& OutEntryKeys)
 {
 	return PCGAttributeAccessorKeys::GetKeys(Entries, InStart, OutEntryKeys, [](PCGMetadataEntryKey& Key) -> PCGMetadataEntryKey* { return &Key; });
diff --git a/Engine/Plugins/Experimental/PCG/Source/PCG/Public/Metadata/Accessors/PCGAttributeAccessorKeys.h b/Engine/Plugins/Experimental/PCG/Source/PCG/Public/Metadata/Accessors/PCGAttributeAccessorKeys.h
--- a/Engine/Plugins/Experimental/PCG/Source/PCG/Public/Metadata/Accessors/PCGAttributeAccessorKeys.h
+++ b/Engine/Plugins/Experimental/PCG/Source/PCG/Public/Metadata/Accessors/PCGAttributeAccessorKeys.h
@@ -8,6 +8,7 @@
 struct FPCGPoint;
 
 class FPCGMetadataAttributeBase;
+class UPCGMetadata;
 
 ///////////////////////////////////////////////////////////////////////
 
@@ -140,6 +141,9 @@ public:
 	explicit FPCGAttributeAccessorKeysEntries(const FPCGMetadataAttributeBase* Attribute);
 	explicit FPCGAttributeAccessorKeysEntries(PCGMetadataEntryKey EntryKey);
 
+	// Read-only case where we will just iterate on all the entries in the metadata.
+	explicit FPCGAttributeAccessorKeysEntries(const UPCGMetadata* Metadata);
+
 	virtual int32 GetNum() const override { return Entries.Num(); }
 
 protected:
