[PCG] Fixed static analysis issue in pcg editor graph, Dereferencing NULL pointer 'OutboundNode'

#rnx

diff --git a/Engine/Plugins/Experimental/PCG/Source/PCGEditor/Private/PCGEditorGraph.cpp b/Engine/Plugins/Experimental/PCG/Source/PCGEditor/Private/PCGEditorGraph.cpp
--- a/Engine/Plugins/Experimental/PCG/Source/PCGEditor/Private/PCGEditorGraph.cpp
+++ b/Engine/Plugins/Experimental/PCG/Source/PCGEditor/Private/PCGEditorGraph.cpp
@@ -62,9 +62,9 @@ void UPCGEditorGraph::ReconstructGraph()
 	{
 		if (!IsValid(PCGNode))
 		{
-			continue;	
+			continue;
 		}
-		
+
 		if (Cast<UPCGRerouteSettings>(PCGNode->GetSettings()))
 		{
 			FGraphNodeCreator<UPCGEditorGraphNodeReroute> NodeCreator(*this);
@@ -156,6 +156,11 @@ void UPCGEditorGraph::CreateLinks(UPCGEditorGraphNodeBase* GraphNode, bool bCrea
 	{
 		for (UPCGPin* InputPin : PCGNode->GetInputPins())
 		{
+			if (!InputPin)
+			{
+				continue;
+			}
+
 			UEdGraphPin* InPin = GraphNode->FindPin(InputPin->Properties.Label, EEdGraphPinDirection::EGPD_Input);
 			if (!InPin)
 			{
@@ -178,6 +183,12 @@ void UPCGEditorGraph::CreateLinks(UPCGEditorGraphNodeBase* GraphNode, bool bCrea
 				}
 
 				const UPCGNode* InboundNode = InboundEdge->InputPin ? InboundEdge->InputPin->Node : nullptr;
+				if (!ensure(InboundNode))
+				{
+					UE_LOG(LogPCGEditor, Error, TEXT("Invalid inbound node for %s"), *InputPin->Properties.Label.ToString());
+					continue;
+				}
+
 				UPCGEditorGraphNodeBase* const* ConnectedGraphNode = InboundNode ? InPCGNodeToPCGEditorNodeMap.Find(InboundNode) : nullptr;
 				UEdGraphPin* OutPin = ConnectedGraphNode ? (*ConnectedGraphNode)->FindPin(InboundEdge->InputPin->Properties.Label, EEdGraphPinDirection::EGPD_Output) : nullptr;
 				if (OutPin)
@@ -197,6 +208,11 @@ void UPCGEditorGraph::CreateLinks(UPCGEditorGraphNodeBase* GraphNode, bool bCrea
 	{
 		for (UPCGPin* OutputPin : PCGNode->GetOutputPins())
 		{
+			if (!OutputPin)
+			{
+				continue;
+			}
+
 			UEdGraphPin* OutPin = GraphNode->FindPin(OutputPin->Properties.Label, EEdGraphPinDirection::EGPD_Output);
 			if (!OutPin)
 			{
@@ -204,7 +220,7 @@ void UPCGEditorGraph::CreateLinks(UPCGEditorGraphNodeBase* GraphNode, bool bCrea
 				{
 					UE_LOG(LogPCGEditor, Error, TEXT("Invalid OutputPin for %s"), *OutputPin->Properties.Label.ToString());
 				}
-				
+
 				continue;
 			}
 
@@ -218,6 +234,12 @@ void UPCGEditorGraph::CreateLinks(UPCGEditorGraphNodeBase* GraphNode, bool bCrea
 				}
 
 				const UPCGNode* OutboundNode = OutboundEdge->OutputPin ? OutboundEdge->OutputPin->Node : nullptr;
+				if (!ensure(OutboundNode))
+				{
+					UE_LOG(LogPCGEditor, Error, TEXT("Invalid outbound node for %s"), *OutputPin->Properties.Label.ToString());
+					continue;
+				}
+
 				UPCGEditorGraphNodeBase* const* ConnectedGraphNode = OutboundNode ? InPCGNodeToPCGEditorNodeMap.Find(OutboundNode) : nullptr;
 				UEdGraphPin* InPin = ConnectedGraphNode ? (*ConnectedGraphNode)->FindPin(OutboundEdge->OutputPin->Properties.Label, EEdGraphPinDirection::EGPD_Input) : nullptr;
 				if (InPin)
@@ -271,4 +293,4 @@ void UPCGEditorGraph::OnGraphUserParametersChanged(UPCGGraphInterface* InGraph,
 	{
 		NodeToRemove->DestroyNode();
 	}
-}
\ No newline at end of file
+}
