[PCG] Fix locks on per-pin output data

#rnx

diff --git a/Engine/Plugins/Experimental/PCG/Source/PCG/Private/PCGComponent.cpp b/Engine/Plugins/Experimental/PCG/Source/PCG/Private/PCGComponent.cpp
--- a/Engine/Plugins/Experimental/PCG/Source/PCG/Private/PCGComponent.cpp
+++ b/Engine/Plugins/Experimental/PCG/Source/PCG/Private/PCGComponent.cpp
@@ -672,7 +672,7 @@ EPCGHiGenGrid UPCGComponent::GetGenerationGrid() const
 
 void UPCGComponent::StoreOutputDataForPin(const FString& InResourceKey, const FPCGDataCollection& InData)
 {
-	FReadScopeLock ScopedWriteLock(PerPinGeneratedOutputLock);
+	FWriteScopeLock ScopedWriteLock(PerPinGeneratedOutputLock);
 	PerPinGeneratedOutput.FindOrAdd(InResourceKey) = InData;
 }
 
@@ -684,7 +684,7 @@ const FPCGDataCollection* UPCGComponent::RetrieveOutputDataForPin(const FString&
 
 void UPCGComponent::ClearPerPinGeneratedOutput()
 {
-	FReadScopeLock ScopedWriteLock(PerPinGeneratedOutputLock);
+	FWriteScopeLock ScopedWriteLock(PerPinGeneratedOutputLock);
 	PerPinGeneratedOutput.Reset();
 }
 
