[PCG] Fixed a crash were the pcg world actor would not unregister properly causing the pcg subsystem to access invalid data

diff --git a/Engine/Plugins/Experimental/PCG/Source/PCG/Private/PCGWorldActor.cpp b/Engine/Plugins/Experimental/PCG/Source/PCG/Private/PCGWorldActor.cpp
--- a/Engine/Plugins/Experimental/PCG/Source/PCG/Private/PCGWorldActor.cpp
+++ b/Engine/Plugins/Experimental/PCG/Source/PCG/Private/PCGWorldActor.cpp
@@ -39,7 +39,7 @@ void APCGWorldActor::BeginCacheForCookedPlatformData(const ITargetPlatform* Targ
 {
 	Super::BeginCacheForCookedPlatformData(TargetPlatform);
 	check(LandscapeCacheObject);
-	
+
 	UWorld* World = GetWorld();
 
 	if (World && LandscapeCacheObject->SerializationMode == EPCGLandscapeCacheSerializationMode::SerializeOnlyAtCook)
@@ -243,7 +243,7 @@ void APCGWorldActor::OnPartitionGridSizeChanged()
 	{
 		return;
 	}
-	
+
 	bool bAllSafeToDelete = true;
 
 	auto AddPartitionComponentAndCheckIfSafeToDelete = [&bAllSafeToDelete](AActor* Actor) -> bool
@@ -297,10 +297,16 @@ void APCGWorldActor::PostLoad()
 
 void APCGWorldActor::PostEditChangeProperty(FPropertyChangedEvent& PropertyChangedEvent)
 {
-	if (PropertyChangedEvent.GetPropertyName() == GET_MEMBER_NAME_CHECKED(APCGWorldActor, PartitionGridSize) 
+	if (PropertyChangedEvent.GetPropertyName() == GET_MEMBER_NAME_CHECKED(APCGWorldActor, PartitionGridSize)
 		|| PropertyChangedEvent.GetPropertyName() == GET_MEMBER_NAME_CHECKED(APCGWorldActor, bUse2DGrid))
 	{
 		OnPartitionGridSizeChanged();
 	}
 }
 #endif // WITH_EDITOR
+
+void APCGWorldActor::BeginDestroy()
+{
+	UnregisterFromSubsystem();
+	Super::BeginDestroy();
+}
diff --git a/Engine/Plugins/Experimental/PCG/Source/PCG/Public/PCGWorldActor.h b/Engine/Plugins/Experimental/PCG/Source/PCG/Public/PCGWorldActor.h
--- a/Engine/Plugins/Experimental/PCG/Source/PCG/Public/PCGWorldActor.h
+++ b/Engine/Plugins/Experimental/PCG/Source/PCG/Public/PCGWorldActor.h
@@ -31,7 +31,7 @@ public:
 
 	void MergeFrom(APCGWorldActor* OtherWorldActor);
 
-#if WITH_EDITOR	
+#if WITH_EDITOR
 	virtual bool CanChangeIsSpatiallyLoadedFlag() const { return false; }
 	virtual bool IsUserManaged() const override { return false; }
 	virtual bool ShouldExport() override { return false; }
@@ -45,10 +45,11 @@ public:
 	static inline constexpr uint32 DefaultPartitionGridSize = 25600; // 256m
 
 	//~ Begin UObject Interface.
-#if WITH_EDITOR	
+#if WITH_EDITOR
 	virtual void PostLoad() override;
 	virtual void PostEditChangeProperty(FPropertyChangedEvent& PropertyChangedEvent) override;
 #endif
+	virtual void BeginDestroy() override;
 	//~ End UObject Interface.
 
 	/** Size of the PCG partition actor grid for non-hierarchical-generation graphs. */
@@ -67,7 +68,7 @@ private:
 	void RegisterToSubsystem();
 	void UnregisterFromSubsystem();
 
-#if WITH_EDITOR	
+#if WITH_EDITOR
 	void OnPartitionGridSizeChanged();
 #endif
 
