[PCG] Migrate DeterminismTestsCommon test functionality upwards to PCGTestsCommon since it was being used in non-determinism tests

#rnx

diff --git a/Engine/Plugins/Experimental/PCG/Source/PCG/Private/Tests/Determinism/PCGCollapseDeterminismTest.cpp b/Engine/Plugins/Experimental/PCG/Source/PCG/Private/Tests/Determinism/PCGCollapseDeterminismTest.cpp
--- a/Engine/Plugins/Experimental/PCG/Source/PCG/Private/Tests/Determinism/PCGCollapseDeterminismTest.cpp
+++ b/Engine/Plugins/Experimental/PCG/Source/PCG/Private/Tests/Determinism/PCGCollapseDeterminismTest.cpp
@@ -33,11 +33,12 @@ namespace
 		check(NumSpatialSetsToAdd > 0 && NumPointsToCreatePerSet > 0 && NumPolyLinePointsToCreatePerSet > 0);
 
 		TestDataSet.Empty();
-		TestDataSet.Init(PCGTestsCommon::FTestData(PCGDeterminismTests::Defaults::Seed), NumTestableSpatialDataTypes);
+		PCGTestsCommon::FTestData TestData;
+		TestDataSet.Init(TestData, NumTestableSpatialDataTypes);
 
 		for (PCGTestsCommon::FTestData& Data : TestDataSet)
 		{
-			PCGDeterminismTests::GenerateSettings<UPCGCollapseSettings>(Data);
+			PCGTestsCommon::GenerateSettings<UPCGCollapseSettings>(Data);
 		}
 
 		for (int32 I = 0; I < NumSpatialSetsToAdd; ++I)
diff --git a/Engine/Plugins/Experimental/PCG/Source/PCG/Private/Tests/Determinism/PCGDifferenceDeterminismTest.cpp b/Engine/Plugins/Experimental/PCG/Source/PCG/Private/Tests/Determinism/PCGDifferenceDeterminismTest.cpp
--- a/Engine/Plugins/Experimental/PCG/Source/PCG/Private/Tests/Determinism/PCGDifferenceDeterminismTest.cpp
+++ b/Engine/Plugins/Experimental/PCG/Source/PCG/Private/Tests/Determinism/PCGDifferenceDeterminismTest.cpp
@@ -51,7 +51,7 @@ namespace
 			break;
 		}
 
-		PCGDeterminismTests::GenerateSettings<UPCGDifferenceSettings>(TestData, AdditionalSettingsDelegate);
+		PCGTestsCommon::GenerateSettings<UPCGDifferenceSettings>(TestData, AdditionalSettingsDelegate);
 		// Source
 		PCGDeterminismTests::AddVolumeInputData(TestData.InputData, FVector::ZeroVector, PCGDeterminismTests::Defaults::LargeVector, PCGDeterminismTests::Defaults::SmallVector, PCGDifferenceConstants::SourceLabel);
 
@@ -100,9 +100,9 @@ namespace PCGDeterminismTests
 bool FPCGDifferenceDeterminismSingleSameDataTest::RunTest(const FString& Parameters)
 {
 	// Test single same data
-	PCGTestsCommon::FTestData TestDataInferred(PCGDeterminismTests::Defaults::Seed);
-	PCGTestsCommon::FTestData TestDataContinuous(PCGDeterminismTests::Defaults::Seed);
-	PCGTestsCommon::FTestData TestDataDiscrete(PCGDeterminismTests::Defaults::Seed);
+	PCGTestsCommon::FTestData TestDataInferred;
+	PCGTestsCommon::FTestData TestDataContinuous;
+	PCGTestsCommon::FTestData TestDataDiscrete;
 
 	DifferenceTestBase(TestDataInferred, EPCGDifferenceMode::Inferred);
 	DifferenceTestBase(TestDataContinuous, EPCGDifferenceMode::Continuous);
@@ -118,12 +118,12 @@ bool FPCGDifferenceDeterminismSingleSameDataTest::RunTest(const FString& Paramet
 bool FPCGDifferenceDeterminismSingleIdenticalDataTest::RunTest(const FString& Parameters)
 {
 	// Test single identical data
-	PCGTestsCommon::FTestData FirstTestDataInferred(PCGDeterminismTests::Defaults::Seed);
-	PCGTestsCommon::FTestData FirstTestDataContinuous(PCGDeterminismTests::Defaults::Seed);
-	PCGTestsCommon::FTestData FirstTestDataDiscrete(PCGDeterminismTests::Defaults::Seed);
-	PCGTestsCommon::FTestData SecondTestDataInferred(PCGDeterminismTests::Defaults::Seed);
-	PCGTestsCommon::FTestData SecondTestDataContinuous(PCGDeterminismTests::Defaults::Seed);
-	PCGTestsCommon::FTestData SecondTestDataDiscrete(PCGDeterminismTests::Defaults::Seed);
+	PCGTestsCommon::FTestData FirstTestDataInferred;
+	PCGTestsCommon::FTestData FirstTestDataContinuous;
+	PCGTestsCommon::FTestData FirstTestDataDiscrete;
+	PCGTestsCommon::FTestData SecondTestDataInferred;
+	PCGTestsCommon::FTestData SecondTestDataContinuous;
+	PCGTestsCommon::FTestData SecondTestDataDiscrete;
 
 	DifferenceTestBase(FirstTestDataInferred, EPCGDifferenceMode::Inferred);
 	DifferenceTestBase(FirstTestDataContinuous, EPCGDifferenceMode::Continuous);
@@ -142,9 +142,9 @@ bool FPCGDifferenceDeterminismSingleIdenticalDataTest::RunTest(const FString& Pa
 bool FPCGDifferenceDeterminismMultipleSameDataTest::RunTest(const FString& Parameters)
 {
 	// Test multiple same data
-	PCGTestsCommon::FTestData TestDataInferred(PCGDeterminismTests::Defaults::Seed);
-	PCGTestsCommon::FTestData TestDataContinuous(PCGDeterminismTests::Defaults::Seed);
-	PCGTestsCommon::FTestData TestDataDiscrete(PCGDeterminismTests::Defaults::Seed);
+	PCGTestsCommon::FTestData TestDataInferred;
+	PCGTestsCommon::FTestData TestDataContinuous;
+	PCGTestsCommon::FTestData TestDataDiscrete;
 
 	DifferenceTestMultiple(TestDataInferred, EPCGDifferenceMode::Inferred);
 	DifferenceTestMultiple(TestDataContinuous, EPCGDifferenceMode::Continuous);
@@ -160,12 +160,12 @@ bool FPCGDifferenceDeterminismMultipleSameDataTest::RunTest(const FString& Param
 bool FPCGDifferenceDeterminismMultipleIdenticalDataTest::RunTest(const FString& Parameters)
 {
 	// Test multiple identical data
-	PCGTestsCommon::FTestData FirstTestDataInferred(PCGDeterminismTests::Defaults::Seed);
-	PCGTestsCommon::FTestData FirstTestDataContinuous(PCGDeterminismTests::Defaults::Seed);
-	PCGTestsCommon::FTestData FirstTestDataDiscrete(PCGDeterminismTests::Defaults::Seed);
-	PCGTestsCommon::FTestData SecondTestDataInferred(PCGDeterminismTests::Defaults::Seed);
-	PCGTestsCommon::FTestData SecondTestDataContinuous(PCGDeterminismTests::Defaults::Seed);
-	PCGTestsCommon::FTestData SecondTestDataDiscrete(PCGDeterminismTests::Defaults::Seed);
+	PCGTestsCommon::FTestData FirstTestDataInferred;
+	PCGTestsCommon::FTestData FirstTestDataContinuous;
+	PCGTestsCommon::FTestData FirstTestDataDiscrete;
+	PCGTestsCommon::FTestData SecondTestDataInferred;
+	PCGTestsCommon::FTestData SecondTestDataContinuous;
+	PCGTestsCommon::FTestData SecondTestDataDiscrete;
 
 	DifferenceTestMultiple(FirstTestDataInferred, EPCGDifferenceMode::Inferred);
 	DifferenceTestMultiple(FirstTestDataContinuous, EPCGDifferenceMode::Continuous);
@@ -184,12 +184,12 @@ bool FPCGDifferenceDeterminismMultipleIdenticalDataTest::RunTest(const FString&
 bool FPCGDifferenceDeterminismOrderIndependenceTest::RunTest(const FString& Parameters)
 {
 	// Test multiple identical data, shuffled
-	PCGTestsCommon::FTestData FirstTestDataInferred(PCGDeterminismTests::Defaults::Seed);
-	PCGTestsCommon::FTestData FirstTestDataContinuous(PCGDeterminismTests::Defaults::Seed);
-	PCGTestsCommon::FTestData FirstTestDataDiscrete(PCGDeterminismTests::Defaults::Seed);
-	PCGTestsCommon::FTestData SecondTestDataInferred(PCGDeterminismTests::Defaults::Seed);
-	PCGTestsCommon::FTestData SecondTestDataContinuous(PCGDeterminismTests::Defaults::Seed);
-	PCGTestsCommon::FTestData SecondTestDataDiscrete(PCGDeterminismTests::Defaults::Seed);
+	PCGTestsCommon::FTestData FirstTestDataInferred;
+	PCGTestsCommon::FTestData FirstTestDataContinuous;
+	PCGTestsCommon::FTestData FirstTestDataDiscrete;
+	PCGTestsCommon::FTestData SecondTestDataInferred;
+	PCGTestsCommon::FTestData SecondTestDataContinuous;
+	PCGTestsCommon::FTestData SecondTestDataDiscrete;
 
 	DifferenceTestMultiple(FirstTestDataInferred, EPCGDifferenceMode::Inferred);
 	DifferenceTestMultiple(FirstTestDataContinuous, EPCGDifferenceMode::Continuous);
diff --git a/Engine/Plugins/Experimental/PCG/Source/PCG/Private/Tests/Determinism/PCGIntersectionDeterminismTest.cpp b/Engine/Plugins/Experimental/PCG/Source/PCG/Private/Tests/Determinism/PCGIntersectionDeterminismTest.cpp
--- a/Engine/Plugins/Experimental/PCG/Source/PCG/Private/Tests/Determinism/PCGIntersectionDeterminismTest.cpp
+++ b/Engine/Plugins/Experimental/PCG/Source/PCG/Private/Tests/Determinism/PCGIntersectionDeterminismTest.cpp
@@ -15,7 +15,7 @@ namespace
 {
 	void IntersectionTestBase(PCGTestsCommon::FTestData& TestData)
 	{
-		PCGDeterminismTests::GenerateSettings<UPCGInnerIntersectionSettings>(TestData);
+		PCGTestsCommon::GenerateSettings<UPCGInnerIntersectionSettings>(TestData);
 		// Source Volumes
 		PCGDeterminismTests::AddVolumeInputData(TestData.InputData, PCGDeterminismTests::Defaults::SmallVector, PCGDeterminismTests::Defaults::MediumVector, PCGDeterminismTests::Defaults::MediumVector);
 		PCGDeterminismTests::AddVolumeInputData(TestData.InputData, PCGDeterminismTests::Defaults::SmallVector * -1.f, PCGDeterminismTests::Defaults::MediumVector, PCGDeterminismTests::Defaults::MediumVector);
@@ -33,7 +33,7 @@ namespace
 bool FPCGIntersectionDeterminismSingleSameDataTest::RunTest(const FString& Parameters)
 {
 	// Test single same data
-	PCGTestsCommon::FTestData TestData(PCGDeterminismTests::Defaults::Seed);
+	PCGTestsCommon::FTestData TestData;
 
 	IntersectionTestBase(TestData);
 
@@ -43,8 +43,8 @@ bool FPCGIntersectionDeterminismSingleSameDataTest::RunTest(const FString& Param
 bool FPCGIntersectionDeterminismSingleIdenticalDataTest::RunTest(const FString& Parameters)
 {
 	// Test single identical data
-	PCGTestsCommon::FTestData FirstTestData(PCGDeterminismTests::Defaults::Seed);
-	PCGTestsCommon::FTestData SecondTestData(PCGDeterminismTests::Defaults::Seed);
+	PCGTestsCommon::FTestData FirstTestData;
+	PCGTestsCommon::FTestData SecondTestData;
 
 	IntersectionTestBase(FirstTestData);
 	IntersectionTestBase(SecondTestData);
@@ -55,7 +55,7 @@ bool FPCGIntersectionDeterminismSingleIdenticalDataTest::RunTest(const FString&
 bool FPCGIntersectionDeterminismMultipleSameDataTest::RunTest(const FString& Parameters)
 {
 	// Test multiple same data
-	PCGTestsCommon::FTestData TestData(PCGDeterminismTests::Defaults::Seed);
+	PCGTestsCommon::FTestData TestData;
 
 	IntersectionTestMultiple(TestData);
 
@@ -65,8 +65,8 @@ bool FPCGIntersectionDeterminismMultipleSameDataTest::RunTest(const FString& Par
 bool FPCGIntersectionDeterminismMultipleIdenticalDataTest::RunTest(const FString& Parameters)
 {
 	// Test multiple identical data
-	PCGTestsCommon::FTestData FirstTestData(PCGDeterminismTests::Defaults::Seed);
-	PCGTestsCommon::FTestData SecondTestData(PCGDeterminismTests::Defaults::Seed);
+	PCGTestsCommon::FTestData FirstTestData;
+	PCGTestsCommon::FTestData SecondTestData;
 
 	IntersectionTestMultiple(FirstTestData);
 	IntersectionTestMultiple(SecondTestData);
@@ -77,8 +77,8 @@ bool FPCGIntersectionDeterminismMultipleIdenticalDataTest::RunTest(const FString
 bool FPCGIntersectionDeterminismOrderIndependenceTest::RunTest(const FString& Parameters)
 {
 	// Test multiple identical, shuffled data
-	PCGTestsCommon::FTestData FirstTestData(PCGDeterminismTests::Defaults::Seed);
-	PCGTestsCommon::FTestData SecondTestData(PCGDeterminismTests::Defaults::Seed);
+	PCGTestsCommon::FTestData FirstTestData;
+	PCGTestsCommon::FTestData SecondTestData;
 
 	IntersectionTestMultiple(FirstTestData);
 	IntersectionTestMultiple(SecondTestData);
diff --git a/Engine/Plugins/Experimental/PCG/Source/PCG/Private/Tests/Determinism/PCGVolumeSamplerDeterminismTest.cpp b/Engine/Plugins/Experimental/PCG/Source/PCG/Private/Tests/Determinism/PCGVolumeSamplerDeterminismTest.cpp
--- a/Engine/Plugins/Experimental/PCG/Source/PCG/Private/Tests/Determinism/PCGVolumeSamplerDeterminismTest.cpp
+++ b/Engine/Plugins/Experimental/PCG/Source/PCG/Private/Tests/Determinism/PCGVolumeSamplerDeterminismTest.cpp
@@ -25,9 +25,9 @@ namespace
 bool FPCGVolumeSamplerDeterminismSingleSameDataTest::RunTest(const FString& Parameters)
 {
 	// Test single same data
-	PCGTestsCommon::FTestData TestData(PCGDeterminismTests::Defaults::Seed);
+	PCGTestsCommon::FTestData TestData;
 
-	PCGDeterminismTests::GenerateSettings<UPCGVolumeSamplerSettings>(TestData, RandomizeVolumeSettingsVoxelSize);
+	PCGTestsCommon::GenerateSettings<UPCGVolumeSamplerSettings>(TestData, RandomizeVolumeSettingsVoxelSize);
 	PCGDeterminismTests::AddRandomizedVolumeInputData(TestData);
 
 	TestTrue("Same single input and settings, same output", PCGDeterminismTests::ExecutionIsDeterministicSameData(TestData));
@@ -38,11 +38,11 @@ bool FPCGVolumeSamplerDeterminismSingleSameDataTest::RunTest(const FString& Para
 bool FPCGVolumeSamplerDeterminismSingleIdenticalDataTest::RunTest(const FString& Parameters)
 {
 	// Test single identical data
-	PCGTestsCommon::FTestData FirstTestData(PCGDeterminismTests::Defaults::Seed);
-	PCGTestsCommon::FTestData SecondTestData(PCGDeterminismTests::Defaults::Seed);
+	PCGTestsCommon::FTestData FirstTestData;
+	PCGTestsCommon::FTestData SecondTestData;
 
-	PCGDeterminismTests::GenerateSettings<UPCGVolumeSamplerSettings>(FirstTestData, RandomizeVolumeSettingsVoxelSize);
-	PCGDeterminismTests::GenerateSettings<UPCGVolumeSamplerSettings>(SecondTestData, RandomizeVolumeSettingsVoxelSize);
+	PCGTestsCommon::GenerateSettings<UPCGVolumeSamplerSettings>(FirstTestData, RandomizeVolumeSettingsVoxelSize);
+	PCGTestsCommon::GenerateSettings<UPCGVolumeSamplerSettings>(SecondTestData, RandomizeVolumeSettingsVoxelSize);
 
 	PCGDeterminismTests::AddRandomizedVolumeInputData(FirstTestData);
 	PCGDeterminismTests::AddRandomizedVolumeInputData(SecondTestData);
@@ -55,9 +55,9 @@ bool FPCGVolumeSamplerDeterminismSingleIdenticalDataTest::RunTest(const FString&
 bool FPCGVolumeSamplerDeterminismMultipleSameDataTest::RunTest(const FString& Parameters)
 {
 	// Test multiple same data
-	PCGTestsCommon::FTestData TestData(PCGDeterminismTests::Defaults::Seed);
+	PCGTestsCommon::FTestData TestData;
 
-	PCGDeterminismTests::GenerateSettings<UPCGVolumeSamplerSettings>(TestData, RandomizeVolumeSettingsVoxelSize);
+	PCGTestsCommon::GenerateSettings<UPCGVolumeSamplerSettings>(TestData, RandomizeVolumeSettingsVoxelSize);
 
 	// Add many random volumes
 	for (int32 I = 0; I < NumInputVolumesToAdd; ++I)
@@ -73,11 +73,11 @@ bool FPCGVolumeSamplerDeterminismMultipleSameDataTest::RunTest(const FString& Pa
 bool FPCGVolumeSamplerDeterminismMultipleIdenticalDataTest::RunTest(const FString& Parameters)
 {
 	// Test multiple identical data
-	PCGTestsCommon::FTestData FirstTestData(PCGDeterminismTests::Defaults::Seed);
-	PCGTestsCommon::FTestData SecondTestData(PCGDeterminismTests::Defaults::Seed);
+	PCGTestsCommon::FTestData FirstTestData;
+	PCGTestsCommon::FTestData SecondTestData;
 
-	PCGDeterminismTests::GenerateSettings<UPCGVolumeSamplerSettings>(FirstTestData, RandomizeVolumeSettingsVoxelSize);
-	PCGDeterminismTests::GenerateSettings<UPCGVolumeSamplerSettings>(SecondTestData, RandomizeVolumeSettingsVoxelSize);
+	PCGTestsCommon::GenerateSettings<UPCGVolumeSamplerSettings>(FirstTestData, RandomizeVolumeSettingsVoxelSize);
+	PCGTestsCommon::GenerateSettings<UPCGVolumeSamplerSettings>(SecondTestData, RandomizeVolumeSettingsVoxelSize);
 
 	// Add many random volumes
 	for (int32 I = 0; I < NumInputVolumesToAdd; ++I)
@@ -94,11 +94,11 @@ bool FPCGVolumeSamplerDeterminismMultipleIdenticalDataTest::RunTest(const FStrin
 bool FPCGVolumeSamplerDeterminismOrderIndependenceTest::RunTest(const FString& Parameters)
 {
 	// Test shuffled input data
-	PCGTestsCommon::FTestData FirstTestData(PCGDeterminismTests::Defaults::Seed);
-	PCGTestsCommon::FTestData SecondTestData(PCGDeterminismTests::Defaults::Seed);
+	PCGTestsCommon::FTestData FirstTestData;
+	PCGTestsCommon::FTestData SecondTestData;
 
-	PCGDeterminismTests::GenerateSettings<UPCGVolumeSamplerSettings>(FirstTestData, RandomizeVolumeSettingsVoxelSize);
-	PCGDeterminismTests::GenerateSettings<UPCGVolumeSamplerSettings>(SecondTestData, RandomizeVolumeSettingsVoxelSize);
+	PCGTestsCommon::GenerateSettings<UPCGVolumeSamplerSettings>(FirstTestData, RandomizeVolumeSettingsVoxelSize);
+	PCGTestsCommon::GenerateSettings<UPCGVolumeSamplerSettings>(SecondTestData, RandomizeVolumeSettingsVoxelSize);
 
 	// Add many random volumes
 	for (int32 I = 0; I < NumInputVolumesToAdd; ++I)
diff --git a/Engine/Plugins/Experimental/PCG/Source/PCG/Private/Tests/Elements/Metadata/PCGMetadataBooleanOpTest.cpp b/Engine/Plugins/Experimental/PCG/Source/PCG/Private/Tests/Elements/Metadata/PCGMetadataBooleanOpTest.cpp
--- a/Engine/Plugins/Experimental/PCG/Source/PCG/Private/Tests/Elements/Metadata/PCGMetadataBooleanOpTest.cpp
+++ b/Engine/Plugins/Experimental/PCG/Source/PCG/Private/Tests/Elements/Metadata/PCGMetadataBooleanOpTest.cpp
@@ -1,6 +1,6 @@
 // Copyright Epic Games, Inc. All Rights Reserved.
 
-#include "Tests/Determinism/PCGDeterminismTestsCommon.h"
+#include "Tests/PCGTestsCommon.h"
 
 #include "PCGComponent.h"
 #include "PCGContext.h"
@@ -14,8 +14,8 @@ IMPLEMENT_CUSTOM_SIMPLE_AUTOMATION_TEST(FPCGMetadataBooleanOpTest, FPCGTestBaseC
 
 bool FPCGMetadataBooleanOpTest::RunTest(const FString& Parameters)
 {
-	PCGTestsCommon::FTestData TestData(PCGDeterminismTests::Defaults::Seed);
-	PCGDeterminismTests::GenerateSettings<UPCGMetadataBooleanSettings>(TestData);
+	PCGTestsCommon::FTestData TestData;
+	PCGTestsCommon::GenerateSettings<UPCGMetadataBooleanSettings>(TestData);
 	UPCGMetadataBooleanSettings* Settings = CastChecked<UPCGMetadataBooleanSettings>(TestData.Settings);
 	FPCGElementPtr MetadataBooleanElement = TestData.Settings->GetElement();
 
diff --git a/Engine/Plugins/Experimental/PCG/Source/PCG/Private/Tests/Elements/Metadata/PCGMetadataBreakTransformTest.cpp b/Engine/Plugins/Experimental/PCG/Source/PCG/Private/Tests/Elements/Metadata/PCGMetadataBreakTransformTest.cpp
--- a/Engine/Plugins/Experimental/PCG/Source/PCG/Private/Tests/Elements/Metadata/PCGMetadataBreakTransformTest.cpp
+++ b/Engine/Plugins/Experimental/PCG/Source/PCG/Private/Tests/Elements/Metadata/PCGMetadataBreakTransformTest.cpp
@@ -1,6 +1,6 @@
 // Copyright Epic Games, Inc. All Rights Reserved.
 
-#include "Tests/Determinism/PCGDeterminismTestsCommon.h"
+#include "Tests/PCGTestsCommon.h"
 
 #include "Elements/Metadata/PCGMetadataMakeTransform.h"
 #include "PCGComponent.h"
@@ -242,12 +242,12 @@ bool FPCGMetadataBreakTransformTest::RunTest(const FString& Parameters)
 
 	bool bTestPassed = true;
 
-	PCGTestsCommon::FTestData TestDataSpatial(PCGDeterminismTests::Defaults::Seed);
-	PCGDeterminismTests::GenerateSettings<UPCGMetadataBreakTransformSettings>(TestDataSpatial);
+	PCGTestsCommon::FTestData TestDataSpatial;
+	PCGTestsCommon::GenerateSettings<UPCGMetadataBreakTransformSettings>(TestDataSpatial);
 	PCGBreakTransformTest::GenerateSpatialData(TestDataSpatial);
 
-	PCGTestsCommon::FTestData TestDataParams(PCGDeterminismTests::Defaults::Seed);
-	PCGDeterminismTests::GenerateSettings<UPCGMetadataBreakTransformSettings>(TestDataParams);
+	PCGTestsCommon::FTestData TestDataParams;
+	PCGTestsCommon::GenerateSettings<UPCGMetadataBreakTransformSettings>(TestDataParams);
 	PCGBreakTransformTest::GenerateParamData(TestDataParams);
 
 	using PairWhatData = TPair<FString, PCGTestsCommon::FTestData*>;
diff --git a/Engine/Plugins/Experimental/PCG/Source/PCG/Private/Tests/Elements/Metadata/PCGMetadataBreakVectorTest.cpp b/Engine/Plugins/Experimental/PCG/Source/PCG/Private/Tests/Elements/Metadata/PCGMetadataBreakVectorTest.cpp
--- a/Engine/Plugins/Experimental/PCG/Source/PCG/Private/Tests/Elements/Metadata/PCGMetadataBreakVectorTest.cpp
+++ b/Engine/Plugins/Experimental/PCG/Source/PCG/Private/Tests/Elements/Metadata/PCGMetadataBreakVectorTest.cpp
@@ -1,6 +1,6 @@
 // Copyright Epic Games, Inc. All Rights Reserved.
 
-#include "Tests/Determinism/PCGDeterminismTestsCommon.h"
+#include "Tests/PCGTestsCommon.h"
 
 #include "Metadata/PCGMetadataAccessor.h"
 #include "PCGComponent.h"
@@ -308,12 +308,12 @@ bool FPCGMetadataBreakVectorTest::RunTest(const FString& Parameters)
 
 	bool bTestPassed = true;
 
-	PCGTestsCommon::FTestData TestDataSpatial(PCGDeterminismTests::Defaults::Seed);
-	PCGDeterminismTests::GenerateSettings<UPCGMetadataBreakVectorSettings>(TestDataSpatial);
+	PCGTestsCommon::FTestData TestDataSpatial;
+	PCGTestsCommon::GenerateSettings<UPCGMetadataBreakVectorSettings>(TestDataSpatial);
 	PCGBreakVectorTest::GenerateSpatialData(TestDataSpatial);
 
-	PCGTestsCommon::FTestData TestDataParams(PCGDeterminismTests::Defaults::Seed);
-	PCGDeterminismTests::GenerateSettings<UPCGMetadataBreakVectorSettings>(TestDataParams);
+	PCGTestsCommon::FTestData TestDataParams;
+	PCGTestsCommon::GenerateSettings<UPCGMetadataBreakVectorSettings>(TestDataParams);
 	PCGBreakVectorTest::GenerateParamData(TestDataParams);
 
 	using PairWhatData = TPair<FString, PCGTestsCommon::FTestData*>;
diff --git a/Engine/Plugins/Experimental/PCG/Source/PCG/Private/Tests/Elements/Metadata/PCGMetadataMakeTransformTest.cpp b/Engine/Plugins/Experimental/PCG/Source/PCG/Private/Tests/Elements/Metadata/PCGMetadataMakeTransformTest.cpp
--- a/Engine/Plugins/Experimental/PCG/Source/PCG/Private/Tests/Elements/Metadata/PCGMetadataMakeTransformTest.cpp
+++ b/Engine/Plugins/Experimental/PCG/Source/PCG/Private/Tests/Elements/Metadata/PCGMetadataMakeTransformTest.cpp
@@ -1,6 +1,6 @@
 // Copyright Epic Games, Inc. All Rights Reserved.
 
-#include "Tests/Determinism/PCGDeterminismTestsCommon.h"
+#include "Tests/PCGTestsCommon.h"
 
 #include "Metadata/PCGMetadataAccessor.h"
 #include "PCGComponent.h"
@@ -212,12 +212,12 @@ bool FPCGMetadataMakeTransformTest::RunTest(const FString& Parameters)
 
 	bool bTestPassed = true;
 
-	PCGTestsCommon::FTestData TestDataSpatial(PCGDeterminismTests::Defaults::Seed);
-	PCGDeterminismTests::GenerateSettings<UPCGMetadataMakeTransformSettings>(TestDataSpatial);
+	PCGTestsCommon::FTestData TestDataSpatial;
+	PCGTestsCommon::GenerateSettings<UPCGMetadataMakeTransformSettings>(TestDataSpatial);
 	GenerateSpatialData(TestDataSpatial);
 
-	PCGTestsCommon::FTestData TestDataParams(PCGDeterminismTests::Defaults::Seed);
-	PCGDeterminismTests::GenerateSettings<UPCGMetadataMakeTransformSettings>(TestDataParams);
+	PCGTestsCommon::FTestData TestDataParams;
+	PCGTestsCommon::GenerateSettings<UPCGMetadataMakeTransformSettings>(TestDataParams);
 	GenerateParamData(TestDataParams);
 
 	using PairWhatData = TPair<FString, PCGTestsCommon::FTestData*>;
diff --git a/Engine/Plugins/Experimental/PCG/Source/PCG/Private/Tests/Elements/PCGAddTagTest.cpp b/Engine/Plugins/Experimental/PCG/Source/PCG/Private/Tests/Elements/PCGAddTagTest.cpp
--- a/Engine/Plugins/Experimental/PCG/Source/PCG/Private/Tests/Elements/PCGAddTagTest.cpp
+++ b/Engine/Plugins/Experimental/PCG/Source/PCG/Private/Tests/Elements/PCGAddTagTest.cpp
@@ -1,6 +1,6 @@
 // Copyright Epic Games, Inc. All Rights Reserved.
 
-#include "Tests/Determinism/PCGDeterminismTestsCommon.h"
+#include "Tests/PCGTestsCommon.h"
 
 #include "PCGComponent.h"
 #include "PCGContext.h"
@@ -11,11 +11,10 @@ IMPLEMENT_CUSTOM_SIMPLE_AUTOMATION_TEST(FPCGAddTagTest_ZeroTags, FPCGTestBaseCla
 
 bool FPCGAddTagTest_ZeroTags::RunTest(const FString& Parameters)
 {
-	PCGTestsCommon::FTestData TestData(PCGDeterminismTests::Defaults::Seed);
+	PCGTestsCommon::FTestData TestData;
 
 	//auto generate settings for this node
-	PCGDeterminismTests::GenerateSettings<UPCGAddTagSettings>(TestData);
-	UPCGAddTagSettings* Settings = CastChecked<UPCGAddTagSettings>(TestData.Settings);
+	PCGTestsCommon::GenerateSettings<UPCGAddTagSettings>(TestData);
 
 	FPCGElementPtr TestElement = TestData.Settings->GetElement();
 
@@ -44,10 +43,10 @@ IMPLEMENT_CUSTOM_SIMPLE_AUTOMATION_TEST(FPCGAddTagTest_OneTag, FPCGTestBaseClass
 
 bool FPCGAddTagTest_OneTag::RunTest(const FString& Parameters)
 {
-	PCGTestsCommon::FTestData TestData(PCGDeterminismTests::Defaults::Seed);
+	PCGTestsCommon::FTestData TestData;
 
 	//auto generate settings for this node
-	PCGDeterminismTests::GenerateSettings<UPCGAddTagSettings>(TestData);
+	PCGTestsCommon::GenerateSettings<UPCGAddTagSettings>(TestData);
 	UPCGAddTagSettings* Settings = CastChecked<UPCGAddTagSettings>(TestData.Settings);
 
 	FPCGElementPtr TestElement = TestData.Settings->GetElement();
@@ -86,10 +85,10 @@ IMPLEMENT_CUSTOM_SIMPLE_AUTOMATION_TEST(FPCGAddTagTest_MultipleTags, FPCGTestBas
 
 bool FPCGAddTagTest_MultipleTags::RunTest(const FString& Parameters)
 {
-	PCGTestsCommon::FTestData TestData(PCGDeterminismTests::Defaults::Seed);
+	PCGTestsCommon::FTestData TestData;
 
 	//auto generate settings for this node
-	PCGDeterminismTests::GenerateSettings<UPCGAddTagSettings>(TestData);
+	PCGTestsCommon::GenerateSettings<UPCGAddTagSettings>(TestData);
 	UPCGAddTagSettings* Settings = CastChecked<UPCGAddTagSettings>(TestData.Settings);
 
 	FPCGElementPtr TestElement = TestData.Settings->GetElement();
diff --git a/Engine/Plugins/Experimental/PCG/Source/PCG/Private/Tests/Elements/PCGAttributeNoiseTest.cpp b/Engine/Plugins/Experimental/PCG/Source/PCG/Private/Tests/Elements/PCGAttributeNoiseTest.cpp
--- a/Engine/Plugins/Experimental/PCG/Source/PCG/Private/Tests/Elements/PCGAttributeNoiseTest.cpp
+++ b/Engine/Plugins/Experimental/PCG/Source/PCG/Private/Tests/Elements/PCGAttributeNoiseTest.cpp
@@ -1,6 +1,6 @@
 // Copyright Epic Games, Inc. All Rights Reserved.
 
-#include "Tests/Determinism/PCGDeterminismTestsCommon.h"
+#include "Tests/PCGTestsCommon.h"
 
 #include "Data/PCGSpatialData.h"
 #include "PCGComponent.h"
@@ -13,8 +13,8 @@ IMPLEMENT_CUSTOM_SIMPLE_AUTOMATION_TEST(FPCGAttributeNoiseDensityTest, FPCGTestB
 
 bool FPCGAttributeNoiseDensityTest::RunTest(const FString& Parameters)
 {
-	PCGTestsCommon::FTestData TestData(PCGDeterminismTests::Defaults::Seed);
-	PCGDeterminismTests::GenerateSettings<UPCGAttributeNoiseSettings>(TestData);
+	PCGTestsCommon::FTestData TestData;
+	PCGTestsCommon::GenerateSettings<UPCGAttributeNoiseSettings>(TestData);
 	UPCGAttributeNoiseSettings* Settings = CastChecked<UPCGAttributeNoiseSettings>(TestData.Settings);
 	Settings->InputSource.SetPointProperty(EPCGPointProperties::Density);
 	FPCGElementPtr NoiseElement = TestData.Settings->GetElement();
diff --git a/Engine/Plugins/Experimental/PCG/Source/PCG/Private/Tests/Elements/PCGCopyPointsTest.cpp b/Engine/Plugins/Experimental/PCG/Source/PCG/Private/Tests/Elements/PCGCopyPointsTest.cpp
--- a/Engine/Plugins/Experimental/PCG/Source/PCG/Private/Tests/Elements/PCGCopyPointsTest.cpp
+++ b/Engine/Plugins/Experimental/PCG/Source/PCG/Private/Tests/Elements/PCGCopyPointsTest.cpp
@@ -1,6 +1,6 @@
 // Copyright Epic Games, Inc. All Rights Reserved.
 
-#include "Tests/Determinism/PCGDeterminismTestsCommon.h"
+#include "Tests/PCGTestsCommon.h"
 
 #include "PCGComponent.h"
 #include "PCGContext.h"
@@ -16,8 +16,8 @@ IMPLEMENT_CUSTOM_SIMPLE_AUTOMATION_TEST(FPCGCopyPointsTest, FPCGTestBaseClass, "
 
 bool FPCGCopyPointsTest::RunTest(const FString& Parameters)
 {
-	PCGTestsCommon::FTestData TestData(PCGDeterminismTests::Defaults::Seed);
-	PCGDeterminismTests::GenerateSettings<UPCGCopyPointsSettings>(TestData);
+	PCGTestsCommon::FTestData TestData;
+	PCGTestsCommon::GenerateSettings<UPCGCopyPointsSettings>(TestData);
 	UPCGCopyPointsSettings* Settings = CastChecked<UPCGCopyPointsSettings>(TestData.Settings);
 	FPCGElementPtr CopyPointsElement = TestData.Settings->GetElement();
 
diff --git a/Engine/Plugins/Experimental/PCG/Source/PCG/Private/Tests/Elements/PCGCreatePointsTest.cpp b/Engine/Plugins/Experimental/PCG/Source/PCG/Private/Tests/Elements/PCGCreatePointsTest.cpp
--- a/Engine/Plugins/Experimental/PCG/Source/PCG/Private/Tests/Elements/PCGCreatePointsTest.cpp
+++ b/Engine/Plugins/Experimental/PCG/Source/PCG/Private/Tests/Elements/PCGCreatePointsTest.cpp
@@ -1,6 +1,6 @@
 // Copyright Epic Games, Inc. All Rights Reserved.
 
-#include "Tests/Determinism/PCGDeterminismTestsCommon.h"
+#include "Tests/PCGTestsCommon.h"
 
 #include "PCGContext.h"
 #include "Data/PCGPointData.h"
@@ -11,8 +11,8 @@ IMPLEMENT_CUSTOM_SIMPLE_AUTOMATION_TEST(FPCGCreatePointsTest_Basic, FPCGTestBase
 
 bool FPCGCreatePointsTest_Basic::RunTest(const FString& Parameters)
 {
-	PCGTestsCommon::FTestData TestData(PCGDeterminismTests::Defaults::Seed);
-	PCGDeterminismTests::GenerateSettings<UPCGCreatePointsSettings>(TestData);
+	PCGTestsCommon::FTestData TestData;
+	PCGTestsCommon::GenerateSettings<UPCGCreatePointsSettings>(TestData);
 	UPCGCreatePointsSettings* Settings = CastChecked<UPCGCreatePointsSettings>(TestData.Settings);
 	
 	TObjectPtr<UPCGPointData> PointData = PCGTestsCommon::CreateRandomPointData(100, 42, false);
diff --git a/Engine/Plugins/Experimental/PCG/Source/PCG/Private/Tests/Elements/PCGDataTableRowToParamDataTest.cpp b/Engine/Plugins/Experimental/PCG/Source/PCG/Private/Tests/Elements/PCGDataTableRowToParamDataTest.cpp
--- a/Engine/Plugins/Experimental/PCG/Source/PCG/Private/Tests/Elements/PCGDataTableRowToParamDataTest.cpp
+++ b/Engine/Plugins/Experimental/PCG/Source/PCG/Private/Tests/Elements/PCGDataTableRowToParamDataTest.cpp
@@ -1,6 +1,6 @@
 // Copyright Epic Games, Inc. All Rights Reserved.
 
-#include "Tests/Determinism/PCGDeterminismTestsCommon.h"
+#include "Tests/PCGTestsCommon.h"
 
 #include "PCGComponent.h"
 
@@ -13,8 +13,8 @@ IMPLEMENT_CUSTOM_SIMPLE_AUTOMATION_TEST(FPCGDataTableRowToParamDataTest, FPCGTes
 
 bool FPCGDataTableRowToParamDataTest::RunTest(const FString& Parameters)
 {
-	PCGTestsCommon::FTestData TestData(PCGDeterminismTests::Defaults::Seed);
-	PCGDeterminismTests::GenerateSettings<UPCGDataTableRowToParamDataSettings>(TestData);
+	PCGTestsCommon::FTestData TestData;
+	PCGTestsCommon::GenerateSettings<UPCGDataTableRowToParamDataSettings>(TestData);
 	UPCGDataTableRowToParamDataSettings* Settings = CastChecked<UPCGDataTableRowToParamDataSettings>(TestData.Settings);
 
 	UDataTable *TestDataTable = NewObject<UDataTable>();
@@ -72,9 +72,8 @@ IMPLEMENT_CUSTOM_SIMPLE_AUTOMATION_TEST(FPCGDataTableRowToParamDataTest_ParamsOv
 
 bool FPCGDataTableRowToParamDataTest_ParamsOverride::RunTest(const FString& Parameters)
 {
-	PCGTestsCommon::FTestData TestData(PCGDeterminismTests::Defaults::Seed);
-	PCGDeterminismTests::GenerateSettings<UPCGDataTableRowToParamDataSettings>(TestData);
-	UPCGDataTableRowToParamDataSettings* Settings = CastChecked<UPCGDataTableRowToParamDataSettings>(TestData.Settings);
+	PCGTestsCommon::FTestData TestData;
+	PCGTestsCommon::GenerateSettings<UPCGDataTableRowToParamDataSettings>(TestData);
 
 	UDataTable *TestDataTable = NewObject<UDataTable>();
 	TestDataTable->RowStruct = FPCGDataTableRowToParamDataTestStruct::StaticStruct();
diff --git a/Engine/Plugins/Experimental/PCG/Source/PCG/Private/Tests/Elements/PCGDensityRemapTest.cpp b/Engine/Plugins/Experimental/PCG/Source/PCG/Private/Tests/Elements/PCGDensityRemapTest.cpp
--- a/Engine/Plugins/Experimental/PCG/Source/PCG/Private/Tests/Elements/PCGDensityRemapTest.cpp
+++ b/Engine/Plugins/Experimental/PCG/Source/PCG/Private/Tests/Elements/PCGDensityRemapTest.cpp
@@ -1,6 +1,6 @@
 // Copyright Epic Games, Inc. All Rights Reserved.
 
-#include "Tests/Determinism/PCGDeterminismTestsCommon.h"
+#include "Tests/PCGTestsCommon.h"
 
 #include "Data/PCGSpatialData.h"
 #include "PCGComponent.h"
@@ -13,8 +13,8 @@ IMPLEMENT_CUSTOM_SIMPLE_AUTOMATION_TEST(FPCGDensityRemapTest, FPCGTestBaseClass,
 
 bool FPCGDensityRemapTest::RunTest(const FString& Parameters)
 {
-	PCGTestsCommon::FTestData TestData(PCGDeterminismTests::Defaults::Seed);
-	PCGDeterminismTests::GenerateSettings<UPCGDensityRemapSettings>(TestData);
+	PCGTestsCommon::FTestData TestData;
+	PCGTestsCommon::GenerateSettings<UPCGDensityRemapSettings>(TestData);
 	UPCGDensityRemapSettings* Settings = CastChecked<UPCGDensityRemapSettings>(TestData.Settings);
 	FPCGElementPtr DensityRemapElement = TestData.Settings->GetElement();
 
diff --git a/Engine/Plugins/Experimental/PCG/Source/PCG/Private/Tests/Elements/PCGDistanceTest.cpp b/Engine/Plugins/Experimental/PCG/Source/PCG/Private/Tests/Elements/PCGDistanceTest.cpp
--- a/Engine/Plugins/Experimental/PCG/Source/PCG/Private/Tests/Elements/PCGDistanceTest.cpp
+++ b/Engine/Plugins/Experimental/PCG/Source/PCG/Private/Tests/Elements/PCGDistanceTest.cpp
@@ -1,6 +1,6 @@
 // Copyright Epic Games, Inc. All Rights Reserved.
 
-#include "Tests/Determinism/PCGDeterminismTestsCommon.h"
+#include "Tests/PCGTestsCommon.h"
 
 #include "PCGComponent.h"
 #include "PCGContext.h"
@@ -12,8 +12,8 @@ IMPLEMENT_CUSTOM_SIMPLE_AUTOMATION_TEST(FPCGDistanceTest_PointToPoint, FPCGTestB
 
 bool FPCGDistanceTest_PointToPoint::RunTest(const FString& Parameters)
 {
-	PCGTestsCommon::FTestData TestData(PCGDeterminismTests::Defaults::Seed);
-	PCGDeterminismTests::GenerateSettings<UPCGDistanceSettings>(TestData);
+	PCGTestsCommon::FTestData TestData;
+	PCGTestsCommon::GenerateSettings<UPCGDistanceSettings>(TestData);
 	UPCGDistanceSettings* Settings = CastChecked<UPCGDistanceSettings>(TestData.Settings);
 
 	Settings->SourceShape = PCGDistanceShape::Center;
@@ -69,8 +69,8 @@ IMPLEMENT_CUSTOM_SIMPLE_AUTOMATION_TEST(FPCGDistanceTest_SetDensity, FPCGTestBas
 
 bool FPCGDistanceTest_SetDensity::RunTest(const FString& Parameters)
 {
-	PCGTestsCommon::FTestData TestData(PCGDeterminismTests::Defaults::Seed);
-	PCGDeterminismTests::GenerateSettings<UPCGDistanceSettings>(TestData);
+	PCGTestsCommon::FTestData TestData;
+	PCGTestsCommon::GenerateSettings<UPCGDistanceSettings>(TestData);
 	UPCGDistanceSettings* Settings = CastChecked<UPCGDistanceSettings>(TestData.Settings);
 
 	Settings->SourceShape = PCGDistanceShape::Center;
@@ -126,8 +126,8 @@ IMPLEMENT_CUSTOM_SIMPLE_AUTOMATION_TEST(FPCGDistanceTest_PointToSphere, FPCGTest
 
 bool FPCGDistanceTest_PointToSphere::RunTest(const FString& Parameters)
 {
-	PCGTestsCommon::FTestData TestData(PCGDeterminismTests::Defaults::Seed);
-	PCGDeterminismTests::GenerateSettings<UPCGDistanceSettings>(TestData);
+	PCGTestsCommon::FTestData TestData;
+	PCGTestsCommon::GenerateSettings<UPCGDistanceSettings>(TestData);
 	UPCGDistanceSettings* Settings = CastChecked<UPCGDistanceSettings>(TestData.Settings);
 
 	Settings->SourceShape = PCGDistanceShape::Center;
@@ -191,8 +191,8 @@ IMPLEMENT_CUSTOM_SIMPLE_AUTOMATION_TEST(FPCGDistanceTest_PointToBox, FPCGTestBas
 
 bool FPCGDistanceTest_PointToBox::RunTest(const FString& Parameters)
 {
-	PCGTestsCommon::FTestData TestData(PCGDeterminismTests::Defaults::Seed);
-	PCGDeterminismTests::GenerateSettings<UPCGDistanceSettings>(TestData);
+	PCGTestsCommon::FTestData TestData;
+	PCGTestsCommon::GenerateSettings<UPCGDistanceSettings>(TestData);
 	UPCGDistanceSettings* Settings = CastChecked<UPCGDistanceSettings>(TestData.Settings);
 
 	Settings->SourceShape = PCGDistanceShape::Center;
diff --git a/Engine/Plugins/Experimental/PCG/Source/PCG/Private/Tests/Elements/PCGElementTest.cpp b/Engine/Plugins/Experimental/PCG/Source/PCG/Private/Tests/Elements/PCGElementTest.cpp
--- a/Engine/Plugins/Experimental/PCG/Source/PCG/Private/Tests/Elements/PCGElementTest.cpp
+++ b/Engine/Plugins/Experimental/PCG/Source/PCG/Private/Tests/Elements/PCGElementTest.cpp
@@ -2,7 +2,7 @@
 
 #include "Tests/Elements/PCGElementTest.h"
 
-#include "Tests/Determinism/PCGDeterminismTestsCommon.h"
+#include "Tests/PCGTestsCommon.h"
 #include "PCGData.h"
 #include "PCGContext.h"
 
@@ -33,8 +33,8 @@ IMPLEMENT_CUSTOM_SIMPLE_AUTOMATION_TEST(FPCGElementTest_CleanupAndValidateOutput
 
 bool FPCGElementTest_CleanupAndValidateOutput::RunTest(const FString& Parameters)
 {
-	PCGTestsCommon::FTestData TestData(PCGDeterminismTests::Defaults::Seed);
-	PCGDeterminismTests::GenerateSettings<UPCGBadOutputsNodeSettings>(TestData);
+	PCGTestsCommon::FTestData TestData;
+	PCGTestsCommon::GenerateSettings<UPCGBadOutputsNodeSettings>(TestData);
 
 	TUniquePtr<FPCGContext> Context = TestData.InitializeTestContext();
 
diff --git a/Engine/Plugins/Experimental/PCG/Source/PCG/Private/Tests/Elements/PCGGatherTest.cpp b/Engine/Plugins/Experimental/PCG/Source/PCG/Private/Tests/Elements/PCGGatherTest.cpp
--- a/Engine/Plugins/Experimental/PCG/Source/PCG/Private/Tests/Elements/PCGGatherTest.cpp
+++ b/Engine/Plugins/Experimental/PCG/Source/PCG/Private/Tests/Elements/PCGGatherTest.cpp
@@ -1,6 +1,6 @@
 // Copyright Epic Games, Inc. All Rights Reserved.
 
-#include "Tests/Determinism/PCGDeterminismTestsCommon.h"
+#include "Tests/PCGTestsCommon.h"
 
 #include "PCGComponent.h"
 #include "PCGContext.h"
@@ -12,9 +12,8 @@ IMPLEMENT_CUSTOM_SIMPLE_AUTOMATION_TEST(FPCGGatherTest_Basic, FPCGTestBaseClass,
 
 bool FPCGGatherTest_Basic::RunTest(const FString& Parameters)
 {
-	PCGTestsCommon::FTestData TestData(PCGDeterminismTests::Defaults::Seed);
-	PCGDeterminismTests::GenerateSettings<UPCGGatherSettings>(TestData);
-	UPCGGatherSettings* Settings = CastChecked<UPCGGatherSettings>(TestData.Settings);
+	PCGTestsCommon::FTestData TestData;
+	PCGTestsCommon::GenerateSettings<UPCGGatherSettings>(TestData);
 
 	FPCGElementPtr TestElement = TestData.Settings->GetElement();
 
diff --git a/Engine/Plugins/Experimental/PCG/Source/PCG/Private/Tests/Elements/PCGInstancePackerTest.cpp b/Engine/Plugins/Experimental/PCG/Source/PCG/Private/Tests/Elements/PCGInstancePackerTest.cpp
--- a/Engine/Plugins/Experimental/PCG/Source/PCG/Private/Tests/Elements/PCGInstancePackerTest.cpp
+++ b/Engine/Plugins/Experimental/PCG/Source/PCG/Private/Tests/Elements/PCGInstancePackerTest.cpp
@@ -1,7 +1,7 @@
 // Copyright Epic Games, Inc. All Rights Reserved.
 
 #include "MeshSelectors/PCGMeshSelectorBase.h"
-#include "Tests/Determinism/PCGDeterminismTestsCommon.h"
+#include "Tests/PCGTestsCommon.h"
 
 #include "Metadata/PCGMetadataAccessor.h"
 #include "PCGComponent.h"
@@ -101,8 +101,8 @@ namespace
 
 bool FPCGStaticMeshSpawnerInstancePackerByAttributeTest::RunTest(const FString& Parameters)
 {
-	PCGTestsCommon::FTestData TestData(PCGDeterminismTests::Defaults::Seed);
-	PCGDeterminismTests::GenerateSettings<UPCGStaticMeshSpawnerSettings>(TestData);
+	PCGTestsCommon::FTestData TestData;
+	PCGTestsCommon::GenerateSettings<UPCGStaticMeshSpawnerSettings>(TestData);
 	UPCGStaticMeshSpawnerSettings* Settings = CastChecked<UPCGStaticMeshSpawnerSettings>(TestData.Settings);
 
 	FPCGTaggedData& SourceTaggedData = TestData.InputData.TaggedData.Emplace_GetRef(FPCGTaggedData());
@@ -258,8 +258,8 @@ bool FPCGStaticMeshSpawnerInstancePackerByAttributeTest::RunTest(const FString&
 
 bool FPCGStaticMeshSpawnerInstancePackerByRegexTest::RunTest(const FString& Parameters)
 {
-	PCGTestsCommon::FTestData TestData(PCGDeterminismTests::Defaults::Seed);
-	PCGDeterminismTests::GenerateSettings<UPCGStaticMeshSpawnerSettings>(TestData);
+	PCGTestsCommon::FTestData TestData;
+	PCGTestsCommon::GenerateSettings<UPCGStaticMeshSpawnerSettings>(TestData);
 	UPCGStaticMeshSpawnerSettings* Settings = CastChecked<UPCGStaticMeshSpawnerSettings>(TestData.Settings);
 
 	FPCGTaggedData& SourceTaggedData = TestData.InputData.TaggedData.Emplace_GetRef(FPCGTaggedData());
diff --git a/Engine/Plugins/Experimental/PCG/Source/PCG/Private/Tests/Elements/PCGNormalToDensityTest.cpp b/Engine/Plugins/Experimental/PCG/Source/PCG/Private/Tests/Elements/PCGNormalToDensityTest.cpp
--- a/Engine/Plugins/Experimental/PCG/Source/PCG/Private/Tests/Elements/PCGNormalToDensityTest.cpp
+++ b/Engine/Plugins/Experimental/PCG/Source/PCG/Private/Tests/Elements/PCGNormalToDensityTest.cpp
@@ -1,6 +1,6 @@
 // Copyright Epic Games, Inc. All Rights Reserved.
 
-#include "Tests/Determinism/PCGDeterminismTestsCommon.h"
+#include "Tests/PCGTestsCommon.h"
 
 #include "Data/PCGSpatialData.h"
 #include "PCGComponent.h"
@@ -13,8 +13,8 @@ IMPLEMENT_CUSTOM_SIMPLE_AUTOMATION_TEST(FPCGNormalToDensityTest_Set, FPCGTestBas
 
 bool FPCGNormalToDensityTest_Set::RunTest(const FString& Parameters)
 {
-	PCGTestsCommon::FTestData TestData(PCGDeterminismTests::Defaults::Seed);
-	PCGDeterminismTests::GenerateSettings<UPCGNormalToDensitySettings>(TestData);
+	PCGTestsCommon::FTestData TestData;
+	PCGTestsCommon::GenerateSettings<UPCGNormalToDensitySettings>(TestData);
 	UPCGNormalToDensitySettings* Settings = CastChecked<UPCGNormalToDensitySettings>(TestData.Settings);
 	FPCGElementPtr Element = TestData.Settings->GetElement();
 
@@ -61,8 +61,8 @@ IMPLEMENT_CUSTOM_SIMPLE_AUTOMATION_TEST(FPCGNormalToDensityTest_CustomNormal, FP
 
 bool FPCGNormalToDensityTest_CustomNormal::RunTest(const FString& Parameters)
 {
-	PCGTestsCommon::FTestData TestData(PCGDeterminismTests::Defaults::Seed);
-	PCGDeterminismTests::GenerateSettings<UPCGNormalToDensitySettings>(TestData);
+	PCGTestsCommon::FTestData TestData;
+	PCGTestsCommon::GenerateSettings<UPCGNormalToDensitySettings>(TestData);
 	UPCGNormalToDensitySettings* Settings = CastChecked<UPCGNormalToDensitySettings>(TestData.Settings);
 	FPCGElementPtr Element = TestData.Settings->GetElement();
 
@@ -108,8 +108,8 @@ IMPLEMENT_CUSTOM_SIMPLE_AUTOMATION_TEST(FPCGNormalToDensityTest_Strength, FPCGTe
 
 bool FPCGNormalToDensityTest_Strength::RunTest(const FString& Parameters)
 {
-	PCGTestsCommon::FTestData TestData(PCGDeterminismTests::Defaults::Seed);
-	PCGDeterminismTests::GenerateSettings<UPCGNormalToDensitySettings>(TestData);
+	PCGTestsCommon::FTestData TestData;
+	PCGTestsCommon::GenerateSettings<UPCGNormalToDensitySettings>(TestData);
 	UPCGNormalToDensitySettings* Settings = CastChecked<UPCGNormalToDensitySettings>(TestData.Settings);
 	FPCGElementPtr Element = TestData.Settings->GetElement();
 
@@ -155,8 +155,8 @@ IMPLEMENT_CUSTOM_SIMPLE_AUTOMATION_TEST(FPCGNormalToDensityTest_Minimum, FPCGTes
 
 bool FPCGNormalToDensityTest_Minimum::RunTest(const FString& Parameters)
 {
-	PCGTestsCommon::FTestData TestData(PCGDeterminismTests::Defaults::Seed);
-	PCGDeterminismTests::GenerateSettings<UPCGNormalToDensitySettings>(TestData);
+	PCGTestsCommon::FTestData TestData;
+	PCGTestsCommon::GenerateSettings<UPCGNormalToDensitySettings>(TestData);
 	UPCGNormalToDensitySettings* Settings = CastChecked<UPCGNormalToDensitySettings>(TestData.Settings);
 	FPCGElementPtr Element = TestData.Settings->GetElement();
 
@@ -203,8 +203,8 @@ IMPLEMENT_CUSTOM_SIMPLE_AUTOMATION_TEST(FPCGNormalToDensityTest_Maximum, FPCGTes
 
 bool FPCGNormalToDensityTest_Maximum::RunTest(const FString& Parameters)
 {
-	PCGTestsCommon::FTestData TestData(PCGDeterminismTests::Defaults::Seed);
-	PCGDeterminismTests::GenerateSettings<UPCGNormalToDensitySettings>(TestData);
+	PCGTestsCommon::FTestData TestData;
+	PCGTestsCommon::GenerateSettings<UPCGNormalToDensitySettings>(TestData);
 	UPCGNormalToDensitySettings* Settings = CastChecked<UPCGNormalToDensitySettings>(TestData.Settings);
 	FPCGElementPtr Element = TestData.Settings->GetElement();
 
@@ -252,8 +252,8 @@ IMPLEMENT_CUSTOM_SIMPLE_AUTOMATION_TEST(FPCGNormalToDensityTest_Add, FPCGTestBas
 
 bool FPCGNormalToDensityTest_Add::RunTest(const FString& Parameters)
 {
-	PCGTestsCommon::FTestData TestData(PCGDeterminismTests::Defaults::Seed);
-	PCGDeterminismTests::GenerateSettings<UPCGNormalToDensitySettings>(TestData);
+	PCGTestsCommon::FTestData TestData;
+	PCGTestsCommon::GenerateSettings<UPCGNormalToDensitySettings>(TestData);
 	UPCGNormalToDensitySettings* Settings = CastChecked<UPCGNormalToDensitySettings>(TestData.Settings);
 	FPCGElementPtr Element = TestData.Settings->GetElement();
 
@@ -300,8 +300,8 @@ IMPLEMENT_CUSTOM_SIMPLE_AUTOMATION_TEST(FPCGNormalToDensityTest_Subtract, FPCGTe
 
 bool FPCGNormalToDensityTest_Subtract::RunTest(const FString& Parameters)
 {
-	PCGTestsCommon::FTestData TestData(PCGDeterminismTests::Defaults::Seed);
-	PCGDeterminismTests::GenerateSettings<UPCGNormalToDensitySettings>(TestData);
+	PCGTestsCommon::FTestData TestData;
+	PCGTestsCommon::GenerateSettings<UPCGNormalToDensitySettings>(TestData);
 	UPCGNormalToDensitySettings* Settings = CastChecked<UPCGNormalToDensitySettings>(TestData.Settings);
 	FPCGElementPtr Element = TestData.Settings->GetElement();
 
@@ -348,8 +348,8 @@ IMPLEMENT_CUSTOM_SIMPLE_AUTOMATION_TEST(FPCGNormalToDensityTest_Multiply, FPCGTe
 
 bool FPCGNormalToDensityTest_Multiply::RunTest(const FString& Parameters)
 {
-	PCGTestsCommon::FTestData TestData(PCGDeterminismTests::Defaults::Seed);
-	PCGDeterminismTests::GenerateSettings<UPCGNormalToDensitySettings>(TestData);
+	PCGTestsCommon::FTestData TestData;
+	PCGTestsCommon::GenerateSettings<UPCGNormalToDensitySettings>(TestData);
 	UPCGNormalToDensitySettings* Settings = CastChecked<UPCGNormalToDensitySettings>(TestData.Settings);
 	FPCGElementPtr Element = TestData.Settings->GetElement();
 
@@ -396,8 +396,8 @@ IMPLEMENT_CUSTOM_SIMPLE_AUTOMATION_TEST(FPCGNormalToDensityTest_Divide, FPCGTest
 
 bool FPCGNormalToDensityTest_Divide::RunTest(const FString& Parameters)
 {
-	PCGTestsCommon::FTestData TestData(PCGDeterminismTests::Defaults::Seed);
-	PCGDeterminismTests::GenerateSettings<UPCGNormalToDensitySettings>(TestData);
+	PCGTestsCommon::FTestData TestData;
+	PCGTestsCommon::GenerateSettings<UPCGNormalToDensitySettings>(TestData);
 	UPCGNormalToDensitySettings* Settings = CastChecked<UPCGNormalToDensitySettings>(TestData.Settings);
 	FPCGElementPtr Element = TestData.Settings->GetElement();
 
diff --git a/Engine/Plugins/Experimental/PCG/Source/PCG/Private/Tests/Elements/PCGPointExtentsModifierTest.cpp b/Engine/Plugins/Experimental/PCG/Source/PCG/Private/Tests/Elements/PCGPointExtentsModifierTest.cpp
--- a/Engine/Plugins/Experimental/PCG/Source/PCG/Private/Tests/Elements/PCGPointExtentsModifierTest.cpp
+++ b/Engine/Plugins/Experimental/PCG/Source/PCG/Private/Tests/Elements/PCGPointExtentsModifierTest.cpp
@@ -1,6 +1,6 @@
 // Copyright Epic Games, Inc. All Rights Reserved.
 
-#include "Tests/Determinism/PCGDeterminismTestsCommon.h"
+#include "Tests/PCGTestsCommon.h"
 
 #include "Data/PCGSpatialData.h"
 #include "PCGComponent.h"
@@ -13,8 +13,8 @@ IMPLEMENT_CUSTOM_SIMPLE_AUTOMATION_TEST(FPCGPointsBoundsModiferTest, FPCGTestBas
 
 bool FPCGPointsBoundsModiferTest::RunTest(const FString& Parameters)
 {
-	PCGTestsCommon::FTestData TestData(PCGDeterminismTests::Defaults::Seed);
-	PCGDeterminismTests::GenerateSettings<UPCGPointExtentsModifierSettings>(TestData);
+	PCGTestsCommon::FTestData TestData;
+	PCGTestsCommon::GenerateSettings<UPCGPointExtentsModifierSettings>(TestData);
 	UPCGPointExtentsModifierSettings* Settings = CastChecked<UPCGPointExtentsModifierSettings>(TestData.Settings);
 	FPCGElementPtr BoundsModifierElement = TestData.Settings->GetElement();
 
diff --git a/Engine/Plugins/Experimental/PCG/Source/PCG/Private/Tests/Elements/PCGPointFilterTest.cpp b/Engine/Plugins/Experimental/PCG/Source/PCG/Private/Tests/Elements/PCGPointFilterTest.cpp
--- a/Engine/Plugins/Experimental/PCG/Source/PCG/Private/Tests/Elements/PCGPointFilterTest.cpp
+++ b/Engine/Plugins/Experimental/PCG/Source/PCG/Private/Tests/Elements/PCGPointFilterTest.cpp
@@ -1,6 +1,6 @@
 // Copyright Epic Games, Inc. All Rights Reserved.
 
-#include "Tests/Determinism/PCGDeterminismTestsCommon.h"
+#include "Tests/PCGTestsCommon.h"
 
 #include "PCGComponent.h"
 #include "PCGContext.h"
@@ -37,8 +37,8 @@ namespace PCGPointFilterTest
 
 bool FPCGPointFilterDensity::RunTest(const FString& Parameters)
 {
-	PCGTestsCommon::FTestData TestData(PCGDeterminismTests::Defaults::Seed);
-	PCGDeterminismTests::GenerateSettings<UPCGPointFilterSettings>(TestData);
+	PCGTestsCommon::FTestData TestData;
+	PCGTestsCommon::GenerateSettings<UPCGPointFilterSettings>(TestData);
 	UPCGPointFilterSettings* Settings = CastChecked<UPCGPointFilterSettings>(TestData.Settings);
 
 	FPCGElementPtr TestElement = TestData.Settings->GetElement();
@@ -54,7 +54,7 @@ bool FPCGPointFilterDensity::RunTest(const FString& Parameters)
 
 	FPCGTaggedData& TaggedData = TestData.InputData.TaggedData.Emplace_GetRef(FPCGTaggedData());
 	TaggedData.Pin = PCGPinConstants::DefaultInputLabel;
-	TaggedData.Data = PCGPointFilterTest::GeneratePointDataWithRandomDensity(NumPoints, PCGDeterminismTests::Defaults::Seed);
+	TaggedData.Data = PCGPointFilterTest::GeneratePointDataWithRandomDensity(NumPoints, TestData.Seed);
 
 	TUniquePtr<FPCGContext> Context = TestData.InitializeTestContext();
 
@@ -90,8 +90,8 @@ bool FPCGPointFilterDensity::RunTest(const FString& Parameters)
 
 bool FPCGPointFilterDensityRange::RunTest(const FString& Parameters)
 {
-	PCGTestsCommon::FTestData TestData(PCGDeterminismTests::Defaults::Seed);
-	PCGDeterminismTests::GenerateSettings<UPCGPointFilterRangeSettings>(TestData);
+	PCGTestsCommon::FTestData TestData;
+	PCGTestsCommon::GenerateSettings<UPCGPointFilterRangeSettings>(TestData);
 	UPCGPointFilterRangeSettings* Settings = CastChecked<UPCGPointFilterRangeSettings>(TestData.Settings);
 
 	FPCGElementPtr TestElement = TestData.Settings->GetElement();
@@ -112,7 +112,7 @@ bool FPCGPointFilterDensityRange::RunTest(const FString& Parameters)
 
 	FPCGTaggedData& TaggedData = TestData.InputData.TaggedData.Emplace_GetRef(FPCGTaggedData());
 	TaggedData.Pin = PCGPinConstants::DefaultInputLabel;
-	TaggedData.Data = PCGPointFilterTest::GeneratePointDataWithRandomDensity(NumPoints, PCGDeterminismTests::Defaults::Seed);
+	TaggedData.Data = PCGPointFilterTest::GeneratePointDataWithRandomDensity(NumPoints, TestData.Seed);
 
 	TUniquePtr<FPCGContext> Context = TestData.InitializeTestContext();
 
diff --git a/Engine/Plugins/Experimental/PCG/Source/PCG/Private/Tests/Elements/PCGSanityCheckPointDataTest.cpp b/Engine/Plugins/Experimental/PCG/Source/PCG/Private/Tests/Elements/PCGSanityCheckPointDataTest.cpp
--- a/Engine/Plugins/Experimental/PCG/Source/PCG/Private/Tests/Elements/PCGSanityCheckPointDataTest.cpp
+++ b/Engine/Plugins/Experimental/PCG/Source/PCG/Private/Tests/Elements/PCGSanityCheckPointDataTest.cpp
@@ -1,6 +1,6 @@
 // Copyright Epic Games, Inc. All Rights Reserved.
 
-#include "Tests/Determinism/PCGDeterminismTestsCommon.h"
+#include "Tests/PCGTestsCommon.h"
 
 #include "PCGComponent.h"
 #include "PCGContext.h"
@@ -15,8 +15,8 @@ IMPLEMENT_CUSTOM_SIMPLE_AUTOMATION_TEST(FPCGSanityCheckPointData_MinPointCount,
 
 bool FPCGSanityCheckPointData_MinPointCount::RunTest(const FString& Parameters)
 {
-	PCGTestsCommon::FTestData TestData(PCGDeterminismTests::Defaults::Seed);
-	PCGDeterminismTests::GenerateSettings<UPCGSanityCheckPointDataSettings>(TestData);
+	PCGTestsCommon::FTestData TestData;
+	PCGTestsCommon::GenerateSettings<UPCGSanityCheckPointDataSettings>(TestData);
 	UPCGSanityCheckPointDataSettings* Settings = CastChecked<UPCGSanityCheckPointDataSettings>(TestData.Settings);
 
 	Settings->MinPointCount = 10;
@@ -43,8 +43,8 @@ IMPLEMENT_CUSTOM_SIMPLE_AUTOMATION_TEST(FPCGSanityCheckPointData_MaxPointCount,
 
 bool FPCGSanityCheckPointData_MaxPointCount::RunTest(const FString& Parameters)
 {
-	PCGTestsCommon::FTestData TestData(PCGDeterminismTests::Defaults::Seed);
-	PCGDeterminismTests::GenerateSettings<UPCGSanityCheckPointDataSettings>(TestData);
+	PCGTestsCommon::FTestData TestData;
+	PCGTestsCommon::GenerateSettings<UPCGSanityCheckPointDataSettings>(TestData);
 	UPCGSanityCheckPointDataSettings* Settings = CastChecked<UPCGSanityCheckPointDataSettings>(TestData.Settings);
 
 	Settings->MinPointCount = 10;
@@ -71,8 +71,8 @@ IMPLEMENT_CUSTOM_SIMPLE_AUTOMATION_TEST(FPCGSanityCheckPointData_Success, FPCGTe
 
 bool FPCGSanityCheckPointData_Success::RunTest(const FString& Parameters)
 {
-	PCGTestsCommon::FTestData TestData(PCGDeterminismTests::Defaults::Seed);
-	PCGDeterminismTests::GenerateSettings<UPCGSanityCheckPointDataSettings>(TestData);
+	PCGTestsCommon::FTestData TestData;
+	PCGTestsCommon::GenerateSettings<UPCGSanityCheckPointDataSettings>(TestData);
 	UPCGSanityCheckPointDataSettings* Settings = CastChecked<UPCGSanityCheckPointDataSettings>(TestData.Settings);
 
 	Settings->MinPointCount = 10;
diff --git a/Engine/Plugins/Experimental/PCG/Source/PCG/Private/Tests/Elements/PCGSortPointsTest.cpp b/Engine/Plugins/Experimental/PCG/Source/PCG/Private/Tests/Elements/PCGSortPointsTest.cpp
--- a/Engine/Plugins/Experimental/PCG/Source/PCG/Private/Tests/Elements/PCGSortPointsTest.cpp
+++ b/Engine/Plugins/Experimental/PCG/Source/PCG/Private/Tests/Elements/PCGSortPointsTest.cpp
@@ -1,6 +1,6 @@
 // Copyright Epic Games, Inc. All Rights Reserved.
 
-#include "Tests/Determinism/PCGDeterminismTestsCommon.h"
+#include "Tests/PCGTestsCommon.h"
 
 #include "PCGComponent.h"
 #include "PCGContext.h"
@@ -11,8 +11,8 @@ namespace SortCommonTestData
 {
 	TUniquePtr<FPCGContext> GenerateTestDataAndRunSort(EPCGSortMethod Method, bool bRandomDensity)
 	{
-		PCGTestsCommon::FTestData TestData(PCGDeterminismTests::Defaults::Seed);
-		PCGDeterminismTests::GenerateSettings<UPCGSortPointsSettings>(TestData);
+		PCGTestsCommon::FTestData TestData;
+		PCGTestsCommon::GenerateSettings<UPCGSortPointsSettings>(TestData);
 		UPCGSortPointsSettings* Settings = CastChecked<UPCGSortPointsSettings>(TestData.Settings);
 		Settings->InputSource.SetPointProperty(EPCGPointProperties::Density);
 		Settings->SortMethod = Method;
diff --git a/Engine/Plugins/Experimental/PCG/Source/PCG/Private/Tests/Elements/PCGSpatialNoiseTest.cpp b/Engine/Plugins/Experimental/PCG/Source/PCG/Private/Tests/Elements/PCGSpatialNoiseTest.cpp
--- a/Engine/Plugins/Experimental/PCG/Source/PCG/Private/Tests/Elements/PCGSpatialNoiseTest.cpp
+++ b/Engine/Plugins/Experimental/PCG/Source/PCG/Private/Tests/Elements/PCGSpatialNoiseTest.cpp
@@ -1,6 +1,6 @@
 // Copyright Epic Games, Inc. All Rights Reserved.
 
-#include "Tests/Determinism/PCGDeterminismTestsCommon.h"
+#include "Tests/PCGTestsCommon.h"
 
 #include "PCGContext.h"
 #include "Data/PCGPointData.h"
@@ -43,8 +43,8 @@ IMPLEMENT_CUSTOM_SIMPLE_AUTOMATION_TEST(FPCGSpatialNoise_Perlin2D, FPCGTestBaseC
 
 bool FPCGSpatialNoise_Perlin2D::RunTest(const FString& Parameters)
 {
-	PCGTestsCommon::FTestData TestData(PCGDeterminismTests::Defaults::Seed);
-	PCGDeterminismTests::GenerateSettings<UPCGSpatialNoiseSettings>(TestData);
+	PCGTestsCommon::FTestData TestData;
+	PCGTestsCommon::GenerateSettings<UPCGSpatialNoiseSettings>(TestData);
 	UPCGSpatialNoiseSettings* Settings = CastChecked<UPCGSpatialNoiseSettings>(TestData.Settings);
 
 	Settings->Mode = PCGSpatialNoiseMode::Perlin2D;
@@ -81,8 +81,8 @@ IMPLEMENT_CUSTOM_SIMPLE_AUTOMATION_TEST(FPCGSpatialNoise_Caustic2D, FPCGTestBase
 
 bool FPCGSpatialNoise_Caustic2D::RunTest(const FString& Parameters)
 {
-	PCGTestsCommon::FTestData TestData(PCGDeterminismTests::Defaults::Seed);
-	PCGDeterminismTests::GenerateSettings<UPCGSpatialNoiseSettings>(TestData);
+	PCGTestsCommon::FTestData TestData;
+	PCGTestsCommon::GenerateSettings<UPCGSpatialNoiseSettings>(TestData);
 	UPCGSpatialNoiseSettings* Settings = CastChecked<UPCGSpatialNoiseSettings>(TestData.Settings);
 
 	Settings->Mode = PCGSpatialNoiseMode::Caustic2D;
@@ -119,8 +119,8 @@ IMPLEMENT_CUSTOM_SIMPLE_AUTOMATION_TEST(FPCGSpatialNoise_Voronoi2D, FPCGTestBase
 
 bool FPCGSpatialNoise_Voronoi2D::RunTest(const FString& Parameters)
 {
-	PCGTestsCommon::FTestData TestData(PCGDeterminismTests::Defaults::Seed);
-	PCGDeterminismTests::GenerateSettings<UPCGSpatialNoiseSettings>(TestData);
+	PCGTestsCommon::FTestData TestData;
+	PCGTestsCommon::GenerateSettings<UPCGSpatialNoiseSettings>(TestData);
 	UPCGSpatialNoiseSettings* Settings = CastChecked<UPCGSpatialNoiseSettings>(TestData.Settings);
 
 	Settings->Mode = PCGSpatialNoiseMode::Voronoi2D;
@@ -162,8 +162,8 @@ IMPLEMENT_CUSTOM_SIMPLE_AUTOMATION_TEST(FPCGSpatialNoise_FractionalBrownian2D, F
 
 bool FPCGSpatialNoise_FractionalBrownian2D::RunTest(const FString& Parameters)
 {
-	PCGTestsCommon::FTestData TestData(PCGDeterminismTests::Defaults::Seed);
-	PCGDeterminismTests::GenerateSettings<UPCGSpatialNoiseSettings>(TestData);
+	PCGTestsCommon::FTestData TestData;
+	PCGTestsCommon::GenerateSettings<UPCGSpatialNoiseSettings>(TestData);
 	UPCGSpatialNoiseSettings* Settings = CastChecked<UPCGSpatialNoiseSettings>(TestData.Settings);
 
 	Settings->Mode = PCGSpatialNoiseMode::FractionalBrownian2D;
@@ -200,8 +200,8 @@ IMPLEMENT_CUSTOM_SIMPLE_AUTOMATION_TEST(FPCGSpatialNoise_EdgeMask2D, FPCGTestBas
 
 bool FPCGSpatialNoise_EdgeMask2D::RunTest(const FString& Parameters)
 {
-	PCGTestsCommon::FTestData TestData(PCGDeterminismTests::Defaults::Seed);
-	PCGDeterminismTests::GenerateSettings<UPCGSpatialNoiseSettings>(TestData);
+	PCGTestsCommon::FTestData TestData;
+	PCGTestsCommon::GenerateSettings<UPCGSpatialNoiseSettings>(TestData);
 	UPCGSpatialNoiseSettings* Settings = CastChecked<UPCGSpatialNoiseSettings>(TestData.Settings);
 
 	Settings->Mode = PCGSpatialNoiseMode::EdgeMask2D;
@@ -237,8 +237,8 @@ bool FPCGSpatialNoise_EdgeMask2D::RunTest(const FString& Parameters)
 IMPLEMENT_CUSTOM_SIMPLE_AUTOMATION_TEST(FPCGSpatialNoise_TilingPerlin2D, FPCGTestBaseClass, "pcg.tests.Noise.TilingPerlin2D", PCGTestsCommon::TestFlags)
 bool FPCGSpatialNoise_TilingPerlin2D::RunTest(const FString& Parameters)
 {
-	PCGTestsCommon::FTestData TestData(PCGDeterminismTests::Defaults::Seed);
-	PCGDeterminismTests::GenerateSettings<UPCGSpatialNoiseSettings>(TestData);
+	PCGTestsCommon::FTestData TestData;
+	PCGTestsCommon::GenerateSettings<UPCGSpatialNoiseSettings>(TestData);
 	UPCGSpatialNoiseSettings* Settings = CastChecked<UPCGSpatialNoiseSettings>(TestData.Settings);
 
 	Settings->Mode = PCGSpatialNoiseMode::Perlin2D;
@@ -275,8 +275,8 @@ bool FPCGSpatialNoise_TilingPerlin2D::RunTest(const FString& Parameters)
 IMPLEMENT_CUSTOM_SIMPLE_AUTOMATION_TEST(FPCGSpatialNoise_TilingVoronoi2D, FPCGTestBaseClass, "pcg.tests.Noise.TilingVoronoi2D", PCGTestsCommon::TestFlags)
 bool FPCGSpatialNoise_TilingVoronoi2D::RunTest(const FString& Parameters)
 {
-	PCGTestsCommon::FTestData TestData(PCGDeterminismTests::Defaults::Seed);
-	PCGDeterminismTests::GenerateSettings<UPCGSpatialNoiseSettings>(TestData);
+	PCGTestsCommon::FTestData TestData;
+	PCGTestsCommon::GenerateSettings<UPCGSpatialNoiseSettings>(TestData);
 	UPCGSpatialNoiseSettings* Settings = CastChecked<UPCGSpatialNoiseSettings>(TestData.Settings);
 
 	Settings->Mode = PCGSpatialNoiseMode::Voronoi2D;
diff --git a/Engine/Plugins/Experimental/PCG/Source/PCG/Private/Tests/Elements/PCGStaticMeshSpawnerTest.cpp b/Engine/Plugins/Experimental/PCG/Source/PCG/Private/Tests/Elements/PCGStaticMeshSpawnerTest.cpp
--- a/Engine/Plugins/Experimental/PCG/Source/PCG/Private/Tests/Elements/PCGStaticMeshSpawnerTest.cpp
+++ b/Engine/Plugins/Experimental/PCG/Source/PCG/Private/Tests/Elements/PCGStaticMeshSpawnerTest.cpp
@@ -1,7 +1,7 @@
 // Copyright Epic Games, Inc. All Rights Reserved.
 
 #include "Metadata/PCGMetadataAccessor.h"
-#include "Tests/Determinism/PCGDeterminismTestsCommon.h"
+#include "Tests/PCGTestsCommon.h"
 
 #include "PCGComponent.h"
 #include "PCGContext.h"
@@ -409,8 +409,8 @@ void TestMeshSelectorWeightedByCategory(
 
 bool FPCGStaticMeshSpawnerByAttributeTest::RunTest(const FString& Parameters)
 {
-	PCGTestsCommon::FTestData TestData(PCGDeterminismTests::Defaults::Seed);
-	PCGDeterminismTests::GenerateSettings<UPCGStaticMeshSpawnerSettings>(TestData);
+	PCGTestsCommon::FTestData TestData;
+	PCGTestsCommon::GenerateSettings<UPCGStaticMeshSpawnerSettings>(TestData);
 
 	// create test attribute data
 	const FName RockAttribute = "Rock";
@@ -499,8 +499,8 @@ bool FPCGStaticMeshSpawnerByAttributeTest::RunTest(const FString& Parameters)
 
 bool FPCGStaticMeshSpawnerWeightedTest::RunTest(const FString& Parameters)
 {
-	PCGTestsCommon::FTestData TestData(PCGDeterminismTests::Defaults::Seed);
-	PCGDeterminismTests::GenerateSettings<UPCGStaticMeshSpawnerSettings>(TestData);
+	PCGTestsCommon::FTestData TestData;
+	PCGTestsCommon::GenerateSettings<UPCGStaticMeshSpawnerSettings>(TestData);
 
 	// create test meshes
 	const FString CubePath = TEXT("StaticMesh'/Engine/BasicShapes/Cube.Cube'");
@@ -568,8 +568,8 @@ bool FPCGStaticMeshSpawnerWeightedTest::RunTest(const FString& Parameters)
 
 bool FPCGStaticMeshSpawnerWeightedByCategoryTest::RunTest(const FString& Parameters)
 {
-	PCGTestsCommon::FTestData TestData(PCGDeterminismTests::Defaults::Seed);
-	PCGDeterminismTests::GenerateSettings<UPCGStaticMeshSpawnerSettings>(TestData);
+	PCGTestsCommon::FTestData TestData;
+	PCGTestsCommon::GenerateSettings<UPCGStaticMeshSpawnerSettings>(TestData);
 
 	// create test meshes
 	const FString CubePath = TEXT("StaticMesh'/Engine/BasicShapes/Cube.Cube'");
diff --git a/Engine/Plugins/Experimental/PCG/Source/PCG/Private/Tests/Elements/PCGTransformPointsTest.cpp b/Engine/Plugins/Experimental/PCG/Source/PCG/Private/Tests/Elements/PCGTransformPointsTest.cpp
--- a/Engine/Plugins/Experimental/PCG/Source/PCG/Private/Tests/Elements/PCGTransformPointsTest.cpp
+++ b/Engine/Plugins/Experimental/PCG/Source/PCG/Private/Tests/Elements/PCGTransformPointsTest.cpp
@@ -1,7 +1,7 @@
 // Copyright Epic Games, Inc. All Rights Reserved.
 
 #include "Data/PCGSpatialData.h"
-#include "Tests/Determinism/PCGDeterminismTestsCommon.h"
+#include "Tests/PCGTestsCommon.h"
 
 #include "PCGComponent.h"
 #include "PCGContext.h"
@@ -15,8 +15,8 @@ IMPLEMENT_CUSTOM_SIMPLE_AUTOMATION_TEST(FPCGTransformPointsTest, FPCGTestBaseCla
 
 bool FPCGTransformPointsTest::RunTest(const FString& Parameters)
 {
-	PCGTestsCommon::FTestData TestData(PCGDeterminismTests::Defaults::Seed);
-	PCGDeterminismTests::GenerateSettings<UPCGTransformPointsSettings>(TestData);
+	PCGTestsCommon::FTestData TestData;
+	PCGTestsCommon::GenerateSettings<UPCGTransformPointsSettings>(TestData);
 	UPCGTransformPointsSettings* Settings = CastChecked<UPCGTransformPointsSettings>(TestData.Settings);
 	FPCGElementPtr TransformPointsElement = TestData.Settings->GetElement();
 
diff --git a/Engine/Plugins/Experimental/PCG/Source/PCG/Public/Tests/Determinism/PCGDeterminismTestsCommon.h b/Engine/Plugins/Experimental/PCG/Source/PCG/Public/Tests/Determinism/PCGDeterminismTestsCommon.h
--- a/Engine/Plugins/Experimental/PCG/Source/PCG/Public/Tests/Determinism/PCGDeterminismTestsCommon.h
+++ b/Engine/Plugins/Experimental/PCG/Source/PCG/Public/Tests/Determinism/PCGDeterminismTestsCommon.h
@@ -70,7 +70,6 @@ namespace PCGDeterminismTests
 {
 	namespace Defaults
 	{
-		constexpr static int32 Seed = 42;
 		constexpr static int32 NumPointsToGenerate = 1;
 		constexpr static int32 NumTestPointsToGenerate = 100;
 		constexpr static int32 NumPolyLinePointsToGenerate = 6;
@@ -291,27 +290,6 @@ namespace PCGDeterminismTests
 	/** Execute the same element twice compare if all the outputs are at least orthogonally deterministic */
 	bool ExecutionIsDeterministicSameData(const PCGTestsCommon::FTestData& TestData, const UPCGNode* PCGNode = nullptr);
 
-	/** Generates settings based upon a UPCGSettings subclass */
-	template<typename SettingsType>
-	SettingsType* GenerateSettings(PCGTestsCommon::FTestData& TestData, TFunction<void(PCGTestsCommon::FTestData&)> ExtraSettingsDelegate = nullptr)
-	{
-		SettingsType* TypedSettings = NewObject<SettingsType>();
-		check(TypedSettings);
-
-		TestData.Settings = TypedSettings;
-		TestData.Settings->Seed = TestData.Seed;
-
-		TestData.InputData.TaggedData.Emplace_GetRef().Data = TestData.Settings;
-		TestData.InputData.TaggedData.Last().Pin = FName(TEXT("Settings"));
-
-		if (ExtraSettingsDelegate)
-		{
-			ExtraSettingsDelegate(TestData);
-		}
-
-		return TypedSettings;
-	}
-
 	/** Validates whether both UPCGData can be cast to a specified subclass */
 	template<typename DataType>
 	bool BothDataCastsToDataType(const UPCGData* FirstData, const UPCGData* SecondData)
diff --git a/Engine/Plugins/Experimental/PCG/Source/PCG/Public/Tests/PCGTestsCommon.h b/Engine/Plugins/Experimental/PCG/Source/PCG/Public/Tests/PCGTestsCommon.h
--- a/Engine/Plugins/Experimental/PCG/Source/PCG/Public/Tests/PCGTestsCommon.h
+++ b/Engine/Plugins/Experimental/PCG/Source/PCG/Public/Tests/PCGTestsCommon.h
@@ -3,6 +3,7 @@
 #pragma once
 
 #include "PCGData.h"
+#include "PCGSettings.h"
 
 #include "GameFramework/Actor.h"
 #include "Misc/AutomationTest.h"
@@ -14,7 +15,6 @@ class UPCGParamData;
 class UPCGPointData;
 class UPCGPolyLineData;
 class UPCGPrimitiveData;
-class UPCGSettings;
 class UPCGSurfaceData;
 class UPCGVolumeData;
 struct FPCGContext;
@@ -30,7 +30,7 @@ namespace PCGTestsCommon
 
 	struct PCG_API FTestData
 	{
-		FTestData(int32 Seed, UPCGSettings* DefaultSettings = nullptr, TSubclassOf<AActor> ActorClass = AActor::StaticClass());
+		explicit FTestData(int32 Seed = 42, UPCGSettings* DefaultSettings = nullptr, TSubclassOf<AActor> ActorClass = AActor::StaticClass());
 		~FTestData();
 
 		void Reset(UPCGSettings* InSettings = nullptr);
@@ -68,6 +68,27 @@ namespace PCGTestsCommon
 
 	/** Validates that two Spatial Points are identical */
 	PCG_API bool PointsAreIdentical(const FPCGPoint& FirstPoint, const FPCGPoint& SecondPoint);
+
+	/** Generates settings based upon a UPCGSettings subclass */
+	template<typename SettingsType>
+	SettingsType* GenerateSettings(FTestData& TestData, TFunction<void(FTestData&)> ExtraSettingsDelegate = nullptr)
+	{
+		SettingsType* TypedSettings = NewObject<SettingsType>();
+		check(TypedSettings);
+
+		TestData.Settings = TypedSettings;
+		TestData.Settings->Seed = TestData.Seed;
+
+		TestData.InputData.TaggedData.Emplace_GetRef().Data = TestData.Settings;
+		TestData.InputData.TaggedData.Last().Pin = FName(TEXT("Settings"));
+
+		if (ExtraSettingsDelegate)
+		{
+			ExtraSettingsDelegate(TestData);
+		}
+
+		return TypedSettings;
+	}
 }
 
 class PCG_API FPCGTestBaseClass : public FAutomationTestBase
