[PCG] Update Difference Node determinism test

Updated the order independence test to account for a change where running the difference of two PointData should now result in PointData, instead of DifferenceData

#rnx

diff --git a/Engine/Plugins/Experimental/PCG/Source/PCG/Private/Tests/Determinism/PCGDeterminismTestsCommon.cpp b/Engine/Plugins/Experimental/PCG/Source/PCG/Private/Tests/Determinism/PCGDeterminismTestsCommon.cpp
--- a/Engine/Plugins/Experimental/PCG/Source/PCG/Private/Tests/Determinism/PCGDeterminismTestsCommon.cpp
+++ b/Engine/Plugins/Experimental/PCG/Source/PCG/Private/Tests/Determinism/PCGDeterminismTestsCommon.cpp
@@ -1555,6 +1555,24 @@ namespace PCGDeterminismTests
 		return DataCollectionsContainSameData(FirstTestData.OutputData, SecondTestData.OutputData);
 	}
 
+	bool ExecutionIsConcretelyDeterministic(PCGTestsCommon::FTestData& FirstTestData, PCGTestsCommon::FTestData& SecondTestData, const UPCGNode* PCGNode)
+	{
+		ExecuteWithTestData(FirstTestData, PCGNode);
+		ExecuteWithTestData(SecondTestData, PCGNode);
+
+		for (FPCGTaggedData& TaggedData : FirstTestData.OutputData.TaggedData)
+		{
+			TaggedData.Data = CastChecked<UPCGSpatialData>(TaggedData.Data)->ToPointData();
+		}
+
+		for (FPCGTaggedData& TaggedData : SecondTestData.OutputData.TaggedData)
+		{
+			TaggedData.Data = CastChecked<UPCGSpatialData>(TaggedData.Data)->ToPointData();
+		}
+
+		return DataCollectionsContainSameData(FirstTestData.OutputData, SecondTestData.OutputData);
+	}
+
 	bool ExecutionIsDeterministicSameData(const PCGTestsCommon::FTestData& TestData, const UPCGNode* PCGNode)
 	{
 		FPCGDataCollection FirstOutput;
diff --git a/Engine/Plugins/Experimental/PCG/Source/PCG/Private/Tests/Determinism/PCGDifferenceDeterminismTest.cpp b/Engine/Plugins/Experimental/PCG/Source/PCG/Private/Tests/Determinism/PCGDifferenceDeterminismTest.cpp
--- a/Engine/Plugins/Experimental/PCG/Source/PCG/Private/Tests/Determinism/PCGDifferenceDeterminismTest.cpp
+++ b/Engine/Plugins/Experimental/PCG/Source/PCG/Private/Tests/Determinism/PCGDifferenceDeterminismTest.cpp
@@ -59,7 +59,33 @@ namespace
 		PCGDeterminismTests::AddVolumeInputData(TestData.InputData, FVector::ZeroVector, PCGDeterminismTests::Defaults::MediumVector, PCGDeterminismTests::Defaults::SmallVector, PCGDifferenceConstants::DifferencesLabel);
 	}
 
-	void DifferenceTestMultiple(PCGTestsCommon::FTestData& TestData, EPCGDifferenceMode DifferenceMode)
+	void DifferenceTestMultiplePoint(PCGTestsCommon::FTestData& TestData, EPCGDifferenceMode DifferenceMode)
+	{
+		DifferenceTestBase(TestData, DifferenceMode);
+
+		// Randomized Sources
+		PCGDeterminismTests::AddRandomizedMultiplePointInputData(TestData, 10, PCGDifferenceConstants::SourceLabel);
+		PCGDeterminismTests::AddRandomizedMultiplePointInputData(TestData, 20, PCGDifferenceConstants::SourceLabel);
+
+		// Randomized Differences
+		PCGDeterminismTests::AddRandomizedMultiplePointInputData(TestData, 10, PCGDifferenceConstants::DifferencesLabel);
+		PCGDeterminismTests::AddRandomizedMultiplePointInputData(TestData, 20, PCGDifferenceConstants::DifferencesLabel);
+	}
+
+	void DifferenceTestMultipleVolume(PCGTestsCommon::FTestData& TestData, EPCGDifferenceMode DifferenceMode)
+	{
+		DifferenceTestBase(TestData, DifferenceMode);
+
+		// Randomized Sources
+		PCGDeterminismTests::AddRandomizedVolumeInputData(TestData, PCGDifferenceConstants::SourceLabel);
+		PCGDeterminismTests::AddRandomizedVolumeInputData(TestData, PCGDifferenceConstants::SourceLabel);
+
+		// Randomized Differences
+		PCGDeterminismTests::AddRandomizedVolumeInputData(TestData, PCGDifferenceConstants::DifferencesLabel);
+		PCGDeterminismTests::AddRandomizedVolumeInputData(TestData, PCGDifferenceConstants::DifferencesLabel);
+	}
+
+	void DifferenceTestMultiplePointAndVolume(PCGTestsCommon::FTestData& TestData, EPCGDifferenceMode DifferenceMode)
 	{
 		DifferenceTestBase(TestData, DifferenceMode);
 
@@ -146,9 +172,9 @@ bool FPCGDifferenceDeterminismMultipleSameDataTest::RunTest(const FString& Param
 	PCGTestsCommon::FTestData TestDataContinuous;
 	PCGTestsCommon::FTestData TestDataDiscrete;
 
-	DifferenceTestMultiple(TestDataInferred, EPCGDifferenceMode::Inferred);
-	DifferenceTestMultiple(TestDataContinuous, EPCGDifferenceMode::Continuous);
-	DifferenceTestMultiple(TestDataDiscrete, EPCGDifferenceMode::Discrete);
+	DifferenceTestMultiplePointAndVolume(TestDataInferred, EPCGDifferenceMode::Inferred);
+	DifferenceTestMultiplePointAndVolume(TestDataContinuous, EPCGDifferenceMode::Continuous);
+	DifferenceTestMultiplePointAndVolume(TestDataDiscrete, EPCGDifferenceMode::Discrete);
 
 	TestTrue("Same single input, same output, 'Inferred' mode", PCGDeterminismTests::ExecutionIsDeterministicSameData(TestDataInferred));
 	TestTrue("Same single input, same output, 'Continuous' mode", PCGDeterminismTests::ExecutionIsDeterministicSameData(TestDataContinuous));
@@ -167,12 +193,12 @@ bool FPCGDifferenceDeterminismMultipleIdenticalDataTest::RunTest(const FString&
 	PCGTestsCommon::FTestData SecondTestDataContinuous;
 	PCGTestsCommon::FTestData SecondTestDataDiscrete;
 
-	DifferenceTestMultiple(FirstTestDataInferred, EPCGDifferenceMode::Inferred);
-	DifferenceTestMultiple(FirstTestDataContinuous, EPCGDifferenceMode::Continuous);
-	DifferenceTestMultiple(FirstTestDataDiscrete, EPCGDifferenceMode::Discrete);
-	DifferenceTestMultiple(SecondTestDataInferred, EPCGDifferenceMode::Inferred);
-	DifferenceTestMultiple(SecondTestDataContinuous, EPCGDifferenceMode::Continuous);
-	DifferenceTestMultiple(SecondTestDataDiscrete, EPCGDifferenceMode::Discrete);
+	DifferenceTestMultiplePointAndVolume(FirstTestDataInferred, EPCGDifferenceMode::Inferred);
+	DifferenceTestMultiplePointAndVolume(FirstTestDataContinuous, EPCGDifferenceMode::Continuous);
+	DifferenceTestMultiplePointAndVolume(FirstTestDataDiscrete, EPCGDifferenceMode::Discrete);
+	DifferenceTestMultiplePointAndVolume(SecondTestDataInferred, EPCGDifferenceMode::Inferred);
+	DifferenceTestMultiplePointAndVolume(SecondTestDataContinuous, EPCGDifferenceMode::Continuous);
+	DifferenceTestMultiplePointAndVolume(SecondTestDataDiscrete, EPCGDifferenceMode::Discrete);
 
 	TestTrue("Identical single input, same output, 'Inferred' mode", PCGDeterminismTests::ExecutionIsDeterministic(FirstTestDataInferred, SecondTestDataInferred));
 	TestTrue("Identical single input, same output, 'Continuous' mode", PCGDeterminismTests::ExecutionIsDeterministic(FirstTestDataContinuous, SecondTestDataContinuous));
@@ -191,20 +217,65 @@ bool FPCGDifferenceDeterminismOrderIndependenceTest::RunTest(const FString& Para
 	PCGTestsCommon::FTestData SecondTestDataContinuous;
 	PCGTestsCommon::FTestData SecondTestDataDiscrete;
 
-	DifferenceTestMultiple(FirstTestDataInferred, EPCGDifferenceMode::Inferred);
-	DifferenceTestMultiple(FirstTestDataContinuous, EPCGDifferenceMode::Continuous);
-	DifferenceTestMultiple(FirstTestDataDiscrete, EPCGDifferenceMode::Discrete);
-	DifferenceTestMultiple(SecondTestDataInferred, EPCGDifferenceMode::Inferred);
-	DifferenceTestMultiple(SecondTestDataContinuous, EPCGDifferenceMode::Continuous);
-	DifferenceTestMultiple(SecondTestDataDiscrete, EPCGDifferenceMode::Discrete);
+	DifferenceTestMultipleVolume(FirstTestDataInferred, EPCGDifferenceMode::Inferred);
+	DifferenceTestMultipleVolume(FirstTestDataContinuous, EPCGDifferenceMode::Continuous);
+	DifferenceTestMultipleVolume(FirstTestDataDiscrete, EPCGDifferenceMode::Discrete);
+	DifferenceTestMultipleVolume(SecondTestDataInferred, EPCGDifferenceMode::Inferred);
+	DifferenceTestMultipleVolume(SecondTestDataContinuous, EPCGDifferenceMode::Continuous);
+	DifferenceTestMultipleVolume(SecondTestDataDiscrete, EPCGDifferenceMode::Discrete);
+
+	PCGDeterminismTests::ShuffleInputOrder(SecondTestDataInferred);
+	PCGDeterminismTests::ShuffleInputOrder(SecondTestDataContinuous);
+	PCGDeterminismTests::ShuffleInputOrder(SecondTestDataDiscrete);
+
+	TestTrue("Volume Only - Shuffled input order, same output, 'Inferred' mode", PCGDeterminismTests::ExecutionIsDeterministic(FirstTestDataInferred, SecondTestDataInferred));
+	TestTrue("Volume Only - Shuffled input order, same output, 'Continuous' mode", PCGDeterminismTests::ExecutionIsDeterministic(FirstTestDataContinuous, SecondTestDataContinuous));
+	TestTrue("Volume Only - Shuffled input order, same output, 'Discrete' mode", PCGDeterminismTests::ExecutionIsDeterministic(FirstTestDataDiscrete, SecondTestDataDiscrete));
+
+	FirstTestDataInferred.Reset();
+	FirstTestDataContinuous.Reset();
+	FirstTestDataDiscrete.Reset();
+	SecondTestDataInferred.Reset();
+	SecondTestDataContinuous.Reset();
+	SecondTestDataDiscrete.Reset();
+
+	DifferenceTestMultiplePoint(FirstTestDataInferred, EPCGDifferenceMode::Inferred);
+	DifferenceTestMultiplePoint(FirstTestDataContinuous, EPCGDifferenceMode::Continuous);
+	DifferenceTestMultiplePoint(FirstTestDataDiscrete, EPCGDifferenceMode::Discrete);
+	DifferenceTestMultiplePoint(SecondTestDataInferred, EPCGDifferenceMode::Inferred);
+	DifferenceTestMultiplePoint(SecondTestDataContinuous, EPCGDifferenceMode::Continuous);
+	DifferenceTestMultiplePoint(SecondTestDataDiscrete, EPCGDifferenceMode::Discrete);
+
+	PCGDeterminismTests::ShuffleInputOrder(SecondTestDataInferred);
+	PCGDeterminismTests::ShuffleInputOrder(SecondTestDataContinuous);
+	PCGDeterminismTests::ShuffleInputOrder(SecondTestDataDiscrete);
+
+	TestTrue("Point Only - Shuffled input order, same output, 'Inferred' mode", PCGDeterminismTests::ExecutionIsDeterministic(FirstTestDataInferred, SecondTestDataInferred));
+	TestTrue("Point Only - Shuffled input order, same output, 'Continuous' mode", PCGDeterminismTests::ExecutionIsDeterministic(FirstTestDataContinuous, SecondTestDataContinuous));
+	TestTrue("Point Only - Shuffled input order, same output, 'Discrete' mode", PCGDeterminismTests::ExecutionIsDeterministic(FirstTestDataDiscrete, SecondTestDataDiscrete));
+
+	FirstTestDataInferred.Reset();
+	FirstTestDataContinuous.Reset();
+	FirstTestDataDiscrete.Reset();
+	SecondTestDataInferred.Reset();
+	SecondTestDataContinuous.Reset();
+	SecondTestDataDiscrete.Reset();
+
+	DifferenceTestMultiplePointAndVolume(FirstTestDataInferred, EPCGDifferenceMode::Inferred);
+	DifferenceTestMultiplePointAndVolume(FirstTestDataContinuous, EPCGDifferenceMode::Continuous);
+	DifferenceTestMultiplePointAndVolume(FirstTestDataDiscrete, EPCGDifferenceMode::Discrete);
+	DifferenceTestMultiplePointAndVolume(SecondTestDataInferred, EPCGDifferenceMode::Inferred);
+	DifferenceTestMultiplePointAndVolume(SecondTestDataContinuous, EPCGDifferenceMode::Continuous);
+	DifferenceTestMultiplePointAndVolume(SecondTestDataDiscrete, EPCGDifferenceMode::Discrete);
 
 	PCGDeterminismTests::ShuffleInputOrder(SecondTestDataInferred);
 	PCGDeterminismTests::ShuffleInputOrder(SecondTestDataContinuous);
 	PCGDeterminismTests::ShuffleInputOrder(SecondTestDataDiscrete);
 
-	TestTrue("Shuffled input order, same output, 'Inferred' mode", PCGDeterminismTests::ExecutionIsDeterministic(FirstTestDataInferred, SecondTestDataInferred));
-	TestTrue("Shuffled input order, same output, 'Continuous' mode", PCGDeterminismTests::ExecutionIsDeterministic(FirstTestDataContinuous, SecondTestDataContinuous));
-	TestTrue("Shuffled input order, same output, 'Discrete' mode", PCGDeterminismTests::ExecutionIsDeterministic(FirstTestDataDiscrete, SecondTestDataDiscrete));
+	// Since the difference node will return PointData - rather than DifferenceData during a PointData to PointData comparison, just convert all output into PointData
+	TestTrue("Point and Volume - Shuffled input order, same output, 'Inferred' mode", PCGDeterminismTests::ExecutionIsConcretelyDeterministic(FirstTestDataInferred, SecondTestDataInferred));
+	TestTrue("Point and Volume - Shuffled input order, same output, 'Continuous' mode", PCGDeterminismTests::ExecutionIsConcretelyDeterministic(FirstTestDataContinuous, SecondTestDataContinuous));
+	TestTrue("Point and Volume - Shuffled input order, same output, 'Discrete' mode", PCGDeterminismTests::ExecutionIsConcretelyDeterministic(FirstTestDataDiscrete, SecondTestDataDiscrete));
 
 	return true;
 }
diff --git a/Engine/Plugins/Experimental/PCG/Source/PCG/Public/Tests/Determinism/PCGDeterminismTestsCommon.h b/Engine/Plugins/Experimental/PCG/Source/PCG/Public/Tests/Determinism/PCGDeterminismTestsCommon.h
--- a/Engine/Plugins/Experimental/PCG/Source/PCG/Public/Tests/Determinism/PCGDeterminismTestsCommon.h
+++ b/Engine/Plugins/Experimental/PCG/Source/PCG/Public/Tests/Determinism/PCGDeterminismTestsCommon.h
@@ -287,6 +287,8 @@ namespace PCGDeterminismTests
 	void ExecuteWithSameTestDataSameElement(const PCGTestsCommon::FTestData& TestData, const UPCGNode* PCGNode, FPCGDataCollection& OutFirstOutputData, FPCGDataCollection& OutSecondOutputData);
 	/** Execute the elements for each valid input and compare if all the outputs are at least orthogonally deterministic */
 	bool ExecutionIsDeterministic(PCGTestsCommon::FTestData& FirstTestData, PCGTestsCommon::FTestData& SecondTestData, const UPCGNode* PCGNode = nullptr);
+	/** Execute the elements for each valid input, convert output to points, and compare if all the outputs are at least orthogonally deterministic */
+	bool ExecutionIsConcretelyDeterministic(PCGTestsCommon::FTestData& FirstTestData, PCGTestsCommon::FTestData& SecondTestData, const UPCGNode* PCGNode = nullptr);
 	/** Execute the same element twice compare if all the outputs are at least orthogonally deterministic */
 	bool ExecutionIsDeterministicSameData(const PCGTestsCommon::FTestData& TestData, const UPCGNode* PCGNode = nullptr);
 
