Revisited the delete pcg world actor so it removes all instances of the pcg world actor

diff --git a/Engine/Plugins/Experimental/PCG/Source/PCG/Private/PCGSubsystem.cpp b/Engine/Plugins/Experimental/PCG/Source/PCG/Private/PCGSubsystem.cpp
--- a/Engine/Plugins/Experimental/PCG/Source/PCG/Private/PCGSubsystem.cpp
+++ b/Engine/Plugins/Experimental/PCG/Source/PCG/Private/PCGSubsystem.cpp
@@ -72,6 +72,28 @@ namespace PCGSubsystemConsole
 		TEXT("pcg.PIE.RebuildLandscapeOnPIE"),
 		true,
 		TEXT("Controls whether the landscape cache will be rebuilt on PIE"));
+
+	static FAutoConsoleCommand CommandDeleteCurrentPCGWorldActor(
+		TEXT("pcg.DeleteCurrentPCGWorldActor"),
+		TEXT("Deletes the PCG World Actor currently registered to the PCG Subsystem."),
+		FConsoleCommandDelegate::CreateLambda([]()
+		{
+			if(UPCGSubsystem* PCGSubsystem = UPCGSubsystem::GetSubsystemForCurrentWorld())
+			{
+				PCGSubsystem->DestroyCurrentPCGWorldActor();
+			}
+		}));
+
+	static FAutoConsoleCommand CommandDeleteAllPCGWorldActors(
+		TEXT("pcg.DeleteAllPCGWorldActors"),
+		TEXT("Deletes all PCG World Actors in current World.."),
+		FConsoleCommandDelegate::CreateLambda([]()
+		{
+			if (UPCGSubsystem* PCGSubsystem = UPCGSubsystem::GetSubsystemForCurrentWorld())
+			{
+				PCGSubsystem->DestroyAllPCGWorldActors();
+			}
+		}));
 #endif
 }
 
@@ -235,7 +257,7 @@ APCGWorldActor* UPCGSubsystem::FindPCGWorldActor()
 }
 
 #if WITH_EDITOR
-void UPCGSubsystem::DestroyPCGWorldActor()
+void UPCGSubsystem::DestroyCurrentPCGWorldActor()
 {
 	if (PCGWorldActor)
 	{
@@ -245,6 +267,30 @@ void UPCGSubsystem::DestroyPCGWorldActor()
 		PCGWorldActorLock.Unlock();
 	}
 }
+
+void UPCGSubsystem::DestroyAllPCGWorldActors()
+{
+	// Get rid of current PCG world actor first
+	DestroyCurrentPCGWorldActor();
+	
+	// Pick up any strays in the current world
+	TArray<APCGWorldActor*> ActorsToDestroy;
+	ForEachObjectWithOuter(GetWorld(), [&ActorsToDestroy](UObject* Object)
+	{
+		if (APCGWorldActor* WorldActor = Cast<APCGWorldActor>(Object))
+		{
+			if (IsValid(WorldActor))
+			{
+				ActorsToDestroy.Add(WorldActor);
+			}
+		}
+	});
+
+	for (APCGWorldActor* ActorToDestroy : ActorsToDestroy)
+	{
+		ActorToDestroy->Destroy();
+	}
+}
 #endif
 
 void UPCGSubsystem::RegisterPCGWorldActor(APCGWorldActor* InActor)
diff --git a/Engine/Plugins/Experimental/PCG/Source/PCG/Public/PCGSubsystem.h b/Engine/Plugins/Experimental/PCG/Source/PCG/Public/PCGSubsystem.h
--- a/Engine/Plugins/Experimental/PCG/Source/PCG/Public/PCGSubsystem.h
+++ b/Engine/Plugins/Experimental/PCG/Source/PCG/Public/PCGSubsystem.h
@@ -78,7 +78,8 @@ public:
 	APCGWorldActor* GetPCGWorldActor();
 	APCGWorldActor* FindPCGWorldActor();
 #if WITH_EDITOR
-	void DestroyPCGWorldActor();
+	void DestroyAllPCGWorldActors();
+	void DestroyCurrentPCGWorldActor();
 #endif
 	void RegisterPCGWorldActor(APCGWorldActor* InActor);
 	void UnregisterPCGWorldActor(APCGWorldActor* InActor);
diff --git a/Engine/Plugins/Experimental/PCG/Source/PCGEditor/Private/PCGEditorModule.cpp b/Engine/Plugins/Experimental/PCG/Source/PCGEditor/Private/PCGEditorModule.cpp
--- a/Engine/Plugins/Experimental/PCG/Source/PCGEditor/Private/PCGEditorModule.cpp
+++ b/Engine/Plugins/Experimental/PCG/Source/PCGEditor/Private/PCGEditorModule.cpp
@@ -191,8 +191,8 @@ void FPCGEditorModule::PopulateMenuActions(FMenuBuilder& MenuBuilder)
 		NAME_None);
 
 	MenuBuilder.AddMenuEntry(
-		LOCTEXT("DeletePCGWorldActor", "Delete PCG World Actor"),
-		LOCTEXT("DeletePCGWorldActor_Tooltip", "Deletes the PCG World Actor"),
+		LOCTEXT("DeletePCGWorldActor", "Deletes all PCG World Actors"),
+		LOCTEXT("DeletePCGWorldActor_Tooltip", "Deletes all PCG World Actors"),
 		FSlateIcon(),
 		FUIAction(
 			FExecuteAction::CreateLambda([]() {
@@ -200,10 +200,7 @@ void FPCGEditorModule::PopulateMenuActions(FMenuBuilder& MenuBuilder)
 				{
 					if (UPCGSubsystem* PCGSubsystem = UPCGSubsystem::GetInstance(GEditor->GetEditorWorldContext().World()))
 					{
-						if (APCGWorldActor* PCGWorldActor = PCGSubsystem->GetPCGWorldActor())
-						{
-							PCGSubsystem->DestroyPCGWorldActor();
-						}
+						PCGSubsystem->DestroyAllPCGWorldActors();
 					}
 				}
 			})),
