Added some robustness around data chunks that couldn't be loaded

diff --git a/Engine/Plugins/Experimental/PCG/Source/PCG/Private/Grid/PCGLandscapeCache.cpp b/Engine/Plugins/Experimental/PCG/Source/PCG/Private/Grid/PCGLandscapeCache.cpp
--- a/Engine/Plugins/Experimental/PCG/Source/PCG/Private/Grid/PCGLandscapeCache.cpp
+++ b/Engine/Plugins/Experimental/PCG/Source/PCG/Private/Grid/PCGLandscapeCache.cpp
@@ -323,8 +323,7 @@ bool FPCGLandscapeCacheEntry::TouchAndLoad(int32 InTouch) const
 		FScopeLock ScopeDataLock(&DataLock);
 		if (!bDataLoaded)
 		{
-			SerializeFromBulkData();
-			return true;
+			return SerializeFromBulkData();
 		}
 	}
 
@@ -398,7 +397,7 @@ void FPCGLandscapeCacheEntry::SerializeToBulkData(EPCGLandscapeCacheSerializatio
 	BulkData.Unlock();
 }
 
-void FPCGLandscapeCacheEntry::SerializeFromBulkData() const
+bool FPCGLandscapeCacheEntry::SerializeFromBulkData() const
 {
 	check(!bDataLoaded);
 
@@ -407,6 +406,12 @@ void FPCGLandscapeCacheEntry::SerializeFromBulkData() const
 	BulkData.GetCopy((void**)&Data);
 	int32 DataSize = BulkData.GetBulkDataSize();
 
+	if (!Data)
+	{
+		UE_LOG(LogPCG, Error, TEXT("Unable to load Landscape Cache Entry bulk data"));
+		return false;
+	}
+
 	FBufferReader Ar(Data, DataSize, /*bInFreeOnClose=*/true, /*bIsPersistent=*/true);
 
 	Ar << PositionsAndNormals;
@@ -422,6 +427,7 @@ void FPCGLandscapeCacheEntry::SerializeFromBulkData() const
 	}
 
 	bDataLoaded = true;
+	return true;
 }
 
 void FPCGLandscapeCacheEntry::Serialize(FArchive& Archive, UObject* Owner, int32 Index, EPCGLandscapeCacheSerializationContents SerializeContents)
@@ -748,6 +754,10 @@ const FPCGLandscapeCacheEntry* UPCGLandscapeCache::GetCacheEntry(ULandscapeCompo
 			{
 				CacheMemorySize += CacheEntry->GetMemorySize();
 			}
+			else if(!CacheEntry->bDataLoaded)
+			{
+				CacheEntry = nullptr;
+			}
 		}
 	}
 
@@ -776,6 +786,10 @@ const FPCGLandscapeCacheEntry* UPCGLandscapeCache::GetCacheEntry(const FGuid& La
 		{
 			CacheMemorySize += CacheEntry->GetMemorySize();
 		}
+		else if(!CacheEntry->bDataLoaded)
+		{
+			CacheEntry = nullptr;
+		}
 	}
 
 	return CacheEntry;
diff --git a/Engine/Plugins/Experimental/PCG/Source/PCG/Public/Grid/PCGLandscapeCache.h b/Engine/Plugins/Experimental/PCG/Source/PCG/Public/Grid/PCGLandscapeCache.h
--- a/Engine/Plugins/Experimental/PCG/Source/PCG/Public/Grid/PCGLandscapeCache.h
+++ b/Engine/Plugins/Experimental/PCG/Source/PCG/Public/Grid/PCGLandscapeCache.h
@@ -90,7 +90,7 @@ private:
 
 	// Internal usage methods
 	void SerializeToBulkData(EPCGLandscapeCacheSerializationContents SerializationContents);
-	void SerializeFromBulkData() const;
+	bool SerializeFromBulkData() const;
 
 	// Serialized data
 	TArray<FName> LayerDataNames;
