[PCG] Disable threaded execution of SM Spawner PrepareData

#rnx

diff --git a/Engine/Plugins/Experimental/PCG/Source/PCG/Private/Elements/PCGStaticMeshSpawner.cpp b/Engine/Plugins/Experimental/PCG/Source/PCG/Private/Elements/PCGStaticMeshSpawner.cpp
--- a/Engine/Plugins/Experimental/PCG/Source/PCG/Private/Elements/PCGStaticMeshSpawner.cpp
+++ b/Engine/Plugins/Experimental/PCG/Source/PCG/Private/Elements/PCGStaticMeshSpawner.cpp
@@ -312,7 +312,10 @@ FPCGContext* FPCGStaticMeshSpawnerElement::Initialize(const FPCGDataCollection&
 
 bool FPCGStaticMeshSpawnerElement::CanExecuteOnlyOnMainThread(FPCGContext* Context) const
 {
-	return Context->CurrentPhase == EPCGExecutionPhase::Execute;
+	// PrepareData can call UPCGManagedComponent::MarkAsReused which registers the ISMC, which can go into Chaos code that asserts if not on main thread.
+	// TODO: We can likely re-enable multi-threading for PrepareData if we move the call to MarkAsReused to Execute. There should hopefully not be
+	// wider contention on resources resources are not shared across nodes and are also per-component.
+	return Context->CurrentPhase == EPCGExecutionPhase::Execute || Context->CurrentPhase == EPCGExecutionPhase::PrepareData;
 }
 
 void FPCGStaticMeshSpawnerElement::SpawnStaticMeshInstances(FPCGContext* Context, const FPCGMeshInstanceList& InstanceList, AActor* TargetActor, const FPCGPackedCustomData& PackedCustomData) const
diff --git a/Engine/Plugins/Experimental/PCG/Source/PCG/Private/PCGSubsystem.cpp b/Engine/Plugins/Experimental/PCG/Source/PCG/Private/PCGSubsystem.cpp
--- a/Engine/Plugins/Experimental/PCG/Source/PCG/Private/PCGSubsystem.cpp
+++ b/Engine/Plugins/Experimental/PCG/Source/PCG/Private/PCGSubsystem.cpp
@@ -391,7 +391,6 @@ FPCGTaskId UPCGSubsystem::ScheduleComponent(UPCGComponent* PCGComponent, bool bS
 				LocalComponent->CleanupLocalImmediate(true);
 			}
 
-			// TODO If i don't forward the force flag, then nothing gets scheduled. Also why is it called bSave?
 			return LocalComponent->GenerateInternal(/*bForce=*/bSave, EPCGComponentGenerationTrigger::GenerateOnDemand, Dependencies);
 		};
 
