[PCG] Modernize attribute reduce

* Quat average is now correctly approximated. Still an approximation, that requires the quat to be close to each other.
* Add "predefined" configs for average,min,max
* Change input type to Any to support any type that has attributes
* Support multi input
* Add option to chose to have seperate param data or single param data with multiple entries.
* Add unit tests
* Small update to transform test to remove anonym namespace, breaking the build.

#rnx

diff --git a/Engine/Plugins/Experimental/PCG/Source/PCG/Private/Elements/PCGAttributeReduceElement.cpp b/Engine/Plugins/Experimental/PCG/Source/PCG/Private/Elements/PCGAttributeReduceElement.cpp
--- a/Engine/Plugins/Experimental/PCG/Source/PCG/Private/Elements/PCGAttributeReduceElement.cpp
+++ b/Engine/Plugins/Experimental/PCG/Source/PCG/Private/Elements/PCGAttributeReduceElement.cpp
@@ -26,19 +26,51 @@ namespace PCGAttributeReduceElement
 		}
 		else
 		{
-			OutValue = PCG::Private::MetadataTraits<T>::ZeroValue();
+			if constexpr (std::is_same_v<FQuat, T>)
+			{
+				// Simple averaging will only work for similar quaternions, but the real answer needs a full solver with eigenvectors.
+				FQuat ZeroQuat{ 0.0, 0.0, 0.0, 0.0 };
+				OutValue = ZeroQuat;
 
-			bool bSuccess = PCGMetadataElementCommon::ApplyOnAccessor<T>(Keys, Accessor, [&OutValue](const T& InValue, int32)
+				FQuat FirstQuat = ZeroQuat;
+				const double Weight = 1.0f / Keys.GetNum();
+
+				bool bSuccess = PCGMetadataElementCommon::ApplyOnAccessor<T>(Keys, Accessor, [&FirstQuat, Weight, &OutValue, &ZeroQuat](const T& InValue, int32 Index)
 				{
-					OutValue = PCG::Private::MetadataTraits<T>::Add(OutValue, InValue);
+					// Since doing a dot product with a quat equals to 0 will give us no info, we'll take the first non null quat.
+					if (FirstQuat.Equals(ZeroQuat))
+					{
+						FirstQuat = InValue;
+					}
+
+					// Because q and -q represent the same quaternion, but the average would not be correct. We need to inverse the quaternion if needed.
+					const double ThisWeight = (InValue | FirstQuat) < 0 ? -Weight : Weight;
+					OutValue += (ThisWeight * InValue);
 				});
 
-			if (bSuccess)
-			{
-				OutValue = PCG::Private::MetadataTraits<T>::WeightedSum(PCG::Private::MetadataTraits<T>::ZeroValue(), OutValue, 1.0f / Keys.GetNum());
+				if (bSuccess)
+				{
+					OutValue.Normalize();
+				}
+
+				return bSuccess;
 			}
+			else
+			{
+				OutValue = PCG::Private::MetadataTraits<T>::ZeroValue();
 
-			return bSuccess;
+				bool bSuccess = PCGMetadataElementCommon::ApplyOnAccessor<T>(Keys, Accessor, [&OutValue](const T& InValue, int32)
+				{
+					OutValue = PCG::Private::MetadataTraits<T>::Add(OutValue, InValue);
+				});
+
+				if (bSuccess)
+				{
+					OutValue = PCG::Private::MetadataTraits<T>::WeightedSum(PCG::Private::MetadataTraits<T>::ZeroValue(), OutValue, 1.0f / Keys.GetNum());
+				}
+
+				return bSuccess;
+			}
 		}
 	}
 
@@ -100,8 +132,24 @@ void UPCGAttributeReduceSettings::ApplyDeprecation(UPCGNode* InOutNode)
 
 	Super::ApplyDeprecation(InOutNode);
 }
+
+TArray<FPCGPreConfiguredSettingsInfo> UPCGAttributeReduceSettings::GetPreconfiguredInfo() const
+{
+	return PCGMetadataElementCommon::FillPreconfiguredSettingsInfoFromEnum<EPCGAttributeReduceOperation>();
+}
 #endif
 
+void UPCGAttributeReduceSettings::ApplyPreconfiguredSettings(const FPCGPreConfiguredSettingsInfo& PreconfiguredInfo)
+{
+	if (const UEnum* EnumPtr = StaticEnum<EPCGAttributeReduceOperation>())
+	{
+		if (EnumPtr->IsValidEnumValue(PreconfiguredInfo.PreconfiguredIndex))
+		{
+			Operation = EPCGAttributeReduceOperation(PreconfiguredInfo.PreconfiguredIndex);
+		}
+	}
+}
+
 void UPCGAttributeReduceSettings::PostLoad()
 {
 	Super::PostLoad();
@@ -145,7 +193,7 @@ FName UPCGAttributeReduceSettings::AdditionalTaskName() const
 TArray<FPCGPinProperties> UPCGAttributeReduceSettings::InputPinProperties() const
 {
 	TArray<FPCGPinProperties> PinProperties;
-	PinProperties.Emplace(PCGPinConstants::DefaultInputLabel, EPCGDataType::Spatial, /*bInAllowMultipleConnections=*/ false);
+	PinProperties.Emplace(PCGPinConstants::DefaultInputLabel, EPCGDataType::Any);
 
 	return PinProperties;
 }
@@ -173,87 +221,92 @@ bool FPCGAttributeReduceElement::ExecuteInternal(FPCGContext* Context) const
 	check(Settings);
 
 	TArray<FPCGTaggedData> Inputs = Context->InputData.GetInputsByPin(PCGPinConstants::DefaultInputLabel);
+	TArray<FPCGTaggedData>& Outputs = Context->OutputData.TaggedData;
+	UPCGParamData* OutputParams = nullptr;
+	FPCGMetadataAttributeBase* NewAttribute = nullptr;
 
-	if (Inputs.Num() != 1)
+	for (int32 i = 0; i < Inputs.Num(); ++i)
 	{
-		PCGE_LOG(Warning, LogOnly, FText::Format(LOCTEXT("WrongNumberOfInputs", "Input pin expected to have one input data element, encountered {0}"), Inputs.Num()));
-		return true;
-	}
+		const UPCGData* InputData = Inputs[i].Data;
 
-	const UPCGSpatialData* SpatialData = Cast<UPCGSpatialData>(Inputs[0].Data);
+		if (!InputData)
+		{
+			PCGE_LOG(Warning, GraphAndLog, FText::Format(LOCTEXT("InputNotSpatialData", "Input {0} is invalid, skipped"), FText::AsNumber(i)));
+			continue;
+		}
 
-	if (!SpatialData)
-	{
-		PCGE_LOG(Error, GraphAndLog, LOCTEXT("InputNotSpatialData", "Input is not a spatial data"));
-		return true;
-	}
+		FPCGAttributePropertyInputSelector InputSource = Settings->InputSource.CopyAndFixLast(InputData);
 
-	const UPCGPointData* PointData = Cast<UPCGPointData>(SpatialData);
+		const FName OutputAttributeName = (Settings->OutputAttributeName == PCGMetadataAttributeConstants::SourceNameAttributeName) ? InputSource.GetName() : Settings->OutputAttributeName;
 
-	FPCGAttributePropertyInputSelector InputSource = Settings->InputSource.CopyAndFixLast(SpatialData);
+		TUniquePtr<const IPCGAttributeAccessor> Accessor = PCGAttributeAccessorHelpers::CreateConstAccessor(InputData, InputSource);
+		TUniquePtr<const IPCGAttributeAccessorKeys> Keys = PCGAttributeAccessorHelpers::CreateConstKeys(InputData, InputSource);
 
-	const FName OutputAttributeName = (Settings->OutputAttributeName == PCGMetadataAttributeConstants::SourceNameAttributeName) ? InputSource.GetName() : Settings->OutputAttributeName;
-	UPCGParamData* OutputParamData = NewObject<UPCGParamData>();
+		if (!Accessor.IsValid() || !Keys.IsValid())
+		{
+			PCGE_LOG(Warning, GraphAndLog, FText::Format(LOCTEXT("AttributeDoesNotExist", "Input attribute/property '{0}' does not exist on input {1}, skipped"), InputSource.GetDisplayText(), FText::AsNumber(i)));
+			continue;
+		}
 
-	TUniquePtr<const IPCGAttributeAccessor> Accessor = PCGAttributeAccessorHelpers::CreateConstAccessor(PointData, InputSource);
-	TUniquePtr<const IPCGAttributeAccessorKeys> Keys = PCGAttributeAccessorHelpers::CreateConstKeys(PointData, InputSource);
+		auto DoOperation = [&Accessor, &Keys, Operation = Settings->Operation, bMergeOutputAttributes = Settings->bMergeOutputAttributes, &OutputParams, &NewAttribute, OutputAttributeName](auto DummyValue) -> bool
+		{
+			using AttributeType = decltype(DummyValue);
 
-	if (!Accessor.IsValid() || !Keys.IsValid())
-	{
-		PCGE_LOG(Error, GraphAndLog, LOCTEXT("AttributeDoesNotExist", "Input attribute/property does not exist"));
-		return true;
-	}
+			bool bSuccess = false;
 
-	auto DoOperation = [&Accessor, &Keys, Operation = Settings->Operation, OutputParamData, OutputAttributeName](auto DummyValue) -> bool
-	{
-		using AttributeType = decltype(DummyValue);
+			AttributeType OutputValue = PCG::Private::MetadataTraits<AttributeType>::ZeroValue();
 
-		bool bSuccess = false;
+			if (!OutputParams || !bMergeOutputAttributes)
+			{
+				OutputParams = NewObject<UPCGParamData>();
+				NewAttribute = OutputParams->Metadata->CreateAttribute<AttributeType>(OutputAttributeName, OutputValue, /*bAllowInterpolation=*/ true, /*bOverrideParent=*/false);
 
-		AttributeType OutputValue{};
+				if (!NewAttribute)
+				{
+					OutputParams = nullptr;
+					return false;
+				}
+			}
 
-		FPCGMetadataAttribute<AttributeType>* NewAttribute = static_cast<FPCGMetadataAttribute<AttributeType>*>(
-			OutputParamData->Metadata->CreateAttribute<AttributeType>(OutputAttributeName, OutputValue, /*bAllowInterpolation=*/ true, /*bOverrideParent=*/false));
+			FPCGMetadataAttribute<AttributeType>* TypedNewAttribute = static_cast<FPCGMetadataAttribute<AttributeType>*>(NewAttribute);
+			check(TypedNewAttribute);
 
-		if (!NewAttribute)
-		{
-			return false;
-		}
+			switch (Operation)
+			{
+			case EPCGAttributeReduceOperation::Average:
+				bSuccess = PCGAttributeReduceElement::Average<AttributeType>(*Keys, *Accessor, OutputValue);
+				break;
+			case EPCGAttributeReduceOperation::Max:
+				bSuccess = PCGAttributeReduceElement::MinMax<AttributeType, /*bIsMin*/false>(*Keys, *Accessor, OutputValue);
+				break;
+			case EPCGAttributeReduceOperation::Min:
+				bSuccess = PCGAttributeReduceElement::MinMax<AttributeType, /*bIsMin*/true>(*Keys, *Accessor, OutputValue);
+				break;
+			default:
+				break;
+			}
+
+			if (bSuccess)
+			{
+				TypedNewAttribute->SetValue(OutputParams->Metadata->AddEntry(), OutputValue);
+			}
 
-		switch (Operation)
+			return bSuccess;
+		};
+
+		if (!PCGMetadataAttribute::CallbackWithRightType(Accessor->GetUnderlyingType(), DoOperation))
 		{
-		case EPCGAttributeReduceOperation::Average:
-			bSuccess = PCGAttributeReduceElement::Average<AttributeType>(*Keys, *Accessor, OutputValue);
-			break;
-		case EPCGAttributeReduceOperation::Max:
-			bSuccess = PCGAttributeReduceElement::MinMax<AttributeType, /*bIsMin*/false>(*Keys, *Accessor, OutputValue);
-			break;
-		case EPCGAttributeReduceOperation::Min:
-			bSuccess = PCGAttributeReduceElement::MinMax<AttributeType, /*bIsMin*/true>(*Keys, *Accessor, OutputValue);
-			break;
-		default:
-			break;
+			PCGE_LOG(Error, GraphAndLog, FText::Format(LOCTEXT("AttributeOperationFailed", "Operation was not compatible with the attribute type {0} or could not create attribute '{1}' for input {2}"), PCG::Private::GetTypeNameText(Accessor->GetUnderlyingType()), FText::FromName(OutputAttributeName), FText::AsNumber(i)));
+			continue;
 		}
 
-		if (bSuccess)
+		if (ensure(OutputParams) && (Outputs.IsEmpty() || !Settings->bMergeOutputAttributes))
 		{
-			NewAttribute->SetDefaultValue(OutputValue);
-			NewAttribute->SetValue(OutputParamData->Metadata->AddEntry(), OutputValue);
+			FPCGTaggedData& Output = Outputs.Emplace_GetRef();
+			Output.Data = OutputParams;
 		}
-
-		return bSuccess;
-	};
-
-	if (!PCGMetadataAttribute::CallbackWithRightType(Accessor->GetUnderlyingType(), DoOperation))
-	{
-		PCGE_LOG(Error, GraphAndLog, FText::Format(LOCTEXT("AttributeOperationFailed", "Operation was not compatible with the attribute type or could not create attribute '{0}'"), FText::FromName(OutputAttributeName)));
-		return true;
 	}
 
-	TArray<FPCGTaggedData>& Outputs = Context->OutputData.TaggedData;
-	FPCGTaggedData& Output = Outputs.Emplace_GetRef();
-	Output.Data = OutputParamData;
-
 	return true;
 }
 
diff --git a/Engine/Plugins/Experimental/PCG/Source/PCG/Private/Tests/Elements/Metadata/PCGMetadataMakeTransformTest.cpp b/Engine/Plugins/Experimental/PCG/Source/PCG/Private/Tests/Elements/Metadata/PCGMetadataMakeTransformTest.cpp
--- a/Engine/Plugins/Experimental/PCG/Source/PCG/Private/Tests/Elements/Metadata/PCGMetadataMakeTransformTest.cpp
+++ b/Engine/Plugins/Experimental/PCG/Source/PCG/Private/Tests/Elements/Metadata/PCGMetadataMakeTransformTest.cpp
@@ -15,7 +15,7 @@
 
 IMPLEMENT_CUSTOM_SIMPLE_AUTOMATION_TEST(FPCGMetadataMakeTransformTest, FPCGTestBaseClass, "pcg.tests.Metadata.MakeTransform", PCGTestsCommon::TestFlags)
 
-namespace
+namespace PCGMetadataMakeTransformTest
 {
 	const FName Vec3Attribute = TEXT("Vec3");
 	const FName QuatAttribute = TEXT("Quat");
@@ -214,11 +214,11 @@ bool FPCGMetadataMakeTransformTest::RunTest(const FString& Parameters)
 
 	PCGTestsCommon::FTestData TestDataSpatial;
 	PCGTestsCommon::GenerateSettings<UPCGMetadataMakeTransformSettings>(TestDataSpatial);
-	GenerateSpatialData(TestDataSpatial);
+	PCGMetadataMakeTransformTest::GenerateSpatialData(TestDataSpatial);
 
 	PCGTestsCommon::FTestData TestDataParams;
 	PCGTestsCommon::GenerateSettings<UPCGMetadataMakeTransformSettings>(TestDataParams);
-	GenerateParamData(TestDataParams);
+	PCGMetadataMakeTransformTest::GenerateParamData(TestDataParams);
 
 	using PairWhatData = TPair<FString, PCGTestsCommon::FTestData*>;
 
@@ -235,9 +235,9 @@ bool FPCGMetadataMakeTransformTest::RunTest(const FString& Parameters)
 
 		Settings->ForceOutputConnections[0] = true;
 
-		Settings->InputSource1.SetAttributeName(Vec3Attribute);
-		Settings->InputSource2.SetAttributeName(QuatAttribute);
-		Settings->InputSource3.SetAttributeName(Vec3Attribute);
+		Settings->InputSource1.SetAttributeName(PCGMetadataMakeTransformTest::Vec3Attribute);
+		Settings->InputSource2.SetAttributeName(PCGMetadataMakeTransformTest::QuatAttribute);
+		Settings->InputSource3.SetAttributeName(PCGMetadataMakeTransformTest::Vec3Attribute);
 
 		bTestPassed &= ValidateMetadataMakeVector(*TestData);
 	}
diff --git a/Engine/Plugins/Experimental/PCG/Source/PCG/Private/Tests/Elements/PCGAttributeReduceTest.cpp b/Engine/Plugins/Experimental/PCG/Source/PCG/Private/Tests/Elements/PCGAttributeReduceTest.cpp
new file mode 100644
--- /dev/null
+++ b/Engine/Plugins/Experimental/PCG/Source/PCG/Private/Tests/Elements/PCGAttributeReduceTest.cpp
@@ -0,0 +1,362 @@
+// Copyright Epic Games, Inc. All Rights Reserved.
+
+#include "Tests/PCGTestsCommon.h"
+
+#include "PCGContext.h"
+#include "PCGComponent.h"
+#include "PCGParamData.h"
+#include "Data/PCGPointData.h"
+#include "Elements/PCGAttributeReduceElement.h"
+#include "Metadata/PCGMetadata.h"
+#include "Metadata/PCGMetadataAttributeTpl.h"
+
+class FPCGAttributeReduceTests : public FPCGTestBaseClass
+{
+	using FPCGTestBaseClass::FPCGTestBaseClass;
+
+protected:
+	template <typename... T>
+	bool RunAndValidate(TArray<FPCGTaggedData> InData, EPCGAttributeReduceOperation Operation, FPCGAttributePropertyInputSelector Selector, bool bMerged, T&& ...ExpectedValues)
+	{
+		AddInfo(*FString::Printf(TEXT("Test with '%s' as input attribute"), *Selector.GetDisplayText().ToString()));
+
+		PCGTestsCommon::FTestData TestData;
+		UPCGAttributeReduceSettings* Settings = PCGTestsCommon::GenerateSettings<UPCGAttributeReduceSettings>(TestData);
+		FPCGElementPtr Element = TestData.Settings->GetElement();
+
+		constexpr int32 NumExpectedValues = sizeof...(ExpectedValues);
+
+		for (FPCGTaggedData& InTaggedData : InData)
+		{
+			InTaggedData.Pin = PCGPinConstants::DefaultInputLabel;
+		}
+
+		TestData.InputData.TaggedData.Append(std::move(InData));
+		Settings->InputSource = std::move(Selector);
+		Settings->bMergeOutputAttributes = bMerged;
+		Settings->Operation = Operation;
+		Settings->OutputAttributeName = NAME_None;
+
+		TUniquePtr<FPCGContext> Context = TestData.InitializeTestContext();
+
+		while (!Element->Execute(Context.Get())) {}
+
+		const TArray<FPCGTaggedData>& Outputs = Context->OutputData.TaggedData;
+
+		// If it is merged, we should have just 1 output with N values
+		// If not merge, N outputs with 1 value
+		if (bMerged)
+		{
+			UTEST_EQUAL("Merged output has a single value.", Outputs.Num(), 1);
+		}
+		else
+		{
+			UTEST_EQUAL("Not Merged output is the same number of inputs.", Outputs.Num(), NumExpectedValues);
+		}
+
+		// All outputs should be param data
+		TArray<const UPCGParamData*> AllParams;
+		AllParams.Reserve(Outputs.Num());
+		for (int32 i = 0; i < Outputs.Num(); ++i)
+		{
+			const UPCGParamData* Param = Cast<const UPCGParamData>(Outputs[i].Data);
+			UTEST_NOT_NULL(*FString::Printf(TEXT("Output %d is a param data"), i), Param);
+			AllParams.Add(Param);
+		}
+
+		// Finally validate all values matches, use a fold expression with Index noting the current param to check or entry to check
+		int32 Index = 0;
+		bool bSuccess = true;
+		([&Index, this, &bSuccess, bMerged, &AllParams, &ExpectedValues]()
+		{
+			if (!bSuccess)
+			{
+				return;
+			}
+
+			const FPCGMetadataAttribute<T>* Attribute = (bMerged ? AllParams[0] : AllParams[Index])->Metadata->GetConstTypedAttribute<T>(NAME_None);
+			bSuccess = TestNotNull(TEXT("Attribute 'None' exists and is the right type."), Attribute);
+
+			if (bSuccess)
+			{
+				check(Attribute);
+				T ComputedValue = Attribute->GetValueFromItemKey(bMerged ? Index : 0);
+				if constexpr (std::is_same_v<T, FVector2D>)
+				{
+					bSuccess = TestTrue(*FString::Printf(TEXT("Value for index %d matches. {%f, %f} vs {%f, %f}"), Index, ComputedValue.X, ComputedValue.Y, ExpectedValues.X, ExpectedValues.Y), ComputedValue.Equals(ExpectedValues));
+				}
+				else
+				{
+					bSuccess = TestEqual(*FString::Printf(TEXT("Value for index %d matches."), Index), ComputedValue, ExpectedValues);
+				}
+			}
+
+			Index++;
+		}(), ...);
+
+		return bSuccess;
+	}
+
+	TArray<FPCGTaggedData> PreparePointsAndSelectors(FPCGAttributePropertyInputSelector& DensitySelector, FPCGAttributePropertyInputSelector& PositionSelector)
+	{
+		UPCGPointData* PointData = PCGTestsCommon::CreateEmptyPointData();
+		TArray<FPCGPoint>& Points = PointData->GetMutablePoints();
+
+		constexpr int32 PointCount = 5;
+		for (int32 i = 0; i < PointCount; ++i)
+		{
+			const FVector Position{ 3.0 * i, 3.0 * i + 1.0, 3.0 * i + 2.0 };
+			FPCGPoint& Point = Points.Emplace_GetRef(FTransform(Position), 0.1 * (i + 1), i);
+		}
+
+		TArray<FPCGTaggedData> Inputs;
+		FPCGTaggedData& Data = Inputs.Emplace_GetRef();
+		Data.Data = PointData;
+
+		DensitySelector.SetPointProperty(EPCGPointProperties::Density);
+		PositionSelector.SetPointProperty(EPCGPointProperties::Position);
+
+		return Inputs;
+	}
+
+	TArray<FPCGTaggedData> PrepareParamAndSelectors(FPCGAttributePropertyInputSelector& IntSelector, FPCGAttributePropertyInputSelector& Vec2Selector)
+	{
+		const FName IntName = TEXT("int");
+		const FName Vec2Name = TEXT("vec2");
+
+		UPCGParamData* ParamData = NewObject<UPCGParamData>();
+		FPCGMetadataAttribute<int32>* IntAttribute = ParamData->Metadata->CreateAttribute<int32>(IntName, 0, /*bAllowsInterpolation=*/ true, /*bOverrideParent=*/false);
+		FPCGMetadataAttribute<FVector2D>* Vec2Attribute = ParamData->Metadata->CreateAttribute<FVector2D>(Vec2Name, FVector2D{}, /*bAllowsInterpolation=*/ true, /*bOverrideParent=*/false);
+
+		constexpr int32 EntryCount = 5;
+		for (int32 i = 0; i < EntryCount; ++i)
+		{
+			const FVector2D Position{ 2.0 * i, 2.0 * i + 1.0 };
+			PCGMetadataEntryKey EntryKey = ParamData->Metadata->AddEntry();
+
+			IntAttribute->SetValue(EntryKey, i);
+			Vec2Attribute->SetValue(EntryKey, Position);
+		}
+
+		TArray<FPCGTaggedData> Inputs;
+		FPCGTaggedData& Data = Inputs.Emplace_GetRef();
+		Data.Data = ParamData;
+
+		IntSelector.SetAttributeName(IntName);
+		Vec2Selector.SetAttributeName(Vec2Name);
+
+		return Inputs;
+	}
+
+	TArray<FPCGTaggedData> PrepareMixedAndSelector(FPCGAttributePropertyInputSelector& Selector)
+	{
+		const FName DoubleName = TEXT("double");
+
+		UPCGPointData* PointData = PCGTestsCommon::CreateEmptyPointData();
+		FPCGMetadataAttribute<double>* PointsAttribute = PointData->Metadata->CreateAttribute<double>(DoubleName, 0.0, /*bAllowsInterpolation=*/ true, /*bOverrideParent=*/false);
+
+		UPCGParamData* ParamData = NewObject<UPCGParamData>();
+		FPCGMetadataAttribute<double>* ParamAttribute = ParamData->Metadata->CreateAttribute<double>(DoubleName, 0.0, /*bAllowsInterpolation=*/ true, /*bOverrideParent=*/false);
+
+		TArray<FPCGPoint>& Points = PointData->GetMutablePoints();
+		constexpr int32 EntryCount = 5;
+		for (int32 i = 0; i < EntryCount; ++i)
+		{
+			FPCGPoint& Point = Points.Emplace_GetRef();
+			PointData->Metadata->InitializeOnSet(Point.MetadataEntry);
+			PointsAttribute->SetValue(Point.MetadataEntry, 0.1 * (i + 1));
+
+			ParamAttribute->SetValue(ParamData->Metadata->AddEntry(), 0.2 * (i + 1));
+		}
+
+		TArray<FPCGTaggedData> Inputs;
+		FPCGTaggedData& Data1 = Inputs.Emplace_GetRef();
+		FPCGTaggedData& Data2 = Inputs.Emplace_GetRef();
+		Data1.Data = PointData;
+		Data2.Data = ParamData;
+
+		Selector.SetAttributeName(DoubleName);
+
+		return Inputs;
+	}
+};
+
+
+IMPLEMENT_CUSTOM_SIMPLE_AUTOMATION_TEST(FPCGAttributeReducePoints_Average, FPCGAttributeReduceTests, "pcg.tests.AttributeReduce.Points.Average", PCGTestsCommon::TestFlags)
+IMPLEMENT_CUSTOM_SIMPLE_AUTOMATION_TEST(FPCGAttributeReducePoints_Min, FPCGAttributeReduceTests, "pcg.tests.AttributeReduce.Points.Min", PCGTestsCommon::TestFlags)
+IMPLEMENT_CUSTOM_SIMPLE_AUTOMATION_TEST(FPCGAttributeReducePoints_Max, FPCGAttributeReduceTests, "pcg.tests.AttributeReduce.Points.Max", PCGTestsCommon::TestFlags)
+
+IMPLEMENT_CUSTOM_SIMPLE_AUTOMATION_TEST(FPCGAttributeReduceParams_Average, FPCGAttributeReduceTests, "pcg.tests.AttributeReduce.Params.Average", PCGTestsCommon::TestFlags)
+IMPLEMENT_CUSTOM_SIMPLE_AUTOMATION_TEST(FPCGAttributeReduceParams_Min, FPCGAttributeReduceTests, "pcg.tests.AttributeReduce.Params.Min", PCGTestsCommon::TestFlags)
+IMPLEMENT_CUSTOM_SIMPLE_AUTOMATION_TEST(FPCGAttributeReduceParams_Max, FPCGAttributeReduceTests, "pcg.tests.AttributeReduce.Params.Max", PCGTestsCommon::TestFlags)
+
+IMPLEMENT_CUSTOM_SIMPLE_AUTOMATION_TEST(FPCGAttributeReduceMixed_NoMerge, FPCGAttributeReduceTests, "pcg.tests.AttributeReduce.Mixed.NoMerge", PCGTestsCommon::TestFlags)
+IMPLEMENT_CUSTOM_SIMPLE_AUTOMATION_TEST(FPCGAttributeReduceMixed_Merge, FPCGAttributeReduceTests, "pcg.tests.AttributeReduce.Mixed.Merge", PCGTestsCommon::TestFlags)
+
+IMPLEMENT_CUSTOM_SIMPLE_AUTOMATION_TEST(FPCGAttributeReduceFirstInputInvalid_NoMerge, FPCGAttributeReduceTests, "pcg.tests.AttributeReduce.FirstInputInvalid.NoMerge", PCGTestsCommon::TestFlags)
+IMPLEMENT_CUSTOM_SIMPLE_AUTOMATION_TEST(FPCGAttributeReduceFirstInputInvalid_Merge, FPCGAttributeReduceTests, "pcg.tests.AttributeReduce.FirstInputInvalid.Merge", PCGTestsCommon::TestFlags)
+
+IMPLEMENT_CUSTOM_SIMPLE_AUTOMATION_TEST(FPCGAttributeReduceIncompatibleType, FPCGAttributeReduceTests, "pcg.tests.AttributeReduce.IncompatibleType", PCGTestsCommon::TestFlags)
+
+
+bool FPCGAttributeReducePoints_Average::RunTest(const FString& Parameters)
+{
+	FPCGAttributePropertyInputSelector InputSourceDensity, InputSourcePosition;
+	TArray<FPCGTaggedData> Inputs = PreparePointsAndSelectors(InputSourceDensity, InputSourcePosition);
+
+	bool bSuccess = true;
+
+	bSuccess &= RunAndValidate(Inputs, EPCGAttributeReduceOperation::Average, std::move(InputSourceDensity), /*bMerged=*/ false, 0.3);
+	bSuccess &= RunAndValidate(Inputs, EPCGAttributeReduceOperation::Average, std::move(InputSourcePosition), /*bMerged=*/ false, FVector(6.0, 7.0, 8.0));
+
+	return bSuccess;
+}
+
+bool FPCGAttributeReducePoints_Min::RunTest(const FString& Parameters)
+{
+	FPCGAttributePropertyInputSelector InputSourceDensity, InputSourcePosition;
+	TArray<FPCGTaggedData> Inputs = PreparePointsAndSelectors(InputSourceDensity, InputSourcePosition);
+
+	bool bSuccess = true;
+
+	bSuccess &= RunAndValidate(Inputs, EPCGAttributeReduceOperation::Min, std::move(InputSourceDensity), /*bMerged=*/ false, 0.1);
+	bSuccess &= RunAndValidate(Inputs, EPCGAttributeReduceOperation::Min, std::move(InputSourcePosition), /*bMerged=*/ false, FVector(0.0, 1.0, 2.0));
+
+	return bSuccess;
+}
+
+bool FPCGAttributeReducePoints_Max::RunTest(const FString& Parameters)
+{
+	FPCGAttributePropertyInputSelector InputSourceDensity, InputSourcePosition;
+	TArray<FPCGTaggedData> Inputs = PreparePointsAndSelectors(InputSourceDensity, InputSourcePosition);
+
+	bool bSuccess = true;
+
+	bSuccess &= RunAndValidate(Inputs, EPCGAttributeReduceOperation::Max, std::move(InputSourceDensity), /*bMerged=*/ false, 0.5);
+	bSuccess &= RunAndValidate(Inputs, EPCGAttributeReduceOperation::Max, std::move(InputSourcePosition), /*bMerged=*/ false, FVector(12.0, 13.0, 14.0));
+
+	return bSuccess;
+}
+
+bool FPCGAttributeReduceParams_Average::RunTest(const FString& Parameters)
+{
+	FPCGAttributePropertyInputSelector InputSourceInt, InputSourceVec2;
+	TArray<FPCGTaggedData> Inputs = PrepareParamAndSelectors(InputSourceInt, InputSourceVec2);
+
+	bool bSuccess = true;
+
+	bSuccess &= RunAndValidate(Inputs, EPCGAttributeReduceOperation::Average, std::move(InputSourceInt), /*bMerged=*/ false, 2);
+	bSuccess &= RunAndValidate(Inputs, EPCGAttributeReduceOperation::Average, std::move(InputSourceVec2), /*bMerged=*/ false, FVector2D(4.0, 5.0));
+
+	return bSuccess;
+}
+
+bool FPCGAttributeReduceParams_Min::RunTest(const FString& Parameters)
+{
+	FPCGAttributePropertyInputSelector InputSourceInt, InputSourceVec2;
+	TArray<FPCGTaggedData> Inputs = PrepareParamAndSelectors(InputSourceInt, InputSourceVec2);
+
+	bool bSuccess = true;
+
+	bSuccess &= RunAndValidate(Inputs, EPCGAttributeReduceOperation::Min, std::move(InputSourceInt), /*bMerged=*/ false, 0);
+	bSuccess &= RunAndValidate(Inputs, EPCGAttributeReduceOperation::Min, std::move(InputSourceVec2), /*bMerged=*/ false, FVector2D(0.0, 1.0));
+
+	return bSuccess;
+}
+
+bool FPCGAttributeReduceParams_Max::RunTest(const FString& Parameters)
+{
+	FPCGAttributePropertyInputSelector InputSourceInt, InputSourceVec2;
+	TArray<FPCGTaggedData> Inputs = PrepareParamAndSelectors(InputSourceInt, InputSourceVec2);
+
+	bool bSuccess = true;
+
+	bSuccess &= RunAndValidate(Inputs, EPCGAttributeReduceOperation::Max, std::move(InputSourceInt), /*bMerged=*/ false, 4);
+	bSuccess &= RunAndValidate(Inputs, EPCGAttributeReduceOperation::Max, std::move(InputSourceVec2), /*bMerged=*/ false, FVector2D(8.0, 9.0));
+
+	return bSuccess;
+}
+
+bool FPCGAttributeReduceMixed_NoMerge::RunTest(const FString& Parameters)
+{
+	FPCGAttributePropertyInputSelector InputSource;
+	TArray<FPCGTaggedData> Inputs = PrepareMixedAndSelector(InputSource);
+
+	return RunAndValidate(Inputs, EPCGAttributeReduceOperation::Average, std::move(InputSource), /*bMerged=*/ false, 0.3, 0.6);
+}
+
+bool FPCGAttributeReduceMixed_Merge::RunTest(const FString& Parameters)
+{
+	FPCGAttributePropertyInputSelector InputSource;
+	TArray<FPCGTaggedData> Inputs = PrepareMixedAndSelector(InputSource);
+
+	return RunAndValidate(Inputs, EPCGAttributeReduceOperation::Max, std::move(InputSource), /*bMerged=*/ true, 0.5, 1.0);
+}
+
+bool FPCGAttributeReduceFirstInputInvalid_NoMerge::RunTest(const FString& Parameters)
+{
+	FPCGAttributePropertyInputSelector InputSource;
+	TArray<FPCGTaggedData> Inputs = PrepareMixedAndSelector(InputSource);
+
+	// Add Random Settings at the beginning
+	FPCGTaggedData& TaggedData = Inputs.EmplaceAt_GetRef(0);
+	TaggedData.Data = NewObject<UPCGAttributeReduceSettings>();
+
+	AddExpectedError(TEXT("Input attribute/property 'double' does not exist on input 0, skipped"));
+
+	return RunAndValidate(Inputs, EPCGAttributeReduceOperation::Min, std::move(InputSource), /*bMerged=*/ false, 0.1, 0.2);
+}
+
+bool FPCGAttributeReduceFirstInputInvalid_Merge::RunTest(const FString& Parameters)
+{
+	FPCGAttributePropertyInputSelector InputSource;
+	TArray<FPCGTaggedData> Inputs = PrepareMixedAndSelector(InputSource);
+
+	// Add Random Settings at the beginning
+	FPCGTaggedData& TaggedData = Inputs.EmplaceAt_GetRef(0);
+	TaggedData.Data = NewObject<UPCGAttributeReduceSettings>();
+
+	AddExpectedError(TEXT("Input attribute/property 'double' does not exist on input 0, skipped"));
+
+	return RunAndValidate(Inputs, EPCGAttributeReduceOperation::Min, std::move(InputSource), /*bMerged=*/ true, 0.1, 0.2);
+}
+
+bool FPCGAttributeReduceIncompatibleType::RunTest(const FString& Parameters)
+{
+	const FName StringName = TEXT("string");
+	const FName QuatName = TEXT("quat");
+	FPCGAttributePropertyInputSelector InputSourceString;
+	FPCGAttributePropertyInputSelector InputSourceQuat;
+
+	InputSourceString.SetAttributeName(StringName);
+	InputSourceQuat.SetAttributeName(QuatName);
+
+	UPCGParamData* ParamData = NewObject<UPCGParamData>();
+	// String are always invalid
+	FPCGMetadataAttribute<FString>* StringAttribute = ParamData->Metadata->CreateAttribute<FString>(StringName, FString{}, /*bAllowsInterpolation=*/ false, /*bOverrideParent=*/false);
+	// Quat are invalid for min/max
+	FPCGMetadataAttribute<FQuat>* QuatAttribute = ParamData->Metadata->CreateAttribute<FQuat>(QuatName, FQuat{}, /*bAllowsInterpolation=*/ true, /*bOverrideParent=*/false);
+
+	constexpr int32 EntryCount = 5;
+	for (int32 i = 0; i < EntryCount; ++i)
+	{
+		PCGMetadataEntryKey EntryKey = ParamData->Metadata->AddEntry();
+
+		QuatAttribute->SetValue(EntryKey, FQuat::Identity);
+		StringAttribute->SetValue(EntryKey, FString::FormatAsNumber(i));
+	}
+
+	TArray<FPCGTaggedData> Inputs;
+	FPCGTaggedData& TaggedData = Inputs.Emplace_GetRef();
+	TaggedData.Data = ParamData;
+
+	bool bSuccess = true;
+
+	AddExpectedError(TEXT("Operation was not compatible with the attribute type String or could not create attribute 'None' for input 0"));
+	bSuccess &= RunAndValidate(Inputs, EPCGAttributeReduceOperation::Average, std::move(InputSourceString), /*bMerged=*/ false);
+
+	AddExpectedError(TEXT("Operation was not compatible with the attribute type Quaternion or could not create attribute 'None' for input 0"));
+	bSuccess &= RunAndValidate(Inputs, EPCGAttributeReduceOperation::Min, std::move(InputSourceQuat), /*bMerged=*/ false);
+
+	return bSuccess;
+}
\ No newline at end of file
diff --git a/Engine/Plugins/Experimental/PCG/Source/PCG/Public/Elements/PCGAttributeReduceElement.h b/Engine/Plugins/Experimental/PCG/Source/PCG/Public/Elements/PCGAttributeReduceElement.h
--- a/Engine/Plugins/Experimental/PCG/Source/PCG/Public/Elements/PCGAttributeReduceElement.h
+++ b/Engine/Plugins/Experimental/PCG/Source/PCG/Public/Elements/PCGAttributeReduceElement.h
@@ -19,8 +19,9 @@ enum class EPCGAttributeReduceOperation
 /**
 * Take all the entries/points from the input and perform a reduce operation on the given attribute/property
 * and output the result into a ParamData.
-* 
-* If the OutputAttributeName is None, we will use the InputSource name
+* Note: Special case for average on Quaternion since they are not trivially averageable. We have a simplistic approximation
+* that would be accurate only if the quaternions are close to each other. The accurate version of the average is using eigenvectors/eigenvalues
+* which is way more complicated and computationally expensive. Quaternion will also be normatilzed at the end. Beware if you are using this average.
 */
 UCLASS(BlueprintType, ClassGroup = (Procedural))
 class PCG_API UPCGAttributeReduceSettings : public UPCGSettings
@@ -38,23 +39,31 @@ public:
 	virtual FText GetDefaultNodeTitle() const override;
 	virtual EPCGSettingsType GetType() const override { return EPCGSettingsType::Metadata; }
 	virtual void ApplyDeprecation(UPCGNode* InOutNode) override;
+	virtual TArray<FPCGPreConfiguredSettingsInfo> GetPreconfiguredInfo() const override;
+	virtual bool OnlyExposePreconfiguredSettings() const override { return true; }
 #endif
 
+	virtual void ApplyPreconfiguredSettings(const FPCGPreConfiguredSettingsInfo& PreconfiguredInfo) override;
+
 	virtual FName AdditionalTaskName() const override;
 
 	virtual TArray<FPCGPinProperties> InputPinProperties() const override;
 	virtual TArray<FPCGPinProperties> OutputPinProperties() const override;
 	//~End UPCGSettings interface
 
-	UPROPERTY(BlueprintReadWrite, EditAnywhere, Category = Settings)
+	UPROPERTY(BlueprintReadWrite, EditAnywhere, Category = Settings, meta = (PCG_Overridable))
 	FPCGAttributePropertyInputSelector InputSource;
 
-	UPROPERTY(BlueprintReadWrite, EditAnywhere, Category = Settings)
+	UPROPERTY(BlueprintReadWrite, EditAnywhere, Category = Settings, meta = (PCG_Overridable))
 	FName OutputAttributeName = NAME_None;
 
-	UPROPERTY(BlueprintReadWrite, EditAnywhere, Category = Settings)
+	UPROPERTY(BlueprintReadWrite, EditAnywhere, Category = Settings, meta = (PCG_Overridable))
 	EPCGAttributeReduceOperation Operation = EPCGAttributeReduceOperation::Average;
 
+	/** Option to merge all results into a single attribute set with multiple entries, instead of multiple attribute sets with a single value in them.*/
+	UPROPERTY(BlueprintReadWrite, EditAnywhere, Category = Settings, meta = (PCG_Overridable))
+	bool bMergeOutputAttributes = false;
+
 #if WITH_EDITORONLY_DATA
 	UPROPERTY()
 	FName InputAttributeName_DEPRECATED = NAME_None;
