[PCG] Added support to copy (ctrl+c) selected rows in the attribute list view to clipboard using CSV format
+Moved to using the new index accessors
+Added ellipsis overflow to metadata column headers
+Removed unused column alignment code

diff --git a/Engine/Plugins/Experimental/PCG/Source/PCGEditor/Private/SPCGEditorGraphAttributeListView.cpp b/Engine/Plugins/Experimental/PCG/Source/PCGEditor/Private/SPCGEditorGraphAttributeListView.cpp
--- a/Engine/Plugins/Experimental/PCG/Source/PCGEditor/Private/SPCGEditorGraphAttributeListView.cpp
+++ b/Engine/Plugins/Experimental/PCG/Source/PCGEditor/Private/SPCGEditorGraphAttributeListView.cpp
@@ -21,7 +21,9 @@
 #include "PCGEditorGraphNodeBase.h"
 
 #include "Fonts/FontMeasure.h"
+#include "Framework/Commands/GenericCommands.h"
 #include "Framework/MultiBox/MultiBoxBuilder.h"
+#include "HAL/PlatformApplicationMisc.h"
 #include "Widgets/Images/SLayeredImage.h"
 #include "Widgets/Images/SThrobber.h"
 #include "Widgets/Input/SComboBox.h"
@@ -45,7 +47,7 @@ namespace PCGEditorGraphAttributeListView
 	const FText NoNodeInspectedToolTip = LOCTEXT("NoNodeInspectedToolTip", "Inspect a node using the right click menu");
 
 	/** Names of the columns in the attribute list */
-	const FName NAME_IndexColumn = FName(TEXT("IndexColumn"));
+	const FName NAME_Index = FName(TEXT("$Index"));
 	const FName NAME_PointPositionX = FName(TEXT("$Position.X"));
 	const FName NAME_PointPositionY = FName(TEXT("$Position.Y"));
 	const FName NAME_PointPositionZ = FName(TEXT("$Position.Z"));
@@ -137,29 +139,12 @@ void FPCGListViewUpdater::Launch()
 
 void FPCGListViewUpdater::AsyncSort()
 {
-	if (SortingColumn == PCGEditorGraphAttributeListView::NAME_IndexColumn || SortMode == EColumnSortMode::None)
-	{
-		if (SortMode == EColumnSortMode::Ascending || SortMode == EColumnSortMode::None)
-		{
-			ListViewItems.Sort([](const PCGListviewItemPtr& LHS, const PCGListviewItemPtr& RHS)
-			{
-				return  LHS->Index < RHS->Index;
-			});
-		}
-		else
-		{
-			ListViewItems.Sort([](const PCGListviewItemPtr& LHS, const PCGListviewItemPtr& RHS)
-			{
-				return  LHS->Index > RHS->Index;
-			});
-		}
-	}
-	else if (const FPCGColumnData* Data = ColumnData.Find(SortingColumn))
+	if (const FPCGColumnData* Data = ColumnData.Find(SortMode == EColumnSortMode::None ? PCGEditorGraphAttributeListView::NAME_Index : SortingColumn))
 	{
 		if (Data->DataAccessor.IsValid() && Data->DataKeys.IsValid())
 		{
 			//lambda used here to get the index value of an item in the array for sorting
-			PCGAttributeAccessorHelpers::SortByAttribute(*Data->DataAccessor, *Data->DataKeys, ListViewItems, SortMode == EColumnSortMode::Ascending, [this](int Index) { return ListViewItems[Index]->Index; });
+			PCGAttributeAccessorHelpers::SortByAttribute(*Data->DataAccessor, *Data->DataKeys, ListViewItems, !(SortMode & EColumnSortMode::Descending), [this](int Index) { return ListViewItems[Index]->Index; });
 		}
 	}
 }
@@ -197,12 +182,6 @@ TSharedRef<SWidget> SPCGListViewItemRow::GenerateWidgetForColumn(const FName& Co
 {
 	FText RowText = LOCTEXT("ColumnError", "Unrecognized Column");
 
-	if (ColumnId == PCGEditorGraphAttributeListView::NAME_IndexColumn)
-	{
-		RowText = FText::FromString(FString::FromInt(InternalItem->Index));
-		return SNew(STextBlock).Text(RowText);
-	}
-
 	const TSharedPtr<SPCGEditorGraphAttributeListView> SharedAttributeListView = AttributeListView.Pin();
 	check(SharedAttributeListView.IsValid());
 
@@ -339,7 +318,7 @@ SPCGEditorGraphAttributeListView::~SPCGEditorGraphAttributeListView()
 void SPCGEditorGraphAttributeListView::Construct(const FArguments& InArgs, TSharedPtr<FPCGEditor> InPCGEditor)
 {
 	PCGEditorPtr = InPCGEditor;
-	SortingColumn = PCGEditorGraphAttributeListView::NAME_IndexColumn;
+	SortingColumn = PCGEditorGraphAttributeListView::NAME_Index;
 
 	PCGEditorPtr.Pin()->OnInspectedComponentChangedDelegate.AddSP(this, &SPCGEditorGraphAttributeListView::OnInspectedComponentChanged);
 	PCGEditorPtr.Pin()->OnInspectedStackChangedDelegate.AddSP(this, &SPCGEditorGraphAttributeListView::OnInspectedStackChanged);
@@ -349,6 +328,11 @@ void SPCGEditorGraphAttributeListView::Construct(const FArguments& InArgs, TShar
 
 	ListViewHeader = CreateHeaderRowWidget();
 
+	ListViewCommands = MakeShareable(new FUICommandList);
+	ListViewCommands->MapAction(FGenericCommands::Get().Copy,
+		FExecuteAction::CreateSP(this, &SPCGEditorGraphAttributeListView::CopySelectionToClipboard),
+		FCanExecuteAction::CreateSP(this, &SPCGEditorGraphAttributeListView::CanCopySelectionToClipboard));
+
 	const TSharedRef<SScrollBar> HorizontalScrollBar = SNew(SScrollBar)
 		.Orientation(Orient_Horizontal)
 		.Thickness(FVector2D(12.0f, 12.0f));
@@ -370,6 +354,7 @@ void SPCGEditorGraphAttributeListView::Construct(const FArguments& InArgs, TShar
 		.AllowOverscroll(EAllowOverscroll::No)
 		.ExternalScrollbar(VerticalScrollBar)
 		.Visibility_Lambda(VisibilityTest)
+		.OnKeyDownHandler(this, &SPCGEditorGraphAttributeListView::OnListViewKeyDown)
 		.ConsumeMouseWheel(EConsumeMouseWheel::Always);
 
 	SAssignNew(PinComboBox, SComboBox<TSharedPtr<FPinComboBoxItem>>)
@@ -680,11 +665,11 @@ void SPCGEditorGraphAttributeListView::RefreshAttributeList()
 	{
 		if (const UPCGMetadata* PCGMetadata = PCGParamData->ConstMetadata())
 		{
-			AddIndexColumn();
+			AddColumn(PCGData, PCGEditorGraphAttributeListView::NAME_Index, PCGEditorGraphAttributeListView::TEXT_IndexLabel);
 			GenerateColumnsFromMetadata(PCGData, PCGMetadata);
 
-			PCGMetadataEntryKey ItemKeyLowerBound = PCGMetadata->GetItemKeyCountForParent();
-			PCGMetadataEntryKey ItemKeyUpperBound = PCGMetadata->GetItemCountForChild();
+			const PCGMetadataEntryKey ItemKeyLowerBound = PCGMetadata->GetItemKeyCountForParent();
+			const PCGMetadataEntryKey ItemKeyUpperBound = PCGMetadata->GetItemCountForChild();
 			for (PCGMetadataEntryKey MetadataItemKey = ItemKeyLowerBound; MetadataItemKey < ItemKeyUpperBound; ++MetadataItemKey)
 			{
 				PCGListviewItemPtr ListViewItem = MakeShared<FPCGListViewItem>();
@@ -878,7 +863,7 @@ TSharedRef<SWidget> SPCGEditorGraphAttributeListView::OnGenerateFilterMenu()
 
 	for (const SHeaderRow::FColumn& Column : Columns)
 	{
-		if (Column.ColumnId == PCGEditorGraphAttributeListView::NAME_IndexColumn)
+		if (Column.ColumnId == PCGEditorGraphAttributeListView::NAME_Index)
 		{
 			continue;
 		}
@@ -1061,7 +1046,7 @@ void SPCGEditorGraphAttributeListView::ToggleAllAttributes()
 		const TIndirectArray<SHeaderRow::FColumn>& Columns = ListViewHeader->GetColumns();
 		for (const SHeaderRow::FColumn& Column : Columns)
 		{
-			if (Column.ColumnId != PCGEditorGraphAttributeListView::NAME_IndexColumn)
+			if (Column.ColumnId != PCGEditorGraphAttributeListView::NAME_Index)
 			{
 				ListViewHeader->SetShowGeneratedColumn(Column.ColumnId, /*InShow=*/false);
 			}
@@ -1081,7 +1066,7 @@ ECheckBoxState SPCGEditorGraphAttributeListView::GetAnyAttributeEnabledState() c
 
 	for (const SHeaderRow::FColumn& Column : ListViewHeader->GetColumns())
 	{
-		if (Column.ColumnId == PCGEditorGraphAttributeListView::NAME_IndexColumn)
+		if (Column.ColumnId == PCGEditorGraphAttributeListView::NAME_Index)
 		{
 			continue;
 		}
@@ -1176,11 +1161,21 @@ void SPCGEditorGraphAttributeListView::OnFilterTextCommitted(const FText& NewTex
 	}
 }
 
-void SPCGEditorGraphAttributeListView::AddColumn(const UPCGPointData* InPCGPointData, const FName& InColumnId, const FText& ColumnLabel, EHorizontalAlignment HeaderHAlign, EHorizontalAlignment CellHAlign)
+FReply SPCGEditorGraphAttributeListView::OnListViewKeyDown(const FGeometry& /*InGeometry*/, const FKeyEvent& InKeyEvent) const
+{
+	if (ListViewCommands->ProcessCommandBindings(InKeyEvent))
+	{
+		return FReply::Handled();
+	}
+
+	return FReply::Unhandled();
+}
+
+void SPCGEditorGraphAttributeListView::AddColumn(const UPCGData* InPCGData, const FName& InColumnId, const FText& ColumnLabel)
 {
 	FText ToolTip;
 
-	if (InPCGPointData)
+	if (InPCGData)
 	{
 		const FString ColumnIdString = InColumnId.ToString();
 
@@ -1191,27 +1186,33 @@ void SPCGEditorGraphAttributeListView::AddColumn(const UPCGPointData* InPCGPoint
 
 		if (InColumnId == PCGEditorGraphAttributeListView::NAME_PointMetadataEntry)
 		{
-			TUniquePtr<const IPCGAttributeAccessor> DataAccessor = PCGAttributeAccessorHelpers::CreatePropertyAccessor(GET_MEMBER_NAME_CHECKED(FPCGPoint, MetadataEntry), FPCGPoint::StaticStruct());
-			ColumnData.DataAccessor = TSharedPtr<const IPCGAttributeAccessor>(DataAccessor.Release());
-			ColumnData.DataKeys = MakeShared<FPCGAttributeAccessorKeysPoints>(InPCGPointData->GetPoints());
+			if (const UPCGPointData* PCGPointData = Cast<const UPCGPointData>(InPCGData))
+			{
+				TUniquePtr<const IPCGAttributeAccessor> DataAccessor = PCGAttributeAccessorHelpers::CreatePropertyAccessor(GET_MEMBER_NAME_CHECKED(FPCGPoint, MetadataEntry), FPCGPoint::StaticStruct());
+				ColumnData.DataAccessor = TSharedPtr<const IPCGAttributeAccessor>(DataAccessor.Release());
+				ColumnData.DataKeys = MakeShared<FPCGAttributeAccessorKeysPoints>(PCGPointData->GetPoints());
+			}
 		}
 		else if (InColumnId == PCGEditorGraphAttributeListView::NAME_PointMetadataEntryParent)
 		{
-			ColumnData.DataAccessor = MakeShared<FPCGCustomPointAccessor<int64>>([InPCGPointData](const FPCGPoint& Point, void* OutValue)
+			if (const UPCGPointData* PCGPointData = Cast<const UPCGPointData>(InPCGData))
 			{
-				if (const UPCGMetadata* Metadata = InPCGPointData->Metadata)
+				ColumnData.DataAccessor = MakeShared<FPCGCustomPointAccessor<int64>>([PCGPointData](const FPCGPoint& Point, void* OutValue)
 				{
-					*reinterpret_cast<int64*>(OutValue) = Metadata->GetParentKey(Point.MetadataEntry);
-					return true;
-				}
-				return false;
-			}, nullptr);
-			ColumnData.DataKeys = MakeShared<FPCGAttributeAccessorKeysPoints>(InPCGPointData->GetPoints());
+					if (const UPCGMetadata* Metadata = PCGPointData->Metadata)
+					{
+						*reinterpret_cast<int64*>(OutValue) = Metadata->GetParentKey(Point.MetadataEntry);
+						return true;
+					}
+					return false;
+				}, nullptr);
+				ColumnData.DataKeys = MakeShared<FPCGAttributeAccessorKeysPoints>(PCGPointData->GetPoints());
+			}
 		}
 		else
 		{
-			ColumnData.DataAccessor = TSharedPtr<const IPCGAttributeAccessor>(PCGAttributeAccessorHelpers::CreateConstAccessor(InPCGPointData, TargetSelector).Release());
-			ColumnData.DataKeys = TSharedPtr<const IPCGAttributeAccessorKeys>(PCGAttributeAccessorHelpers::CreateConstKeys(InPCGPointData, TargetSelector).Release());
+			ColumnData.DataAccessor = TSharedPtr<const IPCGAttributeAccessor>(PCGAttributeAccessorHelpers::CreateConstAccessor(InPCGData, TargetSelector).Release());
+			ColumnData.DataKeys = TSharedPtr<const IPCGAttributeAccessorKeys>(PCGAttributeAccessorHelpers::CreateConstKeys(InPCGData, TargetSelector).Release());
 		}
 
 		ToolTip = FText::FromString(PCG::Private::GetTypeName(ColumnData.DataAccessor->GetUnderlyingType()));
@@ -1224,8 +1225,8 @@ void SPCGEditorGraphAttributeListView::AddColumn(const UPCGPointData* InPCGPoint
 	Arguments.DefaultLabel(ColumnLabel);
 	Arguments.DefaultTooltip(ToolTip);
 	Arguments.ManualWidth(ColumnWidth);
-	Arguments.HAlignHeader(HeaderHAlign);
-	Arguments.HAlignCell(CellHAlign);
+	Arguments.HAlignHeader(HAlign_Center);
+	Arguments.HAlignCell(HAlign_Right);
 	Arguments.SortMode(this, &SPCGEditorGraphAttributeListView::GetColumnSortMode, InColumnId);
 	Arguments.OnSort(this, &SPCGEditorGraphAttributeListView::OnColumnSortModeChanged);
 	Arguments.OverflowPolicy(ETextOverflowPolicy::Ellipsis);
@@ -1235,14 +1236,9 @@ void SPCGEditorGraphAttributeListView::AddColumn(const UPCGPointData* InPCGPoint
 	ListViewHeader->AddColumn(*NewColumn);
 }
 
-void SPCGEditorGraphAttributeListView::AddIndexColumn()
-{
-	AddColumn(nullptr, PCGEditorGraphAttributeListView::NAME_IndexColumn, PCGEditorGraphAttributeListView::TEXT_IndexLabel);
-}
-
 void SPCGEditorGraphAttributeListView::AddPointDataColumns(const UPCGPointData* InPCGPointData)
 {
-	AddIndexColumn();
+	AddColumn(InPCGPointData, PCGEditorGraphAttributeListView::NAME_Index, PCGEditorGraphAttributeListView::TEXT_IndexLabel);
 	AddColumn(InPCGPointData, PCGEditorGraphAttributeListView::NAME_PointPositionX, PCGEditorGraphAttributeListView::TEXT_PointPositionLabelX);
 	AddColumn(InPCGPointData, PCGEditorGraphAttributeListView::NAME_PointPositionY, PCGEditorGraphAttributeListView::TEXT_PointPositionLabelY);
 	AddColumn(InPCGPointData, PCGEditorGraphAttributeListView::NAME_PointPositionZ, PCGEditorGraphAttributeListView::TEXT_PointPositionLabelZ);
@@ -1332,10 +1328,85 @@ void SPCGEditorGraphAttributeListView::AddMetadataColumn(const UPCGData* InPCGDa
 	ColumnArguments.ManualWidth(ColumnWidth);
 	ColumnArguments.SortMode(this, &SPCGEditorGraphAttributeListView::GetColumnSortMode, ColumnId);
 	ColumnArguments.OnSort(this, &SPCGEditorGraphAttributeListView::OnColumnSortModeChanged);
+	ColumnArguments.OverflowPolicy(ETextOverflowPolicy::Ellipsis);
 
 	SHeaderRow::FColumn* NewColumn = new SHeaderRow::FColumn(ColumnArguments);
 	NewColumn->bIsVisible = !HiddenAttributes.Contains(ColumnId);
 	ListViewHeader->AddColumn(*NewColumn);
 }
 
+void SPCGEditorGraphAttributeListView::CopySelectionToClipboard() const
+{
+	constexpr TCHAR Delimiter = TEXT(',');
+	constexpr TCHAR LineEnd = TEXT('\n');
+
+	const TArray<FName> HiddenColumnIds = ListViewHeader->GetHiddenColumnIds();
+
+	TArray<TPair<FName, FPCGColumnData>> FilteredPCGColumnData;
+	for (const TPair<FName, FPCGColumnData>& Data : PCGColumnData)
+	{
+		if (!HiddenColumnIds.Contains(Data.Key))
+		{
+			FilteredPCGColumnData.Add(Data);
+		}
+	}
+
+	TStringBuilder<2048> CSVExport;
+
+	// Write column header row
+	for (int ColumnIndex = 0; ColumnIndex < FilteredPCGColumnData.Num(); ColumnIndex++)
+	{
+		const TPair<FName, FPCGColumnData>& Data = FilteredPCGColumnData[ColumnIndex];
+
+		if (ColumnIndex > 0)
+		{
+			CSVExport += Delimiter;
+		}
+		CSVExport += Data.Key.ToString();
+	}
+
+	// Gather selected rows and sort them to match the displayed order instead of selection order
+	TArray<PCGListviewItemPtr> SelectedListViewItems = ListView->GetSelectedItems();
+	if (const FPCGColumnData* ColumnData = PCGColumnData.Find(SortMode == EColumnSortMode::None ? PCGEditorGraphAttributeListView::NAME_Index : SortingColumn))
+	{
+		if (ColumnData->DataAccessor.IsValid() && ColumnData->DataKeys.IsValid())
+		{
+			// Lambda used here to get the index value of an item in the array for sorting
+			PCGAttributeAccessorHelpers::SortByAttribute(*ColumnData->DataAccessor, *ColumnData->DataKeys, SelectedListViewItems, !(SortMode & EColumnSortMode::Descending), [&SelectedListViewItems](int Index) { return SelectedListViewItems[Index]->Index; });
+		}
+	}
+
+	// Write each row
+	for (const PCGListviewItemPtr&  ListViewItem : SelectedListViewItems)
+	{
+		CSVExport += LineEnd;
+
+		for (int ColumnIndex = 0; ColumnIndex < FilteredPCGColumnData.Num(); ColumnIndex++)
+		{
+			if (ColumnIndex > 0)
+			{
+				CSVExport += Delimiter;
+			}
+
+			const TPair<FName, FPCGColumnData>& Data = FilteredPCGColumnData[ColumnIndex];
+			const FPCGColumnData& ColumnData = Data.Value;
+			if (ColumnData.DataAccessor.IsValid() && ColumnData.DataKeys.IsValid())
+			{
+				FString RowString;
+				if (ColumnData.DataAccessor->Get<FString>(RowString, ListViewItem->Index, *ColumnData.DataKeys, EPCGAttributeAccessorFlags::AllowBroadcast))
+				{
+					CSVExport += RowString;
+				}
+			}
+		}
+	}
+
+	FPlatformApplicationMisc::ClipboardCopy(*CSVExport);
+}
+
+bool SPCGEditorGraphAttributeListView::CanCopySelectionToClipboard() const
+{
+	return ListView->GetNumItemsSelected() > 0;
+}
+
 #undef LOCTEXT_NAMESPACE
diff --git a/Engine/Plugins/Experimental/PCG/Source/PCGEditor/Private/SPCGEditorGraphAttributeListView.h b/Engine/Plugins/Experimental/PCG/Source/PCGEditor/Private/SPCGEditorGraphAttributeListView.h
--- a/Engine/Plugins/Experimental/PCG/Source/PCGEditor/Private/SPCGEditorGraphAttributeListView.h
+++ b/Engine/Plugins/Experimental/PCG/Source/PCGEditor/Private/SPCGEditorGraphAttributeListView.h
@@ -17,13 +17,14 @@ template <typename OptionType> class SComboBox;
 
 class FPCGEditor;
 class FPCGMetadataAttributeBase;
+class FUICommandList;
 class SPCGEditorGraphAttributeListView;
 class UPCGComponent;
-class UPCGPointData;
 class UPCGData;
 class UPCGEditorGraphNodeBase;
 class UPCGMetadata;
 class UPCGParamData;
+class UPCGPointData;
 struct FPCGDataCollection;
 struct FPCGPoint;
 enum class EPCGMetadataTypes : uint8;
@@ -191,11 +192,15 @@ private:
 	void OnFilterTextChanged(const FText& InFilterText);
 	void OnFilterTextCommitted(const FText& NewText, ETextCommit::Type CommitInfo);
 
-	void AddColumn(const UPCGPointData* InPCGPointData, const FName& InColumnId, const FText& ColumnLabel, EHorizontalAlignment HeaderHAlign = HAlign_Center, EHorizontalAlignment CellHAlign = HAlign_Right);
-	void AddIndexColumn();
+	FReply OnListViewKeyDown(const FGeometry& InGeometry, const FKeyEvent& InKeyEvent) const;
+
+	void AddColumn(const UPCGData* InPCGData, const FName& InColumnId, const FText& ColumnLabel);
 	void AddPointDataColumns(const UPCGPointData* InPCGPointData);
 	void AddMetadataColumn(const UPCGData* InPCGData, const FName& InColumnId, EPCGMetadataTypes InMetadataType, const TCHAR* PostFix = nullptr);
 
+	void CopySelectionToClipboard() const;
+	bool CanCopySelectionToClipboard() const;
+
 	/** Pointer back to the PCG editor that owns us */
 	TWeakPtr<FPCGEditor> PCGEditorPtr;
 
@@ -205,6 +210,8 @@ private:
 	/** Cached PCGGraphNode being viewed */
 	TWeakObjectPtr<UPCGEditorGraphNodeBase> PCGEditorGraphNode;
 
+	TSharedPtr<FUICommandList> ListViewCommands;
+
 	TSharedPtr<FTextFilterExpressionEvaluator> TextFilter;
 
 	TSharedPtr<SSearchBox> SearchBoxWidget;
