[PCG] Pass on unit tests

* Give them some love! They needed it!

#rnx

diff --git a/Engine/Plugins/Experimental/PCG/Source/PCG/Private/Tests/Accessors/PCGAttributeExtractorTest.cpp b/Engine/Plugins/Experimental/PCG/Source/PCG/Private/Tests/Accessors/PCGAttributeExtractorTest.cpp
--- a/Engine/Plugins/Experimental/PCG/Source/PCG/Private/Tests/Accessors/PCGAttributeExtractorTest.cpp
+++ b/Engine/Plugins/Experimental/PCG/Source/PCG/Private/Tests/Accessors/PCGAttributeExtractorTest.cpp
@@ -113,8 +113,8 @@ bool FPCGAttributeSingleGetPropertyTest::RunTest(const FString& Parameters)
 	bTestPassed &= Verify(Struct.RotatorValue, PCG_STRUCT_NAME_CHECK(RotatorValue));
 	bTestPassed &= Verify(Struct.QuatValue, PCG_STRUCT_NAME_CHECK(QuatValue));
 	bTestPassed &= Verify(Struct.TransformValue, PCG_STRUCT_NAME_CHECK(TransformValue));
-	bTestPassed &= Verify(Struct.SoftClassPathValue.ToString(), PCG_STRUCT_NAME_CHECK(SoftClassPathValue));
-	bTestPassed &= Verify(Struct.SoftObjectPathValue.ToString(), PCG_STRUCT_NAME_CHECK(SoftObjectPathValue));
+	bTestPassed &= Verify(Struct.SoftClassPathValue, PCG_STRUCT_NAME_CHECK(SoftClassPathValue));
+	bTestPassed &= Verify(Struct.SoftObjectPathValue, PCG_STRUCT_NAME_CHECK(SoftObjectPathValue));
 
 	TempMetadata->MarkAsGarbage();
 
@@ -419,6 +419,7 @@ bool FPCGAttributePropertyMultipleDepthTest::RunTest(const FString& Parameters)
 	FPCGAttributeExtractorTestStruct TestStruct{};
 	TestStruct.Object = NewObject<UPCGAttributeExtractorTestObject>();
 	TestStruct.Object->DoubleValue = 0.5;
+	TestStruct.Object->SetFlags(RF_Transient);
 
 	TestStruct.DepthStruct.FloatValue = 0.324f;
 	TestStruct.DepthStruct.Depth2Struct.IntValue = 5;
@@ -442,8 +443,7 @@ bool FPCGAttributePropertyMultipleDepthTest::RunTest(const FString& Parameters)
 	TUniquePtr<IPCGAttributeAccessor> DepthStructFloatAccessor = PCGAttributeAccessorHelpers::CreatePropertyAccessor(GET_MEMBER_NAME_CHECKED(FPCGAttributeExtractorTestStructDepth1, FloatValue), FPCGAttributeExtractorTestStructDepth1::StaticStruct());
 	TUniquePtr<IPCGAttributeAccessor> Depth2StructIntAccessor = PCGAttributeAccessorHelpers::CreatePropertyAccessor(GET_MEMBER_NAME_CHECKED(FPCGAttributeExtractorTestStructDepth2, IntValue), FPCGAttributeExtractorTestStructDepth2::StaticStruct());
 
-	// Objects not yet supported
-	UTEST_TRUE(TEXT("Invalid object accessor"), !ObjectAccessor.IsValid());
+	UTEST_TRUE(TEXT("Valid object accessor"), ObjectAccessor.IsValid());
 	UTEST_TRUE(TEXT("Valid object double accessor"), ObjectDoubleAccessor.IsValid());
 	UTEST_TRUE(TEXT("Valid depth struct float accessor"), DepthStructFloatAccessor.IsValid());
 	UTEST_TRUE(TEXT("Valid deptch struct 2 int accessor"), Depth2StructIntAccessor.IsValid());
diff --git a/Engine/Plugins/Experimental/PCG/Source/PCG/Private/Tests/Elements/PCGCreatePointsTest.cpp b/Engine/Plugins/Experimental/PCG/Source/PCG/Private/Tests/Elements/PCGCreatePointsTest.cpp
--- a/Engine/Plugins/Experimental/PCG/Source/PCG/Private/Tests/Elements/PCGCreatePointsTest.cpp
+++ b/Engine/Plugins/Experimental/PCG/Source/PCG/Private/Tests/Elements/PCGCreatePointsTest.cpp
@@ -16,8 +16,16 @@ bool FPCGCreatePointsTest_Basic::RunTest(const FString& Parameters)
 	UPCGCreatePointsSettings* Settings = CastChecked<UPCGCreatePointsSettings>(TestData.Settings);
 	
 	TObjectPtr<UPCGPointData> PointData = PCGTestsCommon::CreateRandomPointData(100, 42, false);
-	Settings->PointsToCreate = PointData->GetMutablePoints();
 
+	// Setting the seed explicitly, first point will have a seed of 0, which mean it going to be changed and computed depending on the point position
+	// All the others should stay the same.
+	TArray<FPCGPoint>& Points = PointData->GetMutablePoints();
+	for (int32 i = 0; i < Points.Num(); ++i)
+	{
+		Points[i].Seed = i;
+	}
+
+	Settings->PointsToCreate = Points;
 	FPCGElementPtr TestElement = TestData.Settings->GetElement();
 
 	TUniquePtr<FPCGContext> Context = TestData.InitializeTestContext();
@@ -39,7 +47,9 @@ bool FPCGCreatePointsTest_Basic::RunTest(const FString& Parameters)
 
 	for (int i = 0; i < OutPoints.Num(); ++i)
 	{
-		UTEST_EQUAL(FString::Format(TEXT("InArray[{0}].Seed is equal to OutArray[{0}].Seed"), { i }), OutPoints[i].Seed, Settings->PointsToCreate[i].Seed);
+		// Seed at 0 will compute a seed depending on the point position
+		static constexpr int32 SeedAtZero = 907633527;
+		UTEST_EQUAL(FString::Format(TEXT("InArray[{0}].Seed is equal to OutArray[{0}].Seed"), { i }), OutPoints[i].Seed, i == 0 ? SeedAtZero : Settings->PointsToCreate[i].Seed);
 		UTEST_EQUAL(FString::Format(TEXT("InArray[{0}].Density is equal to OutArray[{0}].Density"), { i }), OutPoints[i].Density, Settings->PointsToCreate[i].Density);
 		UTEST_EQUAL(FString::Format(TEXT("InArray[{0}].Transform is equal to OutArray[{0}].Transform"), { i }), OutPoints[i].Transform, Settings->PointsToCreate[i].Transform);
 	}
diff --git a/Engine/Plugins/Experimental/PCG/Source/PCG/Private/Tests/Elements/PCGDataTableRowToParamDataTest.cpp b/Engine/Plugins/Experimental/PCG/Source/PCG/Private/Tests/Elements/PCGDataTableRowToParamDataTest.cpp
--- a/Engine/Plugins/Experimental/PCG/Source/PCG/Private/Tests/Elements/PCGDataTableRowToParamDataTest.cpp
+++ b/Engine/Plugins/Experimental/PCG/Source/PCG/Private/Tests/Elements/PCGDataTableRowToParamDataTest.cpp
@@ -63,7 +63,7 @@ CCC,"C Name","C String","333","3333","3.3","3.33","(X=3.0,Y=3.0)","(X=3.0,Y=3.0,
 	bSuccess |= PCGDataTableRowToParamDataTestHelpers::TestAttribute<FVector2D,FVector2D>(*this, RowData, *Params, TEXT("V2"), FVector2D(2.0, 2.0));
 	bSuccess |= PCGDataTableRowToParamDataTestHelpers::TestAttribute<FVector,FVector>(*this, RowData, *Params, TEXT("V3"), FVector(2.0, 2.0, 2.0));
 	bSuccess |= PCGDataTableRowToParamDataTestHelpers::TestAttribute<FVector4,FVector4>(*this, RowData, *Params, TEXT("V4"), FVector4(2.0, 2.0, 2.0, 2.0));
-	bSuccess |= PCGDataTableRowToParamDataTestHelpers::TestAttribute<FSoftObjectPath,FString>(*this, RowData, *Params, TEXT("SoftPath"), TEXT("/Script/PCG"));
+	bSuccess |= PCGDataTableRowToParamDataTestHelpers::TestAttribute<FSoftObjectPath, FSoftObjectPath>(*this, RowData, *Params, TEXT("SoftPath"), FSoftObjectPath(TEXT("/Script/PCG")));
 
 	return bSuccess;
 }
@@ -88,7 +88,7 @@ CCC,"C Name","C String","333","3333","3.3","3.33","(X=3.0,Y=3.0)","(X=3.0,Y=3.0,
 
 	{
 		UPCGParamData* InputParamData = NewObject<UPCGParamData>();
-		InputParamData->MutableMetadata()->CreateStringAttribute(TEXT("DataTable"), TestDataTable->GetPathName(), false);
+		InputParamData->MutableMetadata()->CreateStringAttribute(TEXT("PathOverride"), TestDataTable->GetPathName(), false);
 		InputParamData->MutableMetadata()->CreateStringAttribute(TEXT("RowName"), RowNameOverride.ToString(), false);
 		FPCGTaggedData& ParamInput = TestData.InputData.TaggedData.Emplace_GetRef();
 		ParamInput.Data = InputParamData;
@@ -101,8 +101,10 @@ CCC,"C Name","C String","333","3333","3.3","3.33","(X=3.0,Y=3.0)","(X=3.0,Y=3.0,
 
 	while (!TestElement->Execute(Context.Get())) {}
 
-	const UPCGParamData* Params = Context->OutputData.GetFirstParamsOnParamsPin();
-	UTEST_NOT_NULL("Output params", Params);
+	TArray<FPCGTaggedData> AllParams = Context->OutputData.GetParamsByPin(PCGPinConstants::DefaultOutputLabel);
+	UTEST_EQUAL("Output params", AllParams.Num(), 1);
+
+	const UPCGParamData* Params = CastChecked<const UPCGParamData>(AllParams[0].Data);
 
 	// For static analysis
 	if (!Params)
diff --git a/Engine/Plugins/Experimental/PCG/Source/PCG/Private/Tests/Elements/PCGInstancePackerTest.cpp b/Engine/Plugins/Experimental/PCG/Source/PCG/Private/Tests/Elements/PCGInstancePackerTest.cpp
--- a/Engine/Plugins/Experimental/PCG/Source/PCG/Private/Tests/Elements/PCGInstancePackerTest.cpp
+++ b/Engine/Plugins/Experimental/PCG/Source/PCG/Private/Tests/Elements/PCGInstancePackerTest.cpp
@@ -101,12 +101,20 @@ namespace
 
 bool FPCGStaticMeshSpawnerInstancePackerByAttributeTest::RunTest(const FString& Parameters)
 {
-	PCGTestsCommon::FTestData TestData;
-	PCGTestsCommon::GenerateSettings<UPCGStaticMeshSpawnerSettings>(TestData);
-	UPCGStaticMeshSpawnerSettings* Settings = CastChecked<UPCGStaticMeshSpawnerSettings>(TestData.Settings);
+	PCGTestsCommon::FTestData TestData{};
+	UPCGStaticMeshSpawnerSettings* Settings = PCGTestsCommon::GenerateSettings<UPCGStaticMeshSpawnerSettings>(TestData);
+	check(Settings);
 
 	FPCGTaggedData& SourceTaggedData = TestData.InputData.TaggedData.Emplace_GetRef(FPCGTaggedData());
 	TObjectPtr<UPCGPointData> SourceData = PCGTestsCommon::CreateRandomPointData(5, TestData.Seed);
+
+	// TODO - remove this code when the tests support flipped instances (generates multiple ISMC components)
+	TArray<FPCGPoint>& Points = SourceData->GetMutablePoints();
+	for (FPCGPoint& Point : Points)
+	{
+		Point.Transform.SetScale3D(Point.Transform.GetScale3D().GetAbs());
+	}
+
 	SourceTaggedData.Data = SourceData;
 	SourceData->TargetActor = TestData.TestActor;
 
@@ -258,12 +266,20 @@ bool FPCGStaticMeshSpawnerInstancePackerByAttributeTest::RunTest(const FString&
 
 bool FPCGStaticMeshSpawnerInstancePackerByRegexTest::RunTest(const FString& Parameters)
 {
-	PCGTestsCommon::FTestData TestData;
-	PCGTestsCommon::GenerateSettings<UPCGStaticMeshSpawnerSettings>(TestData);
-	UPCGStaticMeshSpawnerSettings* Settings = CastChecked<UPCGStaticMeshSpawnerSettings>(TestData.Settings);
+	PCGTestsCommon::FTestData TestData{};
+	UPCGStaticMeshSpawnerSettings* Settings = PCGTestsCommon::GenerateSettings<UPCGStaticMeshSpawnerSettings>(TestData);
+	check(Settings);
 
 	FPCGTaggedData& SourceTaggedData = TestData.InputData.TaggedData.Emplace_GetRef(FPCGTaggedData());
 	TObjectPtr<UPCGPointData> SourceData = PCGTestsCommon::CreateRandomPointData(5, TestData.Seed);
+
+	// TODO - remove this code when the tests support flipped instances (generates multiple ISMC components)
+	TArray<FPCGPoint>& Points = SourceData->GetMutablePoints();
+	for (FPCGPoint& Point : Points)
+	{
+		Point.Transform.SetScale3D(Point.Transform.GetScale3D().GetAbs());
+	}
+
 	SourceTaggedData.Data = SourceData;
 	SourceData->TargetActor = TestData.TestActor;
 
diff --git a/Engine/Plugins/Experimental/PCG/Source/PCG/Private/Tests/PCGTestsCommon.cpp b/Engine/Plugins/Experimental/PCG/Source/PCG/Private/Tests/PCGTestsCommon.cpp
--- a/Engine/Plugins/Experimental/PCG/Source/PCG/Private/Tests/PCGTestsCommon.cpp
+++ b/Engine/Plugins/Experimental/PCG/Source/PCG/Private/Tests/PCGTestsCommon.cpp
@@ -88,6 +88,7 @@ namespace PCGTestsCommon
 	{
 		check(InElement);
 		TUniquePtr<FPCGContext> Context{ InElement->Initialize(InputData, InSourceComponent, InNode) };
+		Context->InitializeSettings();
 		Context->AsyncState.NumAvailableTasks = 1;
 		return Context;
 	}
@@ -149,8 +150,6 @@ namespace PCGTestsCommon
 			FPCGPoint& Point = Points.Emplace_GetRef(FTransform(Rotation, Location, Scale), 1.f, I);
 			Point.Color = RandomSource.VRand();
 			Point.Density = bRandomDensity ? RandomSource.GetFraction() : 1.0f;
-			
-			Point.Seed = I;
 		}
 
 		return PointData;
