[PCG] Replace search for overlapping components with direct lookup.

Significant speedup for higen, I've seen up to 30% of the execution time go into this search.

#rnx

diff --git a/Engine/Plugins/Experimental/PCG/Source/PCG/Private/Graph/PCGGraphExecutor.cpp b/Engine/Plugins/Experimental/PCG/Source/PCG/Private/Graph/PCGGraphExecutor.cpp
--- a/Engine/Plugins/Experimental/PCG/Source/PCG/Private/Graph/PCGGraphExecutor.cpp
+++ b/Engine/Plugins/Experimental/PCG/Source/PCG/Private/Graph/PCGGraphExecutor.cpp
@@ -17,6 +17,7 @@
 #include "Graph/PCGGraphCache.h"
 #include "Graph/PCGStackContext.h"
 #include "Grid/PCGPartitionActor.h"
+#include "Helpers/PCGActorHelpers.h"
 #include "Helpers/PCGHelpers.h"
 #include "Metadata/PCGMetadata.h"
 #include "Utils/PCGGraphExecutionLogging.h"
@@ -1464,6 +1465,7 @@ namespace PCGGraphExecutor
 		const UPCGNode* InDownstreamNode,
 		FPCGGridLinkageContext* InContext)
 	{
+		check(InContext && InContext->SourceComponent.IsValid());
 		TRACE_CPUPROFILER_EVENT_SCOPE(PCGGraphExecutor::ExecuteGridLinkage);
 
 		// Non-hierarchical generation - no linkage required - data should just pass through.
@@ -1540,15 +1542,14 @@ namespace PCGGraphExecutor
 			}
 			else
 			{
-				TRACE_CPUPROFILER_EVENT_SCOPE(PCGGraphExecutor::ExecuteGridLinkage::FindComponent);
-
-				Subsystem->ForAllOverlappingComponentsInHierarchy(InContext->SourceComponent.Get(), [FromGridSize, &ComponentWithData](UPCGComponent* InLocalComponent)
+				const APCGWorldActor* PCGWorldActor = Subsystem->GetPCGWorldActor();
+				const AActor* ComponentActor = InContext->SourceComponent->GetOwner();
+				if (PCGWorldActor && ComponentActor)
 				{
-					if (InLocalComponent->GetGenerationGridSize() == FromGridSize)
-					{
-						ComponentWithData = InLocalComponent;
-					}
-				});
+					// Get grid coords using the parent grid (FromGridSize).
+					const FIntVector CellCoords = UPCGActorHelpers::GetCellCoord(ComponentActor->GetActorLocation(), FromGridSize, PCGWorldActor->bUse2DGrid);
+					ComponentWithData = Subsystem->GetLocalComponent(FromGridSize, CellCoords, InContext->SourceComponent->GetOriginalComponent());
+				}
 			}
 
 			if (!ComponentWithData)
diff --git a/Engine/Plugins/Experimental/PCG/Source/PCG/Private/PCGActorAndComponentMapping.cpp b/Engine/Plugins/Experimental/PCG/Source/PCG/Private/PCGActorAndComponentMapping.cpp
--- a/Engine/Plugins/Experimental/PCG/Source/PCG/Private/PCGActorAndComponentMapping.cpp
+++ b/Engine/Plugins/Experimental/PCG/Source/PCG/Private/PCGActorAndComponentMapping.cpp
@@ -623,6 +623,30 @@ TSet<UPCGComponent*> UPCGActorAndComponentMapping::GetAllRegisteredComponents()
 	return Res;
 }
 
+UPCGComponent* UPCGActorAndComponentMapping::GetLocalComponent(uint32 GridSize, const FIntVector& CellCoords, const UPCGComponent* InOriginalComponent)
+{
+	FReadScopeLock ReadLock(PartitionActorsMapLock);
+
+	if (const TMap<FIntVector, TObjectPtr<APCGPartitionActor>>* PartitionActorsOnGrid = PartitionActorsMap.Find(GridSize))
+	{
+		const TObjectPtr<APCGPartitionActor>* PartitionActor = PartitionActorsOnGrid->Find(CellCoords);
+		if (PartitionActor && *PartitionActor)
+		{
+			TArray<UPCGComponent*, TInlineAllocator<4>> PCGComponents;
+			(*PartitionActor)->GetComponents(PCGComponents);
+
+			UPCGComponent** MatchingComponent = PCGComponents.FindByPredicate([InOriginalComponent](UPCGComponent* Comp)
+			{
+				return Comp->GetOriginalComponent() == InOriginalComponent;
+			});
+
+			return MatchingComponent ? *MatchingComponent : nullptr;
+		}
+	}
+
+	return nullptr;
+}
+
 #if WITH_EDITOR
 void UPCGActorAndComponentMapping::RegisterOrUpdateTracking(UPCGComponent* InComponent, bool bInShouldDirtyActors)
 {
diff --git a/Engine/Plugins/Experimental/PCG/Source/PCG/Private/PCGSubsystem.cpp b/Engine/Plugins/Experimental/PCG/Source/PCG/Private/PCGSubsystem.cpp
--- a/Engine/Plugins/Experimental/PCG/Source/PCG/Private/PCGSubsystem.cpp
+++ b/Engine/Plugins/Experimental/PCG/Source/PCG/Private/PCGSubsystem.cpp
@@ -652,6 +652,11 @@ void UPCGSubsystem::ForAllOverlappingComponentsInHierarchy(UPCGComponent* InComp
 	});
 }
 
+UPCGComponent* UPCGSubsystem::GetLocalComponent(uint32 GridSize, const FIntVector& CellCoords, const UPCGComponent* InOriginalComponent)
+{
+	return ActorAndComponentMapping.GetLocalComponent(GridSize, CellCoords, InOriginalComponent);
+}
+
 bool UPCGSubsystem::IsGraphCacheDebuggingEnabled() const
 {
 	return GraphExecutor && GraphExecutor->IsGraphCacheDebuggingEnabled();
diff --git a/Engine/Plugins/Experimental/PCG/Source/PCG/Public/PCGActorAndComponentMapping.h b/Engine/Plugins/Experimental/PCG/Source/PCG/Public/PCGActorAndComponentMapping.h
--- a/Engine/Plugins/Experimental/PCG/Source/PCG/Public/PCGActorAndComponentMapping.h
+++ b/Engine/Plugins/Experimental/PCG/Source/PCG/Public/PCGActorAndComponentMapping.h
@@ -21,6 +21,7 @@ class APCGPartitionActor;
 class FLandscapeProxyComponentDataChangedParams;
 class UObject;
 class UPCGComponent;
+class UPCGGraph;
 class UPCGSubsystem;
 
 /**
@@ -78,6 +79,9 @@ public:
 	/** Return a copy of all the registered components. Thread safe */
 	TSet<UPCGComponent*> GetAllRegisteredComponents() const;
 
+	/** Retrieves a local component using grid size and grid coordinates, returns nullptr if no such component found. */
+	UPCGComponent* GetLocalComponent(uint32 GridSize, const FIntVector& CellCoords, const UPCGComponent* InOriginalComponent);
+
 private:
 	// This class is only meant to be used as part of the PCG Subsytem and owned by it.
 	// So we put constructors private.
diff --git a/Engine/Plugins/Experimental/PCG/Source/PCG/Public/PCGSubsystem.h b/Engine/Plugins/Experimental/PCG/Source/PCG/Public/PCGSubsystem.h
--- a/Engine/Plugins/Experimental/PCG/Source/PCG/Public/PCGSubsystem.h
+++ b/Engine/Plugins/Experimental/PCG/Source/PCG/Public/PCGSubsystem.h
@@ -154,6 +154,9 @@ public:
 	/** Traverses the hierarchy associated with the given component and calls InFunc for each overlapping component. */
 	void ForAllOverlappingComponentsInHierarchy(UPCGComponent* InComponent, const TFunction<void(UPCGComponent*)>& InFunc) const;
 
+	/** Retrieves a local component using grid size and grid coordinates, returns nullptr if no such component found. */
+	UPCGComponent* GetLocalComponent(uint32 GridSize, const FIntVector& CellCoords, const UPCGComponent* InOriginalComponent);
+
 	/** True if graph cache debugging is enabled. */
 	bool IsGraphCacheDebuggingEnabled() const;
 
