[PCG] Improve performance when doing changes in a PCG graph regarding tracking

* Update the tracking only if keys have changed
* Add the Landscape as a tracked settings in the component, to avoid custom logic for that in the tracking
* Also unregister actors that are no longer tracked by the component.
* Also rename UPCGActorAndComponentMapping to FPCGActorAndComponentMapping, since it is not an UObject

diff --git a/Engine/Plugins/Experimental/PCG/Source/PCG/Private/PCGActorAndComponentMapping.cpp b/Engine/Plugins/Experimental/PCG/Source/PCG/Private/PCGActorAndComponentMapping.cpp
--- a/Engine/Plugins/Experimental/PCG/Source/PCG/Private/PCGActorAndComponentMapping.cpp
+++ b/Engine/Plugins/Experimental/PCG/Source/PCG/Private/PCGActorAndComponentMapping.cpp
@@ -50,7 +50,7 @@ namespace PCGActorAndComponentMapping
 
 	static TAutoConsoleVariable<bool> CVarDisableDelayedUnregister(
 		TEXT("pcg.DisableDelayedUnregister"),
-		false,
+		true,
 		TEXT("If delayed unregister for all is introducing bad behavior, disables it, allowing people to continue working while we investigate."));
 
 #if WITH_EDITOR
@@ -72,7 +72,7 @@ namespace PCGActorAndComponentMapping
 #endif
 }
 
-UPCGActorAndComponentMapping::UPCGActorAndComponentMapping(UPCGSubsystem* InPCGSubsystem)
+FPCGActorAndComponentMapping::FPCGActorAndComponentMapping(UPCGSubsystem* InPCGSubsystem)
 	: PCGSubsystem(InPCGSubsystem)
 {
 	check(PCGSubsystem);
@@ -84,7 +84,7 @@ UPCGActorAndComponentMapping::UPCGActorAndComponentMapping(UPCGSubsystem* InPCGS
 	NonPartitionedOctree.Reset(FVector::ZeroVector, OctreeExtent);
 }
 
-void UPCGActorAndComponentMapping::Tick()
+void FPCGActorAndComponentMapping::Tick()
 {
 	TSet<UPCGComponent*> ComponentToUnregister;
 	{
@@ -102,9 +102,9 @@ void UPCGActorAndComponentMapping::Tick()
 #endif // WITH_EDITOR
 }
 
-TArray<FPCGTaskId> UPCGActorAndComponentMapping::DispatchToRegisteredLocalComponents(UPCGComponent* OriginalComponent, const TFunction<FPCGTaskId(UPCGComponent*)>& InFunc) const
+TArray<FPCGTaskId> FPCGActorAndComponentMapping::DispatchToRegisteredLocalComponents(UPCGComponent* OriginalComponent, const TFunction<FPCGTaskId(UPCGComponent*)>& InFunc) const
 {
-	TRACE_CPUPROFILER_EVENT_SCOPE(UPCGActorAndComponentMapping::DispatchToRegisteredLocalComponents);
+	TRACE_CPUPROFILER_EVENT_SCOPE(FPCGActorAndComponentMapping::DispatchToRegisteredLocalComponents);
 
 	// TODO: Might be more interesting to copy the set and release the lock.
 	FReadScopeLock ReadLock(ComponentToPartitionActorsMapLock);
@@ -118,7 +118,7 @@ TArray<FPCGTaskId> UPCGActorAndComponentMapping::DispatchToRegisteredLocalCompon
 	return DispatchToLocalComponents(OriginalComponent, *PartitionActorsPtr, InFunc);
 }
 
-TArray<FPCGTaskId> UPCGActorAndComponentMapping::DispatchToLocalComponents(UPCGComponent* OriginalComponent, const TSet<TObjectPtr<APCGPartitionActor>>& PartitionActors, const TFunction<FPCGTaskId(UPCGComponent*)>& InFunc) const
+TArray<FPCGTaskId> FPCGActorAndComponentMapping::DispatchToLocalComponents(UPCGComponent* OriginalComponent, const TSet<TObjectPtr<APCGPartitionActor>>& PartitionActors, const TFunction<FPCGTaskId(UPCGComponent*)>& InFunc) const
 {
 	TArray<FPCGTaskId> TaskIds;
 	for (APCGPartitionActor* PartitionActor : PartitionActors)
@@ -144,7 +144,7 @@ TArray<FPCGTaskId> UPCGActorAndComponentMapping::DispatchToLocalComponents(UPCGC
 	return TaskIds;
 }
 
-bool UPCGActorAndComponentMapping::RegisterOrUpdatePCGComponent(UPCGComponent* InComponent, bool bDoActorMapping)
+bool FPCGActorAndComponentMapping::RegisterOrUpdatePCGComponent(UPCGComponent* InComponent, bool bDoActorMapping)
 {
 	check(InComponent);
 
@@ -161,7 +161,7 @@ bool UPCGActorAndComponentMapping::RegisterOrUpdatePCGComponent(UPCGComponent* I
 		return false;
 	}
 
-	const bool bWasAlreadyRegistered = NonPartitionedOctree.Contains(InComponent) || PartitionedOctree.Contains(InComponent);
+	const bool bWasAlreadyRegistered = IsComponentRegistered(InComponent);
 
 	// First check if the component has changed its partitioned flag.
 	const bool bIsPartitioned = InComponent->IsPartitioned();
@@ -195,14 +195,14 @@ bool UPCGActorAndComponentMapping::RegisterOrUpdatePCGComponent(UPCGComponent* I
 #if WITH_EDITOR
 	if (!bWasAlreadyRegistered && bHasChanged)
 	{
-		RegisterOrUpdateTracking(InComponent, /*bInShouldDirtyActors=*/ false);
+		RegisterTracking(InComponent);
 	}
 #endif // WITH_EDITOR
 
 	return bHasChanged;
 }
 
-bool UPCGActorAndComponentMapping::RegisterOrUpdatePartitionedPCGComponent(UPCGComponent* InComponent, bool bDoActorMapping)
+bool FPCGActorAndComponentMapping::RegisterOrUpdatePartitionedPCGComponent(UPCGComponent* InComponent, bool bDoActorMapping)
 {
 	FBox Bounds(EForceInit::ForceInit);
 	bool bComponentHasChanged = false;
@@ -243,7 +243,7 @@ bool UPCGActorAndComponentMapping::RegisterOrUpdatePartitionedPCGComponent(UPCGC
 	return bComponentHasChanged;
 }
 
-bool UPCGActorAndComponentMapping::RegisterOrUpdateNonPartitionedPCGComponent(UPCGComponent* InComponent)
+bool FPCGActorAndComponentMapping::RegisterOrUpdateNonPartitionedPCGComponent(UPCGComponent* InComponent)
 {
 	// Tracking is only done in Editor for now
 #if WITH_EDITOR
@@ -259,7 +259,7 @@ bool UPCGActorAndComponentMapping::RegisterOrUpdateNonPartitionedPCGComponent(UP
 #endif // WITH_EDITOR
 }
 
-bool UPCGActorAndComponentMapping::RemapPCGComponent(const UPCGComponent* OldComponent, UPCGComponent* NewComponent, bool bDoActorMapping)
+bool FPCGActorAndComponentMapping::RemapPCGComponent(const UPCGComponent* OldComponent, UPCGComponent* NewComponent, bool bDoActorMapping)
 {
 	check(OldComponent && NewComponent);
 
@@ -318,7 +318,7 @@ bool UPCGActorAndComponentMapping::RemapPCGComponent(const UPCGComponent* OldCom
 	return true;
 }
 
-void UPCGActorAndComponentMapping::UnregisterPCGComponent(UPCGComponent* InComponent, bool bForce)
+void FPCGActorAndComponentMapping::UnregisterPCGComponent(UPCGComponent* InComponent, bool bForce)
 {
 	if (!InComponent)
 	{
@@ -360,7 +360,7 @@ void UPCGActorAndComponentMapping::UnregisterPCGComponent(UPCGComponent* InCompo
 	}
 }
 
-void UPCGActorAndComponentMapping::UnregisterPartitionedPCGComponent(UPCGComponent* InComponent)
+void FPCGActorAndComponentMapping::UnregisterPartitionedPCGComponent(UPCGComponent* InComponent)
 {
 	if (!PartitionedOctree.RemoveComponent(InComponent))
 	{
@@ -387,12 +387,12 @@ void UPCGActorAndComponentMapping::UnregisterPartitionedPCGComponent(UPCGCompone
 	}
 }
 
-void UPCGActorAndComponentMapping::UnregisterNonPartitionedPCGComponent(UPCGComponent* InComponent)
+void FPCGActorAndComponentMapping::UnregisterNonPartitionedPCGComponent(UPCGComponent* InComponent)
 {
 	NonPartitionedOctree.RemoveComponent(InComponent);
 }
 
-void UPCGActorAndComponentMapping::ForAllIntersectingComponents(const FBoxCenterAndExtent& InBounds, TFunction<void(UPCGComponent*)> InFunc) const
+void FPCGActorAndComponentMapping::ForAllIntersectingComponents(const FBoxCenterAndExtent& InBounds, TFunction<void(UPCGComponent*)> InFunc) const
 {
 	PartitionedOctree.FindElementsWithBoundsTest(InBounds, [&InFunc](const FPCGComponentRef& ComponentRef)
 	{
@@ -400,7 +400,7 @@ void UPCGActorAndComponentMapping::ForAllIntersectingComponents(const FBoxCenter
 	});
 }
 
-void UPCGActorAndComponentMapping::RegisterPartitionActor(APCGPartitionActor* Actor, bool bDoComponentMapping)
+void FPCGActorAndComponentMapping::RegisterPartitionActor(APCGPartitionActor* Actor, bool bDoComponentMapping)
 {
 	check(Actor);
 
@@ -448,7 +448,7 @@ void UPCGActorAndComponentMapping::RegisterPartitionActor(APCGPartitionActor* Ac
 	}
 }
 
-void UPCGActorAndComponentMapping::UnregisterPartitionActor(APCGPartitionActor* Actor)
+void FPCGActorAndComponentMapping::UnregisterPartitionActor(APCGPartitionActor* Actor)
 {
 	check(Actor);
 
@@ -474,7 +474,7 @@ void UPCGActorAndComponentMapping::UnregisterPartitionActor(APCGPartitionActor*
 	}
 }
 
-void UPCGActorAndComponentMapping::ForAllIntersectingPartitionActors(const FBox& InBounds, TFunction<void(APCGPartitionActor*)> InFunc) const
+void FPCGActorAndComponentMapping::ForAllIntersectingPartitionActors(const FBox& InBounds, TFunction<void(APCGPartitionActor*)> InFunc) const
 {
 	// No PCGWorldActor just early out. Same for invalid bounds.
 	APCGWorldActor* PCGWorldActor = PCGSubsystem->GetPCGWorldActor();
@@ -522,7 +522,7 @@ void UPCGActorAndComponentMapping::ForAllIntersectingPartitionActors(const FBox&
 	}
 }
 
-void UPCGActorAndComponentMapping::UpdateMappingPCGComponentPartitionActor(UPCGComponent* InComponent)
+void FPCGActorAndComponentMapping::UpdateMappingPCGComponentPartitionActor(UPCGComponent* InComponent)
 {
 	if (!PCGSubsystem->IsInitialized())
 	{
@@ -583,7 +583,7 @@ void UPCGActorAndComponentMapping::UpdateMappingPCGComponentPartitionActor(UPCGC
 	}
 }
 
-void UPCGActorAndComponentMapping::DeleteMappingPCGComponentPartitionActor(UPCGComponent* InComponent)
+void FPCGActorAndComponentMapping::DeleteMappingPCGComponentPartitionActor(UPCGComponent* InComponent)
 {
 	check(InComponent);
 
@@ -606,24 +606,29 @@ void UPCGActorAndComponentMapping::DeleteMappingPCGComponentPartitionActor(UPCGC
 	}
 }
 
-TSet<UPCGComponent*> UPCGActorAndComponentMapping::GetAllRegisteredPartitionedComponents() const
+bool FPCGActorAndComponentMapping::IsComponentRegistered(const UPCGComponent* InComponent) const
+{
+	return PartitionedOctree.Contains(InComponent) || NonPartitionedOctree.Contains(InComponent);
+}
+
+TSet<UPCGComponent*> FPCGActorAndComponentMapping::GetAllRegisteredPartitionedComponents() const
 {
 	return PartitionedOctree.GetAllComponents();
 }
 
-TSet<UPCGComponent*> UPCGActorAndComponentMapping::GetAllRegisteredNonPartitionedComponents() const
+TSet<UPCGComponent*> FPCGActorAndComponentMapping::GetAllRegisteredNonPartitionedComponents() const
 {
 	return NonPartitionedOctree.GetAllComponents();
 }
 
-TSet<UPCGComponent*> UPCGActorAndComponentMapping::GetAllRegisteredComponents() const
+TSet<UPCGComponent*> FPCGActorAndComponentMapping::GetAllRegisteredComponents() const
 {
 	TSet<UPCGComponent*> Res = GetAllRegisteredPartitionedComponents();
 	Res.Append(GetAllRegisteredNonPartitionedComponents());
 	return Res;
 }
 
-UPCGComponent* UPCGActorAndComponentMapping::GetLocalComponent(uint32 GridSize, const FIntVector& CellCoords, const UPCGComponent* InOriginalComponent)
+UPCGComponent* FPCGActorAndComponentMapping::GetLocalComponent(uint32 GridSize, const FIntVector& CellCoords, const UPCGComponent* InOriginalComponent)
 {
 	FReadScopeLock ReadLock(PartitionActorsMapLock);
 
@@ -648,7 +653,7 @@ UPCGComponent* UPCGActorAndComponentMapping::GetLocalComponent(uint32 GridSize,
 }
 
 #if WITH_EDITOR
-void UPCGActorAndComponentMapping::RegisterOrUpdateTracking(UPCGComponent* InComponent, bool bInShouldDirtyActors)
+void FPCGActorAndComponentMapping::RegisterTracking(UPCGComponent* InComponent)
 {
 	// Discard BP templates, local components and invalid component
 	if (!IsValid(InComponent) || !InComponent->GetOwner() || InComponent->GetOwner()->IsA<APCGPartitionActor>())
@@ -673,35 +678,68 @@ void UPCGActorAndComponentMapping::RegisterOrUpdateTracking(UPCGComponent* InCom
 
 	// Components owner needs to be always tracked
 	RegisterActor(ComponentOwner);
-	TSet<UPCGComponent*>& AllComponents = AlwaysTrackedActorsToComponentsMap.FindOrAdd(ComponentOwner);
-	AllComponents.Add(InComponent);
+	AlwaysTrackedActorsToComponentsMap.FindOrAdd(ComponentOwner).Add(InComponent);
 
 	FPCGActorSelectionKey OwnerKey = FPCGActorSelectionKey(EPCGActorFilter::Self);
 	KeysToComponentsMap.FindOrAdd(OwnerKey).Add(InComponent);
 
-	const bool bDisableDelayedActorRegistering = PCGActorAndComponentMapping::CVarDisableDelayedActorRegistering.GetValueOnAnyThread();
+	UpdateTracking(InComponent, /*bInShouldDirtyActors=*/ false);
+
+	// Add tracking for when the graph was generated/cleaned, only once
+	if (!InComponent->OnPCGGraphGeneratedDelegate.IsBoundToObject(this))
+	{
+		InComponent->OnPCGGraphGeneratedDelegate.AddRaw(this, &FPCGActorAndComponentMapping::OnPCGGraphGeneratedOrCleaned);
+		InComponent->OnPCGGraphCleanedDelegate.AddRaw(this, &FPCGActorAndComponentMapping::OnPCGGraphGeneratedOrCleaned);
+	}
+}
+
+void FPCGActorAndComponentMapping::UpdateTracking(UPCGComponent* InComponent, bool bInShouldDirtyActors, const TArray<FPCGActorSelectionKey>* ChangedKeys)
+{
+	// Discard BP templates, local components and invalid component
+	if (!IsValid(InComponent) || !InComponent->GetOwner() || InComponent->GetOwner()->IsA<APCGPartitionActor>())
+	{
+		return;
+	}
+
+	AActor* ComponentOwner = InComponent->GetOwner();
+
+	// If we have no owner, we might be in a BP so don't track
+	if (!ComponentOwner)
+	{
+		return;
+	}
+
+	UWorld* World = PCGSubsystem ? PCGSubsystem->GetWorld() : nullptr;
+
+	if (!World)
+	{
+		return;
+	}
+
+	// If no keys are provided, update all tracking keys.
+	TArray<FPCGActorSelectionKey> AllKeys;
+	if (ChangedKeys == nullptr)
+	{
+		InComponent->CachedTrackedKeysToSettings.GenerateKeyArray(AllKeys);
+		ChangedKeys = &AllKeys;
+	}
+
+	check(ChangedKeys);
 
 	// And we also need to find all actors that should be tracked
 	if (UPCGGraph* PCGGraph = InComponent->GetGraph())
 	{
-		auto FindActorsAndTrack = [this, InComponent, bInShouldDirtyActors, bDisableDelayedActorRegistering](const FPCGActorSelectionKey& InKey, const TArray<FPCGSettingsAndCulling>& InSettingsAndCulling)
+		TSet<AActor*> CandidatesForTracking;
+		TSet<TObjectKey<AActor>> CandidatesForUntracking;
+		TSet<FPCGActorSelectionKey> CandidateKeysToRemove;
+
+		auto GatherActors = [this, InComponent, &CandidatesForTracking, &CandidatesForUntracking, &CandidateKeysToRemove](const FPCGActorSelectionKey& InKey, bool bInShouldUntrack)
 		{
 			// InKey provide the info for selecting a given actor.
 			// We reconstruct the selector settings from this key, and we also force it to SelectMultiple, since
 			// we want to gather all the actors that matches this given key.
 			FPCGActorSelectorSettings SelectorSettings = FPCGActorSelectorSettings::ReconstructFromKey(InKey);
 			SelectorSettings.bSelectMultiple = true;
-
-			bool bShouldCull = true;
-			for (const FPCGSettingsAndCulling& SettingsAndCulling : InSettingsAndCulling)
-			{
-				if (!SettingsAndCulling.Value)
-				{
-					bShouldCull = false;
-					break;
-				}
-			}
-
 			TArray<AActor*> AllActors = PCGActorSelector::FindActors(SelectorSettings, InComponent, [](const AActor*) { return true; }, [](const AActor*) { return true; });
 
 			for (AActor* Actor : AllActors)
@@ -711,60 +749,67 @@ void UPCGActorAndComponentMapping::RegisterOrUpdateTracking(UPCGComponent* InCom
 					continue;
 				}
 
-				if (!Actor->HasActorRegisteredAllComponents() && !bDisableDelayedActorRegistering)
+				if (bInShouldUntrack && !CandidatesForTracking.Contains(Actor))
 				{
-					DelayedAddedActors.Emplace({ Actor, bInShouldDirtyActors });
-					continue;
-				}
-
-				if (bShouldCull)
-				{
-					CulledTrackedActorsToComponentsMap.FindOrAdd(Actor).Add(InComponent);
+					CandidatesForUntracking.Add(Actor);
+					CandidateKeysToRemove.Add(InKey);
 				}
 				else
 				{
-					AlwaysTrackedActorsToComponentsMap.FindOrAdd(Actor).Add(InComponent);
-				}
-
-				RegisterActor(Actor);
-
-				if (bInShouldDirtyActors)
-				{
-					// If we need to force dirty, disregard culling (always intersect).
-					InComponent->DirtyTrackedActor(Actor, /*bIntersect=*/ true, /*InRemovedTags=*/ {}, /*InOriginatingChangeObject=*/ nullptr);
+					CandidatesForTracking.Add(Actor);
+					CandidatesForUntracking.Remove(Actor);
+					CandidateKeysToRemove.Remove(InKey);
 				}
 			}
 		};
 
-		for (TPair<FPCGActorSelectionKey, TArray<FPCGSettingsAndCulling>>& It : PCGGraph->GetTrackedActorKeysToSettings())
+		for (const FPCGActorSelectionKey& Key : *ChangedKeys)
 		{
-			FindActorsAndTrack(It.Key, It.Value);
-			KeysToComponentsMap.FindOrAdd(It.Key).Add(InComponent);
+			const bool bShouldUntrack = !InComponent->CachedTrackedKeysToSettings.Contains(Key);
+			GatherActors(Key, bShouldUntrack);
+			if (!bShouldUntrack)
+			{
+				KeysToComponentsMap.FindOrAdd(Key).Add(InComponent);
+			}
 		}
 
-		// Also while we support landscape pins on input node, we need to track landscape if we uses it, or the input is landscape.
-		if (InComponent->ShouldTrackLandscape())
+		const bool bDisableDelayedActorRegistering = PCGActorAndComponentMapping::CVarDisableDelayedActorRegistering.GetValueOnAnyThread();
+		for (AActor* Actor : CandidatesForTracking)
 		{
-			// Landscape doesn't have an associated setting and is always culled.
-			FPCGActorSelectionKey LandscapeKey = FPCGActorSelectionKey(ALandscapeProxy::StaticClass());
-			FindActorsAndTrack(LandscapeKey, { {nullptr, true} });
-			KeysToComponentsMap.FindOrAdd(LandscapeKey).Add(InComponent);
+			if (!Actor->HasActorRegisteredAllComponents() && !bDisableDelayedActorRegistering)
+			{
+				DelayedAddedActors.Emplace({ Actor, false });
+				continue;
+			}
+
+			bool bShouldCull = false;
+			if (!InComponent->IsActorTracked(Actor, bShouldCull))
+			{
+				continue;
+			}
+
+			if (bShouldCull)
+			{
+				CulledTrackedActorsToComponentsMap.FindOrAdd(Actor).Add(InComponent);
+			}
+			else
+			{
+				AlwaysTrackedActorsToComponentsMap.FindOrAdd(Actor).Add(InComponent);
+			}
+
+			RegisterActor(Actor);
 		}
-	}
 
-	// Add tracking for when the graph was generated/cleaned, only once
-	if (!InComponent->OnPCGGraphGeneratedDelegate.IsBoundToObject(this))
-	{
-		InComponent->OnPCGGraphGeneratedDelegate.AddRaw(this, &UPCGActorAndComponentMapping::OnPCGGraphGeneratedOrCleaned);
-		InComponent->OnPCGGraphCleanedDelegate.AddRaw(this, &UPCGActorAndComponentMapping::OnPCGGraphGeneratedOrCleaned);
+		// Also unregister all keys that are not tracked anymore.
+		UnregisterTracking(InComponent, &CandidatesForUntracking, &CandidateKeysToRemove);
 	}
 }
 
-void UPCGActorAndComponentMapping::RemapTracking(const UPCGComponent* InOldComponent, UPCGComponent* InNewComponent)
+void FPCGActorAndComponentMapping::RemapTracking(const UPCGComponent* InOldComponent, UPCGComponent* InNewComponent)
 {
-	auto ReplaceInMap = [InOldComponent, InNewComponent](auto& InMap)
+	auto ReplaceInMap = [InOldComponent, InNewComponent]<typename Key>(TMap<Key, TSet<UPCGComponent*>>& InMap)
 	{
-		for (auto& It : InMap)
+		for (TPair<Key, TSet<UPCGComponent*>>& It : InMap)
 		{
 			if (It.Value.Remove(InOldComponent) > 0)
 			{
@@ -787,12 +832,12 @@ void UPCGActorAndComponentMapping::RemapTracking(const UPCGComponent* InOldCompo
 	// And just making sure we are not registering multiple times
 	if (!InNewComponent->OnPCGGraphGeneratedDelegate.IsBoundToObject(this))
 	{
-		InNewComponent->OnPCGGraphGeneratedDelegate.AddRaw(this, &UPCGActorAndComponentMapping::OnPCGGraphGeneratedOrCleaned);
-		InNewComponent->OnPCGGraphCleanedDelegate.AddRaw(this, &UPCGActorAndComponentMapping::OnPCGGraphGeneratedOrCleaned);
+		InNewComponent->OnPCGGraphGeneratedDelegate.AddRaw(this, &FPCGActorAndComponentMapping::OnPCGGraphGeneratedOrCleaned);
+		InNewComponent->OnPCGGraphCleanedDelegate.AddRaw(this, &FPCGActorAndComponentMapping::OnPCGGraphGeneratedOrCleaned);
 	}
 }
 
-void UPCGActorAndComponentMapping::UnregisterTracking(UPCGComponent* InComponent)
+void FPCGActorAndComponentMapping::UnregisterTracking(UPCGComponent* InComponent, const TSet<TObjectKey<AActor>>* OptionalActorsToUntrack, const TSet<FPCGActorSelectionKey>* OptionalKeysToUntrack)
 {
 	if (!InComponent)
 	{
@@ -802,9 +847,9 @@ void UPCGActorAndComponentMapping::UnregisterTracking(UPCGComponent* InComponent
 	TSet<TObjectKey<AActor>> CandidatesForUntrack;
 	TSet<FPCGActorSelectionKey> KeysToRemove;
 
-	auto RemoveFromMap = [InComponent](auto& InMap, auto& InCandidateToRemove)
+	auto RemoveAllFromMap = [InComponent]<typename Key>(TMap<Key, TSet<UPCGComponent*>>& InMap, TSet<Key>& InCandidateToRemove)
 	{
-		for (auto& It : InMap)
+		for (TPair<Key, TSet<UPCGComponent*>>& It : InMap)
 		{
 			It.Value.Remove(InComponent);
 			if (It.Value.IsEmpty())
@@ -814,9 +859,39 @@ void UPCGActorAndComponentMapping::UnregisterTracking(UPCGComponent* InComponent
 		}
 	};
 
-	RemoveFromMap(CulledTrackedActorsToComponentsMap, CandidatesForUntrack);
-	RemoveFromMap(AlwaysTrackedActorsToComponentsMap, CandidatesForUntrack);
-	RemoveFromMap(KeysToComponentsMap, KeysToRemove);
+	auto RemoveKeysFromMap = [InComponent]<typename Key>(TMap<Key, TSet<UPCGComponent*>>& InMap, TSet<Key>& InCandidateToRemove, const TSet<Key>& SetToIterateOn)
+	{
+		for (const Key& KeyIt : SetToIterateOn)
+		{
+			if (TSet<UPCGComponent*>* ComponentSetPtr = InMap.Find(KeyIt))
+			{
+				ComponentSetPtr->Remove(InComponent);
+				if (ComponentSetPtr->IsEmpty())
+				{
+					InCandidateToRemove.Add(KeyIt);
+				}
+			}
+		}
+	};
+
+	if (OptionalActorsToUntrack)
+	{
+		RemoveKeysFromMap(CulledTrackedActorsToComponentsMap, CandidatesForUntrack, *OptionalActorsToUntrack);
+		RemoveKeysFromMap(AlwaysTrackedActorsToComponentsMap, CandidatesForUntrack, *OptionalActorsToUntrack);
+	}
+	else
+	{
+		RemoveAllFromMap(CulledTrackedActorsToComponentsMap, CandidatesForUntrack);
+	}
+
+	if (OptionalKeysToUntrack)
+	{
+		RemoveKeysFromMap(KeysToComponentsMap, KeysToRemove, *OptionalKeysToUntrack);
+	}
+	else
+	{
+		RemoveAllFromMap(KeysToComponentsMap, KeysToRemove);
+	}
 
 	for (const FPCGActorSelectionKey& Key : KeysToRemove)
 	{
@@ -837,28 +912,38 @@ void UPCGActorAndComponentMapping::UnregisterTracking(UPCGComponent* InComponent
 			UnregisterActor(Candidate.ResolveObjectPtr());
 		}
 	}
+}
+
+void FPCGActorAndComponentMapping::UnregisterTracking(UPCGComponent* InComponent)
+{
+	if (!InComponent)
+	{
+		return;
+	}
+
+	UnregisterTracking(InComponent, nullptr, nullptr);
 
 	InComponent->OnPCGGraphGeneratedDelegate.RemoveAll(this);
 	InComponent->OnPCGGraphCleanedDelegate.RemoveAll(this);
 }
 
-void UPCGActorAndComponentMapping::ResetPartitionActorsMap()
+void FPCGActorAndComponentMapping::ResetPartitionActorsMap()
 {
 	PartitionActorsMapLock.WriteLock();
 	PartitionActorsMap.Empty();
 	PartitionActorsMapLock.WriteUnlock();
 }
 
-void UPCGActorAndComponentMapping::RegisterTrackingCallbacks()
+void FPCGActorAndComponentMapping::RegisterTrackingCallbacks()
 {
-	GEngine->OnActorMoved().AddRaw(this, &UPCGActorAndComponentMapping::OnActorMoved);
-	GEngine->OnLevelActorAdded().AddRaw(this, &UPCGActorAndComponentMapping::OnActorAdded);
-	GEngine->OnLevelActorDeleted().AddRaw(this, &UPCGActorAndComponentMapping::OnActorDeleted);
-	FCoreUObjectDelegates::OnObjectPropertyChanged.AddRaw(this, &UPCGActorAndComponentMapping::OnObjectPropertyChanged);
-	FCoreUObjectDelegates::OnPreObjectPropertyChanged.AddRaw(this, &UPCGActorAndComponentMapping::OnPreObjectPropertyChanged);
+	GEngine->OnActorMoved().AddRaw(this, &FPCGActorAndComponentMapping::OnActorMoved);
+	GEngine->OnLevelActorAdded().AddRaw(this, &FPCGActorAndComponentMapping::OnActorAdded);
+	GEngine->OnLevelActorDeleted().AddRaw(this, &FPCGActorAndComponentMapping::OnActorDeleted);
+	FCoreUObjectDelegates::OnObjectPropertyChanged.AddRaw(this, &FPCGActorAndComponentMapping::OnObjectPropertyChanged);
+	FCoreUObjectDelegates::OnPreObjectPropertyChanged.AddRaw(this, &FPCGActorAndComponentMapping::OnPreObjectPropertyChanged);
 }
 
-void UPCGActorAndComponentMapping::TeardownTrackingCallbacks()
+void FPCGActorAndComponentMapping::TeardownTrackingCallbacks()
 {
 	GEngine->OnActorMoved().RemoveAll(this);
 	GEngine->OnLevelActorAdded().RemoveAll(this);
@@ -867,7 +952,7 @@ void UPCGActorAndComponentMapping::TeardownTrackingCallbacks()
 	FCoreUObjectDelegates::OnPreObjectPropertyChanged.RemoveAll(this);
 }
 
-void UPCGActorAndComponentMapping::AddDelayedActors()
+void FPCGActorAndComponentMapping::AddDelayedActors()
 {
 	// Safeguard, we can't add delayed actors if the subsystem is not initialized
 	if (!PCGSubsystem || !PCGSubsystem->IsInitialized() || DelayedAddedActors.IsEmpty())
@@ -900,15 +985,15 @@ void UPCGActorAndComponentMapping::AddDelayedActors()
 	DelayedAddedActors = MoveTemp(StillDelayedActors);
 }
 
-void UPCGActorAndComponentMapping::OnActorAdded(AActor* InActor)
+void FPCGActorAndComponentMapping::OnActorAdded(AActor* InActor)
 {
 	// Implementation note: since this is called only for actors directly in the current level, the depth here is 0.
 	OnActorAdded_Internal(InActor, true, 0);
 }
 
-void UPCGActorAndComponentMapping::OnActorAdded_Internal(AActor* InActor, bool bShouldDirty, int32 LevelInstanceDepth)
+void FPCGActorAndComponentMapping::OnActorAdded_Internal(AActor* InActor, bool bShouldDirty, int32 LevelInstanceDepth)
 {
-	TRACE_CPUPROFILER_EVENT_SCOPE(UPCGActorAndComponentMapping::OnActorAdded);
+	TRACE_CPUPROFILER_EVENT_SCOPE(FPCGActorAndComponentMapping::OnActorAdded);
 
 	// We have to make sure to not create a infinite loop
 	if (!InActor || InActor->IsA<APCGWorldActor>() || !PCGSubsystem || InActor->GetWorld() != PCGSubsystem->GetWorld())
@@ -969,7 +1054,7 @@ void UPCGActorAndComponentMapping::OnActorAdded_Internal(AActor* InActor, bool b
 	}
 }
 
-bool UPCGActorAndComponentMapping::AddOrUpdateTrackedActor(AActor* InActor)
+bool FPCGActorAndComponentMapping::AddOrUpdateTrackedActor(AActor* InActor)
 {
 	// We have to make sure to not create a infinite loop
 	if (!InActor || InActor->IsA<APCGWorldActor>() || !PCGSubsystem || InActor->GetWorld() != PCGSubsystem->GetWorld())
@@ -1046,7 +1131,7 @@ bool UPCGActorAndComponentMapping::AddOrUpdateTrackedActor(AActor* InActor)
 	}
 }
 
-void UPCGActorAndComponentMapping::RegisterActor(AActor* InActor)
+void FPCGActorAndComponentMapping::RegisterActor(AActor* InActor)
 {
 	if (!InActor)
 	{
@@ -1058,7 +1143,7 @@ void UPCGActorAndComponentMapping::RegisterActor(AActor* InActor)
 		// Only add it once.
 		if (!TrackedActorToPositionMap.Contains(InActor))
 		{
-			LandscapeProxy->OnComponentDataChanged.AddRaw(this, &UPCGActorAndComponentMapping::OnLandscapeChanged);
+			LandscapeProxy->OnComponentDataChanged.AddRaw(this, &FPCGActorAndComponentMapping::OnLandscapeChanged);
 		}
 	}
 
@@ -1068,7 +1153,7 @@ void UPCGActorAndComponentMapping::RegisterActor(AActor* InActor)
 	UpdateActorDependencies(InActor);
 }
 
-bool UPCGActorAndComponentMapping::UnregisterActor(AActor* InActor)
+bool FPCGActorAndComponentMapping::UnregisterActor(AActor* InActor)
 {
 	if (!InActor)
 	{
@@ -1095,15 +1180,15 @@ bool UPCGActorAndComponentMapping::UnregisterActor(AActor* InActor)
 	}
 }
 
-void UPCGActorAndComponentMapping::OnActorDeleted(AActor* InActor)
+void FPCGActorAndComponentMapping::OnActorDeleted(AActor* InActor)
 {
 	// Implementation note: since this is called only for actors directly in the current level, the depth here is 0.
 	OnActorDeleted_Internal(InActor, 0);
 }
 
-void UPCGActorAndComponentMapping::OnActorDeleted_Internal(AActor* InActor, int32 LevelInstanceDepth)
+void FPCGActorAndComponentMapping::OnActorDeleted_Internal(AActor* InActor, int32 LevelInstanceDepth)
 {
-	TRACE_CPUPROFILER_EVENT_SCOPE(UPCGActorAndComponentMapping::OnActorDeleted);
+	TRACE_CPUPROFILER_EVENT_SCOPE(FPCGActorAndComponentMapping::OnActorDeleted);
 
 	if (!InActor || !PCGSubsystem || InActor->GetWorld() != PCGSubsystem->GetWorld())
 	{
@@ -1158,14 +1243,14 @@ void UPCGActorAndComponentMapping::OnActorDeleted_Internal(AActor* InActor, int3
 	UnregisterActor(InActor);
 }
 
-void UPCGActorAndComponentMapping::OnActorMoved(AActor* InActor)
+void FPCGActorAndComponentMapping::OnActorMoved(AActor* InActor)
 {
 	OnActorMoved_Internal(InActor, /*LevelInstanceDepth=*/0);
 }
 
-void UPCGActorAndComponentMapping::OnActorMoved_Internal(AActor* InActor, int32 LevelInstanceDepth)
+void FPCGActorAndComponentMapping::OnActorMoved_Internal(AActor* InActor, int32 LevelInstanceDepth)
 {
-	TRACE_CPUPROFILER_EVENT_SCOPE(UPCGActorAndComponentMapping::OnActorMoved);
+	TRACE_CPUPROFILER_EVENT_SCOPE(FPCGActorAndComponentMapping::OnActorMoved);
 
 	if (!InActor || (PCGSubsystem && InActor->GetWorld() != PCGSubsystem->GetWorld()))
 	{
@@ -1202,7 +1287,7 @@ void UPCGActorAndComponentMapping::OnActorMoved_Internal(AActor* InActor, int32
 	}
 }
 
-void UPCGActorAndComponentMapping::OnPreObjectPropertyChanged(UObject* InObject, const FEditPropertyChain& InEditPropertyChain)
+void FPCGActorAndComponentMapping::OnPreObjectPropertyChanged(UObject* InObject, const FEditPropertyChain& InEditPropertyChain)
 {
 	// We want to track tags, to see if a tag was removed
 	TempTrackedActorTags.Empty();
@@ -1224,9 +1309,9 @@ void UPCGActorAndComponentMapping::OnPreObjectPropertyChanged(UObject* InObject,
 	TempTrackedActorTags = TSet<FName>(Actor->Tags);
 }
 
-void UPCGActorAndComponentMapping::OnObjectPropertyChanged(UObject* InObject, FPropertyChangedEvent& InEvent)
+void FPCGActorAndComponentMapping::OnObjectPropertyChanged(UObject* InObject, FPropertyChangedEvent& InEvent)
 {
-	TRACE_CPUPROFILER_EVENT_SCOPE(UPCGActorAndComponentMapping::OnObjectPropertyChanged);
+	TRACE_CPUPROFILER_EVENT_SCOPE(FPCGActorAndComponentMapping::OnObjectPropertyChanged);
 
 	const bool bValueNotInteractive = (InEvent.ChangeType != EPropertyChangeType::Interactive);
 	// Special exception for actor tags, as we can't track otherwise an actor "losing" a tag
@@ -1305,9 +1390,9 @@ void UPCGActorAndComponentMapping::OnObjectPropertyChanged(UObject* InObject, FP
 	}
 }
 
-void UPCGActorAndComponentMapping::OnActorChanged(AActor* InActor, bool bInHasMoved, const UObject* InOriginatingChangeObject, int32 LevelInstanceDepth)
+void FPCGActorAndComponentMapping::OnActorChanged(AActor* InActor, bool bInHasMoved, const UObject* InOriginatingChangeObject, int32 LevelInstanceDepth)
 {
-	TRACE_CPUPROFILER_EVENT_SCOPE(UPCGActorAndComponentMapping::OnActorChanged);
+	TRACE_CPUPROFILER_EVENT_SCOPE(FPCGActorAndComponentMapping::OnActorChanged);
 
 	check(InActor);
 	ensure(!PCGSubsystem || InActor->GetWorld() == PCGSubsystem->GetWorld());
@@ -1332,7 +1417,7 @@ void UPCGActorAndComponentMapping::OnActorChanged(AActor* InActor, bool bInHasMo
 		// We first do it for non-partitioned, then we do it for partitioned
 		auto UpdateNonPartitioned = [&DirtyComponents, InActor, CulledTrackedComponents, &RemovedTags, DirtyFlag, InOriginatingChangeObject](const FPCGComponentRef& ComponentRef) -> void
 		{
-			TRACE_CPUPROFILER_EVENT_SCOPE(UPCGActorAndComponentMapping::OnActorChanged::UpdateNonPartitioned);
+			TRACE_CPUPROFILER_EVENT_SCOPE(FPCGActorAndComponentMapping::OnActorChanged::UpdateNonPartitioned);
 
 			// Don't dirty if the component was already dirtied, not tracked, or the origin of the change.
 			if (DirtyComponents.Contains(ComponentRef.Component) || 
@@ -1354,7 +1439,7 @@ void UPCGActorAndComponentMapping::OnActorChanged(AActor* InActor, bool bInHasMo
 		// For partitioned, we first need to find all components that intersect with our actor and then forward the dirty call to all local components that intersect.
 		auto UpdatePartitioned = [this, &DirtyComponents, InActor, CulledTrackedComponents, &ActorBounds, &RemovedTags, DirtyFlag, InOriginatingChangeObject](const FPCGComponentRef& ComponentRef)  -> void
 		{
-			TRACE_CPUPROFILER_EVENT_SCOPE(UPCGActorAndComponentMapping::OnActorChanged::UpdatePartitioned);
+			TRACE_CPUPROFILER_EVENT_SCOPE(FPCGActorAndComponentMapping::OnActorChanged::UpdatePartitioned);
 
 			// Don't dirty if the component is not tracked, or the origin of the change.
 			// We can "re-dirty" it because changes can impact different local components, from the same
@@ -1393,7 +1478,7 @@ void UPCGActorAndComponentMapping::OnActorChanged(AActor* InActor, bool bInHasMo
 		// If it has moved, redo it with the old bounds.
 		if (bInHasMoved)
 		{
-			TRACE_CPUPROFILER_EVENT_SCOPE(UPCGActorAndComponentMapping::OnActorChanged::SecondUpdateHasMoved);
+			TRACE_CPUPROFILER_EVENT_SCOPE(FPCGActorAndComponentMapping::OnActorChanged::SecondUpdateHasMoved);
 
 			if (FBox* OldActorBounds = TrackedActorToPositionMap.Find(InActor))
 			{
@@ -1411,7 +1496,7 @@ void UPCGActorAndComponentMapping::OnActorChanged(AActor* InActor, bool bInHasMo
 	// Finally, dirty all components that always track this actor that are not yet notified.
 	if (TSet<UPCGComponent*>* AlwaysTrackedComponents = AlwaysTrackedActorsToComponentsMap.Find(InActor))
 	{
-		TRACE_CPUPROFILER_EVENT_SCOPE(UPCGActorAndComponentMapping::OnActorChanged::AlwaysTrackedUpdate);
+		TRACE_CPUPROFILER_EVENT_SCOPE(FPCGActorAndComponentMapping::OnActorChanged::AlwaysTrackedUpdate);
 
 		for (UPCGComponent* PCGComponent : *AlwaysTrackedComponents)
 		{
@@ -1490,7 +1575,7 @@ void UPCGActorAndComponentMapping::OnActorChanged(AActor* InActor, bool bInHasMo
 	}
 }
 
-void UPCGActorAndComponentMapping::OnLandscapeChanged(ALandscapeProxy* InLandscape, const FLandscapeProxyComponentDataChangedParams& InChangeParams)
+void FPCGActorAndComponentMapping::OnLandscapeChanged(ALandscapeProxy* InLandscape, const FLandscapeProxyComponentDataChangedParams& InChangeParams)
 {
 	if (!InLandscape)
 	{
@@ -1504,7 +1589,7 @@ void UPCGActorAndComponentMapping::OnLandscapeChanged(ALandscapeProxy* InLandsca
 	UpdateActorDependencies(InLandscape);
 }
 
-void UPCGActorAndComponentMapping::UpdateActorDependencies(AActor* InActor)
+void FPCGActorAndComponentMapping::UpdateActorDependencies(AActor* InActor)
 {
 	if (!InActor)
 	{
@@ -1527,7 +1612,7 @@ void UPCGActorAndComponentMapping::UpdateActorDependencies(AActor* InActor)
 	}
 }
 
-void UPCGActorAndComponentMapping::OnPCGGraphGeneratedOrCleaned(UPCGComponent* InComponent)
+void FPCGActorAndComponentMapping::OnPCGGraphGeneratedOrCleaned(UPCGComponent* InComponent)
 {
 	if (!InComponent || !InComponent->GetOwner())
 	{
@@ -1537,7 +1622,7 @@ void UPCGActorAndComponentMapping::OnPCGGraphGeneratedOrCleaned(UPCGComponent* I
 	OnActorChanged(InComponent->GetOwner(), /*bInHasMoved=*/false, InComponent);
 }
 
-bool UPCGActorAndComponentMapping::IsActorTracked(const AActor* InActor) const
+bool FPCGActorAndComponentMapping::IsActorTracked(const AActor* InActor) const
 {
 	return InActor && (CulledTrackedActorsToComponentsMap.Contains(InActor) || AlwaysTrackedActorsToComponentsMap.Contains(InActor));
 }
diff --git a/Engine/Plugins/Experimental/PCG/Source/PCG/Private/PCGComponent.cpp b/Engine/Plugins/Experimental/PCG/Source/PCG/Private/PCGComponent.cpp
--- a/Engine/Plugins/Experimental/PCG/Source/PCG/Private/PCGComponent.cpp
+++ b/Engine/Plugins/Experimental/PCG/Source/PCG/Private/PCGComponent.cpp
@@ -70,18 +70,6 @@ UPCGComponent::UPCGComponent(const FObjectInitializer& InObjectInitializer)
 #endif // WITH_EDITOR
 }
 
-UPCGComponent::~UPCGComponent()
-{
-#if WITH_EDITOR
-	// For the special case where a component is part of a reconstruction script (from a BP),
-	// but gets destroyed immediately, we need to force the unregistering. 
-	if (UPCGSubsystem* PCGSubsystem = UPCGSubsystem::GetSubsystemForCurrentWorld())
-	{
-		PCGSubsystem->UnregisterPCGComponent(this, /*bForce=*/true);
-	}
-#endif // WITH_EDITOR
-}
-
 bool UPCGComponent::CanPartition() const
 {
 	// Support/Force partitioning on non-PCG partition actors in WP worlds.
@@ -1055,7 +1043,6 @@ void UPCGComponent::OnUnregister()
 		if (!PCGHelpers::IsRuntimeOrPIE())
 		{
 			Subsystem->CancelGeneration(this);
-			Subsystem->UnregisterPCGComponent(this);
 		}
 	}
 #endif // WITH_EDITOR
@@ -1083,6 +1070,17 @@ void UPCGComponent::OnComponentDestroyed(bool bDestroyingHierarchy)
 		PAOwner->RemoveLocalComponent(this);
 	}
 
+#if WITH_EDITOR
+	// Don't do the unregister in OnUnregister, because we have flows where the component gets Unregistered/Registered without getting destroyed.
+	if (UPCGSubsystem* Subsystem = GetSubsystem())
+	{
+		if (!PCGHelpers::IsRuntimeOrPIE())
+		{
+			Subsystem->UnregisterPCGComponent(this);
+		}
+	}
+#endif // WITH_EDITOR
+
 	Super::OnComponentDestroyed(bDestroyingHierarchy);
 }
 
@@ -1166,7 +1164,14 @@ void UPCGComponent::BeginDestroy()
 	{
 		GraphInstance->OnGraphChangedDelegate.RemoveAll(this);
 	}
-#endif
+
+	// For the special case where a component is part of a reconstruction script (from a BP),
+	// but gets destroyed immediately, we need to force the unregistering. 
+	if (UPCGSubsystem* PCGSubsystem = UPCGSubsystem::GetSubsystemForCurrentWorld())
+	{
+		PCGSubsystem->UnregisterPCGComponent(this, /*bForce=*/true);
+	}
+#endif // WITH_EDITOR
 
 	Super::BeginDestroy();
 }
@@ -1244,10 +1249,14 @@ void UPCGComponent::RefreshAfterGraphChanged(UPCGGraphInterface* InGraph, bool b
 	// In editor, since we've changed the graph, we might have changed the tracked actor tags as well
 	if (!PCGHelpers::IsRuntimeOrPIE())
 	{
-		UpdateTrackingCache();
 		if (UPCGSubsystem* Subsystem = GetSubsystem())
 		{
-			Subsystem->UpdateComponentTracking(this, /*bInShouldDirtyActors=*/ true);
+			// Don't update the tracking if nothing changed for the tracking.
+			TArray<FPCGActorSelectionKey> ChangedKeys;
+			if (UpdateTrackingCache(&ChangedKeys))
+			{
+				Subsystem->UpdateComponentTracking(this, /*bInShouldDirtyActors=*/ true, &ChangedKeys);
+			}
 		}
 
 		DirtyGenerated(bDirtyInputs ? (EPCGComponentDirtyFlag::Actor | EPCGComponentDirtyFlag::Landscape) : EPCGComponentDirtyFlag::None);
@@ -1320,7 +1329,11 @@ void UPCGComponent::PostEditChangeProperty(FPropertyChangedEvent& PropertyChange
 	{
 		if (UPCGSubsystem* Subsystem = GetSubsystem())
 		{
-			Subsystem->UpdateComponentTracking(this);
+			TArray<FPCGActorSelectionKey> ChangedKeys;
+			if (UpdateTrackingCache(&ChangedKeys))
+			{
+				Subsystem->UpdateComponentTracking(this, /*bInShouldDirtyActors=*/true, &ChangedKeys);
+			}
 		}
 
 		DirtyGenerated(EPCGComponentDirtyFlag::Input);
@@ -1382,21 +1395,30 @@ void UPCGComponent::PostEditUndo()
 	Super::PostEditUndo();
 }
 
-void UPCGComponent::UpdateTrackingCache()
+bool UPCGComponent::UpdateTrackingCache(TArray<FPCGActorSelectionKey>* OptionalChangedKeys)
 {
 	// Without an owner, it probably means we are in a BP template, so no need to setup callbacks
 	if (!GetOwner())
 	{
-		return;
+		return false;
 	}
 
-	CachedTrackedKeysToSettings.Reset();
-	CachedTrackedKeysToCulling.Reset();
+	// Store in a temporary map to detect key changes.
+	TMap<FPCGActorSelectionKey, bool> NewTrackedKeysToCulling;
+
+	int32 FoundKeys = 0;
 
 	if (UPCGGraph* PCGGraph = GetGraph())
 	{
 		CachedTrackedKeysToSettings = PCGGraph->GetTrackedActorKeysToSettings();
 
+		// Also add a key for the landscape, with settings null and always culled, if we should track the landscape
+		if (ShouldTrackLandscape())
+		{
+			FPCGActorSelectionKey LandscapeKey = FPCGActorSelectionKey(ALandscapeProxy::StaticClass());
+			CachedTrackedKeysToSettings.FindOrAdd(LandscapeKey).Emplace(/*Settings*/nullptr, /*bIsCulled*/true);
+		}
+
 		// A tag should be culled, if only all the settings that track this tag should cull.
 		// Note that is only impact the fact that we track (or not) this tag.
 		// If a setting is marked as "should cull", it will only be dirtied (at least by default), if the actor with the
@@ -1415,9 +1437,38 @@ void UPCGComponent::UpdateTrackingCache()
 				}
 			}
 
-			CachedTrackedKeysToCulling.Emplace(Key, bShouldCull);
+			NewTrackedKeysToCulling.Emplace(Key, bShouldCull);
+
+			bool* OldCulling = CachedTrackedKeysToCulling.Find(Key);
+			if (OldCulling && (*OldCulling == bShouldCull))
+			{
+				++FoundKeys;
+			}
+
+			// Remove the key from the previous cached map. If nothing is removed, it means it is a new key
+			if ((CachedTrackedKeysToCulling.Remove(Key) == 0) && OptionalChangedKeys)
+			{
+				OptionalChangedKeys->Add(Key);
+			}
+		}
+
+		// At the end, we also have keys that were tracked but no more, so add them at the list of tracked keys
+		if (OptionalChangedKeys)
+		{
+			OptionalChangedKeys->Reserve(OptionalChangedKeys->Num() + CachedTrackedKeysToCulling.Num());
+
+			for (const TPair<FPCGActorSelectionKey, bool>& It : CachedTrackedKeysToCulling)
+			{
+				OptionalChangedKeys->Add(It.Key);
+			}
 		}
 	}
+
+	bool bHasChanged = CachedTrackedKeysToSettings.Num() != FoundKeys;
+
+	CachedTrackedKeysToCulling = MoveTemp(NewTrackedKeysToCulling);
+
+	return bHasChanged;
 }
 
 void UPCGComponent::DirtyGenerated(EPCGComponentDirtyFlag DirtyFlag, const bool bDispatchToLocalComponents)
@@ -2269,7 +2320,7 @@ bool UPCGComponent::DirtyTrackedActor(AActor* InActor, bool bIntersect, const TS
 				}
 
 				const TWeakObjectPtr<const UPCGSettings>& Settings = SettingsAndCulling.Key;
-				if (ensure(Settings.IsValid()))
+				if (Settings.IsValid())
 				{
 					GetSubsystem()->CleanFromCache(Settings->GetElement().Get(), Settings.Get());
 				}
diff --git a/Engine/Plugins/Experimental/PCG/Source/PCG/Public/PCGActorAndComponentMapping.h b/Engine/Plugins/Experimental/PCG/Source/PCG/Public/PCGActorAndComponentMapping.h
--- a/Engine/Plugins/Experimental/PCG/Source/PCG/Public/PCGActorAndComponentMapping.h
+++ b/Engine/Plugins/Experimental/PCG/Source/PCG/Public/PCGActorAndComponentMapping.h
@@ -32,12 +32,12 @@ class UPCGSubsystem;
 * - Tracking non partitioned PCG Components that has tracking actors needs.
 * - Tracking actors: Be able to react to actors changes, and potentially dirty/refresh affected components
 */
-class UPCGActorAndComponentMapping
+class FPCGActorAndComponentMapping
 {
 public:
 	friend UPCGSubsystem;
 
-	~UPCGActorAndComponentMapping() = default;
+	~FPCGActorAndComponentMapping() = default;
 	
 	/** Should be called by the subsystem to handle delayed operations. */
 	void Tick();
@@ -51,8 +51,8 @@ public:
 
 	void AddDelayedActors();
 
-	/** Will register/update tracking if a component was registered/updated. */
-	void RegisterOrUpdateTracking(UPCGComponent* InComponent, bool bInShouldDirtyActors);
+	void RegisterTracking(UPCGComponent* InComponent);
+	void UpdateTracking(UPCGComponent* InComponent, bool bInShouldDirtyActors, const TArray<FPCGActorSelectionKey>* OptionalChangedKeys = nullptr);
 #endif // WITH_EDITOR
 
 	/** Register a new PCG Component or update it. Returns true if it was added/updated. Thread safe */
@@ -86,8 +86,8 @@ private:
 	// This class is only meant to be used as part of the PCG Subsytem and owned by it.
 	// So we put constructors private.
 	// We also need this class to be default constructible, since the PCGSubsytem needs to be default constructible.
-	UPCGActorAndComponentMapping() = default;
-	explicit UPCGActorAndComponentMapping(UPCGSubsystem* PCGSubsystem);
+	FPCGActorAndComponentMapping() = default;
+	explicit FPCGActorAndComponentMapping(UPCGSubsystem* PCGSubsystem);
 
 	bool RegisterOrUpdatePartitionedPCGComponent(UPCGComponent* InComponent, bool bDoActorMapping = true);
 	bool RegisterOrUpdateNonPartitionedPCGComponent(UPCGComponent* InComponent);
@@ -113,6 +113,9 @@ private:
 	/** Delete the current mapping between a PCG component and its PCG Partition actors */
 	void DeleteMappingPCGComponentPartitionActor(UPCGComponent* InComponent);
 
+	/** Returns true if a component is registered */
+	bool IsComponentRegistered(const UPCGComponent* InComponent) const;
+
 #if WITH_EDITOR
 	/* Return true if something is still tracked or was just untracked. */
 	bool AddOrUpdateTrackedActor(AActor* InActor);
@@ -136,9 +139,12 @@ private:
 	/** Remap the tracking in case of BP components. */
 	void RemapTracking(const UPCGComponent* InOldComponent, UPCGComponent* InNewComponent);
 
-	/** Unregister tracking when a component is removed. */
+	/** Unregister tracking when a component is removed */
 	void UnregisterTracking(UPCGComponent* InComponent);
 
+	/** Unregister tracking when a component is removed or keys are unregistered */
+	void UnregisterTracking(UPCGComponent* InComponent, const TSet<TObjectKey<AActor>>* OptionalActorsToUntrack, const TSet<FPCGActorSelectionKey>* OptionalKeysToUntrack);
+
 	/** Trigger an update when the actor changed. 
 	* Can specify if the actor has moved to also update components that were at its previous position.
 	* Can also specify an optional object, originating the change, to avoid re-dirtying a component if it was the origin.
diff --git a/Engine/Plugins/Experimental/PCG/Source/PCG/Public/PCGComponent.h b/Engine/Plugins/Experimental/PCG/Source/PCG/Public/PCGComponent.h
--- a/Engine/Plugins/Experimental/PCG/Source/PCG/Public/PCGComponent.h
+++ b/Engine/Plugins/Experimental/PCG/Source/PCG/Public/PCGComponent.h
@@ -17,7 +17,7 @@ namespace EEndPlayReason { enum Type : int; }
 
 class APCGPartitionActor;
 struct FPCGContext;
-class UPCGActorAndComponentMapping;
+class FPCGActorAndComponentMapping;
 class UPCGComponent;
 class UPCGGraph;
 class UPCGGraphInterface;
@@ -67,13 +67,12 @@ UCLASS(BlueprintType, ClassGroup = (Procedural), meta = (BlueprintSpawnableCompo
 class PCG_API UPCGComponent : public UActorComponent
 {
 	UPCGComponent(const FObjectInitializer& InObjectInitializer);
-	~UPCGComponent();
 
 	GENERATED_BODY()
 
 	friend class UPCGManagedActors;
 	friend class UPCGSubsystem;
-	friend class UPCGActorAndComponentMapping;
+	friend class FPCGActorAndComponentMapping;
 
 public:
 	/** ~Begin UObject interface */
@@ -349,7 +348,8 @@ private:
 	/** Sets up actor, tracking, landscape and graph callbacks */
 	void SetupCallbacksOnCreation();
 
-	void UpdateTrackingCache();
+	/** Returns true if something changed in the tracking. */
+	bool UpdateTrackingCache(TArray<FPCGActorSelectionKey>* OptionalChangedKeys = nullptr);
 
 	// InOriginatingChangeObject can be null
 	bool DirtyTrackedActor(AActor* InActor, bool bIntersect, const TSet<FName>& InRemovedTags, const UObject* InOriginatingChangeObject);
diff --git a/Engine/Plugins/Experimental/PCG/Source/PCG/Public/PCGSubsystem.h b/Engine/Plugins/Experimental/PCG/Source/PCG/Public/PCGSubsystem.h
--- a/Engine/Plugins/Experimental/PCG/Source/PCG/Public/PCGSubsystem.h
+++ b/Engine/Plugins/Experimental/PCG/Source/PCG/Public/PCGSubsystem.h
@@ -41,7 +41,7 @@ class PCG_API UPCGSubsystem : public UTickableWorldSubsystem
 	GENERATED_BODY()
 
 public:
-	friend UPCGActorAndComponentMapping;
+	friend FPCGActorAndComponentMapping;
 
 	UPCGSubsystem();
 
@@ -179,7 +179,7 @@ public:
 	void NotifyGraphChanged(UPCGGraph* InGraph);
 
 	/** Update the tracking on a given component. */
-	void UpdateComponentTracking(UPCGComponent* InComponent, bool bInShouldDirtyActors = false) { ActorAndComponentMapping.RegisterOrUpdateTracking(InComponent, bInShouldDirtyActors); }
+	void UpdateComponentTracking(UPCGComponent* InComponent, bool bShouldDirtyActors, const TArray<FPCGActorSelectionKey>* OptionalChangedKeys = nullptr) { ActorAndComponentMapping.UpdateTracking(InComponent, bShouldDirtyActors, OptionalChangedKeys); }
 
 	/** Cleans up the graph cache on an element basis. InSettings is used for debugging and is optional. */
 	void CleanFromCache(const IPCGElement* InElement, const UPCGSettings* InSettings = nullptr);
@@ -224,7 +224,7 @@ private:
 	APCGWorldActor* PCGWorldActor = nullptr;
 	FPCGGraphExecutor* GraphExecutor = nullptr;
 	bool bHasTickedOnce = false;
-	UPCGActorAndComponentMapping ActorAndComponentMapping;
+	FPCGActorAndComponentMapping ActorAndComponentMapping;
 
 #if WITH_EDITOR
 	FCriticalSection PCGWorldActorLock;
