Better support for embedded level instances inside of other level instances wrt PCG.

diff --git a/Engine/Plugins/Experimental/PCG/Source/PCG/Private/PCGActorAndComponentMapping.cpp b/Engine/Plugins/Experimental/PCG/Source/PCG/Private/PCGActorAndComponentMapping.cpp
--- a/Engine/Plugins/Experimental/PCG/Source/PCG/Private/PCGActorAndComponentMapping.cpp
+++ b/Engine/Plugins/Experimental/PCG/Source/PCG/Private/PCGActorAndComponentMapping.cpp
@@ -778,7 +778,7 @@ void FPCGActorAndComponentMapping::UpdateTracking(UPCGComponent* InComponent, bo
 		{
 			if (!Actor->HasActorRegisteredAllComponents() && !bDisableDelayedActorRegistering)
 			{
-				DelayedAddedActors.Emplace({ Actor, false });
+				DelayedAddedActors.Emplace({ Actor, false, 0 });
 				continue;
 			}
 
@@ -960,10 +960,10 @@ void FPCGActorAndComponentMapping::AddDelayedActors()
 		return;
 	}
 
-	TSet<TTuple<TObjectKey<AActor>, bool>> StillDelayedActors;
+	TSet<TTuple<TObjectKey<AActor>, bool, int>> StillDelayedActors;
 	const bool bDisableDelayedActorRegistering = PCGActorAndComponentMapping::CVarDisableDelayedActorRegistering.GetValueOnAnyThread();
 
-	for (TTuple<TObjectKey<AActor>, bool>& ActorPtrAndShouldDirty : DelayedAddedActors)
+	for (TTuple<TObjectKey<AActor>, bool, int>& ActorPtrAndShouldDirty : DelayedAddedActors)
 	{
 		AActor* Actor = ActorPtrAndShouldDirty.Get<0>().ResolveObjectPtr();
 		if (!Actor)
@@ -978,7 +978,7 @@ void FPCGActorAndComponentMapping::AddDelayedActors()
 		else
 		{
 			// Implementation note: since the delayed actors list is built from top-level actors (e.g. directly in the level) then depth here is 0.
-			OnActorAdded_Internal(Actor, ActorPtrAndShouldDirty.Get<1>(), 0);
+			OnActorAdded_Internal(Actor, ActorPtrAndShouldDirty.Get<1>(), ActorPtrAndShouldDirty.Get<2>());
 		}
 	}
 
@@ -987,31 +987,45 @@ void FPCGActorAndComponentMapping::AddDelayedActors()
 
 void FPCGActorAndComponentMapping::OnActorAdded(AActor* InActor)
 {
-	// Implementation note: since this is called only for actors directly in the current level, the depth here is 0.
-	OnActorAdded_Internal(InActor, true, 0);
-}
-
-void FPCGActorAndComponentMapping::OnActorAdded_Internal(AActor* InActor, bool bShouldDirty, int32 LevelInstanceDepth)
-{
-	TRACE_CPUPROFILER_EVENT_SCOPE(FPCGActorAndComponentMapping::OnActorAdded);
-
 	// We have to make sure to not create a infinite loop
 	if (!InActor || InActor->IsA<APCGWorldActor>() || !PCGSubsystem || InActor->GetWorld() != PCGSubsystem->GetWorld())
 	{
 		return;
 	}
 
+	int LevelInstanceDepth = 0;
+
 #if WITH_EDITOR
 	if (InActor->bIsEditorPreviewActor)
 	{
 		return;
 	}
-#endif
+
+	if (ALevelInstanceEditorInstanceActor* LevelInstanceEditorInstance = Cast<ALevelInstanceEditorInstanceActor>(InActor))
+	{
+		const ULevelInstanceSubsystem* LevelInstanceSubsystem = PCGSubsystem->GetWorld() ? PCGSubsystem->GetWorld()->GetSubsystem<ULevelInstanceSubsystem>() : nullptr;
+		const ILevelInstanceInterface* ActorLevelInstance = LevelInstanceSubsystem ? LevelInstanceSubsystem->GetParentLevelInstance(InActor) : nullptr;
+		while (ActorLevelInstance)
+		{
+			++LevelInstanceDepth;
+			ActorLevelInstance = LevelInstanceSubsystem->GetParentLevelInstance(CastChecked<AActor>(ActorLevelInstance));
+		}
+	}
+#endif // WITH_EDITOR
+
+	// Implementation note: since this is called only for actors directly in the current level, the depth here is 0.
+	OnActorAdded_Internal(InActor, true, LevelInstanceDepth);
+}
+
+void FPCGActorAndComponentMapping::OnActorAdded_Internal(AActor* InActor, bool bShouldDirty, int32 LevelInstanceDepth)
+{
+	TRACE_CPUPROFILER_EVENT_SCOPE(FPCGActorAndComponentMapping::OnActorAdded);
+	check(InActor && !InActor->IsA<APCGWorldActor>() && PCGSubsystem && InActor->GetWorld() == PCGSubsystem->GetWorld());
 
 	// If the subsystem is not initialized, wait for it to be, and store all the actors to check
 	if (!PCGSubsystem->IsInitialized())
 	{
-		DelayedAddedActors.Emplace({ InActor, bShouldDirty });
+		DelayedAddedActors.Emplace({ InActor, bShouldDirty, LevelInstanceDepth });
 		return;
 	}
 
@@ -1022,7 +1036,7 @@ void FPCGActorAndComponentMapping::OnActorAdded_Internal(AActor* InActor, bool b
 	if (ALevelInstanceEditorInstanceActor* LevelInstanceEditorInstance = Cast<ALevelInstanceEditorInstanceActor>(InActor))
 	{
 		const FLevelInstanceID& LevelInstanceId = LevelInstanceEditorInstance->GetLevelInstanceID();
-		const ULevelInstanceSubsystem* LevelInstanceSubsystem = PCGSubsystem->GetWorld()->GetSubsystem<ULevelInstanceSubsystem>();
+		const ULevelInstanceSubsystem* LevelInstanceSubsystem = PCGSubsystem->GetWorld() ? PCGSubsystem->GetWorld()->GetSubsystem<ULevelInstanceSubsystem>() : nullptr;
 
 		if (LevelInstanceId.IsValid() && LevelInstanceSubsystem)
 		{
@@ -1030,7 +1044,7 @@ void FPCGActorAndComponentMapping::OnActorAdded_Internal(AActor* InActor, bool b
 			{
 				LevelInstanceSubsystem->ForEachActorInLevelInstance(LevelInstance, [this, bShouldDirty, LevelInstanceDepth, LevelInstanceEditorInstance](AActor* LevelActor)
 				{
-					if (LevelActor != LevelInstanceEditorInstance)
+					if (LevelActor && LevelActor != LevelInstanceEditorInstance && !LevelActor->IsA<APCGWorldActor>())
 					{
 						OnActorAdded_Internal(LevelActor, bShouldDirty, LevelInstanceDepth + 1);
 					}
@@ -1041,11 +1055,11 @@ void FPCGActorAndComponentMapping::OnActorAdded_Internal(AActor* InActor, bool b
 		}
 		else
 		{
-			DelayedAddedActors.Emplace({ InActor, bShouldDirty });
+			DelayedAddedActors.Emplace({ InActor, bShouldDirty, LevelInstanceDepth });
 			return;
 		}
 	}
-#endif
+#endif // WITH_EDITOR
 
 	if (AddOrUpdateTrackedActor(InActor) && bShouldDirty)
 	{
@@ -1182,30 +1196,46 @@ bool FPCGActorAndComponentMapping::UnregisterActor(AActor* InActor)
 
 void FPCGActorAndComponentMapping::OnActorDeleted(AActor* InActor)
 {
-	// Implementation note: since this is called only for actors directly in the current level, the depth here is 0.
-	OnActorDeleted_Internal(InActor, 0);
-}
-
-void FPCGActorAndComponentMapping::OnActorDeleted_Internal(AActor* InActor, int32 LevelInstanceDepth)
-{
-	TRACE_CPUPROFILER_EVENT_SCOPE(FPCGActorAndComponentMapping::OnActorDeleted);
-
 	if (!InActor || !PCGSubsystem || InActor->GetWorld() != PCGSubsystem->GetWorld())
 	{
 		return;
 	}
 
+	int LevelInstanceDepth = 0;
+
 #if WITH_EDITOR
 	if (InActor->bIsEditorPreviewActor)
 	{
 		return;
 	}
 
+	if (ALevelInstanceEditorInstanceActor* LevelInstanceEditorInstance = Cast<ALevelInstanceEditorInstanceActor>(InActor))
+	{
+		const ULevelInstanceSubsystem* LevelInstanceSubsystem = PCGSubsystem->GetWorld() ? PCGSubsystem->GetWorld()->GetSubsystem<ULevelInstanceSubsystem>() : nullptr;
+		const ILevelInstanceInterface* ActorLevelInstance = LevelInstanceSubsystem ? LevelInstanceSubsystem->GetParentLevelInstance(InActor) : nullptr;
+		while (ActorLevelInstance)
+		{
+			++LevelInstanceDepth;
+			ActorLevelInstance = LevelInstanceSubsystem->GetParentLevelInstance(CastChecked<AActor>(ActorLevelInstance));
+		}
+	}
+#endif
+
+	// Implementation note: since this is called only for actors directly in the current level, the depth here is 0.
+	OnActorDeleted_Internal(InActor, LevelInstanceDepth);
+}
+
+void FPCGActorAndComponentMapping::OnActorDeleted_Internal(AActor* InActor, int32 LevelInstanceDepth)
+{
+	TRACE_CPUPROFILER_EVENT_SCOPE(UPCGActorAndComponentMapping::OnActorDeleted);
+	check(InActor && PCGSubsystem && InActor->GetWorld() == PCGSubsystem->GetWorld());
+
+#if WITH_EDITOR
 	// When editing/deleting/replacing a level instance, this will allow us to first remove the previously tracked actors that were under the level instance
 	if (ALevelInstanceEditorInstanceActor* LevelInstanceEditorInstance = Cast<ALevelInstanceEditorInstanceActor>(InActor))
 	{
 		const FLevelInstanceID& LevelInstanceId = LevelInstanceEditorInstance->GetLevelInstanceID();
-		const ULevelInstanceSubsystem* LevelInstanceSubsystem = PCGSubsystem->GetWorld()->GetSubsystem<ULevelInstanceSubsystem>();
+		const ULevelInstanceSubsystem* LevelInstanceSubsystem = PCGSubsystem->GetWorld() ? PCGSubsystem->GetWorld()->GetSubsystem<ULevelInstanceSubsystem>() : nullptr;
 
 		if (LevelInstanceId.IsValid() && LevelInstanceSubsystem)
 		{
@@ -1229,7 +1259,7 @@ void FPCGActorAndComponentMapping::OnActorDeleted_Internal(AActor* InActor, int3
 		OnActorDeleted_Internal(LevelActor, LevelInstanceDepth + 1);
 		return true;
 	});
-#endif
+#endif // WITH_EDITOR
 
 	if (!IsActorTracked(InActor))
 	{
diff --git a/Engine/Plugins/Experimental/PCG/Source/PCG/Public/PCGActorAndComponentMapping.h b/Engine/Plugins/Experimental/PCG/Source/PCG/Public/PCGActorAndComponentMapping.h
--- a/Engine/Plugins/Experimental/PCG/Source/PCG/Public/PCGActorAndComponentMapping.h
+++ b/Engine/Plugins/Experimental/PCG/Source/PCG/Public/PCGActorAndComponentMapping.h
@@ -195,8 +195,8 @@ private:
 
 	mutable FRWLock TrackedComponentsLock;
 
-	// Keep track of actors that aren't yet ready (or if the subsystem is not yet ready) and add them in next tick.
-	TSet<TTuple<TObjectKey<AActor>, bool>> DelayedAddedActors;
+	// Keep track of actors that aren't yet ready (or if the subsystem is not yet ready), whether we should dirty them and their instance level depth so we can add them in next tick.
+	TSet<TTuple<TObjectKey<AActor>, bool, int>> DelayedAddedActors;
 
 	/** Keep a mapping between tracked actors and their dependencies. */
 	TMap<TObjectKey<AActor>, TSet<TObjectPtr<UObject>>> TrackedActorsToDependenciesMap;
