Fixed runtime merging of landscape caches

diff --git a/Engine/Plugins/Experimental/PCG/Source/PCG/Private/Grid/PCGLandscapeCache.cpp b/Engine/Plugins/Experimental/PCG/Source/PCG/Private/Grid/PCGLandscapeCache.cpp
--- a/Engine/Plugins/Experimental/PCG/Source/PCG/Private/Grid/PCGLandscapeCache.cpp
+++ b/Engine/Plugins/Experimental/PCG/Source/PCG/Private/Grid/PCGLandscapeCache.cpp
@@ -717,6 +717,11 @@ void UPCGLandscapeCache::TakeOwnership(UPCGLandscapeCache* InLandscapeCache)
 	CacheEntryCount = CachedData.Num();
 #endif
 
+	if (bShouldDirty)
+	{
+		EmbeddedCaches.AddUnique(InLandscapeCache);
+	}
+
 	if (bShouldDirty && SerializationMode != EPCGLandscapeCacheSerializationMode::NeverSerialize)
 	{
 		MarkPackageDirty();
diff --git a/Engine/Plugins/Experimental/PCG/Source/PCG/Public/Grid/PCGLandscapeCache.h b/Engine/Plugins/Experimental/PCG/Source/PCG/Public/Grid/PCGLandscapeCache.h
--- a/Engine/Plugins/Experimental/PCG/Source/PCG/Public/Grid/PCGLandscapeCache.h
+++ b/Engine/Plugins/Experimental/PCG/Source/PCG/Public/Grid/PCGLandscapeCache.h
@@ -186,6 +186,10 @@ private:
 	TSet<TWeakObjectPtr<ALandscapeProxy>> Landscapes;
 #endif
 
+	// Keep "merged" landscape caches alive so that the bulk data is reachable
+	UPROPERTY(Transient)
+	TArray<TObjectPtr<UPCGLandscapeCache>> EmbeddedCaches;
+
 #if WITH_EDITOR
 	FRWLock CacheLock;
 #endif
