[PCG] Fix up un-placed title of the Filter By Type node which currently appears in node palette as "Filter - Any".

#rnx

diff --git a/Engine/Plugins/Experimental/PCG/Source/PCG/Private/Elements/PCGFilterByType.cpp b/Engine/Plugins/Experimental/PCG/Source/PCG/Private/Elements/PCGFilterByType.cpp
--- a/Engine/Plugins/Experimental/PCG/Source/PCG/Private/Elements/PCGFilterByType.cpp
+++ b/Engine/Plugins/Experimental/PCG/Source/PCG/Private/Elements/PCGFilterByType.cpp
@@ -12,6 +12,11 @@
 #if WITH_EDITOR
 FText UPCGFilterByTypeSettings::GetDefaultNodeTitle() const
 {
+	if (this == GetClass()->GetDefaultObject())
+	{
+		return LOCTEXT("DefaultNodeTitle", "Filter By Type");
+	}
+
 	const FText TypeText = FText::FromString(StaticEnum<EPCGDataType>()->GetDisplayNameTextByValue(static_cast<int64>(TargetType)).ToString());
 	return FText::Format(LOCTEXT("NodeTitleFormat", "Filter - {0}"), TypeText);
 }
