[PCG] Improve edge validation code to report graph name in errors

#rnx

diff --git a/Engine/Plugins/Experimental/PCG/Source/PCG/Private/PCGGraph.cpp b/Engine/Plugins/Experimental/PCG/Source/PCG/Private/PCGGraph.cpp
--- a/Engine/Plugins/Experimental/PCG/Source/PCG/Private/PCGGraph.cpp
+++ b/Engine/Plugins/Experimental/PCG/Source/PCG/Private/PCGGraph.cpp
@@ -1069,22 +1069,23 @@ void UPCGGraph::FixInvalidEdges()
 			for (int32 i = Pin->Edges.Num() - 1; i >= 0; --i)
 			{
 				UPCGPin* OtherPin = Pin->Edges[i] ? (bPinsAreInputs ? Pin->Edges[i]->InputPin : Pin->Edges[i]->OutputPin) : nullptr;
+				UPCGNode* ConnectedNode = OtherPin ? OtherPin->Node : nullptr;
 
 				// Remove trivially invalid edges.
-				if (!OtherPin || !OtherPin->Node)
+				if (!ensure(OtherPin && OtherPin->Node))
 				{
-					UE_LOG(LogPCG, Error, TEXT("Removed edge to a missing pin or pin that has no node."));
-					ensure(false);
+					UE_LOG(LogPCG, Error, TEXT("Removed edge to a missing pin or pin that has no node, from graph '%s'."), *GetFName().ToString());
+
 					Pin->Edges.RemoveAt(i);
-					continue;
 				}
-
-				// Remove edges to nodes that are not present in the graph.
-				UPCGNode* ConnectedNode = OtherPin->Node;
-				if (!ConnectedNode || (GetInputNode() != ConnectedNode && GetOutputNode() != ConnectedNode && !Nodes.Contains(ConnectedNode)))
+				else if (!ConnectedNode || (GetInputNode() != ConnectedNode && GetOutputNode() != ConnectedNode && !Nodes.Contains(ConnectedNode)))
 				{
-					UE_LOG(LogPCG, Error, TEXT("Removed edge to a node '%s' that is not registered in the graph."), ConnectedNode ? *ConnectedNode->GetFName().ToString() : TEXT("NULL"));
+					// Remove edges to nodes that are not present in the graph.
+					UE_LOG(LogPCG, Error, TEXT("Removed edge to a node '%s' that is not registered in the graph."),
+						ConnectedNode ? *ConnectedNode->GetFName().ToString() : TEXT("NULL"),
+						*GetFName().ToString());
 					ensure(false);
+
 					Pin->Edges.RemoveAt(i);
 				}
 			}
