[PCG] New TimeSlice context to make timesliced node authoring simpler and consistent

* Updated the Surface Sampler to the new context

diff --git a/Engine/Plugins/Experimental/PCG/Source/PCG/Private/Data/PCGWorldData.cpp b/Engine/Plugins/Experimental/PCG/Source/PCG/Private/Data/PCGWorldData.cpp
--- a/Engine/Plugins/Experimental/PCG/Source/PCG/Private/Data/PCGWorldData.cpp
+++ b/Engine/Plugins/Experimental/PCG/Source/PCG/Private/Data/PCGWorldData.cpp
@@ -366,7 +366,6 @@ const UPCGPointData* UPCGWorldRayHitData::CreatePointData(FPCGContext* Context,
 {
 	UPCGPointData* Data = NewObject<UPCGPointData>();
 	Data->InitializeFromData(this);
-	TArray<FPCGPoint>& Points = Data->GetMutablePoints();
 
 	FBox EffectiveBounds = Bounds;
 	if (InBounds.IsValid)
@@ -392,10 +391,10 @@ const UPCGPointData* UPCGWorldRayHitData::CreatePointData(FPCGContext* Context,
 		return Data;
 	}
 
-	PCGSurfaceSampler::FSurfaceSamplerSettings SamplerSettings;
+	PCGSurfaceSampler::FSurfaceSamplerExecutionSettings SamplerSettings;
 	if (SamplerSettings.Initialize(nullptr, Context, EffectiveBounds))
 	{
-		PCGSurfaceSampler::SampleSurface(Context, this, nullptr, SamplerSettings, Data);
+		Data = PCGSurfaceSampler::SampleSurface(Context, this, nullptr, SamplerSettings);
 	}
 
 	return Data;
diff --git a/Engine/Plugins/Experimental/PCG/Source/PCG/Private/Elements/PCGSurfaceSampler.cpp b/Engine/Plugins/Experimental/PCG/Source/PCG/Private/Elements/PCGSurfaceSampler.cpp
--- a/Engine/Plugins/Experimental/PCG/Source/PCG/Private/Elements/PCGSurfaceSampler.cpp
+++ b/Engine/Plugins/Experimental/PCG/Source/PCG/Private/Elements/PCGSurfaceSampler.cpp
@@ -29,15 +29,13 @@ namespace PCGSurfaceSamplerConstants
 
 namespace PCGSurfaceSampler
 {
-	bool FSurfaceSamplerSettings::Initialize(const UPCGSurfaceSamplerSettings* InSettings, FPCGContext* Context, const FBox& InputBounds)
+	bool FSurfaceSamplerExecutionSettings::Initialize(const UPCGSurfaceSamplerSettings* Settings, const FPCGContext* Context, const FBox& InputBounds)
 	{
 		if (!Context)
 		{
 			return false;
 		}
 
-		Settings = InSettings;
-
 		if (Settings)
 		{
 			// Compute used values
@@ -121,7 +119,7 @@ namespace PCGSurfaceSampler
 		return true;
 	}
 
-	FIntVector2 FSurfaceSamplerSettings::ComputeCellIndices(int32 Index) const
+	FIntVector2 FSurfaceSamplerExecutionSettings::ComputeCellIndices(int32 Index) const
 	{
 		check(Index >= 0 && Index < CellCount);
 		const int32 CellCountX = 1 + CellMaxX - CellMinX;
@@ -129,26 +127,20 @@ namespace PCGSurfaceSampler
 		return FIntVector2(CellMinX + (Index % CellCountX), CellMinY + (Index / CellCountX));
 	}
 
-	UPCGPointData* SampleSurface(FPCGContext* Context, const UPCGSpatialData* InSurface, const UPCGSpatialData* InBoundingShape, const FSurfaceSamplerSettings& LoopData)
+	UPCGPointData* SampleSurface(FPCGContext* Context, const UPCGSpatialData* InSurface, const UPCGSpatialData* InBoundingShape, const FSurfaceSamplerExecutionSettings& ExecutionSettings)
 	{
-		// We don't support time slicing here
-		if (LoopData.bEnableTimeSlicing)
-		{
-			PCGE_LOG_C(Error, LogOnly, Context, LOCTEXT("MissingOutputPointData", "An output point data must be provided to support sampling with time slicing"));
-			return nullptr;
-		}
-
 		UPCGPointData* SampledData = NewObject<UPCGPointData>();
 		SampledData->InitializeFromData(InSurface);
 
-		SampleSurface(Context, InSurface, InBoundingShape, LoopData, SampledData);
+		// We don't support time slicing here
+		SampleSurface(Context, ExecutionSettings, InSurface, InBoundingShape, SampledData, /*bTimeSlicingIsEnabled=*/false);
 
 		return SampledData;
 	}
 
-	bool SampleSurface(FPCGContext* Context, const UPCGSpatialData* InSurface, const UPCGSpatialData* InBoundingShape, const FSurfaceSamplerSettings& LoopData, UPCGPointData* SampledData)
+	bool SampleSurface(FPCGContext* Context, const FSurfaceSamplerExecutionSettings& Settings, const UPCGSpatialData* InSurface, const UPCGSpatialData* InBoundingShape, UPCGPointData* SampledData, const bool bTimeSlicingIsEnabled)
 	{
-		check(InSurface);
+		check(InSurface && SampledData);
 
 		TArray<FPCGPoint>& SampledPoints = SampledData->GetMutablePoints();
 
@@ -160,26 +152,26 @@ namespace PCGSurfaceSampler
 		constexpr FVector::FReal ZMultiplier = 1.0 - UE_DOUBLE_SMALL_NUMBER;
 		// Try to use a multiplier instead of a simply offset to combat loss of precision in floats. However if MaxZ is very small,
 		// then multiplier will not work, so just use an offset.
-		FVector::FReal SampleZ = (FMath::Abs(LoopData.InputBoundsMaxZ) > UE_DOUBLE_SMALL_NUMBER) ? LoopData.InputBoundsMaxZ * ZMultiplier : -UE_DOUBLE_SMALL_NUMBER;
+		FVector::FReal SampleZ = (FMath::Abs(Settings.InputBoundsMaxZ) > UE_DOUBLE_SMALL_NUMBER) ? Settings.InputBoundsMaxZ * ZMultiplier : -UE_DOUBLE_SMALL_NUMBER;
 		// Make sure we're still in bounds though!
-		SampleZ = FMath::Max(SampleZ, LoopData.InputBoundsMinZ);
+		SampleZ = FMath::Max(SampleZ, Settings.InputBoundsMinZ);
 
 		// Cache pointer ahead of time to avoid dereferencing object pointer which does access tracking and supports lazy loading, and can come with substantial
 		// overhead (add trace marker to FObjectPtr::Get to see).
 		UPCGMetadata* OutMetadata = SampledData->Metadata.Get();
 
-		auto AsyncProcessFunc = [&LoopData, SampledData, InBoundingShape, InSurface, &ProjectionParams, SampleZ, OutMetadata](int32 Index, FPCGPoint& OutPoint)
+		auto AsyncProcessFunc = [&Settings, InBoundingShape, InSurface, &ProjectionParams, SampleZ, OutMetadata](int32 Index, FPCGPoint& OutPoint)
 		{
-			const FIntVector2 Indices = LoopData.ComputeCellIndices(Index);
+			const FIntVector2 Indices = Settings.ComputeCellIndices(Index);
 
-			const FVector::FReal CurrentX = Indices.X * LoopData.CellSize.X;
-			const FVector::FReal CurrentY = Indices.Y * LoopData.CellSize.Y;
-			const FVector InnerCellSize = LoopData.InnerCellSize;
+			const FVector::FReal CurrentX = Indices.X * Settings.CellSize.X;
+			const FVector::FReal CurrentY = Indices.Y * Settings.CellSize.Y;
+			const FVector InnerCellSize = Settings.InnerCellSize;
 
-			FRandomStream RandomSource(PCGHelpers::ComputeSeed(LoopData.Seed, Indices.X, Indices.Y));
-			float Chance = RandomSource.FRand();
+			FRandomStream RandomSource(PCGHelpers::ComputeSeed(Settings.Seed, Indices.X, Indices.Y));
+			const float Chance = RandomSource.FRand();
 
-			const float Ratio = LoopData.Ratio;
+			const float Ratio = Settings.Ratio;
 
 			if (Chance >= Ratio)
 			{
@@ -190,7 +182,7 @@ namespace PCGSurfaceSampler
 			const float RandY = RandomSource.FRand();
 
 			const FVector TentativeLocation = FVector(CurrentX + RandX * InnerCellSize.X, CurrentY + RandY * InnerCellSize.Y, SampleZ);
-			const FBox LocalBound(-LoopData.PointExtents, LoopData.PointExtents);
+			const FBox LocalBound(-Settings.PointExtents, Settings.PointExtents);
 
 			// The output at this point is not initialized
 			OutPoint = FPCGPoint();
@@ -202,15 +194,15 @@ namespace PCGSurfaceSampler
 			}
 
 			// Set physical properties that are needed for the bounding shape checks, etc.
-			OutPoint.SetExtents(LoopData.PointExtents);
-			OutPoint.Steepness = LoopData.PointSteepness;
+			OutPoint.SetExtents(Settings.PointExtents);
+			OutPoint.Steepness = Settings.PointSteepness;
 
 			// Now run gauntlet of shape network (if there is one) to accept or reject the point.
 			if (InBoundingShape)
 			{
 				FPCGPoint BoundingShapeSample;
 #if WITH_EDITOR
-				if (!InBoundingShape->SamplePoint(OutPoint.Transform, OutPoint.GetLocalBounds(), BoundingShapeSample, nullptr) && !LoopData.bKeepZeroDensityPoints)
+				if (!InBoundingShape->SamplePoint(OutPoint.Transform, OutPoint.GetLocalBounds(), BoundingShapeSample, nullptr) && !Settings.bKeepZeroDensityPoints)
 #else
 				if (!InBoundingShape->SamplePoint(OutPoint.Transform, OutPoint.GetLocalBounds(), BoundingShapeSample, nullptr))
 #endif
@@ -223,20 +215,15 @@ namespace PCGSurfaceSampler
 			}
 
 			// Apply final parameters on the point
-			OutPoint.Density *= (LoopData.bApplyDensityToPoints ? ((Ratio - Chance) / Ratio) : 1.0f);
+			OutPoint.Density *= (Settings.bApplyDensityToPoints ? ((Ratio - Chance) / Ratio) : 1.0f);
 			OutPoint.Seed = RandomSource.GetCurrentSeed();
 
 			return true;
 		};
 
-		bool bAsyncDone = FPCGAsync::AsyncProcessing<FPCGPoint>(Context ? &Context->AsyncState : nullptr, LoopData.CellCount, SampledPoints, AsyncProcessFunc, LoopData.bEnableTimeSlicing);
-
-		if (Context && bAsyncDone)
-		{
-			PCGE_LOG_C(Verbose, LogOnly, Context, FText::Format(LOCTEXT("GenerationInfo", "Generated {0} points in {1} cells"), SampledPoints.Num(), LoopData.CellCount));
-		}
-
-		return bAsyncDone;
+		FPCGAsyncState* AsyncState = Context ? &Context->AsyncState : nullptr;
+		const bool bEnableTimeSlicing = Context ? bTimeSlicingIsEnabled : false;
+		return FPCGAsync::AsyncProcessing<FPCGPoint>(AsyncState, Settings.CellCount, SampledPoints, AsyncProcessFunc, bEnableTimeSlicing);
 	}
 
 #if WITH_EDITOR
@@ -313,192 +300,185 @@ bool UPCGSurfaceSamplerSettings::IsPinUsedByNodeExecution(const UPCGPin* InPin)
 }
 #endif
 
-FPCGSurfaceSamplerContext::~FPCGSurfaceSamplerContext()
-{
-	if (bUnionDataCreated && BoundingShapeSpatialInput)
-	{
-		// We created the data so can cast away constness
-		UPCGSpatialData* BoundingShapeMutable = const_cast<UPCGSpatialData*>(BoundingShapeSpatialInput);
-		BoundingShapeMutable->RemoveFromRoot();
-		BoundingShapeMutable->MarkAsGarbage();
-
-		BoundingShapeSpatialInput = nullptr;
-	}
-}
-
 FPCGElementPtr UPCGSurfaceSamplerSettings::CreateElement() const
 {
 	return MakeShared<FPCGSurfaceSamplerElement>();
 }
 
-FPCGContext* FPCGSurfaceSamplerElement::Initialize(const FPCGDataCollection& InInputData, TWeakObjectPtr<UPCGComponent> InSourceComponent, const UPCGNode* InNode)
+namespace PCGSurfaceSamplerHelpers
 {
-	FPCGSurfaceSamplerContext* Context = new FPCGSurfaceSamplerContext();
-	Context->InputData = InInputData;
-	Context->SourceComponent = InSourceComponent;
-	Context->Node = InNode;
+	using ContextType = FPCGSurfaceSamplerElement::ContextType;
+	using ExecStateType = FPCGSurfaceSamplerElement::ExecStateType;
 
-	return Context;
-}
-
-bool FPCGSurfaceSamplerElement::AddGeneratingShapesToContext(FPCGSurfaceSamplerContext* InContext) const
-{
-	TRACE_CPUPROFILER_EVENT_SCOPE(FPCGSurfaceSamplerElement::SetupContextForGeneratingShape);
-
-	check(InContext);
-	const UPCGSurfaceSamplerSettings* Settings = InContext->GetInputSettings<UPCGSurfaceSamplerSettings>();
-	check(Settings);
-
-	TArray<FPCGTaggedData>& Outputs = InContext->OutputData.TaggedData;
-
-	// Early out on invalid settings
-	// TODO: we could compute an approximate radius based on the points per squared meters if that's useful
-	const FVector& PointExtents = Settings->PointExtents;
-	if (PointExtents.X <= 0 || PointExtents.Y <= 0)
+	const UPCGSpatialData* FindBoundingShape(const FPCGContext* Context, bool& bOutUnionWasCreated)
 	{
-		PCGE_LOG_C(Warning, GraphAndLog, InContext, LOCTEXT("SkippedInvalidPointExtents", "Skipped - Invalid point extents"));
-		return false;
-	}
-
-	TArray<FPCGTaggedData> SurfaceInputs = InContext->InputData.GetInputsByPin(PCGSurfaceSamplerConstants::SurfaceLabel);
+		const UPCGSpatialData* BoundingShape = Context->InputData.GetSpatialUnionOfInputsByPin(PCGSurfaceSamplerConstants::BoundingShapeLabel, bOutUnionWasCreated);
 
-	// If there are no surfaces to sample, early out
-	if (SurfaceInputs.IsEmpty())
-	{
-		return false;
-	}
-
-	// Construct a list of shapes to generate samples from. Get these directly from the first input pin.
-	for (FPCGTaggedData& TaggedData : SurfaceInputs)
-	{
-		if (const UPCGSpatialData* SurfaceData = Cast<UPCGSpatialData>(TaggedData.Data))
+		// Fallback to getting bounds from actor
+		if (!BoundingShape && Context->SourceComponent.IsValid())
 		{
-			InContext->GeneratingShapes.Add(SurfaceData);
-			Outputs.Add(TaggedData);
+			check(bOutUnionWasCreated == false);
+			BoundingShape = Cast<UPCGSpatialData>(Context->SourceComponent->GetActorPCGData());
 		}
-	}
 
-	// If there are no generating shapes, early out
-	if (InContext->GeneratingShapes.IsEmpty())
-	{
-		PCGE_LOG_C(Warning, GraphAndLog, InContext, LOCTEXT("NoSurfaceFound", "No valid surfaces found from which to generate"));
-		return false;
+		return BoundingShape;
 	}
 
-	// Grab the Bounding Shape input if there is one.
-	const TArray<FPCGTaggedData> BoundingShapeInputs = InContext->InputData.GetInputsByPin(PCGSurfaceSamplerConstants::BoundingShapeLabel);
-
-	if (!Settings->bUnbounded)
+	bool InitializePerExecutionData(ContextType* Context, ExecStateType& OutState)
 	{
-		InContext->BoundingShapeSpatialInput = InContext->InputData.GetSpatialUnionOfInputsByPin(PCGSurfaceSamplerConstants::BoundingShapeLabel, InContext->bUnionDataCreated);
+		TRACE_CPUPROFILER_EVENT_SCOPE(FPCGSurfaceSamplerElement::InitializePerExecutionData);
+
+		check(Context);
+		const UPCGSurfaceSamplerSettings* Settings = Context->GetInputSettings<UPCGSurfaceSamplerSettings>();
+		check(Settings);
 
-		// Root dynamically created data to keep it safe from GC
-		if (InContext->bUnionDataCreated && InContext->BoundingShapeSpatialInput)
+		const TArray<FPCGTaggedData> SurfaceInputs = Context->InputData.GetInputsByPin(PCGSurfaceSamplerConstants::SurfaceLabel);
+		// If there are no surfaces to sample, early out
+		if (SurfaceInputs.IsEmpty())
 		{
-			const_cast<UPCGSpatialData*>(InContext->BoundingShapeSpatialInput)->AddToRoot();
+			PCGE_LOG_C(Warning, GraphAndLog, Context, LOCTEXT("NoSurfaceInputs", "No Surface input was provided."));
+			return false;
 		}
 
-		// Fallback to getting bounds from actor but only if we actually have some inputs
-		if (!InContext->BoundingShapeSpatialInput && InContext->SourceComponent.IsValid() && !InContext->GeneratingShapes.IsEmpty())
+		// Early out on invalid settings
+		// TODO: we could compute an approximate radius based on the points per squared meters if that's useful
+		const FVector& PointExtents = Settings->PointExtents;
+		if (PointExtents.X <= 0 || PointExtents.Y <= 0)
 		{
-			InContext->BoundingShapeSpatialInput = Cast<UPCGSpatialData>(InContext->SourceComponent->GetActorPCGData());
+			PCGE_LOG_C(Warning, GraphAndLog, Context, LOCTEXT("SkippedInvalidPointExtents", "Skipped - Invalid point extents"));
+			return false;
 		}
-	}
-	else if (BoundingShapeInputs.Num() > 0)
-	{
-		PCGE_LOG_C(Verbose, LogOnly, InContext, LOCTEXT("BoundsIgnored", "The bounds of the Bounding Shape input pin will be ignored because the Unbounded option is enabled."));
-	}
-
-	if (InContext->BoundingShapeSpatialInput)
-	{
-		InContext->BoundingShapeBounds = InContext->BoundingShapeSpatialInput->GetBounds();
-	}
-
-	// Warn if something is connected but no shape could be obtained for sampling
-	if (InContext->GeneratingShapes.Num() == 0 && (BoundingShapeInputs.Num() > 0 || SurfaceInputs.Num() > 0))
-	{
-		PCGE_LOG_C(Warning, GraphAndLog, InContext, LOCTEXT("NoGenerator", "No Surface input was provided, and no surface could be found in the Bounding Shape input for sampling. Connect the surface to be sampled to the Surface input."));
-		return false;
-	}
-
-	return true;
-}
-
-void FPCGSurfaceSamplerElement::AllocateOutputs(FPCGSurfaceSamplerContext* InContext) const
-{
-	TRACE_CPUPROFILER_EVENT_SCOPE(FPCGSurfaceSamplerElement::AllocateOutputs);
-
-	check(InContext);
-	TArray<FPCGTaggedData>& Outputs = InContext->OutputData.TaggedData;
 
-	const UPCGSurfaceSamplerSettings* Settings = InContext->GetInputSettings<UPCGSurfaceSamplerSettings>();
-	check(Settings);
+		// Grab the Bounding Shape input if there is one.
+		const TArray<FPCGTaggedData> BoundingShapeInputs = Context->InputData.GetInputsByPin(PCGSurfaceSamplerConstants::BoundingShapeLabel);
+		if (!Settings->bUnbounded)
+		{
+			bool bUnionWasCreated;
+			const UPCGSpatialData* BoundingShape = FindBoundingShape(Context, bUnionWasCreated);
+			if (BoundingShape && bUnionWasCreated)
+			{
+				// Must cast away const to root the union object
+				Context->RootAndTrackObject(const_cast<UPCGSpatialData*>(BoundingShape));
+			}
 
-	int32 GeneratingShapeIndex = 0;
-	while (GeneratingShapeIndex < InContext->GeneratingShapes.Num())
-	{
-		// If we have generating shape inputs, use them
-		const UPCGSpatialData* GeneratingShape = InContext->GeneratingShapes[GeneratingShapeIndex];
-		check(GeneratingShape);
+			if (BoundingShape)
+			{
+				OutState.BoundingShapeBounds = BoundingShape->GetBounds();
+			}
+		}
+		else if (BoundingShapeInputs.Num() > 0)
+		{
+			PCGE_LOG_C(Verbose, LogOnly, Context, LOCTEXT("BoundsIgnored", "The bounds of the Bounding Shape input pin will be ignored because the Unbounded option is enabled."));
+		}
 
-		// Calculate the intersection of bounds of the provided inputs
-		FBox InputBounds = FBox(EForceInit::ForceInit);
+		TArray<FPCGTaggedData>& Outputs = Context->OutputData.TaggedData;
+		// Find the generating shapes to sample
+		TArray<const UPCGSpatialData*>& GeneratingShapes = OutState.GeneratingShapes;
+		GeneratingShapes.Reserve(SurfaceInputs.Num());
 
-		if (GeneratingShape->IsBounded())
+		// Construct a list of shapes to generate samples from. Get these directly from the first input pin.
+		for (const FPCGTaggedData& TaggedData : SurfaceInputs)
 		{
-			InputBounds = GeneratingShape->GetBounds();
-
-			if (InContext->BoundingShapeBounds.IsValid)
+			if (const UPCGSpatialData* SurfaceData = Cast<UPCGSpatialData>(TaggedData.Data))
 			{
-				InputBounds = PCGHelpers::OverlapBounds(InputBounds, InContext->BoundingShapeBounds);
+				GeneratingShapes.Add(SurfaceData);
+				Outputs.Add(TaggedData);
 			}
 		}
-		else
+
+		// If there are no generating shapes, early out
+		if (GeneratingShapes.IsEmpty())
 		{
-			InputBounds = InContext->BoundingShapeBounds;
+			PCGE_LOG_C(Warning, GraphAndLog, Context, LOCTEXT("NoSurfaceFound", "No surfaces found from which to generate"));
+			return false;
 		}
 
-		PCGSurfaceSampler::FSurfaceSamplerSettings TentativeLoopData = PCGSurfaceSampler::FSurfaceSamplerSettings{};
-		TentativeLoopData.bEnableTimeSlicing = true;
-		if (!InputBounds.IsValid || !TentativeLoopData.Initialize(Settings, InContext, InputBounds))
+		// Iterate over the generating shapes to compute the bounds
+		int32 GeneratingShapeIndex = 0;
+		while (GeneratingShapeIndex < GeneratingShapes.Num())
 		{
-			if (!GeneratingShape->IsBounded())
+			// If we have generating shape inputs, use them
+			const UPCGSpatialData* GeneratingShape = GeneratingShapes[GeneratingShapeIndex];
+			check(GeneratingShape);
+
+			// Calculate the intersection of bounds of the provided inputs
+			FBox InputBounds = FBox(EForceInit::ForceInit);
+
+			if (GeneratingShape->IsBounded())
 			{
-				// Some inputs are unable to provide bounds, like the WorldRayHit, in which case the user must provide bounds.
-				PCGE_LOG_C(Warning, GraphAndLog, InContext, LOCTEXT("CouldNotObtainInputBounds", "Input data is not bounded, so bounds must be provided for sampling. Consider providing a Bounding Shape input."));
+				InputBounds = GeneratingShape->GetBounds();
+
+				if (OutState.BoundingShapeBounds.IsValid)
+				{
+					InputBounds = PCGHelpers::OverlapBounds(InputBounds, OutState.BoundingShapeBounds);
+				}
 			}
-			else if(!InputBounds.IsValid)
+			else
 			{
-				PCGE_LOG_C(Verbose, LogOnly, InContext, LOCTEXT("InvalidSamplingBounds", "Final sampling bounds is invalid/zero-sized."));
+				InputBounds = OutState.BoundingShapeBounds;
+			}
+
+			if (!InputBounds.IsValid || !OutState.Settings.Initialize(Settings, Context, InputBounds))
+			{
+				if (!GeneratingShape->IsBounded())
+				{
+					// Some inputs are unable to provide bounds, like the WorldRayHit, in which case the user must provide bounds.
+					PCGE_LOG_C(Warning, GraphAndLog, Context, LOCTEXT("CouldNotObtainInputBounds", "Input data is not bounded, so bounds must be provided for sampling. Consider providing a Bounding Shape input."));
+				}
+				else if(!InputBounds.IsValid)
+				{
+					// TODO [REVIEWERS]: Should this not be at least a warning, since the generating shape is being culled here?
+					PCGE_LOG_C(Verbose, LogOnly, Context, LOCTEXT("InvalidSamplingBounds", "Final sampling bounds is invalid/zero-sized."));
+				}
+
+				Outputs.RemoveAt(GeneratingShapeIndex);
+				GeneratingShapes.RemoveAt(GeneratingShapeIndex);
+				continue;
 			}
 
-			Outputs.RemoveAt(GeneratingShapeIndex);
-			InContext->GeneratingShapes.RemoveAt(GeneratingShapeIndex);
-			continue;
+			++GeneratingShapeIndex;
 		}
-		
-		InContext->LoopData.Emplace(std::move(TentativeLoopData));
 
-		UPCGPointData* NewPointData = NewObject<UPCGPointData>();
-		NewPointData->InitializeFromData(GeneratingShape);
+		// If there are still no generating shapes
+		if (GeneratingShapes.IsEmpty())
+		{
+			PCGE_LOG_C(Warning, GraphAndLog, Context, LOCTEXT("InvalidSurface", "No surfaces from which to generate are valid"));
+		}
 
-		// Directly set in the output, will allow for the data to be rooted if we are postponed.
-		Outputs[GeneratingShapeIndex].Data = NewPointData;
-		++GeneratingShapeIndex;
+		return true;
 	}
 }
 
 bool FPCGSurfaceSamplerElement::PrepareDataInternal(FPCGContext* InContext) const
 {
-	TRACE_CPUPROFILER_EVENT_SCOPE(FPCGSurfaceSamplerElement::PrepareData);
-	FPCGSurfaceSamplerContext* Context = static_cast<FPCGSurfaceSamplerContext*>(InContext);
-
+	TRACE_CPUPROFILER_EVENT_SCOPE(FPCGSurfaceSamplerElement::PrepareDataInternal);
+	ContextType* Context = static_cast<ContextType*>(InContext);
 	check(Context);
 
-	if (AddGeneratingShapesToContext(Context))
+	// Initialize the per-execution state data that won't change over the duration of the time slicing
+	if (!Context->InitializePerExecutionState(PCGSurfaceSamplerHelpers::InitializePerExecutionData))
+	{
+		PCGE_LOG(Warning, GraphAndLog, LOCTEXT("CouldNotInitializeExecutionState", "Could not initialize per-execution timeslice state data"));
+		return true;
+	}
+
+	TArray<const UPCGSpatialData*>& GeneratingShapes = Context->GetPerExecutionState().GeneratingShapes;
+	TArray<FPCGTaggedData>& Outputs = Context->OutputData.TaggedData;
+
+	// Initialize the per-iteration data, using the generating shapes as the source of iteration
+	const bool bInitSuccess = Context->InitializePerIterationStates(GeneratingShapes.Num(),
+		[&GeneratingShapes, &Outputs](IterStateType& OutState, const uint32 IterationIndex)->bool
+		{
+			OutState.OutputPoints = NewObject<UPCGPointData>();
+			OutState.OutputPoints->InitializeFromData(GeneratingShapes[IterationIndex]);
+			// Assigning this here prevents the need to root
+			Outputs[IterationIndex].Data = OutState.OutputPoints;
+			return true;
+		});
+
+	if (!bInitSuccess)
 	{
-		AllocateOutputs(Context);
-		Context->bDataPrepared = true;
+		PCGE_LOG(Warning, GraphAndLog, LOCTEXT("CouldNotInitializeIterationState", "Could not initialize per-iteration timeslice state data"));
 	}
 
 	return true;
@@ -507,36 +487,28 @@ bool FPCGSurfaceSamplerElement::PrepareDataInternal(FPCGContext* InContext) cons
 bool FPCGSurfaceSamplerElement::ExecuteInternal(FPCGContext* InContext) const
 {
 	TRACE_CPUPROFILER_EVENT_SCOPE(FPCGSurfaceSamplerElement::Execute);
-	FPCGSurfaceSamplerContext* Context = static_cast<FPCGSurfaceSamplerContext*>(InContext);
-
-	check(Context);
+	ContextType* TimeSlicedContext = static_cast<ContextType*>(InContext);
+	check(TimeSlicedContext);
 
 	// Prepare data failed, no need to execute
-	if (!Context->bDataPrepared)
+	if (!TimeSlicedContext->DataIsPrepared())
 	{
 		return true;
 	}
 
-	TArray<FPCGTaggedData>& Outputs = Context->OutputData.TaggedData;
-
-	while (Context->CurrentGeneratingShape < Context->GeneratingShapes.Num())
+	// The context will iterate over per-iteration states and execute the lambda until it returns true
+	return ExecuteSlice(TimeSlicedContext, [](ContextType* Context, const ExecStateType& ExecutionState, const IterStateType& IterationState, const uint32 IterationIndex)->bool
 	{
-		UPCGPointData* PointData = const_cast<UPCGPointData*>(Cast<UPCGPointData>(Outputs[Context->CurrentGeneratingShape].Data));
+		// Run the execution until the time slice is finished
+		const bool bAsyncDone = PCGSurfaceSampler::SampleSurface(Context, ExecutionState.Settings, ExecutionState.GeneratingShapes[IterationIndex], ExecutionState.BoundingShape, IterationState.OutputPoints, Context->TimeSliceIsEnabled());
 
-		if (PointData)
+		if (Context && bAsyncDone)
 		{
-			bool bIsDone = PCGSurfaceSampler::SampleSurface(Context, Context->GeneratingShapes[Context->CurrentGeneratingShape], Context->BoundingShapeSpatialInput, Context->LoopData[Context->CurrentGeneratingShape], PointData);
-
-			if (!bIsDone)
-			{
-				return false;
-			}
+			PCGE_LOG_C(Verbose, LogOnly, Context, FText::Format(LOCTEXT("GenerationInfo", "Generated {0} points in {1} cells"), IterationState.OutputPoints->GetPoints().Num(), ExecutionState.Settings.CellCount));
 		}
 
-		Context->CurrentGeneratingShape++;
-	}
-
-	return true;
+		return bAsyncDone;
+	});
 }
 
 void FPCGSurfaceSamplerElement::GetDependenciesCrc(const FPCGDataCollection& InInput, const UPCGSettings* InSettings, UPCGComponent* InComponent, FPCGCrc& OutCrc) const
diff --git a/Engine/Plugins/Experimental/PCG/Source/PCG/Private/PCGElement.cpp b/Engine/Plugins/Experimental/PCG/Source/PCG/Private/PCGElement.cpp
--- a/Engine/Plugins/Experimental/PCG/Source/PCG/Private/PCGElement.cpp
+++ b/Engine/Plugins/Experimental/PCG/Source/PCG/Private/PCGElement.cpp
@@ -9,7 +9,6 @@
 #include "Data/PCGPointData.h"
 #include "Elements/PCGDebugElement.h"
 #include "Elements/PCGSelfPruning.h"
-#include "Grid/PCGPartitionActor.h"
 
 #include "HAL/IConsoleManager.h"
 #include "Utils/PCGExtraCapture.h"
@@ -329,6 +328,11 @@ void IPCGElement::DisabledPassThroughData(FPCGContext* Context) const
 	}
 }
 
+FPCGContext* IPCGElement::CreateContext()
+{
+	return new FPCGContext();
+}
+
 #if WITH_EDITOR
 void IPCGElement::DebugDisplay(FPCGContext* Context) const
 {
@@ -480,6 +484,16 @@ void IPCGElement::CleanupAndValidateOutput(FPCGContext* Context) const
 	}
 }
 
+FPCGContext* IPCGElement::Initialize(const FPCGDataCollection& InputData, TWeakObjectPtr<UPCGComponent> SourceComponent, const UPCGNode* Node)
+{
+	FPCGContext* Context = CreateContext();
+	Context->InputData = InputData;
+	Context->SourceComponent = SourceComponent;
+	Context->Node = Node;
+
+	return Context;
+}
+
 bool IPCGElement::IsCacheableInstance(const UPCGSettingsInterface* InSettingsInterface) const
 {
 	if (InSettingsInterface)
@@ -521,14 +535,4 @@ void IPCGElement::GetDependenciesCrc(const FPCGDataCollection& InInput, const UP
 	OutCrc = Crc;
 }
 
-FPCGContext* FSimplePCGElement::Initialize(const FPCGDataCollection& InputData, TWeakObjectPtr<UPCGComponent> SourceComponent, const UPCGNode* Node)
-{
-	FPCGContext* Context = new FPCGContext();
-	Context->InputData = InputData;
-	Context->SourceComponent = SourceComponent;
-	Context->Node = Node;
-
-	return Context;
-}
-
 #undef LOCTEXT_NAMESPACE
diff --git a/Engine/Plugins/Experimental/PCG/Source/PCG/Public/Elements/PCGSurfaceSampler.h b/Engine/Plugins/Experimental/PCG/Source/PCG/Public/Elements/PCGSurfaceSampler.h
--- a/Engine/Plugins/Experimental/PCG/Source/PCG/Public/Elements/PCGSurfaceSampler.h
+++ b/Engine/Plugins/Experimental/PCG/Source/PCG/Public/Elements/PCGSurfaceSampler.h
@@ -2,8 +2,8 @@
 
 #pragma once
 
-#include "PCGContext.h"
 #include "PCGSettings.h"
+#include "PCGTimeSlicedElementBase.h"
 
 #include "Math/Box.h"
 
@@ -12,7 +12,62 @@
 class UPCGNode;
 class UPCGPointData;
 class UPCGSpatialData;
-struct FPCGContext;
+class UPCGSurfaceSamplerSettings;
 
+namespace PCGSurfaceSampler
+{
+	struct FSurfaceSamplerExecutionSettings
+	{
+		float PointsPerSquaredMeter = 1.0f;
+		FVector PointExtents = FVector::One() * 0.5f;
+		float Looseness = 0.0f;
+		bool bApplyDensityToPoints = false;
+		float PointSteepness = 0.0f;
+#if WITH_EDITORONLY_DATA
+		bool bKeepZeroDensityPoints = false;
+#endif
+
+		bool Initialize(const UPCGSurfaceSamplerSettings* Settings, const FPCGContext* Context, const FBox& InputBounds);
+		FIntVector2 ComputeCellIndices(int32 Index) const;
+
+		/** Computed values **/
+		FVector InterstitialDistance;
+		FVector InnerCellSize;
+		FVector CellSize;
+
+		int32 CellMinX;
+		int32 CellMaxX;
+		int32 CellMinY;
+		int32 CellMaxY;
+		int32 CellCount;
+		int64 TargetPointCount;
+		float Ratio;
+		int Seed;
+
+		FVector::FReal InputBoundsMinZ;
+		FVector::FReal InputBoundsMaxZ;
+	};
+
+	struct FSurfaceSamplerExecutionState
+	{
+		FSurfaceSamplerExecutionSettings Settings;
+
+		FBox BoundingShapeBounds = FBox(EForceInit::ForceInit);
+		const UPCGSpatialData* BoundingShape = nullptr;
+		TArray<const UPCGSpatialData*> GeneratingShapes;
+	};
+
+	struct FSurfaceSamplerIterationState
+	{
+		UPCGPointData* OutputPoints = nullptr;
+	};
+
+	/** Sample a surface and returns the resulting point data. Can't be timesliced. */
+	UPCGPointData* SampleSurface(FPCGContext* Context, const UPCGSpatialData* InSurface, const UPCGSpatialData* InBoundingShape, const FSurfaceSamplerExecutionSettings& ExecutionSettings);
+
+	/** Sample a surface and write the results in the given point data. Can be timesliced, and will return false if the processing is not done, true otherwise. */
+	bool SampleSurface(FPCGContext* Context, const FSurfaceSamplerExecutionSettings& Settings, const UPCGSpatialData* InSurface, const UPCGSpatialData* InBoundingShape, UPCGPointData* SampledData, const bool bTimeSlicingIsEnabled = false);
+}
+
 UCLASS(BlueprintType, ClassGroup = (Procedural))
 class PCG_API UPCGSurfaceSamplerSettings : public UPCGSettings
@@ -84,82 +139,16 @@
 #endif
 };
 
-namespace PCGSurfaceSampler
+class FPCGSurfaceSamplerElement : public TPCGTimeSlicedElementBase<PCGSurfaceSampler::FSurfaceSamplerExecutionState, PCGSurfaceSampler::FSurfaceSamplerIterationState>
 {
-	struct FSurfaceSamplerSettings
-	{
-		const UPCGSurfaceSamplerSettings* Settings = nullptr;
-
-		float PointsPerSquaredMeter = 1.0f;
-		FVector PointExtents = FVector::One() * 0.5f;
-		float Looseness = 0.0f;
-		bool bApplyDensityToPoints = false;
-		float PointSteepness = 0.0f;
-#if WITH_EDITORONLY_DATA
-		bool bKeepZeroDensityPoints = false;
-#endif
-
-		bool Initialize(const UPCGSurfaceSamplerSettings* InSettings, FPCGContext* Context, const FBox& InputBounds);
-		FIntVector2 ComputeCellIndices(int32 Index) const;
-
-		/** Computed values **/
-		FVector InterstitialDistance;
-		FVector InnerCellSize;
-		FVector CellSize;
-
-		int32 CellMinX;
-		int32 CellMaxX;
-		int32 CellMinY;
-		int32 CellMaxY;
-		int32 CellCount;
-		int64 TargetPointCount;
-		float Ratio;
-		int Seed;
-
-		FVector::FReal InputBoundsMinZ;
-		FVector::FReal InputBoundsMaxZ;
-
-		bool bEnableTimeSlicing = false;
-	};
-
-	/** Sample a surface and returns the resulting point data. Can't be timesliced. */
-	UPCGPointData* SampleSurface(FPCGContext* Context, const UPCGSpatialData* InSurface, const UPCGSpatialData* InBoundingShape, const FSurfaceSamplerSettings& LoopData);
-
-	/** Sample a surface and write the results in the given point data. Can be timesliced, and will return false if the processing is not done, true otherwise. */
-	bool SampleSurface(FPCGContext* Context, const UPCGSpatialData* InSurface, const UPCGSpatialData* InBoundingShape, const FSurfaceSamplerSettings& LoopData, UPCGPointData* SampledData);
-}
-
-struct FPCGSurfaceSamplerContext : public FPCGContext
-{
-	virtual ~FPCGSurfaceSamplerContext();
-
-	// Keep all data that won't be recomputed at every execution
-	const UPCGSpatialData* BoundingShapeSpatialInput = nullptr;
-	TArray<const UPCGSpatialData*, TInlineAllocator<16>> GeneratingShapes;
-	TArray<PCGSurfaceSampler::FSurfaceSamplerSettings, TInlineAllocator<16>> LoopData;
-	FBox BoundingShapeBounds = FBox(EForceInit::ForceInit);
-	int32 CurrentGeneratingShape = 0;
-	bool bDataPrepared = false;
-
-	/** True if BoundingShapeSpatialInput is data that was createdand rooted during data preparation */
-	bool bUnionDataCreated = false;
-};
-
-class FPCGSurfaceSamplerElement : public IPCGElement
-{
 public:
 	virtual void GetDependenciesCrc(const FPCGDataCollection& InInput, const UPCGSettings* InSettings, UPCGComponent* InComponent, FPCGCrc& OutCrc) const override;
-	virtual FPCGContext* Initialize(const FPCGDataCollection& InputData, TWeakObjectPtr<UPCGComponent> SourceComponent, const UPCGNode* Node) override;
 	// Might be sampling landscape or other external data, worth computing a full CRC in case we can halt change propagation/re-executions
 	virtual bool ShouldComputeFullOutputDataCrc(FPCGContext* Context) const override { return true; }
 
 protected:
 	virtual bool PrepareDataInternal(FPCGContext* InContext) const override;
 	virtual bool ExecuteInternal(FPCGContext* InContext) const override;
-
-private:
-	bool AddGeneratingShapesToContext(FPCGSurfaceSamplerContext* InContext) const;
-	void AllocateOutputs(FPCGSurfaceSamplerContext* InContext) const;
 };
 
 #if UE_ENABLE_INCLUDE_ORDER_DEPRECATED_IN_5_2
diff --git a/Engine/Plugins/Experimental/PCG/Source/PCG/Public/Elements/PCGTimeSlicedElementBase.h b/Engine/Plugins/Experimental/PCG/Source/PCG/Public/Elements/PCGTimeSlicedElementBase.h
new file mode 100644
--- /dev/null
+++ b/Engine/Plugins/Experimental/PCG/Source/PCG/Public/Elements/PCGTimeSlicedElementBase.h
@@ -0,0 +1,215 @@
+﻿// Copyright Epic Games, Inc. All Rights Reserved.
+
+#pragma once
+
+#include "PCGContext.h"
+#include "PCGElement.h"
+
+/**
+ * Rundown for utilizing the Time Slice Element and Context:
+ * 
+ * Create a struct to contain the state data for the element that will contain "per-execution" state data to be calculated only once
+ * Create a struct to contain the state data for the element that will contain "per-iteration" state data for each item (usually validated inputs) to iterate through
+ * Note: Either struct is optional and can be substituted with an empty struct if that aspect of the element is stateless
+ * Override the element type with TTimeSlicedPCGElement with two template arguments, the first being the static struct,and the second being the "per-iteration"
+ * Pass a function or lambda matching the correct signature to InitializePerExecutionState and initialize the static struct within
+ * Do the same for InitializePerIterationStates, but also pass the number of iterations. This will iterate through a state initialization for each iteration
+ * [Optional] If needed, mark UObjects that need to bypass GC with RootAndTrackObject or RootAndTrackObjectByName if you need to retrieve it later. The objects will be garbage collected at the end of the context lifetime
+ * DataIsPrepared SHOULD BE used to verify initialization was successful, such as in the ExecuteInternal if data was previously initialized in PrepareDataInternal
+ * Call ExecuteSlice with an execution function or lambda that returns a boolean that is true once the full execution is completed, or false otherwise
+ *
+ * Note: See PCGSurfaceSampler.h/.cpp as an example.
+ */
+
+// Forward declaration for friending
+template <typename PerExecutionStateT, typename PerIterationStateT> class TPCGTimeSlicedElementBase;
+
+namespace PCGTimeSlice
+{
+	struct FEmptyStruct {};
+}
+
+/**
+ * A PCG context with helper utility to enable element authors to more easily implement timeslicing.
+ * @tparam PerExecutionStateT Struct type of the "per-execution" static data state
+ * @tparam PerIterationStateT Struct type of the "per-iteration" data state
+ */
+template <typename PerExecutionStateT = PCGTimeSlice::FEmptyStruct, typename PerIterationStateT = PCGTimeSlice::FEmptyStruct>
+struct TPCGTimeSlicedContext : public FPCGContext
+{
+	virtual ~TPCGTimeSlicedContext() override;
+
+	virtual bool TimeSliceIsEnabled() const override final { return bTimeSliceIsEnabled; }
+	void SetTimeSliceIsEnabled(const bool bEnableTimeSlice = true) { bTimeSliceIsEnabled = bEnableTimeSlice; }
+
+	/** Retrieves the number of times this context was executed */
+	uint32 GetExecutionCount() const { return ExecutionCount; }
+
+	using InitSignature = bool(TPCGTimeSlicedContext* Context, PerExecutionStateT& OutState);
+
+	/** Initializes per execution state data if required. Returns true if data is completely initialized properly and false if a problem should result in early termination. */
+	bool InitializePerExecutionState(TFunctionRef<InitSignature> InitFunc = []{ return true; });
+
+	using IterSignature = bool(PerIterationStateT& OutState, const uint32 IterationIndex);
+
+	/** Initializes per execution state data if required. An array will be created with a state element for every execution iteration in the context. */
+	bool InitializePerIterationStates(int32 NumIterations = 1, TFunctionRef<IterSignature> IterFunc = []{ return true; });
+
+	PerExecutionStateT& GetPerExecutionState() { return PerExecutionStateData; }
+	const PerExecutionStateT& GetPerExecutionState() const { return PerExecutionStateData; }
+
+	TArray<PerIterationStateT>& GetPerIterationStateArray() { return PerIterationStateArray; }
+	const TArray<PerIterationStateT>& GetPerIterationStateArray() const { return PerIterationStateArray; }
+
+	/** Returns true if both the execution state and iteration state were fully initialized. Can be used to bypass any further initialization. */
+	bool DataIsPrepared() const { return bPerExecutionStateIsInitialized && bPerIterationStateIsInitialized; }
+
+	/** Fire and forget function to root a UObject for the duration of the Context, and then mark for garbage collection at context lifecycle end. */
+	void RootAndTrackObject(UObject* Object);
+
+private:
+	// Allow exposure to iteration index, etc
+	friend class TPCGTimeSlicedElementBase<PerExecutionStateT, PerIterationStateT>;
+
+	/** The number of times this context has been time sliced */
+	uint32 ExecutionCount = 0u;
+	/** The index of which iteration being processed. Ie. If a volume sampler has two volume inputs, it will be processed twice. */
+	int32 IterationIndex = 0;
+
+	bool bTimeSliceIsEnabled = true;
+	bool bPerExecutionStateIsInitialized = false;
+	bool bPerIterationStateIsInitialized = false;
+
+	/** The state of the timesliced context that won't change each iteration. Ie. the node settings. */
+	PerExecutionStateT PerExecutionStateData;
+	/** An array of the various states of timesliced context that will change per iteration. Ie. generating shape */
+	TArray<PerIterationStateT> PerIterationStateArray;
+
+	/** Tracks rooted UObjects that need to avoid garbage collection during the lifetime of the context */
+	TArray<UObject*> RootedAndTrackedObjectArray;
+};
+
+/**
+ * A PCG Element that will utilize a Time Slice Context.
+ * @tparam PerExecutionStateT Struct type of the "per-execution" static data state
+ * @tparam PerIterationStateT Struct type of the "per-iteration" data state
+ */
+template <typename PerExecutionStateT, typename PerIterationStateT>
+class TPCGTimeSlicedElementBase : public FSimplePCGElement
+{
+public:
+	// Aliases, for ease of use with the template
+	using ExecStateType = PerExecutionStateT;
+	using IterStateType = PerIterationStateT;
+	using ContextType = TPCGTimeSlicedContext<ExecStateType, IterStateType>;
+
+	virtual FPCGContext* CreateContext() override { return new ContextType(); }
+
+	using ExecSignature = bool(ContextType* Context, const PerExecutionStateT& PerExecutionState, const PerIterationStateT& PerIterationState, const uint32 IterationIndex);
+
+	/** Executes the delegate for every iteration. Will return false while still processing and true only when all tasks for all iterations are complete. */
+	bool ExecuteSlice(ContextType* Context, TFunctionRef<ExecSignature> ExecFunc) const;
+
+	// TODO: [FUTURE WORK] Consider callback 'on completion' and other Execution styles, like parallel iterations
+};
+
+template <typename PerExecutionStateT, typename PerIterationStateT>
+TPCGTimeSlicedContext<PerExecutionStateT, PerIterationStateT>::~TPCGTimeSlicedContext()
+{
+	for (UObject* Object : RootedAndTrackedObjectArray)
+	{
+		// Sanity check
+		check(Object);
+		if (!ensure(IsValid(Object) && Object->IsRooted()))
+		{
+			continue;
+		}
+
+		Object->RemoveFromRoot();
+		Object->MarkAsGarbage();
+	}
+
+	RootedAndTrackedObjectArray.Empty();
+}
+
+template <typename PerExecutionStateT, typename PerIterationStateT>
+bool TPCGTimeSlicedContext<PerExecutionStateT, PerIterationStateT>::InitializePerExecutionState(TFunctionRef<InitSignature> InitFunc)
+{
+	// Should only ever be initialized once, so if its called again, ignore it. This allows flexibility for the call to be somewhere that might be invoked numerous times
+	if (!bPerExecutionStateIsInitialized)
+	{
+		bPerExecutionStateIsInitialized = InitFunc(this, PerExecutionStateData);
+	}
+
+	return bPerExecutionStateIsInitialized;
+}
+
+template <typename PerExecutionStateT, typename PerIterationStateT>
+bool TPCGTimeSlicedContext<PerExecutionStateT, PerIterationStateT>::InitializePerIterationStates(int32 NumIterations, TFunctionRef<IterSignature> IterFunc)
+{
+	// Same as InitializePerExecutionState. Should only ever be initialized once, so if its called again, ignore it.
+	if (bPerIterationStateIsInitialized)
+	{
+		return true;
+	}
+
+	bPerIterationStateIsInitialized = true;
+
+	// An empty iteration state is still valid
+	if (NumIterations < 1)
+	{
+		return true;
+	}
+
+	bool bAnySucceeded{false};
+	PerIterationStateArray.Reserve(NumIterations);
+
+	for (int32 I = 0; I < NumIterations; ++I)
+	{
+		const bool bCurrentSucceeded = IterFunc(PerIterationStateArray.Emplace_GetRef(), I);
+		bAnySucceeded |= bCurrentSucceeded;
+
+		if (!bCurrentSucceeded)
+		{
+			// If it fails, remove the new state from the array and continue
+			PerIterationStateArray.RemoveAt(PerIterationStateArray.Num() - 1);
+		}
+	}
+
+	// Returns true if any succeeded
+	return bAnySucceeded;
+}
+
+template <typename PerExecutionStateT, typename PerIterationStateT>
+void TPCGTimeSlicedContext<PerExecutionStateT, PerIterationStateT>::RootAndTrackObject(UObject* Object)
+{
+	check(Object && IsValid(Object) && !Object->IsRooted());
+	Object->AddToRoot();
+	RootedAndTrackedObjectArray.AddUnique(Object);
+}
+
+template <typename PerExecutionStateT, typename PerIterationStateT>
+bool TPCGTimeSlicedElementBase<PerExecutionStateT, PerIterationStateT>::ExecuteSlice(ContextType* Context, TFunctionRef<ExecSignature> ExecFunc) const
+{
+	++Context->ExecutionCount;
+
+	// The user is responsible to check for this before execution, but just in case
+	if (!ensureMsgf(Context->DataIsPrepared() && !Context->PerIterationStateArray.IsEmpty(), TEXT("State data was not properly initialized.")))
+	{
+		return true;
+	}
+
+	do // The elements must execute at least once
+	{
+		if (ExecFunc(Context, Context->PerExecutionStateData, Context->PerIterationStateArray[Context->IterationIndex], Context->IterationIndex))
+		{
+			++Context->IterationIndex;
+		}
+		else
+		{
+			return false;
+		}
+	} while (Context->IterationIndex < Context->PerIterationStateArray.Num() && !Context->ShouldStop());
+
+	return Context->IterationIndex == Context->PerIterationStateArray.Num();
+}
diff --git a/Engine/Plugins/Experimental/PCG/Source/PCG/Public/PCGContext.h b/Engine/Plugins/Experimental/PCG/Source/PCG/Public/PCGContext.h
--- a/Engine/Plugins/Experimental/PCG/Source/PCG/Public/PCGContext.h
+++ b/Engine/Plugins/Experimental/PCG/Source/PCG/Public/PCGContext.h
@@ -101,6 +101,9 @@ struct PCG_API FPCGContext
 
 	AActor* GetTargetActor(const UPCGSpatialData* InSpatialData) const;
 
+	/** Time slicing is not enabled by default. */
+	virtual bool TimeSliceIsEnabled() const { return false; }
+
 #if WITH_EDITOR
 	/** Log warnings and errors to be displayed on node in graph editor. */
 	void LogVisual(ELogVerbosity::Type InVerbosity, const FText& InMessage) const;
diff --git a/Engine/Plugins/Experimental/PCG/Source/PCG/Public/PCGElement.h b/Engine/Plugins/Experimental/PCG/Source/PCG/Public/PCGElement.h
--- a/Engine/Plugins/Experimental/PCG/Source/PCG/Public/PCGElement.h
+++ b/Engine/Plugins/Experimental/PCG/Source/PCG/Public/PCGElement.h
@@ -70,7 +70,7 @@ public:
 	virtual ~IPCGElement() = default;
 
 	/** Creates a custom context object paired to this element */
-	virtual FPCGContext* Initialize(const FPCGDataCollection& InputData, TWeakObjectPtr<UPCGComponent> SourceComponent, const UPCGNode* Node) = 0;
+	virtual FPCGContext* Initialize(const FPCGDataCollection& InputData, TWeakObjectPtr<UPCGComponent> SourceComponent, const UPCGNode* Node);
 
 	/** Returns true if the element, in its current phase can be executed only from the main thread */
 	virtual bool CanExecuteOnlyOnMainThread(FPCGContext* Context) const { return false; }
@@ -119,6 +119,9 @@ protected:
 	/** Passes through data when the element is Disabled. Can be implemented to override what gets passed through. */
 	virtual void DisabledPassThroughData(FPCGContext* Context) const;
 
+	/** Let each element optionally act as a concrete factory for its own context */
+	virtual FPCGContext* CreateContext();
+
 #if WITH_EDITOR
 	virtual bool ShouldLog() const { return true; }
 #endif
@@ -132,8 +135,8 @@ private:
 */
 class PCG_API FSimplePCGElement : public IPCGElement
 {
-public:
-	virtual FPCGContext* Initialize(const FPCGDataCollection& InputData, TWeakObjectPtr<UPCGComponent> SourceComponent, const UPCGNode* Node) override;
+	// TODO: Will be scrubbed and cleaned up in a following CL
+	// TODO: In the same CL, go through and have classes with simple initializations replaced with CreateContext only
 };
 
 #if UE_ENABLE_INCLUDE_ORDER_DEPRECATED_IN_5_2
