Added an indirection in the landscape cache so we can retain more world information if we stream in levels.

diff --git a/Engine/Plugins/Experimental/PCG/Source/PCG/Private/Data/PCGLandscapeData.cpp b/Engine/Plugins/Experimental/PCG/Source/PCG/Private/Data/PCGLandscapeData.cpp
--- a/Engine/Plugins/Experimental/PCG/Source/PCG/Private/Data/PCGLandscapeData.cpp
+++ b/Engine/Plugins/Experimental/PCG/Source/PCG/Private/Data/PCGLandscapeData.cpp
@@ -171,7 +171,7 @@ bool UPCGLandscapeData::ProjectPoint(const FTransform& InTransform, const FBox&
 	ULandscapeComponent* LandscapeComponent = LandscapeInfo->XYtoComponentMap.FindRef(ComponentMapKey);
 	const FPCGLandscapeCacheEntry * LandscapeCacheEntry = (LandscapeComponent ? LandscapeCache->GetCacheEntry(LandscapeComponent, ComponentMapKey) : nullptr);
 #else
-	const FPCGLandscapeCacheEntry* LandscapeCacheEntry = LandscapeCache->GetCacheEntry(LandscapeInfo->GetLandscapeProxy()->GetOriginalLandscapeGuid(), ComponentMapKey);
+	const FPCGLandscapeCacheEntry* LandscapeCacheEntry = LandscapeCache->GetCacheEntry(LandscapeInfo->GetLandscapeProxy(), LandscapeInfo->GetLandscapeProxy()->GetOriginalLandscapeGuid(), ComponentMapKey);
 #endif
 
 	if (!LandscapeCacheEntry)
@@ -304,7 +304,7 @@ const UPCGPointData* UPCGLandscapeData::CreatePointData(FPCGContext* Context, co
 				ULandscapeComponent* LandscapeComponent = LandscapeInfo->XYtoComponentMap.FindRef(ComponentMapKey);
 				const FPCGLandscapeCacheEntry* LandscapeCacheEntry = LandscapeComponent ? LandscapeCache->GetCacheEntry(LandscapeComponent, ComponentMapKey) : nullptr;
 #else
-				const FPCGLandscapeCacheEntry* LandscapeCacheEntry = LandscapeCache->GetCacheEntry(LandscapeGuid, ComponentMapKey);
+				const FPCGLandscapeCacheEntry* LandscapeCacheEntry = LandscapeCache->GetCacheEntry(LandscapeInfo->GetLandscapeProxy(), LandscapeGuid, ComponentMapKey);
 #endif
 
 				if (!LandscapeCacheEntry)
diff --git a/Engine/Plugins/Experimental/PCG/Source/PCG/Private/Grid/PCGLandscapeCache.cpp b/Engine/Plugins/Experimental/PCG/Source/PCG/Private/Grid/PCGLandscapeCache.cpp
--- a/Engine/Plugins/Experimental/PCG/Source/PCG/Private/Grid/PCGLandscapeCache.cpp
+++ b/Engine/Plugins/Experimental/PCG/Source/PCG/Private/Grid/PCGLandscapeCache.cpp
@@ -503,7 +503,8 @@ void UPCGLandscapeCache::Serialize(FArchive& Archive)
 			FPCGLandscapeCacheEntry* Entry = new FPCGLandscapeCacheEntry();
 			Entry->Serialize(Archive, this, EntryIndex, SerializedContents);
 
-			CachedData.Add(Key, Entry);
+			CacheMapKey MapKey(Key.Key, Key.Value, nullptr);
+			CachedData.Add(MapKey, Entry);
 		}
 
 #if WITH_EDITOR
@@ -515,7 +516,8 @@ void UPCGLandscapeCache::Serialize(FArchive& Archive)
 		int32 EntryIndex = 0;
 		for (auto& CacheEntry : CachedData)
 		{
-			Archive << CacheEntry.Key;
+			TPair<FGuid, FIntPoint> Key(CacheEntry.Key.LandscapeGuid, CacheEntry.Key.Coordinate);
+			Archive << Key;
 			CacheEntry.Value->Serialize(Archive, this, EntryIndex++, SerializedContents);
 		}
 	}
@@ -528,11 +530,39 @@ void UPCGLandscapeCache::Initialize()
 	{
 		SetupLandscapeCallbacks();
 		CacheLayerNames();
+		UpdateCacheWorldKeys();
 		bInitialized = true;
 	}
 #endif
 }
 
+void UPCGLandscapeCache::UpdateCacheWorldKeys()
+{
+	check(IsInGameThread());
+
+	AActor* HintActor = Cast<AActor>(GetOuter());
+	if (!HintActor)
+	{
+		return;
+	}
+
+	TMap<CacheMapKey, FPCGLandscapeCacheEntry*> UpdatedCachedData;
+	UpdatedCachedData.Reserve(CachedData.Num());
+
+	for (const TPair<CacheMapKey, FPCGLandscapeCacheEntry*>& CacheEntry : CachedData)
+	{
+		CacheMapKey Key = CacheEntry.Key;
+		if (Key.WorldKey == FObjectKey())
+		{
+			Key = CacheMapKey(Key.LandscapeGuid, Key.Coordinate, HintActor);
+		}
+
+		UpdatedCachedData.Add(Key, CacheEntry.Value);
+	}
+
+	CachedData = MoveTemp(UpdatedCachedData);
+}
+
 void UPCGLandscapeCache::Tick(float DeltaSeconds)
 {
 	Initialize();
@@ -555,7 +585,7 @@ void UPCGLandscapeCache::Tick(float DeltaSeconds)
 		if (CacheMemorySize > MemoryThresholdInBytes)
 		{
 			TArray<FPCGLandscapeCacheEntry*> LoadedEntries;
-			for (TPair<TPair<FGuid, FIntPoint>, FPCGLandscapeCacheEntry*>& CacheEntry : CachedData)
+			for(TPair<CacheMapKey, FPCGLandscapeCacheEntry*>& CacheEntry : CachedData)
 			{
 				if (CacheEntry.Value->bDataLoaded)
 				{
@@ -603,7 +633,7 @@ void UPCGLandscapeCache::PrimeCache()
 	}
 
 	// First, gather all potential cached data to create, emplace a nullptr at these locations
-	TArray<TTuple<ULandscapeInfo*, ULandscapeComponent*, TPair<FGuid, FIntPoint>>> CacheEntriesToBuild;
+	TArray<TTuple<ULandscapeInfo*, ULandscapeComponent*, CacheMapKey>> CacheEntriesToBuild;
 
 	for (auto It = ULandscapeInfoMap::GetLandscapeInfoMap(World).Map.CreateIterator(); It; ++It)
 	{
@@ -624,8 +654,7 @@ void UPCGLandscapeCache::PrimeCache()
 						continue;
 					}
 
-					TPair<FGuid, FIntPoint> ComponentKey(LandscapeGuid, PCGLandscapeCache::GetCoordinates(LandscapeComponent));
-
+					CacheMapKey ComponentKey(LandscapeGuid, PCGLandscapeCache::GetCoordinates(LandscapeComponent), Cast<AActor>(GetOuter()));
 					if (!CachedData.Contains(ComponentKey))
 					{
 						CacheEntriesToBuild.Emplace(LandscapeInfo, LandscapeComponent, ComponentKey);
@@ -642,7 +671,7 @@ void UPCGLandscapeCache::PrimeCache()
 
 	ParallelFor(CacheEntriesToBuild.Num(), [&CacheEntriesToBuild, &NewEntries](int32 EntryIndex)
 	{
-		const TTuple<ULandscapeInfo*, ULandscapeComponent*, TPair<FGuid, FIntPoint>>& CacheEntryInfo = CacheEntriesToBuild[EntryIndex];
+		const TTuple<ULandscapeInfo*, ULandscapeComponent*, CacheMapKey>& CacheEntryInfo = CacheEntriesToBuild[EntryIndex];
 		NewEntries[EntryIndex] = FPCGLandscapeCacheEntry::CreateCacheEntry(CacheEntryInfo.Get<0>(), CacheEntryInfo.Get<1>());
 	});
 
@@ -672,7 +701,7 @@ void UPCGLandscapeCache::ClearCache()
 		Modify();
 	}
 
-	for (TPair<TPair<FGuid, FIntPoint>, FPCGLandscapeCacheEntry*>& CacheEntry : CachedData)
+	for(TPair<CacheMapKey, FPCGLandscapeCacheEntry*>& CacheEntry : CachedData)
 	{
 		delete CacheEntry.Value;
 		CacheEntry.Value = nullptr;
@@ -692,9 +721,10 @@ void UPCGLandscapeCache::TakeOwnership(UPCGLandscapeCache* InLandscapeCache)
 		Modify(/*bAlwaysMarkDirty=*/false);
 	}
 
+	InLandscapeCache->UpdateCacheWorldKeys();
 	bool bShouldDirty = false;
 
-	for (TPair<TPair<FGuid, FIntPoint>, FPCGLandscapeCacheEntry*>& CacheEntryPair : InLandscapeCache->CachedData)
+	for(TPair<CacheMapKey, FPCGLandscapeCacheEntry*>& CacheEntryPair : InLandscapeCache->CachedData)
 	{
 		FPCGLandscapeCacheEntry* CacheEntry = CachedData.FindOrAdd(CacheEntryPair.Key, CacheEntryPair.Value);
 		if (CacheEntry == CacheEntryPair.Value)
@@ -732,19 +762,19 @@ void UPCGLandscapeCache::TakeOwnership(UPCGLandscapeCache* InLandscapeCache)
 const FPCGLandscapeCacheEntry* UPCGLandscapeCache::GetCacheEntry(ULandscapeComponent* LandscapeComponent, const FIntPoint& ComponentCoordinate)
 {
 	const FGuid LandscapeGuid = (LandscapeComponent && LandscapeComponent->GetLandscapeProxy() ? LandscapeComponent->GetLandscapeProxy()->GetOriginalLandscapeGuid() : FGuid());
-	const FPCGLandscapeCacheEntry* CacheEntry = GetCacheEntry(LandscapeGuid, ComponentCoordinate);
+	const FPCGLandscapeCacheEntry* CacheEntry = GetCacheEntry(LandscapeComponent->GetLandscapeProxy(), LandscapeGuid, ComponentCoordinate);
 
 	if (!CacheEntry && LandscapeComponent && LandscapeComponent->GetLandscapeInfo())
 	{
 		FWriteScopeLock ScopeLock(CacheLock);
-		TPair<FGuid, FIntPoint> ComponentKey(LandscapeGuid, ComponentCoordinate);
+		CacheMapKey ComponentKey(LandscapeGuid, ComponentCoordinate, LandscapeComponent->GetOwner());
 		if (FPCGLandscapeCacheEntry** FoundEntry = CachedData.Find(ComponentKey))
 		{
 			CacheEntry = *FoundEntry;
 		}
 		else
 		{
-			check(LandscapeComponent->SectionBaseX / LandscapeComponent->ComponentSizeQuads == ComponentKey.Value.X && LandscapeComponent->SectionBaseY / LandscapeComponent->ComponentSizeQuads == ComponentKey.Value.Y);
+			check(LandscapeComponent->SectionBaseX / LandscapeComponent->ComponentSizeQuads == ComponentKey.Coordinate.X && LandscapeComponent->SectionBaseY / LandscapeComponent->ComponentSizeQuads == ComponentKey.Coordinate.Y);
 			if (FPCGLandscapeCacheEntry* NewEntry = FPCGLandscapeCacheEntry::CreateCacheEntry(LandscapeComponent->GetLandscapeInfo(), LandscapeComponent))
 			{
 				CacheEntry = NewEntry;
@@ -770,10 +800,10 @@ const FPCGLandscapeCacheEntry* UPCGLandscapeCache::GetCacheEntry(ULandscapeCompo
 }
 #endif
 
-const FPCGLandscapeCacheEntry* UPCGLandscapeCache::GetCacheEntry(const FGuid& LandscapeGuid, const FIntPoint& ComponentCoordinate)
+const FPCGLandscapeCacheEntry* UPCGLandscapeCache::GetCacheEntry(AActor* HintActor, const FGuid& LandscapeGuid, const FIntPoint& ComponentCoordinate)
 {
 	const FPCGLandscapeCacheEntry* CacheEntry = nullptr;
-	TPair<FGuid, FIntPoint> ComponentKey(LandscapeGuid, ComponentCoordinate);
+	CacheMapKey ComponentKey(LandscapeGuid, ComponentCoordinate, HintActor);
 
 	{
 #if WITH_EDITOR
@@ -829,7 +859,7 @@ void UPCGLandscapeCache::SampleMetadataOnPoint(ALandscapeProxy* Landscape, FPCGP
 	ULandscapeComponent* LandscapeComponent = LandscapeInfo->XYtoComponentMap.FindRef(ComponentMapKey);
 	const FPCGLandscapeCacheEntry* CacheEntry = GetCacheEntry(LandscapeComponent, ComponentMapKey);
 #else
-	const FPCGLandscapeCacheEntry* CacheEntry = GetCacheEntry(Landscape->GetOriginalLandscapeGuid(), ComponentMapKey);
+	const FPCGLandscapeCacheEntry* CacheEntry = GetCacheEntry(Landscape, Landscape->GetOriginalLandscapeGuid(), ComponentMapKey);
 #endif
 
 	if (!CacheEntry || CacheEntry->LayerData.IsEmpty())
@@ -906,8 +936,7 @@ void UPCGLandscapeCache::OnLandscapeChanged(ALandscapeProxy* InLandscape, const
 	{
 		if (LandscapeComponent)
 		{
-			const TPair<FGuid, FIntPoint> ComponentKey(InLandscape->GetOriginalLandscapeGuid(), PCGLandscapeCache::GetCoordinates(LandscapeComponent));
-
+			const CacheMapKey ComponentKey(InLandscape->GetOriginalLandscapeGuid(), PCGLandscapeCache::GetCoordinates(LandscapeComponent), InLandscape);
 			FPCGLandscapeCacheEntry* EntryToDelete = nullptr;
 
 			if (CachedData.RemoveAndCopyValue(ComponentKey, EntryToDelete))
@@ -990,8 +1019,7 @@ void UPCGLandscapeCache::RemoveComponentFromCache(const ALandscapeProxy* Landsca
 	{
 		if (LandscapeComponent)
 		{
-			const TPair<FGuid, FIntPoint> ComponentKey(LandscapeProxy->GetOriginalLandscapeGuid(), PCGLandscapeCache::GetCoordinates(LandscapeComponent));
-
+			const CacheMapKey ComponentKey(LandscapeProxy->GetOriginalLandscapeGuid(), PCGLandscapeCache::GetCoordinates(LandscapeComponent), LandscapeProxy);
 			if (FPCGLandscapeCacheEntry** FoundEntry = CachedData.Find(ComponentKey))
 			{
 				CachedData.Remove(ComponentKey);
diff --git a/Engine/Plugins/Experimental/PCG/Source/PCG/Private/Helpers/PCGBlueprintHelpers.cpp b/Engine/Plugins/Experimental/PCG/Source/PCG/Private/Helpers/PCGBlueprintHelpers.cpp
--- a/Engine/Plugins/Experimental/PCG/Source/PCG/Private/Helpers/PCGBlueprintHelpers.cpp
+++ b/Engine/Plugins/Experimental/PCG/Source/PCG/Private/Helpers/PCGBlueprintHelpers.cpp
@@ -203,7 +203,7 @@ TArray<FPCGLandscapeLayerWeight> UPCGBlueprintHelpers::GetInterpolatedPCGLandsca
 		ULandscapeComponent* LandscapeComponent = LandscapeInfo->XYtoComponentMap.FindRef(ComponentMapKey);
 		const FPCGLandscapeCacheEntry* CacheEntry = LandscapeCache->GetCacheEntry(LandscapeComponent, ComponentMapKey);
 		#else
-		const FPCGLandscapeCacheEntry* CacheEntry = LandscapeCache->GetCacheEntry(Landscape->GetOriginalLandscapeGuid(), ComponentMapKey);
+		const FPCGLandscapeCacheEntry* CacheEntry = LandscapeCache->GetCacheEntry(Landscape, Landscape->GetOriginalLandscapeGuid(), ComponentMapKey);
 		#endif
 
 		if (!CacheEntry)
diff --git a/Engine/Plugins/Experimental/PCG/Source/PCG/Public/Grid/PCGLandscapeCache.h b/Engine/Plugins/Experimental/PCG/Source/PCG/Public/Grid/PCGLandscapeCache.h
--- a/Engine/Plugins/Experimental/PCG/Source/PCG/Public/Grid/PCGLandscapeCache.h
+++ b/Engine/Plugins/Experimental/PCG/Source/PCG/Public/Grid/PCGLandscapeCache.h
@@ -2,7 +2,11 @@
 
 #pragma once
 
+#include "Engine/World.h"
+#include "Engine/Level.h"
+#include "GameFramework/Actor.h"
 #include "Serialization/BulkData.h"
+#include "WorldPartition/WorldPartitionRuntimeCellInterface.h"
 
 #include "PCGLandscapeCache.generated.h"
 
@@ -139,7 +143,7 @@ public:
 #endif
 
 	/** Gets landscape cache entry, works both in editor (but does not create) but works in game mode too. */
-	const FPCGLandscapeCacheEntry* GetCacheEntry(const FGuid& LandscapeGuid, const FIntPoint& ComponentKey);
+	const FPCGLandscapeCacheEntry* GetCacheEntry(AActor* HintActor, const FGuid& LandscapeGuid, const FIntPoint& ComponentKey);
 
 	TArray<FName> GetLayerNames(ALandscapeProxy* Landscape);
 
@@ -170,8 +174,44 @@ private:
 	void RemoveComponentFromCache(const ALandscapeProxy* LandscapeProxy);
 #endif
 
+	struct CacheMapKey
+	{
+		FGuid LandscapeGuid;
+		FIntPoint Coordinate;
+		FObjectKey WorldKey;
+
+		CacheMapKey(const FGuid& InLandscapeGuid, const FIntPoint& InCoordinate, const FObjectKey& InWorldKey)
+			: LandscapeGuid(InLandscapeGuid), Coordinate(InCoordinate), WorldKey(InWorldKey)
+		{}
+
+		CacheMapKey(const FGuid& InLandscapeGuid, const FIntPoint& InCoordinate, AActor* InHintActor)
+			: LandscapeGuid(InLandscapeGuid), Coordinate(InCoordinate)
+		{
+			if (InHintActor && InHintActor->GetLevel())
+			{
+				if (InHintActor->GetLevel()->GetWorldPartitionRuntimeCell())
+				{
+					WorldKey = FObjectKey(InHintActor->GetLevel()->GetWorldPartitionRuntimeCell()->GetOuterWorld());
+				}
+				else
+				{
+					WorldKey = FObjectKey(InHintActor->GetTypedOuter<UWorld>());
+				}
+			}
+		}
+
+		friend inline uint32 GetTypeHash(const CacheMapKey& Key)
+		{
+			return HashCombine(HashCombine(GetTypeHash(Key.LandscapeGuid), GetTypeHash(Key.Coordinate)), GetTypeHash(Key.WorldKey));
+		}
+
+		bool operator==(const CacheMapKey& Rhs) const { return LandscapeGuid == Rhs.LandscapeGuid && Coordinate == Rhs.Coordinate && WorldKey == Rhs.WorldKey; }
+	};
+
+	void UpdateCacheWorldKeys();
+
 	// Mapping of landscape guid + coordinates to entries. This is manually serialized as needed (depends on the serialize options).
-	TMap<TPair<FGuid, FIntPoint>, FPCGLandscapeCacheEntry*> CachedData;
+	TMap<CacheMapKey, FPCGLandscapeCacheEntry*> CachedData;
 
 	//TODO: separate by landscape
 	UPROPERTY(VisibleAnywhere, Category = "Cache")
