Fix for nullptr crash in PIE within the PCG landscape cache.

diff --git a/Engine/Plugins/Experimental/PCG/Source/PCG/Private/Grid/PCGLandscapeCache.cpp b/Engine/Plugins/Experimental/PCG/Source/PCG/Private/Grid/PCGLandscapeCache.cpp
--- a/Engine/Plugins/Experimental/PCG/Source/PCG/Private/Grid/PCGLandscapeCache.cpp
+++ b/Engine/Plugins/Experimental/PCG/Source/PCG/Private/Grid/PCGLandscapeCache.cpp
@@ -761,7 +761,12 @@ void UPCGLandscapeCache::TakeOwnership(UPCGLandscapeCache* InLandscapeCache)
 #if WITH_EDITOR
 const FPCGLandscapeCacheEntry* UPCGLandscapeCache::GetCacheEntry(ULandscapeComponent* LandscapeComponent, const FIntPoint& ComponentCoordinate)
 {
-	const FGuid LandscapeGuid = (LandscapeComponent && LandscapeComponent->GetLandscapeProxy() ? LandscapeComponent->GetLandscapeProxy()->GetOriginalLandscapeGuid() : FGuid());
+	if (!LandscapeComponent)
+	{
+		return nullptr;
+	}
+
+	const FGuid LandscapeGuid = (LandscapeComponent->GetLandscapeProxy() ? LandscapeComponent->GetLandscapeProxy()->GetOriginalLandscapeGuid() : FGuid());
 	const FPCGLandscapeCacheEntry* CacheEntry = GetCacheEntry(LandscapeComponent->GetLandscapeProxy(), LandscapeGuid, ComponentCoordinate);
 
 	if (!CacheEntry && LandscapeComponent && LandscapeComponent->GetLandscapeInfo())
