[PCG] Tweak ensure - !DataIsPrepared() is bad, but iterations empty is not (necessarily)

#rnx

diff --git a/Engine/Plugins/Experimental/PCG/Source/PCG/Public/Elements/PCGTimeSlicedElementBase.h b/Engine/Plugins/Experimental/PCG/Source/PCG/Public/Elements/PCGTimeSlicedElementBase.h
--- a/Engine/Plugins/Experimental/PCG/Source/PCG/Public/Elements/PCGTimeSlicedElementBase.h
+++ b/Engine/Plugins/Experimental/PCG/Source/PCG/Public/Elements/PCGTimeSlicedElementBase.h
@@ -194,7 +194,7 @@ bool TPCGTimeSlicedElementBase<PerExecutionStateT, PerIterationStateT>::ExecuteS
 	++Context->ExecutionCount;
 
 	// The user is responsible to check for this before execution, but just in case
-	if (!ensureMsgf(Context->DataIsPrepared() && !Context->PerIterationStateArray.IsEmpty(), TEXT("State data was not properly initialized.")))
+	if (!ensureMsgf(Context->DataIsPrepared(), TEXT("State data was not properly initialized.")) || Context->PerIterationStateArray.IsEmpty())
 	{
 		return true;
 	}
