Added convenience BP method to get copy of points for better efficiency

diff --git a/Engine/Plugins/Experimental/PCG/Source/PCG/Public/Data/PCGPointData.h b/Engine/Plugins/Experimental/PCG/Source/PCG/Public/Data/PCGPointData.h
--- a/Engine/Plugins/Experimental/PCG/Source/PCG/Public/Data/PCGPointData.h
+++ b/Engine/Plugins/Experimental/PCG/Source/PCG/Public/Data/PCGPointData.h
@@ -101,6 +101,9 @@ public:
 	UFUNCTION(BlueprintCallable, Category = SpatialData)
 	const TArray<FPCGPoint>& GetPoints() const { return Points; }
 
+	UFUNCTION(BlueprintCallable, BlueprintPure = false, Category = SpatialData)
+	TArray<FPCGPoint> GetPointsCopy() const { return Points; }
+
 	UFUNCTION(BlueprintCallable, Category = SpatialData)
 	FPCGPoint GetPoint(int32 Index) const;
 
