[PCG] Add warning if Volume Data points to a brush that won't function at runtime

#rnx

diff --git a/Engine/Plugins/Experimental/PCG/Source/PCG/Private/Data/PCGVolumeData.cpp b/Engine/Plugins/Experimental/PCG/Source/PCG/Private/Data/PCGVolumeData.cpp
--- a/Engine/Plugins/Experimental/PCG/Source/PCG/Private/Data/PCGVolumeData.cpp
+++ b/Engine/Plugins/Experimental/PCG/Source/PCG/Private/Data/PCGVolumeData.cpp
@@ -6,6 +6,8 @@
 #include "Elements/PCGVolumeSampler.h"
 #include "Helpers/PCGHelpers.h"
 
+#include "Components/BrushComponent.h"
+#include "Engine/CollisionProfile.h"
 #include "GameFramework/Volume.h"
 
 #include UE_INLINE_GENERATED_CPP_BY_NAME(PCGVolumeData)
@@ -15,6 +17,15 @@ void UPCGVolumeData::Initialize(AVolume* InVolume)
 	check(InVolume);
 	Volume = InVolume;
 	
+	if (PCGHelpers::IsRuntimeOrPIE())
+	{
+		const UBrushComponent* Brush = Volume->GetBrushComponent();
+		if (Brush && Brush->BodyInstance.GetCollisionProfileName() == UCollisionProfile::NoCollision_ProfileName)
+		{
+			UE_LOG(LogPCG, Warning, TEXT("Volume Data points to a Brush Component which is set to NoCollision and may not function outside of editor."));
+		}
+	}
+
 	FBoxSphereBounds BoxSphereBounds = Volume->GetBounds();
 	Bounds = FBox::BuildAABB(BoxSphereBounds.Origin, BoxSphereBounds.BoxExtent);
 
