[PCG] Fixed a crash where we had old PCGEditorGraphs being kept alive causing old callbacks to complain about node mismatches

diff --git a/Engine/Plugins/Experimental/PCG/Source/PCG/Private/PCGGraph.cpp b/Engine/Plugins/Experimental/PCG/Source/PCG/Private/PCGGraph.cpp
--- a/Engine/Plugins/Experimental/PCG/Source/PCG/Private/PCGGraph.cpp
+++ b/Engine/Plugins/Experimental/PCG/Source/PCG/Private/PCGGraph.cpp
@@ -242,7 +242,7 @@ UPCGGraph::UPCGGraph(const FObjectInitializer& ObjectInitializer)
 #if WITH_EDITOR
 	InputNode->OnNodeChangedDelegate.AddUObject(this, &UPCGGraph::OnNodeChanged);
 	OutputNode->OnNodeChangedDelegate.AddUObject(this, &UPCGGraph::OnNodeChanged);
-#endif 
+#endif
 
 	// Note: default connection from input to output
 	// should be added when creating from scratch,
@@ -285,7 +285,7 @@ void UPCGGraph::PostLoad()
 	{
 		Node->ConditionalPostLoad();
 	}
-	
+
 	// Also do this for ExtraNodes
 	for (UObject* ExtraNode : ExtraEditorNodes)
 	{
@@ -372,6 +372,15 @@ void UPCGGraph::BeginDestroy()
 	Super::BeginDestroy();
 }
 
+void UPCGGraph::AddReferencedObjects(UObject* InThis, FReferenceCollector& Collector)
+{
+#if WITH_EDITOR
+	UPCGGraph* This = CastChecked<UPCGGraph>(InThis);
+	Collector.AddReferencedObject(This->PCGEditorGraph, This);
+#endif
+	Super::AddReferencedObjects(InThis, Collector);
+}
+
 uint32 UPCGGraph::GetDefaultGridSize() const
 {
 	ensure(IsHierarchicalGenerationEnabled());
@@ -512,7 +521,7 @@ void UPCGGraph::OnNodesRemoved(TArrayView<UPCGNode*> InNodes)
 		}
 	}
 
-	NotifyGraphChanged(EPCGChangeType::Structural); 
+	NotifyGraphChanged(EPCGChangeType::Structural);
 #endif
 }
 
@@ -552,13 +561,13 @@ bool UPCGGraph::AddLabeledEdge(UPCGNode* From, const FName& FromPinLabel, UPCGNo
 	FromPin->AddEdgeTo(ToPin, &TouchedNodes);
 
 	bool bToPinBrokeOtherEdges = false;
-	
+
 	// Add an edge to a pin that doesn't allow multiple connections requires to do some cleanup
 	if (!ToPin->AllowsMultipleConnections())
 	{
 		bToPinBrokeOtherEdges = ToPin->BreakAllIncompatibleEdges(&TouchedNodes);
 	}
-	
+
 	const EPCGChangeType ChangeType = PCGGraphUtils::NotifyTouchedNodes(TouchedNodes);
 
 #if WITH_EDITOR
@@ -642,7 +651,7 @@ void UPCGGraph::RemoveNodes_Internal(TArrayView<UPCGNode*> InNodes)
 	}
 
 	Modify();
-	
+
 	TSet<UPCGNode*> TouchedNodes;
 
 	for (UPCGNode* Node : InNodes)
@@ -1569,7 +1578,7 @@ void UPCGGraphInstance::RefreshParameters(EPCGGraphParameterEvent InChangeType,
 		ParametersOverrides.Reset();
 	}
 	else
-	{		
+	{
 		const FInstancedPropertyBag* ParentUserParameters = Graph->GetUserParametersStruct();
 
 		// Refresh can modify nothing, but we still need to keep a snapshot of this object state, if it ever change.
@@ -1644,7 +1653,7 @@ bool UPCGGraphInstance::IsPropertyOverriddenAndNotDefault(const FProperty* InPro
 
 bool UPCGGraphInstance::IsGraphParameterOverridden(const FName PropertyName) const
 {
-	return (ParametersOverrides.Parameters.FindPropertyDescByName(PropertyName) != nullptr); 
+	return (ParametersOverrides.Parameters.FindPropertyDescByName(PropertyName) != nullptr);
 }
 
 bool FPCGOverrideInstancedPropertyBag::RefreshParameters(const FInstancedPropertyBag* ParentUserParameters, EPCGGraphParameterEvent InChangeType, FName InChangedPropertyName)
@@ -1880,4 +1889,4 @@ void FPCGOverrideInstancedPropertyBag::MigrateToNewBagInstance(const FInstancedP
 	}
 }
 
-#undef LOCTEXT_NAMESPACE
\ No newline at end of file
+#undef LOCTEXT_NAMESPACE
diff --git a/Engine/Plugins/Experimental/PCG/Source/PCG/Public/PCGGraph.h b/Engine/Plugins/Experimental/PCG/Source/PCG/Public/PCGGraph.h
--- a/Engine/Plugins/Experimental/PCG/Source/PCG/Public/PCGGraph.h
+++ b/Engine/Plugins/Experimental/PCG/Source/PCG/Public/PCGGraph.h
@@ -15,6 +15,7 @@
 
 class UPCGGraphInterface;
 #if WITH_EDITOR
+class UPCGEditorGraph;
 struct FEdGraphPinType;
 #endif // WITH_EDITOR
 
@@ -182,6 +183,7 @@ public:
 #endif
 
 	virtual void BeginDestroy() override;
+	static void AddReferencedObjects(UObject* InThis, FReferenceCollector& Collector);
 
 #if WITH_EDITOR
 	virtual void PreEditChange(FProperty* InProperty) override;
@@ -331,6 +333,9 @@ protected:
 	// Extra data to hold information that is useful only in editor (like comments)
 	UPROPERTY()
 	TArray<TObjectPtr<UObject>> ExtraEditorNodes;
+
+	// Editor graph created from PCG Editor but owned by this, reference is collected using AddReferencedObjects
+	TObjectPtr<UPCGEditorGraph> PCGEditorGraph = nullptr;
 #endif // WITH_EDITORONLY_DATA
 
 	// Parameters
@@ -340,7 +345,7 @@ protected:
 #if WITH_EDITOR
 	UFUNCTION(BlueprintInternalUseOnly)
 	bool UserParametersIsPinTypeAccepted(FEdGraphPinType InPinType, bool bIsChild);
-	
+
 	UFUNCTION(BlueprintInternalUseOnly)
 	bool UserParametersCanRemoveProperty(FGuid InPropertyID, FName InPropertyName);
 #endif // WITH_EDITOR
diff --git a/Engine/Plugins/Experimental/PCG/Source/PCGEditor/Private/PCGEditor.cpp b/Engine/Plugins/Experimental/PCG/Source/PCGEditor/Private/PCGEditor.cpp
--- a/Engine/Plugins/Experimental/PCG/Source/PCGEditor/Private/PCGEditor.cpp
+++ b/Engine/Plugins/Experimental/PCG/Source/PCGEditor/Private/PCGEditor.cpp
@@ -90,10 +90,15 @@ void FPCGEditor::Initialize(const EToolkitMode::Type InMode, const TSharedPtr<cl
 		PCGGraphBeingEdited->OnGraphDynamicallyExecutedDelegate.AddRaw(this, &FPCGEditor::OnGraphDynamicallyExecuted);
 	}
 
-	PCGEditorGraph = NewObject<UPCGEditorGraph>(PCGGraphBeingEdited, UPCGEditorGraph::StaticClass(), NAME_None, RF_Transactional | RF_Transient);
-	PCGEditorGraph->Schema = UPCGEditorGraphSchema::StaticClass();
-	PCGEditorGraph->InitFromNodeGraph(InPCGGraph);
-	PCGEditorGraph->SetEditor(SharedThis(this));
+	if (!PCGGraphBeingEdited->PCGEditorGraph)
+	{
+		PCGGraphBeingEdited->PCGEditorGraph = NewObject<UPCGEditorGraph>(PCGGraphBeingEdited, UPCGEditorGraph::StaticClass(), NAME_None, RF_Transactional | RF_Transient);
+		PCGGraphBeingEdited->PCGEditorGraph->Schema = UPCGEditorGraphSchema::StaticClass();
+		PCGGraphBeingEdited->PCGEditorGraph->InitFromNodeGraph(InPCGGraph);
+	}
+
+	PCGGraphBeingEdited->PCGEditorGraph->SetEditor(SharedThis(this));
+	PCGEditorGraph = PCGGraphBeingEdited->PCGEditorGraph;
 
 	FPropertyEditorModule& PropertyEditorModule = FModuleManager::GetModuleChecked<FPropertyEditorModule>("PropertyEditor");
 
@@ -326,7 +331,7 @@ void FPCGEditor::RegisterTabSpawners(const TSharedRef<FTabManager>& InTabManager
 	InTabManager->RegisterTabSpawner(FPCGEditor_private::PaletteID, FOnSpawnTab::CreateSP(this, &FPCGEditor::SpawnTab_Palette))
 		.SetDisplayName(LOCTEXT("PaletteTab", "Palette"))
 		.SetGroup(WorkspaceMenuCategoryRef);
-	
+
 	InTabManager->RegisterTabSpawner(FPCGEditor_private::DebugObjectID, FOnSpawnTab::CreateSP(this, &FPCGEditor::SpawnTab_DebugObject))
 		.SetDisplayName(LOCTEXT("DebugTab", "Debug Object Tree"))
 		.SetGroup(WorkspaceMenuCategoryRef);
@@ -538,7 +543,7 @@ void FPCGEditor::BindCommands()
 		FExecuteAction::CreateSP(this, &FPCGEditor::OnEditGraphSettings),
 		FCanExecuteAction(),
 		FIsActionChecked::CreateSP(this, &FPCGEditor::IsEditGraphSettingsToggled));
-	
+
 	GraphEditorCommands->MapAction(
 		PCGEditorCommands.CollapseNodes,
 		FExecuteAction::CreateSP(this, &FPCGEditor::OnCollapseNodesInSubgraph),
@@ -853,7 +858,7 @@ void FPCGEditor::OnAddDynamicInputPin()
 	check(GraphEditorWidget.IsValid());
 
 	const FGraphPanelSelectionSet SelectedNodes = GraphEditorWidget->GetSelectedNodes();
-	
+
 	if (!ensure(SelectedNodes.Num() == 1))
 	{
 		UE_LOG(LogPCGEditor, Warning, TEXT("Attempting to add new input pin to multiple nodes."));
@@ -885,7 +890,7 @@ void FPCGEditor::OnCollapseNodesInSubgraph()
 		UE_LOG(LogPCGEditor, Error, TEXT("GraphEditorWidget or PCGEditorGraph is null, aborting"));
 		return;
 	}
-	
+
 	UPCGGraph* PCGGraph = PCGEditorGraph->GetPCGGraph();
 	if (PCGGraph == nullptr)
 	{
@@ -1123,7 +1128,7 @@ void FPCGEditor::OnToggleInspected()
 	{
 		PCGGraphNodeBeingInspected->SetInspected(false);
 	}
-	
+
 	UEdGraphNode* GraphNode = GraphEditorWidget->GetSingleSelectedNode();
 	UPCGEditorGraphNodeBase* PCGGraphNodeBase = Cast<UPCGEditorGraphNodeBase>(GraphNode);
 	if (PCGGraphNodeBase && PCGGraphNodeBase != PCGGraphNodeBeingInspected)
@@ -1135,7 +1140,7 @@ void FPCGEditor::OnToggleInspected()
 	{
 		PCGGraphNodeBeingInspected = nullptr;
 	}
-	
+
 	OnInspectedNodeChangedDelegate.Broadcast(PCGGraphNodeBeingInspected);
 	GetTabManager()->TryInvokeTab(FPCGEditor_private::AttributesID);
 }
@@ -1167,10 +1172,10 @@ ECheckBoxState FPCGEditor::GetInspectedCheckState() const
 		{
 			return ECheckBoxState::Unchecked;
 		}
-	
+
 		bool bAllEnabled = true;
 		bool bAnyEnabled = false;
-		
+
 		for (UObject* Object : SelectedNodes)
 		{
 			const UPCGEditorGraphNodeBase* PCGEditorGraphNode = Cast<UPCGEditorGraphNodeBase>(Object);
@@ -1178,7 +1183,7 @@ ECheckBoxState FPCGEditor::GetInspectedCheckState() const
 			{
 				continue;
 			}
-			
+
 			bAllEnabled &= PCGEditorGraphNode->GetInspected();
 			bAnyEnabled |= PCGEditorGraphNode->GetInspected();
 		}
@@ -1189,11 +1194,11 @@ ECheckBoxState FPCGEditor::GetInspectedCheckState() const
 		}
 		else if (bAnyEnabled)
 		{
-			return ECheckBoxState::Undetermined; 
+			return ECheckBoxState::Undetermined;
 		}
 	}
-	
-	return ECheckBoxState::Unchecked;	
+
+	return ECheckBoxState::Unchecked;
 }
 
 void FPCGEditor::OnToggleEnabled()
@@ -1232,7 +1237,7 @@ void FPCGEditor::OnToggleEnabled()
 				bChanged = true;
 			}
 		}
-		
+
 		if (bChanged)
 		{
 			GraphEditorWidget->NotifyGraphChanged();
@@ -1241,12 +1246,12 @@ void FPCGEditor::OnToggleEnabled()
 }
 
 ECheckBoxState FPCGEditor::GetEnabledCheckState() const
-{	
+{
 	if (GraphEditorWidget.IsValid())
 	{
 		bool bAllEnabled = true;
 		bool bAnyEnabled = false;
-		
+
 		for (UObject* Object : GraphEditorWidget->GetSelectedNodes())
 		{
 			const UPCGEditorGraphNodeBase* PCGEditorGraphNode = Cast<UPCGEditorGraphNodeBase>(Object);
@@ -1277,10 +1282,10 @@ ECheckBoxState FPCGEditor::GetEnabledCheckState() const
 		}
 		else if (bAnyEnabled)
 		{
-			return ECheckBoxState::Undetermined; 
+			return ECheckBoxState::Undetermined;
 		}
 	}
-	
+
 	return ECheckBoxState::Unchecked;
 }
 
@@ -1288,11 +1293,11 @@ void FPCGEditor::OnToggleDebug()
 {
 	const ECheckBoxState CheckState = GetDebugCheckState();
 	const bool bNewCheckState = !(CheckState != ECheckBoxState::Unchecked);
-	
+
 	if (GraphEditorWidget.IsValid())
 	{
 		const FScopedTransaction Transaction(*FPCGEditorCommon::ContextIdentifier, LOCTEXT("PCGEditorToggleDebugTransactionMessage", "PCG Editor: Toggle Debug Nodes"), nullptr);
-		
+
 		for (UObject* Object : GraphEditorWidget->GetSelectedNodes())
 		{
 			UPCGEditorGraphNodeBase* PCGEditorGraphNode = Cast<UPCGEditorGraphNodeBase>(Object);
@@ -1431,7 +1436,7 @@ ECheckBoxState FPCGEditor::GetDebugCheckState() const
 	{
 		bool bAllDebug = true;
 		bool bAnyDebug = false;
-		
+
 		for (UObject* Object : GraphEditorWidget->GetSelectedNodes())
 		{
 			const UPCGEditorGraphNodeBase* PCGEditorGraphNode = Cast<UPCGEditorGraphNodeBase>(Object);
@@ -1465,7 +1470,7 @@ ECheckBoxState FPCGEditor::GetDebugCheckState() const
 			return ECheckBoxState::Undetermined;
 		}
 	}
-	
+
 	return ECheckBoxState::Unchecked;
 }
 
@@ -2131,8 +2136,8 @@ void FPCGEditor::JumpToDefinition(const UClass* Class) const
 bool FPCGEditor::IsReadOnlyProperty(const FPropertyAndParent& InPropertyAndParent, IDetailsView* InDetailsView) const
 {
 	// Everything is writeable when not in an instance
-	if (!InDetailsView || 
-		InPropertyAndParent.ParentProperties.IsEmpty() || 
+	if (!InDetailsView ||
+		InPropertyAndParent.ParentProperties.IsEmpty() ||
 		InPropertyAndParent.ParentProperties.Last()->GetFName() != GET_MEMBER_NAME_CHECKED(UPCGSettingsInstance, Settings))
 	{
 		return false;
@@ -2145,7 +2150,7 @@ bool FPCGEditor::IsReadOnlyProperty(const FPropertyAndParent& InPropertyAndParen
 		{
 			continue;
 		}
-		
+
 		if (UPCGSettingsInstance* Instance = Cast<UPCGSettingsInstance>(SelectedObject.Get()))
 		{
 			return true;
@@ -2185,7 +2190,7 @@ bool FPCGEditor::IsVisibleProperty(const FPropertyAndParent& InPropertyAndParent
 	// Hide debug settings from the setting when showing the instance settings.
 	if (InPropertyAndParent.Property.GetFName() == GET_MEMBER_NAME_CHECKED(UPCGSettings, bEnabled) ||
 		InPropertyAndParent.Property.GetFName() == GET_MEMBER_NAME_CHECKED(UPCGSettings, bDebug) ||
-		(InPropertyAndParent.ParentProperties.Num() >= 2 && 
+		(InPropertyAndParent.ParentProperties.Num() >= 2 &&
 			InPropertyAndParent.ParentProperties[InPropertyAndParent.ParentProperties.Num() - 2]->GetFName() == GET_MEMBER_NAME_CHECKED(UPCGSettings, DebugSettings)))
 	{
 		return false;
