[PCG] Don't generate components with grid sizes that should not be part of the current execution.

Still clean them up though, to not have situations where user can't clean up some components.

#rnx

diff --git a/Engine/Plugins/Experimental/PCG/Source/PCG/Private/PCGSubsystem.cpp b/Engine/Plugins/Experimental/PCG/Source/PCG/Private/PCGSubsystem.cpp
--- a/Engine/Plugins/Experimental/PCG/Source/PCG/Private/PCGSubsystem.cpp
+++ b/Engine/Plugins/Experimental/PCG/Source/PCG/Private/PCGSubsystem.cpp
@@ -372,8 +372,13 @@ FPCGTaskId UPCGSubsystem::ScheduleComponent(UPCGComponent* PCGComponent, bool bS
 			Dependencies.Add(OriginalComponentTask);
 		}
 
-		auto LocalGenerateTask = [OriginalComponent = PCGComponent, &Dependencies, bSave](UPCGComponent* LocalComponent)
+		auto LocalGenerateTask = [OriginalComponent = PCGComponent, &Dependencies, bSave, &GridSizes](UPCGComponent* LocalComponent)
 		{
+			if (!GridSizes.Contains(LocalComponent->GetGenerationGridSize()))
+			{
+				return InvalidPCGTaskId;
+			}
+
 			// If the local component is currently generating, it's probably because it was requested by a refresh.
 			// Wait after this one instead
 			if (LocalComponent->IsGenerating())
@@ -463,6 +468,7 @@ FPCGTaskId UPCGSubsystem::ScheduleCleanup(UPCGComponent* PCGComponent, bool bRem
 				return LocalComponent->CurrentCleanupTask;
 			}
 
+			// Always executes regardless of local component grid size - clean up as much as possible.
 			return LocalComponent->CleanupInternal(bRemoveComponents, /*bSave=*/ false, Dependencies);
 		};
 
