Added some options to the landscape data (physical material retrieval, component key, proxy touched).

diff --git a/Engine/Plugins/Experimental/PCG/Source/PCG/Private/Data/PCGLandscapeData.cpp b/Engine/Plugins/Experimental/PCG/Source/PCG/Private/Data/PCGLandscapeData.cpp
--- a/Engine/Plugins/Experimental/PCG/Source/PCG/Private/Data/PCGLandscapeData.cpp
+++ b/Engine/Plugins/Experimental/PCG/Source/PCG/Private/Data/PCGLandscapeData.cpp
@@ -5,6 +5,7 @@
 #include "PCGSubsystem.h"
 #include "Data/PCGPointData.h"
 #include "Data/PCGSurfaceData.h"
+#include "Data/PCGWorldData.h"
 #include "Grid/PCGLandscapeCache.h"
 #include "Helpers/PCGHelpers.h"
 
@@ -12,10 +13,17 @@
 #include "LandscapeInfo.h"
 #include "LandscapeProxy.h"
 #include "Engine/World.h"
+#include "PhysicalMaterials/PhysicalMaterial.h"
 
 #include UE_INLINE_GENERATED_CPP_BY_NAME(PCGLandscapeData)
 
-void UPCGLandscapeData::Initialize(const TArray<TWeakObjectPtr<ALandscapeProxy>>& InLandscapes, const FBox& InBounds, bool bInHeightOnly, bool bInUseMetadata)
+namespace PCGLandscapeDataConstants
+{
+	const FName ComponentXAttribute = TEXT("ComponentX");
+	const FName ComponentYAttribute = TEXT("ComponentY");
+}
+
+void UPCGLandscapeData::Initialize(const TArray<TWeakObjectPtr<ALandscapeProxy>>& InLandscapes, const FBox& InBounds, const FPCGLandscapeDataProps& InDataProps)
 {
 	TSet<ALandscapeProxy*> LandscapesToIgnore;
 
@@ -51,8 +59,7 @@ void UPCGLandscapeData::Initialize(const TArray<TWeakObjectPtr<ALandscapeProxy>>
 	check(FirstLandscape);
 
 	Bounds = InBounds;
-	bHeightOnly = bInHeightOnly;
-	bUseMetadata = bInUseMetadata;
+	DataProps = InDataProps;
 
 	Transform = FirstLandscape->GetActorTransform();
 
@@ -63,7 +70,7 @@ void UPCGLandscapeData::Initialize(const TArray<TWeakObjectPtr<ALandscapeProxy>>
 	// TODO: find a better way to do this - maybe there should be a prototype metadata in the landscape cache
 	if (LandscapeCache)
 	{
-		if (bUseMetadata)
+		if (DataProps.bGetLayerWeights)
 		{
 			for (TSoftObjectPtr<ALandscapeProxy> Landscape : Landscapes)
 			{
@@ -83,6 +90,23 @@ void UPCGLandscapeData::Initialize(const TArray<TWeakObjectPtr<ALandscapeProxy>>
 	{
 		UE_LOG(LogPCG, Error, TEXT("Landscape is unable to access the landscape cache"));
 	}
+
+	// Create secondary attributes as we need them
+	if (DataProps.bGetActorReference)
+	{
+		Metadata->CreateAttribute<FSoftObjectPath>(PCGPointDataConstants::ActorReferenceAttribute, FSoftObjectPath(), /*bAllowInterpolation=*/false, /*bOverrideParent=*/false);
+	}
+
+	if (DataProps.bGetPhysicalMaterial)
+	{
+		Metadata->CreateAttribute<FSoftObjectPath>(PCGWorldRayHitConstants::PhysicalMaterialReferenceAttribute, FSoftObjectPath(), /*bAllowInterpolation=*/false, /*bOverrideParent*/false);
+	}
+
+	if (DataProps.bGetComponentCoordinates)
+	{
+		Metadata->CreateInteger32Attribute(PCGLandscapeDataConstants::ComponentXAttribute, 0, /*bAllowsInterpolation=*/false);
+		Metadata->CreateInteger32Attribute(PCGLandscapeDataConstants::ComponentYAttribute, 0, /*bAllowsInterpolation=*/false);
+	}
 }
 
 void UPCGLandscapeData::PostLoad()
@@ -111,6 +135,20 @@ void UPCGLandscapeData::PostLoad()
 
 	UPCGSubsystem* PCGSubsystem = UPCGSubsystem::GetInstance(FirstLandscape ? FirstLandscape->GetWorld() : nullptr);
 	LandscapeCache = PCGSubsystem ? PCGSubsystem->GetLandscapeCache() : nullptr;
+
+#if WITH_EDITOR
+	if (bHeightOnly_DEPRECATED)
+	{
+		DataProps.bGetHeightOnly = bHeightOnly_DEPRECATED;
+		bHeightOnly_DEPRECATED = false;
+	}
+
+	if (!bUseMetadata_DEPRECATED)
+	{
+		DataProps.bGetLayerWeights = bUseMetadata_DEPRECATED;
+		bUseMetadata_DEPRECATED = true;
+	}
+#endif
 }
 
 FBox UPCGLandscapeData::GetBounds() const
@@ -181,13 +219,52 @@ bool UPCGLandscapeData::ProjectPoint(const FTransform& InTransform, const FBox&
 
 	const FVector2D ComponentLocalPoint(LocalPoint.X - ComponentMapKey.X * LandscapeInfo->ComponentSizeQuads, LocalPoint.Y - ComponentMapKey.Y * LandscapeInfo->ComponentSizeQuads);
 
-	if (bHeightOnly)
+	if (DataProps.bGetHeightOnly)
 	{
-		LandscapeCacheEntry->GetInterpolatedPointHeightOnly(ComponentLocalPoint, OutPoint, bUseMetadata ? OutMetadata : nullptr);
+		LandscapeCacheEntry->GetInterpolatedPointHeightOnly(ComponentLocalPoint, OutPoint, DataProps.bGetLayerWeights ? OutMetadata : nullptr);
 	}
 	else
 	{
-		LandscapeCacheEntry->GetInterpolatedPoint(ComponentLocalPoint, OutPoint, bUseMetadata ? OutMetadata : nullptr);
+		LandscapeCacheEntry->GetInterpolatedPoint(ComponentLocalPoint, OutPoint, DataProps.bGetLayerWeights ? OutMetadata : nullptr);
+	}
+
+	ULandscapeHeightfieldCollisionComponent* LandscapeCollisionComponent = LandscapeInfo->XYtoCollisionComponentMap.FindRef(ComponentMapKey);
+
+	if (DataProps.bGetActorReference && OutMetadata && LandscapeCollisionComponent)
+	{
+		if (FPCGMetadataAttribute<FSoftObjectPath>* ActorReferenceAttribute = OutMetadata->GetMutableTypedAttribute<FSoftObjectPath>(PCGPointDataConstants::ActorReferenceAttribute))
+		{
+			// Landscape code seems to indicate the XYtoComponentMap can be sometimes invalid, so rely on the collision map instead
+			OutMetadata->InitializeOnSet(OutPoint.MetadataEntry);
+			ActorReferenceAttribute->SetValue(OutPoint.MetadataEntry, FSoftObjectPath(LandscapeCollisionComponent->GetOwner()));
+		}
+	}
+
+	if (DataProps.bGetPhysicalMaterial && OutMetadata && LandscapeCollisionComponent)
+	{
+		if (FPCGMetadataAttribute<FSoftObjectPath>* PhysicalMaterialAttribute = OutMetadata->GetMutableTypedAttribute<FSoftObjectPath>(PCGWorldRayHitConstants::PhysicalMaterialReferenceAttribute))
+		{
+			if(UPhysicalMaterial* PhysicalMaterial = LandscapeCollisionComponent->GetPhysicalMaterial(static_cast<float>(ComponentLocalPoint.X), static_cast<float>(ComponentLocalPoint.Y), EHeightfieldSource::Complex))
+			{
+				OutMetadata->InitializeOnSet(OutPoint.MetadataEntry);
+				PhysicalMaterialAttribute->SetValue(OutPoint.MetadataEntry, FSoftObjectPath(PhysicalMaterial));
+			}
+		}
+	}
+
+	if (DataProps.bGetComponentCoordinates && OutMetadata)
+	{
+		if(FPCGMetadataAttribute<int32>* ComponentXAttribute = OutMetadata->GetMutableTypedAttribute<int32>(PCGLandscapeDataConstants::ComponentXAttribute))
+		{
+			OutMetadata->InitializeOnSet(OutPoint.MetadataEntry);
+			ComponentXAttribute->SetValue(OutPoint.MetadataEntry, ComponentMapKey.X);
+		}
+
+		if (FPCGMetadataAttribute<int32>* ComponentYAttribute = OutMetadata->GetMutableTypedAttribute<int32>(PCGLandscapeDataConstants::ComponentYAttribute))
+		{
+			OutMetadata->InitializeOnSet(OutPoint.MetadataEntry);
+			ComponentYAttribute->SetValue(OutPoint.MetadataEntry, ComponentMapKey.Y);
+		}
 	}
 
 	// Respect projection settings
@@ -242,7 +319,7 @@ const UPCGPointData* UPCGLandscapeData::CreatePointData(FPCGContext* Context, co
 		return Data;
 	}
 
-	UPCGMetadata* OutMetadata = bUseMetadata ? Data->Metadata : nullptr;
+	UPCGMetadata* OutMetadata = DataProps.bGetLayerWeights ? Data->Metadata : nullptr;
 
 	// Most proxies we gathered will have the same landscape info, we shouldn't loop multiple times
 	// on them, unless we add the box filtering - but even then, depending on the transform we could have overlaps
@@ -328,7 +405,7 @@ const UPCGPointData* UPCGLandscapeData::CreatePointData(FPCGContext* Context, co
 						const int32 PointIndex = LocalX + LocalY * (ComponentSizeQuads + 1);
 
 						FPCGPoint& Point = Points.Emplace_GetRef();
-						if (bHeightOnly)
+						if (DataProps.bGetHeightOnly)
 						{
 							LandscapeCacheEntry->GetPointHeightOnly(PointIndex, Point);
 						}
@@ -373,16 +450,15 @@ const ULandscapeInfo* UPCGLandscapeData::GetLandscapeInfo(const FVector& InPosit
 
 UPCGSpatialData* UPCGLandscapeData::CopyInternal() const
 {
-	UPCGLandscapeData* NewLandsacapeData = NewObject<UPCGLandscapeData>();
+	UPCGLandscapeData* NewLandscapeData = NewObject<UPCGLandscapeData>();
 
-	CopyBaseSurfaceData(NewLandsacapeData);
+	CopyBaseSurfaceData(NewLandscapeData);
 
-	NewLandsacapeData->Landscapes = Landscapes;
-	NewLandsacapeData->Bounds = Bounds;
-	NewLandsacapeData->bHeightOnly = bHeightOnly;
-	NewLandsacapeData->bUseMetadata = bUseMetadata;
-	NewLandsacapeData->LandscapeInfos = LandscapeInfos;
-	NewLandsacapeData->LandscapeCache = LandscapeCache;
+	NewLandscapeData->Landscapes = Landscapes;
+	NewLandscapeData->Bounds = Bounds;
+	NewLandscapeData->DataProps = DataProps;
+	NewLandscapeData->LandscapeInfos = LandscapeInfos;
+	NewLandscapeData->LandscapeCache = LandscapeCache;
 
-	return NewLandsacapeData;
+	return NewLandscapeData;
 }
diff --git a/Engine/Plugins/Experimental/PCG/Source/PCG/Private/Elements/PCGTypedGetter.cpp b/Engine/Plugins/Experimental/PCG/Source/PCG/Private/Elements/PCGTypedGetter.cpp
--- a/Engine/Plugins/Experimental/PCG/Source/PCG/Private/Elements/PCGTypedGetter.cpp
+++ b/Engine/Plugins/Experimental/PCG/Source/PCG/Private/Elements/PCGTypedGetter.cpp
@@ -34,6 +34,25 @@ UPCGGetLandscapeSettings::UPCGGetLandscapeSettings()
 	}
 }
 
+void UPCGGetLandscapeSettings::PostLoad()
+{
+	Super::PostLoad();
+
+#if WITH_EDITOR
+	if (bGetHeightOnly_DEPRECATED)
+	{
+		SamplingProperties.bGetHeightOnly = bGetHeightOnly_DEPRECATED;
+		bGetHeightOnly_DEPRECATED = false;
+	}
+
+	if (bGetLayerWeights_DEPRECATED)
+	{
+		SamplingProperties.bGetLayerWeights = bGetLayerWeights_DEPRECATED;
+		bGetLayerWeights_DEPRECATED = false;
+	}
+#endif
+}
+
 TSubclassOf<AActor> UPCGGetLandscapeSettings::GetDefaultActorSelectorClass() const
 {
 	return ALandscapeProxy::StaticClass();
@@ -103,7 +122,7 @@ void FPCGGetLandscapeDataElement::ProcessActors(FPCGContext* Context, const UPCG
 	if (!Landscapes.IsEmpty())
 	{
 		UPCGLandscapeData* LandscapeData = NewObject<UPCGLandscapeData>();
-		LandscapeData->Initialize(Landscapes, LandscapeBounds, Settings->bGetHeightOnly, Settings->bGetLayerWeights);
+		LandscapeData->Initialize(Landscapes, LandscapeBounds, Settings->SamplingProperties);
 		
 		FPCGTaggedData& TaggedData = Context->OutputData.TaggedData.Emplace_GetRef();
 		TaggedData.Data = LandscapeData;
diff --git a/Engine/Plugins/Experimental/PCG/Source/PCG/Private/PCGComponent.cpp b/Engine/Plugins/Experimental/PCG/Source/PCG/Private/PCGComponent.cpp
--- a/Engine/Plugins/Experimental/PCG/Source/PCG/Private/PCGComponent.cpp
+++ b/Engine/Plugins/Experimental/PCG/Source/PCG/Private/PCGComponent.cpp
@@ -1976,9 +1976,12 @@ FPCGDataCollection UPCGComponent::CreateActorPCGDataCollection(AActor* Actor, co
 	{
 		UPCGLandscapeData* Data = NewObject<UPCGLandscapeData>();
 		const UPCGGraph* PCGGraph = Component ? Component->GetGraph() : nullptr;
-		const bool bUseLandscapeMetadata = (!PCGGraph || PCGGraph->bLandscapeUsesMetadata);
 
-		Data->Initialize({ LandscapeActor }, PCGHelpers::GetGridBounds(Actor, Component), /*bHeightOnly=*/false, bUseLandscapeMetadata);
+		FPCGLandscapeDataProps LandscapeDataProps;
+		LandscapeDataProps.bGetHeightOnly = false;
+		LandscapeDataProps.bGetLayerWeights = (!PCGGraph || PCGGraph->bLandscapeUsesMetadata);
+
+		Data->Initialize({ LandscapeActor }, PCGHelpers::GetGridBounds(Actor, Component), LandscapeDataProps);
 
 		FPCGTaggedData& TaggedData = Collection.TaggedData.Emplace_GetRef();
 		TaggedData.Data = Data;
@@ -2196,7 +2199,12 @@ UPCGData* UPCGComponent::CreateLandscapePCGData(bool bHeightOnly)
 	// TODO: we're creating separate landscape data instances here so we can do some tweaks on it (such as storing the right target actor) but this probably should change
 	UPCGLandscapeData* LandscapeData = NewObject<UPCGLandscapeData>();
 	const UPCGGraph* PCGGraph = GetGraph();
-	LandscapeData->Initialize(Landscapes, LandscapeBounds, bHeightOnly, /*bUseMetadata=*/PCGGraph && PCGGraph->bLandscapeUsesMetadata);
+
+	FPCGLandscapeDataProps LandscapeDataProps;
+	LandscapeDataProps.bGetHeightOnly = bHeightOnly;
+	LandscapeDataProps.bGetLayerWeights = (PCGGraph && PCGGraph->bLandscapeUsesMetadata);
+
+	LandscapeData->Initialize(Landscapes, LandscapeBounds, LandscapeDataProps);
 
 	return LandscapeData;
 }
diff --git a/Engine/Plugins/Experimental/PCG/Source/PCG/Public/Data/PCGLandscapeData.h b/Engine/Plugins/Experimental/PCG/Source/PCG/Public/Data/PCGLandscapeData.h
--- a/Engine/Plugins/Experimental/PCG/Source/PCG/Public/Data/PCGLandscapeData.h
+++ b/Engine/Plugins/Experimental/PCG/Source/PCG/Public/Data/PCGLandscapeData.h
@@ -13,6 +13,32 @@ class ALandscapeProxy;
 class ULandscapeInfo;
 class UPCGLandscapeCache;
 
+USTRUCT(BlueprintType)
+struct PCG_API FPCGLandscapeDataProps
+{
+	GENERATED_BODY()
+
+	/** Controls whether the points projected on the landscape will return the normal/tangent (if false) or only the position (if true) */
+	UPROPERTY(BlueprintReadWrite, EditAnywhere, Category = Settings, meta = (PCG_Overridable))
+	bool bGetHeightOnly = false;
+
+	/** Controls whether data from landscape layers will be retrieved (turning it off is an optimization if that data is not needed) */
+	UPROPERTY(BlueprintReadWrite, EditAnywhere, Category = Settings, meta = (PCG_Overridable))
+	bool bGetLayerWeights = true;
+
+	/** Controls whether the points from this landscape will return the actor from which they originate (e.g. which Landscape Proxy) */
+	UPROPERTY(BlueprintReadWrite, EditAnywhere, Category = Settings, meta = (PCG_Overridable))
+	bool bGetActorReference = false;
+
+	/** Controls whether the points from the landscape will have their physical material added as the "PhysicalMaterial" attribute */
+	UPROPERTY(BlueprintReadWrite, EditAnywhere, Category = Settings, meta = (PCG_Overridable))
+	bool bGetPhysicalMaterial = false;
+
+	/** Controls whether the component coordinates will be added the point as attributes ('CoordinateX', 'CoordinateY') */
+	UPROPERTY(BlueprintReadWrite, EditAnywhere, Category = Settings, meta = (PCG_Overridable))
+	bool bGetComponentCoordinates = false;
+};
+
 /**
 * Landscape data access abstraction for PCG. Supports multi-landscape access, but it assumes that they are not overlapping.
 */
@@ -22,7 +48,7 @@ class PCG_API UPCGLandscapeData : public UPCGSurfaceData
 	GENERATED_BODY()
 
 public:
-	void Initialize(const TArray<TWeakObjectPtr<ALandscapeProxy>>& InLandscapes, const FBox& InBounds, bool bInHeightOnly, bool bInUseMetadata);
+	void Initialize(const TArray<TWeakObjectPtr<ALandscapeProxy>>& InLandscapes, const FBox& InBounds, const FPCGLandscapeDataProps& InDataProps);
 
 	// ~Begin UObject interface
 	virtual void PostLoad();
@@ -53,7 +79,7 @@ public:
 	UPROPERTY(BlueprintReadWrite, EditAnywhere, Category = SourceData)
 	TArray<TSoftObjectPtr<ALandscapeProxy>> Landscapes;
 
-	bool IsUsingMetadata() const { return bUseMetadata; }
+	bool IsUsingMetadata() const { return DataProps.bGetLayerWeights; }
 
 protected:
 	/** Returns the landscape info associated to the first landscape that contains the given position
@@ -65,10 +91,15 @@ protected:
 	FBox Bounds = FBox(EForceInit::ForceInit);
 
 	UPROPERTY()
-	bool bHeightOnly = false;
+	FPCGLandscapeDataProps DataProps;
+
+#if WITH_EDITORONLY_DATA
+	UPROPERTY()
+	bool bHeightOnly_DEPRECATED = false;
 
 	UPROPERTY()
-	bool bUseMetadata = true;
+	bool bUseMetadata_DEPRECATED = true;
+#endif // WITH_EDITORONLY_DATA
 
 private:
 	// Transient data
diff --git a/Engine/Plugins/Experimental/PCG/Source/PCG/Public/Elements/PCGTypedGetter.h b/Engine/Plugins/Experimental/PCG/Source/PCG/Public/Elements/PCGTypedGetter.h
--- a/Engine/Plugins/Experimental/PCG/Source/PCG/Public/Elements/PCGTypedGetter.h
+++ b/Engine/Plugins/Experimental/PCG/Source/PCG/Public/Elements/PCGTypedGetter.h
@@ -2,6 +2,7 @@
 
 #pragma once
 
+#include "Data/PCGLandscapeData.h"
 #include "Elements/PCGDataFromActor.h"
 
 #include "PCGTypedGetter.generated.h"
@@ -15,6 +16,10 @@ class PCG_API UPCGGetLandscapeSettings : public UPCGDataFromActorSettings
 public:
 	UPCGGetLandscapeSettings();
 
+	//~Begin UObject interface
+	virtual void PostLoad() override;
+	//~End UObject interface
+
 	//~Begin UPCGSettings interface
 #if WITH_EDITOR
 	virtual FName GetDefaultNodeName() const override { return FName(TEXT("GetLandscapeData")); }
@@ -36,11 +41,16 @@ public:
 	virtual TSubclassOf<AActor> GetDefaultActorSelectorClass() const override;
 	//~End UPCGDataFromActorSettings
 
-	UPROPERTY(BlueprintReadWrite, EditAnywhere, Category = Settings, meta = (PCG_Overridable))
-	bool bGetHeightOnly = false;
+	UPROPERTY(BlueprintReadWrite, EditAnywhere, Category = Settings, meta = (PCG_Overridable, ShowOnlyInnerProperties))
+	FPCGLandscapeDataProps SamplingProperties;
+
+#if WITH_EDITORONLY_DATA
+	UPROPERTY()
+	bool bGetHeightOnly_DEPRECATED = false;
 
-	UPROPERTY(BlueprintReadWrite, EditAnywhere, Category = Settings, meta = (PCG_Overridable))
-	bool bGetLayerWeights = true;
+	UPROPERTY()
+	bool bGetLayerWeights_DEPRECATED = true;
+#endif // WITH_EDITORONLY_DATA
 };
 
 class FPCGGetLandscapeDataElement : public FPCGDataFromActorElement
