[PCG] Hotfix to Surface Sampler to respect invalid bounding shape bounds

Fix improper bounds checking
Add initialization results to Time Slice context

#rnx

diff --git a/Engine/Plugins/Experimental/PCG/Source/PCG/Private/Elements/PCGSurfaceSampler.cpp b/Engine/Plugins/Experimental/PCG/Source/PCG/Private/Elements/PCGSurfaceSampler.cpp
--- a/Engine/Plugins/Experimental/PCG/Source/PCG/Private/Elements/PCGSurfaceSampler.cpp
+++ b/Engine/Plugins/Experimental/PCG/Source/PCG/Private/Elements/PCGSurfaceSampler.cpp
@@ -5,6 +5,7 @@
 #include "PCGCommon.h"
 #include "PCGComponent.h"
 #include "PCGCustomVersion.h"
+#include "PCGData.h"
 #include "PCGEdge.h"
 #include "PCGGraph.h"
 #include "PCGPin.h"
@@ -324,7 +325,7 @@ namespace PCGSurfaceSamplerHelpers
 		return BoundingShape;
 	}
 
-	bool InitializePerExecutionData(ContextType* Context, ExecStateType& OutState)
+	EPCGTimeSliceInitResult InitializePerExecutionData(ContextType* Context, ExecStateType& OutState)
 	{
 		TRACE_CPUPROFILER_EVENT_SCOPE(FPCGSurfaceSamplerElement::InitializePerExecutionData);
 
@@ -337,7 +338,7 @@ namespace PCGSurfaceSamplerHelpers
 		if (SurfaceInputs.IsEmpty())
 		{
 			PCGE_LOG_C(Warning, GraphAndLog, Context, LOCTEXT("NoSurfaceInputs", "No Surface input was provided."));
-			return false;
+			return EPCGTimeSliceInitResult::AbortExecution;
 		}
 
 		// Early out on invalid settings
@@ -346,24 +347,31 @@ namespace PCGSurfaceSamplerHelpers
 		if (PointExtents.X <= 0 || PointExtents.Y <= 0)
 		{
 			PCGE_LOG_C(Warning, GraphAndLog, Context, LOCTEXT("SkippedInvalidPointExtents", "Skipped - Invalid point extents"));
-			return false;
+			return EPCGTimeSliceInitResult::AbortExecution;
 		}
 
 		// Grab the Bounding Shape input if there is one.
 		const TArray<FPCGTaggedData> BoundingShapeInputs = Context->InputData.GetInputsByPin(PCGSurfaceSamplerConstants::BoundingShapeLabel);
+		FBox BoundingShapeBounds(EForceInit::ForceInit);
 		if (!Settings->bUnbounded)
 		{
 			bool bUnionWasCreated;
-			const UPCGSpatialData* BoundingShape = FindBoundingShape(Context, bUnionWasCreated);
-			if (BoundingShape && bUnionWasCreated)
+			OutState.BoundingShape = FindBoundingShape(Context, bUnionWasCreated);
+			if (OutState.BoundingShape && bUnionWasCreated)
 			{
 				// Must cast away const to root the union object
-				Context->RootAndTrackObject(const_cast<UPCGSpatialData*>(BoundingShape));
+				Context->RootAndTrackObject(const_cast<UPCGSpatialData*>(OutState.BoundingShape));
+			}
+
+			if (OutState.BoundingShape)
+			{
+				BoundingShapeBounds = OutState.BoundingShape->GetBounds();
 			}
 
-			if (BoundingShape)
+			if (!BoundingShapeBounds.IsValid)
 			{
-				OutState.BoundingShapeBounds = BoundingShape->GetBounds();
+				// The bounding shape bounds is invalid, such as an empty intersection, so no operation will need to be performed.
+				return EPCGTimeSliceInitResult::NoOperation;
 			}
 		}
 		else if (BoundingShapeInputs.Num() > 0)
@@ -390,7 +398,7 @@ namespace PCGSurfaceSamplerHelpers
 		if (GeneratingShapes.IsEmpty())
 		{
 			PCGE_LOG_C(Warning, GraphAndLog, Context, LOCTEXT("NoSurfaceFound", "No surfaces found from which to generate"));
-			return false;
+			return EPCGTimeSliceInitResult::AbortExecution;
 		}
 
 		// Iterate over the generating shapes to compute the bounds
@@ -408,14 +416,14 @@ namespace PCGSurfaceSamplerHelpers
 			{
 				InputBounds = GeneratingShape->GetBounds();
 
-				if (OutState.BoundingShapeBounds.IsValid)
+				if (BoundingShapeBounds.IsValid)
 				{
-					InputBounds = PCGHelpers::OverlapBounds(InputBounds, OutState.BoundingShapeBounds);
+					InputBounds = PCGHelpers::OverlapBounds(InputBounds, BoundingShapeBounds);
 				}
 			}
 			else
 			{
-				InputBounds = OutState.BoundingShapeBounds;
+				InputBounds = BoundingShapeBounds;
 			}
 
 			if (!InputBounds.IsValid || !OutState.Settings.Initialize(Settings, Context, InputBounds))
@@ -445,7 +453,7 @@ namespace PCGSurfaceSamplerHelpers
 			PCGE_LOG_C(Warning, GraphAndLog, Context, LOCTEXT("InvalidSurface", "No surfaces from which to generate are valid"));
 		}
 
-		return true;
+		return EPCGTimeSliceInitResult::Success;
 	}
 }
 
@@ -456,7 +464,7 @@ bool FPCGSurfaceSamplerElement::PrepareDataInternal(FPCGContext* InContext) cons
 	check(Context);
 
 	// Initialize the per-execution state data that won't change over the duration of the time slicing
-	if (!Context->InitializePerExecutionState(PCGSurfaceSamplerHelpers::InitializePerExecutionData))
+	if (Context->InitializePerExecutionState(PCGSurfaceSamplerHelpers::InitializePerExecutionData) == EPCGTimeSliceInitResult::AbortExecution)
 	{
 		PCGE_LOG(Warning, GraphAndLog, LOCTEXT("CouldNotInitializeExecutionState", "Could not initialize per-execution timeslice state data"));
 		return true;
@@ -466,17 +474,17 @@ bool FPCGSurfaceSamplerElement::PrepareDataInternal(FPCGContext* InContext) cons
 	TArray<FPCGTaggedData>& Outputs = Context->OutputData.TaggedData;
 
 	// Initialize the per-iteration data, using the generating shapes as the source of iteration
-	const bool bInitSuccess = Context->InitializePerIterationStates(GeneratingShapes.Num(),
-		[&GeneratingShapes, &Outputs](IterStateType& OutState, const uint32 IterationIndex)->bool
+	Context->InitializePerIterationStates(GeneratingShapes.Num(),
+		[&GeneratingShapes, &Outputs](IterStateType& OutState, const uint32 IterationIndex)
 		{
 			OutState.OutputPoints = NewObject<UPCGPointData>();
 			OutState.OutputPoints->InitializeFromData(GeneratingShapes[IterationIndex]);
 			// Assigning this here prevents the need to root
 			Outputs[IterationIndex].Data = OutState.OutputPoints;
-			return true;
+			return EPCGTimeSliceInitResult::Success;
 		});
 
-	if (!bInitSuccess)
+	if (!Context->DataIsPreparedForExecution())
 	{
 		PCGE_LOG(Warning, GraphAndLog, LOCTEXT("CouldNotInitializeIterationState", "Could not initialize per-iteration timeslice state data"));
 	}
@@ -490,15 +498,43 @@ bool FPCGSurfaceSamplerElement::ExecuteInternal(FPCGContext* InContext) const
 	ContextType* TimeSlicedContext = static_cast<ContextType*>(InContext);
 	check(TimeSlicedContext);
 
-	// Prepare data failed, no need to execute
-	if (!TimeSlicedContext->DataIsPrepared())
+	// Prepare data failed, no need to execute. Return an empty output
+	if (!TimeSlicedContext->DataIsPreparedForExecution())
+	{
+		TimeSlicedContext->OutputData.TaggedData.Empty();
+
+		return true;
+	}
+
+	// The execution would have resulted in an empty set of points for all iterations
+	if (TimeSlicedContext->GetExecutionStateResult() == EPCGTimeSliceInitResult::NoOperation)
 	{
+		for (FPCGTaggedData& Input : TimeSlicedContext->InputData.GetInputs())
+		{
+			FPCGTaggedData& Output = TimeSlicedContext->OutputData.TaggedData.Emplace_GetRef();
+			UPCGPointData* PointData = NewObject<UPCGPointData>();
+			PointData->InitializeFromData(Cast<UPCGSpatialData>(Input.Data));
+			Output.Data = PointData;
+		}
+
 		return true;
 	}
 
 	// The context will iterate over per-iteration states and execute the lambda until it returns true
 	return ExecuteSlice(TimeSlicedContext, [](ContextType* Context, const ExecStateType& ExecutionState, const IterStateType& IterationState, const uint32 IterationIndex)->bool
 	{
+		// This iteration resulted in an early out for no sampling operation. Early out with empty point data.
+		const EPCGTimeSliceInitResult InitResult = Context->GetIterationStateResult(IterationIndex);
+		if (InitResult == EPCGTimeSliceInitResult::NoOperation)
+		{
+			Context->OutputData.TaggedData[IterationIndex].Data = NewObject<UPCGPointData>();
+
+			return true;
+		}
+
+		// It should be guaranteed to be a success at this point
+		check(InitResult == EPCGTimeSliceInitResult::Success);
+
 		// Run the execution until the time slice is finished
 		const bool bAsyncDone = PCGSurfaceSampler::SampleSurface(Context, ExecutionState.Settings, ExecutionState.GeneratingShapes[IterationIndex], ExecutionState.BoundingShape, IterationState.OutputPoints, Context->TimeSliceIsEnabled());
 
diff --git a/Engine/Plugins/Experimental/PCG/Source/PCG/Public/Elements/PCGSurfaceSampler.h b/Engine/Plugins/Experimental/PCG/Source/PCG/Public/Elements/PCGSurfaceSampler.h
--- a/Engine/Plugins/Experimental/PCG/Source/PCG/Public/Elements/PCGSurfaceSampler.h
+++ b/Engine/Plugins/Experimental/PCG/Source/PCG/Public/Elements/PCGSurfaceSampler.h
@@ -52,7 +52,6 @@ namespace PCGSurfaceSampler
 	{
 		FSurfaceSamplerExecutionSettings Settings;
 
-		FBox BoundingShapeBounds = FBox(EForceInit::ForceInit);
 		const UPCGSpatialData* BoundingShape = nullptr;
 		TArray<const UPCGSpatialData*> GeneratingShapes;
 	};
diff --git a/Engine/Plugins/Experimental/PCG/Source/PCG/Public/Elements/PCGTimeSlicedElementBase.h b/Engine/Plugins/Experimental/PCG/Source/PCG/Public/Elements/PCGTimeSlicedElementBase.h
--- a/Engine/Plugins/Experimental/PCG/Source/PCG/Public/Elements/PCGTimeSlicedElementBase.h
+++ b/Engine/Plugins/Experimental/PCG/Source/PCG/Public/Elements/PCGTimeSlicedElementBase.h
@@ -24,6 +24,15 @@
 // Forward declaration for friending
 template <typename PerExecutionStateT, typename PerIterationStateT> class TPCGTimeSlicedElementBase;
 
+/** The result of initializing the Per Execution or Per Iteration Time Slice States */
+enum class EPCGTimeSliceInitResult : uint8
+{
+	Uninitialized = 0, // Initialization has not occurred yet
+	Success,           // Initialization was a success
+	NoOperation,       // Initialization was a success, but resulted in no operation
+	AbortExecution     // Initialization was a failure. Should abort.
+};
+
 namespace PCGTimeSlice
 {
 	struct FEmptyStruct {};
@@ -45,15 +54,28 @@ struct TPCGTimeSlicedContext : public FPCGContext
 	/** Retrieves the number of times this context was executed */
 	uint32 GetExecutionCount() const { return ExecutionCount; }
 
-	using InitSignature = bool(TPCGTimeSlicedContext* Context, PerExecutionStateT& OutState);
+	using InitExecSignature = EPCGTimeSliceInitResult(TPCGTimeSlicedContext* Context, PerExecutionStateT& OutState);
+
+	/** Initializes per execution state data if required. */
+	EPCGTimeSliceInitResult InitializePerExecutionState(TFunctionRef<InitExecSignature> InitFunc = []{ return true; });
 
-	/** Initializes per execution state data if required. Returns true if data is completely initialized properly and false if a problem should result in early termination. */
-	bool InitializePerExecutionState(TFunctionRef<InitSignature> InitFunc = []{ return true; });
+	using InitIterSignature = EPCGTimeSliceInitResult(PerIterationStateT& OutState, const uint32 IterationIndex);
 
-	using IterSignature = bool(PerIterationStateT& OutState, const uint32 IterationIndex);
+	/** Initializes per execution state data if required. An array will be created with a state element for every execution iteration in the context. Returns the Init Result for each iteration's initialization. */
+	const TArray<EPCGTimeSliceInitResult>& InitializePerIterationStates(int32 NumIterations = 1, TFunctionRef<InitIterSignature> IterFunc = []{ return true; });
 
-	/** Initializes per execution state data if required. An array will be created with a state element for every execution iteration in the context. */
-	bool InitializePerIterationStates(int32 NumIterations = 1, TFunctionRef<IterSignature> IterFunc = []{ return true; });
+	/** Will return the result of the attempt to initialize the per execution state */
+	EPCGTimeSliceInitResult GetExecutionStateResult() const { return ExecutionStateResult; }
+
+	/** Will return the result of the attempt to initialize a specific iteration, by Index */
+	EPCGTimeSliceInitResult GetIterationStateResult(const int32 Index)
+	{
+		check(Index >= 0 && Index < PerIterationStateResultArray.Num());
+
+		return PerIterationStateResultArray[Index];
+	}
+
+	const TArray<EPCGTimeSliceInitResult>& GetPerIterationStateInitResultArray() { return PerIterationStateResultArray; }
 
 	PerExecutionStateT& GetPerExecutionState() { return PerExecutionStateData; }
 	const PerExecutionStateT& GetPerExecutionState() const { return PerExecutionStateData; }
@@ -61,8 +83,26 @@ struct TPCGTimeSlicedContext : public FPCGContext
 	TArray<PerIterationStateT>& GetPerIterationStateArray() { return PerIterationStateArray; }
 	const TArray<PerIterationStateT>& GetPerIterationStateArray() const { return PerIterationStateArray; }
 
-	/** Returns true if both the execution state and iteration state were fully initialized. Can be used to bypass any further initialization. */
-	bool DataIsPrepared() const { return bPerExecutionStateIsInitialized && bPerIterationStateIsInitialized; }
+	/** Returns true if both the execution state and iteration state were fully initialized and no problems arose that should abort the execution. Can be used to bypass any further initialization. */
+	bool DataIsPreparedForExecution() const
+	{
+		if (ExecutionStateResult == EPCGTimeSliceInitResult::Uninitialized ||
+			ExecutionStateResult == EPCGTimeSliceInitResult::AbortExecution ||
+			!bPerIterationStateIsInitialized)
+		{
+			return false;
+		}
+
+		for (const EPCGTimeSliceInitResult Result : PerIterationStateResultArray)
+		{
+			if (Result == EPCGTimeSliceInitResult::Uninitialized || Result == EPCGTimeSliceInitResult::AbortExecution)
+			{
+				return false;
+			}
+		}
+
+		return true;
+	}
 
 	/** Fire and forget function to root a UObject for the duration of the Context, and then mark for garbage collection at context lifecycle end. */
 	void RootAndTrackObject(UObject* Object);
@@ -76,12 +116,21 @@ private:
 	/** The index of which iteration being processed. Ie. If a volume sampler has two volume inputs, it will be processed twice. */
 	int32 IterationIndex = 0;
 
+	/** True if time slicing is enabled for this context */
 	bool bTimeSliceIsEnabled = true;
-	bool bPerExecutionStateIsInitialized = false;
+
+	/** The result of initializing the per execution state */
+	EPCGTimeSliceInitResult ExecutionStateResult = EPCGTimeSliceInitResult::Uninitialized;
+
+	/** True if an attempt has been made to initialize the per iteration states */
 	bool bPerIterationStateIsInitialized = false;
 
+	/** Tracks the results of each attempt to initialize an iteration state */
+	TArray<EPCGTimeSliceInitResult> PerIterationStateResultArray;
+
 	/** The state of the timesliced context that won't change each iteration. Ie. the node settings. */
 	PerExecutionStateT PerExecutionStateData;
+
 	/** An array of the various states of timesliced context that will change per iteration. Ie. generating shape */
 	TArray<PerIterationStateT> PerIterationStateArray;
 
@@ -133,51 +182,53 @@ TPCGTimeSlicedContext<PerExecutionStateT, PerIterationStateT>::~TPCGTimeSlicedCo
 }
 
 template <typename PerExecutionStateT, typename PerIterationStateT>
-bool TPCGTimeSlicedContext<PerExecutionStateT, PerIterationStateT>::InitializePerExecutionState(TFunctionRef<InitSignature> InitFunc)
+EPCGTimeSliceInitResult TPCGTimeSlicedContext<PerExecutionStateT, PerIterationStateT>::InitializePerExecutionState(TFunctionRef<InitExecSignature> InitFunc)
 {
 	// Should only ever be initialized once, so if its called again, ignore it. This allows flexibility for the call to be somewhere that might be invoked numerous times
-	if (!bPerExecutionStateIsInitialized)
+	if (ExecutionStateResult == EPCGTimeSliceInitResult::Uninitialized)
 	{
-		bPerExecutionStateIsInitialized = InitFunc(this, PerExecutionStateData);
+		ExecutionStateResult = InitFunc(this, PerExecutionStateData);
 	}
 
-	return bPerExecutionStateIsInitialized;
+	return ExecutionStateResult;
 }
 
 template <typename PerExecutionStateT, typename PerIterationStateT>
-bool TPCGTimeSlicedContext<PerExecutionStateT, PerIterationStateT>::InitializePerIterationStates(int32 NumIterations, TFunctionRef<IterSignature> IterFunc)
+const TArray<EPCGTimeSliceInitResult>& TPCGTimeSlicedContext<PerExecutionStateT, PerIterationStateT>::InitializePerIterationStates(int32 NumIterations, TFunctionRef<InitIterSignature> IterFunc)
 {
 	// Same as InitializePerExecutionState. Should only ever be initialized once, so if its called again, ignore it.
 	if (bPerIterationStateIsInitialized)
 	{
-		return true;
+		return PerIterationStateResultArray;
 	}
 
 	bPerIterationStateIsInitialized = true;
 
+	// Should be guaranteed to be uninitialized
+	check(PerIterationStateResultArray.IsEmpty() && PerIterationStateArray.IsEmpty());
+	
 	// An empty iteration state is still valid
 	if (NumIterations < 1)
 	{
-		return true;
+		return PerIterationStateResultArray;
 	}
 
-	bool bAnySucceeded{false};
+	PerIterationStateResultArray.AddUninitialized(NumIterations);
 	PerIterationStateArray.Reserve(NumIterations);
 
 	for (int32 I = 0; I < NumIterations; ++I)
 	{
-		const bool bCurrentSucceeded = IterFunc(PerIterationStateArray.Emplace_GetRef(), I);
-		bAnySucceeded |= bCurrentSucceeded;
+		PerIterationStateResultArray[I] = IterFunc(PerIterationStateArray.Emplace_GetRef(), I);
 
-		if (!bCurrentSucceeded)
+		if (PerIterationStateResultArray[I] == EPCGTimeSliceInitResult::AbortExecution)
 		{
 			// If it fails, remove the new state from the array and continue
 			PerIterationStateArray.RemoveAt(PerIterationStateArray.Num() - 1);
 		}
 	}
 
-	// Returns true if any succeeded
-	return bAnySucceeded;
+	// Returns a complete array of the results
+	return PerIterationStateResultArray;
 }
 
 template <typename PerExecutionStateT, typename PerIterationStateT>
@@ -194,7 +245,7 @@ bool TPCGTimeSlicedElementBase<PerExecutionStateT, PerIterationStateT>::ExecuteS
 	++Context->ExecutionCount;
 
 	// The user is responsible to check for this before execution, but just in case
-	if (!ensureMsgf(Context->DataIsPrepared(), TEXT("State data was not properly initialized.")) || Context->PerIterationStateArray.IsEmpty())
+	if (!ensureMsgf(Context->DataIsPreparedForExecution(), TEXT("State data was not properly initialized.")) || Context->PerIterationStateArray.IsEmpty())
 	{
 		return true;
 	}
