[PCG] Replace "Can Be Multithreaded" option with clearer "Requires Game Thread"

diff --git a/Engine/Plugins/Experimental/PCG/Source/PCG/Private/Elements/PCGExecuteBlueprint.cpp b/Engine/Plugins/Experimental/PCG/Source/PCG/Private/Elements/PCGExecuteBlueprint.cpp
--- a/Engine/Plugins/Experimental/PCG/Source/PCG/Private/Elements/PCGExecuteBlueprint.cpp
+++ b/Engine/Plugins/Experimental/PCG/Source/PCG/Private/Elements/PCGExecuteBlueprint.cpp
@@ -78,6 +78,12 @@ void UPCGBlueprintElement::PostLoad()
 			CustomInputPins.RemoveAt(i);
 		}
 	}
+
+	if (bCanBeMultithreaded_DEPRECATED)
+	{
+		bRequiresGameThread = false;
+	}
+	bCanBeMultithreaded_DEPRECATED = false;
 #endif
 }
 
@@ -356,7 +362,10 @@ void UPCGBlueprintSettings::PostLoad()
 		BlueprintElementInstance->ConditionalPostLoad();
 		BlueprintElementInstance->SetFlags(RF_Transactional);
 #if WITH_EDITOR
-		BlueprintElementInstance->bCanBeMultithreaded |= bCanBeMultithreaded_DEPRECATED;
+		if (bCanBeMultithreaded_DEPRECATED)
+		{
+			BlueprintElementInstance->bRequiresGameThread = false;
+		}
 #endif
 	}
 
@@ -943,7 +952,7 @@ bool FPCGExecuteBlueprintElement::CanExecuteOnlyOnMainThread(FPCGContext* Contex
 	const UPCGBlueprintSettings* BPSettings = Context->GetInputSettings<UPCGBlueprintSettings>();
 	if (BPSettings && BPSettings->BlueprintElementInstance)
 	{
-		return !BPSettings->BlueprintElementInstance->bCanBeMultithreaded;
+		return BPSettings->BlueprintElementInstance->bRequiresGameThread;
 	}
 	else
 	{
diff --git a/Engine/Plugins/Experimental/PCG/Source/PCG/Public/Elements/PCGExecuteBlueprint.h b/Engine/Plugins/Experimental/PCG/Source/PCG/Public/Elements/PCGExecuteBlueprint.h
--- a/Engine/Plugins/Experimental/PCG/Source/PCG/Public/Elements/PCGExecuteBlueprint.h
+++ b/Engine/Plugins/Experimental/PCG/Source/PCG/Public/Elements/PCGExecuteBlueprint.h
@@ -162,9 +162,13 @@ public:
 
 	/** Controls whether this node execution can be run from a non-game thread. This is not related to the Loop functions provided/implemented in this class, which should always run on any thread. */
 	UPROPERTY(BlueprintReadWrite, EditDefaultsOnly, Category = Settings)
-	bool bCanBeMultithreaded = false;
+	bool bRequiresGameThread = true;
 
 #if WITH_EDITORONLY_DATA
+	/** This property is deprecated, please use 'Requires Game Thread' instead. */
+	UPROPERTY()
+	bool bCanBeMultithreaded_DEPRECATED = false;
+
 	UPROPERTY(BlueprintGetter=CustomInputLabels, Category = "Settings|Input & Output", meta = (DeprecatedProperty, DeprecatedMessage = "Input Pin Labels are deprecated - use Input Labels instead."))
 	TSet<FName> InputPinLabels_DEPRECATED;
 
