[PCG] Removed GraphAuthoringTestHelper node from the library since it's for testing purpose only

#rnx

diff --git a/Engine/Plugins/Experimental/PCG/Source/PCG/Private/Tests/GraphAuthoring/PCGGraphAuthoringTestHelperSettings.cpp b/Engine/Plugins/Experimental/PCG/Source/PCG/Private/Tests/GraphAuthoring/PCGGraphAuthoringTestHelperSettings.cpp
--- a/Engine/Plugins/Experimental/PCG/Source/PCG/Private/Tests/GraphAuthoring/PCGGraphAuthoringTestHelperSettings.cpp
+++ b/Engine/Plugins/Experimental/PCG/Source/PCG/Private/Tests/GraphAuthoring/PCGGraphAuthoringTestHelperSettings.cpp
@@ -2,6 +2,13 @@
 
 #include "Tests/GraphAuthoring/PCGGraphAuthoringTestHelperSettings.h"
 
+UPCGGraphAuthoringTestHelperSettings::UPCGGraphAuthoringTestHelperSettings()
+{
+#if WITH_EDITORONLY_DATA
+	bExposeToLibrary = false;
+#endif
+}
+
 TArray<FPCGPinProperties> UPCGGraphAuthoringTestHelperSettings::InputPinProperties() const
 {
 	TArray<FPCGPinProperties> PinProperties;
diff --git a/Engine/Plugins/Experimental/PCG/Source/PCG/Private/Tests/GraphAuthoring/PCGGraphAuthoringTestHelperSettings.h b/Engine/Plugins/Experimental/PCG/Source/PCG/Private/Tests/GraphAuthoring/PCGGraphAuthoringTestHelperSettings.h
--- a/Engine/Plugins/Experimental/PCG/Source/PCG/Private/Tests/GraphAuthoring/PCGGraphAuthoringTestHelperSettings.h
+++ b/Engine/Plugins/Experimental/PCG/Source/PCG/Private/Tests/GraphAuthoring/PCGGraphAuthoringTestHelperSettings.h
@@ -11,6 +11,8 @@ class UPCGGraphAuthoringTestHelperSettings : public UPCGSettings
 {
 	GENERATED_BODY()
 
+	UPCGGraphAuthoringTestHelperSettings();
+	
 protected:
 	virtual TArray<FPCGPinProperties> InputPinProperties() const override;
 	virtual TArray<FPCGPinProperties> OutputPinProperties() const override;
@@ -19,5 +21,5 @@ protected:
 	virtual FPCGElementPtr CreateElement() const override { return nullptr; }
 
 public:
-	EPCGDataType PinType;
+	EPCGDataType PinType = EPCGDataType::None;
 };
