[PCG] Tracking improvements

* Make all Self/Original keys to track as "not-cull"
* Add intersection test in AlwaysTrack since we can have keys that are always tracked but culled for some settings.
* Intersection test also need to include old bounds, as for the culled version.
* Since caching is not per-component, only do the dirty once (using original component) to clear the cache.
* Making sure we only track either in cull mode, or always track. Never both.
* Now track actors on load/unload, should fix load order issues.
* Re-enable IsCacheable on GetActorData, with a new dependency CRC to make it work in partitioned.
* Also the seed was not refreshing the graph.

#rnx

diff --git a/Engine/Plugins/Experimental/PCG/Source/PCG/Private/Elements/PCGDataFromActor.cpp b/Engine/Plugins/Experimental/PCG/Source/PCG/Private/Elements/PCGDataFromActor.cpp
--- a/Engine/Plugins/Experimental/PCG/Source/PCG/Private/Elements/PCGDataFromActor.cpp
+++ b/Engine/Plugins/Experimental/PCG/Source/PCG/Private/Elements/PCGDataFromActor.cpp
@@ -474,4 +474,27 @@ void FPCGDataFromActorElement::ProcessActor(FPCGContext* Context, const UPCGData
 	}
 }
 
+void FPCGDataFromActorElement::GetDependenciesCrc(const FPCGDataCollection& InInput, const UPCGSettings* InSettings, UPCGComponent* InComponent, FPCGCrc& OutCrc) const
+{
+	FPCGCrc Crc;
+	IPCGElement::GetDependenciesCrc(InInput, InSettings, InComponent, Crc);
+
+	// If we track self or original, we are dependant on the actor data
+	if (const UPCGDataFromActorSettings* Settings = Cast<const UPCGDataFromActorSettings>(InSettings))
+	{
+		if (InComponent && (Settings->ActorSelector.ActorFilter == EPCGActorFilter::Self || Settings->ActorSelector.ActorFilter == EPCGActorFilter::Original))
+		{
+			UPCGComponent* ComponentToCheck = (Settings->ActorSelector.ActorFilter == EPCGActorFilter::Original) ? InComponent->GetOriginalComponent() : InComponent;
+			const UPCGData* ActorData = ComponentToCheck ? ComponentToCheck->GetActorPCGData() : nullptr;
+
+			if (ActorData)
+			{
+				Crc.Combine(ActorData->GetOrComputeCrc(/*bFullDataCrc=*/false));
+			}
+		}
+	}
+
+	OutCrc = Crc;
+}
+
 #undef LOCTEXT_NAMESPACE
diff --git a/Engine/Plugins/Experimental/PCG/Source/PCG/Private/PCGActorAndComponentMapping.cpp b/Engine/Plugins/Experimental/PCG/Source/PCG/Private/PCGActorAndComponentMapping.cpp
--- a/Engine/Plugins/Experimental/PCG/Source/PCG/Private/PCGActorAndComponentMapping.cpp
+++ b/Engine/Plugins/Experimental/PCG/Source/PCG/Private/PCGActorAndComponentMapping.cpp
@@ -788,13 +788,28 @@ void FPCGActorAndComponentMapping::UpdateTracking(UPCGComponent* InComponent, bo
 				continue;
 			}
 
+			// Making sure that we only have the component in one map.
+			auto RemoveFromMap = [InComponent, Actor](TMap<TObjectKey<AActor>, TSet<UPCGComponent*>>& InMap)
+			{
+				if (TSet<UPCGComponent*>* Components = InMap.Find(Actor))
+				{
+					Components->Remove(InComponent);
+					if (Components->IsEmpty())
+					{
+						InMap.Remove(Actor);
+					}
+				}
+			};
+
 			if (bShouldCull)
 			{
 				CulledTrackedActorsToComponentsMap.FindOrAdd(Actor).Add(InComponent);
+				RemoveFromMap(AlwaysTrackedActorsToComponentsMap);
 			}
 			else
 			{
 				AlwaysTrackedActorsToComponentsMap.FindOrAdd(Actor).Add(InComponent);
+				RemoveFromMap(CulledTrackedActorsToComponentsMap);
 			}
 
 			RegisterActor(Actor);
@@ -942,6 +957,13 @@ void FPCGActorAndComponentMapping::RegisterTrackingCallbacks()
 	GEngine->OnLevelActorDeleted().AddRaw(this, &FPCGActorAndComponentMapping::OnActorDeleted);
 	FCoreUObjectDelegates::OnObjectPropertyChanged.AddRaw(this, &FPCGActorAndComponentMapping::OnObjectPropertyChanged);
 	FCoreUObjectDelegates::OnPreObjectPropertyChanged.AddRaw(this, &FPCGActorAndComponentMapping::OnPreObjectPropertyChanged);
+
+	UWorld* World = PCGSubsystem ? PCGSubsystem->GetWorld() : nullptr;
+	if (IsValid(World) && World->PersistentLevel)
+	{
+		World->PersistentLevel->OnLoadedActorAddedToLevelEvent.AddRaw(this, &FPCGActorAndComponentMapping::OnActorLoaded);
+		World->PersistentLevel->OnLoadedActorRemovedFromLevelEvent.AddRaw(this, &FPCGActorAndComponentMapping::OnActorUnloaded);
+	}
 }
 
 void FPCGActorAndComponentMapping::TeardownTrackingCallbacks()
@@ -951,6 +973,13 @@ void FPCGActorAndComponentMapping::TeardownTrackingCallbacks()
 	GEngine->OnLevelActorDeleted().RemoveAll(this);
 	FCoreUObjectDelegates::OnObjectPropertyChanged.RemoveAll(this);
 	FCoreUObjectDelegates::OnPreObjectPropertyChanged.RemoveAll(this);
+
+	UWorld* World = PCGSubsystem ? PCGSubsystem->GetWorld() : nullptr;
+	if (IsValid(World) && World->PersistentLevel)
+	{
+		World->PersistentLevel->OnLoadedActorAddedToLevelEvent.RemoveAll(this);
+		World->PersistentLevel->OnLoadedActorRemovedFromLevelEvent.RemoveAll(this);
+	}
 }
 
 void FPCGActorAndComponentMapping::AddDelayedActors()
@@ -986,6 +1015,18 @@ void FPCGActorAndComponentMapping::AddDelayedActors()
 	DelayedAddedActors = MoveTemp(StillDelayedActors);
 }
 
+void FPCGActorAndComponentMapping::OnActorLoaded(AActor& InActor)
+{
+	// We have to make sure to not create a infinite loop
+	if (InActor.IsA<APCGWorldActor>() || !PCGSubsystem || InActor.GetWorld() != PCGSubsystem->GetWorld())
+	{
+		return;
+	}
+
+	// Loaded actors should not dirty.
+	OnActorAdded_Internal(&InActor, /*bShouldDirty=*/ false, /*LevelInstanceDepth=*/ 0);
+}
+
 void FPCGActorAndComponentMapping::OnActorAdded(AActor* InActor)
 {
 	// We have to make sure to not create a infinite loop
@@ -1195,6 +1236,11 @@ bool FPCGActorAndComponentMapping::UnregisterActor(AActor* InActor)
 	}
 }
 
+void FPCGActorAndComponentMapping::OnActorUnloaded(AActor& InActor)
+{
+	OnActorDeleted(&InActor);
+}
+
 void FPCGActorAndComponentMapping::OnActorDeleted(AActor* InActor)
 {
 	if (!InActor || !PCGSubsystem || InActor->GetWorld() != PCGSubsystem->GetWorld())
@@ -1427,6 +1473,16 @@ void FPCGActorAndComponentMapping::OnActorChanged(AActor* InActor, bool bInHasMo
 
 	check(InActor);
 	ensure(!PCGSubsystem || InActor->GetWorld() == PCGSubsystem->GetWorld());
+
+	TSet<UPCGComponent*>* CulledTrackedComponents = CulledTrackedActorsToComponentsMap.Find(InActor);
+	TSet<UPCGComponent*>* AlwaysTrackedComponents = AlwaysTrackedActorsToComponentsMap.Find(InActor);
+
+	// If this actor is not tracked, just early out
+	if (!CulledTrackedComponents && !AlwaysTrackedComponents)
+	{
+		return;
+	}
+
 	TSet<UPCGComponent*> DirtyComponents;
 
 	EPCGComponentDirtyFlag DirtyFlag = EPCGComponentDirtyFlag::Actor;
@@ -1438,15 +1494,24 @@ void FPCGActorAndComponentMapping::OnActorChanged(AActor* InActor, bool bInHasMo
 	// Check if we have a change of tag too
 	TSet<FName> RemovedTags = TempTrackedActorTags.Difference(TSet<FName>(InActor->Tags));
 
-	if (TSet<UPCGComponent*>* CulledTrackedComponents = CulledTrackedActorsToComponentsMap.Find(InActor))
+	// We need the actor bounds to know if we are intersecting, for the culled settings to be dirtied.
+	const FBox ActorBounds = PCGActorAndComponentMapping::GetActorBounds(InActor);
+	// We also need the OldActorBounds if it has moved, and is different than the new bounds.
+	FBox* OldActorBoundsPtr = bInHasMoved ? TrackedActorToPositionMap.Find(InActor) : nullptr;
+	if (OldActorBoundsPtr && OldActorBoundsPtr->Equals(ActorBounds))
 	{
-		// Not const, since it will be updated with old actor bounds
-		FBox ActorBounds = PCGActorAndComponentMapping::GetActorBounds(InActor);
+		OldActorBoundsPtr = nullptr;
+	}
+
+	if (CulledTrackedComponents)
+	{
+		// Keep a pointer to the current bounds we are looking at for the partition case. Set it to Actorbounds first, then will be set to old bounds if necessary.
+		const FBox* CurrentActorBoundsPtr = &ActorBounds;
 
 		// Then do an octree find to get all components that intersect with this actor.
 		// If the actor has moved, we also need to find components that intersected with it before
 		// We first do it for non-partitioned, then we do it for partitioned
-		auto UpdateNonPartitioned = [&DirtyComponents, InActor, CulledTrackedComponents, &RemovedTags, DirtyFlag, InOriginatingChangeObject](const FPCGComponentRef& ComponentRef) -> void
+		auto UpdateNonPartitioned = [this, &DirtyComponents, InActor, CulledTrackedComponents, &RemovedTags, DirtyFlag, InOriginatingChangeObject](const FPCGComponentRef& ComponentRef) -> void
 		{
 			TRACE_CPUPROFILER_EVENT_SCOPE(FPCGActorAndComponentMapping::OnActorChanged::UpdateNonPartitioned);
 
@@ -1458,7 +1523,7 @@ void FPCGActorAndComponentMapping::OnActorChanged(AActor* InActor, bool bInHasMo
 				return;
 			}
 
-			if (ComponentRef.Component->DirtyTrackedActor(InActor, /*bIntersect=*/true, RemovedTags, InOriginatingChangeObject))
+			if (ClearCacheForActor(InActor, ComponentRef.Component, /*bIntersect=*/true, RemovedTags, InOriginatingChangeObject))
 			{
 				ComponentRef.Component->DirtyGenerated(DirtyFlag);
 				DirtyComponents.Add(ComponentRef.Component);
@@ -1467,8 +1532,9 @@ void FPCGActorAndComponentMapping::OnActorChanged(AActor* InActor, bool bInHasMo
 
 		NonPartitionedOctree.FindElementsWithBoundsTest(ActorBounds, UpdateNonPartitioned);
 
-		// For partitioned, we first need to find all components that intersect with our actor and then forward the dirty call to all local components that intersect.
-		auto UpdatePartitioned = [this, &DirtyComponents, InActor, CulledTrackedComponents, &ActorBounds, &RemovedTags, DirtyFlag, InOriginatingChangeObject](const FPCGComponentRef& ComponentRef)  -> void
+		// For partitioned, we first need check if the original component intersect with the bounds, then forward the dirty call only to locals that intersect with the bounds.
+		// Note: CurrentActorBoundsPtr is passed by reference because it will be modified between lambda calls (cf comment above).
+		auto UpdatePartitioned = [this, &DirtyComponents, InActor, CulledTrackedComponents, &CurrentActorBoundsPtr, &RemovedTags, DirtyFlag, InOriginatingChangeObject](const FPCGComponentRef& ComponentRef)  -> void
 		{
 			TRACE_CPUPROFILER_EVENT_SCOPE(FPCGActorAndComponentMapping::OnActorChanged::UpdatePartitioned);
 
@@ -1481,20 +1547,29 @@ void FPCGActorAndComponentMapping::OnActorChanged(AActor* InActor, bool bInHasMo
 				return;
 			}
 
-			const FBox Overlap = ActorBounds.Overlap(ComponentRef.Bounds.GetBox());
+			check(CurrentActorBoundsPtr);
+			const FBox ComponentBounds = ComponentRef.Bounds.GetBox();
+			const bool bIntersect = CurrentActorBoundsPtr->Intersect(ComponentBounds);
+			if (!bIntersect)
+			{
+				return;
+			}
+
+			const FBox Overlap = CurrentActorBoundsPtr->Overlap(ComponentBounds);
 			bool bWasDirtied = false;
 
-			ForAllIntersectingPartitionActors(Overlap, [InActor, Component = ComponentRef.Component, &RemovedTags, &bWasDirtied, DirtyFlag, InOriginatingChangeObject](APCGPartitionActor* InPartitionActor) -> void
+			// Since when we clear the cache for a settings, we clear it all, it's only necessary to do it once on the original component, then only dirty the local that intersects.
+			if (ClearCacheForActor(InActor, ComponentRef.Component, /*bIntersect=*/true, RemovedTags, InOriginatingChangeObject))
 			{
-				if (UPCGComponent* LocalComponent = InPartitionActor->GetLocalComponent(Component))
+				ForAllIntersectingPartitionActors(Overlap, [InActor, Component = ComponentRef.Component, &bWasDirtied, DirtyFlag, InOriginatingChangeObject](APCGPartitionActor* InPartitionActor) -> void
 				{
-					if (LocalComponent->DirtyTrackedActor(InActor, /*bIntersect=*/true, RemovedTags, InOriginatingChangeObject))
+					if (UPCGComponent* LocalComponent = InPartitionActor->GetLocalComponent(Component))
 					{
 						bWasDirtied = true;
 						LocalComponent->DirtyGenerated(DirtyFlag);
 					}
-				}
-			});
+				});
+			}
 
 			if (bWasDirtied)
 			{
@@ -1507,25 +1582,19 @@ void FPCGActorAndComponentMapping::OnActorChanged(AActor* InActor, bool bInHasMo
 		PartitionedOctree.FindElementsWithBoundsTest(ActorBounds, UpdatePartitioned);
 
 		// If it has moved, redo it with the old bounds.
-		if (bInHasMoved)
+		if (OldActorBoundsPtr)
 		{
 			TRACE_CPUPROFILER_EVENT_SCOPE(FPCGActorAndComponentMapping::OnActorChanged::SecondUpdateHasMoved);
 
-			if (FBox* OldActorBounds = TrackedActorToPositionMap.Find(InActor))
-			{
-				if (!OldActorBounds->Equals(ActorBounds))
-				{
-					// Set the actor bounds with the old one, to have the right Overlap in the Partition case.
-					ActorBounds = *OldActorBounds;
-					NonPartitionedOctree.FindElementsWithBoundsTest(*OldActorBounds, UpdateNonPartitioned);
-					PartitionedOctree.FindElementsWithBoundsTest(*OldActorBounds, UpdatePartitioned);
-				}
-			}
+			// Set the actor bounds with the old one, to have the right Overlap in the Partition case.
+			CurrentActorBoundsPtr = OldActorBoundsPtr;
+			NonPartitionedOctree.FindElementsWithBoundsTest(*OldActorBoundsPtr, UpdateNonPartitioned);
+			PartitionedOctree.FindElementsWithBoundsTest(*OldActorBoundsPtr, UpdatePartitioned);
 		}
 	}
 
 	// Finally, dirty all components that always track this actor that are not yet notified.
-	if (TSet<UPCGComponent*>* AlwaysTrackedComponents = AlwaysTrackedActorsToComponentsMap.Find(InActor))
+	if (AlwaysTrackedComponents)
 	{
 		TRACE_CPUPROFILER_EVENT_SCOPE(FPCGActorAndComponentMapping::OnActorChanged::AlwaysTrackedUpdate);
 
@@ -1536,34 +1605,22 @@ void FPCGActorAndComponentMapping::OnActorChanged(AActor* InActor, bool bInHasMo
 				continue;
 			}
 
-			// Don't mark "Owner changed" if the change originate from a PCG Component. It will be delegated to the DirtyTrackedActor
+			// Don't mark "Owner changed" if the change originate from a PCG Component. It will be delegated to the ClearCacheForActor.
 			// It is necessary to avoid infine loops when there are multiple PCG components on one actor, and one component was generated.
 			const bool bOwnerChanged = (PCGComponent->GetOwner() == InActor) && (!InOriginatingChangeObject || !InOriginatingChangeObject->IsA<UPCGComponent>());
 			bool bWasDirtied = false;
 
-			if (!DirtyComponents.Contains(PCGComponent) && !bOwnerChanged)
+			// Since when we clear the cache for a settings, we clear it all, it's only necessary to do it once on the original component. And we will dirty all the local components (with dirty dispatch)
+			if (!DirtyComponents.Contains(PCGComponent))
 			{
-				if (PCGComponent->IsPartitioned())
-				{
-					DispatchToRegisteredLocalComponents(PCGComponent, [InActor, &RemovedTags, &bWasDirtied, DirtyFlag, InOriginatingChangeObject](UPCGComponent* InLocalComponent) -> FPCGTaskId
-					{
-						if (InLocalComponent->DirtyTrackedActor(InActor, /*bIntersect=*/false, RemovedTags, InOriginatingChangeObject))
-						{
-							bWasDirtied = true;
-							InLocalComponent->DirtyGenerated(DirtyFlag);
-						}
-						return InvalidPCGTaskId;
-					});
-				}
-				else
-				{
-					bWasDirtied = PCGComponent->DirtyTrackedActor(InActor, /*bIntersect=*/false, RemovedTags, InOriginatingChangeObject);
-				}
+				const FBox ComponentBounds = PCGComponent->GetGridBounds();
+				const bool bIntersect = ActorBounds.Intersect(ComponentBounds) || (OldActorBoundsPtr && OldActorBoundsPtr->Intersect(ComponentBounds));
+				bWasDirtied = ClearCacheForActor(InActor, PCGComponent, /*bIntersect=*/bIntersect, RemovedTags, InOriginatingChangeObject);
 			}
 
 			if (bWasDirtied || bOwnerChanged)
 			{
-				PCGComponent->DirtyGenerated(DirtyFlag, /*bDispatchToLocalComponents=*/bOwnerChanged);
+				PCGComponent->DirtyGenerated(DirtyFlag, /*bDispatchToLocalComponents=*/true);
 				DirtyComponents.Add(PCGComponent);
 			}
 		}
@@ -1658,4 +1715,30 @@ bool FPCGActorAndComponentMapping::IsActorTracked(const AActor* InActor) const
 	return InActor && (CulledTrackedActorsToComponentsMap.Contains(InActor) || AlwaysTrackedActorsToComponentsMap.Contains(InActor));
 }
 
+bool FPCGActorAndComponentMapping::ClearCacheForActor(const AActor* InActor, const UPCGComponent* InComponent, const bool bIntersect, const TSet<FName>& InRemovedTags, const UObject* InOriginatingChange) const
+{
+	check(InActor && InComponent);
+
+	if (!PCGSubsystem)
+	{
+		return false;
+	}
+
+	// Special case for the landscape. No settings associated, but we should dirty if the component is tracking the landscape.
+	bool bShouldDirty = InComponent->ShouldTrackLandscape() && InActor->IsA<ALandscapeProxy>();
+
+	for (const UPCGSettings* Settings : InComponent->GatherSettingsTrackingActor(InActor, bIntersect, InRemovedTags, InOriginatingChange))
+	{
+		if (!Settings) 
+		{
+			continue;
+		}
+
+		bShouldDirty = true;
+		PCGSubsystem->CleanFromCache(Settings->GetElement().Get(), Settings);
+	}
+
+	return bShouldDirty;
+}
+
 #endif // WITH_EDITOR
\ No newline at end of file
diff --git a/Engine/Plugins/Experimental/PCG/Source/PCG/Private/PCGComponent.cpp b/Engine/Plugins/Experimental/PCG/Source/PCG/Private/PCGComponent.cpp
--- a/Engine/Plugins/Experimental/PCG/Source/PCG/Private/PCGComponent.cpp
+++ b/Engine/Plugins/Experimental/PCG/Source/PCG/Private/PCGComponent.cpp
@@ -30,6 +30,7 @@
 
 #include "LandscapeComponent.h"
 #include "LandscapeProxy.h"
+#include "Algo/AllOf.h"
 #include "Algo/AnyOf.h"
 #include "Algo/Transform.h"
 #include "Components/BillboardComponent.h"
@@ -202,15 +203,7 @@ void UPCGComponent::SetPropertiesFromOriginal(const UPCGComponent* Original)
 
 #if WITH_EDITOR
 	const bool bHasDirtyInput = InputType != NewInputType;
-
-	TSet<FPCGActorSelectionKey> TrackedKeys;
-	TSet<FPCGActorSelectionKey> OriginalTrackedKeys;
-	CachedTrackedKeysToSettings.GetKeys(TrackedKeys);
-	Original->CachedTrackedKeysToSettings.GetKeys(OriginalTrackedKeys);
-
-	const bool bHasDirtyTracking = !(TrackedKeys.Num() == OriginalTrackedKeys.Num() && TrackedKeys.Includes(OriginalTrackedKeys));
-
-	const bool bIsDirty = bHasDirtyInput || bHasDirtyTracking || bGraphInstanceIsDifferent;
+	const bool bIsDirty = bHasDirtyInput || bGraphInstanceIsDifferent;
 #endif // WITH_EDITOR
 
 	InputType = NewInputType;
@@ -229,11 +222,6 @@ void UPCGComponent::SetPropertiesFromOriginal(const UPCGComponent* Original)
 	}
 
 #if WITH_EDITOR
-	if (bHasDirtyTracking)
-	{
-		UpdateTrackingCache();
-	}
-
 	// Note that while we dirty here, we won't trigger a refresh since we don't have the required context
 	if (bIsDirty)
 	{
@@ -1344,6 +1332,11 @@ void UPCGComponent::PostEditChangeProperty(FPropertyChangedEvent& PropertyChange
 		DirtyGenerated(EPCGComponentDirtyFlag::Input);
 		Refresh();
 	}
+	else if (PropName == GET_MEMBER_NAME_CHECKED(UPCGComponent, Seed))
+	{
+		DirtyGenerated();
+		Refresh();
+	}
 	// General properties that don't affect behavior
 	else
 	{
@@ -1397,8 +1390,9 @@ void UPCGComponent::PostEditUndo()
 
 bool UPCGComponent::UpdateTrackingCache(TArray<FPCGActorSelectionKey>* OptionalChangedKeys)
 {
-	// Without an owner, it probably means we are in a BP template, so no need to setup callbacks
-	if (!GetOwner())
+	// Without an owner, it probably means we are in a BP template, so no need to update the tracking cache.
+	// Same for local components, as we will use the cache of the original component.
+	if (!GetOwner() || IsLocalComponent())
 	{
 		return false;
 	}
@@ -1427,26 +1421,23 @@ bool UPCGComponent::UpdateTrackingCache(TArray<FPCGActorSelectionKey>* OptionalC
 		{
 			const FPCGActorSelectionKey& Key = It.Key;
 
-			bool bShouldCull = true;
-			for (const FPCGSettingsAndCulling& SettingsAndCulling : It.Value)
-			{
-				if (!SettingsAndCulling.Value)
-				{
-					bShouldCull = false;
-					break;
-				}
-			}
+			// Should cull only if all the settings requires a cull.
+			const bool bShouldCull = Algo::AllOf(It.Value, [](const FPCGSettingsAndCulling& SettingsAndCullingPair) { return SettingsAndCullingPair.Value; });
 
 			NewTrackedKeysToCulling.Emplace(Key, bShouldCull);
 
-			bool* OldCulling = CachedTrackedKeysToCulling.Find(Key);
-			if (OldCulling && (*OldCulling == bShouldCull))
+			// Look for the key in the previous cached keys.
+			const bool* OldCulling = CachedTrackedKeysToCulling.Find(Key);
+			// It is a new key if we didn't find it. We also mark it changed if the key existed, but the culling changed.
+			const bool bNewKeyOrCullChanged = !OldCulling || (*OldCulling != bShouldCull);
+			// Then remove the key.
+			CachedTrackedKeysToCulling.Remove(Key);
+
+			if (!bNewKeyOrCullChanged)
 			{
 				++FoundKeys;
 			}
-
-			// Remove the key from the previous cached map. If nothing is removed, it means it is a new key
-			if ((CachedTrackedKeysToCulling.Remove(Key) == 0) && OptionalChangedKeys)
+			else if (OptionalChangedKeys)
 			{
 				OptionalChangedKeys->Add(Key);
 			}
@@ -2294,16 +2285,16 @@ UPCGSubsystem* UPCGComponent::GetSubsystem() const
 }
 
 #if WITH_EDITOR
-bool UPCGComponent::DirtyTrackedActor(AActor* InActor, bool bIntersect, const TSet<FName>& InRemovedTags, const UObject* InOriginatingChangeObject)
+TArray<const UPCGSettings*> UPCGComponent::GatherSettingsTrackingActor(const AActor* InActor, bool bIntersect, const TSet<FName>& InRemovedTags, const UObject* InOriginatingChangeObject) const
 {
-	TRACE_CPUPROFILER_EVENT_SCOPE(UPCGComponent::DirtyTrackedActor);
+	TRACE_CPUPROFILER_EVENT_SCOPE(UPCGComponent::GatherSettingsTrackingActor);
 
 	if (!InActor)
 	{
-		return false;
+		return {};
 	}
 
-	bool bWasDirtied = false;
+	TArray<const UPCGSettings*> TrackedSettings;
 
 	for (const auto& It : CachedTrackedKeysToSettings)
 	{
@@ -2327,24 +2318,15 @@ bool UPCGComponent::DirtyTrackedActor(AActor* InActor, bool bIntersect, const TS
 					continue;
 				}
 
-				const TWeakObjectPtr<const UPCGSettings>& Settings = SettingsAndCulling.Key;
-				if (Settings.IsValid())
+				if (const UPCGSettings* Settings = SettingsAndCulling.Key.Get())
 				{
-					GetSubsystem()->CleanFromCache(Settings->GetElement().Get(), Settings.Get());
+					TrackedSettings.Add(Settings);
 				}
-
-				bWasDirtied = true;
 			}
 		}
 	}
 
-	// Special case for landscape, we should dirty.
-	if (ShouldTrackLandscape() && InActor->IsA<ALandscapeProxy>())
-	{
-		bWasDirtied = true;
-	}
-
-	return bWasDirtied;
+	return TrackedSettings;
 }
 
 bool UPCGComponent::ShouldTrackLandscape() const
diff --git a/Engine/Plugins/Experimental/PCG/Source/PCG/Private/PCGGraph.cpp b/Engine/Plugins/Experimental/PCG/Source/PCG/Private/PCGGraph.cpp
--- a/Engine/Plugins/Experimental/PCG/Source/PCG/Private/PCGGraph.cpp
+++ b/Engine/Plugins/Experimental/PCG/Source/PCG/Private/PCGGraph.cpp
@@ -915,6 +915,21 @@ void UPCGGraph::GetTrackedActorKeysToSettings(FPCGActorSelectionKeyToSettingsMap
 			Node->GetSettings()->GetTrackedActorKeys(OutTagsToSettings, OutVisitedGraphs);
 		}
 	}
+
+	// Make sure that all Self/Original keys are marked as not-cull, since the component will always intersect with its owner/original
+	// We need to loop because we can have multiple keys that have Filter to self or original.
+	for (auto& It : OutTagsToSettings)
+	{
+		const FPCGActorSelectionKey& Key = It.Key;
+		TArray<FPCGSettingsAndCulling>& SettingsAndCullingArray = It.Value;
+		if (Key.ActorFilter == EPCGActorFilter::Self || Key.ActorFilter == EPCGActorFilter::Original)
+		{
+			for (FPCGSettingsAndCulling& SettingsAndCullingPair : SettingsAndCullingArray)
+			{
+				SettingsAndCullingPair.Value = false;
+			}
+		}
+	}
 }
 
 void UPCGGraph::NotifyGraphChanged(EPCGChangeType ChangeType)
diff --git a/Engine/Plugins/Experimental/PCG/Source/PCG/Public/Elements/PCGDataFromActor.h b/Engine/Plugins/Experimental/PCG/Source/PCG/Public/Elements/PCGDataFromActor.h
--- a/Engine/Plugins/Experimental/PCG/Source/PCG/Public/Elements/PCGDataFromActor.h
+++ b/Engine/Plugins/Experimental/PCG/Source/PCG/Public/Elements/PCGDataFromActor.h
@@ -101,7 +101,7 @@ class FPCGDataFromActorElement : public IPCGElement
 public:
 	virtual FPCGContext* Initialize(const FPCGDataCollection& InputData, TWeakObjectPtr<UPCGComponent> SourceComponent, const UPCGNode* Node) override;
 	virtual bool CanExecuteOnlyOnMainThread(FPCGContext* Context) const override { return true; }
-	virtual bool IsCacheable(const UPCGSettings* InSettings) const override { return false; }
+	virtual void GetDependenciesCrc(const FPCGDataCollection& InInput, const UPCGSettings* InSettings, UPCGComponent* InComponent, FPCGCrc& OutCrc) const override;
 
 protected:
 	virtual bool ExecuteInternal(FPCGContext* Context) const;
diff --git a/Engine/Plugins/Experimental/PCG/Source/PCG/Public/PCGActorAndComponentMapping.h b/Engine/Plugins/Experimental/PCG/Source/PCG/Public/PCGActorAndComponentMapping.h
--- a/Engine/Plugins/Experimental/PCG/Source/PCG/Public/PCGActorAndComponentMapping.h
+++ b/Engine/Plugins/Experimental/PCG/Source/PCG/Public/PCGActorAndComponentMapping.h
@@ -126,8 +126,10 @@ private:
 	bool UnregisterActor(AActor* InActor);
 
 	void OnActorAdded(AActor* InActor);
+	void OnActorLoaded(AActor& InActor);
 	void OnActorAdded_Internal(AActor* InActor, bool bShouldDirty, int32 LevelInstanceDepth);
 	void OnActorDeleted(AActor* InActor);
+	void OnActorUnloaded(AActor& InActor);
 	void OnActorDeleted_Internal(AActor* InActor, int32 LevelInstanceDepth);
 	void OnActorMoved(AActor* InActor);
 	void OnActorMoved_Internal(AActor* InActor, int32 LevelInstanceDepth);
@@ -155,6 +157,9 @@ private:
 	void UpdateActorDependencies(AActor* InActor);
 
 	bool IsActorTracked(const AActor* InActor) const;
+
+	/** Gather all settings from a given component that track the actor, and clear the cache for them. Returns true if we should dirty afterwards (aka at least one settings was cleared and/or landscape changed). */
+	bool ClearCacheForActor(const AActor* InActor, const UPCGComponent* InComponent, const bool bIntersect, const TSet<FName>& InRemovedTags, const UObject* InOriginatingChange) const;
 #endif // WITH_EDITOR
 
 private:
diff --git a/Engine/Plugins/Experimental/PCG/Source/PCG/Public/PCGComponent.h b/Engine/Plugins/Experimental/PCG/Source/PCG/Public/PCGComponent.h
--- a/Engine/Plugins/Experimental/PCG/Source/PCG/Public/PCGComponent.h
+++ b/Engine/Plugins/Experimental/PCG/Source/PCG/Public/PCGComponent.h
@@ -351,8 +351,14 @@ private:
 	/** Returns true if something changed in the tracking. */
 	bool UpdateTrackingCache(TArray<FPCGActorSelectionKey>* OptionalChangedKeys = nullptr);
 
-	// InOriginatingChangeObject can be null
-	bool DirtyTrackedActor(AActor* InActor, bool bIntersect, const TSet<FName>& InRemovedTags, const UObject* InOriginatingChangeObject);
+	/** Gather all the settings that are tracking the given actor. If the settings should cull and there is no intersection, it will not be added to the list.
+	* @param InActor                    The actor the component is tracking
+	* @param bIntersect                 If the actor is intersecting with the component. Useful for culling.
+	* @param InRemovedTags              List of tags that were removed, because the actor won't have those tags on it, but we still need to gather the settings that were tracking this tag. Can be empty.
+	* @param InOriginatingChangeObject: Optional pointer on the originating object that triggered the update. Useful to track PCGComponents generation/cleanup. Can be null.
+	* @return                           Array of all the settings that track the actor.
+	*/
+	TArray<const UPCGSettings*> GatherSettingsTrackingActor(const AActor* InActor, const bool bIntersect, const TSet<FName>& InRemovedTags, const UObject* InOriginatingChangeObject) const;
 
 	bool ShouldTrackLandscape() const;
 #endif
