[PCG] Static analysis fix

#rnx

diff --git a/Engine/Plugins/Experimental/PCG/Source/PCG/Private/PCGActorAndComponentMapping.cpp b/Engine/Plugins/Experimental/PCG/Source/PCG/Private/PCGActorAndComponentMapping.cpp
--- a/Engine/Plugins/Experimental/PCG/Source/PCG/Private/PCGActorAndComponentMapping.cpp
+++ b/Engine/Plugins/Experimental/PCG/Source/PCG/Private/PCGActorAndComponentMapping.cpp
@@ -959,7 +959,8 @@ void FPCGActorAndComponentMapping::RegisterTrackingCallbacks()
 	FCoreUObjectDelegates::OnPreObjectPropertyChanged.AddRaw(this, &FPCGActorAndComponentMapping::OnPreObjectPropertyChanged);
 
 	UWorld* World = PCGSubsystem ? PCGSubsystem->GetWorld() : nullptr;
-	if (IsValid(World) && World->PersistentLevel)
+	// Need the World condition for static analysis...
+	if (World && IsValid(World) && World->PersistentLevel)
 	{
 		World->PersistentLevel->OnLoadedActorAddedToLevelEvent.AddRaw(this, &FPCGActorAndComponentMapping::OnActorLoaded);
 		World->PersistentLevel->OnLoadedActorRemovedFromLevelEvent.AddRaw(this, &FPCGActorAndComponentMapping::OnActorUnloaded);
@@ -975,7 +976,8 @@ void FPCGActorAndComponentMapping::TeardownTrackingCallbacks()
 	FCoreUObjectDelegates::OnPreObjectPropertyChanged.RemoveAll(this);
 
 	UWorld* World = PCGSubsystem ? PCGSubsystem->GetWorld() : nullptr;
-	if (IsValid(World) && World->PersistentLevel)
+	// Need the World condition for static analysis...
+	if (World && IsValid(World) && World->PersistentLevel)
 	{
 		World->PersistentLevel->OnLoadedActorAddedToLevelEvent.RemoveAll(this);
 		World->PersistentLevel->OnLoadedActorRemovedFromLevelEvent.RemoveAll(this);
