[PCG] Fix NaN from Distance node if MaximumDistance == 0.0

#rnx

diff --git a/Engine/Plugins/Experimental/PCG/Source/PCG/Private/Elements/PCGDistance.cpp b/Engine/Plugins/Experimental/PCG/Source/PCG/Private/Elements/PCGDistance.cpp
--- a/Engine/Plugins/Experimental/PCG/Source/PCG/Private/Elements/PCGDistance.cpp
+++ b/Engine/Plugins/Experimental/PCG/Source/PCG/Private/Elements/PCGDistance.cpp
@@ -90,10 +90,12 @@ bool FPCGDistanceElement::ExecuteInternal(FPCGContext* Context) const
 	const FName AttributeName = Settings->AttributeName;
 	const bool bSetDensity = Settings->bSetDensity;
 	const bool bOutputDistanceVector = Settings->bOutputDistanceVector;
-	const double MaximumDistance = Settings->MaximumDistance;
 	const PCGDistanceShape SourceShape = Settings->SourceShape;
 	const PCGDistanceShape TargetShape = Settings->TargetShape;
 
+	const double MaximumDistance = FMath::Max(0.0, Settings->MaximumDistance);
+	const double MaximumDistanceRecip = MaximumDistance > UE_DOUBLE_SMALL_NUMBER ? 1.0 / MaximumDistance : 0.0;
+
 	TArray<FPCGTaggedData> Sources = Context->InputData.GetInputsByPin(PCGDistance::SourceLabel);
 	TArray<FPCGTaggedData> Targets = Context->InputData.GetInputsByPin(PCGDistance::TargetLabel);
 	TArray<FPCGTaggedData>& Outputs = Context->OutputData.TaggedData;
@@ -121,10 +123,10 @@ bool FPCGDistanceElement::ExecuteInternal(FPCGContext* Context) const
 		TargetPointDatas.Add(TargetPointData);
 	}
 
-	// first find the total Input bounds which will determine the size of each cell
+	// First find the total Input bounds which will determine the size of each cell
 	for (const FPCGTaggedData& Source : Sources) 
 	{
-		// add the point bounds to the input cell
+		// Add the point bounds to the input cell
 
 		const UPCGSpatialData* SourceData = Cast<UPCGSpatialData>(Source.Data);
 
@@ -159,7 +161,7 @@ bool FPCGDistanceElement::ExecuteInternal(FPCGContext* Context) const
 		}
 
 		FPCGAsync::AsyncPointProcessing(Context, SourcePointData->GetPoints(), OutputData->GetMutablePoints(),
-			[OutputData, SourceShape, TargetShape, &TargetPointDatas, MaximumDistance, ScalarAttribute, VectorAttribute, bSetDensity](const FPCGPoint& SourcePoint, FPCGPoint& OutPoint) {
+			[OutputData, SourceShape, TargetShape, &TargetPointDatas, MaximumDistance, MaximumDistanceRecip, ScalarAttribute, VectorAttribute, bSetDensity](const FPCGPoint& SourcePoint, FPCGPoint& OutPoint) {
 
 				OutPoint = SourcePoint;
 
@@ -167,7 +169,7 @@ bool FPCGDistanceElement::ExecuteInternal(FPCGContext* Context) const
 
 				const FVector SourceCenter = SourcePoint.Transform.TransformPosition(SourcePoint.GetLocalCenter());
 
-				double MinDistanceSquared = MaximumDistance*MaximumDistance;
+				double MinDistanceSquared = MaximumDistance * MaximumDistance;
 				FVector MinDistanceVector = FVector::ZeroVector;
 
 				// Signed distance field for calculating the closest point of source and target
@@ -208,7 +210,7 @@ bool FPCGDistanceElement::ExecuteInternal(FPCGContext* Context) const
 					);
 				}
 
-				const float Distance = FMath::Sign(MinDistanceSquared) * FMath::Sqrt(FMath::Abs(MinDistanceSquared));
+				const double Distance = FMath::Sign(MinDistanceSquared) * FMath::Sqrt(FMath::Abs(MinDistanceSquared));
 
 				if (ScalarAttribute || VectorAttribute)
 				{
@@ -227,8 +229,8 @@ bool FPCGDistanceElement::ExecuteInternal(FPCGContext* Context) const
 				
 				if (bSetDensity)
 				{
-					// set density instead
-					OutPoint.Density = FMath::Clamp(Distance, -MaximumDistance, MaximumDistance) / MaximumDistance;
+					// Set density instead
+					OutPoint.Density = MaximumDistance > UE_DOUBLE_SMALL_NUMBER ? (FMath::Clamp(Distance, -MaximumDistance, MaximumDistance) * MaximumDistanceRecip) : 1.0f;
 				}
 
 				return true;
