[PCG] Clean up and deprecate SimplePCGElement

#rnx

diff --git a/Engine/Plugins/Experimental/PCG/Source/PCG/Private/Elements/IO/PCGExternalData.cpp b/Engine/Plugins/Experimental/PCG/Source/PCG/Private/Elements/IO/PCGExternalData.cpp
--- a/Engine/Plugins/Experimental/PCG/Source/PCG/Private/Elements/IO/PCGExternalData.cpp
+++ b/Engine/Plugins/Experimental/PCG/Source/PCG/Private/Elements/IO/PCGExternalData.cpp
@@ -8,14 +8,9 @@
 
 #define LOCTEXT_NAMESPACE "PCGExternalData"
 
-FPCGContext* FPCGExternalDataElement::Initialize(const FPCGDataCollection& InputData, TWeakObjectPtr<UPCGComponent> SourceComponent, const UPCGNode* Node)
+FPCGContext* FPCGExternalDataElement::CreateContext()
 {
-	FPCGExternalDataContext* Context = new FPCGExternalDataContext();
-	Context->InputData = InputData;
-	Context->SourceComponent = SourceComponent;
-	Context->Node = Node;
-
-	return Context;
+	return new FPCGExternalDataContext();
 }
 
 bool FPCGExternalDataElement::PrepareDataInternal(FPCGContext* InContext) const
diff --git a/Engine/Plugins/Experimental/PCG/Source/PCG/Private/Elements/PCGAttributeNoise.cpp b/Engine/Plugins/Experimental/PCG/Source/PCG/Private/Elements/PCGAttributeNoise.cpp
--- a/Engine/Plugins/Experimental/PCG/Source/PCG/Private/Elements/PCGAttributeNoise.cpp
+++ b/Engine/Plugins/Experimental/PCG/Source/PCG/Private/Elements/PCGAttributeNoise.cpp
@@ -153,14 +153,9 @@ FPCGElementPtr UPCGAttributeNoiseSettings::CreateElement() const
 	return MakeShared<FPCGAttributeNoiseElement>();
 }
 
-FPCGContext* FPCGAttributeNoiseElement::Initialize(const FPCGDataCollection& InInputData, TWeakObjectPtr<UPCGComponent> InSourceComponent, const UPCGNode* InNode)
+FPCGContext* FPCGAttributeNoiseElement::CreateContext()
 {
-	FPCGAttributeNoiseContext* Context = new FPCGAttributeNoiseContext();
-	Context->InputData = InInputData;
-	Context->SourceComponent = InSourceComponent;
-	Context->Node = InNode;
-
-	return Context;
+	return new FPCGAttributeNoiseContext();
 }
 
 bool FPCGAttributeNoiseElement::ExecuteInternal(FPCGContext* InContext) const
diff --git a/Engine/Plugins/Experimental/PCG/Source/PCG/Private/Elements/PCGCullPointsOutsideActorBounds.cpp b/Engine/Plugins/Experimental/PCG/Source/PCG/Private/Elements/PCGCullPointsOutsideActorBounds.cpp
--- a/Engine/Plugins/Experimental/PCG/Source/PCG/Private/Elements/PCGCullPointsOutsideActorBounds.cpp
+++ b/Engine/Plugins/Experimental/PCG/Source/PCG/Private/Elements/PCGCullPointsOutsideActorBounds.cpp
@@ -84,7 +84,7 @@ bool FPCGCullPointsOutsideActorBoundsElement::ExecuteInternal(FPCGContext* Conte
 void FPCGCullPointsOutsideActorBoundsElement::GetDependenciesCrc(const FPCGDataCollection& InInput, const UPCGSettings* InSettings, UPCGComponent* InComponent, FPCGCrc& OutCrc) const
 {
 	FPCGCrc Crc;
-	FSimplePCGElement::GetDependenciesCrc(InInput, InSettings, InComponent, Crc);
+	IPCGElement::GetDependenciesCrc(InInput, InSettings, InComponent, Crc);
 
 	// The culling volume depends on the component transform.
 	const UPCGData* ActorData = InComponent ? InComponent->GetActorPCGData() : nullptr;
diff --git a/Engine/Plugins/Experimental/PCG/Source/PCG/Private/Elements/PCGDataFromActor.cpp b/Engine/Plugins/Experimental/PCG/Source/PCG/Private/Elements/PCGDataFromActor.cpp
--- a/Engine/Plugins/Experimental/PCG/Source/PCG/Private/Elements/PCGDataFromActor.cpp
+++ b/Engine/Plugins/Experimental/PCG/Source/PCG/Private/Elements/PCGDataFromActor.cpp
@@ -132,14 +132,9 @@ TArray<FPCGPinProperties> UPCGDataFromActorSettings::OutputPinProperties() const
 	return Pins;
 }
 
-FPCGContext* FPCGDataFromActorElement::Initialize(const FPCGDataCollection& InputData, TWeakObjectPtr<UPCGComponent> SourceComponent, const UPCGNode* Node)
+FPCGContext* FPCGDataFromActorElement::CreateContext()
 {
-	FPCGDataFromActorContext* Context = new FPCGDataFromActorContext();
-	Context->InputData = InputData;
-	Context->SourceComponent = SourceComponent;
-	Context->Node = Node;
-
-	return Context;
+	return new FPCGDataFromActorContext();
 }
 
 bool FPCGDataFromActorElement::ExecuteInternal(FPCGContext* InContext) const
diff --git a/Engine/Plugins/Experimental/PCG/Source/PCG/Private/Elements/PCGHiGenGridSize.cpp b/Engine/Plugins/Experimental/PCG/Source/PCG/Private/Elements/PCGHiGenGridSize.cpp
--- a/Engine/Plugins/Experimental/PCG/Source/PCG/Private/Elements/PCGHiGenGridSize.cpp
+++ b/Engine/Plugins/Experimental/PCG/Source/PCG/Private/Elements/PCGHiGenGridSize.cpp
@@ -140,7 +140,7 @@ bool FPCGHiGenGridSizeElement::ExecuteInternal(FPCGContext* Context) const
 void FPCGHiGenGridSizeElement::GetDependenciesCrc(const FPCGDataCollection& InInput, const UPCGSettings* InSettings, UPCGComponent* InComponent, FPCGCrc& OutCrc) const
 {
 	FPCGCrc Crc;
-	FSimplePCGElement::GetDependenciesCrc(InInput, InSettings, InComponent, Crc);
+	IPCGElement::GetDependenciesCrc(InInput, InSettings, InComponent, Crc);
 
 	// The grid cell volume output depends on the component transform.
 	// NOTE: It might be interesting to only incorporate the transform if the cell volume output pin is connected (and if we have a node obviously).
diff --git a/Engine/Plugins/Experimental/PCG/Source/PCG/Private/Elements/PCGIndirectionElement.cpp b/Engine/Plugins/Experimental/PCG/Source/PCG/Private/Elements/PCGIndirectionElement.cpp
--- a/Engine/Plugins/Experimental/PCG/Source/PCG/Private/Elements/PCGIndirectionElement.cpp
+++ b/Engine/Plugins/Experimental/PCG/Source/PCG/Private/Elements/PCGIndirectionElement.cpp
@@ -76,14 +76,9 @@ FPCGIndirectionContext::~FPCGIndirectionContext()
 	InnerContext = nullptr;
 }
 
-FPCGContext* FPCGIndirectionElement::Initialize(const FPCGDataCollection& InputData, TWeakObjectPtr<UPCGComponent> SourceComponent, const UPCGNode* Node)
+FPCGContext* FPCGIndirectionElement::CreateContext()
 {
-	FPCGIndirectionContext* Context = new FPCGIndirectionContext();
-	Context->InputData = InputData;
-	Context->SourceComponent = SourceComponent;
-	Context->Node = Node;
-
-	return Context;
+	return new FPCGIndirectionContext();
 }
 
 bool FPCGIndirectionElement::PrepareDataInternal(FPCGContext* InContext) const
diff --git a/Engine/Plugins/Experimental/PCG/Source/PCG/Private/Elements/PCGStaticMeshSpawner.cpp b/Engine/Plugins/Experimental/PCG/Source/PCG/Private/Elements/PCGStaticMeshSpawner.cpp
--- a/Engine/Plugins/Experimental/PCG/Source/PCG/Private/Elements/PCGStaticMeshSpawner.cpp
+++ b/Engine/Plugins/Experimental/PCG/Source/PCG/Private/Elements/PCGStaticMeshSpawner.cpp
@@ -66,6 +66,11 @@ FPCGElementPtr UPCGStaticMeshSpawnerSettings::CreateElement() const
 	return MakeShared<FPCGStaticMeshSpawnerElement>();
 }
 
+FPCGContext* FPCGStaticMeshSpawnerElement::CreateContext()
+{
+	return new FPCGStaticMeshSpawnerContext();
+}
+
 bool FPCGStaticMeshSpawnerElement::PrepareDataInternal(FPCGContext* InContext) const
 {
 	TRACE_CPUPROFILER_EVENT_SCOPE(FPCGStaticMeshSpawnerElement::PrepareDataInternal);
@@ -300,16 +305,6 @@ bool FPCGStaticMeshSpawnerElement::ExecuteInternal(FPCGContext* InContext) const
 	return Context->MeshInstancesData.IsEmpty();
 }
 
-FPCGContext* FPCGStaticMeshSpawnerElement::Initialize(const FPCGDataCollection& InputData, TWeakObjectPtr<UPCGComponent> SourceComponent, const UPCGNode* Node)
-{
-	FPCGStaticMeshSpawnerContext* Context = new FPCGStaticMeshSpawnerContext();
-	Context->InputData = InputData;
-	Context->SourceComponent = SourceComponent;
-	Context->Node = Node;
-
-	return Context;
-}
-
 bool FPCGStaticMeshSpawnerElement::CanExecuteOnlyOnMainThread(FPCGContext* Context) const
 {
 	// PrepareData can call UPCGManagedComponent::MarkAsReused which registers the ISMC, which can go into Chaos code that asserts if not on main thread.
diff --git a/Engine/Plugins/Experimental/PCG/Source/PCG/Private/Elements/PCGTextureSampler.cpp b/Engine/Plugins/Experimental/PCG/Source/PCG/Private/Elements/PCGTextureSampler.cpp
--- a/Engine/Plugins/Experimental/PCG/Source/PCG/Private/Elements/PCGTextureSampler.cpp
+++ b/Engine/Plugins/Experimental/PCG/Source/PCG/Private/Elements/PCGTextureSampler.cpp
@@ -129,20 +129,15 @@ bool FPCGTextureSamplerElement::ExecuteInternal(FPCGContext* InContext) const
 	return false;
 }
 
-FPCGContext* FPCGTextureSamplerElement::Initialize(const FPCGDataCollection& InputData, TWeakObjectPtr<UPCGComponent> SourceComponent, const UPCGNode* Node)
+FPCGContext* FPCGTextureSamplerElement::CreateContext()
 {
-	FPCGTextureSamplerContext* Context = new FPCGTextureSamplerContext();
-	Context->InputData = InputData;
-	Context->SourceComponent = SourceComponent;
-	Context->Node = Node;
-
-	return Context;
+	return new FPCGTextureSamplerContext();
 }
 
 void FPCGTextureSamplerElement::GetDependenciesCrc(const FPCGDataCollection& InInput, const UPCGSettings* InSettings, UPCGComponent* InComponent, FPCGCrc& OutCrc) const
 {
 	FPCGCrc Crc;
-	FSimplePCGElement::GetDependenciesCrc(InInput, InSettings, InComponent, Crc);
+	IPCGElement::GetDependenciesCrc(InInput, InSettings, InComponent, Crc);
 
 	if (const UPCGTextureSamplerSettings* Settings = Cast<UPCGTextureSamplerSettings>(InSettings))
 	{
diff --git a/Engine/Plugins/Experimental/PCG/Source/PCG/Private/Graph/PCGGraphExecutor.h b/Engine/Plugins/Experimental/PCG/Source/PCG/Private/Graph/PCGGraphExecutor.h
--- a/Engine/Plugins/Experimental/PCG/Source/PCG/Private/Graph/PCGGraphExecutor.h
+++ b/Engine/Plugins/Experimental/PCG/Source/PCG/Private/Graph/PCGGraphExecutor.h
@@ -209,7 +209,7 @@ private:
 #endif
 };
 
-class FPCGFetchInputElement : public FSimplePCGElement
+class FPCGFetchInputElement : public IPCGElement
 {
 public:
 	virtual bool IsCacheable(const UPCGSettings* InSettings) const override { return false; }
@@ -220,7 +220,7 @@ protected:
 	virtual bool IsPassthrough(const UPCGSettings* InSettings) const override { return true; }
 };
 
-class FPCGGenericElement : public FSimplePCGElement
+class FPCGGenericElement : public IPCGElement
 {
 public:
 	using FContextAllocator = TFunction<FPCGContext*(const FPCGDataCollection&, TWeakObjectPtr<UPCGComponent>, const UPCGNode*)>;
diff --git a/Engine/Plugins/Experimental/PCG/Source/PCG/Private/Tests/Elements/PCGElementTest.h b/Engine/Plugins/Experimental/PCG/Source/PCG/Private/Tests/Elements/PCGElementTest.h
--- a/Engine/Plugins/Experimental/PCG/Source/PCG/Private/Tests/Elements/PCGElementTest.h
+++ b/Engine/Plugins/Experimental/PCG/Source/PCG/Private/Tests/Elements/PCGElementTest.h
@@ -27,7 +27,7 @@ public:
 
 };
 
-class FPCGBadOutputNodeElement : public FSimplePCGElement
+class FPCGBadOutputNodeElement : public IPCGElement
 {
 public:
 	virtual bool ExecuteInternal(FPCGContext* Context) const override;
diff --git a/Engine/Plugins/Experimental/PCG/Source/PCG/Public/Elements/IO/PCGExternalData.h b/Engine/Plugins/Experimental/PCG/Source/PCG/Public/Elements/IO/PCGExternalData.h
--- a/Engine/Plugins/Experimental/PCG/Source/PCG/Public/Elements/IO/PCGExternalData.h
+++ b/Engine/Plugins/Experimental/PCG/Source/PCG/Public/Elements/IO/PCGExternalData.h
@@ -35,9 +35,8 @@ public:
 
 class PCG_API FPCGExternalDataElement : public IPCGElement
 {
-public:
-	virtual FPCGContext* Initialize(const FPCGDataCollection& InputData, TWeakObjectPtr<UPCGComponent> SourceComponent, const UPCGNode* Node) override;
 protected:
+	virtual FPCGContext* CreateContext() override;
 	virtual bool PrepareDataInternal(FPCGContext* InContext) const override;
 	virtual bool ExecuteInternal(FPCGContext* InContext) const override;
 
diff --git a/Engine/Plugins/Experimental/PCG/Source/PCG/Public/Elements/Metadata/PCGMetadataElement.h b/Engine/Plugins/Experimental/PCG/Source/PCG/Public/Elements/Metadata/PCGMetadataElement.h
--- a/Engine/Plugins/Experimental/PCG/Source/PCG/Public/Elements/Metadata/PCGMetadataElement.h
+++ b/Engine/Plugins/Experimental/PCG/Source/PCG/Public/Elements/Metadata/PCGMetadataElement.h
@@ -67,7 +67,7 @@ PRAGMA_ENABLE_DEPRECATION_WARNINGS
 #endif // WITH_EDITORONLY_DATA
 };
 
-class FPCGMetadataOperationElement : public FSimplePCGElement
+class FPCGMetadataOperationElement : public IPCGElement
 {
 protected:
 	virtual bool ExecuteInternal(FPCGContext* Context) const override;
diff --git a/Engine/Plugins/Experimental/PCG/Source/PCG/Public/Elements/Metadata/PCGMetadataOpElementBase.h b/Engine/Plugins/Experimental/PCG/Source/PCG/Public/Elements/Metadata/PCGMetadataOpElementBase.h
--- a/Engine/Plugins/Experimental/PCG/Source/PCG/Public/Elements/Metadata/PCGMetadataOpElementBase.h
+++ b/Engine/Plugins/Experimental/PCG/Source/PCG/Public/Elements/Metadata/PCGMetadataOpElementBase.h
@@ -160,7 +160,7 @@ protected:
 };
 
 
-class FPCGMetadataElementBase : public FSimplePCGElement
+class FPCGMetadataElementBase : public IPCGElement
 {
 public:
 	struct FOperationData
diff --git a/Engine/Plugins/Experimental/PCG/Source/PCG/Public/Elements/Metadata/PCGMetadataPartition.h b/Engine/Plugins/Experimental/PCG/Source/PCG/Public/Elements/Metadata/PCGMetadataPartition.h
--- a/Engine/Plugins/Experimental/PCG/Source/PCG/Public/Elements/Metadata/PCGMetadataPartition.h
+++ b/Engine/Plugins/Experimental/PCG/Source/PCG/Public/Elements/Metadata/PCGMetadataPartition.h
@@ -30,7 +30,7 @@ public:
 	FName PartitionAttribute = NAME_None;
 };
 
-class FPCGMetadataPartitionElement : public FSimplePCGElement
+class FPCGMetadataPartitionElement : public IPCGElement
 {
 protected:
 	virtual bool ExecuteInternal(FPCGContext* Context) const override;
diff --git a/Engine/Plugins/Experimental/PCG/Source/PCG/Public/Elements/Metadata/PCGMetadataRenameElement.h b/Engine/Plugins/Experimental/PCG/Source/PCG/Public/Elements/Metadata/PCGMetadataRenameElement.h
--- a/Engine/Plugins/Experimental/PCG/Source/PCG/Public/Elements/Metadata/PCGMetadataRenameElement.h
+++ b/Engine/Plugins/Experimental/PCG/Source/PCG/Public/Elements/Metadata/PCGMetadataRenameElement.h
@@ -35,7 +35,7 @@ public:
 	FName NewAttributeName = NAME_None;
 };
 
-class FPCGMetadataRenameElement : public FSimplePCGElement
+class FPCGMetadataRenameElement : public IPCGElement
 {
 protected:
 	virtual bool ExecuteInternal(FPCGContext* Context) const override;
diff --git a/Engine/Plugins/Experimental/PCG/Source/PCG/Public/Elements/PCGAddTag.h b/Engine/Plugins/Experimental/PCG/Source/PCG/Public/Elements/PCGAddTag.h
--- a/Engine/Plugins/Experimental/PCG/Source/PCG/Public/Elements/PCGAddTag.h
+++ b/Engine/Plugins/Experimental/PCG/Source/PCG/Public/Elements/PCGAddTag.h
@@ -37,7 +37,7 @@ public:
 	FString TagsToAdd;
 };
 
-class FPCGAddTagElement : public FSimplePCGElement
+class FPCGAddTagElement : public IPCGElement
 {
 protected:
 	virtual bool ExecuteInternal(FPCGContext* Context) const override;
diff --git a/Engine/Plugins/Experimental/PCG/Source/PCG/Public/Elements/PCGAttributeFilterElement.h b/Engine/Plugins/Experimental/PCG/Source/PCG/Public/Elements/PCGAttributeFilterElement.h
--- a/Engine/Plugins/Experimental/PCG/Source/PCG/Public/Elements/PCGAttributeFilterElement.h
+++ b/Engine/Plugins/Experimental/PCG/Source/PCG/Public/Elements/PCGAttributeFilterElement.h
@@ -60,7 +60,7 @@ public:
 };
 
 
-class FPCGAttributeFilterElement : public FSimplePCGElement
+class FPCGAttributeFilterElement : public IPCGElement
 {
 protected:
 	virtual bool ExecuteInternal(FPCGContext* Context) const override;
diff --git a/Engine/Plugins/Experimental/PCG/Source/PCG/Public/Elements/PCGAttributeGetFromIndexElement.h b/Engine/Plugins/Experimental/PCG/Source/PCG/Public/Elements/PCGAttributeGetFromIndexElement.h
--- a/Engine/Plugins/Experimental/PCG/Source/PCG/Public/Elements/PCGAttributeGetFromIndexElement.h
+++ b/Engine/Plugins/Experimental/PCG/Source/PCG/Public/Elements/PCGAttributeGetFromIndexElement.h
@@ -31,7 +31,7 @@ public:
 	int32 Index = 0;
 };
 
-class FPCGAttributeGetFromIndexElement : public FSimplePCGElement
+class FPCGAttributeGetFromIndexElement : public IPCGElement
 {
 protected:
 	virtual bool ExecuteInternal(FPCGContext* Context) const override;
diff --git a/Engine/Plugins/Experimental/PCG/Source/PCG/Public/Elements/PCGAttributeGetFromPointIndexElement.h b/Engine/Plugins/Experimental/PCG/Source/PCG/Public/Elements/PCGAttributeGetFromPointIndexElement.h
--- a/Engine/Plugins/Experimental/PCG/Source/PCG/Public/Elements/PCGAttributeGetFromPointIndexElement.h
+++ b/Engine/Plugins/Experimental/PCG/Source/PCG/Public/Elements/PCGAttributeGetFromPointIndexElement.h
@@ -62,7 +62,7 @@ protected:
 };
 
 
-class FPCGAttributeGetFromPointIndexElement : public FSimplePCGElement
+class FPCGAttributeGetFromPointIndexElement : public IPCGElement
 {
 protected:
 	virtual bool ExecuteInternal(FPCGContext* Context) const override;
diff --git a/Engine/Plugins/Experimental/PCG/Source/PCG/Public/Elements/PCGAttributeNoise.h b/Engine/Plugins/Experimental/PCG/Source/PCG/Public/Elements/PCGAttributeNoise.h
--- a/Engine/Plugins/Experimental/PCG/Source/PCG/Public/Elements/PCGAttributeNoise.h
+++ b/Engine/Plugins/Experimental/PCG/Source/PCG/Public/Elements/PCGAttributeNoise.h
@@ -112,10 +112,9 @@ struct FPCGAttributeNoiseContext : public FPCGContext
 
 class FPCGAttributeNoiseElement : public IPCGElement
 {
-public:
-	virtual FPCGContext* Initialize(const FPCGDataCollection& InputData, TWeakObjectPtr<UPCGComponent> SourceComponent, const UPCGNode* Node) override;
 protected:
-	virtual bool ExecuteInternal(FPCGContext* Context) const;
+	virtual FPCGContext* CreateContext() override;
+	virtual bool ExecuteInternal(FPCGContext* Context) const override;
 };
 
 #if UE_ENABLE_INCLUDE_ORDER_DEPRECATED_IN_5_2
diff --git a/Engine/Plugins/Experimental/PCG/Source/PCG/Public/Elements/PCGAttributeReduceElement.h b/Engine/Plugins/Experimental/PCG/Source/PCG/Public/Elements/PCGAttributeReduceElement.h
--- a/Engine/Plugins/Experimental/PCG/Source/PCG/Public/Elements/PCGAttributeReduceElement.h
+++ b/Engine/Plugins/Experimental/PCG/Source/PCG/Public/Elements/PCGAttributeReduceElement.h
@@ -74,7 +74,7 @@ protected:
 };
 
 
-class FPCGAttributeReduceElement : public FSimplePCGElement
+class FPCGAttributeReduceElement : public IPCGElement
 {
 protected:
 	virtual bool ExecuteInternal(FPCGContext* Context) const override;
diff --git a/Engine/Plugins/Experimental/PCG/Source/PCG/Public/Elements/PCGAttributeSelectElement.h b/Engine/Plugins/Experimental/PCG/Source/PCG/Public/Elements/PCGAttributeSelectElement.h
--- a/Engine/Plugins/Experimental/PCG/Source/PCG/Public/Elements/PCGAttributeSelectElement.h
+++ b/Engine/Plugins/Experimental/PCG/Source/PCG/Public/Elements/PCGAttributeSelectElement.h
@@ -94,7 +94,7 @@ protected:
 };
 
 
-class FPCGAttributeSelectElement : public FSimplePCGElement
+class FPCGAttributeSelectElement : public IPCGElement
 {
 protected:
 	virtual bool ExecuteInternal(FPCGContext* Context) const override;
diff --git a/Engine/Plugins/Experimental/PCG/Source/PCG/Public/Elements/PCGAttributeTransferElement.h b/Engine/Plugins/Experimental/PCG/Source/PCG/Public/Elements/PCGAttributeTransferElement.h
--- a/Engine/Plugins/Experimental/PCG/Source/PCG/Public/Elements/PCGAttributeTransferElement.h
+++ b/Engine/Plugins/Experimental/PCG/Source/PCG/Public/Elements/PCGAttributeTransferElement.h
@@ -69,7 +69,7 @@ public:
 #endif
 };
 
-class FPCGAttributeTransferElement : public FSimplePCGElement
+class FPCGAttributeTransferElement : public IPCGElement
 {
 protected:
 	virtual bool ExecuteInternal(FPCGContext* Context) const override;
diff --git a/Engine/Plugins/Experimental/PCG/Source/PCG/Public/Elements/PCGBooleanSelect.h b/Engine/Plugins/Experimental/PCG/Source/PCG/Public/Elements/PCGBooleanSelect.h
--- a/Engine/Plugins/Experimental/PCG/Source/PCG/Public/Elements/PCGBooleanSelect.h
+++ b/Engine/Plugins/Experimental/PCG/Source/PCG/Public/Elements/PCGBooleanSelect.h
@@ -36,7 +36,7 @@ public:
 	bool bUseInputB = false;
 };
 
-class FPCGBooleanSelectElement : public FSimplePCGElement
+class FPCGBooleanSelectElement : public IPCGElement
 {
 protected:
 	virtual bool ExecuteInternal(FPCGContext* Context) const override;
diff --git a/Engine/Plugins/Experimental/PCG/Source/PCG/Public/Elements/PCGBranch.h b/Engine/Plugins/Experimental/PCG/Source/PCG/Public/Elements/PCGBranch.h
--- a/Engine/Plugins/Experimental/PCG/Source/PCG/Public/Elements/PCGBranch.h
+++ b/Engine/Plugins/Experimental/PCG/Source/PCG/Public/Elements/PCGBranch.h
@@ -35,7 +35,7 @@ public:
 	bool bOutputToB = false;
 };
 
-class FPCGBranchElement : public FSimplePCGElement
+class FPCGBranchElement : public IPCGElement
 {
 protected:
 	virtual bool ExecuteInternal(FPCGContext* Context) const override;
diff --git a/Engine/Plugins/Experimental/PCG/Source/PCG/Public/Elements/PCGCollapseElement.h b/Engine/Plugins/Experimental/PCG/Source/PCG/Public/Elements/PCGCollapseElement.h
--- a/Engine/Plugins/Experimental/PCG/Source/PCG/Public/Elements/PCGCollapseElement.h
+++ b/Engine/Plugins/Experimental/PCG/Source/PCG/Public/Elements/PCGCollapseElement.h
@@ -28,7 +28,7 @@ protected:
 	//~End UPCGSettings interface
 };
 
-class FPCGCollapseElement : public FSimplePCGElement
+class FPCGCollapseElement : public IPCGElement
 {
 public:
 	// Might be sampling spline/landscape or other external data, worth computing a full CRC in case we can halt change propagation/re-executions
diff --git a/Engine/Plugins/Experimental/PCG/Source/PCG/Public/Elements/PCGCopyPoints.h b/Engine/Plugins/Experimental/PCG/Source/PCG/Public/Elements/PCGCopyPoints.h
--- a/Engine/Plugins/Experimental/PCG/Source/PCG/Public/Elements/PCGCopyPoints.h
+++ b/Engine/Plugins/Experimental/PCG/Source/PCG/Public/Elements/PCGCopyPoints.h
@@ -72,7 +72,7 @@ public:
 	EPCGCopyPointsMetadataInheritanceMode AttributeInheritance = EPCGCopyPointsMetadataInheritanceMode::SourceFirst;
 };
 
-class FPCGCopyPointsElement : public FSimplePCGElement
+class FPCGCopyPointsElement : public IPCGElement
 {
 protected:
 	virtual bool ExecuteInternal(FPCGContext* Context) const override;
diff --git a/Engine/Plugins/Experimental/PCG/Source/PCG/Public/Elements/PCGCreateAttribute.h b/Engine/Plugins/Experimental/PCG/Source/PCG/Public/Elements/PCGCreateAttribute.h
--- a/Engine/Plugins/Experimental/PCG/Source/PCG/Public/Elements/PCGCreateAttribute.h
+++ b/Engine/Plugins/Experimental/PCG/Source/PCG/Public/Elements/PCGCreateAttribute.h
@@ -187,7 +187,7 @@ protected:
 	virtual bool ShouldAddAttributesPin() const override { return false; }
 };
 
-class FPCGCreateAttributeElement : public FSimplePCGElement
+class FPCGCreateAttributeElement : public IPCGElement
 {
 protected:
 	virtual bool ExecuteInternal(FPCGContext* Context) const override;
diff --git a/Engine/Plugins/Experimental/PCG/Source/PCG/Public/Elements/PCGCreatePoints.h b/Engine/Plugins/Experimental/PCG/Source/PCG/Public/Elements/PCGCreatePoints.h
--- a/Engine/Plugins/Experimental/PCG/Source/PCG/Public/Elements/PCGCreatePoints.h
+++ b/Engine/Plugins/Experimental/PCG/Source/PCG/Public/Elements/PCGCreatePoints.h
@@ -42,10 +42,10 @@ public:
 	bool bCullPointsOutsideVolume = false;
 };
 
-class FPCGCreatePointsElement : public FSimplePCGElement
+class FPCGCreatePointsElement : public IPCGElement
 {
 protected:
 	virtual bool ExecuteInternal(FPCGContext* Context) const override;
 	virtual bool IsCacheable(const UPCGSettings* InSettings) const override;
-	virtual bool ShouldComputeFullOutputDataCrc(FPCGContext* Context) const { return true; }
+	virtual bool ShouldComputeFullOutputDataCrc(FPCGContext* Context) const override { return true; }
 };
\ No newline at end of file
diff --git a/Engine/Plugins/Experimental/PCG/Source/PCG/Public/Elements/PCGCreateSpline.h b/Engine/Plugins/Experimental/PCG/Source/PCG/Public/Elements/PCGCreateSpline.h
--- a/Engine/Plugins/Experimental/PCG/Source/PCG/Public/Elements/PCGCreateSpline.h
+++ b/Engine/Plugins/Experimental/PCG/Source/PCG/Public/Elements/PCGCreateSpline.h
@@ -68,7 +68,7 @@ public:
 	EPCGAttachOptions AttachOptions = EPCGAttachOptions::Attached;
 };
 
-class FPCGCreateSplineElement : public FSimplePCGElement
+class FPCGCreateSplineElement : public IPCGElement
 {
 protected:
 	virtual bool CanExecuteOnlyOnMainThread(FPCGContext* Context) const override;
diff --git a/Engine/Plugins/Experimental/PCG/Source/PCG/Public/Elements/PCGCreateTargetActor.h b/Engine/Plugins/Experimental/PCG/Source/PCG/Public/Elements/PCGCreateTargetActor.h
--- a/Engine/Plugins/Experimental/PCG/Source/PCG/Public/Elements/PCGCreateTargetActor.h
+++ b/Engine/Plugins/Experimental/PCG/Source/PCG/Public/Elements/PCGCreateTargetActor.h
@@ -73,12 +73,11 @@ public:
 	FTransform ActorPivot;
 };
 
-class FPCGCreateTargetActorElement : public FSimplePCGElement
+class FPCGCreateTargetActorElement : public IPCGElement
 {
 protected:
 	// Since this element creates an actor, it needs to run on the game thread and cannot be cached
 	virtual bool CanExecuteOnlyOnMainThread(FPCGContext* Context) const override { return true; }
 	virtual bool IsCacheable(const UPCGSettings* Settings) const override { return false; }
-	
 	virtual bool ExecuteInternal(FPCGContext* Context) const override;
 };
\ No newline at end of file
diff --git a/Engine/Plugins/Experimental/PCG/Source/PCG/Public/Elements/PCGCullPointsOutsideActorBounds.h b/Engine/Plugins/Experimental/PCG/Source/PCG/Public/Elements/PCGCullPointsOutsideActorBounds.h
--- a/Engine/Plugins/Experimental/PCG/Source/PCG/Public/Elements/PCGCullPointsOutsideActorBounds.h
+++ b/Engine/Plugins/Experimental/PCG/Source/PCG/Public/Elements/PCGCullPointsOutsideActorBounds.h
@@ -30,7 +30,7 @@ protected:
 	//~End UPCGSettings interface
 };
 
-class FPCGCullPointsOutsideActorBoundsElement : public FSimplePCGElement
+class FPCGCullPointsOutsideActorBoundsElement : public IPCGElement
 {
 public:
 	virtual void GetDependenciesCrc(const FPCGDataCollection& InInput, const UPCGSettings* InSettings, UPCGComponent* InComponent, FPCGCrc& OutCrc) const override;
diff --git a/Engine/Plugins/Experimental/PCG/Source/PCG/Public/Elements/PCGDataFromActor.h b/Engine/Plugins/Experimental/PCG/Source/PCG/Public/Elements/PCGDataFromActor.h
--- a/Engine/Plugins/Experimental/PCG/Source/PCG/Public/Elements/PCGDataFromActor.h
+++ b/Engine/Plugins/Experimental/PCG/Source/PCG/Public/Elements/PCGDataFromActor.h
@@ -99,12 +99,12 @@ public:
 class FPCGDataFromActorElement : public IPCGElement
 {
 public:
-	virtual FPCGContext* Initialize(const FPCGDataCollection& InputData, TWeakObjectPtr<UPCGComponent> SourceComponent, const UPCGNode* Node) override;
 	virtual bool CanExecuteOnlyOnMainThread(FPCGContext* Context) const override { return true; }
 	virtual void GetDependenciesCrc(const FPCGDataCollection& InInput, const UPCGSettings* InSettings, UPCGComponent* InComponent, FPCGCrc& OutCrc) const override;
 
 protected:
-	virtual bool ExecuteInternal(FPCGContext* Context) const;
+	virtual FPCGContext* CreateContext() override;
+	virtual bool ExecuteInternal(FPCGContext* Context) const override;
 	void GatherWaitTasks(AActor* FoundActor, FPCGContext* Context, TArray<FPCGTaskId>& OutWaitTasks) const;
 	virtual void ProcessActors(FPCGContext* Context, const UPCGDataFromActorSettings* Settings, const TArray<AActor*>& FoundActors) const;
 	virtual void ProcessActor(FPCGContext* Context, const UPCGDataFromActorSettings* Settings, AActor* FoundActor) const;
diff --git a/Engine/Plugins/Experimental/PCG/Source/PCG/Public/Elements/PCGDataNum.h b/Engine/Plugins/Experimental/PCG/Source/PCG/Public/Elements/PCGDataNum.h
--- a/Engine/Plugins/Experimental/PCG/Source/PCG/Public/Elements/PCGDataNum.h
+++ b/Engine/Plugins/Experimental/PCG/Source/PCG/Public/Elements/PCGDataNum.h
@@ -31,7 +31,7 @@ public:
 	FName OutputAttributeName = NAME_None;
 };
 
-class FPCGDataNumElement : public FSimplePCGElement
+class FPCGDataNumElement : public IPCGElement
 {
 protected:
 	virtual bool ExecuteInternal(FPCGContext* Context) const override;
diff --git a/Engine/Plugins/Experimental/PCG/Source/PCG/Public/Elements/PCGDataTableRowToParamData.h b/Engine/Plugins/Experimental/PCG/Source/PCG/Public/Elements/PCGDataTableRowToParamData.h
--- a/Engine/Plugins/Experimental/PCG/Source/PCG/Public/Elements/PCGDataTableRowToParamData.h
+++ b/Engine/Plugins/Experimental/PCG/Source/PCG/Public/Elements/PCGDataTableRowToParamData.h
@@ -43,11 +43,11 @@ protected:
 	// ~End UPCGSettings interface
 };
 
-class FPCGDataTableRowToParamData : public FSimplePCGElement
+class FPCGDataTableRowToParamData : public IPCGElement
 {
 public:
 	virtual bool CanExecuteOnlyOnMainThread(FPCGContext* Context) const override { return true; }
-	virtual bool IsCacheable(const UPCGSettings* InSettings) const { return false; }
+	virtual bool IsCacheable(const UPCGSettings* InSettings) const override { return false; }
 
 protected:
 	virtual bool ExecuteInternal(FPCGContext* Context) const override;
diff --git a/Engine/Plugins/Experimental/PCG/Source/PCG/Public/Elements/PCGDebugElement.h b/Engine/Plugins/Experimental/PCG/Source/PCG/Public/Elements/PCGDebugElement.h
--- a/Engine/Plugins/Experimental/PCG/Source/PCG/Public/Elements/PCGDebugElement.h
+++ b/Engine/Plugins/Experimental/PCG/Source/PCG/Public/Elements/PCGDebugElement.h
@@ -38,14 +38,14 @@ public:
 	TSoftObjectPtr<AActor> TargetActor;
 };
 
-class FPCGDebugElement : public FSimplePCGElement
+class FPCGDebugElement : public IPCGElement
 {
 public:
 	virtual bool IsCacheable(const UPCGSettings* InSettings) const override { return false; }
-	virtual bool CanExecuteOnlyOnMainThread(FPCGContext* Context) const { return true; }
+	virtual bool CanExecuteOnlyOnMainThread(FPCGContext* Context) const override { return true; }
 
 protected:
-	virtual bool ExecuteInternal(FPCGContext* Context) const override;	
+	virtual bool ExecuteInternal(FPCGContext* Context) const override;
 };
 
 #if UE_ENABLE_INCLUDE_ORDER_DEPRECATED_IN_5_2
diff --git a/Engine/Plugins/Experimental/PCG/Source/PCG/Public/Elements/PCGDensityFilter.h b/Engine/Plugins/Experimental/PCG/Source/PCG/Public/Elements/PCGDensityFilter.h
--- a/Engine/Plugins/Experimental/PCG/Source/PCG/Public/Elements/PCGDensityFilter.h
+++ b/Engine/Plugins/Experimental/PCG/Source/PCG/Public/Elements/PCGDensityFilter.h
@@ -41,7 +41,7 @@ public:
 #endif
 };
 
-class FPCGDensityFilterElement : public FSimplePCGElement
+class FPCGDensityFilterElement : public IPCGElement
 {
 protected:
 	virtual bool ExecuteInternal(FPCGContext* Context) const override;
diff --git a/Engine/Plugins/Experimental/PCG/Source/PCG/Public/Elements/PCGDensityRemap.h b/Engine/Plugins/Experimental/PCG/Source/PCG/Public/Elements/PCGDensityRemap.h
--- a/Engine/Plugins/Experimental/PCG/Source/PCG/Public/Elements/PCGDensityRemap.h
+++ b/Engine/Plugins/Experimental/PCG/Source/PCG/Public/Elements/PCGDensityRemap.h
@@ -38,7 +38,7 @@ public:
 	bool bMultiplyDensity = true;
 };
 
-class FPCGLinearDensityRemapElement : public FSimplePCGElement
+class FPCGLinearDensityRemapElement : public IPCGElement
 {
 protected:
 	virtual bool ExecuteInternal(FPCGContext* Context) const;
diff --git a/Engine/Plugins/Experimental/PCG/Source/PCG/Public/Elements/PCGDifferenceElement.h b/Engine/Plugins/Experimental/PCG/Source/PCG/Public/Elements/PCGDifferenceElement.h
--- a/Engine/Plugins/Experimental/PCG/Source/PCG/Public/Elements/PCGDifferenceElement.h
+++ b/Engine/Plugins/Experimental/PCG/Source/PCG/Public/Elements/PCGDifferenceElement.h
@@ -59,7 +59,7 @@ public:
 #endif
 };
 
-class FPCGDifferenceElement : public FSimplePCGElement
+class FPCGDifferenceElement : public IPCGElement
 {
 protected:
 	virtual bool ExecuteInternal(FPCGContext* Context) const override;
diff --git a/Engine/Plugins/Experimental/PCG/Source/PCG/Public/Elements/PCGFilterByTag.h b/Engine/Plugins/Experimental/PCG/Source/PCG/Public/Elements/PCGFilterByTag.h
--- a/Engine/Plugins/Experimental/PCG/Source/PCG/Public/Elements/PCGFilterByTag.h
+++ b/Engine/Plugins/Experimental/PCG/Source/PCG/Public/Elements/PCGFilterByTag.h
@@ -46,7 +46,7 @@ public:
 	FString SelectedTags;
 };
 
-class FPCGFilterByTagElement : public FSimplePCGElement
+class FPCGFilterByTagElement : public IPCGElement
 {
 protected:
 	virtual bool ExecuteInternal(FPCGContext* Context) const override;
diff --git a/Engine/Plugins/Experimental/PCG/Source/PCG/Public/Elements/PCGFilterByType.h b/Engine/Plugins/Experimental/PCG/Source/PCG/Public/Elements/PCGFilterByType.h
--- a/Engine/Plugins/Experimental/PCG/Source/PCG/Public/Elements/PCGFilterByType.h
+++ b/Engine/Plugins/Experimental/PCG/Source/PCG/Public/Elements/PCGFilterByType.h
@@ -33,8 +33,8 @@ public:
 	EPCGDataType TargetType = EPCGDataType::Any;
 };
 
-class FPCGFilterByTypeElement : public FSimplePCGElement
+class FPCGFilterByTypeElement : public IPCGElement
 {
 protected:
-	virtual bool ExecuteInternal(FPCGContext* InContext) const;
+	virtual bool ExecuteInternal(FPCGContext* InContext) const override;
 };
diff --git a/Engine/Plugins/Experimental/PCG/Source/PCG/Public/Elements/PCGGather.h b/Engine/Plugins/Experimental/PCG/Source/PCG/Public/Elements/PCGGather.h
--- a/Engine/Plugins/Experimental/PCG/Source/PCG/Public/Elements/PCGGather.h
+++ b/Engine/Plugins/Experimental/PCG/Source/PCG/Public/Elements/PCGGather.h
@@ -31,7 +31,7 @@ protected:
 	//~End UPCGSettings interface
 };
 
-class FPCGGatherElement : public FSimplePCGElement
+class FPCGGatherElement : public IPCGElement
 {
 protected:
 	virtual bool ExecuteInternal(FPCGContext* Context) const override;
diff --git a/Engine/Plugins/Experimental/PCG/Source/PCG/Public/Elements/PCGGetBounds.h b/Engine/Plugins/Experimental/PCG/Source/PCG/Public/Elements/PCGGetBounds.h
--- a/Engine/Plugins/Experimental/PCG/Source/PCG/Public/Elements/PCGGetBounds.h
+++ b/Engine/Plugins/Experimental/PCG/Source/PCG/Public/Elements/PCGGetBounds.h
@@ -30,7 +30,7 @@ protected:
 	//~End UPCGSettings interface
 };
 
-class FPCGGetBoundsElement : public FSimplePCGElement
+class FPCGGetBoundsElement : public IPCGElement
 {
 protected:
 	virtual bool ExecuteInternal(FPCGContext* Context) const override;
diff --git a/Engine/Plugins/Experimental/PCG/Source/PCG/Public/Elements/PCGHiGenGridSize.h b/Engine/Plugins/Experimental/PCG/Source/PCG/Public/Elements/PCGHiGenGridSize.h
--- a/Engine/Plugins/Experimental/PCG/Source/PCG/Public/Elements/PCGHiGenGridSize.h
+++ b/Engine/Plugins/Experimental/PCG/Source/PCG/Public/Elements/PCGHiGenGridSize.h
@@ -42,7 +42,7 @@ protected:
 	EPCGHiGenGrid HiGenGridSize = EPCGHiGenGrid::Grid256;
 };
 
-class FPCGHiGenGridSizeElement : public FSimplePCGElement
+class FPCGHiGenGridSizeElement : public IPCGElement
 {
 public:
 	virtual void GetDependenciesCrc(const FPCGDataCollection& InInput, const UPCGSettings* InSettings, UPCGComponent* InComponent, FPCGCrc& OutCrc) const override;
diff --git a/Engine/Plugins/Experimental/PCG/Source/PCG/Public/Elements/PCGIndirectionElement.h b/Engine/Plugins/Experimental/PCG/Source/PCG/Public/Elements/PCGIndirectionElement.h
--- a/Engine/Plugins/Experimental/PCG/Source/PCG/Public/Elements/PCGIndirectionElement.h
+++ b/Engine/Plugins/Experimental/PCG/Source/PCG/Public/Elements/PCGIndirectionElement.h
@@ -47,12 +47,12 @@ struct FPCGIndirectionContext : public FPCGContext
 class FPCGIndirectionElement : public IPCGElement
 {
 public:
-	virtual FPCGContext* Initialize(const FPCGDataCollection& InputData, TWeakObjectPtr<UPCGComponent> SourceComponent, const UPCGNode* Node) override;
 	// TODO: investigate how we could make this cacheable, might require to pass in context instead of settings
 	virtual bool IsCacheable(const UPCGSettings* InSettings) const override { return false; }
 	virtual bool CanExecuteOnlyOnMainThread(FPCGContext* Context) const override;
 
 protected:
+	virtual FPCGContext* CreateContext() override;
 	virtual bool PrepareDataInternal(FPCGContext* InContext) const override;
 	virtual bool ExecuteInternal(FPCGContext* InContext) const override;
 };
\ No newline at end of file
diff --git a/Engine/Plugins/Experimental/PCG/Source/PCG/Public/Elements/PCGInnerIntersectionElement.h b/Engine/Plugins/Experimental/PCG/Source/PCG/Public/Elements/PCGInnerIntersectionElement.h
--- a/Engine/Plugins/Experimental/PCG/Source/PCG/Public/Elements/PCGInnerIntersectionElement.h
+++ b/Engine/Plugins/Experimental/PCG/Source/PCG/Public/Elements/PCGInnerIntersectionElement.h
@@ -42,10 +42,10 @@ public:
 #endif
 };
 
-class FPCGInnerIntersectionElement : public FSimplePCGElement
+class FPCGInnerIntersectionElement : public IPCGElement
 {
 protected:
-	virtual bool ExecuteInternal(FPCGContext* Context) const;
+	virtual bool ExecuteInternal(FPCGContext* Context) const override;
 };
 
 #if UE_ENABLE_INCLUDE_ORDER_DEPRECATED_IN_5_2
diff --git a/Engine/Plugins/Experimental/PCG/Source/PCG/Public/Elements/PCGMakeConcreteElement.h b/Engine/Plugins/Experimental/PCG/Source/PCG/Public/Elements/PCGMakeConcreteElement.h
--- a/Engine/Plugins/Experimental/PCG/Source/PCG/Public/Elements/PCGMakeConcreteElement.h
+++ b/Engine/Plugins/Experimental/PCG/Source/PCG/Public/Elements/PCGMakeConcreteElement.h
@@ -28,7 +28,7 @@ protected:
 	//~End UPCGSettings interface
 };
 
-class FPCGMakeConcreteElement : public FSimplePCGElement
+class FPCGMakeConcreteElement : public IPCGElement
 {
 public:
 	// Might be sampling spline/landscape or other external data, worth computing a full CRC in case we can halt change propagation/re-executions
diff --git a/Engine/Plugins/Experimental/PCG/Source/PCG/Public/Elements/PCGMatchAndSetAttributes.h b/Engine/Plugins/Experimental/PCG/Source/PCG/Public/Elements/PCGMatchAndSetAttributes.h
--- a/Engine/Plugins/Experimental/PCG/Source/PCG/Public/Elements/PCGMatchAndSetAttributes.h
+++ b/Engine/Plugins/Experimental/PCG/Source/PCG/Public/Elements/PCGMatchAndSetAttributes.h
@@ -66,7 +66,7 @@ public:
 	FName WeightAttribute = NAME_None;
 };
 
-class FPCGMatchAndSetAttributesElement : public FSimplePCGElement
+class FPCGMatchAndSetAttributesElement : public IPCGElement
 {
 public:
 	virtual bool ExecuteInternal(FPCGContext* Context) const override;
diff --git a/Engine/Plugins/Experimental/PCG/Source/PCG/Public/Elements/PCGMergeAttributes.h b/Engine/Plugins/Experimental/PCG/Source/PCG/Public/Elements/PCGMergeAttributes.h
--- a/Engine/Plugins/Experimental/PCG/Source/PCG/Public/Elements/PCGMergeAttributes.h
+++ b/Engine/Plugins/Experimental/PCG/Source/PCG/Public/Elements/PCGMergeAttributes.h
@@ -28,7 +28,7 @@ protected:
 	//~End UPCGSettings interface
 };
 
-class FPCGMergeAttributesElement : public FSimplePCGElement
+class FPCGMergeAttributesElement : public IPCGElement
 {
 protected:
 	virtual bool ExecuteInternal(FPCGContext* Context) const override;
diff --git a/Engine/Plugins/Experimental/PCG/Source/PCG/Public/Elements/PCGMergeElement.h b/Engine/Plugins/Experimental/PCG/Source/PCG/Public/Elements/PCGMergeElement.h
--- a/Engine/Plugins/Experimental/PCG/Source/PCG/Public/Elements/PCGMergeElement.h
+++ b/Engine/Plugins/Experimental/PCG/Source/PCG/Public/Elements/PCGMergeElement.h
@@ -34,7 +34,7 @@ public:
 	bool bMergeMetadata = true;
 };
 
-class FPCGMergeElement : public FSimplePCGElement
+class FPCGMergeElement : public IPCGElement
 {
 protected:
 	virtual bool ExecuteInternal(FPCGContext* Context) const override;
diff --git a/Engine/Plugins/Experimental/PCG/Source/PCG/Public/Elements/PCGNumberOfElements.h b/Engine/Plugins/Experimental/PCG/Source/PCG/Public/Elements/PCGNumberOfElements.h
--- a/Engine/Plugins/Experimental/PCG/Source/PCG/Public/Elements/PCGNumberOfElements.h
+++ b/Engine/Plugins/Experimental/PCG/Source/PCG/Public/Elements/PCGNumberOfElements.h
@@ -78,7 +78,7 @@ protected:
 
 
 template <typename DataType>
-class FPCGNumberOfElementsBaseElement : public FSimplePCGElement
+class FPCGNumberOfElementsBaseElement : public IPCGElement
 {
 protected:
 	virtual bool ExecuteInternal(FPCGContext* Context) const override;
diff --git a/Engine/Plugins/Experimental/PCG/Source/PCG/Public/Elements/PCGOuterIntersectionElement.h b/Engine/Plugins/Experimental/PCG/Source/PCG/Public/Elements/PCGOuterIntersectionElement.h
--- a/Engine/Plugins/Experimental/PCG/Source/PCG/Public/Elements/PCGOuterIntersectionElement.h
+++ b/Engine/Plugins/Experimental/PCG/Source/PCG/Public/Elements/PCGOuterIntersectionElement.h
@@ -56,7 +56,7 @@ public:
 #endif // WITH_EDITORONLY_DATA
 };
 
-class FPCGOuterIntersectionElement : public FSimplePCGElement
+class FPCGOuterIntersectionElement : public IPCGElement
 {
 protected:
 	virtual bool ExecuteInternal(FPCGContext* Context) const override;
diff --git a/Engine/Plugins/Experimental/PCG/Source/PCG/Public/Elements/PCGPointFilter.h b/Engine/Plugins/Experimental/PCG/Source/PCG/Public/Elements/PCGPointFilter.h
--- a/Engine/Plugins/Experimental/PCG/Source/PCG/Public/Elements/PCGPointFilter.h
+++ b/Engine/Plugins/Experimental/PCG/Source/PCG/Public/Elements/PCGPointFilter.h
@@ -229,7 +229,7 @@ public:
 	FPCGPointFilterThresholdSettings MaxThreshold;
 };
 
-class FPCGPointFilterElementBase : public FSimplePCGElement
+class FPCGPointFilterElementBase : public IPCGElement
 {
 protected:
 	bool DoFiltering(FPCGContext* Context, EPCGPointFilterOperator InOperation, const FPCGAttributePropertyInputSelector& TargetAttribute, const FPCGPointFilterThresholdSettings& FirstThreshold, const FPCGPointFilterThresholdSettings* SecondThreshold = nullptr) const;
diff --git a/Engine/Plugins/Experimental/PCG/Source/PCG/Public/Elements/PCGPointFromMeshElement.h b/Engine/Plugins/Experimental/PCG/Source/PCG/Public/Elements/PCGPointFromMeshElement.h
--- a/Engine/Plugins/Experimental/PCG/Source/PCG/Public/Elements/PCGPointFromMeshElement.h
+++ b/Engine/Plugins/Experimental/PCG/Source/PCG/Public/Elements/PCGPointFromMeshElement.h
@@ -40,8 +40,8 @@ public:
 	FName MeshPathAttributeName = NAME_None;
 };
 
-class FPCGPointFromMeshElement : public FSimplePCGElement
+class FPCGPointFromMeshElement : public IPCGElement
 {
 protected:
-	virtual bool ExecuteInternal(FPCGContext* Context) const;
+	virtual bool ExecuteInternal(FPCGContext* Context) const override;
 };
diff --git a/Engine/Plugins/Experimental/PCG/Source/PCG/Public/Elements/PCGPointMatchAndSet.h b/Engine/Plugins/Experimental/PCG/Source/PCG/Public/Elements/PCGPointMatchAndSet.h
--- a/Engine/Plugins/Experimental/PCG/Source/PCG/Public/Elements/PCGPointMatchAndSet.h
+++ b/Engine/Plugins/Experimental/PCG/Source/PCG/Public/Elements/PCGPointMatchAndSet.h
@@ -82,7 +82,7 @@ private:
 #endif
 };
 
-class FPCGPointMatchAndSetElement : public FSimplePCGElement
+class FPCGPointMatchAndSetElement : public IPCGElement
 {
 public:
 	virtual bool ExecuteInternal(FPCGContext* Context) const override;
diff --git a/Engine/Plugins/Experimental/PCG/Source/PCG/Public/Elements/PCGPointProcessingElementBase.h b/Engine/Plugins/Experimental/PCG/Source/PCG/Public/Elements/PCGPointProcessingElementBase.h
--- a/Engine/Plugins/Experimental/PCG/Source/PCG/Public/Elements/PCGPointProcessingElementBase.h
+++ b/Engine/Plugins/Experimental/PCG/Source/PCG/Public/Elements/PCGPointProcessingElementBase.h
@@ -8,10 +8,9 @@
 struct FPCGPoint;
 struct FPCGTaggedData;
 
-class PCG_API FPCGPointProcessingElementBase : public FSimplePCGElement
+class PCG_API FPCGPointProcessingElementBase : public IPCGElement
 {
 protected:
-	virtual bool ExecuteInternal(FPCGContext* Context) const = 0;
 	void ProcessPoints(FPCGContext* Context, const TArray<FPCGTaggedData>& Inputs, TArray<FPCGTaggedData>& Outputs, const TFunction<bool(const FPCGPoint&, FPCGPoint&)>& PointFunc) const;
 };
 
diff --git a/Engine/Plugins/Experimental/PCG/Source/PCG/Public/Elements/PCGProjectionElement.h b/Engine/Plugins/Experimental/PCG/Source/PCG/Public/Elements/PCGProjectionElement.h
--- a/Engine/Plugins/Experimental/PCG/Source/PCG/Public/Elements/PCGProjectionElement.h
+++ b/Engine/Plugins/Experimental/PCG/Source/PCG/Public/Elements/PCGProjectionElement.h
@@ -45,10 +45,10 @@ protected:
 	//~End UPCGSettings interface
 };
 
-class FPCGProjectionElement : public FSimplePCGElement
+class FPCGProjectionElement : public IPCGElement
 {
 protected:
-	virtual bool ExecuteInternal(FPCGContext* Context) const;
+	virtual bool ExecuteInternal(FPCGContext* Context) const override;
 };
 
 #if UE_ENABLE_INCLUDE_ORDER_DEPRECATED_IN_5_2
diff --git a/Engine/Plugins/Experimental/PCG/Source/PCG/Public/Elements/PCGPropertyToParamData.h b/Engine/Plugins/Experimental/PCG/Source/PCG/Public/Elements/PCGPropertyToParamData.h
--- a/Engine/Plugins/Experimental/PCG/Source/PCG/Public/Elements/PCGPropertyToParamData.h
+++ b/Engine/Plugins/Experimental/PCG/Source/PCG/Public/Elements/PCGPropertyToParamData.h
@@ -96,11 +96,11 @@ private:
 	bool bIncludeChildren_DEPRECATED = false;
 };
 
-class FPCGPropertyToParamDataElement : public FSimplePCGElement
+class FPCGPropertyToParamDataElement : public IPCGElement
 {
 public:
 	virtual bool CanExecuteOnlyOnMainThread(FPCGContext* Context) const override { return true; }
 	virtual bool IsCacheable(const UPCGSettings* InSettings) const override { return !CastChecked<UPCGPropertyToParamDataSettings>(InSettings)->bAlwaysRequeryActors; }
 protected:
-	virtual bool ExecuteInternal(FPCGContext* Context) const;
+	virtual bool ExecuteInternal(FPCGContext* Context) const override;
 };
diff --git a/Engine/Plugins/Experimental/PCG/Source/PCG/Public/Elements/PCGReroute.h b/Engine/Plugins/Experimental/PCG/Source/PCG/Public/Elements/PCGReroute.h
--- a/Engine/Plugins/Experimental/PCG/Source/PCG/Public/Elements/PCGReroute.h
+++ b/Engine/Plugins/Experimental/PCG/Source/PCG/Public/Elements/PCGReroute.h
@@ -31,7 +31,7 @@ protected:
 #endif
 };
 
-class PCG_API FPCGRerouteElement : public FSimplePCGElement
+class PCG_API FPCGRerouteElement : public IPCGElement
 {
 public:
 	virtual bool IsCacheable(const UPCGSettings* InSettings) const override { return true; }
diff --git a/Engine/Plugins/Experimental/PCG/Source/PCG/Public/Elements/PCGSanityCheckPointData.h b/Engine/Plugins/Experimental/PCG/Source/PCG/Public/Elements/PCGSanityCheckPointData.h
--- a/Engine/Plugins/Experimental/PCG/Source/PCG/Public/Elements/PCGSanityCheckPointData.h
+++ b/Engine/Plugins/Experimental/PCG/Source/PCG/Public/Elements/PCGSanityCheckPointData.h
@@ -34,10 +34,10 @@ public:
 	int32 MaxPointCount = 100;
 };
 
-class FPCGSanityCheckPointDataElement : public FSimplePCGElement
+class FPCGSanityCheckPointDataElement : public IPCGElement
 {
 protected:
-	virtual bool ExecuteInternal(FPCGContext* Context) const;
+	virtual bool ExecuteInternal(FPCGContext* Context) const override;
 };
 
 #if UE_ENABLE_INCLUDE_ORDER_DEPRECATED_IN_5_2
diff --git a/Engine/Plugins/Experimental/PCG/Source/PCG/Public/Elements/PCGSelectPoints.h b/Engine/Plugins/Experimental/PCG/Source/PCG/Public/Elements/PCGSelectPoints.h
--- a/Engine/Plugins/Experimental/PCG/Source/PCG/Public/Elements/PCGSelectPoints.h
+++ b/Engine/Plugins/Experimental/PCG/Source/PCG/Public/Elements/PCGSelectPoints.h
@@ -38,7 +38,7 @@ public:
 #endif
 };
 
-class FPCGSelectPointsElement : public FSimplePCGElement
+class FPCGSelectPointsElement : public IPCGElement
 {
 protected:
 	virtual bool ExecuteInternal(FPCGContext* Context) const override;
diff --git a/Engine/Plugins/Experimental/PCG/Source/PCG/Public/Elements/PCGSelfPruning.h b/Engine/Plugins/Experimental/PCG/Source/PCG/Public/Elements/PCGSelfPruning.h
--- a/Engine/Plugins/Experimental/PCG/Source/PCG/Public/Elements/PCGSelfPruning.h
+++ b/Engine/Plugins/Experimental/PCG/Source/PCG/Public/Elements/PCGSelfPruning.h
@@ -50,7 +50,7 @@ protected:
 	// ~End UPCGSettings interface
 };
 
-class FPCGSelfPruningElement : public FSimplePCGElement
+class FPCGSelfPruningElement : public IPCGElement
 {
 protected:
 	virtual bool ExecuteInternal(FPCGContext* Context) const override;
diff --git a/Engine/Plugins/Experimental/PCG/Source/PCG/Public/Elements/PCGSortPoints.h b/Engine/Plugins/Experimental/PCG/Source/PCG/Public/Elements/PCGSortPoints.h
--- a/Engine/Plugins/Experimental/PCG/Source/PCG/Public/Elements/PCGSortPoints.h
+++ b/Engine/Plugins/Experimental/PCG/Source/PCG/Public/Elements/PCGSortPoints.h
@@ -45,7 +45,7 @@ public:
 	EPCGSortMethod SortMethod = EPCGSortMethod::Ascending;
 };
 
-class FPCGSortPointsElement : public FSimplePCGElement
+class FPCGSortPointsElement : public IPCGElement
 {
 protected:
 	virtual bool ExecuteInternal(FPCGContext* Context) const override;
diff --git a/Engine/Plugins/Experimental/PCG/Source/PCG/Public/Elements/PCGSplineSampler.h b/Engine/Plugins/Experimental/PCG/Source/PCG/Public/Elements/PCGSplineSampler.h
--- a/Engine/Plugins/Experimental/PCG/Source/PCG/Public/Elements/PCGSplineSampler.h
+++ b/Engine/Plugins/Experimental/PCG/Source/PCG/Public/Elements/PCGSplineSampler.h
@@ -177,7 +177,7 @@ public:
 	FPCGSplineSamplerParams SamplerParams;
 };
 
-class FPCGSplineSamplerElement : public FSimplePCGElement
+class FPCGSplineSamplerElement : public IPCGElement
 {
 public:
 	// Worth computing a full CRC in case we can halt change propagation/re-executions
diff --git a/Engine/Plugins/Experimental/PCG/Source/PCG/Public/Elements/PCGStaticMeshSpawner.h b/Engine/Plugins/Experimental/PCG/Source/PCG/Public/Elements/PCGStaticMeshSpawner.h
--- a/Engine/Plugins/Experimental/PCG/Source/PCG/Public/Elements/PCGStaticMeshSpawner.h
+++ b/Engine/Plugins/Experimental/PCG/Source/PCG/Public/Elements/PCGStaticMeshSpawner.h
@@ -108,11 +108,11 @@ protected:
 class FPCGStaticMeshSpawnerElement : public IPCGElement
 {
 public:
-	virtual FPCGContext* Initialize(const FPCGDataCollection& InputData, TWeakObjectPtr<UPCGComponent> SourceComponent, const UPCGNode* Node) override;
-	virtual bool CanExecuteOnlyOnMainThread(FPCGContext* Context) const;
+	virtual bool CanExecuteOnlyOnMainThread(FPCGContext* Context) const override;
 	virtual bool IsCacheable(const UPCGSettings* InSettings) const override { return false; }
 
 protected:
+	virtual FPCGContext* CreateContext() override;
 	virtual bool PrepareDataInternal(FPCGContext* Context) const override;
 	virtual bool ExecuteInternal(FPCGContext* Context) const override;	
 	void SpawnStaticMeshInstances(FPCGContext* Context, const FPCGMeshInstanceList& InstanceList, AActor* TargetActor, const FPCGPackedCustomData& PackedCustomData) const;
diff --git a/Engine/Plugins/Experimental/PCG/Source/PCG/Public/Elements/PCGTextureSampler.h b/Engine/Plugins/Experimental/PCG/Source/PCG/Public/Elements/PCGTextureSampler.h
--- a/Engine/Plugins/Experimental/PCG/Source/PCG/Public/Elements/PCGTextureSampler.h
+++ b/Engine/Plugins/Experimental/PCG/Source/PCG/Public/Elements/PCGTextureSampler.h
@@ -80,12 +80,12 @@ struct PCG_API FPCGTextureSamplerContext : public FPCGContext
 	bool bTextureReadbackDone = false;
 };
 
-class FPCGTextureSamplerElement : public FSimplePCGElement
+class FPCGTextureSamplerElement : public IPCGElement
 {
 public:
-	virtual FPCGContext* Initialize(const FPCGDataCollection& InputData, TWeakObjectPtr<UPCGComponent> SourceComponent, const UPCGNode* Node) override;
 	virtual void GetDependenciesCrc(const FPCGDataCollection& InInput, const UPCGSettings* InSettings, UPCGComponent* InComponent, FPCGCrc& OutCrc) const override;
 
 protected:
+	virtual FPCGContext* CreateContext() override;
 	virtual bool ExecuteInternal(FPCGContext* InContext) const override;
 };
diff --git a/Engine/Plugins/Experimental/PCG/Source/PCG/Public/Elements/PCGTimeSlicedElementBase.h b/Engine/Plugins/Experimental/PCG/Source/PCG/Public/Elements/PCGTimeSlicedElementBase.h
--- a/Engine/Plugins/Experimental/PCG/Source/PCG/Public/Elements/PCGTimeSlicedElementBase.h
+++ b/Engine/Plugins/Experimental/PCG/Source/PCG/Public/Elements/PCGTimeSlicedElementBase.h
@@ -144,7 +144,7 @@ private:
  * @tparam PerIterationStateT Struct type of the "per-iteration" data state
  */
 template <typename PerExecutionStateT, typename PerIterationStateT>
-class TPCGTimeSlicedElementBase : public FSimplePCGElement
+class TPCGTimeSlicedElementBase : public IPCGElement
 {
 public:
 	// Aliases, for ease of use with the template
diff --git a/Engine/Plugins/Experimental/PCG/Source/PCG/Public/Elements/PCGTransformPoints.h b/Engine/Plugins/Experimental/PCG/Source/PCG/Public/Elements/PCGTransformPoints.h
--- a/Engine/Plugins/Experimental/PCG/Source/PCG/Public/Elements/PCGTransformPoints.h
+++ b/Engine/Plugins/Experimental/PCG/Source/PCG/Public/Elements/PCGTransformPoints.h
@@ -73,7 +73,7 @@ public:
 	bool bRecomputeSeed = false;
 };
 
-class FPCGTransformPointsElement : public FSimplePCGElement
+class FPCGTransformPointsElement : public IPCGElement
 {
 protected:
 	virtual bool ExecuteInternal(FPCGContext* Context) const override;
diff --git a/Engine/Plugins/Experimental/PCG/Source/PCG/Public/Elements/PCGUnionElement.h b/Engine/Plugins/Experimental/PCG/Source/PCG/Public/Elements/PCGUnionElement.h
--- a/Engine/Plugins/Experimental/PCG/Source/PCG/Public/Elements/PCGUnionElement.h
+++ b/Engine/Plugins/Experimental/PCG/Source/PCG/Public/Elements/PCGUnionElement.h
@@ -35,10 +35,10 @@ public:
 	EPCGUnionDensityFunction DensityFunction = EPCGUnionDensityFunction::Maximum;
 };
 
-class FPCGUnionElement : public FSimplePCGElement
+class FPCGUnionElement : public IPCGElement
 {
 protected:
-	virtual bool ExecuteInternal(FPCGContext* Context) const;
+	virtual bool ExecuteInternal(FPCGContext* Context) const override;
 };
 
 #if UE_ENABLE_INCLUDE_ORDER_DEPRECATED_IN_5_2
diff --git a/Engine/Plugins/Experimental/PCG/Source/PCG/Public/Elements/PCGUserParameterGet.h b/Engine/Plugins/Experimental/PCG/Source/PCG/Public/Elements/PCGUserParameterGet.h
--- a/Engine/Plugins/Experimental/PCG/Source/PCG/Public/Elements/PCGUserParameterGet.h
+++ b/Engine/Plugins/Experimental/PCG/Source/PCG/Public/Elements/PCGUserParameterGet.h
@@ -43,7 +43,7 @@ protected:
 	//~End UPCGSettings interface
 };
 
-class PCG_API FPCGUserParameterGetElement : public FSimplePCGElement
+class PCG_API FPCGUserParameterGetElement : public IPCGElement
 {
 public:
 	virtual bool IsCacheable(const UPCGSettings* InSettings) const override { return false; }
diff --git a/Engine/Plugins/Experimental/PCG/Source/PCG/Public/Elements/PCGVolumeSampler.h b/Engine/Plugins/Experimental/PCG/Source/PCG/Public/Elements/PCGVolumeSampler.h
--- a/Engine/Plugins/Experimental/PCG/Source/PCG/Public/Elements/PCGVolumeSampler.h
+++ b/Engine/Plugins/Experimental/PCG/Source/PCG/Public/Elements/PCGVolumeSampler.h
@@ -67,12 +67,12 @@ protected:
 	//~End UPCGSettings interface
 };
 
-class FPCGVolumeSamplerElement : public FSimplePCGElement
+class FPCGVolumeSamplerElement : public IPCGElement
 {
 public:
 	// Might be sampling external data like brush, worth computing a full CRC in case we can halt change propagation/re-executions
 	virtual bool ShouldComputeFullOutputDataCrc(FPCGContext* Context) const override { return true; }
 
 protected:
-	bool ExecuteInternal(FPCGContext* Context) const override;
+	virtual bool ExecuteInternal(FPCGContext* Context) const override;
 };
diff --git a/Engine/Plugins/Experimental/PCG/Source/PCG/Public/Elements/PCGWorldQuery.h b/Engine/Plugins/Experimental/PCG/Source/PCG/Public/Elements/PCGWorldQuery.h
--- a/Engine/Plugins/Experimental/PCG/Source/PCG/Public/Elements/PCGWorldQuery.h
+++ b/Engine/Plugins/Experimental/PCG/Source/PCG/Public/Elements/PCGWorldQuery.h
@@ -33,12 +33,12 @@ public:
 	FPCGWorldVolumetricQueryParams QueryParams;
 };
 
-class FPCGWorldVolumetricQueryElement : public FSimplePCGElement
+class FPCGWorldVolumetricQueryElement : public IPCGElement
 {
 public:
-	virtual bool IsCacheable(const UPCGSettings* InSettings) const { return false; }
+	virtual bool IsCacheable(const UPCGSettings* InSettings) const override { return false; }
 protected:
-	virtual bool ExecuteInternal(FPCGContext* Context) const;
+	virtual bool ExecuteInternal(FPCGContext* Context) const override;
 };
 
 UCLASS(BlueprintType, ClassGroup = (Procedural))
@@ -67,10 +67,10 @@ public:
 	FPCGWorldRayHitQueryParams QueryParams;
 };
 
-class FPCGWorldRayHitQueryElement : public FSimplePCGElement
+class FPCGWorldRayHitQueryElement : public IPCGElement
 {
 public:
-	virtual bool IsCacheable(const UPCGSettings* InSettings) const { return false; }
+	virtual bool IsCacheable(const UPCGSettings* InSettings) const override { return false; }
 protected:
-	virtual bool ExecuteInternal(FPCGContext* Context) const;
+	virtual bool ExecuteInternal(FPCGContext* Context) const override;
 };
diff --git a/Engine/Plugins/Experimental/PCG/Source/PCG/Public/PCGElement.h b/Engine/Plugins/Experimental/PCG/Source/PCG/Public/PCGElement.h
--- a/Engine/Plugins/Experimental/PCG/Source/PCG/Public/PCGElement.h
+++ b/Engine/Plugins/Experimental/PCG/Source/PCG/Public/PCGElement.h
@@ -130,14 +130,8 @@ private:
 	void CleanupAndValidateOutput(FPCGContext* Context) const;
 };
 
-/**
-* Basic PCG element class for elements that do not store any intermediate data in the context
-*/
-class PCG_API FSimplePCGElement : public IPCGElement
-{
-	// TODO: Will be scrubbed and cleaned up in a following CL
-	// TODO: In the same CL, go through and have classes with simple initializations replaced with CreateContext only
-};
+class UE_DEPRECATED(5.4, "This class has been deprecated. Please inherit from IPCGElement directly.") PCG_API FSimplePCGElement : public IPCGElement
+{};
 
 #if UE_ENABLE_INCLUDE_ORDER_DEPRECATED_IN_5_2
 #include "Misc/OutputDeviceRedirector.h"
diff --git a/Engine/Plugins/Experimental/PCG/Source/PCG/Public/PCGInputOutputSettings.h b/Engine/Plugins/Experimental/PCG/Source/PCG/Public/PCGInputOutputSettings.h
--- a/Engine/Plugins/Experimental/PCG/Source/PCG/Public/PCGInputOutputSettings.h
+++ b/Engine/Plugins/Experimental/PCG/Source/PCG/Public/PCGInputOutputSettings.h
@@ -16,7 +16,7 @@ namespace PCGInputOutputConstants
 	const FName DefaultNewCustomPinName = TEXT("NewPin");
 }
 
-class PCG_API FPCGInputOutputElement : public FSimplePCGElement
+class PCG_API FPCGInputOutputElement : public IPCGElement
 {
 protected:
 	virtual bool ExecuteInternal(FPCGContext* Context) const override;
diff --git a/Engine/Plugins/Experimental/PCG/Source/PCG/Public/PCGSettings.h b/Engine/Plugins/Experimental/PCG/Source/PCG/Public/PCGSettings.h
--- a/Engine/Plugins/Experimental/PCG/Source/PCG/Public/PCGSettings.h
+++ b/Engine/Plugins/Experimental/PCG/Source/PCG/Public/PCGSettings.h
@@ -477,11 +477,8 @@ protected:
 	virtual FPCGElementPtr CreateElement() const override;
 };
 
-class PCG_API FPCGTrivialElement : public FSimplePCGElement
+class PCG_API FPCGTrivialElement : public IPCGElement
 {
-public:
-	virtual bool IsCacheable(const UPCGSettings* InSettings) const override { return true; }
-
 protected:
 	virtual bool ExecuteInternal(FPCGContext* Context) const override;
 	virtual bool IsPassthrough(const UPCGSettings* InSettings) const override { return true; }
diff --git a/Engine/Plugins/Experimental/PCG/Source/PCG/Public/PCGSubgraph.h b/Engine/Plugins/Experimental/PCG/Source/PCG/Public/PCGSubgraph.h
--- a/Engine/Plugins/Experimental/PCG/Source/PCG/Public/PCGSubgraph.h
+++ b/Engine/Plugins/Experimental/PCG/Source/PCG/Public/PCGSubgraph.h
@@ -171,7 +171,7 @@ protected:
 	void PrepareSubgraphUserParameters(const UPCGSubgraphSettings* Settings, FPCGSubgraphContext* Context, FPCGDataCollection& OutputData) const;
 };
 
-class PCG_API FPCGInputForwardingElement : public FSimplePCGElement
+class PCG_API FPCGInputForwardingElement : public IPCGElement
 {
 public:
 	FPCGInputForwardingElement(const FPCGDataCollection& InputToForward);
diff --git a/Engine/Plugins/Experimental/PCGExternalDataInterop/Source/PCGExternalDataInterop/Private/Elements/PCGLoadAlembicElement.cpp b/Engine/Plugins/Experimental/PCGExternalDataInterop/Source/PCGExternalDataInterop/Private/Elements/PCGLoadAlembicElement.cpp
--- a/Engine/Plugins/Experimental/PCGExternalDataInterop/Source/PCGExternalDataInterop/Private/Elements/PCGLoadAlembicElement.cpp
+++ b/Engine/Plugins/Experimental/PCGExternalDataInterop/Source/PCGExternalDataInterop/Private/Elements/PCGLoadAlembicElement.cpp
@@ -73,14 +73,9 @@ void UPCGLoadAlembicSettings::SetupFromStandard(EPCGLoadAlembicStandardSetup InS
 	}
 }
 
-FPCGContext* FPCGLoadAlembicElement::Initialize(const FPCGDataCollection& InputData, TWeakObjectPtr<UPCGComponent> SourceComponent, const UPCGNode* Node)
+FPCGContext* FPCGLoadAlembicElement::CreateContext()
 {
-	FPCGLoadAlembicContext* Context = new FPCGLoadAlembicContext();
-	Context->InputData = InputData;
-	Context->SourceComponent = SourceComponent;
-	Context->Node = Node;
-
-	return Context;
+	return new FPCGLoadAlembicContext();
 }
 
 bool FPCGLoadAlembicElement::PrepareLoad(FPCGExternalDataContext* InContext) const
diff --git a/Engine/Plugins/Experimental/PCGExternalDataInterop/Source/PCGExternalDataInterop/Public/Elements/PCGLoadAlembicElement.h b/Engine/Plugins/Experimental/PCGExternalDataInterop/Source/PCGExternalDataInterop/Public/Elements/PCGLoadAlembicElement.h
--- a/Engine/Plugins/Experimental/PCGExternalDataInterop/Source/PCGExternalDataInterop/Public/Elements/PCGLoadAlembicElement.h
+++ b/Engine/Plugins/Experimental/PCGExternalDataInterop/Source/PCGExternalDataInterop/Public/Elements/PCGLoadAlembicElement.h
@@ -68,10 +68,8 @@ public:
 
 class PCGEXTERNALDATAINTEROP_API FPCGLoadAlembicElement : public FPCGExternalDataElement
 {
-public:
-	virtual FPCGContext* Initialize(const FPCGDataCollection& InputData, TWeakObjectPtr<UPCGComponent> SourceComponent, const UPCGNode* Node) override;
-
 protected:
+	virtual FPCGContext* CreateContext() override;
 	virtual bool PrepareLoad(FPCGExternalDataContext* Context) const override;
 	virtual bool ExecuteLoad(FPCGExternalDataContext* Context) const override;
 };
\ No newline at end of file
diff --git a/Engine/Plugins/Experimental/PCGGeometryScriptInterop/Source/PCGGeometryScriptInterop/Private/Elements/PCGMeshSampler.cpp b/Engine/Plugins/Experimental/PCGGeometryScriptInterop/Source/PCGGeometryScriptInterop/Private/Elements/PCGMeshSampler.cpp
--- a/Engine/Plugins/Experimental/PCGGeometryScriptInterop/Source/PCGGeometryScriptInterop/Private/Elements/PCGMeshSampler.cpp
+++ b/Engine/Plugins/Experimental/PCGGeometryScriptInterop/Source/PCGGeometryScriptInterop/Private/Elements/PCGMeshSampler.cpp
@@ -95,14 +95,9 @@ FPCGMeshSamplerContext::~FPCGMeshSamplerContext()
 	}
 }
 
-FPCGContext* FPCGMeshSamplerElement::Initialize(const FPCGDataCollection& InputData, TWeakObjectPtr<UPCGComponent> SourceComponent, const UPCGNode* Node)
+FPCGContext* FPCGMeshSamplerElement::CreateContext()
 {
-	FPCGMeshSamplerContext* Context = new FPCGMeshSamplerContext();
-	Context->InputData = InputData;
-	Context->SourceComponent = SourceComponent;
-	Context->Node = Node;
-
-	return Context;
+	return new FPCGMeshSamplerContext();
 }
 
 bool FPCGMeshSamplerElement::PrepareDataInternal(FPCGContext* InContext) const
diff --git a/Engine/Plugins/Experimental/PCGGeometryScriptInterop/Source/PCGGeometryScriptInterop/Public/Elements/PCGMeshSampler.h b/Engine/Plugins/Experimental/PCGGeometryScriptInterop/Source/PCGGeometryScriptInterop/Public/Elements/PCGMeshSampler.h
--- a/Engine/Plugins/Experimental/PCGGeometryScriptInterop/Source/PCGGeometryScriptInterop/Public/Elements/PCGMeshSampler.h
+++ b/Engine/Plugins/Experimental/PCGGeometryScriptInterop/Source/PCGGeometryScriptInterop/Public/Elements/PCGMeshSampler.h
@@ -190,9 +190,8 @@ struct FPCGMeshSamplerContext : public FPCGContext
 
 class FPCGMeshSamplerElement : public IPCGElement
 {
-public:
-	virtual FPCGContext* Initialize(const FPCGDataCollection& InputData, TWeakObjectPtr<UPCGComponent> SourceComponent, const UPCGNode* Node) override;
 protected:
+	virtual FPCGContext* CreateContext() override;
 	virtual bool PrepareDataInternal(FPCGContext* InContext) const override;
 	virtual bool ExecuteInternal(FPCGContext* InContext) const override;
 };
\ No newline at end of file
