Made the search recursive so we can easily find things in subgraphs as well

diff --git a/Engine/Plugins/Experimental/PCG/Source/PCGEditor/Private/PCGEditor.cpp b/Engine/Plugins/Experimental/PCG/Source/PCGEditor/Private/PCGEditor.cpp
--- a/Engine/Plugins/Experimental/PCG/Source/PCGEditor/Private/PCGEditor.cpp
+++ b/Engine/Plugins/Experimental/PCG/Source/PCGEditor/Private/PCGEditor.cpp
@@ -80,6 +80,23 @@ namespace FPCGEditor_private
 	const FName LogID = FName(TEXT("Log"));
 }
 
+UPCGEditorGraph* FPCGEditor::GetPCGEditorGraph(UPCGGraph* InGraph)
+{
+	if (!InGraph)
+	{
+		return nullptr;
+	}
+
+	if (!InGraph->PCGEditorGraph)
+	{
+		InGraph->PCGEditorGraph = NewObject<UPCGEditorGraph>(InGraph, UPCGEditorGraph::StaticClass(), NAME_None, RF_Transactional | RF_Transient);
+		InGraph->PCGEditorGraph->Schema = UPCGEditorGraphSchema::StaticClass();
+		InGraph->PCGEditorGraph->InitFromNodeGraph(InGraph);
+	}
+
+	return InGraph->PCGEditorGraph;
+}
+
 void FPCGEditor::Initialize(const EToolkitMode::Type InMode, const TSharedPtr<class IToolkitHost>& InToolkitHost, UPCGGraph* InPCGGraph)
 {
 	PCGGraphBeingEdited = InPCGGraph;
@@ -90,12 +107,8 @@ void FPCGEditor::Initialize(const EToolkitMode::Type InMode, const TSharedPtr<cl
 		PCGGraphBeingEdited->OnGraphDynamicallyExecutedDelegate.AddRaw(this, &FPCGEditor::OnGraphDynamicallyExecuted);
 	}
 
-	if (!PCGGraphBeingEdited->PCGEditorGraph)
-	{
-		PCGGraphBeingEdited->PCGEditorGraph = NewObject<UPCGEditorGraph>(PCGGraphBeingEdited, UPCGEditorGraph::StaticClass(), NAME_None, RF_Transactional | RF_Transient);
-		PCGGraphBeingEdited->PCGEditorGraph->Schema = UPCGEditorGraphSchema::StaticClass();
-		PCGGraphBeingEdited->PCGEditorGraph->InitFromNodeGraph(InPCGGraph);
-	}
+	// Initializes the UPCGEditorGraph if needed
+	GetPCGEditorGraph(InPCGGraph);
 
 	PCGGraphBeingEdited->PCGEditorGraph->SetEditor(SharedThis(this));
 	PCGEditorGraph = PCGGraphBeingEdited->PCGEditorGraph;
diff --git a/Engine/Plugins/Experimental/PCG/Source/PCGEditor/Private/PCGEditor.h b/Engine/Plugins/Experimental/PCG/Source/PCGEditor/Private/PCGEditor.h
--- a/Engine/Plugins/Experimental/PCG/Source/PCGEditor/Private/PCGEditor.h
+++ b/Engine/Plugins/Experimental/PCG/Source/PCGEditor/Private/PCGEditor.h
@@ -44,6 +44,9 @@ public:
 	/** Get the PCG graph being edited */
 	UPCGEditorGraph* GetPCGEditorGraph();
 
+	/** Gets/Creates the PCG graph editor for a given PCG graph */
+	static UPCGEditorGraph* GetPCGEditorGraph(UPCGGraph* InGraph);
+
 	/** Sets the PCG component and stack that we want to inspect */
 	void SetComponentAndStackBeingInspected(UPCGComponent* InPCGComponent, const FPCGStack& InPCGStack);
 
diff --git a/Engine/Plugins/Experimental/PCG/Source/PCGEditor/Private/PCGEditorGraph.h b/Engine/Plugins/Experimental/PCG/Source/PCGEditor/Private/PCGEditorGraph.h
--- a/Engine/Plugins/Experimental/PCG/Source/PCGEditor/Private/PCGEditorGraph.h
+++ b/Engine/Plugins/Experimental/PCG/Source/PCGEditor/Private/PCGEditorGraph.h
@@ -37,8 +37,8 @@ public:
 
 	UPCGGraph* GetPCGGraph() const { return PCGGraph; }
 
-	void SetEditor(TWeakPtr<const FPCGEditor> InEditor) { PCGEditor = InEditor; }
-	TWeakPtr<const FPCGEditor> GetEditor() const { return PCGEditor; }
+	void SetEditor(TWeakPtr<FPCGEditor> InEditor) { PCGEditor = InEditor; }
+	TWeakPtr<FPCGEditor> GetEditor() const { return PCGEditor; }
 
 	/** Updates the grid size visualization in the editor. */
 	void UpdateGridSizeVisualization(UPCGComponent* InPCGComponentBeingInspected);
@@ -52,5 +52,5 @@ private:
 	UPROPERTY()
 	TObjectPtr<UPCGGraph> PCGGraph = nullptr;
 
-	TWeakPtr<const FPCGEditor> PCGEditor = nullptr;
+	TWeakPtr<FPCGEditor> PCGEditor = nullptr;
 };
diff --git a/Engine/Plugins/Experimental/PCG/Source/PCGEditor/Private/SPCGEditorGraphFind.cpp b/Engine/Plugins/Experimental/PCG/Source/PCGEditor/Private/SPCGEditorGraphFind.cpp
--- a/Engine/Plugins/Experimental/PCG/Source/PCGEditor/Private/SPCGEditorGraphFind.cpp
+++ b/Engine/Plugins/Experimental/PCG/Source/PCGEditor/Private/SPCGEditorGraphFind.cpp
@@ -6,6 +6,7 @@
 #include "PCGEditorGraph.h"
 #include "PCGEditorGraphNodeBase.h"
 #include "PCGNode.h"
+#include "PCGSubgraph.h"
 
 #include "EdGraph/EdGraphSchema.h"
 #include "Framework/Application/SlateApplication.h"
@@ -47,13 +48,29 @@ FPCGEditorGraphFindResult::FPCGEditorGraphFindResult(const FString& InValue, TSh
 
 FReply FPCGEditorGraphFindResult::OnClick(TWeakPtr<FPCGEditor> InPCGEditorPtr)
 {
-	if (UEdGraphPin* ResolvedPin = Pin.Get())
+	TSharedPtr<FPCGEditor> Editor;
+
+	// If result points to another graph, make that visible
+	if (ParentGraph)
 	{
-		InPCGEditorPtr.Pin()->JumpToNode(ResolvedPin->GetOwningNode());
+		GEditor->GetEditorSubsystem<UAssetEditorSubsystem>()->OpenEditorForAsset(ParentGraph->GetPCGGraph());
+		Editor = ParentGraph->GetEditor().Pin();
 	}
-	else if (GraphNode.IsValid())
+	else
 	{
-		InPCGEditorPtr.Pin()->JumpToNode(GraphNode.Get());
+		Editor = InPCGEditorPtr.Pin();
+	}
+
+	if (Editor)
+	{
+		if (UEdGraphPin* ResolvedPin = Pin.Get())
+		{
+			Editor->JumpToNode(ResolvedPin->GetOwningNode());
+		}
+		else if (GraphNode.IsValid())
+		{
+			Editor->JumpToNode(GraphNode.Get());
+		}
 	}
 
 	return FReply::Handled();
@@ -265,35 +282,62 @@ void SPCGEditorGraphFind::MatchTokens(const TArray<FString>& InTokens)
 	RootFindResult.Reset();
 
 	UPCGEditorGraph* PCGEditorGraph = PCGEditorPtr.Pin()->GetPCGEditorGraph();
+
+	if (PCGEditorGraph)
+	{
+		RootFindResult = MakeShared<FPCGEditorGraphFindResult>(FString("PCGTreeRoot"));
+		auto GetParentFunc = []() -> FPCGEditorGraphFindResultPtr { return nullptr; };
+
+		TArray<UPCGGraph*> VisitedGraphs;
+		VisitedGraphs.Add(PCGEditorGraph->GetPCGGraph());
+		MatchTokensInternal(InTokens, PCGEditorGraph, GetParentFunc, VisitedGraphs);
+	}
+}
+
+void SPCGEditorGraphFind::MatchTokensInternal(const TArray<FString>& InTokens, UPCGEditorGraph* PCGEditorGraph, TFunctionRef<FPCGEditorGraphFindResultPtr()> GetParentFunc, TArray<UPCGGraph*>& VisitedSubgraphs)
+{
 	if (!PCGEditorGraph)
 	{
 		return;
 	}
 
-	RootFindResult = MakeShared<FPCGEditorGraphFindResult>(FString("PCGTreeRoot"));
-
 	for (UEdGraphNode* Node : PCGEditorGraph->Nodes)
 	{
+		check(Node);
 		const FString NodeName = Node->GetNodeTitle(ENodeTitleType::FullTitle).ToString();
 		const FString NodeType = Node->GetNodeTitle(ENodeTitleType::ListView).ToString();
 
 		FString NodeSearchString = NodeName + NodeType + Node->NodeComment;
 
-		if (const UPCGEditorGraphNodeBase* PCGEditorGraphNodeBase = Cast<UPCGEditorGraphNodeBase>(Node))
-        {
-        	if (const UPCGNode* PCGNode = PCGEditorGraphNodeBase->GetPCGNode())
-        	{
-        		NodeSearchString.Append(PCGNode->GetName());
-        	}
-        }
+		const UPCGEditorGraphNodeBase* PCGEditorGraphNodeBase = Cast<UPCGEditorGraphNodeBase>(Node);
+
+		if (PCGEditorGraphNodeBase)
+		{
+			if (const UPCGNode* PCGNode = PCGEditorGraphNodeBase->GetPCGNode())
+			{
+				NodeSearchString.Append(PCGNode->GetName());
+			}
+		}
+
 		NodeSearchString = NodeSearchString.Replace(TEXT(" "), TEXT(""));
 
 		FPCGEditorGraphFindResultPtr NodeResult;
-		auto GetOrCreateNodeResult = [this, &NodeResult, &NodeName, &NodeType, Node]() -> FPCGEditorGraphFindResultPtr&
+		auto GetOrCreateNodeResult = [this, &NodeResult, &NodeName, &NodeType, Node, PCGEditorGraph, &GetParentFunc]() -> FPCGEditorGraphFindResultPtr&
 		{
 			if (!NodeResult.IsValid())
 			{
 				NodeResult = MakeShared<FPCGEditorGraphFindResult>((NodeName == NodeType) ? NodeName : NodeName + " - " + NodeType, RootFindResult, Node);
+
+				FPCGEditorGraphFindResultPtr Parent = GetParentFunc();
+				if(Parent)
+				{
+					NodeResult->ParentGraph = PCGEditorGraph;
+					Parent->Children.Add(NodeResult);
+				}
+				else
+				{
+					ItemsFound.Add(NodeResult);
+				}
 			}
 			return NodeResult;
 		};
@@ -313,14 +357,28 @@ void SPCGEditorGraphFind::MatchTokens(const TArray<FString>& InTokens)
 				if (StringMatchesSearchTokens(InTokens, PinSearchString))
 				{
 					FPCGEditorGraphFindResultPtr PinResult(MakeShared<FPCGEditorGraphFindResult>(PinName.ToString(), GetOrCreateNodeResult(), Pin));
+					PinResult->ParentGraph = PCGEditorGraph;
 					NodeResult->Children.Add(PinResult);
 				}
 			}
 		}
 
-		if (NodeResult.IsValid())
+		// Search recursively in subgraph nodes
+		const UPCGBaseSubgraphNode* SubgraphNode = PCGEditorGraphNodeBase ? Cast<UPCGBaseSubgraphNode>(PCGEditorGraphNodeBase->GetPCGNode()) : nullptr;
+		if(UPCGGraph* Subgraph = SubgraphNode ? SubgraphNode->GetSubgraph() : nullptr)
 		{
-			ItemsFound.Add(NodeResult);
+			// Skip already visited subgraph nodes to prevent recursion issues
+			if (!VisitedSubgraphs.Contains(Subgraph))
+			{
+				VisitedSubgraphs.Add(Subgraph);
+
+				if (UPCGEditorGraph* EditorSubgraph = FPCGEditor::GetPCGEditorGraph(Subgraph))
+				{
+					const int CountBeforeRecursion = VisitedSubgraphs.Num();
+					MatchTokensInternal(InTokens, EditorSubgraph, GetOrCreateNodeResult, VisitedSubgraphs);
+					VisitedSubgraphs.SetNum(CountBeforeRecursion);
+				}
+			}
 		}
 	}
 }
diff --git a/Engine/Plugins/Experimental/PCG/Source/PCGEditor/Private/SPCGEditorGraphFind.h b/Engine/Plugins/Experimental/PCG/Source/PCGEditor/Private/SPCGEditorGraphFind.h
--- a/Engine/Plugins/Experimental/PCG/Source/PCGEditor/Private/SPCGEditorGraphFind.h
+++ b/Engine/Plugins/Experimental/PCG/Source/PCGEditor/Private/SPCGEditorGraphFind.h
@@ -14,6 +14,8 @@ template <typename ItemType> class STreeView;
 
 class FPCGEditor;
 class UEdGraphNode;
+class UPCGEditorGraph;
+class UPCGGraph;
 
 struct FPCGEditorGraphFindResult
 {
@@ -44,6 +46,9 @@ struct FPCGEditorGraphFindResult
 	/** Create an icon to represent the result */
 	TSharedRef<SWidget> CreateIcon() const;
 
+	/** ParentGraph if not current */
+	UPCGEditorGraph* ParentGraph = nullptr;
+
 	/** Search result parent */
 	TWeakPtr<FPCGEditorGraphFindResult> Parent;
 
@@ -97,6 +102,9 @@ private:
 	/** Find any results that contain all of the tokens */
 	void MatchTokens(const TArray<FString>& InTokens);
 
+	/** Recursive internal implementation of the MatchTokens */
+	void MatchTokensInternal(const TArray<FString>& InTokens, UPCGEditorGraph* PCGEditorGraph, TFunctionRef<FPCGEditorGraphFindResultPtr()> GetParentFunc, TArray<UPCGGraph*>& VisitedSubgraphNodes);
+
 	/** Determines if a string matches the search tokens */
 	static bool StringMatchesSearchTokens(const TArray<FString>& InTokens, const FString& InComparisonString);
 
