Added the option to tag some graphs as editor-only to prevent some data validation issues. This is a transitive property (through subgraphs) and it is propagated to PCG components as well.

diff --git a/Engine/Plugins/Experimental/PCG/Source/PCG/Private/PCGComponent.cpp b/Engine/Plugins/Experimental/PCG/Source/PCG/Private/PCGComponent.cpp
--- a/Engine/Plugins/Experimental/PCG/Source/PCG/Private/PCGComponent.cpp
+++ b/Engine/Plugins/Experimental/PCG/Source/PCG/Private/PCGComponent.cpp
@@ -1164,6 +1164,11 @@ void UPCGComponent::BeginDestroy()
 	Super::BeginDestroy();
 }
 
+bool UPCGComponent::IsEditorOnly() const
+{
+	return Super::IsEditorOnly() || (GraphInstance && GraphInstance->IsEditorOnly());
+}
+
 void UPCGComponent::PostInitProperties()
 {
 #if WITH_EDITOR
diff --git a/Engine/Plugins/Experimental/PCG/Source/PCG/Private/PCGGraph.cpp b/Engine/Plugins/Experimental/PCG/Source/PCG/Private/PCGGraph.cpp
--- a/Engine/Plugins/Experimental/PCG/Source/PCG/Private/PCGGraph.cpp
+++ b/Engine/Plugins/Experimental/PCG/Source/PCG/Private/PCGGraph.cpp
@@ -332,6 +332,28 @@ void UPCGGraph::PostLoad()
 #endif
 }
 
+bool UPCGGraph::IsEditorOnly() const
+{
+	bool bIsCurrentlyEditorOnly = (Super::IsEditorOnly() || bIsEditorOnly);
+
+	if (!bIsCurrentlyEditorOnly)
+	{
+		// Check for any subgraphs...
+		ForEachNode([&bIsCurrentlyEditorOnly](UPCGNode* Node)
+		{
+			if (UPCGBaseSubgraphNode* SubgraphNode = Cast<UPCGBaseSubgraphNode>(Node))
+			{
+				if (SubgraphNode->GetSubgraph())
+				{
+					bIsCurrentlyEditorOnly |= SubgraphNode->GetSubgraph()->IsEditorOnly();
+				}
+			}
+		});
+	}
+
+	return bIsCurrentlyEditorOnly;
+}
+
 #if WITH_EDITOR
 void UPCGGraph::DeclareConstructClasses(TArray<FTopLevelAssetPath>& OutConstructClasses, const UClass* SpecificSubclass)
 {
@@ -708,7 +730,7 @@ bool UPCGGraph::RemoveEdge(UPCGNode* From, const FName& FromLabel, UPCGNode* To,
 	return TouchedNodes.Num() > 0;
 }
 
-void UPCGGraph::ForEachNode(const TFunction<void(UPCGNode*)>& Action)
+void UPCGGraph::ForEachNode(const TFunction<void(UPCGNode*)>& Action) const
 {
 	Action(InputNode);
 	Action(OutputNode);
diff --git a/Engine/Plugins/Experimental/PCG/Source/PCG/Public/PCGComponent.h b/Engine/Plugins/Experimental/PCG/Source/PCG/Public/PCGComponent.h
--- a/Engine/Plugins/Experimental/PCG/Source/PCG/Public/PCGComponent.h
+++ b/Engine/Plugins/Experimental/PCG/Source/PCG/Public/PCGComponent.h
@@ -79,6 +79,7 @@ public:
 	virtual void PostLoad() override;
 	virtual void PostInitProperties() override;
 	virtual void BeginDestroy() override;
+	virtual bool IsEditorOnly() const override;
 
 #if WITH_EDITOR
 	virtual bool CanEditChange(const FProperty* InProperty) const override;
diff --git a/Engine/Plugins/Experimental/PCG/Source/PCG/Public/PCGGraph.h b/Engine/Plugins/Experimental/PCG/Source/PCG/Public/PCGGraph.h
--- a/Engine/Plugins/Experimental/PCG/Source/PCG/Public/PCGGraph.h
+++ b/Engine/Plugins/Experimental/PCG/Source/PCG/Public/PCGGraph.h
@@ -178,6 +178,7 @@ public:
 	UPCGGraph(const FObjectInitializer& ObjectInitializer);
 	/** ~Begin UObject interface */
 	virtual void PostLoad() override;
+	virtual bool IsEditorOnly() const override;
 #if WITH_EDITOR
 	static void DeclareConstructClasses(TArray<FTopLevelAssetPath>& OutConstructClasses, const UClass* SpecificSubclass);
 #endif
@@ -269,7 +270,7 @@ public:
 	void AddNodes(TArray<UPCGNode*>& InNodes);
 
 	/** Calls the lambda on every node in graph. */
-	void ForEachNode(const TFunction<void(UPCGNode*)>& Action);
+	void ForEachNode(const TFunction<void(UPCGNode*)>& Action) const;
 
 	bool RemoveInboundEdges(UPCGNode* InNode, const FName& InboundLabel);
 	bool RemoveOutboundEdges(UPCGNode* InNode, const FName& OutboundLabel);
@@ -362,6 +363,10 @@ protected:
 	mutable TMap<const UPCGNode*, uint32> NodeToGridSize;
 	mutable FRWLock NodeToGridSizeLock;
 
+	/** Sets whether this graph is marked as editor-only; note that the IsEditorOnly call depends on the local graph value and the value in all subgraphs, recursively. */
+	UPROPERTY(EditAnywhere, BlueprintReadWrite, Category = Cooking)
+	bool bIsEditorOnly = false;
+
 #if WITH_EDITORONLY_DATA
 	/** When true the Debug flag in the graph editor will display debug information contextually for the selected debug object. Otherwise
 	* debug information is displayed for all components using a graph (requires regenerate).
@@ -417,6 +422,7 @@ public:
 	virtual void PostDuplicate(bool bDuplicateForPIE) override;
 	virtual void PostEditImport() override;
 	virtual void BeginDestroy() override;
+	virtual bool IsEditorOnly() const override { return GetGraph() && GetGraph()->IsEditorOnly(); }
 
 #if WITH_EDITOR
 	virtual void PreEditChange(FProperty* InProperty) override;
