[PCG] Fixed an issue were opening multiple PCG editors would cause multiple separators to show between toolbar sections

diff --git a/Engine/Plugins/Experimental/PCG/Source/PCGEditor/Private/PCGEditor.cpp b/Engine/Plugins/Experimental/PCG/Source/PCGEditor/Private/PCGEditor.cpp
--- a/Engine/Plugins/Experimental/PCG/Source/PCGEditor/Private/PCGEditor.cpp
+++ b/Engine/Plugins/Experimental/PCG/Source/PCGEditor/Private/PCGEditor.cpp
@@ -454,21 +454,14 @@ FString FPCGEditor::GetWorldCentricTabPrefix() const
 void FPCGEditor::RegisterToolbar() const
 {
 	UToolMenus* ToolMenus = UToolMenus::Get();
-	UToolMenu* ToolBar;
 	FName ParentName;
 	const FName MenuName = GetToolMenuToolbarName(ParentName);
-	if (ToolMenus->IsMenuRegistered(MenuName))
+	if (!ToolMenus->IsMenuRegistered(MenuName))
 	{
-		ToolBar = ToolMenus->ExtendMenu(MenuName);
-	}
-	else
-	{
-		ToolBar = ToolMenus->RegisterMenu(MenuName, ParentName, EMultiBoxType::ToolBar);
-	}
+		UToolMenu* ToolBar = ToolMenus->RegisterMenu(MenuName, ParentName, EMultiBoxType::ToolBar);
 
-	const FPCGEditorCommands& PCGEditorCommands = FPCGEditorCommands::Get();
-	const FToolMenuInsert InsertAfterAssetSection("Asset", EToolMenuInsertType::After);
-	{
+		const FPCGEditorCommands& PCGEditorCommands = FPCGEditorCommands::Get();
+		const FToolMenuInsert InsertAfterAssetSection("Asset", EToolMenuInsertType::After);
 		FToolMenuSection& Section = ToolBar->AddSection("PCGToolbar", TAttribute<FText>(), InsertAfterAssetSection);
 		Section.AddEntry(FToolMenuEntry::InitToolBarButton(
 			PCGEditorCommands.Find,
