[PCG] Temporary fix to support adding attribute to multi-entry attribute set

#rnx

diff --git a/Engine/Plugins/Experimental/PCG/Source/PCG/Private/Elements/PCGCreateAttribute.cpp b/Engine/Plugins/Experimental/PCG/Source/PCG/Private/Elements/PCGCreateAttribute.cpp
--- a/Engine/Plugins/Experimental/PCG/Source/PCG/Private/Elements/PCGCreateAttribute.cpp
+++ b/Engine/Plugins/Experimental/PCG/Source/PCG/Private/Elements/PCGCreateAttribute.cpp
@@ -291,6 +291,12 @@ bool FPCGCreateAttributeElement::ExecuteInternal(FPCGContext* Context) const
 			PCGE_LOG(Error, GraphAndLog, FText::Format(LOCTEXT("ParamMissingAttribute", "Source Attribute Set data does not have an attribute '{0}'"), FText::FromName(SourceParamAttributeName)));
 			return true;
 		}
+
+		// We don't support multi-entry yet
+		if (SourceParamData->Metadata->GetLocalItemCount() > 1)
+		{
+			PCGE_LOG(Warning, GraphAndLog,LOCTEXT("NotSupportingMultiEntry", "Source Attribute Set data has more than one entry. Not yet supported. Will use the default value of the attribute."));
+		}
 	}
 	else
 	{
@@ -306,6 +312,7 @@ bool FPCGCreateAttributeElement::ExecuteInternal(FPCGContext* Context) const
 	// If the input is empty, we will create a new ParamData.
 	// We can re-use this newly object as the output
 	bool bCanReuseInputData = false;
+	bool bIsParamData = false;
 	if (!bHasInputConnections)
 	{
 		ensure(Inputs.IsEmpty());
@@ -323,8 +330,6 @@ bool FPCGCreateAttributeElement::ExecuteInternal(FPCGContext* Context) const
 
 		UPCGMetadata* Metadata = nullptr;
 
-		bool bShouldAddNewEntry = false;
-
 		if (const UPCGSpatialData* InputSpatialData = Cast<UPCGSpatialData>(InputData))
 		{
 			UPCGSpatialData* NewSpatialData = InputSpatialData->DuplicateData(/*bInitializeFromData=*/false);
@@ -341,9 +346,7 @@ bool FPCGCreateAttributeElement::ExecuteInternal(FPCGContext* Context) const
 
 			OutputData = NewParamData;
 			Metadata = NewParamData->Metadata;
-
-			// In case of param data, we want to add a new entry too, if needed
-			bShouldAddNewEntry = true;
+			bIsParamData = true;
 		}
 		else
 		{
@@ -361,7 +364,7 @@ bool FPCGCreateAttributeElement::ExecuteInternal(FPCGContext* Context) const
 			if (InputSource.GetExtraNames().IsEmpty())
 			{
 				const FPCGMetadataAttributeBase* SourceAttribute = SourceParamData->Metadata->GetConstAttribute(SourceParamAttributeName);
-				Attribute = Metadata->CopyAttribute(SourceAttribute, OutputAttributeName, /*bKeepParent=*/false, /*bCopyEntries=*/bShouldAddNewEntry, /*bCopyValues=*/bShouldAddNewEntry);
+				Attribute = Metadata->CopyAttribute(SourceAttribute, OutputAttributeName, /*bKeepParent=*/false, /*bCopyEntries=*/false, /*bCopyValues=*/false);
 			}
 			else // Create a new attribute of the accessed field's type manually
 			{
@@ -395,18 +398,25 @@ bool FPCGCreateAttributeElement::ExecuteInternal(FPCGContext* Context) const
 			continue;
 		}
 
+		// In the case of an input param data, remap all entries for the output data (which is also a param data) 
+		// to the default value for the newly created attribute.
+		if (bIsParamData)
+		{
+			// Making sure the metadata has at least one entry.
+			if (Metadata->GetLocalItemCount() == 0)
+			{
+				Metadata->AddEntry();
+			}
+
+			for (PCGMetadataEntryKey Key = 0; Key < Metadata->GetLocalItemCount(); ++Key)
+			{
+				Attribute->SetValueFromValueKey(Key, PCGDefaultValueKey);
+			}
+		}
+
 		TArray<FPCGTaggedData>& Outputs = Context->OutputData.TaggedData;
 		FPCGTaggedData& Output = Outputs.Emplace_GetRef();
 		Output.Data = OutputData;
-
-		// Add a new entry if it is a param data and not from source (because entries are already copied)
-		if (bShouldAddNewEntry && !SourceParamData)
-		{
-			// If the metadata is empty, we need to add a new entry, so set it to PCGInvalidEntryKey.
-			// Otherwise, use the entry key 0.
-			PCGMetadataEntryKey EntryKey = Metadata->GetItemCountForChild() == 0 ? PCGInvalidEntryKey : 0;
-			SetAttribute(Settings, Attribute, Metadata, EntryKey);
-		}
 	}
 
 	return true;
@@ -424,23 +434,4 @@ FPCGMetadataAttributeBase* FPCGCreateAttributeElement::ClearOrCreateAttribute(co
 	return Settings->AttributeTypes.Dispatcher(CreateAttribute);
 }
 
-PCGMetadataEntryKey FPCGCreateAttributeElement::SetAttribute(const UPCGCreateAttributeBaseSettings* Settings, FPCGMetadataAttributeBase* Attribute, UPCGMetadata* Metadata, PCGMetadataEntryKey EntryKey) const
-{
-	check(Attribute && Metadata);
-
-	auto SetAttribute = [Attribute, EntryKey, Metadata](auto&& Value) -> PCGMetadataEntryKey
-	{
-		using AttributeType = std::decay_t<decltype(Value)>;
-
-		check(Attribute->GetTypeId() == PCG::Private::MetadataTypes<AttributeType>::Id);
-
-		const PCGMetadataEntryKey FinalKey = (EntryKey == PCGInvalidEntryKey) ? Metadata->AddEntry() : EntryKey;
-
-		static_cast<FPCGMetadataAttribute<AttributeType>*>(Attribute)->SetValue(FinalKey, std::forward<decltype(Value)>(Value));
-
-		return FinalKey;
-	};
-
-	return Settings->AttributeTypes.Dispatcher(SetAttribute);
-}
 #undef LOCTEXT_NAMESPACE
diff --git a/Engine/Plugins/Experimental/PCG/Source/PCG/Public/Elements/PCGCreateAttribute.h b/Engine/Plugins/Experimental/PCG/Source/PCG/Public/Elements/PCGCreateAttribute.h
--- a/Engine/Plugins/Experimental/PCG/Source/PCG/Public/Elements/PCGCreateAttribute.h
+++ b/Engine/Plugins/Experimental/PCG/Source/PCG/Public/Elements/PCGCreateAttribute.h
@@ -195,9 +195,6 @@ protected:
 private:
 	/* Create (or clear) an attribute named by OutputAttributeName and depending on the selected type. Value can be overridden by params. Default value will be set to the specified value. */
 	FPCGMetadataAttributeBase* ClearOrCreateAttribute(const UPCGCreateAttributeBaseSettings* Settings, UPCGMetadata* Metadata, const FName OutputAttributeName) const;
-
-	/* Set an entry defined by EntryKey (or create one if it is invalid) in the Attribute depending on the selected type. Value can be overridden by params. */
-	PCGMetadataEntryKey SetAttribute(const UPCGCreateAttributeBaseSettings* Settings, FPCGMetadataAttributeBase* Attribute, UPCGMetadata* Metadata, PCGMetadataEntryKey EntryKey) const;
 };
 
 #if UE_ENABLE_INCLUDE_ORDER_DEPRECATED_IN_5_2
