[PCG] Attribute CRC fixes - attribute name missing from CRC, and added Last attribute

diff --git a/Engine/Plugins/Experimental/PCG/Source/PCG/Private/Data/PCGCollisionShapeData.cpp b/Engine/Plugins/Experimental/PCG/Source/PCG/Private/Data/PCGCollisionShapeData.cpp
--- a/Engine/Plugins/Experimental/PCG/Source/PCG/Private/Data/PCGCollisionShapeData.cpp
+++ b/Engine/Plugins/Experimental/PCG/Source/PCG/Private/Data/PCGCollisionShapeData.cpp
@@ -43,6 +43,8 @@ bool UPCGCollisionShapeData::IsSupported(UShapeComponent* InComponent)
 
 void UPCGCollisionShapeData::AddToCrc(FArchiveCrc32& Ar, bool bFullDataCrc) const
 {
+	Super::AddToCrc(Ar, bFullDataCrc);
+
 	uint32 UniqueTypeID = StaticClass()->GetDefaultObject()->GetUniqueID();
 	Ar << UniqueTypeID;
 
diff --git a/Engine/Plugins/Experimental/PCG/Source/PCG/Private/Data/PCGDifferenceData.cpp b/Engine/Plugins/Experimental/PCG/Source/PCG/Private/Data/PCGDifferenceData.cpp
--- a/Engine/Plugins/Experimental/PCG/Source/PCG/Private/Data/PCGDifferenceData.cpp
+++ b/Engine/Plugins/Experimental/PCG/Source/PCG/Private/Data/PCGDifferenceData.cpp
@@ -145,7 +145,7 @@ FPCGCrc UPCGDifferenceData::ComputeCrc(bool bFullDataCrc) const
 	}
 	else
 	{
-		UPCGData::AddToCrc(Ar, bFullDataCrc);
+		AddUIDToCrc(Ar);
 	}
 
 	return FPCGCrc(Ar.GetCrc());
@@ -153,6 +153,8 @@ FPCGCrc UPCGDifferenceData::ComputeCrc(bool bFullDataCrc) const
 
 void UPCGDifferenceData::AddToCrc(FArchiveCrc32& Ar, bool bFullDataCrc) const
 {
+	Super::AddToCrc(Ar, bFullDataCrc);
+
 	uint32 UniqueTypeID = StaticClass()->GetDefaultObject()->GetUniqueID();
 	Ar << UniqueTypeID;
 
diff --git a/Engine/Plugins/Experimental/PCG/Source/PCG/Private/Data/PCGIntersectionData.cpp b/Engine/Plugins/Experimental/PCG/Source/PCG/Private/Data/PCGIntersectionData.cpp
--- a/Engine/Plugins/Experimental/PCG/Source/PCG/Private/Data/PCGIntersectionData.cpp
+++ b/Engine/Plugins/Experimental/PCG/Source/PCG/Private/Data/PCGIntersectionData.cpp
@@ -88,7 +88,7 @@ FPCGCrc UPCGIntersectionData::ComputeCrc(bool bFullDataCrc) const
 	}
 	else
 	{
-		UPCGData::AddToCrc(Ar, bFullDataCrc);
+		AddUIDToCrc(Ar);
 	}
 
 	return FPCGCrc(Ar.GetCrc());
@@ -96,6 +96,8 @@ FPCGCrc UPCGIntersectionData::ComputeCrc(bool bFullDataCrc) const
 
 void UPCGIntersectionData::AddToCrc(FArchiveCrc32& Ar, bool bFullDataCrc) const
 {
+	Super::AddToCrc(Ar, bFullDataCrc);
+
 	uint32 UniqueTypeID = StaticClass()->GetDefaultObject()->GetUniqueID();
 	Ar << UniqueTypeID;
 
diff --git a/Engine/Plugins/Experimental/PCG/Source/PCG/Private/Data/PCGLandscapeData.cpp b/Engine/Plugins/Experimental/PCG/Source/PCG/Private/Data/PCGLandscapeData.cpp
--- a/Engine/Plugins/Experimental/PCG/Source/PCG/Private/Data/PCGLandscapeData.cpp
+++ b/Engine/Plugins/Experimental/PCG/Source/PCG/Private/Data/PCGLandscapeData.cpp
@@ -151,6 +151,14 @@ void UPCGLandscapeData::PostLoad()
 #endif
 }
 
+void UPCGLandscapeData::AddToCrc(FArchiveCrc32& Ar, bool bFullDataCrc) const
+{
+	Super::AddToCrc(Ar, bFullDataCrc);
+
+	// This data does not have a bespoke CRC implementation so just use a global unique data CRC.
+	AddUIDToCrc(Ar);
+}
+
 FBox UPCGLandscapeData::GetBounds() const
 {
 	return Bounds;
diff --git a/Engine/Plugins/Experimental/PCG/Source/PCG/Private/Data/PCGLandscapeSplineData.cpp b/Engine/Plugins/Experimental/PCG/Source/PCG/Private/Data/PCGLandscapeSplineData.cpp
--- a/Engine/Plugins/Experimental/PCG/Source/PCG/Private/Data/PCGLandscapeSplineData.cpp
+++ b/Engine/Plugins/Experimental/PCG/Source/PCG/Private/Data/PCGLandscapeSplineData.cpp
@@ -43,6 +43,14 @@ void UPCGLandscapeSplineData::Initialize(ULandscapeSplinesComponent* InSplineCom
 	Spline = InSplineComponent;
 }
 
+void UPCGLandscapeSplineData::AddToCrc(FArchiveCrc32& Ar, bool bFullDataCrc) const
+{
+	Super::AddToCrc(Ar, bFullDataCrc);
+
+	// This data does not have a bespoke CRC implementation so just use a global unique data CRC.
+	AddUIDToCrc(Ar);
+}
+
 FTransform UPCGLandscapeSplineData::GetTransform() const
 {
 	return Spline.IsValid() ? Spline->GetComponentTransform() : FTransform::Identity;
diff --git a/Engine/Plugins/Experimental/PCG/Source/PCG/Private/Data/PCGPointData.cpp b/Engine/Plugins/Experimental/PCG/Source/PCG/Private/Data/PCGPointData.cpp
--- a/Engine/Plugins/Experimental/PCG/Source/PCG/Private/Data/PCGPointData.cpp
+++ b/Engine/Plugins/Experimental/PCG/Source/PCG/Private/Data/PCGPointData.cpp
@@ -267,11 +267,13 @@ const UPCGPointData::PointOctree& UPCGPointData::GetOctree() const
 
 void UPCGPointData::AddToCrc(FArchiveCrc32& Ar, bool bFullDataCrc) const
 {
+	Super::AddToCrc(Ar, bFullDataCrc);
+
 	// The code below has non-trivial cost, and can be disabled from console.
 	if (!bFullDataCrc || !CVarCacheFullPointDataCrc.GetValueOnAnyThread())
 	{
 		// Fallback to UID
-		Super::AddToCrc(Ar, bFullDataCrc);
+		AddUIDToCrc(Ar);
 		return;
 	}
 
diff --git a/Engine/Plugins/Experimental/PCG/Source/PCG/Private/Data/PCGPrimitiveData.cpp b/Engine/Plugins/Experimental/PCG/Source/PCG/Private/Data/PCGPrimitiveData.cpp
--- a/Engine/Plugins/Experimental/PCG/Source/PCG/Private/Data/PCGPrimitiveData.cpp
+++ b/Engine/Plugins/Experimental/PCG/Source/PCG/Private/Data/PCGPrimitiveData.cpp
@@ -18,6 +18,14 @@ void UPCGPrimitiveData::Initialize(UPrimitiveComponent* InPrimitive)
 	// Not obvious to find strict bounds, leave at the default value
 }
 
+void UPCGPrimitiveData::AddToCrc(FArchiveCrc32& Ar, bool bFullDataCrc) const
+{
+	Super::AddToCrc(Ar, bFullDataCrc);
+
+	// This data does not have a bespoke CRC implementation so just use a global unique data CRC.
+	AddUIDToCrc(Ar);
+}
+
 bool UPCGPrimitiveData::SamplePoint(const FTransform& InTransform, const FBox& InBounds, FPCGPoint& OutPoint, UPCGMetadata* OutMetadata) const
 {
 	// Pure implementation
diff --git a/Engine/Plugins/Experimental/PCG/Source/PCG/Private/Data/PCGProjectionData.cpp b/Engine/Plugins/Experimental/PCG/Source/PCG/Private/Data/PCGProjectionData.cpp
--- a/Engine/Plugins/Experimental/PCG/Source/PCG/Private/Data/PCGProjectionData.cpp
+++ b/Engine/Plugins/Experimental/PCG/Source/PCG/Private/Data/PCGProjectionData.cpp
@@ -28,6 +28,14 @@ void UPCGProjectionData::Initialize(const UPCGSpatialData* InSource, const UPCGS
 	CachedStrictBounds = ProjectBounds(Source->GetStrictBounds());
 }
 
+void UPCGProjectionData::AddToCrc(FArchiveCrc32& Ar, bool bFullDataCrc) const
+{
+	Super::AddToCrc(Ar, bFullDataCrc);
+
+	// This data does not have a bespoke CRC implementation so just use a global unique data CRC.
+	AddUIDToCrc(Ar);
+}
+
 int UPCGProjectionData::GetDimension() const
 {
 	check(Source && Target);
diff --git a/Engine/Plugins/Experimental/PCG/Source/PCG/Private/Data/PCGRenderTargetData.cpp b/Engine/Plugins/Experimental/PCG/Source/PCG/Private/Data/PCGRenderTargetData.cpp
--- a/Engine/Plugins/Experimental/PCG/Source/PCG/Private/Data/PCGRenderTargetData.cpp
+++ b/Engine/Plugins/Experimental/PCG/Source/PCG/Private/Data/PCGRenderTargetData.cpp
@@ -42,6 +42,14 @@ void UPCGRenderTargetData::Initialize(UTextureRenderTarget2D* InRenderTarget, co
 	Bounds = Bounds.TransformBy(Transform);
 }
 
+void UPCGRenderTargetData::AddToCrc(FArchiveCrc32& Ar, bool bFullDataCrc) const
+{
+	Super::AddToCrc(Ar, bFullDataCrc);
+
+	// This data does not have a bespoke CRC implementation so just use a global unique data CRC.
+	AddUIDToCrc(Ar);
+}
+
 UPCGSpatialData* UPCGRenderTargetData::CopyInternal() const
 {
 	UPCGRenderTargetData* NewRenderTargetData = NewObject<UPCGRenderTargetData>();
diff --git a/Engine/Plugins/Experimental/PCG/Source/PCG/Private/Data/PCGSpatialData.cpp b/Engine/Plugins/Experimental/PCG/Source/PCG/Private/Data/PCGSpatialData.cpp
--- a/Engine/Plugins/Experimental/PCG/Source/PCG/Private/Data/PCGSpatialData.cpp
+++ b/Engine/Plugins/Experimental/PCG/Source/PCG/Private/Data/PCGSpatialData.cpp
@@ -7,6 +7,8 @@
 #include "Data/PCGProjectionData.h"
 #include "Data/PCGUnionData.h"
 
+#include "Serialization/ArchiveCrc32.h"
+
 #include UE_INLINE_GENERATED_CPP_BY_NAME(PCGSpatialData)
 
 UPCGSpatialData::UPCGSpatialData(const FObjectInitializer& ObjectInitializer)
@@ -235,6 +237,18 @@ UPCGSpatialData* UPCGSpatialData::DuplicateData(const bool bInitializeMetadata)
 	return NewSpatialData;
 }
 
+void UPCGSpatialData::AddToCrc(FArchiveCrc32& Ar, bool bFullDataCrc) const
+{
+	Super::AddToCrc(Ar, bFullDataCrc);
+
+	if (Metadata)
+	{
+		// Can impact results of downstream node execution.
+		FName LatestAttribute = Metadata->GetLatestAttributeNameOrNone();
+		Ar << LatestAttribute;
+	}
+}
+
 bool UPCGSpatialData::HasCachedLastSelector() const
 {
 	return bHasCachedLastSelector || (Metadata && Metadata->GetAttributeCount() > 0);
diff --git a/Engine/Plugins/Experimental/PCG/Source/PCG/Private/Data/PCGSplineData.cpp b/Engine/Plugins/Experimental/PCG/Source/PCG/Private/Data/PCGSplineData.cpp
--- a/Engine/Plugins/Experimental/PCG/Source/PCG/Private/Data/PCGSplineData.cpp
+++ b/Engine/Plugins/Experimental/PCG/Source/PCG/Private/Data/PCGSplineData.cpp
@@ -57,6 +57,14 @@ void UPCGSplineData::ApplyTo(USplineComponent* InSplineComponent)
 	SplineStruct.ApplyTo(InSplineComponent);
 }
 
+void UPCGSplineData::AddToCrc(FArchiveCrc32& Ar, bool bFullDataCrc) const
+{
+	Super::AddToCrc(Ar, bFullDataCrc);
+
+	// This data does not have a bespoke CRC implementation so just use a global unique data CRC.
+	AddUIDToCrc(Ar);
+}
+
 FTransform UPCGSplineData::GetTransform() const
 {
 	return SplineStruct.GetTransform();
@@ -315,6 +323,12 @@ void UPCGSplineProjectionData::Initialize(const UPCGSplineData* InSourceSpline,
 	}
 }
 
+void UPCGSplineProjectionData::AddToCrc(FArchiveCrc32& Ar, bool bFullDataCrc) const
+{
+	// This data does not have a bespoke CRC implementation so just use a global unique data CRC.
+	AddUIDToCrc(Ar);
+}
+
 const UPCGSplineData* UPCGSplineProjectionData::GetSpline() const
 {
 	return Cast<const UPCGSplineData>(Source);
diff --git a/Engine/Plugins/Experimental/PCG/Source/PCG/Private/Data/PCGTextureData.cpp b/Engine/Plugins/Experimental/PCG/Source/PCG/Private/Data/PCGTextureData.cpp
--- a/Engine/Plugins/Experimental/PCG/Source/PCG/Private/Data/PCGTextureData.cpp
+++ b/Engine/Plugins/Experimental/PCG/Source/PCG/Private/Data/PCGTextureData.cpp
@@ -304,6 +304,14 @@ void UPCGTextureData::Initialize(UTexture2D* InTexture, const FTransform& InTran
 	}
 }
 
+void UPCGTextureData::AddToCrc(FArchiveCrc32& Ar, bool bFullDataCrc) const
+{
+	Super::AddToCrc(Ar, bFullDataCrc);
+
+	// This data does not have a bespoke CRC implementation so just use a global unique data CRC.
+	AddUIDToCrc(Ar);
+}
+
 UPCGSpatialData* UPCGTextureData::CopyInternal() const
 {
 	UPCGTextureData* NewTextureData = NewObject<UPCGTextureData>();
diff --git a/Engine/Plugins/Experimental/PCG/Source/PCG/Private/Data/PCGUnionData.cpp b/Engine/Plugins/Experimental/PCG/Source/PCG/Private/Data/PCGUnionData.cpp
--- a/Engine/Plugins/Experimental/PCG/Source/PCG/Private/Data/PCGUnionData.cpp
+++ b/Engine/Plugins/Experimental/PCG/Source/PCG/Private/Data/PCGUnionData.cpp
@@ -106,7 +106,7 @@ FPCGCrc UPCGUnionData::ComputeCrc(bool bFullDataCrc) const
 	}
 	else
 	{
-		UPCGData::AddToCrc(Ar, bFullDataCrc);
+		AddUIDToCrc(Ar);
 	}
 
 	return FPCGCrc(Ar.GetCrc());
@@ -114,6 +114,8 @@ FPCGCrc UPCGUnionData::ComputeCrc(bool bFullDataCrc) const
 
 void UPCGUnionData::AddToCrc(FArchiveCrc32& Ar, bool bFullDataCrc) const
 {
+	Super::AddToCrc(Ar, bFullDataCrc);
+
 	uint32 UniqueTypeID = StaticClass()->GetDefaultObject()->GetUniqueID();
 	Ar << UniqueTypeID;
 
diff --git a/Engine/Plugins/Experimental/PCG/Source/PCG/Private/Data/PCGVolumeData.cpp b/Engine/Plugins/Experimental/PCG/Source/PCG/Private/Data/PCGVolumeData.cpp
--- a/Engine/Plugins/Experimental/PCG/Source/PCG/Private/Data/PCGVolumeData.cpp
+++ b/Engine/Plugins/Experimental/PCG/Source/PCG/Private/Data/PCGVolumeData.cpp
@@ -39,6 +39,14 @@ void UPCGVolumeData::Initialize(const FBox& InBounds)
 	StrictBounds = InBounds;
 }
 
+void UPCGVolumeData::AddToCrc(FArchiveCrc32& Ar, bool bFullDataCrc) const
+{
+	Super::AddToCrc(Ar, bFullDataCrc);
+
+	// This data does not have a bespoke CRC implementation so just use a global unique data CRC.
+	AddUIDToCrc(Ar);
+}
+
 FBox UPCGVolumeData::GetBounds() const
 {
 	return Bounds;
diff --git a/Engine/Plugins/Experimental/PCG/Source/PCG/Private/Data/PCGWorldData.cpp b/Engine/Plugins/Experimental/PCG/Source/PCG/Private/Data/PCGWorldData.cpp
--- a/Engine/Plugins/Experimental/PCG/Source/PCG/Private/Data/PCGWorldData.cpp
+++ b/Engine/Plugins/Experimental/PCG/Source/PCG/Private/Data/PCGWorldData.cpp
@@ -230,6 +230,14 @@ void UPCGWorldRayHitData::Initialize(UWorld* InWorld, const FBox& InBounds)
 	Bounds = InBounds;
 }
 
+void UPCGWorldRayHitData::AddToCrc(FArchiveCrc32& Ar, bool bFullDataCrc) const
+{
+	Super::AddToCrc(Ar, bFullDataCrc);
+
+	// This data does not have a bespoke CRC implementation so just use a global unique data CRC.
+	AddUIDToCrc(Ar);
+}
+
 bool UPCGWorldRayHitData::SamplePoint(const FTransform& InTransform, const FBox& InBounds, FPCGPoint& OutPoint, UPCGMetadata* OutMetadata) const
 {
 	// TODO: This seems to be a projection - along a direction. I suspect that UPCGWorldVolumetricData is the SamplePoint(), and this is the ProjectPoint() (in a direction)?
diff --git a/Engine/Plugins/Experimental/PCG/Source/PCG/Private/Metadata/PCGAttributePropertySelector.cpp b/Engine/Plugins/Experimental/PCG/Source/PCG/Private/Metadata/PCGAttributePropertySelector.cpp
--- a/Engine/Plugins/Experimental/PCG/Source/PCG/Private/Metadata/PCGAttributePropertySelector.cpp
+++ b/Engine/Plugins/Experimental/PCG/Source/PCG/Private/Metadata/PCGAttributePropertySelector.cpp
@@ -8,6 +8,8 @@
 #include "Metadata/PCGMetadata.h"
 #include "Metadata/PCGMetadataAttribute.h"
 
+#include "Serialization/ArchiveCrc32.h"
+
 namespace PCGAttributePropertySelectorConstants
 {
 	static const TCHAR* PropertyPrefix = TEXT("$");
@@ -247,6 +249,25 @@ bool FPCGAttributePropertySelector::Update(FString NewValue)
 	return SetAttributeName(NewName.IsEmpty() ? NAME_None : FName(NewName), /*bResetExtraNames=*/ false) || bExtraChanged;
 }
 
+void FPCGAttributePropertySelector::AddToCrc(FArchiveCrc32& Ar) const
+{
+	Ar << Selection;
+	Ar << const_cast<TArray<FString>&>(ExtraNames);
+
+	switch (Selection)
+	{
+	case EPCGAttributePropertySelection::Attribute:
+		Ar << const_cast<FName&>(AttributeName);
+		break;
+	case EPCGAttributePropertySelection::PointProperty:
+		Ar << PointProperty;
+		break;
+	case EPCGAttributePropertySelection::ExtraProperty:
+		Ar << ExtraProperty;
+		break;
+	}
+}
+
 ///////////////////////////////////////////////////////////////////////
 
 FPCGAttributePropertyInputSelector::FPCGAttributePropertyInputSelector()
diff --git a/Engine/Plugins/Experimental/PCG/Source/PCG/Private/PCGData.cpp b/Engine/Plugins/Experimental/PCG/Source/PCG/Private/PCGData.cpp
--- a/Engine/Plugins/Experimental/PCG/Source/PCG/Private/PCGData.cpp
+++ b/Engine/Plugins/Experimental/PCG/Source/PCG/Private/PCGData.cpp
@@ -111,9 +111,9 @@ FPCGCrc UPCGData::ComputeCrc(bool bFullDataCrc) const
 
 void UPCGData::AddToCrc(FArchiveCrc32& Ar, bool bFullDataCrc) const
 {
-	// Fallback implementation uses UID to ensure every object returns a different Crc.
-	uint64 UIDValue = UID;
-	Ar << UIDValue;
+	// Add "last attribute" which can affect downstream nodes.
+	FPCGAttributePropertyInputSelector LastSelector = GetCachedLastSelector();
+	LastSelector.AddToCrc(Ar);
 }
 
 bool UPCGData::PropagateCrcThroughBooleanData() const
@@ -126,6 +126,13 @@ void UPCGData::VisitDataNetwork(TFunctionRef<void(const UPCGData*)> Action) cons
 	Action(this);
 }
 
+void UPCGData::AddUIDToCrc(FArchiveCrc32& Ar) const
+{
+	// Fallback implementation uses UID to ensure every object returns a different Crc.
+	uint64 UIDValue = UID;
+	Ar << UIDValue;
+}
+
 void UPCGData::InitUID()
 {
 	UID = ++UIDCounter;
diff --git a/Engine/Plugins/Experimental/PCG/Source/PCG/Private/PCGParamData.cpp b/Engine/Plugins/Experimental/PCG/Source/PCG/Private/PCGParamData.cpp
--- a/Engine/Plugins/Experimental/PCG/Source/PCG/Private/PCGParamData.cpp
+++ b/Engine/Plugins/Experimental/PCG/Source/PCG/Private/PCGParamData.cpp
@@ -35,9 +35,11 @@ void UPCGParamData::AddToCrc(FArchiveCrc32& Ar, bool bFullDataCrc) const
 		Metadata->GetAttributes(AttributeNames, AttributeTypes);
 	}
 
-	// Add all attribute values
+	// Add attributes to CRC
 	for (FName AttributeName : AttributeNames)
 	{
+		Ar << AttributeName;
+
 		FPCGAttributePropertyInputSelector InputSource;
 		InputSource.SetAttributeName(AttributeName);
 
diff --git a/Engine/Plugins/Experimental/PCG/Source/PCG/Private/PCGSettings.cpp b/Engine/Plugins/Experimental/PCG/Source/PCG/Private/PCGSettings.cpp
--- a/Engine/Plugins/Experimental/PCG/Source/PCG/Private/PCGSettings.cpp
+++ b/Engine/Plugins/Experimental/PCG/Source/PCG/Private/PCGSettings.cpp
@@ -232,6 +232,14 @@ void UPCGSettings::PostEditUndo()
 }
 #endif // WITH_EDITOR
 
+void UPCGSettings::AddToCrc(FArchiveCrc32& Ar, bool bFullDataCrc) const
+{
+	Super::AddToCrc(Ar, bFullDataCrc);
+
+	// This data does not have a bespoke CRC implementation so just use a global unique data CRC.
+	AddUIDToCrc(Ar);
+}
+
 void UPCGSettings::PostLoad()
 {
 	Super::PostLoad();
diff --git a/Engine/Plugins/Experimental/PCG/Source/PCG/Public/Data/PCGLandscapeData.h b/Engine/Plugins/Experimental/PCG/Source/PCG/Public/Data/PCGLandscapeData.h
--- a/Engine/Plugins/Experimental/PCG/Source/PCG/Public/Data/PCGLandscapeData.h
+++ b/Engine/Plugins/Experimental/PCG/Source/PCG/Public/Data/PCGLandscapeData.h
@@ -56,6 +56,7 @@ public:
 
 	// ~Begin UPCGData interface
 	virtual EPCGDataType GetDataType() const override { return EPCGDataType::Landscape; }
+	virtual void AddToCrc(FArchiveCrc32& Ar, bool bFullDataCrc) const override;
 	// ~End UPCGData interface
 
 	// ~Begin UPGCSpatialData interface
diff --git a/Engine/Plugins/Experimental/PCG/Source/PCG/Public/Data/PCGLandscapeSplineData.h b/Engine/Plugins/Experimental/PCG/Source/PCG/Public/Data/PCGLandscapeSplineData.h
--- a/Engine/Plugins/Experimental/PCG/Source/PCG/Public/Data/PCGLandscapeSplineData.h
+++ b/Engine/Plugins/Experimental/PCG/Source/PCG/Public/Data/PCGLandscapeSplineData.h
@@ -19,6 +19,7 @@ public:
 
 	// ~Begin UPCGData interface
 	virtual EPCGDataType GetDataType() const override { return EPCGDataType::LandscapeSpline; }
+	virtual void AddToCrc(FArchiveCrc32& Ar, bool bFullDataCrc) const override;
 	// ~End UPCGData interface
 
 	//~Begin UPCGPolyLineData interface
diff --git a/Engine/Plugins/Experimental/PCG/Source/PCG/Public/Data/PCGPrimitiveData.h b/Engine/Plugins/Experimental/PCG/Source/PCG/Public/Data/PCGPrimitiveData.h
--- a/Engine/Plugins/Experimental/PCG/Source/PCG/Public/Data/PCGPrimitiveData.h
+++ b/Engine/Plugins/Experimental/PCG/Source/PCG/Public/Data/PCGPrimitiveData.h
@@ -16,6 +16,7 @@ public:
 
 	// ~Begin UPCGData interface
 	virtual EPCGDataType GetDataType() const override { return EPCGDataType::Primitive; }
+	virtual void AddToCrc(FArchiveCrc32& Ar, bool bFullDataCrc) const override;
 	// ~End UPCGData interface
 
 	// ~Begin UPCGSpatialData interface
diff --git a/Engine/Plugins/Experimental/PCG/Source/PCG/Public/Data/PCGProjectionData.h b/Engine/Plugins/Experimental/PCG/Source/PCG/Public/Data/PCGProjectionData.h
--- a/Engine/Plugins/Experimental/PCG/Source/PCG/Public/Data/PCGProjectionData.h
+++ b/Engine/Plugins/Experimental/PCG/Source/PCG/Public/Data/PCGProjectionData.h
@@ -18,6 +18,10 @@ public:
 
 	const FPCGProjectionParams& GetProjectionParams() const { return ProjectionParams; }
 
+	// ~Begin UPCGData interface
+	virtual void AddToCrc(FArchiveCrc32& Ar, bool bFullDataCrc) const override;
+	// ~End UPCGData interface
+
 	//~Begin UPCGSpatialData interface
 	virtual int GetDimension() const override;
 	virtual FBox GetBounds() const override;
diff --git a/Engine/Plugins/Experimental/PCG/Source/PCG/Public/Data/PCGRenderTargetData.h b/Engine/Plugins/Experimental/PCG/Source/PCG/Public/Data/PCGRenderTargetData.h
--- a/Engine/Plugins/Experimental/PCG/Source/PCG/Public/Data/PCGRenderTargetData.h
+++ b/Engine/Plugins/Experimental/PCG/Source/PCG/Public/Data/PCGRenderTargetData.h
@@ -18,6 +18,7 @@ class PCG_API UPCGRenderTargetData : public UPCGBaseTextureData
 public:
 	// ~Begin UPCGData interface
 	virtual EPCGDataType GetDataType() const override { return EPCGDataType::RenderTarget; }
+	virtual void AddToCrc(FArchiveCrc32& Ar, bool bFullDataCrc) const override;
 	// ~End UPCGData interface
 
 
diff --git a/Engine/Plugins/Experimental/PCG/Source/PCG/Public/Data/PCGSpatialData.h b/Engine/Plugins/Experimental/PCG/Source/PCG/Public/Data/PCGSpatialData.h
--- a/Engine/Plugins/Experimental/PCG/Source/PCG/Public/Data/PCGSpatialData.h
+++ b/Engine/Plugins/Experimental/PCG/Source/PCG/Public/Data/PCGSpatialData.h
@@ -41,6 +41,7 @@ public:
 
 	// ~Begin UPCGData interface
 	virtual EPCGDataType GetDataType() const override { return EPCGDataType::Spatial; }
+	virtual void AddToCrc(FArchiveCrc32& Ar, bool bFullDataCrc) const override;
 
 	virtual bool HasCachedLastSelector() const override;
 	virtual FPCGAttributePropertyInputSelector GetCachedLastSelector() const override;
diff --git a/Engine/Plugins/Experimental/PCG/Source/PCG/Public/Data/PCGSplineData.h b/Engine/Plugins/Experimental/PCG/Source/PCG/Public/Data/PCGSplineData.h
--- a/Engine/Plugins/Experimental/PCG/Source/PCG/Public/Data/PCGSplineData.h
+++ b/Engine/Plugins/Experimental/PCG/Source/PCG/Public/Data/PCGSplineData.h
@@ -27,6 +27,7 @@ public:
 
 	// ~Begin UPCGData interface
 	virtual EPCGDataType GetDataType() const override { return EPCGDataType::Spline; }
+	virtual void AddToCrc(FArchiveCrc32& Ar, bool bFullDataCrc) const override;
 	// ~End UPCGData interface
 
 	//~Begin UPCGPolyLineData interface
@@ -70,6 +71,7 @@ public:
 
 	// ~Begin UPCGData interface
 	virtual EPCGDataType GetDataType() const override { return EPCGDataType::Spline; }
+	virtual void AddToCrc(FArchiveCrc32& Ar, bool bFullDataCrc) const override;
 	// ~End UPCGData interface
 
 	const UPCGSplineData* GetSpline() const;
diff --git a/Engine/Plugins/Experimental/PCG/Source/PCG/Public/Data/PCGTextureData.h b/Engine/Plugins/Experimental/PCG/Source/PCG/Public/Data/PCGTextureData.h
--- a/Engine/Plugins/Experimental/PCG/Source/PCG/Public/Data/PCGTextureData.h
+++ b/Engine/Plugins/Experimental/PCG/Source/PCG/Public/Data/PCGTextureData.h
@@ -100,12 +100,13 @@ class PCG_API UPCGTextureData : public UPCGBaseTextureData
 	GENERATED_BODY()
 
 public:
+	void Initialize(UTexture2D* InTexture, const FTransform& InTransform, const TFunction<void()>& PostInitializeCallback);
+
 	// ~Begin UPCGData interface
 	virtual EPCGDataType GetDataType() const override { return EPCGDataType::Texture; }
+	virtual void AddToCrc(FArchiveCrc32& Ar, bool bFullDataCrc) const override;
 	// ~End UPCGData interface
 
-	void Initialize(UTexture2D* InTexture, const FTransform& InTransform, const TFunction<void()>& PostInitializeCallback);
-
 	//~Begin UPCGSpatialData interface
 protected:
 	virtual UPCGSpatialData* CopyInternal() const override;
diff --git a/Engine/Plugins/Experimental/PCG/Source/PCG/Public/Data/PCGUserParametersData.h b/Engine/Plugins/Experimental/PCG/Source/PCG/Public/Data/PCGUserParametersData.h
--- a/Engine/Plugins/Experimental/PCG/Source/PCG/Public/Data/PCGUserParametersData.h
+++ b/Engine/Plugins/Experimental/PCG/Source/PCG/Public/Data/PCGUserParametersData.h
@@ -26,6 +26,7 @@ class UPCGUserParametersData : public UPCGData
 public:
 	// ~Begin UPCGData interface
 	virtual EPCGDataType GetDataType() const override { return EPCGDataType::Other; }
+	virtual void AddToCrc(FArchiveCrc32& Ar, bool bFullDataCrc) const override { AddUIDToCrc(Ar); }
 	// ~End UPCGData interface
 
 	UPROPERTY()
diff --git a/Engine/Plugins/Experimental/PCG/Source/PCG/Public/Data/PCGVolumeData.h b/Engine/Plugins/Experimental/PCG/Source/PCG/Public/Data/PCGVolumeData.h
--- a/Engine/Plugins/Experimental/PCG/Source/PCG/Public/Data/PCGVolumeData.h
+++ b/Engine/Plugins/Experimental/PCG/Source/PCG/Public/Data/PCGVolumeData.h
@@ -19,6 +19,7 @@ public:
 
 	// ~Begin UPCGData interface
 	virtual EPCGDataType GetDataType() const override { return EPCGDataType::Volume; }
+	virtual void AddToCrc(FArchiveCrc32& Ar, bool bFullDataCrc) const override;
 	// ~End UPCGData interface
 
 	// ~Begin UPGCSpatialData interface
diff --git a/Engine/Plugins/Experimental/PCG/Source/PCG/Public/Data/PCGWorldData.h b/Engine/Plugins/Experimental/PCG/Source/PCG/Public/Data/PCGWorldData.h
--- a/Engine/Plugins/Experimental/PCG/Source/PCG/Public/Data/PCGWorldData.h
+++ b/Engine/Plugins/Experimental/PCG/Source/PCG/Public/Data/PCGWorldData.h
@@ -153,6 +153,10 @@ class PCG_API UPCGWorldRayHitData : public UPCGSurfaceData
 public:
 	void Initialize(UWorld* InWorld, const FBox& InBounds = FBox(EForceInit::ForceInit));
 
+	// ~Begin UPCGData interface
+	virtual void AddToCrc(FArchiveCrc32& Ar, bool bFullDataCrc) const override;
+	// ~End UPCGData interface
+
 	//~Begin UPCGSpatialData interface
 	virtual FBox GetBounds() const override { return Bounds; }
 	virtual FBox GetStrictBounds() const override { return Bounds; }
diff --git a/Engine/Plugins/Experimental/PCG/Source/PCG/Public/Metadata/PCGAttributePropertySelector.h b/Engine/Plugins/Experimental/PCG/Source/PCG/Public/Metadata/PCGAttributePropertySelector.h
--- a/Engine/Plugins/Experimental/PCG/Source/PCG/Public/Metadata/PCGAttributePropertySelector.h
+++ b/Engine/Plugins/Experimental/PCG/Source/PCG/Public/Metadata/PCGAttributePropertySelector.h
@@ -22,6 +22,7 @@ enum class EPCGExtraProperties : uint8
 	Index
 };
 
+class FArchiveCrc32;
 struct FPCGCustomVersion;
 class UPCGData;
 
@@ -57,7 +58,6 @@ public:
 	EPCGPointProperties GetPointProperty() const { return PointProperty; }
 	EPCGExtraProperties GetExtraProperty() const { return ExtraProperty; }
 
-
 	// Return the name of the selector.
 	FName GetName() const;
 
@@ -80,6 +80,8 @@ public:
 
 	void ImportFromOtherSelector(const FPCGAttributePropertySelector& InOther);
 
+	virtual void AddToCrc(FArchiveCrc32& Ar) const;
+
 protected:
 	UPROPERTY()
 	EPCGAttributePropertySelection Selection = EPCGAttributePropertySelection::Attribute;
diff --git a/Engine/Plugins/Experimental/PCG/Source/PCG/Public/Metadata/PCGMetadata.h b/Engine/Plugins/Experimental/PCG/Source/PCG/Public/Metadata/PCGMetadata.h
--- a/Engine/Plugins/Experimental/PCG/Source/PCG/Public/Metadata/PCGMetadata.h
+++ b/Engine/Plugins/Experimental/PCG/Source/PCG/Public/Metadata/PCGMetadata.h
@@ -189,6 +189,7 @@ public:
 	UFUNCTION(BlueprintCallable, Category = "PCG|Metadata")
 	void GetAttributes(TArray<FName>& AttributeNames, TArray<EPCGMetadataTypes>& AttributeTypes) const;
 
+	/** Returns name of the most recently created attribute, or none if no attributes are present. */
 	FName GetLatestAttributeNameOrNone() const;
 
 	/** Delete/Hide attribute */
diff --git a/Engine/Plugins/Experimental/PCG/Source/PCG/Public/PCGData.h b/Engine/Plugins/Experimental/PCG/Source/PCG/Public/PCGData.h
--- a/Engine/Plugins/Experimental/PCG/Source/PCG/Public/PCGData.h
+++ b/Engine/Plugins/Experimental/PCG/Source/PCG/Public/PCGData.h
@@ -64,6 +64,8 @@ protected:
 	/** Whether intersection, union, difference combine Crc values from operands. If false they fall back to using data UID. */
 	bool PropagateCrcThroughBooleanData() const;
 
+	void AddUIDToCrc(FArchiveCrc32& Ar) const;
+
 private:
 	void InitUID();
 
diff --git a/Engine/Plugins/Experimental/PCG/Source/PCG/Public/PCGSettings.h b/Engine/Plugins/Experimental/PCG/Source/PCG/Public/PCGSettings.h
--- a/Engine/Plugins/Experimental/PCG/Source/PCG/Public/PCGSettings.h
+++ b/Engine/Plugins/Experimental/PCG/Source/PCG/Public/PCGSettings.h
@@ -188,6 +188,7 @@ class PCG_API UPCGSettings : public UPCGSettingsInterface
 public:
 	// ~Begin UPCGData interface
 	virtual EPCGDataType GetDataType() const override { return EPCGDataType::Settings; }
+	virtual void AddToCrc(FArchiveCrc32& Ar, bool bFullDataCrc) const override;
 	// ~End UPCGData interface
 
 	// ~Begin UPCGSettingsInterface interface
