[PCG] Fixed an issue where inspecting an input pin could try to retrieve the wrong inspection data from the cache

diff --git a/Engine/Plugins/Experimental/PCG/Source/PCGEditor/Private/SPCGEditorGraphAttributeListView.cpp b/Engine/Plugins/Experimental/PCG/Source/PCGEditor/Private/SPCGEditorGraphAttributeListView.cpp
--- a/Engine/Plugins/Experimental/PCG/Source/PCGEditor/Private/SPCGEditorGraphAttributeListView.cpp
+++ b/Engine/Plugins/Experimental/PCG/Source/PCGEditor/Private/SPCGEditorGraphAttributeListView.cpp
@@ -588,7 +588,7 @@ void SPCGEditorGraphAttributeListView::OnGenerateUpdated(UPCGComponent* /*InPCGC
 	RequestRefresh();
 }
 
-const FPCGDataCollection* SPCGEditorGraphAttributeListView::GetInspectionData()
+const FPCGDataCollection* SPCGEditorGraphAttributeListView::GetInspectionData() const
 {
 	if (!PCGComponent.IsValid())
 	{
@@ -601,28 +601,24 @@ const FPCGDataCollection* SPCGEditorGraphAttributeListView::GetInspectionData()
 		return nullptr;
 	}
 
-	const TSharedPtr<FPCGEditor> PCGEditor = PCGEditorPtr.Pin();
-	const FPCGStack& PCGStack = PCGEditor->GetStackBeingInspected();
-
-	const int32 SelectedIndex = GetSelectedPinIndex();
-	if (SelectedIndex == INDEX_NONE)
-	{
-		return nullptr;
-	}
-
-	// Selected pin is an output if it's at the beginning of the list, otherwise it's an input
 	const UPCGPin* Pin = nullptr;
-	const int32 InputPinCount = PCGNode->GetInputPins().Num();
-	const int32 OutputPinCount = PCGNode->GetOutputPins().Num();
-	if (SelectedIndex < (InputPinCount + OutputPinCount))
+	if (const TSharedPtr<FPinComboBoxItem> SelectedPin = PinComboBox->GetSelectedItem())
 	{
-		Pin = (SelectedIndex < OutputPinCount) ? PCGNode->GetOutputPins()[SelectedIndex] : PCGNode->GetInputPins()[SelectedIndex - OutputPinCount];
+		const TArray<TObjectPtr<UPCGPin>>& Pins = SelectedPin->bIsOutputPin ? PCGNode->GetOutputPins() : PCGNode->GetInputPins();
+		if (Pins.IsValidIndex(SelectedPin->PinIndex))
+		{
+			Pin = Pins[SelectedPin->PinIndex];
+		}
 	}
+
 	if (!Pin)
 	{
 		return nullptr;
 	}
 
+	const TSharedPtr<FPCGEditor> PCGEditor = PCGEditorPtr.Pin();
+	const FPCGStack& PCGStack = PCGEditor->GetStackBeingInspected();
+
 	// Create a temporary stack with Node+Pin to query the exact DataCollection we are inspecting
 	FPCGStack Stack = PCGStack;
 	TArray<FPCGStackFrame>& StackFrames = Stack.GetStackFramesMutable();
@@ -749,21 +745,23 @@ void SPCGEditorGraphAttributeListView::RefreshPinComboBox()
 	// Add output and then input pins to list. Optionally output the first connected item - useful for initializing
 	// the selected item to the first connected output pin.
 	auto PopulatePins = [](
-		const TArray<TObjectPtr<UPCGPin>> InPins,
+		const TArray<TObjectPtr<UPCGPin>>& InPins,
 		const FString& InFormatText,
 		TArray<TSharedPtr<FPinComboBoxItem>>& InOutItems,
 		int32* OutFirstConnectedItemIndex)
 	{
-		for (int32 i = 0; i < InPins.Num(); ++i)
+		for (int32 PinIndex = 0; PinIndex < InPins.Num(); ++PinIndex)
 		{
+			const UPCGPin* PCGPin = InPins[PinIndex];
+			const bool bIsOutputPin = PCGPin->IsOutputPin();
 			// Pin is included in list if it is connected, or if it is an output pin.
-			if (InPins[i] && (InPins[i]->IsConnected() || InPins[i]->IsOutputPin()))
+			if (PCGPin && (PCGPin->IsConnected() || bIsOutputPin))
 			{
-				FString ItemName = FString::Format(*InFormatText, { InPins[i]->Properties.Label.ToString() });
-				InOutItems.Add(MakeShared<FPinComboBoxItem>(FName(ItemName), i));
+				FString ItemName = FString::Format(*InFormatText, { PCGPin->Properties.Label.ToString() });
+				InOutItems.Add(MakeShared<FPinComboBoxItem>(FName(ItemName), PinIndex, bIsOutputPin));
 
 				// Look for first connected, null pointer once found so only first is taken.
-				if (OutFirstConnectedItemIndex && InPins[i]->IsConnected())
+				if (OutFirstConnectedItemIndex && PCGPin->IsConnected())
 				{
 					*OutFirstConnectedItemIndex = InOutItems.Num() - 1;
 					OutFirstConnectedItemIndex = nullptr;
@@ -894,17 +892,6 @@ FText SPCGEditorGraphAttributeListView::OnGenerateSelectedPinText() const
 	}
 }
 
-int32 SPCGEditorGraphAttributeListView::GetSelectedPinIndex() const
-{
-	int32 Index = INDEX_NONE;
-	if (const TSharedPtr<FPinComboBoxItem> SelectedPin = PinComboBox->GetSelectedItem())
-	{
-		PinComboBoxItems.Find(SelectedPin, Index);
-	}
-
-	return Index;
-}
-
 void SPCGEditorGraphAttributeListView::OnSelectionChangedPin(TSharedPtr<FPinComboBoxItem> InItem, ESelectInfo::Type InSelectInfo)
 {
 	RefreshDataComboBox();
diff --git a/Engine/Plugins/Experimental/PCG/Source/PCGEditor/Private/SPCGEditorGraphAttributeListView.h b/Engine/Plugins/Experimental/PCG/Source/PCGEditor/Private/SPCGEditorGraphAttributeListView.h
--- a/Engine/Plugins/Experimental/PCG/Source/PCGEditor/Private/SPCGEditorGraphAttributeListView.h
+++ b/Engine/Plugins/Experimental/PCG/Source/PCGEditor/Private/SPCGEditorGraphAttributeListView.h
@@ -131,7 +131,7 @@ public:
 	SLATE_BEGIN_ARGS(SPCGEditorGraphAttributeListView) {}
 	SLATE_END_ARGS()
 
-	virtual ~SPCGEditorGraphAttributeListView() override;
+	virtual ~SPCGEditorGraphAttributeListView();
 
 	void Construct(const FArguments& InArgs, TSharedPtr<FPCGEditor> InPCGEditor);
 
@@ -148,7 +148,7 @@ private:
 
 	void OnGenerateUpdated(UPCGComponent* InPCGComponent);
 
-	const FPCGDataCollection* GetInspectionData();
+	const FPCGDataCollection* GetInspectionData() const;
 
 	void RefreshAttributeList();
 	void RefreshPinComboBox();
@@ -159,14 +159,18 @@ private:
 	/** Only connected input pins are added to combo box, so keep track of the node pin index for each item. */
 	struct FPinComboBoxItem
 	{
-		explicit FPinComboBoxItem(FName InName, int32 InPinIndex) : Name(InName), PinIndex(InPinIndex) {}
+		explicit FPinComboBoxItem(FName InName, int32 InPinIndex, bool bInIsOutputPin)
+			: Name(InName)
+			, PinIndex(InPinIndex)
+			, bIsOutputPin(bInIsOutputPin)
+		{}
 
 		FName Name;
 		int32 PinIndex = INDEX_NONE;
+		bool bIsOutputPin = true;
 	};
 
 	FText OnGenerateSelectedPinText() const;
-	int32 GetSelectedPinIndex() const;
 	void OnSelectionChangedPin(TSharedPtr<FPinComboBoxItem> InItem, ESelectInfo::Type InSelectInfo);
 	TSharedRef<SWidget> OnGeneratePinWidget(TSharedPtr<FPinComboBoxItem> InItem) const;
 
