[PCG] Refactored PCG AssetTypeActions to AssetDefinitions
+Added right click option to create PCG graph instances from pcg graph interfaces

diff --git a/Engine/Plugins/Experimental/PCG/Source/PCGEditor/PCGEditor.Build.cs b/Engine/Plugins/Experimental/PCG/Source/PCGEditor/PCGEditor.Build.cs
--- a/Engine/Plugins/Experimental/PCG/Source/PCGEditor/PCGEditor.Build.cs
+++ b/Engine/Plugins/Experimental/PCG/Source/PCGEditor/PCGEditor.Build.cs
@@ -39,9 +39,11 @@ namespace UnrealBuildTool.Rules
 				new string[]{
 					"AppFramework",
 					"ApplicationCore",
+					"AssetDefinition",
 					"AssetTools",
 					"AssetRegistry",
 					"BlueprintGraph",
+					"ContentBrowser",
 					"DesktopWidgets",
 					"DetailCustomizations",
 					"EditorScriptingUtilities",
diff --git a/Engine/Plugins/Experimental/PCG/Source/PCGEditor/Private/AssetDefinitions/AssetDefinition_PCGGraph.cpp b/Engine/Plugins/Experimental/PCG/Source/PCGEditor/Private/AssetDefinitions/AssetDefinition_PCGGraph.cpp
new file mode 100644
--- /dev/null
+++ b/Engine/Plugins/Experimental/PCG/Source/PCGEditor/Private/AssetDefinitions/AssetDefinition_PCGGraph.cpp
@@ -0,0 +1,31 @@
+// Copyright Epic Games, Inc. All Rights Reserved.
+
+#include "AssetDefinition_PCGGraph.h"
+
+#include "PCGEditor.h"
+#include "PCGGraph.h"
+
+#define LOCTEXT_NAMESPACE "AssetDefinition_PCGGraph"
+
+FText UAssetDefinition_PCGGraph::GetAssetDisplayName() const
+{
+	return LOCTEXT("DisplayName", "PCG Graph");
+}
+
+TSoftClassPtr<UObject> UAssetDefinition_PCGGraph::GetAssetClass() const
+{
+	return UPCGGraph::StaticClass();
+}
+
+EAssetCommandResult UAssetDefinition_PCGGraph::OpenAssets(const FAssetOpenArgs& OpenArgs) const
+{
+	for (UPCGGraph* PCGGraph : OpenArgs.LoadObjects<UPCGGraph>())
+	{
+		const TSharedRef<FPCGEditor> PCGEditor = MakeShared<FPCGEditor>();
+		PCGEditor->Initialize(EToolkitMode::Standalone, OpenArgs.ToolkitHost, PCGGraph);
+	}
+
+	return EAssetCommandResult::Handled;
+}
+
+#undef LOCTEXT_NAMESPACE
\ No newline at end of file
diff --git a/Engine/Plugins/Experimental/PCG/Source/PCGEditor/Private/AssetDefinitions/AssetDefinition_PCGGraph.h b/Engine/Plugins/Experimental/PCG/Source/PCGEditor/Private/AssetDefinitions/AssetDefinition_PCGGraph.h
new file mode 100644
--- /dev/null
+++ b/Engine/Plugins/Experimental/PCG/Source/PCGEditor/Private/AssetDefinitions/AssetDefinition_PCGGraph.h
@@ -0,0 +1,18 @@
+// Copyright Epic Games, Inc. All Rights Reserved.
+
+#pragma once
+
+#include "AssetDefinition_PCGGraphInterface.h"
+
+#include "AssetDefinition_PCGGraph.generated.h"
+
+UCLASS()
+class UAssetDefinition_PCGGraph : public UAssetDefinition_PCGGraphInterface
+{
+	GENERATED_BODY()
+
+public:
+	virtual FText GetAssetDisplayName() const override;
+	virtual TSoftClassPtr<UObject> GetAssetClass() const override;
+	virtual EAssetCommandResult OpenAssets(const FAssetOpenArgs& OpenArgs) const override;
+};
diff --git a/Engine/Plugins/Experimental/PCG/Source/PCGEditor/Private/AssetDefinitions/AssetDefinition_PCGGraphInstance.cpp b/Engine/Plugins/Experimental/PCG/Source/PCGEditor/Private/AssetDefinitions/AssetDefinition_PCGGraphInstance.cpp
new file mode 100644
--- /dev/null
+++ b/Engine/Plugins/Experimental/PCG/Source/PCGEditor/Private/AssetDefinitions/AssetDefinition_PCGGraphInstance.cpp
@@ -0,0 +1,19 @@
+// Copyright Epic Games, Inc. All Rights Reserved.
+
+#include "AssetDefinition_PCGGraphInstance.h"
+
+#include "PCGGraph.h"
+
+#define LOCTEXT_NAMESPACE "AssetDefinition_PCGGraphInstance"
+
+FText UAssetDefinition_PCGGraphInstance::GetAssetDisplayName() const
+{
+	return LOCTEXT("DisplayName", "PCG Graph Instance");
+}
+
+TSoftClassPtr<UObject> UAssetDefinition_PCGGraphInstance::GetAssetClass() const
+{
+	return UPCGGraphInstance::StaticClass();
+}
+
+#undef LOCTEXT_NAMESPACE
diff --git a/Engine/Plugins/Experimental/PCG/Source/PCGEditor/Private/AssetDefinitions/AssetDefinition_PCGGraphInstance.h b/Engine/Plugins/Experimental/PCG/Source/PCGEditor/Private/AssetDefinitions/AssetDefinition_PCGGraphInstance.h
new file mode 100644
--- /dev/null
+++ b/Engine/Plugins/Experimental/PCG/Source/PCGEditor/Private/AssetDefinitions/AssetDefinition_PCGGraphInstance.h
@@ -0,0 +1,17 @@
+// Copyright Epic Games, Inc. All Rights Reserved.
+
+#pragma once
+
+#include "AssetDefinition_PCGGraphInterface.h"
+
+#include "AssetDefinition_PCGGraphInstance.generated.h"
+
+UCLASS()
+class UAssetDefinition_PCGGraphInstance : public UAssetDefinition_PCGGraphInterface
+{
+	GENERATED_BODY()
+
+public:
+	virtual FText GetAssetDisplayName() const override;
+	virtual TSoftClassPtr<UObject> GetAssetClass() const override;
+};
diff --git a/Engine/Plugins/Experimental/PCG/Source/PCGEditor/Private/AssetDefinitions/AssetDefinition_PCGGraphInterface.cpp b/Engine/Plugins/Experimental/PCG/Source/PCGEditor/Private/AssetDefinitions/AssetDefinition_PCGGraphInterface.cpp
new file mode 100644
--- /dev/null
+++ b/Engine/Plugins/Experimental/PCG/Source/PCGEditor/Private/AssetDefinitions/AssetDefinition_PCGGraphInterface.cpp
@@ -0,0 +1,80 @@
+// Copyright Epic Games, Inc. All Rights Reserved.
+
+#include "AssetDefinition_PCGGraphInterface.h"
+
+#include "PCGGraph.h"
+#include "PCGGraphFactory.h"
+
+#include "ContentBrowserMenuContexts.h"
+#include "IAssetTools.h"
+#include "ToolMenu.h"
+#include "ToolMenuEntry.h"
+#include "ToolMenuSection.h"
+#include "ToolMenus.h"
+#include "Misc/AssetCategoryPath.h"
+#include "Misc/DelayedAutoRegister.h"
+
+#define LOCTEXT_NAMESPACE "AssetDefinition_PCGGraphInterface"
+
+FText UAssetDefinition_PCGGraphInterface::GetAssetDisplayName() const
+{
+	return LOCTEXT("DisplayName", "PCG Graph Interface");
+}
+
+FLinearColor UAssetDefinition_PCGGraphInterface::GetAssetColor() const
+{
+	return FColor::Turquoise;
+}
+
+TSoftClassPtr<UObject> UAssetDefinition_PCGGraphInterface::GetAssetClass() const
+{
+	return UPCGGraphInterface::StaticClass();
+}
+
+TConstArrayView<FAssetCategoryPath> UAssetDefinition_PCGGraphInterface::GetAssetCategories() const
+{
+	static const FAssetCategoryPath Categories[] = { FAssetCategoryPath(LOCTEXT("PCGCategory", "PCG")) };
+	return Categories;
+}
+
+// Menu Extensions
+//--------------------------------------------------------------------
+
+namespace MenuExtension_PCGGraphInterface
+{
+	static void ExecuteNewPCGGraphInstance(const FToolMenuContext& MenuContext)
+	{
+		const UContentBrowserAssetContextMenuContext* CBContext = UContentBrowserAssetContextMenuContext::FindContextWithAssets(MenuContext);
+
+		IAssetTools::Get().CreateAssetsFrom<UPCGGraphInterface>(
+			CBContext->LoadSelectedObjects<UPCGGraphInterface>(), UPCGGraphInstance::StaticClass(), TEXT("_Inst"), [](UPCGGraphInterface* SourceObject)
+			{
+				UPCGGraphInstanceFactory* Factory = NewObject<UPCGGraphInstanceFactory>();
+				Factory->ParentGraph = SourceObject;
+				return Factory;
+			}
+		);
+	}
+
+	static FDelayedAutoRegisterHelper DelayedAutoRegister(EDelayedRegisterRunPhase::EndOfEngineInit, []{
+		UToolMenus::RegisterStartupCallback(FSimpleMulticastDelegate::FDelegate::CreateLambda([]()
+		{
+			FToolMenuOwnerScoped OwnerScoped(UE_MODULE_NAME);
+			UToolMenu* Menu = UE::ContentBrowser::ExtendToolMenu_AssetContextMenu(UPCGGraphInterface::StaticClass());
+
+			FToolMenuSection& Section = Menu->FindOrAddSection("GetAssetActions");
+			Section.AddDynamicEntry(NAME_None, FNewToolMenuSectionDelegate::CreateLambda([](FToolMenuSection& InSection)
+			{
+				{
+					const TAttribute<FText> Label = LOCTEXT("PCGGraph_NewInstance", "Create PCG Graph Instance");
+					const TAttribute<FText> ToolTip = LOCTEXT("PCGGraph_NewInstanceToolTip", "Creates a parameterized PCG graph using this graph as a base.");
+					const FToolMenuExecuteAction UIAction = FToolMenuExecuteAction::CreateStatic(&ExecuteNewPCGGraphInstance);
+
+					InSection.AddMenuEntry("PCGGraph_NewInstance", Label, ToolTip, FSlateIcon(), UIAction);
+				}
+			}));
+		}));
+	});
+}
+
+#undef LOCTEXT_NAMESPACE
diff --git a/Engine/Plugins/Experimental/PCG/Source/PCGEditor/Private/AssetDefinitions/AssetDefinition_PCGGraphInterface.h b/Engine/Plugins/Experimental/PCG/Source/PCGEditor/Private/AssetDefinitions/AssetDefinition_PCGGraphInterface.h
new file mode 100644
--- /dev/null
+++ b/Engine/Plugins/Experimental/PCG/Source/PCGEditor/Private/AssetDefinitions/AssetDefinition_PCGGraphInterface.h
@@ -0,0 +1,19 @@
+// Copyright Epic Games, Inc. All Rights Reserved.
+
+#pragma once
+
+#include "AssetDefinitionDefault.h"
+
+#include "AssetDefinition_PCGGraphInterface.generated.h"
+
+UCLASS()
+class UAssetDefinition_PCGGraphInterface : public UAssetDefinitionDefault
+{
+	GENERATED_BODY()
+
+public:
+	virtual FText GetAssetDisplayName() const override;
+	virtual FLinearColor GetAssetColor() const override;
+	virtual TSoftClassPtr<UObject> GetAssetClass() const override;
+	virtual TConstArrayView<FAssetCategoryPath> GetAssetCategories() const override;
+};
diff --git a/Engine/Plugins/Experimental/PCG/Source/PCGEditor/Private/AssetDefinitions/AssetDefinition_PCGSettings.cpp b/Engine/Plugins/Experimental/PCG/Source/PCGEditor/Private/AssetDefinitions/AssetDefinition_PCGSettings.cpp
new file mode 100644
--- /dev/null
+++ b/Engine/Plugins/Experimental/PCG/Source/PCGEditor/Private/AssetDefinitions/AssetDefinition_PCGSettings.cpp
@@ -0,0 +1,32 @@
+// Copyright Epic Games, Inc. All Rights Reserved.
+
+#include "AssetDefinition_PCGSettings.h"
+
+#include "PCGSettings.h"
+
+#include "Misc/AssetCategoryPath.h"
+
+#define LOCTEXT_NAMESPACE "AssetDefinition_PCGSettings"
+
+FText UAssetDefinition_PCGSettings::GetAssetDisplayName() const
+{
+	return LOCTEXT("DisplayName", "PCG Settings");
+}
+
+FLinearColor UAssetDefinition_PCGSettings::GetAssetColor() const
+{
+	return FColor::Turquoise;
+}
+
+TSoftClassPtr<UObject> UAssetDefinition_PCGSettings::GetAssetClass() const
+{
+	return UPCGSettings::StaticClass();
+}
+
+TConstArrayView<FAssetCategoryPath> UAssetDefinition_PCGSettings::GetAssetCategories() const
+{
+	static const FAssetCategoryPath Categories[] = { FAssetCategoryPath(LOCTEXT("PCGCategory", "PCG")) }; 
+	return Categories;
+}
+
+#undef LOCTEXT_NAMESPACE
diff --git a/Engine/Plugins/Experimental/PCG/Source/PCGEditor/Private/AssetDefinitions/AssetDefinition_PCGSettings.h b/Engine/Plugins/Experimental/PCG/Source/PCGEditor/Private/AssetDefinitions/AssetDefinition_PCGSettings.h
new file mode 100644
--- /dev/null
+++ b/Engine/Plugins/Experimental/PCG/Source/PCGEditor/Private/AssetDefinitions/AssetDefinition_PCGSettings.h
@@ -0,0 +1,19 @@
+// Copyright Epic Games, Inc. All Rights Reserved.
+
+#pragma once
+
+#include "AssetDefinitionDefault.h"
+
+#include "AssetDefinition_PCGSettings.generated.h"
+
+UCLASS()
+class UAssetDefinition_PCGSettings : public UAssetDefinitionDefault
+{
+	GENERATED_BODY()
+
+public:
+	virtual FText GetAssetDisplayName() const override;
+	virtual FLinearColor GetAssetColor() const override;
+	virtual TSoftClassPtr<UObject> GetAssetClass() const override;
+	virtual TConstArrayView<FAssetCategoryPath> GetAssetCategories() const override;
+};
diff --git a/Engine/Plugins/Experimental/PCG/Source/PCGEditor/Private/AssetTypeActions/PCGCommonAssetTypeActions.cpp b/Engine/Plugins/Experimental/PCG/Source/PCGEditor/Private/AssetTypeActions/PCGCommonAssetTypeActions.cpp
deleted file mode 100644
diff --git a/Engine/Plugins/Experimental/PCG/Source/PCGEditor/Private/AssetTypeActions/PCGCommonAssetTypeActions.h b/Engine/Plugins/Experimental/PCG/Source/PCGEditor/Private/AssetTypeActions/PCGCommonAssetTypeActions.h
deleted file mode 100644
diff --git a/Engine/Plugins/Experimental/PCG/Source/PCGEditor/Private/AssetTypeActions/PCGGraphAssetTypeActions.cpp b/Engine/Plugins/Experimental/PCG/Source/PCGEditor/Private/AssetTypeActions/PCGGraphAssetTypeActions.cpp
deleted file mode 100644
diff --git a/Engine/Plugins/Experimental/PCG/Source/PCGEditor/Private/AssetTypeActions/PCGGraphAssetTypeActions.h b/Engine/Plugins/Experimental/PCG/Source/PCGEditor/Private/AssetTypeActions/PCGGraphAssetTypeActions.h
deleted file mode 100644
diff --git a/Engine/Plugins/Experimental/PCG/Source/PCGEditor/Private/AssetTypeActions/PCGSettingsAssetTypeActions.cpp b/Engine/Plugins/Experimental/PCG/Source/PCGEditor/Private/AssetTypeActions/PCGSettingsAssetTypeActions.cpp
deleted file mode 100644
diff --git a/Engine/Plugins/Experimental/PCG/Source/PCGEditor/Private/AssetTypeActions/PCGSettingsAssetTypeActions.h b/Engine/Plugins/Experimental/PCG/Source/PCGEditor/Private/AssetTypeActions/PCGSettingsAssetTypeActions.h
deleted file mode 100644
diff --git a/Engine/Plugins/Experimental/PCG/Source/PCGEditor/Private/PCGEditorModule.cpp b/Engine/Plugins/Experimental/PCG/Source/PCGEditor/Private/PCGEditorModule.cpp
--- a/Engine/Plugins/Experimental/PCG/Source/PCGEditor/Private/PCGEditorModule.cpp
+++ b/Engine/Plugins/Experimental/PCG/Source/PCGEditor/Private/PCGEditorModule.cpp
@@ -9,12 +9,8 @@
 #include "PCGEditorUtils.h"
 #include "PCGSubsystem.h"
 #include "PCGVolumeFactory.h"
-#include "AssetTypeActions/PCGCommonAssetTypeActions.h"
-#include "AssetTypeActions/PCGGraphAssetTypeActions.h"
-#include "AssetTypeActions/PCGSettingsAssetTypeActions.h"
 
 #include "ISettingsModule.h"
-#include "LevelEditor.h"
 #include "PropertyEditorModule.h"
 #include "ToolMenus.h"
 #include "Details/PCGAttributePropertySelectorDetails.h"
@@ -27,12 +23,9 @@
 
 #define LOCTEXT_NAMESPACE "FPCGEditorModule"
 
-EAssetTypeCategories::Type FPCGEditorModule::PCGAssetCategory;
-
 void FPCGEditorModule::StartupModule()
 {
 	RegisterDetailsCustomizations();
-	RegisterAssetTypeActions();
 	RegisterMenuExtensions();
 	RegisterSettings();
 
@@ -51,7 +44,6 @@ void FPCGEditorModule::StartupModule()
 void FPCGEditorModule::ShutdownModule()
 {
 	UnregisterSettings();
-	UnregisterAssetTypeActions();
 	UnregisterDetailsCustomizations();
 	UnregisterMenuExtensions();
 
@@ -109,35 +101,6 @@ void FPCGEditorModule::UnregisterDetailsCustomizations()
 	}
 }
 
-void FPCGEditorModule::RegisterAssetTypeActions()
-{
-	IAssetTools& AssetTools = FModuleManager::LoadModuleChecked<FAssetToolsModule>("AssetTools").Get();
-	PCGAssetCategory = AssetTools.RegisterAdvancedAssetCategory(FName(TEXT("PCG")), LOCTEXT("PCGAssetCategory", "PCG"));
-
-	RegisteredAssetTypeActions.Emplace(MakeShareable(new FPCGGraphAssetTypeActions()));
-	RegisteredAssetTypeActions.Emplace(MakeShareable(new FPCGGraphInstanceAssetTypeActions()));
-	RegisteredAssetTypeActions.Emplace(MakeShareable(new FPCGGraphInterfaceAssetTypeActions()));
-	RegisteredAssetTypeActions.Emplace(MakeShareable(new FPCGSettingsAssetTypeActions()));
-
-	for (auto Action : RegisteredAssetTypeActions)
-	{
-		AssetTools.RegisterAssetTypeActions(Action);
-	}
-}
-
-void FPCGEditorModule::UnregisterAssetTypeActions()
-{
-	if (FAssetToolsModule* AssetToolsModule = FModuleManager::GetModulePtr<FAssetToolsModule>("AssetTools"))
-	{
-		IAssetTools& AssetTools = AssetToolsModule->Get();
-
-		for (auto Action : RegisteredAssetTypeActions)
-		{
-			AssetTools.UnregisterAssetTypeActions(Action);
-		}
-	}
-}
-
 void FPCGEditorModule::RegisterMenuExtensions()
 {
 	FToolMenuOwnerScoped OwnerScoped(this);
diff --git a/Engine/Plugins/Experimental/PCG/Source/PCGEditor/Private/PCGEditorModule.h b/Engine/Plugins/Experimental/PCG/Source/PCGEditor/Private/PCGEditorModule.h
--- a/Engine/Plugins/Experimental/PCG/Source/PCGEditor/Private/PCGEditorModule.h
+++ b/Engine/Plugins/Experimental/PCG/Source/PCGEditor/Private/PCGEditorModule.h
@@ -23,13 +23,9 @@ public:
 	virtual bool SupportsDynamicReloading() override;
 	// ~End IModuleInterface implementation
 
-	static EAssetTypeCategories::Type GetAssetCategory() { return PCGAssetCategory; }
-
 protected:
 	void RegisterDetailsCustomizations();
 	void UnregisterDetailsCustomizations();
-	void RegisterAssetTypeActions();
-	void UnregisterAssetTypeActions();
 	void RegisterMenuExtensions();
 	void UnregisterMenuExtensions();
 	void PopulateMenuActions(FMenuBuilder& MenuBuilder);
diff --git a/Engine/Plugins/Experimental/PCG/Source/PCGEditor/Private/PCGGraphFactory.cpp b/Engine/Plugins/Experimental/PCG/Source/PCGEditor/Private/PCGGraphFactory.cpp
--- a/Engine/Plugins/Experimental/PCG/Source/PCGEditor/Private/PCGGraphFactory.cpp
+++ b/Engine/Plugins/Experimental/PCG/Source/PCGEditor/Private/PCGGraphFactory.cpp
@@ -39,7 +39,14 @@ UPCGGraphInstanceFactory::UPCGGraphInstanceFactory(const FObjectInitializer& Obj
 
 UObject* UPCGGraphInstanceFactory::FactoryCreateNew(UClass* InClass, UObject* InParent, FName InName, EObjectFlags Flags, UObject* Context, FFeedbackContext* Warn)
 {
-	return NewObject<UPCGGraphInstance>(InParent, InClass, InName, Flags);
+	UPCGGraphInstance* PCGGraphInstance = NewObject<UPCGGraphInstance>(InParent, InClass, InName, Flags);
+
+	if (ParentGraph)
+	{
+		PCGGraphInstance->SetGraph(ParentGraph);
+	}
+	
+	return PCGGraphInstance;
 }
 
 bool UPCGGraphInstanceFactory::ShouldShowInNewMenu() const
diff --git a/Engine/Plugins/Experimental/PCG/Source/PCGEditor/Private/PCGGraphFactory.h b/Engine/Plugins/Experimental/PCG/Source/PCGEditor/Private/PCGGraphFactory.h
--- a/Engine/Plugins/Experimental/PCG/Source/PCGEditor/Private/PCGGraphFactory.h
+++ b/Engine/Plugins/Experimental/PCG/Source/PCGEditor/Private/PCGGraphFactory.h
@@ -6,6 +6,8 @@
 
 #include "PCGGraphFactory.generated.h"
 
+class UPCGGraphInterface;
+
 UCLASS(hidecategories=Object)
 class UPCGGraphFactory : public UFactory
 {
@@ -27,6 +29,9 @@ class UPCGGraphInstanceFactory : public UFactory
 public:
 	UPCGGraphInstanceFactory(const FObjectInitializer& ObjectInitializer);
 
+	UPROPERTY()
+	TObjectPtr<UPCGGraphInterface> ParentGraph;
+	
 	//~UFactory interface
 	virtual UObject* FactoryCreateNew(UClass* InClass, UObject* InParent, FName InName, EObjectFlags Flags, UObject* Context, FFeedbackContext* Warn) override;
 	virtual bool ShouldShowInNewMenu() const override;
