[PCG] Expose PCGGraphInterface to BP (and by extension Graph Parameters)

* Also fixing an issue where we were updating the graph overrides if we had a parameter named "Graph"

#rnx

diff --git a/Engine/Plugins/Experimental/PCG/Source/PCG/Private/PCGGraph.cpp b/Engine/Plugins/Experimental/PCG/Source/PCG/Private/PCGGraph.cpp
--- a/Engine/Plugins/Experimental/PCG/Source/PCG/Private/PCGGraph.cpp
+++ b/Engine/Plugins/Experimental/PCG/Source/PCG/Private/PCGGraph.cpp
@@ -1413,7 +1413,8 @@ void UPCGGraphInstance::PreEditChange(FProperty* InProperty)
 		return;
 	}
 
-	if (InProperty->GetFName() == GET_MEMBER_NAME_CHECKED(UPCGGraphInstance, Graph))
+	// We need to be careful and only capture `Graph` if it is our graph and not a graph parameter called `Graph`!
+	if (InProperty->GetOwnerClass() == UPCGGraphInstance::StaticClass() && InProperty->GetFName() == GET_MEMBER_NAME_CHECKED(UPCGGraphInstance, Graph))
 	{
 		TeardownCallbacks();
 	}
@@ -1428,7 +1429,8 @@ void UPCGGraphInstance::PostEditChangeProperty(FPropertyChangedEvent& PropertyCh
 		return;
 	}
 
-	if (PropertyChangedEvent.GetPropertyName() == GET_MEMBER_NAME_CHECKED(UPCGGraphInstance, Graph))
+	// We need to be careful and only capture `Graph` if it is our graph and not a graph parameter called `Graph`!
+	if (PropertyChangedEvent.Property->GetOwnerClass() == UPCGGraphInstance::StaticClass() && PropertyChangedEvent.GetPropertyName() == GET_MEMBER_NAME_CHECKED(UPCGGraphInstance, Graph))
 	{
 		SetupCallbacks();
 
diff --git a/Engine/Plugins/Experimental/PCG/Source/PCG/Public/PCGGraph.h b/Engine/Plugins/Experimental/PCG/Source/PCG/Public/PCGGraph.h
--- a/Engine/Plugins/Experimental/PCG/Source/PCG/Public/PCGGraph.h
+++ b/Engine/Plugins/Experimental/PCG/Source/PCG/Public/PCGGraph.h
@@ -80,7 +80,7 @@ public:
 	TSet<FGuid> PropertiesIDsOverridden;
 };
 
-UCLASS(Abstract)
+UCLASS(BlueprintType, Abstract, ClassGroup = (Procedural))
 class PCG_API UPCGGraphInterface : public UObject
 {
 	GENERATED_BODY()
