[PCG] Fix surface sampler not executing if expected point count < 1

#rnx

diff --git a/Engine/Plugins/Experimental/PCG/Source/PCG/Private/Elements/PCGSurfaceSampler.cpp b/Engine/Plugins/Experimental/PCG/Source/PCG/Private/Elements/PCGSurfaceSampler.cpp
--- a/Engine/Plugins/Experimental/PCG/Source/PCG/Private/Elements/PCGSurfaceSampler.cpp
+++ b/Engine/Plugins/Experimental/PCG/Source/PCG/Private/Elements/PCGSurfaceSampler.cpp
@@ -96,23 +96,18 @@ namespace PCGSurfaceSampler
 		check(CellCount > 0);
 
 		const FVector::FReal InvSquaredMeterUnits = 1.0 / (100.0 * 100.0);
-		TargetPointCount = (InputBounds.Max.X - InputBounds.Min.X) * (InputBounds.Max.Y - InputBounds.Min.Y) * PointsPerSquaredMeter * InvSquaredMeterUnits;
+		const FVector::FReal TargetPointCount = (InputBounds.Max.X - InputBounds.Min.X) * (InputBounds.Max.Y - InputBounds.Min.Y) * PointsPerSquaredMeter * InvSquaredMeterUnits;
+		Ratio = static_cast<float>(FMath::Clamp(TargetPointCount / (FVector::FReal)CellCount, 0.0, 1.0));
 
-		if (TargetPointCount == 0)
+		if (Ratio < UE_SMALL_NUMBER)
 		{
 			if (Context)
 			{
 				PCGE_LOG_C(Verbose, LogOnly, Context, LOCTEXT("NoPointsFromDensity", "Skipped - density yields no points"));
 			}
-			
+
 			return false;
 		}
-		else if (TargetPointCount > CellCount)
-		{
-			TargetPointCount = CellCount;
-		}
-
-		Ratio = TargetPointCount / (FVector::FReal)CellCount;
 
 		InputBoundsMinZ = InputBounds.Min.Z;
 		InputBoundsMaxZ = InputBounds.Max.Z;
diff --git a/Engine/Plugins/Experimental/PCG/Source/PCG/Public/Elements/PCGSurfaceSampler.h b/Engine/Plugins/Experimental/PCG/Source/PCG/Public/Elements/PCGSurfaceSampler.h
--- a/Engine/Plugins/Experimental/PCG/Source/PCG/Public/Elements/PCGSurfaceSampler.h
+++ b/Engine/Plugins/Experimental/PCG/Source/PCG/Public/Elements/PCGSurfaceSampler.h
@@ -40,7 +40,6 @@ namespace PCGSurfaceSampler
 		int32 CellMinY;
 		int32 CellMaxY;
 		int32 CellCount;
-		int64 TargetPointCount;
 		float Ratio;
 		int Seed;
 
