[PCG] Allow SOP/SCP in some attribute nodes, and remove assert which is not required/useful

#rnx

diff --git a/Engine/Plugins/Experimental/PCG/Source/PCG/Private/PCGGraph.cpp b/Engine/Plugins/Experimental/PCG/Source/PCG/Private/PCGGraph.cpp
--- a/Engine/Plugins/Experimental/PCG/Source/PCG/Private/PCGGraph.cpp
+++ b/Engine/Plugins/Experimental/PCG/Source/PCG/Private/PCGGraph.cpp
@@ -1130,7 +1130,6 @@ void UPCGGraph::FixInvalidEdges()
 					UE_LOG(LogPCG, Error, TEXT("Removed edge to a node '%s' that is not registered in the graph."),
 						ConnectedNode ? *ConnectedNode->GetFName().ToString() : TEXT("NULL"),
 						*GetFName().ToString());
-					ensure(false);
 
 					Pin->Edges.RemoveAt(i);
 				}
diff --git a/Engine/Plugins/Experimental/PCG/Source/PCG/Public/Elements/PCGPointMatchAndSet.h b/Engine/Plugins/Experimental/PCG/Source/PCG/Public/Elements/PCGPointMatchAndSet.h
--- a/Engine/Plugins/Experimental/PCG/Source/PCG/Public/Elements/PCGPointMatchAndSet.h
+++ b/Engine/Plugins/Experimental/PCG/Source/PCG/Public/Elements/PCGPointMatchAndSet.h
@@ -67,7 +67,7 @@ public:
 	FPCGAttributePropertyOutputSelector SetTarget;
 
 	/** If the "Set" part is an attribute, then the type must be provided */
-	UPROPERTY(BlueprintReadWrite, EditAnywhere, Category = Settings, meta = (EditCondition = "bSetTargetIsAttribute", HideEditConditionToggle, EditConditionHides, ValidEnumValues = "Float, Double, Integer32, Integer64, Vector2, Vector, Vector4, Quaternion, Transform, String, Boolean, Rotator, Name"))
+	UPROPERTY(BlueprintReadWrite, EditAnywhere, Category = Settings, meta = (EditCondition = "bSetTargetIsAttribute", HideEditConditionToggle, EditConditionHides))
 	EPCGMetadataTypes SetTargetType = EPCGMetadataTypes::Double;
 
 	/** For string types, the subtype is used to cleanup the UI. */
diff --git a/Engine/Plugins/Experimental/PCG/Source/PCG/Public/MatchAndSet/PCGMatchAndSetByAttribute.h b/Engine/Plugins/Experimental/PCG/Source/PCG/Public/MatchAndSet/PCGMatchAndSetByAttribute.h
--- a/Engine/Plugins/Experimental/PCG/Source/PCG/Public/MatchAndSet/PCGMatchAndSetByAttribute.h
+++ b/Engine/Plugins/Experimental/PCG/Source/PCG/Public/MatchAndSet/PCGMatchAndSetByAttribute.h
@@ -62,7 +62,7 @@ public:
 	FName MatchSourceAttribute;
 
 	/** Type of the attribute to match against. */
-	UPROPERTY(BlueprintReadWrite, EditAnywhere, Category = Settings, meta = (ValidEnumValues = "Float, Double, Integer32, Integer64, Vector2, Vector, Vector4, Quaternion, Transform, String, Boolean, Rotator, Name"))
+	UPROPERTY(BlueprintReadWrite, EditAnywhere, Category = Settings)
 	EPCGMetadataTypes MatchSourceType = EPCGMetadataTypes::Double;
 
 	/** String type of the attribute to match against (if required). */
diff --git a/Engine/Plugins/Experimental/PCG/Source/PCG/Public/MatchAndSet/PCGMatchAndSetWeightedByCategory.h b/Engine/Plugins/Experimental/PCG/Source/PCG/Public/MatchAndSet/PCGMatchAndSetWeightedByCategory.h
--- a/Engine/Plugins/Experimental/PCG/Source/PCG/Public/MatchAndSet/PCGMatchAndSetWeightedByCategory.h
+++ b/Engine/Plugins/Experimental/PCG/Source/PCG/Public/MatchAndSet/PCGMatchAndSetWeightedByCategory.h
@@ -69,7 +69,7 @@ public:
 	FName CategoryAttribute;
 
 	/** Type of the attribute to match against. */
-	UPROPERTY(BlueprintReadWrite, EditAnywhere, Category = Settings, meta = (ValidEnumValues = "Float, Double, Integer32, Integer64, Vector2, Vector, Vector4, Quaternion, Transform, String, Boolean, Rotator, Name"))
+	UPROPERTY(BlueprintReadWrite, EditAnywhere, Category = Settings)
 	EPCGMetadataTypes CategoryType = EPCGMetadataTypes::Double;
 
 	UPROPERTY()
diff --git a/Engine/Plugins/Experimental/PCG/Source/PCG/Public/Metadata/PCGMetadataAttributeTraits.h b/Engine/Plugins/Experimental/PCG/Source/PCG/Public/Metadata/PCGMetadataAttributeTraits.h
--- a/Engine/Plugins/Experimental/PCG/Source/PCG/Public/Metadata/PCGMetadataAttributeTraits.h
+++ b/Engine/Plugins/Experimental/PCG/Source/PCG/Public/Metadata/PCGMetadataAttributeTraits.h
@@ -29,7 +29,7 @@ enum class EPCGMetadataTypes : uint8
 
 	Count UMETA(Hidden),
 
-	Unknown = 255
+	Unknown = 255 UMETA(Hidden),
 };
 
 // Convenient macro to avoid duplicating a lot of code with all our supported types.
