[PCG] Support Texture2DArrays in Texture Sampler (in-editor only)

diff --git a/Engine/Plugins/Experimental/PCG/Source/PCG/Private/Elements/PCGTextureSampler.cpp b/Engine/Plugins/Experimental/PCG/Source/PCG/Private/Elements/PCGTextureSampler.cpp
--- a/Engine/Plugins/Experimental/PCG/Source/PCG/Private/Elements/PCGTextureSampler.cpp
+++ b/Engine/Plugins/Experimental/PCG/Source/PCG/Private/Elements/PCGTextureSampler.cpp
@@ -9,13 +9,39 @@
 #include "Helpers/PCGHelpers.h"
 #include "Helpers/PCGSettingsHelpers.h"
 
+#include "Engine/Texture.h"
 #include "Engine/Texture2D.h"
+#include "Engine/Texture2DArray.h"
 #include "GameFramework/Actor.h"
 
 #include UE_INLINE_GENERATED_CPP_BY_NAME(PCGTextureSampler)
 
 #define LOCTEXT_NAMESPACE "PCGTextureSamplerElement"
 
+#if WITH_EDITOR
+void UPCGTextureSamplerSettings::PostEditChangeProperty(FPropertyChangedEvent& PropertyChangedEvent) 
+{
+	if (PropertyChangedEvent.Property)
+	{
+		const FName PropertyName = PropertyChangedEvent.Property->GetFName();
+
+		if (PropertyName == GET_MEMBER_NAME_CHECKED(UPCGTextureSamplerSettings, Texture))
+		{
+			UpdateDisplayTextureArrayIndex();
+		} 
+	}
+
+	Super::PostEditChangeProperty(PropertyChangedEvent);
+}
+
+void UPCGTextureSamplerSettings::PostLoad()
+{
+	Super::PostLoad();
+
+	UpdateDisplayTextureArrayIndex();
+}
+#endif
+
 #if WITH_EDITOR
 FText UPCGTextureSamplerSettings::GetNodeTooltipText() const
 {
@@ -38,6 +64,14 @@ FPCGElementPtr UPCGTextureSamplerSettings::CreateElement() const
 	return MakeShared<FPCGTextureSamplerElement>();
 }
 
+#if WITH_EDITOR
+void UPCGTextureSamplerSettings::UpdateDisplayTextureArrayIndex()
+{
+	UTexture* NewTexture = Texture.LoadSynchronous();
+	bDisplayTextureArrayIndex = NewTexture && NewTexture->IsA<UTexture2DArray>();
+}
+#endif
+
 bool FPCGTextureSamplerElement::ExecuteInternal(FPCGContext* InContext) const
 {
 	TRACE_CPUPROFILER_EVENT_SCOPE(FPCGTextureSamplerElement::Execute);
@@ -62,7 +96,7 @@ bool FPCGTextureSamplerElement::ExecuteInternal(FPCGContext* InContext) const
 		return true;
 	}
 
-	UTexture2D* Texture = Settings->Texture.LoadSynchronous();
+	UTexture* Texture = Settings->Texture.LoadSynchronous();
 
 	if (!Texture)
 	{
@@ -70,6 +104,44 @@ bool FPCGTextureSamplerElement::ExecuteInternal(FPCGContext* InContext) const
 		return true;
 	}
 
+	UTexture2D* Texture2D = Cast<UTexture2D>(Texture);
+
+	// If the type is not a UTexture2D, try UTexture2DArray instead
+	if (!Texture2D)
+	{
+#if WITH_EDITOR
+		// TODO: support Texture2DArray without editor
+		if (UTexture2DArray* Tex2DArray = Cast<UTexture2DArray>(Texture))
+		{
+			if (Tex2DArray->SourceTextures.IsValidIndex(Settings->TextureArrayIndex))
+			{
+				Texture2D = Tex2DArray->SourceTextures[Settings->TextureArrayIndex];
+
+				if (!Texture2D)
+				{
+					PCGE_LOG(Error, GraphAndLog, FText::Format(LOCTEXT("CouldNotResolveTextureArray", "Texture index {0} could not be loaded."), Settings->TextureArrayIndex));
+				}
+			}
+			else
+			{
+				PCGE_LOG(Error, GraphAndLog, FText::Format(LOCTEXT("TextureArrayIndexOutOfBounds", "Texture array index {0} was out of bounds. There are only {1} textures in the array."), Settings->TextureArrayIndex, Tex2DArray->SourceTextures.Num()));
+			}
+		}
+		else
+		{
+			PCGE_LOG(Error, GraphAndLog, FText::Format(LOCTEXT("InvalidTextureType", "Texture at path '{0}' was not of valid type. Must be either Texture2D or Texture2DArray."), FText::FromString(Settings->Texture.ToString())));
+		}
+#else
+		PCGE_LOG(Error, GraphAndLog, FText::Format(LOCTEXT("InvalidTextureTypeWithoutEditor", "Texture at path '{0}' was not of valid type. Must be a Texture2D when built without editor."), FText::FromString(Settings->Texture.ToString())));
+#endif
+	}
+
+	if (!Texture2D)
+	{
+		PCGE_LOG(Error, GraphAndLog, FText::Format(LOCTEXT("InvalidTexture2D", "Texture at path '{0}' could not evaluate to a valid 2D Texture"), FText::FromString(Settings->Texture.ToString())));
+		return true;
+	}
+
 	const FTransform& Transform = Settings->Transform;
 	const bool bUseAbsoluteTransform = Settings->bUseAbsoluteTransform;
 	const EPCGTextureDensityFunction DensityFunction = Settings->DensityFunction;
@@ -114,7 +186,7 @@ bool FPCGTextureSamplerElement::ExecuteInternal(FPCGContext* InContext) const
 		}
 	};
 
-	TextureData->Initialize(Texture, FinalTransform, PostInitializeCallback);
+	TextureData->Initialize(Texture2D, FinalTransform, PostInitializeCallback);
 
 	TextureData->DensityFunction = DensityFunction;
 	TextureData->ColorChannel = ColorChannel;
diff --git a/Engine/Plugins/Experimental/PCG/Source/PCG/Public/Elements/PCGTextureSampler.h b/Engine/Plugins/Experimental/PCG/Source/PCG/Public/Elements/PCGTextureSampler.h
--- a/Engine/Plugins/Experimental/PCG/Source/PCG/Public/Elements/PCGTextureSampler.h
+++ b/Engine/Plugins/Experimental/PCG/Source/PCG/Public/Elements/PCGTextureSampler.h
@@ -8,12 +8,21 @@
 
 #include "PCGTextureSampler.generated.h"
 
+class UTexture;
+
 UCLASS(BlueprintType, ClassGroup = (Procedural))
 class PCG_API UPCGTextureSamplerSettings : public UPCGSettings
 {
 	GENERATED_BODY()
 
 public:
+	//~Begin UObject interface
+#if WITH_EDITOR
+	virtual void PostEditChangeProperty(FPropertyChangedEvent& PropertyChangedEvent) override;
+	virtual void PostLoad() override;
+#endif
+	//~End UObject interface
+
 	//~Begin UPCGSettings interface
 #if WITH_EDITOR
 	virtual FName GetDefaultNodeName() const override { return FName(TEXT("GetTextureData")); }
@@ -28,17 +37,31 @@ protected:
 	virtual FPCGElementPtr CreateElement() const override;
 	//~End UPCGSettings interface
 
+#if WITH_EDITOR
+	void UpdateDisplayTextureArrayIndex();
+#endif
+
 public:
-	// Surface transform
+	/** Surface transform */
 	UPROPERTY(BlueprintReadWrite, EditAnywhere, Category = Settings, meta = (PCG_Overridable))
 	FTransform Transform = FTransform::Identity;
 
 	UPROPERTY(BlueprintReadWrite, EditAnywhere, Category = Settings, meta = (PCG_Overridable))
 	bool bUseAbsoluteTransform = false;
 
-	// Texture specific parameters
+	/** Texture specific parameters */
 	UPROPERTY(BlueprintReadWrite, EditAnywhere, Category = Settings, meta = (PCG_Overridable))
-	TSoftObjectPtr<UTexture2D> Texture = nullptr;
+	TSoftObjectPtr<UTexture> Texture = nullptr;
+
+#if WITH_EDITORONLY_DATA
+	/** Index of texture array slice. Only used when built with editor and if the type of Texture is Texture2DArray. */
+	UPROPERTY(BlueprintReadWrite, EditAnywhere, Category = Settings, meta = (EditCondition = bDisplayTextureArrayIndex, EditConditionHides, HideEditConditionToggle, ClampMin = '0', PCG_Overridable))
+	int TextureArrayIndex = 0;
+
+	// This can be set false by inheriting nodes to hide the 'TextureArrayIndex' property.
+	UPROPERTY(Transient, meta = (EditCondition = false, EditConditionHides))
+	bool bDisplayTextureArrayIndex = false;
+#endif
 
 	// Common members in BaseTextureData
 	UPROPERTY(BlueprintReadWrite, VisibleAnywhere, Category = SpatialData, meta = (PCG_Overridable))
