[PCG] Add more options on the maths node

* For rounding operations, add the option to force the output in int.
* For operations that can yield floating point values, add the option to force the output in double.

diff --git a/Engine/Plugins/Experimental/PCG/Source/PCG/Private/Elements/Metadata/PCGMetadataMathsOpElement.cpp b/Engine/Plugins/Experimental/PCG/Source/PCG/Private/Elements/Metadata/PCGMetadataMathsOpElement.cpp
--- a/Engine/Plugins/Experimental/PCG/Source/PCG/Private/Elements/Metadata/PCGMetadataMathsOpElement.cpp
+++ b/Engine/Plugins/Experimental/PCG/Source/PCG/Private/Elements/Metadata/PCGMetadataMathsOpElement.cpp
@@ -222,6 +222,42 @@ bool UPCGMetadataMathsSettings::IsSupportedInputType(uint16 TypeId, uint32 Input
 	return PCG::Private::IsOfTypes<float, double, int32, int64, FVector2D, FVector, FVector4>(TypeId);
 }
 
+bool UPCGMetadataMathsSettings::ShouldForceOutputToInt(uint16 InputTypeId) const
+{
+	return PCG::Private::IsOfTypes<float, double>(InputTypeId) && bForceRoundingOpToInt &&
+		(Operation == EPCGMedadataMathsOperation::Round ||
+		Operation == EPCGMedadataMathsOperation::Truncate ||
+		Operation == EPCGMedadataMathsOperation::Floor ||
+		Operation == EPCGMedadataMathsOperation::Ceil);
+}
+
+bool UPCGMetadataMathsSettings::ShouldForceOutputToDouble(uint16 InputTypeId) const
+{
+	return PCG::Private::IsOfTypes<int32, int64>(InputTypeId) && bForceOpToDouble &&
+		(Operation == EPCGMedadataMathsOperation::Divide ||
+		Operation == EPCGMedadataMathsOperation::Sqrt ||
+		Operation == EPCGMedadataMathsOperation::Pow ||
+		Operation == EPCGMedadataMathsOperation::Lerp);
+}
+
+uint16 UPCGMetadataMathsSettings::GetOutputType(uint16 InputTypeId) const
+{
+	// If attribute Type is a float or double, can convert to int if it is a rounding op.
+	if (ShouldForceOutputToInt(InputTypeId))
+	{
+		return PCG::Private::MetadataTypes<int64>::Id;
+	}
+	// If attribute Type is an integer, can convert to double if it is an operation that can yield a floating point value.
+	else if (ShouldForceOutputToDouble(InputTypeId))
+	{
+		return PCG::Private::MetadataTypes<double>::Id;
+	}
+	else
+	{
+		return InputTypeId;
+	}
+}
+
 FPCGAttributePropertyInputSelector UPCGMetadataMathsSettings::GetInputSource(uint32 Index) const
 {
 	switch (Index)
@@ -302,7 +338,17 @@ bool FPCGMetadataMathsElement::DoOperation(FOperationData& OperationData) const
 			if (PCGMetadataMathsSettings::IsUnaryOp(Operation))
 			{
 				TRACE_CPUPROFILER_EVENT_SCOPE(FPCGMetadataMathsElement::ExecuteInternal::UnaryOp);
-				DoUnaryOp<AttributeType>(OperationData, [Operation](const AttributeType& Value) -> AttributeType { return PCGMetadataMathsSettings::UnaryOp(Value, Operation); });
+				// For int64 as output of the lambda if the output type is int64, as AttributeType might be different (cf GetOutputType)
+				using OverriddenOutputType = typename std::conditional_t<PCG::Private::IsOfTypes<AttributeType, float, double>(), int64, AttributeType>;
+
+				if (OperationData.OutputType == PCG::Private::MetadataTypes<int64>::Id)
+				{
+					DoUnaryOp<AttributeType>(OperationData, [Operation](const AttributeType& Value) -> OverriddenOutputType { return static_cast<OverriddenOutputType>(PCGMetadataMathsSettings::UnaryOp<AttributeType>(Value, Operation)); });
+				}
+				else
+				{
+					DoUnaryOp<AttributeType>(OperationData, [Operation](const AttributeType& Value) -> AttributeType { return PCGMetadataMathsSettings::UnaryOp(Value, Operation); });
+				}
 			}
 			else if (PCGMetadataMathsSettings::IsBinaryOp(Operation))
 			{
@@ -317,7 +363,15 @@ bool FPCGMetadataMathsElement::DoOperation(FOperationData& OperationData) const
 		}
 	};
 
-	PCGMetadataAttribute::CallbackWithRightType(OperationData.OutputType, MathFunc);
+	// If the output is double, force all to double.
+	if (OperationData.OutputType == PCG::Private::MetadataTypes<double>::Id)
+	{
+		MathFunc(double{});
+	}
+	else
+	{
+		PCGMetadataAttribute::CallbackWithRightType(OperationData.MostComplexInputType, MathFunc);
+	}
 
 	return true;
 }
diff --git a/Engine/Plugins/Experimental/PCG/Source/PCG/Public/Elements/Metadata/PCGMetadataMathsOpElement.h b/Engine/Plugins/Experimental/PCG/Source/PCG/Public/Elements/Metadata/PCGMetadataMathsOpElement.h
--- a/Engine/Plugins/Experimental/PCG/Source/PCG/Public/Elements/Metadata/PCGMetadataMathsOpElement.h
+++ b/Engine/Plugins/Experimental/PCG/Source/PCG/Public/Elements/Metadata/PCGMetadataMathsOpElement.h
@@ -68,6 +68,7 @@ public:
 
 	virtual FName GetInputPinLabel(uint32 Index) const override;
 	virtual uint32 GetInputPinNum() const override;
+	virtual uint16 GetOutputType(uint16 InputTypeId) const;
 
 	virtual bool IsSupportedInputType(uint16 TypeId, uint32 InputIndex, bool& bHasSpecialRequirement) const override;
 	//~End UPCGMetadataSettingsBase interface
@@ -81,6 +82,14 @@ public:
 	UPROPERTY(BlueprintReadWrite, EditAnywhere, Category = Settings)
 	EPCGMedadataMathsOperation Operation = EPCGMedadataMathsOperation::Add;
 
+	/** For rounding operation, if the input type is float or double, use this option to force the output attribute to be int64. */
+	UPROPERTY(BlueprintReadWrite, EditAnywhere, Category = Settings, meta = (EditConditionHides, EditCondition = "Operation == EPCGMedadataMathsOperation::Round || Operation == EPCGMedadataMathsOperation::Truncate || Operation == EPCGMedadataMathsOperation::Floor || Operation == EPCGMedadataMathsOperation::Ceil"))
+	bool bForceRoundingOpToInt = false;
+
+	/** For operations that can yield floating point values, if the input type are ints, use this option to force the output attribute to be double. */
+	UPROPERTY(BlueprintReadWrite, EditAnywhere, Category = Settings, meta = (EditConditionHides, EditCondition = "Operation == EPCGMedadataMathsOperation::Divide || Operation == EPCGMedadataMathsOperation::Sqrt || Operation == EPCGMedadataMathsOperation::Pow || Operation == EPCGMedadataMathsOperation::Lerp"))
+	bool bForceOpToDouble = false;
+
 	UPROPERTY(BlueprintReadWrite, EditAnywhere, Category = Input)
 	FPCGAttributePropertyInputSelector InputSource1;
 
@@ -100,6 +109,10 @@ public:
 	UPROPERTY()
 	FName Input3AttributeName_DEPRECATED = NAME_None;
 #endif
+
+private:
+	bool ShouldForceOutputToInt(uint16 InputTypeId) const;
+	bool ShouldForceOutputToDouble(uint16 InputTypeId) const;
 };
 
 class FPCGMetadataMathsElement : public FPCGMetadataElementBase
