Fixed CRC issue in the data from actor that wasn't considering the culling option

diff --git a/Engine/Plugins/Experimental/PCG/Source/PCG/Private/Elements/PCGDataFromActor.cpp b/Engine/Plugins/Experimental/PCG/Source/PCG/Private/Elements/PCGDataFromActor.cpp
--- a/Engine/Plugins/Experimental/PCG/Source/PCG/Private/Elements/PCGDataFromActor.cpp
+++ b/Engine/Plugins/Experimental/PCG/Source/PCG/Private/Elements/PCGDataFromActor.cpp
@@ -477,7 +477,10 @@ void FPCGDataFromActorElement::GetDependenciesCrc(const FPCGDataCollection& InIn
 	// If we track self or original, we are dependant on the actor data
 	if (const UPCGDataFromActorSettings* Settings = Cast<const UPCGDataFromActorSettings>(InSettings))
 	{
-		if (InComponent && (Settings->ActorSelector.ActorFilter == EPCGActorFilter::Self || Settings->ActorSelector.ActorFilter == EPCGActorFilter::Original))
+		const bool bDependsOnSelfOrHierarchy = (Settings->ActorSelector.ActorFilter == EPCGActorFilter::Self || Settings->ActorSelector.ActorFilter == EPCGActorFilter::Original);
+		const bool bDependsOnSelfBounds = (Settings->ActorSelector.bMustOverlapSelf);
+
+		if (InComponent && (bDependsOnSelfOrHierarchy || bDependsOnSelfBounds))
 		{
 			UPCGComponent* ComponentToCheck = (Settings->ActorSelector.ActorFilter == EPCGActorFilter::Original) ? InComponent->GetOriginalComponent() : InComponent;
 			const UPCGData* ActorData = ComponentToCheck ? ComponentToCheck->GetActorPCGData() : nullptr;
