[PCG] Fix OneMinus op for vectors

* Add a test case for that

#rnx

diff --git a/Engine/Plugins/Experimental/PCG/Source/PCG/Private/Elements/Metadata/PCGMetadataMaths.inl b/Engine/Plugins/Experimental/PCG/Source/PCG/Private/Elements/Metadata/PCGMetadataMaths.inl
--- a/Engine/Plugins/Experimental/PCG/Source/PCG/Private/Elements/Metadata/PCGMetadataMaths.inl
+++ b/Engine/Plugins/Experimental/PCG/Source/PCG/Private/Elements/Metadata/PCGMetadataMaths.inl
@@ -440,5 +440,32 @@ namespace PCGMetadataMaths
 			Modulo(Value1.Y, Value2.Y),
 			Modulo(Value1.Z, Value2.Z),
 			Modulo(Value1.W, Value2.W));
-	}	
-}
\ No newline at end of file
+	}
+
+	////////////////////
+	// One Minus
+	////////////////////
+	template <typename T>
+	inline T OneMinus(const T& Value)
+	{
+		return T{ 1 } - Value;
+	}
+
+	template<>
+	inline FVector2D OneMinus(const FVector2D& Value)
+	{
+		return FVector2D::One() - Value;
+	}
+
+	template<>
+	inline FVector OneMinus(const FVector& Value)
+	{
+		return FVector::One() - Value;
+	}
+
+	template<>
+	inline FVector4 OneMinus(const FVector4& Value)
+	{
+		return FVector4::One() - Value;
+	}
+}
diff --git a/Engine/Plugins/Experimental/PCG/Source/PCG/Private/Elements/Metadata/PCGMetadataMathsOpElement.cpp b/Engine/Plugins/Experimental/PCG/Source/PCG/Private/Elements/Metadata/PCGMetadataMathsOpElement.cpp
--- a/Engine/Plugins/Experimental/PCG/Source/PCG/Private/Elements/Metadata/PCGMetadataMathsOpElement.cpp
+++ b/Engine/Plugins/Experimental/PCG/Source/PCG/Private/Elements/Metadata/PCGMetadataMathsOpElement.cpp
@@ -102,7 +102,7 @@ namespace PCGMetadataMathsSettings
 		case EPCGMedadataMathsOperation::Ceil:
 			return PCGMetadataMaths::Ceil(Value);
 		case EPCGMedadataMathsOperation::OneMinus:
-			return T{ 1 } - Value;
+			return PCGMetadataMaths::OneMinus(Value);
 		default:
 			return T{};
 		}
diff --git a/Engine/Plugins/Experimental/PCG/Source/PCG/Private/Tests/Elements/Metadata/PCGMetadataMathsOpTest.cpp b/Engine/Plugins/Experimental/PCG/Source/PCG/Private/Tests/Elements/Metadata/PCGMetadataMathsOpTest.cpp
new file mode 100644
--- /dev/null
+++ b/Engine/Plugins/Experimental/PCG/Source/PCG/Private/Tests/Elements/Metadata/PCGMetadataMathsOpTest.cpp
@@ -0,0 +1,102 @@
+// Copyright Epic Games, Inc. All Rights Reserved.
+
+#include "Tests/PCGTestsCommon.h"
+
+#include "PCGComponent.h"
+#include "PCGContext.h"
+#include "PCGParamData.h"
+#include "Elements/Metadata/PCGMetadataMathsOpElement.h"
+#include "Metadata/PCGMetadataAccessor.h"
+#include "Metadata/PCGMetadataAttributeTraits.h"
+
+#if WITH_EDITOR
+
+IMPLEMENT_CUSTOM_SIMPLE_AUTOMATION_TEST(FPCGMetadataMathsOp_OneMinus, FPCGTestBaseClass, "pcg.tests.Metadata.MathsOp.OneMinus", PCGTestsCommon::TestFlags)
+
+namespace PCGMathsOpTest
+{
+	const FName AttributeName = TEXT("Attr");
+
+	template <typename T>
+	FPCGTaggedData GenerateParamData(PCGTestsCommon::FTestData& TestData, const T& Value)
+	{
+		FPCGTaggedData& ParamTaggedData = TestData.InputData.TaggedData.Emplace_GetRef();
+		TObjectPtr<UPCGParamData> ParamData = PCGTestsCommon::CreateEmptyParamData();
+		ParamTaggedData.Data = ParamData;
+
+		PCGMetadataEntryKey EntryKey = ParamData->Metadata->AddEntry();
+		const bool bAllowsInterpolation = true;
+		const bool bOverrideParent = false;
+
+		FPCGMetadataAttribute<T>* Attribute = ParamData->Metadata->CreateAttribute<T>(AttributeName, PCG::Private::MetadataTraits<T>::ZeroValue(), bAllowsInterpolation, bOverrideParent);
+		check(Attribute);
+		Attribute->SetValue(EntryKey, Value);
+
+		return ParamTaggedData;
+	}
+}
+
+
+bool FPCGMetadataMathsOp_OneMinus::RunTest(const FString& Parameters)
+{
+	auto RunTest = [this]<typename T>(const T& Value, const T& ExpectedDefaultValue, const T& ExpectedValue)
+	{
+		PCGTestsCommon::FTestData TestData;
+		UPCGMetadataMathsSettings* Settings = PCGTestsCommon::GenerateSettings<UPCGMetadataMathsSettings>(TestData);
+		check(Settings);
+
+		Settings->Operation = EPCGMedadataMathsOperation::OneMinus;
+		Settings->InputSource1.SetAttributeName(PCGMathsOpTest::AttributeName);
+		Settings->OutputTarget.SetAttributeName(PCGMathsOpTest::AttributeName);
+
+		FPCGTaggedData& InputTaggedData = TestData.InputData.TaggedData.Add_GetRef(PCGMathsOpTest::GenerateParamData(TestData, Value));
+		InputTaggedData.Pin = PCGPinConstants::DefaultInputLabel;
+
+		FPCGElementPtr TestElement = TestData.Settings->GetElement();
+		TUniquePtr<FPCGContext> Context = TestData.InitializeTestContext();
+
+		while (!TestElement->Execute(Context.Get())) {}
+
+		auto FormatText = [](const FString& In) -> FString { return FString::Printf(TEXT("OneMinus %s: %s"), *PCG::Private::GetTypeName<T>(), *In); };
+
+		UTEST_EQUAL(*FormatText("There is one output"), Context->OutputData.TaggedData.Num(), 1);
+
+		const UPCGParamData* OutputParam = Cast<const UPCGParamData>(Context->OutputData.TaggedData[0].Data);
+
+		UTEST_NOT_NULL(*FormatText("Output is an attribute set"), OutputParam);
+
+		const FPCGMetadataAttribute<T>* OutputAttribute = OutputParam->Metadata->GetConstTypedAttribute<T>(PCGMathsOpTest::AttributeName);
+
+		UTEST_NOT_NULL(*FormatText("Output metadata has the expected attribute"), OutputAttribute);
+
+		// Vec2 and Vec4 don't have "nearly equal" in the test framework
+		T DefaultValue = OutputAttribute->GetValue(PCGDefaultValueKey);
+		T FirstEntry = OutputAttribute->GetValueFromItemKey(0);
+		if constexpr (std::is_same_v<T, FVector4> || std::is_same_v<T, FVector2D>)
+		{
+			UTEST_TRUE(*FormatText("Default value is correct"), DefaultValue.Equals(ExpectedDefaultValue));
+			UTEST_TRUE(*FormatText("First entry is correct"), FirstEntry.Equals(ExpectedValue));
+		}
+		else
+		{
+			UTEST_EQUAL(*FormatText("Default value is correct"), DefaultValue, ExpectedDefaultValue);
+			UTEST_EQUAL(*FormatText("First entry is correct"), FirstEntry, ExpectedValue);
+		}
+
+		return true;
+	};
+
+	bool bSuccess = true;
+
+	bSuccess &= RunTest(5, 1, -4);
+	bSuccess &= RunTest(6ll, 1ll, -5ll);
+	bSuccess &= RunTest(0.1f, 1.0f, 0.9f);
+	bSuccess &= RunTest(0.2, 1.0, 0.8);
+	bSuccess &= RunTest(FVector2D(0.3, 0.4), FVector2D(1, 1), FVector2D(0.7, 0.6));
+	bSuccess &= RunTest(FVector(0.5, 0.6, 0.7), FVector(1, 1, 1), FVector(0.5, 0.4, 0.3));
+	bSuccess &= RunTest(FVector4(0.5, 0.6, 0.7, 0.8), FVector4(1, 1, 1, 1), FVector4(0.5, 0.4, 0.3, 0.2));
+
+	return bSuccess;
+}
+
+#endif // WITH_EDITOR
