[PCG] Support AddAttribute on multi entry source param

* Add new options in the Metadata Element common to factorize some operations
* Modified the accessor keys on attribute for param data. Will now use Metadata keys if the attribute has no values.

diff --git a/Engine/Plugins/Experimental/PCG/Source/PCG/Private/Elements/Metadata/PCGMetadataElementCommon.cpp b/Engine/Plugins/Experimental/PCG/Source/PCG/Private/Elements/Metadata/PCGMetadataElementCommon.cpp
--- a/Engine/Plugins/Experimental/PCG/Source/PCG/Private/Elements/Metadata/PCGMetadataElementCommon.cpp
+++ b/Engine/Plugins/Experimental/PCG/Source/PCG/Private/Elements/Metadata/PCGMetadataElementCommon.cpp
@@ -46,4 +46,44 @@ namespace PCGMetadataElementCommon
 			OutAttribute->SetValueFromValueKey(EntryKey, ValueKey);
 		}
 	}
+
+	bool CopyFromAccessorToAccessor(FCopyFromAccessorToAccessorParams& Params)
+	{
+		check(Params.InAccessor && Params.OutAccessor && Params.InKeys && Params.OutKeys);
+
+		int32 Count = 0;
+		switch (Params.IterationCount)
+		{
+		case FCopyFromAccessorToAccessorParams::EIterationCount::In:
+			Count = Params.InKeys->GetNum();
+			break;
+		case FCopyFromAccessorToAccessorParams::EIterationCount::Out:
+			Count = Params.OutKeys->GetNum();
+			break;
+		case FCopyFromAccessorToAccessorParams::EIterationCount::Min:
+			Count = FMath::Min(Params.InKeys->GetNum(), Params.OutKeys->GetNum());
+			break;
+		case FCopyFromAccessorToAccessorParams::EIterationCount::Max:
+			Count = FMath::Max(Params.InKeys->GetNum(), Params.OutKeys->GetNum());
+			break;
+		default:
+			checkNoEntry();
+			return false;
+		}
+
+		auto Operation = [&Params, Count](auto Dummy)
+		{
+			using OutputType = decltype(Dummy);
+			OutputType Value{};
+
+			auto SetToAccessor = [&Params](const TArrayView<OutputType>& View, const int32 Start, const int32 Range)
+			{
+				Params.OutAccessor->SetRange<OutputType>(View, Start, *Params.OutKeys, Params.Flags);
+			};
+
+			return PCGMetadataElementCommon::ApplyOnAccessorRange<OutputType>(*Params.InKeys, *Params.InAccessor, SetToAccessor, Params.Flags, Params.ChunkSize, Count);
+		};
+
+		return PCGMetadataAttribute::CallbackWithRightType(Params.OutAccessor->GetUnderlyingType(), Operation);
+	}
 }
diff --git a/Engine/Plugins/Experimental/PCG/Source/PCG/Private/Elements/PCGCreateAttribute.cpp b/Engine/Plugins/Experimental/PCG/Source/PCG/Private/Elements/PCGCreateAttribute.cpp
--- a/Engine/Plugins/Experimental/PCG/Source/PCG/Private/Elements/PCGCreateAttribute.cpp
+++ b/Engine/Plugins/Experimental/PCG/Source/PCG/Private/Elements/PCGCreateAttribute.cpp
@@ -291,12 +291,6 @@ bool FPCGCreateAttributeElement::ExecuteInternal(FPCGContext* Context) const
 			PCGE_LOG(Error, GraphAndLog, FText::Format(LOCTEXT("ParamMissingAttribute", "Source Attribute Set data does not have an attribute '{0}'"), FText::FromName(SourceParamAttributeName)));
 			return true;
 		}
-
-		// We don't support multi-entry yet
-		if (SourceParamData->Metadata->GetLocalItemCount() > 1)
-		{
-			PCGE_LOG(Warning, GraphAndLog,LOCTEXT("NotSupportingMultiEntry", "Source Attribute Set data has more than one entry. Not yet supported. Will use the default value of the attribute."));
-		}
 	}
 	else
 	{
@@ -360,8 +354,11 @@ bool FPCGCreateAttributeElement::ExecuteInternal(FPCGContext* Context) const
 		{
 			const FPCGAttributePropertyInputSelector InputSource = Settings->InputSource.CopyAndFixLast(SourceParamData);
 
-			// If no field accessor, copy over the attribute
-			if (InputSource.GetExtraNames().IsEmpty())
+			// We need accessors if we have a multi entry source attribute or we have extractors
+			const bool bNeedAccessors = SourceParamData->Metadata->GetLocalItemCount() > 1 || !InputSource.GetExtraNames().IsEmpty();
+
+			// If no accessor, copy over the attribute
+			if (!bNeedAccessors)
 			{
 				const FPCGMetadataAttributeBase* SourceAttribute = SourceParamData->Metadata->GetConstAttribute(SourceParamAttributeName);
 				Attribute = Metadata->CopyAttribute(SourceAttribute, OutputAttributeName, /*bKeepParent=*/false, /*bCopyEntries=*/false, /*bCopyValues=*/false);
@@ -376,12 +373,46 @@ bool FPCGCreateAttributeElement::ExecuteInternal(FPCGContext* Context) const
 					return true;
 				}
 
-				auto CreateOutputAttribute = [Metadata, OutputAttributeName, &InputAccessor, &InputKeys]<typename Type>(Type Dummy) -> FPCGMetadataAttributeBase*
+				auto CreateOutputAttribute = [Metadata, OutputAttributeName, &InputAccessor, &InputKeys, OutputData, bIsParamData]<typename Type>(Type Dummy) -> FPCGMetadataAttributeBase*
 				{
-					// Get the value from the input accessor and pass that as the default value
+					// Get the value from the input accessor default value and pass that as the default value
 					Type Value{};
-					InputAccessor->Get<Type>(Value, *InputKeys);
-					return PCGMetadataElementCommon::ClearOrCreateAttribute<Type>(Metadata, OutputAttributeName, Value);
+					InputAccessor->Get<Type>(Value, FPCGAttributeAccessorKeysEntries(PCGInvalidEntryKey));
+					FPCGMetadataAttribute<Type>* Attribute = PCGMetadataElementCommon::ClearOrCreateAttribute<Type>(Metadata, OutputAttributeName, Value);
+
+					FPCGAttributePropertySelector OutputSelector;
+					OutputSelector.SetAttributeName(OutputAttributeName);
+
+					TUniquePtr<IPCGAttributeAccessor> OutputAccessor = PCGAttributeAccessorHelpers::CreateAccessor(OutputData, OutputSelector);
+					TUniquePtr<IPCGAttributeAccessorKeys> OutputKeys = PCGAttributeAccessorHelpers::CreateKeys(OutputData, OutputSelector);
+
+					// We just created the attribute, this should not fail
+					if (!ensure(OutputAccessor.IsValid() && OutputKeys.IsValid()))
+					{
+						return nullptr;
+					}
+
+					PCGMetadataElementCommon::FCopyFromAccessorToAccessorParams Params;
+					Params.InKeys = InputKeys.Get();
+					Params.InAccessor = InputAccessor.Get();
+					Params.OutKeys = OutputKeys.Get();
+					Params.OutAccessor = OutputAccessor.Get();
+					Params.IterationCount = PCGMetadataElementCommon::FCopyFromAccessorToAccessorParams::Min;
+
+					PCGMetadataElementCommon::CopyFromAccessorToAccessor(Params);
+
+					// In case of param data, do some padding at the end if needed
+					if (bIsParamData)
+					{
+						int32 NumInKeys = InputKeys->GetNum();
+						int32 Padding = OutputKeys->GetNum() - NumInKeys;
+						for (int32 Key = 0; Key < Padding; ++Key)
+						{
+							Attribute->SetValueFromValueKey(PCGMetadataEntryKey(NumInKeys + Key), PCGDefaultValueKey);
+						}
+					}
+
+					return Attribute;
 				};
 
 				Attribute = PCGMetadataAttribute::CallbackWithRightType(InputAccessor->GetUnderlyingType(), CreateOutputAttribute);
@@ -398,20 +429,10 @@ bool FPCGCreateAttributeElement::ExecuteInternal(FPCGContext* Context) const
 			continue;
 		}
 
-		// In the case of an input param data, remap all entries for the output data (which is also a param data) 
-		// to the default value for the newly created attribute.
-		if (bIsParamData)
+		// Making sure the metadata has at least one entry.
+		if (bIsParamData && Metadata->GetLocalItemCount() == 0)
 		{
-			// Making sure the metadata has at least one entry.
-			if (Metadata->GetLocalItemCount() == 0)
-			{
-				Metadata->AddEntry();
-			}
-
-			for (PCGMetadataEntryKey Key = 0; Key < Metadata->GetLocalItemCount(); ++Key)
-			{
-				Attribute->SetValueFromValueKey(Key, PCGDefaultValueKey);
-			}
+			Metadata->AddEntry();
 		}
 
 		TArray<FPCGTaggedData>& Outputs = Context->OutputData.TaggedData;
diff --git a/Engine/Plugins/Experimental/PCG/Source/PCG/Private/Metadata/Accessors/PCGAttributeAccessorKeys.cpp b/Engine/Plugins/Experimental/PCG/Source/PCG/Private/Metadata/Accessors/PCGAttributeAccessorKeys.cpp
--- a/Engine/Plugins/Experimental/PCG/Source/PCG/Private/Metadata/Accessors/PCGAttributeAccessorKeys.cpp
+++ b/Engine/Plugins/Experimental/PCG/Source/PCG/Private/Metadata/Accessors/PCGAttributeAccessorKeys.cpp
@@ -6,7 +6,7 @@
 
 #include "PCGPoint.h"
 
-////////////////////////////////////////////////////////////////////
+//////////////////////////////////////////////////////////////////// 
 
 FPCGAttributeAccessorKeysEntries::FPCGAttributeAccessorKeysEntries(const FPCGMetadataAttributeBase* Attribute)
 	: IPCGAttributeAccessorKeys(/*bInReadOnly=*/ false)
@@ -20,10 +20,16 @@ FPCGAttributeAccessorKeysEntries::FPCGAttributeAccessorKeysEntries(const FPCGMet
 		Current = Current->GetParent();
 	}
 
-	// If the attribute doesn't have any entry, we will always take the default value.
+	// If the attribute doesn't have any entry, re-try with metadata entries.
 	if (Entries.IsEmpty())
 	{
-		Entries.Add(PCGInvalidEntryKey);
+		InitializeFromMetadata(Attribute->GetMetadata());
+
+		// If the attribute still doesn't have any entry, we will always take the default value.
+		if (Entries.IsEmpty())
+		{
+			Entries.Add(PCGInvalidEntryKey);
+		}
 	}
 }
 
@@ -35,12 +41,19 @@ FPCGAttributeAccessorKeysEntries::FPCGAttributeAccessorKeysEntries(PCGMetadataEn
 
 FPCGAttributeAccessorKeysEntries::FPCGAttributeAccessorKeysEntries(const UPCGMetadata* Metadata)
 	: IPCGAttributeAccessorKeys(/*bInReadOnly=*/ true)
+{
+	InitializeFromMetadata(Metadata);
+}
+
+void FPCGAttributeAccessorKeysEntries::InitializeFromMetadata(const UPCGMetadata* Metadata)
 {
 	if (!Metadata)
 	{
 		return;
 	}
 
+	check(Entries.IsEmpty());
+
 	const PCGMetadataEntryKey ItemKeyLowerBound = Metadata->GetItemKeyCountForParent();
 	const PCGMetadataEntryKey ItemKeyUpperBound = Metadata->GetItemCountForChild();
 	const int64 Count = ItemKeyUpperBound - ItemKeyLowerBound;
diff --git a/Engine/Plugins/Experimental/PCG/Source/PCG/Public/Elements/Metadata/PCGMetadataElementCommon.h b/Engine/Plugins/Experimental/PCG/Source/PCG/Public/Elements/Metadata/PCGMetadataElementCommon.h
--- a/Engine/Plugins/Experimental/PCG/Source/PCG/Public/Elements/Metadata/PCGMetadataElementCommon.h
+++ b/Engine/Plugins/Experimental/PCG/Source/PCG/Public/Elements/Metadata/PCGMetadataElementCommon.h
@@ -14,10 +14,10 @@ struct FPCGTaggedData;
 
 namespace PCGMetadataElementCommon
 {
-	void DuplicateTaggedData(const FPCGTaggedData& InTaggedData, FPCGTaggedData& OutTaggedData, UPCGMetadata*& OutMetadata);
+	PCG_API void DuplicateTaggedData(const FPCGTaggedData& InTaggedData, FPCGTaggedData& OutTaggedData, UPCGMetadata*& OutMetadata);
 
 	/** Copies the entry to value key relationship stored in the given Metadata, including its parents */
-	void CopyEntryToValueKeyMap(const UPCGMetadata* MetadataToCopy, const FPCGMetadataAttributeBase* AttributeToCopy, FPCGMetadataAttributeBase* OutAttribute);
+	PCG_API void CopyEntryToValueKeyMap(const UPCGMetadata* MetadataToCopy, const FPCGMetadataAttributeBase* AttributeToCopy, FPCGMetadataAttributeBase* OutAttribute);
 
 	/** Creates a new attribute, or clears the attribute if it already exists and is a 'T' type */
 	template<typename T>
@@ -45,15 +45,15 @@ namespace PCGMetadataElementCommon
 	constexpr int32 DefaultChunkSize = 256;
 
 	/**
-	* Iterate over the full range of the keys, calling the callback with values get from the accessor.
+	* Iterate over the full range of the keys (if Count is negative, otherwise, as many times as Count), calling the callback with values range get from the accessor.
 	* ChunkSize influence the max number of values to get in one go with GetRange.
-	* Callback should have a signature: void(const T&, int32)
+	* Callback should have a signature: void(const TArrayView<T>& View, int32 Start, int32 Range)
 	* Return false if it process nothing.
 	*/
 	template <typename T, typename Func>
-	bool ApplyOnAccessor(const IPCGAttributeAccessorKeys& Keys, const IPCGAttributeAccessor& Accessor, Func&& Callback, EPCGAttributeAccessorFlags Flags = EPCGAttributeAccessorFlags::StrictType, const int32 ChunkSize = DefaultChunkSize)
+	bool ApplyOnAccessorRange(const IPCGAttributeAccessorKeys& Keys, const IPCGAttributeAccessor& Accessor, Func&& Callback, EPCGAttributeAccessorFlags Flags = EPCGAttributeAccessorFlags::StrictType, const int32 ChunkSize = DefaultChunkSize, const int32 Count = -1)
 	{
-		const int32 NumberOfEntries = Keys.GetNum();
+		const int32 NumberOfEntries = Count < 0 ? Keys.GetNum() : Count;
 
 		if (NumberOfEntries == 0)
 		{
@@ -76,21 +76,62 @@ namespace PCGMetadataElementCommon
 				return false;
 			}
 
+			Callback(View, StartIndex, Range);
+		}
+
+		return true;
+	}
+
+	/**
+	* Iterate over the full range of the keys (if Count is negative, otherwise, as many times as Count), calling the callback with values get from the accessor.
+	* ChunkSize influence the max number of values to get in one go with GetRange.
+	* Callback should have a signature: void(const T& Value, int32 Index)
+	* Return false if it process nothing.
+	*/
+	template <typename T, typename Func>
+	bool ApplyOnAccessor(const IPCGAttributeAccessorKeys& Keys, const IPCGAttributeAccessor& Accessor, Func&& InCallback, EPCGAttributeAccessorFlags Flags = EPCGAttributeAccessorFlags::StrictType, const int32 ChunkSize = DefaultChunkSize, const int32 Count = -1)
+	{
+		auto RangeCallback = [Callback = std::forward<Func>(InCallback)](const TArrayView<T>& View, int32 Start, int32 Range)
+		{
 			for (int32 j = 0; j < Range; ++j)
 			{
-				Callback(TempValues[j], StartIndex + j);
+				Callback(View[j], Start + j);
 			}
-		}
+		};
 
-		return true;
+		return ApplyOnAccessorRange<T>(Keys, Accessor, std::move(RangeCallback), Flags, ChunkSize, Count);
 	}
 
+	/**
+	* Read all the values from in accessor and write to out accessor. Iterate as many times than the out accessor keys
+	*/
+	struct PCG_API FCopyFromAccessorToAccessorParams
+	{
+		enum EIterationCount
+		{
+			In, // Iterate as many times than the in accessor keys
+			Out, // Iterate as many times than the out accessor keys
+			Min, // Iterate as many times than the min of accessor keys
+			Max, // Iterate as many times than the max of accessor keys
+		};
+
+		EIterationCount IterationCount;
+		const IPCGAttributeAccessor* InAccessor = nullptr;
+		const IPCGAttributeAccessorKeys* InKeys = nullptr;
+		IPCGAttributeAccessor* OutAccessor = nullptr;
+		IPCGAttributeAccessorKeys* OutKeys = nullptr;
+		EPCGAttributeAccessorFlags Flags = EPCGAttributeAccessorFlags::StrictType;
+		int32 ChunkSize = DefaultChunkSize;
+	};
+
+	PCG_API bool CopyFromAccessorToAccessor(FCopyFromAccessorToAccessorParams& Params);
+
 	/**
 	* Automatically fill all preconfigured settings depending on the enum operation.
 	* Can also specify explicitly values that should not be included, as metadata is not available in non-editor builds.
 	*/
 	template <typename EnumOperation>
-	TArray<FPCGPreConfiguredSettingsInfo> FillPreconfiguredSettingsInfoFromEnum(const TSet<EnumOperation>& InValuesToSkip = {})
+	TArray<FPCGPreConfiguredSettingsInfo> FillPreconfiguredSettingsInfoFromEnum(const TSet<EnumOperation>& InValuesToSkip = {}, const FText& InOptionalPrefix = FText())
 	{
 		TArray<FPCGPreConfiguredSettingsInfo> PreconfiguredInfo;
 
@@ -103,7 +144,7 @@ namespace PCGMetadataElementCommon
 
 				if (Value != EnumPtr->GetMaxEnumValue() && !InValuesToSkip.Contains(EnumOperation(Value)))
 				{
-					PreconfiguredInfo.Emplace(Value, EnumPtr->GetDisplayNameTextByValue(Value));
+					PreconfiguredInfo.Emplace(Value, FText::Format(FText::FromString(TEXT("{0}{1}")), InOptionalPrefix, EnumPtr->GetDisplayNameTextByValue(Value)));
 				}
 			}
 		}
diff --git a/Engine/Plugins/Experimental/PCG/Source/PCG/Public/Metadata/Accessors/PCGAttributeAccessorKeys.h b/Engine/Plugins/Experimental/PCG/Source/PCG/Public/Metadata/Accessors/PCGAttributeAccessorKeys.h
--- a/Engine/Plugins/Experimental/PCG/Source/PCG/Public/Metadata/Accessors/PCGAttributeAccessorKeys.h
+++ b/Engine/Plugins/Experimental/PCG/Source/PCG/Public/Metadata/Accessors/PCGAttributeAccessorKeys.h
@@ -150,6 +150,8 @@ protected:
 	virtual bool GetMetadataEntryKeys(int32 InStart, TArrayView<PCGMetadataEntryKey*>& OutEntryKeys) override;
 	virtual bool GetMetadataEntryKeys(int32 InStart, TArrayView<const PCGMetadataEntryKey*>& OutEntryKeys) const override;
 
+	void InitializeFromMetadata(const UPCGMetadata* Metadata);
+
 	TArray<PCGMetadataEntryKey> Entries;
 };
 
