[PCG] Fix static analysis

#rnx

diff --git a/Engine/Plugins/Experimental/PCG/Source/PCG/Private/Metadata/Accessors/PCGAttributeAccessorKeys.cpp b/Engine/Plugins/Experimental/PCG/Source/PCG/Private/Metadata/Accessors/PCGAttributeAccessorKeys.cpp
--- a/Engine/Plugins/Experimental/PCG/Source/PCG/Private/Metadata/Accessors/PCGAttributeAccessorKeys.cpp
+++ b/Engine/Plugins/Experimental/PCG/Source/PCG/Private/Metadata/Accessors/PCGAttributeAccessorKeys.cpp
@@ -21,15 +21,15 @@ FPCGAttributeAccessorKeysEntries::FPCGAttributeAccessorKeysEntries(const FPCGMet
 	}
 
 	// If the attribute doesn't have any entry, re-try with metadata entries.
-	if (Entries.IsEmpty())
+	if (Attribute && Entries.IsEmpty())
 	{
 		InitializeFromMetadata(Attribute->GetMetadata());
+	}
 
-		// If the attribute still doesn't have any entry, we will always take the default value.
-		if (Entries.IsEmpty())
-		{
-			Entries.Add(PCGInvalidEntryKey);
-		}
+	// If the attribute still doesn't have any entry, we will always take the default value.
+	if (Entries.IsEmpty())
+	{
+		Entries.Add(PCGInvalidEntryKey);
 	}
 }
 
