Fixed editor automation test around PCG

diff --git a/Engine/Plugins/Experimental/PCG/Source/PCGEditor/Private/Tests/Display/PCGGrapthAssetIsExposedInContetBrowser.cpp b/Engine/Plugins/Experimental/PCG/Source/PCGEditor/Private/Tests/Display/PCGGrapthAssetIsExposedInContetBrowser.cpp
--- a/Engine/Plugins/Experimental/PCG/Source/PCGEditor/Private/Tests/Display/PCGGrapthAssetIsExposedInContetBrowser.cpp
+++ b/Engine/Plugins/Experimental/PCG/Source/PCGEditor/Private/Tests/Display/PCGGrapthAssetIsExposedInContetBrowser.cpp
@@ -12,32 +12,18 @@ IMPLEMENT_CUSTOM_SIMPLE_AUTOMATION_TEST(FPCGGraphExposedInContentMenuTest, FPCGT
 bool FPCGGraphExposedInContentMenuTest::RunTest(const FString& Parameters)
 {
 	const FName NAME_AssetTools = "AssetTools";
-	const FString PluginName = "PCG";
-	uint32 FactoryCategories = 0;
 	bool pcgGraphFound = false;
 
 	FAssetToolsModule& AssetToolsModule = FModuleManager::LoadModuleChecked<FAssetToolsModule>(NAME_AssetTools);
 
-	TArray<FAdvancedAssetCategory> AdvancedAssetCategories;
-	AssetToolsModule.Get().GetAllAdvancedAssetCategories(/*out*/ AdvancedAssetCategories);
-
-	FAdvancedAssetCategory* PCGAssetCategory = AdvancedAssetCategories.FindByPredicate([&PluginName](const FAdvancedAssetCategory& AdvancedAssetCategory) {
-		return AdvancedAssetCategory.CategoryName.ToString().Contains(PluginName);
-		});
-
-	UTEST_NOT_NULL(TEXT("PCG category not found"), PCGAssetCategory)
-
 	// We check if the PCG category factory's display name contains "PCG Graph"
 	TArray<UFactory*> Factories = AssetToolsModule.Get().GetNewAssetFactories();
 	for (UFactory* Factory : Factories)
 	{
-		FactoryCategories = Factory->GetMenuCategories();
-		if (FactoryCategories & PCGAssetCategory->CategoryType)
+		if (Factory->GetDisplayName().ToString().Contains("PCG Graph"))
 		{
-			if (Factory->GetDisplayName().ToString().Contains("PCG Graph")) {
-				pcgGraphFound = true;
-				break;
-			}
+			pcgGraphFound = true;
+			break;
 		}
 	}
 
