[PCG] CopyPoints node now supports multiple sources and targets

diff --git a/Engine/Plugins/Experimental/PCG/Source/PCG/Private/Elements/PCGCopyPoints.cpp b/Engine/Plugins/Experimental/PCG/Source/PCG/Private/Elements/PCGCopyPoints.cpp
--- a/Engine/Plugins/Experimental/PCG/Source/PCG/Private/Elements/PCGCopyPoints.cpp
+++ b/Engine/Plugins/Experimental/PCG/Source/PCG/Private/Elements/PCGCopyPoints.cpp
@@ -24,8 +24,8 @@ FText UPCGCopyPointsSettings::GetNodeTooltipText() const
 TArray<FPCGPinProperties> UPCGCopyPointsSettings::InputPinProperties() const
 {
 	TArray<FPCGPinProperties> PinProperties;
-	PinProperties.Emplace(PCGCopyPointsConstants::SourcePointsLabel, EPCGDataType::Point, /*bAllowMultipleConnections=*/false);
-	PinProperties.Emplace(PCGCopyPointsConstants::TargetPointsLabel, EPCGDataType::Point, /*bAllowMultipleConnections=*/false);
+	PinProperties.Emplace(PCGCopyPointsConstants::SourcePointsLabel, EPCGDataType::Point, /*bAllowMultipleConnections=*/true);
+	PinProperties.Emplace(PCGCopyPointsConstants::TargetPointsLabel, EPCGDataType::Point, /*bAllowMultipleConnections=*/true);
 	return PinProperties;
 }
 
@@ -52,272 +52,272 @@ bool FPCGCopyPointsElement::ExecuteInternal(FPCGContext* Context) const
 
 	TArray<FPCGTaggedData>& Outputs = Context->OutputData.TaggedData;
 
-	if (Sources.Num() != 1 || Targets.Num() != 1)
+	for (const FPCGTaggedData& Source : Sources)
 	{
-		PCGE_LOG(Warning, LogOnly, FText::Format(LOCTEXT("InvalidNumberOfInputs", "Invalid number of inputs. Expected one source, got {0}; Expected one target, got {1}."), Sources.Num(), Targets.Num()));
-		return true;
-	}
-	
-	const FPCGTaggedData& Source = Sources[0];
-	const FPCGTaggedData& Target = Targets[0];
+		for (const FPCGTaggedData& Target : Targets)
+		{
+			FPCGTaggedData& Output = Outputs.Add_GetRef(Source);
+			if (!Source.Data || !Target.Data) 
+			{
+				PCGE_LOG(Error, GraphAndLog, LOCTEXT("InvalidInputData", "Invalid input data"));
+				return true;
+			}
 
-	FPCGTaggedData& Output = Outputs.Add_GetRef(Source);
+			const UPCGSpatialData* SourceSpatialData = Cast<UPCGSpatialData>(Source.Data);
+			const UPCGSpatialData* TargetSpatialData = Cast<UPCGSpatialData>(Target.Data);
 
-	if (!Source.Data || !Target.Data) 
-	{
-		PCGE_LOG(Error, GraphAndLog, LOCTEXT("InvalidInputData", "Invalid input data"));
-		return true;
-	}
+			if (!SourceSpatialData || !TargetSpatialData)
+			{
+				PCGE_LOG(Error, GraphAndLog, LOCTEXT("CouldNotObtainSpatialData", "Unable to get Spatial Data from input"));
+				return true;
+			}
 
-	const UPCGSpatialData* SourceSpatialData = Cast<UPCGSpatialData>(Source.Data);
-	const UPCGSpatialData* TargetSpatialData = Cast<UPCGSpatialData>(Target.Data);
+			const UPCGPointData* SourcePointData = SourceSpatialData->ToPointData(Context);
+			const UPCGPointData* TargetPointData = TargetSpatialData->ToPointData(Context);
 
-	if (!SourceSpatialData || !TargetSpatialData)
-	{
-		PCGE_LOG(Error, GraphAndLog, LOCTEXT("CouldNotObtainSpatialData", "Unable to get Spatial Data from input"));
-		return true;
-	}
+			if (!SourcePointData || !TargetPointData)
+			{
+				PCGE_LOG(Error, GraphAndLog, LOCTEXT("CouldNotGetPointData", "Unable to get Point Data from input"));
+				return true;
+			}
 
-	const UPCGPointData* SourcePointData = SourceSpatialData->ToPointData(Context);
-	const UPCGPointData* TargetPointData = TargetSpatialData->ToPointData(Context);
+			const UPCGMetadata* SourcePointMetadata = SourcePointData->Metadata;
+			const UPCGMetadata* TargetPointMetadata = TargetPointData->Metadata;
 
-	if (!SourcePointData || !TargetPointData)
-	{
-		PCGE_LOG(Error, GraphAndLog, LOCTEXT("CouldNotGetPointData", "Unable to get Point Data from input"));
-		return true;
-	}
+			const TArray<FPCGPoint>& SourcePoints = SourcePointData->GetPoints();
+			const TArray<FPCGPoint>& TargetPoints = TargetPointData->GetPoints();
 
-	const UPCGMetadata* SourcePointMetadata = SourcePointData->Metadata;
-	const UPCGMetadata* TargetPointMetadata = TargetPointData->Metadata;
 
-	const TArray<FPCGPoint>& SourcePoints = SourcePointData->GetPoints();
-	const TArray<FPCGPoint>& TargetPoints = TargetPointData->GetPoints();
+			UPCGPointData* OutPointData = NewObject<UPCGPointData>();
+			TArray<FPCGPoint>& OutPoints = OutPointData->GetMutablePoints();
+			Output.Data = OutPointData;
 
-	UPCGPointData* OutPointData = NewObject<UPCGPointData>();
-	TArray<FPCGPoint>& OutPoints = OutPointData->GetMutablePoints();
-	Output.Data = OutPointData;
+			// Make sure that output contains both collection of tags from source and target
+			Output.Tags.Append(Target.Tags);
 
-	// RootMetadata will be parent to the ouptut metadata, while NonRootMetadata will carry attributes from the input not selected for inheritance
-	// Note that this is a preference, as we can and should pick more efficiently in the trivial cases
-	const UPCGMetadata* RootMetadata = nullptr;
-	const UPCGMetadata* NonRootMetadata = nullptr;
+			// RootMetadata will be parent to the ouptut metadata, while NonRootMetadata will carry attributes from the input not selected for inheritance
+			// Note that this is a preference, as we can and should pick more efficiently in the trivial cases
+			const UPCGMetadata* RootMetadata = nullptr;
+			const UPCGMetadata* NonRootMetadata = nullptr;
 
-	const bool bSourceHasMetadata = (SourcePointMetadata->GetAttributeCount() > 0 && SourcePointMetadata->GetItemCountForChild() > 0);
-	const bool bTargetHasMetadata = (TargetPointMetadata->GetAttributeCount() > 0 && TargetPointMetadata->GetItemCountForChild() > 0);
+			const bool bSourceHasMetadata = (SourcePointMetadata->GetAttributeCount() > 0 && SourcePointMetadata->GetItemCountForChild() > 0);
+			const bool bTargetHasMetadata = (TargetPointMetadata->GetAttributeCount() > 0 && TargetPointMetadata->GetItemCountForChild() > 0);
 
-	bool bInheritMetadataFromSource = true;
-	bool bProcessMetadata = (bSourceHasMetadata || bTargetHasMetadata);
+			bool bInheritMetadataFromSource = true;
+			bool bProcessMetadata = (bSourceHasMetadata || bTargetHasMetadata);
 
-	if (AttributeInheritance == EPCGCopyPointsMetadataInheritanceMode::SourceOnly)
-	{
-		bInheritMetadataFromSource = true;
-		bProcessMetadata = bSourceHasMetadata;
+			if (AttributeInheritance == EPCGCopyPointsMetadataInheritanceMode::SourceOnly)
+			{
+				bInheritMetadataFromSource = true;
+				bProcessMetadata = bSourceHasMetadata;
 
-		OutPointData->InitializeFromData(SourcePointData);
-		RootMetadata = SourcePointMetadata;
-		NonRootMetadata = nullptr;
-	}
-	else if (AttributeInheritance == EPCGCopyPointsMetadataInheritanceMode::TargetOnly)
-	{
-		bInheritMetadataFromSource = false;
-		bProcessMetadata = bTargetHasMetadata;
+				OutPointData->InitializeFromData(SourcePointData);
+				RootMetadata = SourcePointMetadata;
+				NonRootMetadata = nullptr;
+			}
+			else if (AttributeInheritance == EPCGCopyPointsMetadataInheritanceMode::TargetOnly)
+			{
+				bInheritMetadataFromSource = false;
+				bProcessMetadata = bTargetHasMetadata;
 
-		OutPointData->InitializeFromData(TargetPointData);
+				OutPointData->InitializeFromData(TargetPointData);
 
-		RootMetadata = TargetPointMetadata;
-		NonRootMetadata = nullptr;
-	}
-	else if (AttributeInheritance == EPCGCopyPointsMetadataInheritanceMode::SourceFirst)
-	{
-		bInheritMetadataFromSource = bSourceHasMetadata || !bTargetHasMetadata;
-
-		OutPointData->InitializeFromData(SourcePointData);
-		RootMetadata = SourcePointMetadata;
-		NonRootMetadata = TargetPointMetadata;
-	}
-	else if (AttributeInheritance == EPCGCopyPointsMetadataInheritanceMode::TargetFirst)
-	{
-		bInheritMetadataFromSource = !bTargetHasMetadata;
+				RootMetadata = TargetPointMetadata;
+				NonRootMetadata = nullptr;
+			}
+			else if (AttributeInheritance == EPCGCopyPointsMetadataInheritanceMode::SourceFirst)
+			{
+				bInheritMetadataFromSource = bSourceHasMetadata || !bTargetHasMetadata;
 
-		OutPointData->InitializeFromData(TargetPointData);
+				OutPointData->InitializeFromData(SourcePointData);
+				RootMetadata = SourcePointMetadata;
+				NonRootMetadata = TargetPointMetadata;
+			}
+			else if (AttributeInheritance == EPCGCopyPointsMetadataInheritanceMode::TargetFirst)
+			{
+				bInheritMetadataFromSource = !bTargetHasMetadata;
 
-		RootMetadata = TargetPointMetadata;
-		NonRootMetadata = SourcePointMetadata;
-	}
-	else // None
-	{
-		OutPointData->InitializeFromData(SourcePointData, nullptr, /*bInheritMetadata=*/false, /*bInheritAttributes=*/false);
+				OutPointData->InitializeFromData(TargetPointData);
 
-		bProcessMetadata = false;
-		RootMetadata = NonRootMetadata = nullptr;
-	}
+				RootMetadata = TargetPointMetadata;
+				NonRootMetadata = SourcePointMetadata;
+			}
+			else // None
+			{
+				OutPointData->InitializeFromData(SourcePointData, nullptr, /*bInheritMetadata=*/false, /*bInheritAttributes=*/false);
 
-	// Always use the target actor from the target, irrespective of the source
-	OutPointData->TargetActor = TargetPointData->TargetActor;
+				bProcessMetadata = false;
+				RootMetadata = NonRootMetadata = nullptr;
+			}
 
-	UPCGMetadata* OutMetadata = OutPointData->Metadata;
-	check(OutMetadata);
+			// Always use the target actor from the target, irrespective of the source
+			OutPointData->TargetActor = TargetPointData->TargetActor;
 
-	TArray<FPCGMetadataAttributeBase*> AttributesToSet;
-	TArray<const FPCGMetadataAttributeBase*> NonRootAttributes;
-	TArray<TTuple<int64, int64>> AllMetadataEntries;
-	TArray<TArray<TTuple<PCGMetadataEntryKey, PCGMetadataValueKey>>> AttributeValuesToSet;
+			UPCGMetadata* OutMetadata = OutPointData->Metadata;
+			check(OutMetadata);
 
-	if (bProcessMetadata)
-	{
-		TRACE_CPUPROFILER_EVENT_SCOPE(FPCGCopyPointsElement::Execute::SetupMetadata);
-		if (NonRootMetadata)
-		{
-			// Prepare the attributes from the non-root that we'll need to use to copy values over
-			TArray<FName> AttributeNames;
-			TArray<EPCGMetadataTypes> AttributeTypes;
-			NonRootMetadata->GetAttributes(AttributeNames, AttributeTypes);
+			TArray<FPCGMetadataAttributeBase*> AttributesToSet;
+			TArray<const FPCGMetadataAttributeBase*> NonRootAttributes;
+			TArray<TTuple<int64, int64>> AllMetadataEntries;
+			TArray<TArray<TTuple<PCGMetadataEntryKey, PCGMetadataValueKey>>> AttributeValuesToSet;
 
-			for (const FName& AttributeName : AttributeNames)
+			if (bProcessMetadata)
 			{
-				if (!OutMetadata->HasAttribute(AttributeName))
+				TRACE_CPUPROFILER_EVENT_SCOPE(FPCGCopyPointsElement::Execute::SetupMetadata);
+				if (NonRootMetadata)
 				{
-					const FPCGMetadataAttributeBase* Attribute = NonRootMetadata->GetConstAttribute(AttributeName);
-					if (FPCGMetadataAttributeBase* NewAttribute = OutMetadata->CopyAttribute(Attribute, AttributeName, /*bKeepRoot=*/false, /*bCopyEntries=*/false, /*bCopyValues=*/true))
+					// Prepare the attributes from the non-root that we'll need to use to copy values over
+					TArray<FName> AttributeNames;
+					TArray<EPCGMetadataTypes> AttributeTypes;
+					NonRootMetadata->GetAttributes(AttributeNames, AttributeTypes);
+
+					for (const FName& AttributeName : AttributeNames)
 					{
-						AttributesToSet.Add(NewAttribute);
-						NonRootAttributes.Add(Attribute);
+						if (!OutMetadata->HasAttribute(AttributeName))
+						{
+							const FPCGMetadataAttributeBase* Attribute = NonRootMetadata->GetConstAttribute(AttributeName);
+							if (FPCGMetadataAttributeBase* NewAttribute = OutMetadata->CopyAttribute(Attribute, AttributeName, /*bKeepRoot=*/false, /*bCopyEntries=*/false, /*bCopyValues=*/true))
+							{
+								AttributesToSet.Add(NewAttribute);
+								NonRootAttributes.Add(Attribute);
+							}
+						}
 					}
-				}
-			}
 
-			// Considering writing to the attribute value requires a lock, we'll gather the value keys to write
-			// and do it on a 1-thread-per-attribute basis at the end
-			AttributeValuesToSet.SetNum(AttributesToSet.Num());
+					// Considering writing to the attribute value requires a lock, we'll gather the value keys to write
+					// and do it on a 1-thread-per-attribute basis at the end
+					AttributeValuesToSet.SetNum(AttributesToSet.Num());
 
-			for (TArray<TTuple<PCGMetadataEntryKey, PCGMetadataValueKey>>& AttributeValues : AttributeValuesToSet)
-			{
-				AttributeValues.SetNumUninitialized(SourcePoints.Num() * TargetPoints.Num());
-			}
-		}
-
-		// Preallocate the metadata entries array if we're going to use it
-		AllMetadataEntries.SetNumUninitialized(SourcePoints.Num() * TargetPoints.Num());
-	}
-
-	// Use implicit capture, since we capture a lot
-	FPCGAsync::AsyncPointProcessing(Context, SourcePoints.Num() * TargetPoints.Num(), OutPoints, [&](int32 Index, FPCGPoint& OutPoint)
-	{
-		const FPCGPoint& SourcePoint = SourcePoints[Index / TargetPoints.Num()];
-		const FPCGPoint& TargetPoint = TargetPoints[Index % TargetPoints.Num()];
-
-		OutPoint = SourcePoint;
-
-		// Set Rotation, Scale, and Color based on inheritance mode
-		if (RotationInheritance == EPCGCopyPointsInheritanceMode::Relative)
-		{
-			OutPoint.Transform.SetRotation(TargetPoint.Transform.GetRotation() * SourcePoint.Transform.GetRotation());
-		}
-		else if (RotationInheritance == EPCGCopyPointsInheritanceMode::Source)
-		{
-			OutPoint.Transform.SetRotation(SourcePoint.Transform.GetRotation());
-		}
-		else // if (RotationInheritance == EPCGCopyPointsInheritanceMode::Target)
-		{
-			OutPoint.Transform.SetRotation(TargetPoint.Transform.GetRotation());
-		}
+					for (TArray<TTuple<PCGMetadataEntryKey, PCGMetadataValueKey>>& AttributeValues : AttributeValuesToSet)
+					{
+						AttributeValues.SetNumUninitialized(SourcePoints.Num() * TargetPoints.Num());
+					}
+				}
 
-		if (ScaleInheritance == EPCGCopyPointsInheritanceMode::Relative)
-		{
-			OutPoint.Transform.SetScale3D(SourcePoint.Transform.GetScale3D() * TargetPoint.Transform.GetScale3D());
-		}
-		else if (ScaleInheritance == EPCGCopyPointsInheritanceMode::Source)
-		{ 
-			OutPoint.Transform.SetScale3D(SourcePoint.Transform.GetScale3D());
-		}
-		else // if (ScaleInheritance == EPCGCopyPointsInheritanceMode::Target)
-		{
-			OutPoint.Transform.SetScale3D(TargetPoint.Transform.GetScale3D());
-		}
+				// Preallocate the metadata entries array if we're going to use it
+				AllMetadataEntries.SetNumUninitialized(SourcePoints.Num() * TargetPoints.Num());
+			}
 
-		if (ColorInheritance == EPCGCopyPointsInheritanceMode::Relative)
-		{
-			OutPoint.Color = SourcePoint.Color * TargetPoint.Color;
-		}
-		else if (ColorInheritance == EPCGCopyPointsInheritanceMode::Source)
-		{ 
-			OutPoint.Color = SourcePoint.Color;
-		}
-		else // if (ColorInheritance == EPCGCopyPointsInheritanceMode::Target)
-		{ 
-			OutPoint.Color = TargetPoint.Color;
-		}
+			// Use implicit capture, since we capture a lot
+			FPCGAsync::AsyncPointProcessing(Context, SourcePoints.Num() * TargetPoints.Num(), OutPoints, [&](int32 Index, FPCGPoint& OutPoint)
+			{
+				const FPCGPoint& SourcePoint = SourcePoints[Index / TargetPoints.Num()];
+				const FPCGPoint& TargetPoint = TargetPoints[Index % TargetPoints.Num()];
 
-		const FVector Location = TargetPoint.Transform.TransformPosition(SourcePoint.Transform.GetLocation());
-		OutPoint.Transform.SetLocation(Location);
+				OutPoint = SourcePoint;
 
-		// Set seed based on inheritance mode
-		if (SeedInheritance == EPCGCopyPointsInheritanceMode::Relative)
-		{
-			OutPoint.Seed = PCGHelpers::ComputeSeed(SourcePoint.Seed, TargetPoint.Seed);
-		}
-		else if (SeedInheritance == EPCGCopyPointsInheritanceMode::Target)
-		{
-			OutPoint.Seed = TargetPoint.Seed;
-		}
+				// Set Rotation, Scale, and Color based on inheritance mode
+				if (RotationInheritance == EPCGCopyPointsInheritanceMode::Relative)
+				{
+					OutPoint.Transform.SetRotation(TargetPoint.Transform.GetRotation() * SourcePoint.Transform.GetRotation());
+				}
+				else if (RotationInheritance == EPCGCopyPointsInheritanceMode::Source)
+				{
+					OutPoint.Transform.SetRotation(SourcePoint.Transform.GetRotation());
+				}
+				else // if (RotationInheritance == EPCGCopyPointsInheritanceMode::Target)
+				{
+					OutPoint.Transform.SetRotation(TargetPoint.Transform.GetRotation());
+				}
 
-		if (bProcessMetadata)
-		{
-			const FPCGPoint* RootPoint = (bInheritMetadataFromSource ? &SourcePoint : &TargetPoint);
+				if (ScaleInheritance == EPCGCopyPointsInheritanceMode::Relative)
+				{
+					OutPoint.Transform.SetScale3D(SourcePoint.Transform.GetScale3D() * TargetPoint.Transform.GetScale3D());
+				}
+				else if (ScaleInheritance == EPCGCopyPointsInheritanceMode::Source)
+				{ 
+					OutPoint.Transform.SetScale3D(SourcePoint.Transform.GetScale3D());
+				}
+				else // if (ScaleInheritance == EPCGCopyPointsInheritanceMode::Target)
+				{
+					OutPoint.Transform.SetScale3D(TargetPoint.Transform.GetScale3D());
+				}
 
-			OutPoint.MetadataEntry = OutMetadata->AddEntryPlaceholder();
-			AllMetadataEntries[Index] = TTuple<int64, int64>(OutPoint.MetadataEntry, RootPoint->MetadataEntry);
+				if (ColorInheritance == EPCGCopyPointsInheritanceMode::Relative)
+				{
+					OutPoint.Color = SourcePoint.Color * TargetPoint.Color;
+				}
+				else if (ColorInheritance == EPCGCopyPointsInheritanceMode::Source)
+				{ 
+					OutPoint.Color = SourcePoint.Color;
+				}
+				else // if (ColorInheritance == EPCGCopyPointsInheritanceMode::Target)
+				{ 
+					OutPoint.Color = TargetPoint.Color;
+				}
 
-			if (NonRootMetadata)
-			{
-				const FPCGPoint* NonRootPoint = (bInheritMetadataFromSource ? &TargetPoint : &SourcePoint);
+				const FVector Location = TargetPoint.Transform.TransformPosition(SourcePoint.Transform.GetLocation());
+				OutPoint.Transform.SetLocation(Location);
 
-				// Copy EntryToValue key mappings from NonRootAttributes - no need to do it if the non-root uses the default values
-				if (NonRootPoint->MetadataEntry != PCGInvalidEntryKey)
+				// Set seed based on inheritance mode
+				if (SeedInheritance == EPCGCopyPointsInheritanceMode::Relative)
 				{
-					for (int32 AttributeIndex = 0; AttributeIndex < NonRootAttributes.Num(); ++AttributeIndex)
-					{
-						AttributeValuesToSet[AttributeIndex][Index] = TTuple<PCGMetadataEntryKey, PCGMetadataValueKey>(OutPoint.MetadataEntry, NonRootAttributes[AttributeIndex]->GetValueKey(NonRootPoint->MetadataEntry));
-					}
+					OutPoint.Seed = PCGHelpers::ComputeSeed(SourcePoint.Seed, TargetPoint.Seed);
+				}
+				else if (SeedInheritance == EPCGCopyPointsInheritanceMode::Target)
+				{
+					OutPoint.Seed = TargetPoint.Seed;
 				}
-				else
+
+				if (bProcessMetadata)
 				{
-					for (int32 AttributeIndex = 0; AttributeIndex < NonRootAttributes.Num(); ++AttributeIndex)
+					const FPCGPoint* RootPoint = (bInheritMetadataFromSource ? &SourcePoint : &TargetPoint);
+
+					OutPoint.MetadataEntry = OutMetadata->AddEntryPlaceholder();
+					AllMetadataEntries[Index] = TTuple<int64, int64>(OutPoint.MetadataEntry, RootPoint->MetadataEntry);
+
+					if (NonRootMetadata)
 					{
-						AttributeValuesToSet[AttributeIndex][Index] = TTuple<PCGMetadataEntryKey, PCGMetadataValueKey>(OutPoint.MetadataEntry, PCGDefaultValueKey);
+						const FPCGPoint* NonRootPoint = (bInheritMetadataFromSource ? &TargetPoint : &SourcePoint);
+
+						// Copy EntryToValue key mappings from NonRootAttributes - no need to do it if the non-root uses the default values
+						if (NonRootPoint->MetadataEntry != PCGInvalidEntryKey)
+						{
+							for (int32 AttributeIndex = 0; AttributeIndex < NonRootAttributes.Num(); ++AttributeIndex)
+							{
+								AttributeValuesToSet[AttributeIndex][Index] = TTuple<PCGMetadataEntryKey, PCGMetadataValueKey>(OutPoint.MetadataEntry, NonRootAttributes[AttributeIndex]->GetValueKey(NonRootPoint->MetadataEntry));
+							}
+						}
+						else
+						{
+							for (int32 AttributeIndex = 0; AttributeIndex < NonRootAttributes.Num(); ++AttributeIndex)
+							{
+								AttributeValuesToSet[AttributeIndex][Index] = TTuple<PCGMetadataEntryKey, PCGMetadataValueKey>(OutPoint.MetadataEntry, PCGDefaultValueKey);
+							}
+						}
 					}
 				}
-			}
-		}
 
-		return true;
-	});
+				return true;
+			});
 
-	if (bProcessMetadata)
-	{
-		TRACE_CPUPROFILER_EVENT_SCOPE(FPCGCopyPointsElement::Execute::SetMetadata);
-		check(AttributesToSet.Num() == AttributeValuesToSet.Num());
-		if (AttributesToSet.Num() > 0)
-		{
-			int32 AttributeOffset = 0;
-			const int32 AttributePerDispatch = FMath::Max(1, Context->AsyncState.NumAvailableTasks);
-
-			while (AttributeOffset < AttributesToSet.Num())
+			if (bProcessMetadata)
 			{
-				const int32 AttributeCountInCurrentDispatch = FMath::Min(AttributePerDispatch, AttributesToSet.Num() - AttributeOffset);
-				ParallelFor(AttributeCountInCurrentDispatch, [AttributeOffset, &AttributesToSet, &AttributeValuesToSet](int32 WorkerIndex)
+				TRACE_CPUPROFILER_EVENT_SCOPE(FPCGCopyPointsElement::Execute::SetMetadata);
+				check(AttributesToSet.Num() == AttributeValuesToSet.Num());
+				if (AttributesToSet.Num() > 0)
 				{
-					FPCGMetadataAttributeBase* Attribute = AttributesToSet[AttributeOffset + WorkerIndex];
-					const TArray<TTuple<PCGMetadataEntryKey, PCGMetadataValueKey>>& Values = AttributeValuesToSet[AttributeOffset + WorkerIndex];
-					check(Attribute);
-					Attribute->SetValuesFromValueKeys(Values, /*bResetValueOnDefaultValueKey*/false); // no need for the reset here, our points will not have any prior value for these attributes
-				});
+					int32 AttributeOffset = 0;
+					const int32 AttributePerDispatch = FMath::Max(1, Context->AsyncState.NumAvailableTasks);
+
+					while (AttributeOffset < AttributesToSet.Num())
+					{
+						const int32 AttributeCountInCurrentDispatch = FMath::Min(AttributePerDispatch, AttributesToSet.Num() - AttributeOffset);
+						ParallelFor(AttributeCountInCurrentDispatch, [AttributeOffset, &AttributesToSet, &AttributeValuesToSet](int32 WorkerIndex)
+						{
+							FPCGMetadataAttributeBase* Attribute = AttributesToSet[AttributeOffset + WorkerIndex];
+							const TArray<TTuple<PCGMetadataEntryKey, PCGMetadataValueKey>>& Values = AttributeValuesToSet[AttributeOffset + WorkerIndex];
+							check(Attribute);
+							Attribute->SetValuesFromValueKeys(Values, /*bResetValueOnDefaultValueKey*/false); // no need for the reset here, our points will not have any prior value for these attributes
+						});
+
+						AttributeOffset += AttributeCountInCurrentDispatch;
+					}
+				}
 
-				AttributeOffset += AttributeCountInCurrentDispatch;
+				OutMetadata->AddDelayedEntries(AllMetadataEntries);
 			}
 		}
-
-		OutMetadata->AddDelayedEntries(AllMetadataEntries);
 	}
 
 	return true;
