#UE Disable PCGWorldActor cook caching for WP for now. It causes a bug that prevents the landscape from being cooked

diff --git a/Engine/Plugins/Experimental/PCG/Source/PCG/Private/PCGWorldActor.cpp b/Engine/Plugins/Experimental/PCG/Source/PCG/Private/PCGWorldActor.cpp
--- a/Engine/Plugins/Experimental/PCG/Source/PCG/Private/PCGWorldActor.cpp
+++ b/Engine/Plugins/Experimental/PCG/Source/PCG/Private/PCGWorldActor.cpp
@@ -51,12 +51,14 @@ void APCGWorldActor::BeginCacheForCookedPlatformData(const ITargetPlatform* Targ
 
 		if (UWorldPartition* WorldPartition = World->GetWorldPartition())
 		{
+			/*
 			FWorldPartitionHelpers::ForEachActorDesc<ALandscapeProxy>(WorldPartition, [WorldPartition, &ActorRefs](const FWorldPartitionActorDesc* ActorDesc)
 			{
 				check(ActorDesc);
 				ActorRefs.Add(FWorldPartitionReference(WorldPartition, ActorDesc->GetGuid()));
 				return true;
 			});
+			*/
 		}
 		else
 		{
