[PCG] Add support for arrays and structures in Graph Parameters

* Arrays and structures are now selectable in the Graph Parameter menu.
* Objects are now extractable like structs with an option. Structs are marked always extractable.
* Re-factorize the property getter since it was quasi identical between Get Actor Property and Get Graph Parameter.
* Update unit test for the updated error reporting.
* Update the override UI to reflect property overridden for arrays and struct.

diff --git a/Engine/Plugins/Experimental/PCG/Config/BasePCG.ini b/Engine/Plugins/Experimental/PCG/Config/BasePCG.ini
--- a/Engine/Plugins/Experimental/PCG/Config/BasePCG.ini
+++ b/Engine/Plugins/Experimental/PCG/Config/BasePCG.ini
@@ -54,6 +54,7 @@
 +PropertyRedirects=(OldName="PCGGraphParametersHelpers.SetVectorParameter.GraphInstance",NewName="PCGGraphParametersHelpers.SetVectorParameter.GraphInterface")
 +PropertyRedirects=(OldName="PCGGraphParametersHelpers.SetRotatorParameter.GraphInstance",NewName="PCGGraphParametersHelpers.SetRotatorParameter.GraphInterface")
 +PropertyRedirects=(OldName="PCGGraphParametersHelpers.SetTransformParameter.GraphInstance",NewName="PCGGraphParametersHelpers.SetTransformParameter.GraphInterface")
++PropertyRedirects=(OldName="PCGPropertyToParamDataSettings.bExtractObjectAndStruct",NewName="PCGPropertyToParamDataSettings.bForceObjectAndStructExtraction")
 
 +ClassRedirects=(OldName="PCGPointSamplerSettings",NewName="/Script/PCG.PCGSelectPointsSettings")
 +ClassRedirects=(OldName="PCGDensityNoiseSettings",NewName="/Script/PCG.PCGAttributeNoiseSettings")
diff --git a/Engine/Plugins/Experimental/PCG/Source/PCG/Private/Elements/PCGPropertyToParamData.cpp b/Engine/Plugins/Experimental/PCG/Source/PCG/Private/Elements/PCGPropertyToParamData.cpp
--- a/Engine/Plugins/Experimental/PCG/Source/PCG/Private/Elements/PCGPropertyToParamData.cpp
+++ b/Engine/Plugins/Experimental/PCG/Source/PCG/Private/Elements/PCGPropertyToParamData.cpp
@@ -7,8 +7,8 @@
 #include "PCGComponent.h"
 #include "PCGParamData.h"
 #include "Helpers/PCGBlueprintHelpers.h"
+#include "Helpers/PCGPropertyHelpers.h"
 #include "Helpers/PCGHelpers.h"
-#include "Helpers/PCGSettingsHelpers.h"
 #include "Metadata/Accessors/PCGAttributeAccessorHelpers.h"
 
 #include UE_INLINE_GENERATED_CPP_BY_NAME(PCGPropertyToParamData)
@@ -154,147 +154,8 @@ bool FPCGPropertyToParamDataElement::ExecuteInternal(FPCGContext* Context) const
 		}
 	}
 
-	// Try to get the property
-	FProperty* Property = FindFProperty<FProperty>(ObjectToInspect->GetClass(), Settings->PropertyName);
-	if (!Property)
-	{
-		PCGE_LOG(Error, GraphAndLog, FText::Format(LOCTEXT("PropertyDoesNotExist", "Property '{0}' does not exist in the found actor"), FText::FromName(Settings->PropertyName)));
-		return true;
-	}
-
-	// Make sure the property is visible
-	const uint64 ExcludePropertyFlags = CPF_DisableEditOnInstance;
-	const uint64 IncludePropertyFlags = CPF_BlueprintVisible;
-	if (Property->HasAnyPropertyFlags(ExcludePropertyFlags) || !Property->HasAnyPropertyFlags(IncludePropertyFlags))
-	{
-		PCGE_LOG(Error, GraphAndLog, FText::Format(LOCTEXT("PropertyExistsButNotVisible", "Property '{0}' does exist in the found actor, but is not visible."), FText::FromName(Settings->PropertyName)));
-		return true;
-	}
-
-	// If the property is an array, we will work on the underlying property, and extract each element as an entry in the param data
-	FArrayProperty* ArrayProperty = CastField<FArrayProperty>(Property);
-	if (ArrayProperty)
-	{
-		Property = ArrayProperty->Inner;
-	}
-
-	using GetAddressFunc = TFunction<const void* (const void*)>;
-	using ExtractablePropertyTuple = TTuple<FName, const FProperty*>;
-	TArray<ExtractablePropertyTuple> ExtractableProperties;
-
-	GetAddressFunc AddressFunc;
-
-	// Special case where the property is a struct/object, that is not supported by our metadata, we will try to break it down to multiple attributes in the resulting param data, if asked.
-	if ((Property->IsA<FStructProperty>() || Property->IsA<FObjectProperty>()) && Settings->bExtractObjectAndStruct)
-	{
-		UScriptStruct* UnderlyingStruct = nullptr;
-		UClass* UnderlyingClass = nullptr;
-
-		if (FStructProperty* StructProperty = CastField<FStructProperty>(Property))
-		{
-			UnderlyingStruct = StructProperty->Struct;
-			AddressFunc = [StructProperty](const void* InAddress) { return StructProperty->ContainerPtrToValuePtr<void>(InAddress); };
-		}
-		else if (FObjectProperty* ObjectProperty = CastField<FObjectProperty>(Property))
-		{
-			UnderlyingClass = ObjectProperty->PropertyClass;
-			AddressFunc = [ObjectProperty](const void* InAddress) { return ObjectProperty->GetObjectPropertyValue_InContainer(InAddress); };
-		}
-
-		check(UnderlyingStruct || UnderlyingClass);
-		check(!!AddressFunc);
-
-		// Re-use code from overridable params
-		// Limit ourselves to not recurse into more structs.
-		PCGSettingsHelpers::FPCGGetAllOverridableParamsConfig Config;
-		Config.bUseSeed = true;
-		Config.bExcludeSuperProperties = true;
-		Config.MaxStructDepth = 0;
-		// Can only get exposed properties and visible
-		Config.ExcludePropertyFlags = ExcludePropertyFlags;
-		Config.IncludePropertyFlags = IncludePropertyFlags;
-		TArray<FPCGSettingsOverridableParam> AllChildProperties = UnderlyingStruct ? PCGSettingsHelpers::GetAllOverridableParams(UnderlyingStruct, Config) : PCGSettingsHelpers::GetAllOverridableParams(UnderlyingClass, Config);
-
-		for (const FPCGSettingsOverridableParam& Param : AllChildProperties)
-		{
-			if (ensure(!Param.PropertiesNames.IsEmpty()))
-			{
-				const FName ChildPropertyName = Param.PropertiesNames[0];
-				if (const FProperty* ChildProperty = (UnderlyingStruct ? UnderlyingStruct->FindPropertyByName(ChildPropertyName) : UnderlyingClass->FindPropertyByName(ChildPropertyName)))
-				{
-					// We use authored name as attribute name to avoid issue with noisy property names, like in UUserDefinedStructs, where some random number is appended to the property name.
-					// By default, it will just return the property name anyway.
-					const FString AuthoredName = UnderlyingStruct ? UnderlyingStruct->GetAuthoredNameForField(ChildProperty) : UnderlyingClass->GetAuthoredNameForField(ChildProperty);
-					ExtractableProperties.Emplace(FName(AuthoredName), ChildProperty);
-				}
-			}
-		}
-	}
-	else
-	{
-		const FName AttributeName = (Settings->OutputAttributeName == PCGMetadataAttributeConstants::SourceNameAttributeName) ? Property->GetFName() : Settings->OutputAttributeName;
-		ExtractableProperties.Emplace(AttributeName, Property);
-		// Identity
-		AddressFunc = [](const void* InAddress) { return InAddress; };
-	}
-
-	if (ExtractableProperties.IsEmpty())
-	{
-		PCGE_LOG(Error, GraphAndLog, LOCTEXT("NoPropertiesFound", "No properties found to extract"));
-		return true;
-	}
-
-	// Before we need to compute all the addresses for each entry in our array (or just a single entry if there is no array)
-	TArray<const void*, TInlineAllocator<16>> ElementAddresses;
-	if (ArrayProperty)
-	{
-		FScriptArrayHelper_InContainer Helper(ArrayProperty, ObjectToInspect);
-		ElementAddresses.Reserve(Helper.Num());
-		for (int32 DynamicIndex = 0; DynamicIndex < Helper.Num(); ++DynamicIndex)
-		{
-			ElementAddresses.Add(Helper.GetRawPtr(DynamicIndex));
-		}
-	}
-	else
-	{
-		ElementAddresses.Add(ObjectToInspect);
-	}
-
-	// From there, we should be able to create the data.
-	UPCGParamData* ParamData = NewObject<UPCGParamData>();
-	UPCGMetadata* Metadata = ParamData->MutableMetadata();
-	check(Metadata);
-
-	bool bValidOperation = true;
-
-	for (const void* ElementAddress : ElementAddresses)
-	{
-		if (!bValidOperation)
-		{
-			break;
-		}
-
-		// Add a new entry for all elements
-		PCGMetadataEntryKey EntryKey = Metadata->AddEntry();
-
-		for (ExtractablePropertyTuple& ExtractableProperty : ExtractableProperties)
-		{
-			const FName AttributeName = ExtractableProperty.Get<0>();
-			const FProperty* FinalProperty = ExtractableProperty.Get<1>();
-
-			// Offset the address if needed
-			const void* ContainerPtr = AddressFunc(ElementAddress);
-
-			if (!Metadata->SetAttributeFromDataProperty(AttributeName, EntryKey, ContainerPtr, FinalProperty, /*bCreate=*/ true))
-			{
-				PCGE_LOG(Error, GraphAndLog, FText::Format(LOCTEXT("ErrorCreatingAttribute", "Error while creating an attribute for property '{0}'. Either the property type is not supported by PCG or attribute creation failed."), FText::FromString(FinalProperty->GetName())));
-				bValidOperation = false;
-				break;
-			}
-		}
-	}
-
-	if (bValidOperation)
+	PCGPropertyHelpers::FExtractorParameters Parameters{ ObjectToInspect, ObjectToInspect->GetClass(), Settings->PropertyName, Settings->OutputAttributeName, Settings->bForceObjectAndStructExtraction, /*bPropertyNeedsToBeVisible=*/true};
+	if (UPCGParamData* ParamData = PCGPropertyHelpers::ExtractPropertyAsAttributeSet(Parameters, Context))
 	{
 		TArray<FPCGTaggedData>& Outputs = Context->OutputData.TaggedData;
 		FPCGTaggedData& Output = Outputs.Emplace_GetRef();
diff --git a/Engine/Plugins/Experimental/PCG/Source/PCG/Private/Elements/PCGUserParameterGet.cpp b/Engine/Plugins/Experimental/PCG/Source/PCG/Private/Elements/PCGUserParameterGet.cpp
--- a/Engine/Plugins/Experimental/PCG/Source/PCG/Private/Elements/PCGUserParameterGet.cpp
+++ b/Engine/Plugins/Experimental/PCG/Source/PCG/Private/Elements/PCGUserParameterGet.cpp
@@ -4,18 +4,12 @@
 
 #include "PCGComponent.h"
 #include "PCGContext.h"
-#include "PCGEdge.h"
 #include "PCGGraph.h"
-#include "PCGNode.h"
 #include "PCGParamData.h"
 #include "PCGPin.h"
 #include "PCGSubgraph.h"
 #include "Data/PCGUserParametersData.h"
-#include "Metadata/PCGMetadata.h"
-#include "Metadata/PCGMetadataAttributeTpl.h"
-#include "Metadata/Accessors/IPCGAttributeAccessor.h"
-#include "Metadata/Accessors/PCGAttributeAccessorHelpers.h"
-#include "Metadata/Accessors/PCGAttributeAccessorKeys.h"
+#include "Helpers/PCGPropertyHelpers.h"
 
 #include "PropertyBag.h"
 #include "StructView.h"
@@ -155,6 +149,7 @@ bool FPCGUserParameterGetElement::ExecuteInternal(FPCGContext* Context) const
 
 	FConstStructView Parameters = PCGUserParameterGetSettings::GetFirstValidLayout(*Context);
 	const UScriptStruct* PropertyBag = Parameters.GetScriptStruct();
+
 	const FProperty* Property = PropertyBag ? PropertyBag->FindPropertyByName(PropertyName) : nullptr;
 
 	if (!Property)
@@ -163,35 +158,13 @@ bool FPCGUserParameterGetElement::ExecuteInternal(FPCGContext* Context) const
 		return true;
 	}
 
-	TUniquePtr<IPCGAttributeAccessor> PropertyAccessor = PCGAttributeAccessorHelpers::CreatePropertyAccessor(Property);
+	PCGPropertyHelpers::FExtractorParameters ExtractorParameters{ Parameters.GetMemory(), PropertyBag, PropertyName, PropertyName, Settings->bForceObjectAndStructExtraction, /*bPropertyNeedsToBeVisible=*/false };
 
-	if (!PropertyAccessor)
+	if (UPCGParamData* ParamData = PCGPropertyHelpers::ExtractPropertyAsAttributeSet(ExtractorParameters, Context))
 	{
-		// TODO: Should not happen when we have a working filter on the Property types.
-		PCGE_LOG(Error, GraphAndLog, FText::Format(LOCTEXT("InvalidType", "Parameter '{0}' type is unsupported by metadata. Available types are the ones you can find in the Create Attribute node"), FText::FromName(PropertyName)));
-		return true;
+		Context->OutputData.TaggedData.Emplace_GetRef().Data = ParamData;
 	}
 
-	const FPCGAttributeAccessorKeysSingleObjectPtr PropertyAccessorKey(Parameters.GetMemory());
-
-	UPCGParamData* NewParamData = NewObject<UPCGParamData>();
-	check(NewParamData && NewParamData->Metadata);
-	PCGMetadataEntryKey NewEntry = NewParamData->Metadata->AddEntry();
-
-	auto CreateAndSet = [&PropertyAccessor, &PropertyAccessorKey, PropertyName, NewEntry, Metadata = NewParamData->Metadata](auto&& Dummy)
-	{
-		using AttributeType = std::decay_t<decltype(Dummy)>;
-
-		AttributeType Value{};
-		PropertyAccessor->Get<AttributeType>(Value, PropertyAccessorKey);
-
-		FPCGMetadataAttribute<AttributeType>* Attribute = Metadata->CreateAttribute<AttributeType>(PropertyName, Value, /*bAllowInterpolation=*/true, /*bOverrideParent=*/false);
-		Attribute->SetValue(NewEntry, Value);
-	};
-
-	PCGMetadataAttribute::CallbackWithRightType(PropertyAccessor->GetUnderlyingType(), CreateAndSet);
-
-	Context->OutputData.TaggedData.Emplace_GetRef().Data = NewParamData;
 	return true;
 }
 
diff --git a/Engine/Plugins/Experimental/PCG/Source/PCG/Private/Helpers/PCGPropertyHelpers.cpp b/Engine/Plugins/Experimental/PCG/Source/PCG/Private/Helpers/PCGPropertyHelpers.cpp
--- a/Engine/Plugins/Experimental/PCG/Source/PCG/Private/Helpers/PCGPropertyHelpers.cpp
+++ b/Engine/Plugins/Experimental/PCG/Source/PCG/Private/Helpers/PCGPropertyHelpers.cpp
@@ -2,7 +2,17 @@
 
 #include "Helpers/PCGPropertyHelpers.h"
 
+#include "PCGContext.h"
+#include "PCGElement.h"
+#include "PCGModule.h"
+#include "PCGParamData.h"
+#include "Helpers/PCGSettingsHelpers.h"
+#include "Metadata/PCGMetadata.h"
+
 #include "UObject/Field.h"
+#include "UObject/UnrealType.h"
+
+#define LOCTEXT_NAMESPACE "PCGPropertyHelpers"
 
 EPCGMetadataTypes PCGPropertyHelpers::GetMetadataTypeFromProperty(const FProperty* InProperty)
 {
@@ -21,3 +31,174 @@ EPCGMetadataTypes PCGPropertyHelpers::GetMetadataTypeFromProperty(const FPropert
 
 	return PropertyAccessor.IsValid() ? EPCGMetadataTypes(PropertyAccessor->GetUnderlyingType()) : EPCGMetadataTypes::Unknown;
 }
+
+UPCGParamData* PCGPropertyHelpers::ExtractPropertyAsAttributeSet(const PCGPropertyHelpers::FExtractorParameters& Parameters, FPCGContext* InOptionalContext)
+{
+	check(Parameters.Container && Parameters.Class);
+
+	auto LogError = [InOptionalContext](const FText& ErrorMessage)
+	{
+		if (InOptionalContext)
+		{
+			PCGE_LOG_C(Error, GraphAndLog, InOptionalContext, ErrorMessage);
+		}
+		else
+		{
+			UE_LOG(LogPCG, Error, TEXT("%s"), *ErrorMessage.ToString());
+		}
+	};
+
+	// Try to get the property
+	FProperty* Property = FindFProperty<FProperty>(Parameters.Class, Parameters.PropertyName);
+	if (!Property)
+	{
+		LogError(FText::Format(LOCTEXT("PropertyDoesNotExist", "Property '{0}' does not exist."), FText::FromName(Parameters.PropertyName)));
+		return nullptr;
+	}
+
+	// Make sure the property is visible, if requested
+	const uint64 ExcludePropertyFlags = CPF_DisableEditOnInstance;
+	const uint64 IncludePropertyFlags = CPF_BlueprintVisible;
+	if (Parameters.bPropertyNeedsToBeVisible)
+	{
+		if (Property->HasAnyPropertyFlags(ExcludePropertyFlags) || !Property->HasAnyPropertyFlags(IncludePropertyFlags))
+		{
+			LogError(FText::Format(LOCTEXT("PropertyExistsButNotVisible", "Property '{0}' does exist, but is not visible."), FText::FromName(Parameters.PropertyName)));
+			return nullptr;
+		}
+	}
+
+	// If the property is an array, we will work on the underlying property, and extract each element as an entry in the param data
+	FArrayProperty* ArrayProperty = CastField<FArrayProperty>(Property);
+	if (ArrayProperty)
+	{
+		Property = ArrayProperty->Inner;
+	}
+
+	using GetAddressFunc = TFunction<const void* (const void*)>;
+	using ExtractablePropertyTuple = TTuple<FName, const FProperty*>;
+	TArray<ExtractablePropertyTuple> ExtractableProperties;
+
+	GetAddressFunc AddressFunc;
+
+	// Force extraction if the property is not supported by accessors.
+	const bool bShouldExtract = Parameters.bShouldExtract || !PCGAttributeAccessorHelpers::IsPropertyAccessorSupported(Property);
+
+	// Special case where the property is a struct/object, that is not supported by our metadata, we will try to break it down to multiple attributes in the resulting param data, if asked.
+	if ((Property->IsA<FStructProperty>() || Property->IsA<FObjectProperty>()) && bShouldExtract)
+	{
+		UScriptStruct* UnderlyingStruct = nullptr;
+		UClass* UnderlyingClass = nullptr;
+
+		if (FStructProperty* StructProperty = CastField<FStructProperty>(Property))
+		{
+			UnderlyingStruct = StructProperty->Struct;
+			AddressFunc = [StructProperty](const void* InAddress) { return StructProperty->ContainerPtrToValuePtr<void>(InAddress); };
+		}
+		else if (FObjectProperty* ObjectProperty = CastField<FObjectProperty>(Property))
+		{
+			UnderlyingClass = ObjectProperty->PropertyClass;
+			AddressFunc = [ObjectProperty](const void* InAddress) { return ObjectProperty->GetObjectPropertyValue_InContainer(InAddress); };
+		}
+
+		check(UnderlyingStruct || UnderlyingClass);
+		check(!!AddressFunc);
+
+		// Re-use code from overridable params
+		// Limit ourselves to not recurse into more structs.
+		PCGSettingsHelpers::FPCGGetAllOverridableParamsConfig Config;
+		Config.bUseSeed = true;
+		Config.bExcludeSuperProperties = true;
+		Config.MaxStructDepth = 0;
+		// Can only get exposed properties and visible if requested
+		if (Parameters.bPropertyNeedsToBeVisible)
+		{
+			Config.ExcludePropertyFlags = ExcludePropertyFlags;
+			Config.IncludePropertyFlags = IncludePropertyFlags;
+		}
+		TArray<FPCGSettingsOverridableParam> AllChildProperties = UnderlyingStruct ? PCGSettingsHelpers::GetAllOverridableParams(UnderlyingStruct, Config) : PCGSettingsHelpers::GetAllOverridableParams(UnderlyingClass, Config);
+
+		for (const FPCGSettingsOverridableParam& Param : AllChildProperties)
+		{
+			if (ensure(!Param.PropertiesNames.IsEmpty()))
+			{
+				const FName ChildPropertyName = Param.PropertiesNames[0];
+				if (const FProperty* ChildProperty = (UnderlyingStruct ? UnderlyingStruct->FindPropertyByName(ChildPropertyName) : UnderlyingClass->FindPropertyByName(ChildPropertyName)))
+				{
+					// We use authored name as attribute name to avoid issue with noisy property names, like in UUserDefinedStructs, where some random number is appended to the property name.
+					// By default, it will just return the property name anyway.
+					const FString AuthoredName = UnderlyingStruct ? UnderlyingStruct->GetAuthoredNameForField(ChildProperty) : UnderlyingClass->GetAuthoredNameForField(ChildProperty);
+					ExtractableProperties.Emplace(FName(AuthoredName), ChildProperty);
+				}
+			}
+		}
+	}
+	else
+	{
+		// For non struct/object, there is just a single property to extract with no shenanigans for address indirection.
+		const FName AttributeName = (Parameters.OutputAttributeName == PCGMetadataAttributeConstants::SourceNameAttributeName) ? Property->GetFName() : Parameters.OutputAttributeName;
+		ExtractableProperties.Emplace(AttributeName, Property);
+		// Identity
+		AddressFunc = [](const void* InAddress) { return InAddress; };
+	}
+
+	if (ExtractableProperties.IsEmpty())
+	{
+		LogError(LOCTEXT("NoPropertiesFound", "No properties found to extract"));
+		return nullptr;
+	}
+
+	// Before we need to compute all the addresses for each entry in our array (or just a single entry if there is no array)
+	TArray<const void*, TInlineAllocator<16>> ElementAddresses;
+	if (ArrayProperty)
+	{
+		FScriptArrayHelper_InContainer Helper(ArrayProperty, Parameters.Container);
+		ElementAddresses.Reserve(Helper.Num());
+		for (int32 DynamicIndex = 0; DynamicIndex < Helper.Num(); ++DynamicIndex)
+		{
+			ElementAddresses.Add(Helper.GetRawPtr(DynamicIndex));
+		}
+	}
+	else
+	{
+		ElementAddresses.Add(Parameters.Container);
+	}
+
+	// From there, we should be able to create the data.
+	UPCGParamData* ParamData = NewObject<UPCGParamData>();
+	UPCGMetadata* Metadata = ParamData->MutableMetadata();
+	check(Metadata);
+
+	bool bValidOperation = true;
+
+	for (const void* ElementAddress : ElementAddresses)
+	{
+		if (!bValidOperation)
+		{
+			break;
+		}
+
+		// Add a new entry for all elements
+		PCGMetadataEntryKey EntryKey = Metadata->AddEntry();
+
+		for (ExtractablePropertyTuple& ExtractableProperty : ExtractableProperties)
+		{
+			const FName AttributeName = ExtractableProperty.Get<0>();
+			const FProperty* FinalProperty = ExtractableProperty.Get<1>();
+
+			// Offset the address if needed
+			const void* ContainerPtr = AddressFunc(ElementAddress);
+
+			if (!Metadata->SetAttributeFromDataProperty(AttributeName, EntryKey, ContainerPtr, FinalProperty, /*bCreate=*/ true))
+			{
+				LogError(FText::Format(LOCTEXT("ErrorCreatingAttribute", "Error while creating an attribute for property '{0}'. Either the property type is not supported by PCG or attribute creation failed."), FText::FromString(FinalProperty->GetName())));
+				bValidOperation = false;
+				break;
+			}
+		}
+	}
+
+	return bValidOperation ? ParamData : nullptr;
+}
+
+#undef LOCTEXT_NAMESPACE
\ No newline at end of file
diff --git a/Engine/Plugins/Experimental/PCG/Source/PCG/Private/PCGGraph.cpp b/Engine/Plugins/Experimental/PCG/Source/PCG/Private/PCGGraph.cpp
--- a/Engine/Plugins/Experimental/PCG/Source/PCG/Private/PCGGraph.cpp
+++ b/Engine/Plugins/Experimental/PCG/Source/PCG/Private/PCGGraph.cpp
@@ -1178,35 +1178,7 @@ bool UPCGGraph::UserParametersCanRemoveProperty(FGuid InPropertyID, FName InProp
 bool UPCGGraph::UserParametersIsPinTypeAccepted(FEdGraphPinType InPinType, bool bIsChild)
 {
 	// Text and interface not supported
-	if (InPinType.PinCategory == TEXT("text") || InPinType.PinCategory == TEXT("interface"))
-	{
-		return false;
-	}
-	else if (InPinType.PinCategory == TEXT("struct"))
-	{
-		// Vector/Transform/Rotator are "explicitly" defined. Also return true for a "null" PinSubCategoryObject
-		// to mark it as "valid" so the system will call this function back on all available structures.
-		if (!bIsChild)
-		{
-			return InPinType.PinSubCategoryObject.IsExplicitlyNull()
-				|| InPinType.PinSubCategoryObject == TBaseStructure<FVector>::Get()
-				|| InPinType.PinSubCategoryObject == TBaseStructure<FTransform>::Get()
-				|| InPinType.PinSubCategoryObject == TBaseStructure<FRotator>::Get();
-		}
-		else
-		{
-			// Since we are not supporting all the structures by default, only show the ones that we are supporting with PCG.
-			return InPinType.PinSubCategoryObject == TBaseStructure<FVector2D>::Get()
-				|| InPinType.PinSubCategoryObject == TBaseStructure<FVector4>::Get()
-				|| InPinType.PinSubCategoryObject == TBaseStructure<FQuat>::Get()
-				|| InPinType.PinSubCategoryObject == TBaseStructure<FSoftObjectPath>::Get()
-				|| InPinType.PinSubCategoryObject == TBaseStructure<FSoftClassPath>::Get();
-		}
-	}
-	else
-	{
-		return true;
-	}
+	return InPinType.PinCategory != TEXT("text") && InPinType.PinCategory != TEXT("interface");
 }
 #endif // WITH_EDITOR
 
diff --git a/Engine/Plugins/Experimental/PCG/Source/PCG/Private/Tests/Elements/PCGPropertyToParamDataElementTest.cpp b/Engine/Plugins/Experimental/PCG/Source/PCG/Private/Tests/Elements/PCGPropertyToParamDataElementTest.cpp
--- a/Engine/Plugins/Experimental/PCG/Source/PCG/Private/Tests/Elements/PCGPropertyToParamDataElementTest.cpp
+++ b/Engine/Plugins/Experimental/PCG/Source/PCG/Private/Tests/Elements/PCGPropertyToParamDataElementTest.cpp
@@ -52,7 +52,7 @@ bool VerifyAttributeValue(FPCGTestBaseClass* TestInstance, PCGTestsCommon::FTest
 
 	Settings->PropertyName = PropertyName;
 	Settings->OutputAttributeName = AttributeNames[0];
-	Settings->bExtractObjectAndStruct = AttributeNames.Num() > 1;
+	Settings->bForceObjectAndStructExtraction = AttributeNames.Num() > 1;
 
 	// Add 2 nodes, PropertyToParamDataNode and a Trivial node (just there for the connection)
 	UPCGNode* TestNode = TestData.TestPCGComponent->GetGraph()->AddNode(Settings);
@@ -286,7 +286,7 @@ bool FPCGPropertyToParamDataPropertyTypeTest::RunTest(const FString& Parameters)
 	bSuccess &= VerifyAttributeValueInvalid(this, TestData, GET_MEMBER_NAME_CHECKED(APCGUnitTestDummyActor, ColorProperty), ColorValue, ExtraTestWhat);
 
 	// Unknown property
-	AddExpectedError(TEXT("Property 'DummyMissingProperty' does not exist in the found actor"), EAutomationExpectedErrorFlags::Contains, 1);
+	AddExpectedError(TEXT("Property 'DummyMissingProperty' does not exist"), EAutomationExpectedErrorFlags::Contains, 1);
 	bSuccess &= VerifyAttributeValueInvalid(this, TestData, TEXT("DummyMissingProperty"), 42, ExtraTestWhat);
 
 	// Missing property
diff --git a/Engine/Plugins/Experimental/PCG/Source/PCG/Public/Elements/PCGPropertyToParamData.h b/Engine/Plugins/Experimental/PCG/Source/PCG/Public/Elements/PCGPropertyToParamData.h
--- a/Engine/Plugins/Experimental/PCG/Source/PCG/Public/Elements/PCGPropertyToParamData.h
+++ b/Engine/Plugins/Experimental/PCG/Source/PCG/Public/Elements/PCGPropertyToParamData.h
@@ -58,9 +58,9 @@ public:
 	UPROPERTY(BlueprintReadWrite, EditAnywhere, Category = Settings)
 	FName PropertyName = NAME_None;
 
-	/** If the property is a struct/object unsupported by metadata, this option can be toggled to extract all (compatible) properties contained in this property. For now, only supports direct child properties (and not deeper). */
+	/** If the property is a struct/object supported by metadata, this option can be toggled to force extracting all (compatible) properties contained in this property. Automatically true if unsupported by metadata. For now, only supports direct child properties (and not deeper). */
 	UPROPERTY(BlueprintReadWrite, EditAnywhere, Category = Settings)
-	bool bExtractObjectAndStruct = false;
+	bool bForceObjectAndStructExtraction = false;
 
 	/** By default, attribute name will be None, but it can be overridden by this name. */
 	UPROPERTY(BlueprintReadWrite, EditAnywhere, Category = Settings, meta = (EditCondition = "!bExtractObjectAndStruct", EditConditionHides))
diff --git a/Engine/Plugins/Experimental/PCG/Source/PCG/Public/Elements/PCGUserParameterGet.h b/Engine/Plugins/Experimental/PCG/Source/PCG/Public/Elements/PCGUserParameterGet.h
--- a/Engine/Plugins/Experimental/PCG/Source/PCG/Public/Elements/PCGUserParameterGet.h
+++ b/Engine/Plugins/Experimental/PCG/Source/PCG/Public/Elements/PCGUserParameterGet.h
@@ -27,6 +27,10 @@ public:
 	UPROPERTY()
 	FName PropertyName;
 
+	/** If the property is a struct/object supported by metadata, this option can be toggled to force extracting all (compatible) properties contained in this property. Automatically true if unsupported by metadata. For now, only supports direct child properties (and not deeper). */
+	UPROPERTY(BlueprintReadWrite, EditAnywhere, Category = Settings)
+	bool bForceObjectAndStructExtraction = false;
+	
 	void UpdatePropertyName(FName InNewName);
 
 	//~Begin UPCGSettings interface
diff --git a/Engine/Plugins/Experimental/PCG/Source/PCG/Public/Helpers/PCGPropertyHelpers.h b/Engine/Plugins/Experimental/PCG/Source/PCG/Public/Helpers/PCGPropertyHelpers.h
--- a/Engine/Plugins/Experimental/PCG/Source/PCG/Public/Helpers/PCGPropertyHelpers.h
+++ b/Engine/Plugins/Experimental/PCG/Source/PCG/Public/Helpers/PCGPropertyHelpers.h
@@ -12,7 +12,9 @@
 
 enum class EPCGMetadataTypes : uint8;
 
+struct FPCGContext;
 class UPCGData;
+class UPCGParamData;
 
 namespace PCGPropertyHelpers
 {
@@ -42,6 +44,36 @@ namespace PCGPropertyHelpers
 	* @returns PCG type if the property is supported, Unknown otherwise.
 	*/
 	PCG_API EPCGMetadataTypes GetMetadataTypeFromProperty(const FProperty* InProperty);
+
+	struct FExtractorParameters
+	{
+		// Pointer to the container containing the data we want to extract
+		const void* Container = nullptr;
+
+		// Class or ScriptStruct of the object/struct we want to extract the property from.
+		const UStruct* Class = nullptr;
+
+		// Name of the property we want to extract
+		FName PropertyName = NAME_None;
+
+		// Optional name of the attribute that will receive the extracted property. If None, will take the property name. 
+		// Also not used for Structs/Object extraction, as we will create multiple attributes, and they will be the name of all the extracted members.
+		FName OutputAttributeName = NAME_None;
+
+		// If the property we want to extract is an object/struct, give the possibility to extract all their member in one go.
+		// Only work if their members are not structs or containers (like array).
+		bool bShouldExtract = false;
+
+		// Extra flag if the property needs to be visible to be extractable.
+		bool bPropertyNeedsToBeVisible = false;
+	};
+
+	/**
+	* Extract a given property in an Attribute Set.
+	* @param Parameters - Parameters for extaction, cf above.
+	* @prama InOptionalContext - Optional context if the extraction is done in a PCG Node, so errors are using the context to log.
+	*/
+	UPCGParamData* ExtractPropertyAsAttributeSet(const FExtractorParameters& Parameters, FPCGContext* InOptionalContext = nullptr);
 }
 
 //////
diff --git a/Engine/Plugins/Experimental/PCG/Source/PCG/Public/PCGGraph.h b/Engine/Plugins/Experimental/PCG/Source/PCG/Public/PCGGraph.h
--- a/Engine/Plugins/Experimental/PCG/Source/PCG/Public/PCGGraph.h
+++ b/Engine/Plugins/Experimental/PCG/Source/PCG/Public/PCGGraph.h
@@ -340,7 +340,7 @@ protected:
 #endif // WITH_EDITORONLY_DATA
 
 	// Parameters
-	UPROPERTY(EditAnywhere, Category = Instance, meta = (DisplayName = "Parameters", NoResetToDefault, DefaultType = "EPropertyBagPropertyType::Double", AllowArrays = "false", IsPinTypeAccepted = "UserParametersIsPinTypeAccepted", CanRemoveProperty = "UserParametersCanRemoveProperty"))
+	UPROPERTY(EditAnywhere, Category = Instance, meta = (DisplayName = "Parameters", NoResetToDefault, DefaultType = "EPropertyBagPropertyType::Double", IsPinTypeAccepted = "UserParametersIsPinTypeAccepted", CanRemoveProperty = "UserParametersCanRemoveProperty"))
 	FInstancedPropertyBag UserParameters;
 
 #if WITH_EDITOR
diff --git a/Engine/Plugins/Experimental/PCG/Source/PCGEditor/Private/Details/PCGInstancedPropertyBagOverrideDetails.cpp b/Engine/Plugins/Experimental/PCG/Source/PCGEditor/Private/Details/PCGInstancedPropertyBagOverrideDetails.cpp
--- a/Engine/Plugins/Experimental/PCG/Source/PCGEditor/Private/Details/PCGInstancedPropertyBagOverrideDetails.cpp
+++ b/Engine/Plugins/Experimental/PCG/Source/PCGEditor/Private/Details/PCGInstancedPropertyBagOverrideDetails.cpp
@@ -143,10 +143,6 @@ void FPCGOverrideInstancedPropertyBagDataDetails::OnChildRowAdded(IDetailPropert
 	FDetailWidgetRow Row;
 	ChildRow.GetDefaultWidgets(NameWidget, ValueWidget, Row);
 
-	FGuid PropertyID;
-
-	const bool bIsOverridden = Owner->IsPropertyOverridden(ChildPropertyHandle->GetProperty());
-
 	ChildRow
 	.CustomWidget(/*bShowChildren*/true)
 	.NameContent()
@@ -172,30 +168,6 @@ void FPCGOverrideInstancedPropertyBagDataDetails::OnChildRowAdded(IDetailPropert
 				.Image(FAppStyle::GetBrush("Icons.Error"))
 			]
 		]
-		// Override 
-		+ SHorizontalBox::Slot()
-		.AutoWidth()
-		[
-			SNew(SCheckBox)
-			.IsChecked_Lambda([this, ChildPropertyHandle]()
-			{
-				return (Owner.IsValid() && Owner->IsPropertyOverridden(ChildPropertyHandle->GetProperty())) ? ECheckBoxState::Checked : ECheckBoxState::Unchecked;
-			})
-			.OnCheckStateChanged_Lambda([this, ValueWidget, ChildPropertyHandle](ECheckBoxState NewState)
-			{
-				PropertiesIDsOverriddenHandle->NotifyPreChange();
-
-				FScopedTransaction Transaction(FText::Format(LOCTEXT("OnCheckStateChanged", "Change Override for {0}"), FText::FromName(ChildPropertyHandle->GetProperty()->GetFName())));
-				const bool bIsOverridden = NewState == ECheckBoxState::Checked;
-				if (Owner.IsValid())
-				{
-					Owner->UpdatePropertyOverride(ChildPropertyHandle->GetProperty(), bIsOverridden);
-				}
-
-				PropertiesIDsOverriddenHandle->NotifyPostChange(EPropertyChangeType::ValueSet);
-			})
-			.IsEnabled_Lambda([this, ChildPropertyHandle]() -> bool { return !IsPropertyOverriddenByPin(ChildPropertyHandle); })
-		]
 		// Name
 		+ SHorizontalBox::Slot()
 		.AutoWidth()
@@ -224,8 +196,22 @@ void FPCGOverrideInstancedPropertyBagDataDetails::OnChildRowAdded(IDetailPropert
 		]
 	];
 
-	ValueWidget->SetEnabled(TAttribute<bool>::CreateLambda([this, ChildPropertyHandle]() -> bool { return Owner.IsValid() && Owner->IsPropertyOverridden(ChildPropertyHandle->GetProperty()) && !IsPropertyOverriddenByPin(ChildPropertyHandle); }));
-	NameWidget->SetEnabled(TAttribute<bool>::CreateLambda([this, ChildPropertyHandle]() -> bool { return !IsPropertyOverriddenByPin(ChildPropertyHandle); }));
+	// A property is mark overriden if the property is overridden and the property is not overridden by a pin (in the case of a subgraph).
+	TAttribute<bool> EditConditionValue = TAttribute<bool>::CreateLambda([this, ChildPropertyHandle]() -> bool { return Owner.IsValid() && Owner->IsPropertyOverridden(ChildPropertyHandle->GetProperty()) && !IsPropertyOverriddenByPin(ChildPropertyHandle); });
+	FOnBooleanValueChanged OnEditConditionChanged = FOnBooleanValueChanged::CreateLambda([this, ChildPropertyHandle](bool bNewValue)
+	{
+		PropertiesIDsOverriddenHandle->NotifyPreChange();
+
+		FScopedTransaction Transaction(FText::Format(LOCTEXT("OnCheckStateChanged", "Change Override for {0}"), FText::FromName(ChildPropertyHandle->GetProperty()->GetFName())));
+		if (Owner.IsValid())
+		{
+			Owner->UpdatePropertyOverride(ChildPropertyHandle->GetProperty(), bNewValue);
+		}
+
+		PropertiesIDsOverriddenHandle->NotifyPostChange(EPropertyChangeType::ValueSet);
+	});
+
+	ChildRow.EditCondition(std::move(EditConditionValue), std::move(OnEditConditionChanged));
 }
 
 #undef LOCTEXT_NAMESPACE
\ No newline at end of file
