[PCG] Add Metadata output for Create Attribute functions in BP for easier chaining

#rnx

diff --git a/Engine/Plugins/Experimental/PCG/Source/PCG/Private/Metadata/PCGMetadata.cpp b/Engine/Plugins/Experimental/PCG/Source/PCG/Private/Metadata/PCGMetadata.cpp
--- a/Engine/Plugins/Experimental/PCG/Source/PCG/Private/Metadata/PCGMetadata.cpp
+++ b/Engine/Plugins/Experimental/PCG/Source/PCG/Private/Metadata/PCGMetadata.cpp
@@ -453,79 +453,94 @@ bool UPCGMetadata::ParentHasAttribute(FName AttributeName) const
 	return Parent && Parent->HasAttribute(AttributeName);
 }
 
-void UPCGMetadata::CreateInteger32Attribute(FName AttributeName, int32 DefaultValue, bool bAllowsInterpolation, bool bOverrideParent)
+UPCGMetadata* UPCGMetadata::CreateInteger32Attribute(FName AttributeName, int32 DefaultValue, bool bAllowsInterpolation, bool bOverrideParent)
 {
 	CreateAttribute<int32>(AttributeName, DefaultValue, bAllowsInterpolation, bOverrideParent);
+	return this;
 }
 
-void UPCGMetadata::CreateInteger64Attribute(FName AttributeName, int64 DefaultValue, bool bAllowsInterpolation, bool bOverrideParent)
+UPCGMetadata* UPCGMetadata::CreateInteger64Attribute(FName AttributeName, int64 DefaultValue, bool bAllowsInterpolation, bool bOverrideParent)
 {
 	CreateAttribute<int64>(AttributeName, DefaultValue, bAllowsInterpolation, bOverrideParent);
+	return this;
 }
 
-void UPCGMetadata::CreateFloatAttribute(FName AttributeName, float DefaultValue, bool bAllowsInterpolation, bool bOverrideParent)
+UPCGMetadata* UPCGMetadata::CreateFloatAttribute(FName AttributeName, float DefaultValue, bool bAllowsInterpolation, bool bOverrideParent)
 {
 	CreateAttribute<float>(AttributeName, DefaultValue, bAllowsInterpolation, bOverrideParent);
+	return this;
 }
 
-void UPCGMetadata::CreateDoubleAttribute(FName AttributeName, double DefaultValue, bool bAllowsInterpolation, bool bOverrideParent)
+UPCGMetadata* UPCGMetadata::CreateDoubleAttribute(FName AttributeName, double DefaultValue, bool bAllowsInterpolation, bool bOverrideParent)
 {
 	CreateAttribute<double>(AttributeName, DefaultValue, bAllowsInterpolation, bOverrideParent);
+	return this;
 }
 
-void UPCGMetadata::CreateVectorAttribute(FName AttributeName, FVector DefaultValue, bool bAllowsInterpolation, bool bOverrideParent)
+UPCGMetadata* UPCGMetadata::CreateVectorAttribute(FName AttributeName, FVector DefaultValue, bool bAllowsInterpolation, bool bOverrideParent)
 {
 	CreateAttribute<FVector>(AttributeName, DefaultValue, bAllowsInterpolation, bOverrideParent);
+	return this;
 }
 
-void UPCGMetadata::CreateVector4Attribute(FName AttributeName, FVector4 DefaultValue, bool bAllowsInterpolation, bool bOverrideParent)
+UPCGMetadata* UPCGMetadata::CreateVector4Attribute(FName AttributeName, FVector4 DefaultValue, bool bAllowsInterpolation, bool bOverrideParent)
 {
 	CreateAttribute<FVector4>(AttributeName, DefaultValue, bAllowsInterpolation, bOverrideParent);
+	return this;
 }
 
-void UPCGMetadata::CreateVector2Attribute(FName AttributeName, FVector2D DefaultValue, bool bAllowsInterpolation, bool bOverrideParent)
+UPCGMetadata* UPCGMetadata::CreateVector2Attribute(FName AttributeName, FVector2D DefaultValue, bool bAllowsInterpolation, bool bOverrideParent)
 {
 	CreateAttribute<FVector2D>(AttributeName, DefaultValue, bAllowsInterpolation, bOverrideParent);
+	return this;
 }
 
-void UPCGMetadata::CreateRotatorAttribute(FName AttributeName, FRotator DefaultValue, bool bAllowsInterpolation, bool bOverrideParent)
+UPCGMetadata* UPCGMetadata::CreateRotatorAttribute(FName AttributeName, FRotator DefaultValue, bool bAllowsInterpolation, bool bOverrideParent)
 {
 	CreateAttribute<FRotator>(AttributeName, DefaultValue, bAllowsInterpolation, bOverrideParent);
+	return this;
 }
 
-void UPCGMetadata::CreateQuatAttribute(FName AttributeName, FQuat DefaultValue, bool bAllowsInterpolation, bool bOverrideParent)
+UPCGMetadata* UPCGMetadata::CreateQuatAttribute(FName AttributeName, FQuat DefaultValue, bool bAllowsInterpolation, bool bOverrideParent)
 {
 	CreateAttribute<FQuat>(AttributeName, DefaultValue, bAllowsInterpolation, bOverrideParent);
+	return this;
 }
 
-void UPCGMetadata::CreateTransformAttribute(FName AttributeName, FTransform DefaultValue, bool bAllowsInterpolation, bool bOverrideParent)
+UPCGMetadata* UPCGMetadata::CreateTransformAttribute(FName AttributeName, FTransform DefaultValue, bool bAllowsInterpolation, bool bOverrideParent)
 {
 	CreateAttribute<FTransform>(AttributeName, DefaultValue, bAllowsInterpolation, bOverrideParent);
+	return this;
 }
 
-void UPCGMetadata::CreateStringAttribute(FName AttributeName, FString DefaultValue, bool bAllowsInterpolation, bool bOverrideParent)
+UPCGMetadata* UPCGMetadata::CreateStringAttribute(FName AttributeName, FString DefaultValue, bool bAllowsInterpolation, bool bOverrideParent)
 {
 	CreateAttribute<FString>(AttributeName, DefaultValue, bAllowsInterpolation, bOverrideParent);
+	return this;
 }
 
-void UPCGMetadata::CreateNameAttribute(FName AttributeName, FName DefaultValue, bool bAllowsInterpolation, bool bOverrideParent)
+UPCGMetadata* UPCGMetadata::CreateNameAttribute(FName AttributeName, FName DefaultValue, bool bAllowsInterpolation, bool bOverrideParent)
 {
 	CreateAttribute<FName>(AttributeName, DefaultValue, bAllowsInterpolation, bOverrideParent);
+	return this;
 }
 
-void UPCGMetadata::CreateBoolAttribute(FName AttributeName, bool DefaultValue, bool bAllowsInterpolation, bool bOverrideParent)
+UPCGMetadata* UPCGMetadata::CreateBoolAttribute(FName AttributeName, bool DefaultValue, bool bAllowsInterpolation, bool bOverrideParent)
 {
 	CreateAttribute<bool>(AttributeName, DefaultValue, bAllowsInterpolation, bOverrideParent);
+	return this;
 }
 
-void UPCGMetadata::CreateSoftObjectPathAttribute(FName AttributeName, const FSoftObjectPath& DefaultValue, bool bAllowsInterpolation, bool bOverrideParent)
+UPCGMetadata* UPCGMetadata::CreateSoftObjectPathAttribute(FName AttributeName, const FSoftObjectPath& DefaultValue, bool bAllowsInterpolation, bool bOverrideParent)
 {
 	CreateAttribute<FSoftObjectPath>(AttributeName, DefaultValue, bAllowsInterpolation, bOverrideParent);
+	return this;
 }
 
-void UPCGMetadata::CreateSoftClassPathAttribute(FName AttributeName, const FSoftClassPath& DefaultValue, bool bAllowsInterpolation, bool bOverrideParent)
+UPCGMetadata* UPCGMetadata::CreateSoftClassPathAttribute(FName AttributeName, const FSoftClassPath& DefaultValue, bool bAllowsInterpolation, bool bOverrideParent)
 {
 	CreateAttribute<FSoftClassPath>(AttributeName, DefaultValue, bAllowsInterpolation, bOverrideParent);
+	return this;
 }
 
 namespace PCGMetadata
diff --git a/Engine/Plugins/Experimental/PCG/Source/PCG/Public/Metadata/PCGMetadata.h b/Engine/Plugins/Experimental/PCG/Source/PCG/Public/Metadata/PCGMetadata.h
--- a/Engine/Plugins/Experimental/PCG/Source/PCG/Public/Metadata/PCGMetadata.h
+++ b/Engine/Plugins/Experimental/PCG/Source/PCG/Public/Metadata/PCGMetadata.h
@@ -85,51 +85,50 @@ public:
 	UFUNCTION(BlueprintCallable, Category = "PCG|Metadata|Advanced")
 	void Flatten();
 
-	/** Create new streams */
 	UFUNCTION(BlueprintCallable, Category = "PCG|Metadata")
-	void CreateInteger32Attribute(FName AttributeName, int32 DefaultValue, bool bAllowsInterpolation, bool bOverrideParent = true);
+	UPARAM(DisplayName = "Metadata") UPCGMetadata* CreateInteger32Attribute(FName AttributeName, int32 DefaultValue, bool bAllowsInterpolation, bool bOverrideParent = true);
 
 	UFUNCTION(BlueprintCallable, Category = "PCG|Metadata")
-	void CreateInteger64Attribute(FName AttributeName, int64 DefaultValue, bool bAllowsInterpolation, bool bOverrideParent = true);
+	UPARAM(DisplayName = "Metadata") UPCGMetadata* CreateInteger64Attribute(FName AttributeName, int64 DefaultValue, bool bAllowsInterpolation, bool bOverrideParent = true);
 
 	UFUNCTION(BlueprintCallable, Category = "PCG|Metadata")
-	void CreateFloatAttribute(FName AttributeName, float DefaultValue, bool bAllowsInterpolation, bool bOverrideParent = true);
+	UPARAM(DisplayName = "Metadata") UPCGMetadata* CreateFloatAttribute(FName AttributeName, float DefaultValue, bool bAllowsInterpolation, bool bOverrideParent = true);
 
 	UFUNCTION(BlueprintCallable, Category = "PCG|Metadata")
-	void CreateDoubleAttribute(FName AttributeName, double DefaultValue, bool bAllowsInterpolation, bool bOverrideParent = true);
+	UPARAM(DisplayName = "Metadata") UPCGMetadata* CreateDoubleAttribute(FName AttributeName, double DefaultValue, bool bAllowsInterpolation, bool bOverrideParent = true);
 
 	UFUNCTION(BlueprintCallable, Category = "PCG|Metadata")
-	void CreateVectorAttribute(FName AttributeName, FVector DefaultValue, bool bAllowsInterpolation, bool bOverrideParent = true);
+	UPARAM(DisplayName = "Metadata") UPCGMetadata* CreateVectorAttribute(FName AttributeName, FVector DefaultValue, bool bAllowsInterpolation, bool bOverrideParent = true);
 
 	UFUNCTION(BlueprintCallable, Category = "PCG|Metadata", meta=(DisplayName = "Create Vector4 Attribute"))
-	void CreateVector4Attribute(FName AttributeName, FVector4 DefaultValue, bool bAllowsInterpolation, bool bOverrideParent = true);
+	UPARAM(DisplayName = "Metadata") UPCGMetadata* CreateVector4Attribute(FName AttributeName, FVector4 DefaultValue, bool bAllowsInterpolation, bool bOverrideParent = true);
 
 	UFUNCTION(BlueprintCallable, Category = "PCG|Metadata")
-	void CreateVector2Attribute(FName AttributeName, FVector2D DefaultValue, bool bAllowsInterpolation, bool bOverrideParent = true);
+	UPARAM(DisplayName = "Metadata") UPCGMetadata* CreateVector2Attribute(FName AttributeName, FVector2D DefaultValue, bool bAllowsInterpolation, bool bOverrideParent = true);
 
 	UFUNCTION(BlueprintCallable, Category = "PCG|Metadata")
-	void CreateRotatorAttribute(FName AttributeName, FRotator DefaultValue, bool bAllowsInterpolation, bool bOverrideParent = true);
+	UPARAM(DisplayName = "Metadata") UPCGMetadata* CreateRotatorAttribute(FName AttributeName, FRotator DefaultValue, bool bAllowsInterpolation, bool bOverrideParent = true);
 
 	UFUNCTION(BlueprintCallable, Category = "PCG|Metadata")
-	void CreateQuatAttribute(FName AttributeName, FQuat DefaultValue, bool bAllowsInterpolation, bool bOverrideParent = true);
+	UPARAM(DisplayName = "Metadata") UPCGMetadata* CreateQuatAttribute(FName AttributeName, FQuat DefaultValue, bool bAllowsInterpolation, bool bOverrideParent = true);
 
 	UFUNCTION(BlueprintCallable, Category = "PCG|Metadata")
-	void CreateTransformAttribute(FName AttributeName, FTransform DefaultValue, bool bAllowsInterpolation, bool bOverrideParent = true);
+	UPARAM(DisplayName = "Metadata") UPCGMetadata* CreateTransformAttribute(FName AttributeName, FTransform DefaultValue, bool bAllowsInterpolation, bool bOverrideParent = true);
 
 	UFUNCTION(BlueprintCallable, Category = "PCG|Metadata")
-	void CreateStringAttribute(FName AttributeName, FString DefaultValue, bool bAllowsInterpolation, bool bOverrideParent = true);
+	UPARAM(DisplayName = "Metadata") UPCGMetadata* CreateStringAttribute(FName AttributeName, FString DefaultValue, bool bAllowsInterpolation, bool bOverrideParent = true);
 
 	UFUNCTION(BlueprintCallable, Category = "PCG|Metadata")
-	void CreateNameAttribute(FName AttributeName, FName DefaultValue, bool bAllowsInterpolation, bool bOverrideParent = true);
+	UPARAM(DisplayName = "Metadata") UPCGMetadata* CreateNameAttribute(FName AttributeName, FName DefaultValue, bool bAllowsInterpolation, bool bOverrideParent = true);
 
 	UFUNCTION(BlueprintCallable, Category = "PCG|Metadata")
-	void CreateBoolAttribute(FName AttributeName, bool DefaultValue, bool bAllowsInterpolation, bool bOverrideParent = true);
+	UPARAM(DisplayName = "Metadata") UPCGMetadata* CreateBoolAttribute(FName AttributeName, bool DefaultValue, bool bAllowsInterpolation, bool bOverrideParent = true);
 
 	UFUNCTION(BlueprintCallable, Category = "PCG|Metadata")
-	void CreateSoftObjectPathAttribute(FName AttributeName, const FSoftObjectPath& DefaultValue, bool bAllowsInterpolation, bool bOverrideParent = true);
+	UPARAM(DisplayName = "Metadata") UPCGMetadata* CreateSoftObjectPathAttribute(FName AttributeName, const FSoftObjectPath& DefaultValue, bool bAllowsInterpolation, bool bOverrideParent = true);
 
 	UFUNCTION(BlueprintCallable, Category = "PCG|Metadata")
-	void CreateSoftClassPathAttribute(FName AttributeName, const FSoftClassPath& DefaultValue, bool bAllowsInterpolation, bool bOverrideParent = true);
+	UPARAM(DisplayName = "Metadata") UPCGMetadata* CreateSoftClassPathAttribute(FName AttributeName, const FSoftClassPath& DefaultValue, bool bAllowsInterpolation, bool bOverrideParent = true);
 
 	/** Creates an attribute given a property.
 	* @param AttributeName: Target attribute to create
