Added editing mode to PCG by the way of a preview mode and preview-on-load modes.

diff --git a/Engine/Plugins/Experimental/PCG/Source/PCG/Private/Elements/PCGCreateTargetActor.cpp b/Engine/Plugins/Experimental/PCG/Source/PCG/Private/Elements/PCGCreateTargetActor.cpp
--- a/Engine/Plugins/Experimental/PCG/Source/PCG/Private/Elements/PCGCreateTargetActor.cpp
+++ b/Engine/Plugins/Experimental/PCG/Source/PCG/Private/Elements/PCGCreateTargetActor.cpp
@@ -252,7 +252,7 @@ bool FPCGCreateTargetActorElement::ExecuteInternal(FPCGContext* Context) const
 	SpawnParams.Template = TemplateActor;
 	SpawnParams.SpawnCollisionHandlingOverride = ESpawnActorCollisionHandlingMethod::AlwaysSpawn;
 
-	if (PCGHelpers::IsRuntimeOrPIE())
+	if (PCGHelpers::IsRuntimeOrPIE() || (Context->SourceComponent.IsValid() && Context->SourceComponent->IsInPreviewMode()))
 	{
 		SpawnParams.ObjectFlags |= RF_Transient;
 	}
diff --git a/Engine/Plugins/Experimental/PCG/Source/PCG/Private/Grid/PCGPartitionActor.cpp b/Engine/Plugins/Experimental/PCG/Source/PCG/Private/Grid/PCGPartitionActor.cpp
--- a/Engine/Plugins/Experimental/PCG/Source/PCG/Private/Grid/PCGPartitionActor.cpp
+++ b/Engine/Plugins/Experimental/PCG/Source/PCG/Private/Grid/PCGPartitionActor.cpp
@@ -93,6 +93,50 @@ void APCGPartitionActor::BeginDestroy()
 	Super::BeginDestroy();
 }
 
+void APCGPartitionActor::Serialize(FArchive& Ar)
+{
+#if WITH_EDITOR
+	TMap<TObjectPtr<UPCGComponent>, TSoftObjectPtr<UPCGComponent>> LocalToOriginalCopy;
+
+	if (Ar.IsSaving())
+	{
+		LocalToOriginalCopy = MoveTemp(LocalToOriginal);
+
+		LocalToOriginal.Reset();
+		for (const auto& It : LocalToOriginalCopy)
+		{
+			if (!It.Key->HasAnyFlags(RF_Transient))
+			{
+				LocalToOriginal.Add(It);
+			}
+		}
+
+		// For components that were cleared but marked as load-as-preview
+		// we need to keep those around only if the original is still in its original state,
+		// i.e. the serialization mode is load on preview and the current mode is preview.
+		for (const auto& It : LoadedPreviewComponents)
+		{
+			if (UPCGComponent* OriginalComponent = It.Value.Get())
+			{
+				if (OriginalComponent->GetSerializedEditingMode() == EPCGEditorDirtyMode::LoadAsPreview && OriginalComponent->GetSerializedEditingMode() == EPCGEditorDirtyMode::Preview)
+				{
+					LocalToOriginal.Add(It);
+				}
+			}
+		}
+	}
+#endif
+
+	Super::Serialize(Ar);
+
+#if WITH_EDITOR
+	if (Ar.IsSaving())
+	{
+		LocalToOriginal = MoveTemp(LocalToOriginalCopy);
+	}
+#endif
+}
+
 void APCGPartitionActor::Destroyed()
 {
 #if WITH_EDITOR
@@ -303,16 +347,17 @@ void APCGPartitionActor::AddGraphInstance(UPCGComponent* OriginalComponent)
 	if (LocalComponent)
 	{
 		// Update properties as needed and early out
+		LocalComponent->SetEditingMode(OriginalComponent->GetEditingMode(), OriginalComponent->GetSerializedEditingMode());
 		LocalComponent->SetPropertiesFromOriginal(OriginalComponent);
 		LocalComponent->MarkAsLocalComponent();
 		LocalComponent->SetGenerationGridSize(PCGGridSize);
 		return;
 	}
 
-	Modify();
+	Modify(!OriginalComponent->IsInPreviewMode());
 
 	// Create a new local component
-	LocalComponent = NewObject<UPCGComponent>(this);
+	LocalComponent = NewObject<UPCGComponent>(this, NAME_None, OriginalComponent->IsInPreviewMode() ? RF_Transient : RF_NoFlags);
 	LocalComponent->MarkAsLocalComponent();
 	LocalComponent->SetGenerationGridSize(PCGGridSize);
 
@@ -320,6 +365,8 @@ void APCGPartitionActor::AddGraphInstance(UPCGComponent* OriginalComponent)
 	OriginalToLocal.Emplace(OriginalComponent, LocalComponent);
 	LocalToOriginal.Emplace(LocalComponent, OriginalComponent);
 
+	// Implementation note: since this is a new component, we need to use the current editing mode only for both the current & serialized editing modes
+	LocalComponent->SetEditingMode(/*InEditingMode=*/OriginalComponent->GetEditingMode(), /*InSerializedEditingMode=*/OriginalComponent->GetEditingMode());
 	LocalComponent->SetPropertiesFromOriginal(OriginalComponent);
 
 	LocalComponent->RegisterComponent();
@@ -345,12 +392,13 @@ void APCGPartitionActor::RemapGraphInstance(const UPCGComponent* OldOriginalComp
 
 	if (!bIsLoading)
 	{
-		Modify();
+		Modify(!LocalComponent->IsInPreviewMode());
 	}
 
 	OriginalToLocal.Remove(OldOriginalComponent);
 	LocalToOriginal.Remove(LocalComponent);
 
+	LocalComponent->SetEditingMode(NewOriginalComponent->GetEditingMode(), NewOriginalComponent->GetSerializedEditingMode());
 	LocalComponent->SetPropertiesFromOriginal(NewOriginalComponent);
 	OriginalToLocal.Emplace(NewOriginalComponent, LocalComponent);
 	LocalToOriginal.Emplace(LocalComponent, NewOriginalComponent);
@@ -376,16 +424,26 @@ bool APCGPartitionActor::RemoveGraphInstance(UPCGComponent* OriginalComponent)
 		return false;
 	}
 
-	Modify();
+	Modify(!LocalComponent->IsInPreviewMode());
 
 	OriginalToLocal.Remove(OriginalComponent);
 	LocalToOriginal.Remove(LocalComponent);
 
 	// TODO Add option to not cleanup?
 	LocalComponent->CleanupLocalImmediate(/*bRemoveComponents=*/true);
-	LocalComponent->DestroyComponent();
 
-	return OriginalToLocal.IsEmpty();
+	// If the component is tagged as "preview-on-load" we shouldn't actually remove it, otherwise we'll cause a change on the actor.
+	if (LocalComponent->GetEditingMode() == EPCGEditorDirtyMode::Preview && LocalComponent->GetSerializedEditingMode() == EPCGEditorDirtyMode::LoadAsPreview)
+	{
+		LocalComponent->UnregisterComponent();
+		LoadedPreviewComponents.Emplace(LocalComponent, OriginalComponent);
+	}
+	else
+	{
+		LocalComponent->DestroyComponent();
+	}
+
+	return OriginalToLocal.IsEmpty() && LoadedPreviewComponents.IsEmpty();
 }
 
 void APCGPartitionActor::RemoveLocalComponent(UPCGComponent* LocalComponent)
@@ -484,4 +542,43 @@ TSet<TObjectPtr<UPCGComponent>> APCGPartitionActor::GetAllOriginalPCGComponents(
 	return ResultComponents;
 }
 
+bool APCGPartitionActor::ChangeTransientState(UPCGComponent* OriginalComponent, EPCGEditorDirtyMode EditingMode)
+{
+	check(OriginalComponent);
+
+	// First, propagate the transient state to the matching local component if any
+	if (UPCGComponent* LocalComponent = GetLocalComponent(OriginalComponent))
+	{
+		LocalComponent->SetEditingMode(/*CurrentEditingMode=*/EditingMode, /*SerializedEditingMode=*/EditingMode);
+		LocalComponent->ChangeTransientState(EditingMode);
+	}
+
+	// Then, when switching to anything but preview, we must get rid of any still-loaded components that would have been removed otherwise
+	if (EditingMode != EPCGEditorDirtyMode::Preview)
+	{
+		TArray<UPCGComponent*> LocalComponentsToDelete;
+
+		for (const auto& LoadedPreviewLocalToOriginal : LoadedPreviewComponents)
+		{
+			if (LoadedPreviewLocalToOriginal.Value == OriginalComponent)
+			{
+				LocalComponentsToDelete.Add(LoadedPreviewLocalToOriginal.Key);
+			}
+		}
+
+		// Streamlined version of RemoveGraphInstance
+		Modify(!LocalComponentsToDelete.IsEmpty());
+
+		for (UPCGComponent* LocalComponent : LocalComponentsToDelete)
+		{
+			LoadedPreviewComponents.Remove(LocalComponent);
+			LocalComponent->ChangeTransientState(EditingMode);
+			LocalComponent->CleanupLocalImmediate(/*bRemoveComponents=*/true);
+			LocalComponent->DestroyComponent();
+		}
+	}
+
+	return OriginalToLocal.IsEmpty() && LoadedPreviewComponents.IsEmpty();
+}
+
 #endif // WITH_EDITOR
diff --git a/Engine/Plugins/Experimental/PCG/Source/PCG/Private/Grid/PCGPartitionActor.h b/Engine/Plugins/Experimental/PCG/Source/PCG/Private/Grid/PCGPartitionActor.h
--- a/Engine/Plugins/Experimental/PCG/Source/PCG/Private/Grid/PCGPartitionActor.h
+++ b/Engine/Plugins/Experimental/PCG/Source/PCG/Private/Grid/PCGPartitionActor.h
@@ -4,6 +4,8 @@
 
 #include "ActorPartition/PartitionActor.h"
 
+#include "PCGCommon.h"
+
 #include "PCGPartitionActor.generated.h"
 
 namespace EEndPlayReason { enum Type : int; }
@@ -27,6 +29,7 @@ public:
 	//~Begin UObject Interface
 	virtual void PostLoad() override;
 	virtual void BeginDestroy() override;
+	virtual void Serialize(FArchive& Ar) override;
 	//~End UObject Interface
 
 	//~Begin AActor Interface
@@ -75,6 +78,9 @@ public:
 
 	/** Return a set of all the PCGComponents linked to this actor */
 	TSet<TObjectPtr<UPCGComponent>> GetAllOriginalPCGComponents() const;
+
+	/** Changes transient state for the local component matching the given original component. Returns true if PA becomes empty */
+	bool ChangeTransientState(UPCGComponent* OriginalComponent, EPCGEditorDirtyMode EditingMode);
 #endif
 
 	// TODO: Make this in-editor only; during runtime, we should keep a map of component to bounds/volume only
@@ -97,6 +103,10 @@ private:
 	UPROPERTY()
 	TMap<TObjectPtr<UPCGComponent>, TSoftObjectPtr<UPCGComponent>> LocalToOriginal;
 
+	// PCG components that are cleared when in preview-on-load mode are kept aside and put back when serializing to prevent changes
+	UPROPERTY(Transient)
+	TMap<TObjectPtr<UPCGComponent>, TSoftObjectPtr<UPCGComponent>> LoadedPreviewComponents;
+
 #if WITH_EDITORONLY_DATA
 	UPROPERTY()
 	TMap<TObjectPtr<UPCGComponent>, TWeakObjectPtr<UPCGComponent>> LocalToOriginalMap_DEPRECATED;
diff --git a/Engine/Plugins/Experimental/PCG/Source/PCG/Private/Helpers/PCGActorHelpers.cpp b/Engine/Plugins/Experimental/PCG/Source/PCG/Private/Helpers/PCGActorHelpers.cpp
--- a/Engine/Plugins/Experimental/PCG/Source/PCG/Private/Helpers/PCGActorHelpers.cpp
+++ b/Engine/Plugins/Experimental/PCG/Source/PCG/Private/Helpers/PCGActorHelpers.cpp
@@ -55,7 +55,7 @@ UPCGManagedISMComponent* UPCGActorHelpers::GetOrCreateManagedISMC(AActor* InTarg
 
 			if (UPCGManagedISMComponent* Resource = Cast<UPCGManagedISMComponent>(InResource))
 			{
-				if (Resource->GetSettingsUID() != SettingsUID)
+				if (Resource->GetSettingsUID() != SettingsUID || !Resource->CanBeUsed())
 				{
 					return;
 				}
@@ -83,7 +83,7 @@ UPCGManagedISMComponent* UPCGActorHelpers::GetOrCreateManagedISMC(AActor* InTarg
 	}
 
 	// No matching ISM component found, let's create a new one
-	InTargetActor->Modify();
+	InTargetActor->Modify(!InSourceComponent->IsInPreviewMode());
 
 	// Done as in InstancedStaticMesh.cpp
 #if WITH_EDITOR
@@ -122,7 +122,8 @@ UPCGManagedISMComponent* UPCGActorHelpers::GetOrCreateManagedISMC(AActor* InTarg
 
 	ComponentName += StaticMesh->GetName();
 
-	UInstancedStaticMeshComponent* ISMC = NewObject<UInstancedStaticMeshComponent>(InTargetActor, ComponentClass, MakeUniqueObjectName(InTargetActor, ComponentClass, FName(ComponentName)));
+	const EObjectFlags ObjectFlags = (InSourceComponent->IsInPreviewMode() ? RF_Transient : RF_NoFlags);
+	UInstancedStaticMeshComponent* ISMC = NewObject<UInstancedStaticMeshComponent>(InTargetActor, ComponentClass, MakeUniqueObjectName(InTargetActor, ComponentClass, FName(ComponentName)), ObjectFlags);
 	InParams.Descriptor.InitComponent(ISMC);
 	ISMC->NumCustomDataFloats = InParams.NumCustomDataFloats;
 
diff --git a/Engine/Plugins/Experimental/PCG/Source/PCG/Private/PCGComponent.cpp b/Engine/Plugins/Experimental/PCG/Source/PCG/Private/PCGComponent.cpp
--- a/Engine/Plugins/Experimental/PCG/Source/PCG/Private/PCGComponent.cpp
+++ b/Engine/Plugins/Experimental/PCG/Source/PCG/Private/PCGComponent.cpp
@@ -46,8 +46,10 @@
 #include UE_INLINE_GENERATED_CPP_BY_NAME(PCGComponent)
 
 #if WITH_EDITOR
+#include "Editor.h"
 #include "EditorActorFolders.h"
 #include "ScopedTransaction.h"
+#include "Editor/Transactor.h"
 #endif
 
 #define LOCTEXT_NAMESPACE "UPCGComponent"
@@ -225,7 +227,7 @@ void UPCGComponent::SetPropertiesFromOriginal(const UPCGComponent* Original)
 	// Note that while we dirty here, we won't trigger a refresh since we don't have the required context
 	if (bIsDirty)
 	{
-		Modify();
+		Modify(!IsInPreviewMode());
 		DirtyGenerated(bHasDirtyInput ? EPCGComponentDirtyFlag::Input : EPCGComponentDirtyFlag::None);
 	}
 #endif
@@ -267,7 +269,7 @@ FPCGTaskId UPCGComponent::GenerateInternal(bool bForce, EPCGComponentGenerationT
 		return InvalidPCGTaskId;
 	}
 
-	Modify();
+	Modify(!IsInPreviewMode());
 
 	CurrentGenerationTask = GetSubsystem()->ScheduleComponent(this, /*bSave=*/bForce, Dependencies);
 
@@ -528,7 +530,7 @@ FPCGTaskId UPCGComponent::CleanupInternal(bool bRemoveComponents, bool bSave, co
 
 	PCGGeneratedResourcesLogging::LogCleanupInternal(bRemoveComponents);
 
-	Modify();
+	Modify(!IsInPreviewMode());
 
 #if WITH_EDITOR
 	ExtraCapture.ResetTimers();
@@ -684,7 +686,7 @@ bool UPCGComponent::MoveResourcesToNewActor(AActor* InNewActor, bool bCreateChil
 
 	bool bHasMovedResources = false;
 
-	Modify();
+	Modify(!IsInPreviewMode());
 
 	if (bCreateChild)
 	{
@@ -750,7 +752,16 @@ void UPCGComponent::CleanupLocalImmediate(bool bRemoveComponents)
 			{
 				if (Resource)
 				{
-					Resource->Rename(nullptr, nullptr, REN_ForceNoResetLoaders | REN_DoNotDirty | REN_DontCreateRedirectors | REN_NonTransactional);
+#if WITH_EDITOR
+					if (Resource->IsMarkedTransientOnLoad())
+					{
+						LoadedPreviewResources.Add(Resource);
+					}
+					else
+#endif
+					{
+						Resource->Rename(nullptr, nullptr, REN_ForceNoResetLoaders | REN_DoNotDirty | REN_DontCreateRedirectors | REN_NonTransactional);
+					}
 				}
 
 				GeneratedResources.RemoveAtSwap(ResourceIndex);
@@ -843,7 +854,16 @@ FPCGTaskId UPCGComponent::CreateCleanupTask(bool bRemoveComponents, const TArray
 				{
 					if (Resource)
 					{
-						Resource->Rename(nullptr, nullptr, REN_ForceNoResetLoaders | REN_DoNotDirty | REN_DontCreateRedirectors | REN_NonTransactional);
+#if WITH_EDITOR
+						if (Resource->IsMarkedTransientOnLoad())
+						{
+							ThisComponent->LoadedPreviewResources.Add(Resource);
+						}
+						else
+#endif
+						{
+							Resource->Rename(nullptr, nullptr, REN_ForceNoResetLoaders | REN_DoNotDirty | REN_DontCreateRedirectors | REN_NonTransactional);
+						}
 					}
 
 					ThisComponent->GeneratedResources.RemoveAtSwap(Context->ResourceIndex);
@@ -957,7 +977,16 @@ void UPCGComponent::CleanupUnusedManagedResources()
 			{
 				if (Resource)
 				{
-					Resource->Rename(nullptr, nullptr, REN_ForceNoResetLoaders | REN_DoNotDirty | REN_DontCreateRedirectors | REN_NonTransactional);
+#if WITH_EDITOR
+					if (Resource->IsMarkedTransientOnLoad())
+					{
+						LoadedPreviewResources.Add(Resource);
+					}
+					else
+#endif
+					{
+						Resource->Rename(nullptr, nullptr, REN_ForceNoResetLoaders | REN_DoNotDirty | REN_DontCreateRedirectors | REN_NonTransactional);
+					}
 				}
 
 				GeneratedResources.RemoveAtSwap(ResourceIndex);
@@ -1072,6 +1101,28 @@ void UPCGComponent::OnComponentDestroyed(bool bDestroyingHierarchy)
 	Super::OnComponentDestroyed(bDestroyingHierarchy);
 }
 
+void UPCGComponent::Serialize(FArchive& Ar)
+{
+#if WITH_EDITOR
+	TArray<TObjectPtr<UPCGManagedResource>> GeneratedResourcesCopy;
+
+	if (Ar.IsSaving() && CurrentEditingMode == EPCGEditorDirtyMode::Preview)
+	{
+		GeneratedResourcesCopy = GeneratedResources;
+		GeneratedResources = LoadedPreviewResources;
+	}
+#endif // WITH_EDITOR
+
+	Super::Serialize(Ar);
+
+#if WITH_EDITOR
+	if (Ar.IsSaving() && CurrentEditingMode == EPCGEditorDirtyMode::Preview)
+	{
+		GeneratedResources = GeneratedResourcesCopy;
+	}
+#endif // WITH_EDITOR
+}
+
 void UPCGComponent::PostLoad()
 {
 	Super::PostLoad();
@@ -1121,6 +1172,19 @@ void UPCGComponent::PostLoad()
 	}
 
 	SetupCallbacksOnCreation();
+
+	CurrentEditingMode = SerializedEditingMode;
+
+	if (CurrentEditingMode == EPCGEditorDirtyMode::Preview)
+	{
+		bGenerated = false;
+	}
+	else if (CurrentEditingMode == EPCGEditorDirtyMode::LoadAsPreview && !PCGHelpers::IsRuntimeOrPIE())
+	{
+		CurrentEditingMode = EPCGEditorDirtyMode::Preview;
+		MarkResourcesAsTransientOnLoad();
+		bDirtyGenerated = true;
+	}
 #endif
 }
 
@@ -1282,8 +1346,6 @@ void UPCGComponent::RefreshAfterGraphChanged(UPCGGraphInterface* InGraph, bool b
 #if WITH_EDITOR
 void UPCGComponent::PostEditChangeProperty(FPropertyChangedEvent& PropertyChangedEvent)
 {
-	Super::PostEditChangeProperty(PropertyChangedEvent);
-
 	if (!PropertyChangedEvent.Property || !IsValid(this))
 	{
 		return;
@@ -1291,6 +1353,23 @@ void UPCGComponent::PostEditChangeProperty(FPropertyChangedEvent& PropertyChange
 
 	const FName PropName = PropertyChangedEvent.Property->GetFName();
 
+	// Implementation note:
+	// Since the current editing mode is a transient variable, if we do not do this transition here before going in the Super call,
+	//  we can end up in a situation where BP actors are reconstructed (... this component included ...) which makes this fall into the !IsValid case just after
+	if (PropName == GET_MEMBER_NAME_CHECKED(UPCGComponent, CurrentEditingMode))
+	{
+		// When affecting the editing mode from the user's point of view, we need to change both the current & serialized values
+		SetEditingMode(CurrentEditingMode, CurrentEditingMode);
+		ChangeTransientState(CurrentEditingMode);
+	}
+
+	Super::PostEditChangeProperty(PropertyChangedEvent);
+
+	if (!IsValid(this))
+	{
+		return;
+	}
+
 	// Important note: all property changes already go through the OnObjectPropertyChanged, and will be dirtied here.
 	// So where only a Refresh is needed, it goes through the "capture all" else case.
 	if (PropName == GET_MEMBER_NAME_CHECKED(UPCGComponent, bIsComponentPartitioned))
@@ -2374,6 +2453,122 @@ void UPCGComponent::GetManagedResources(TArray<TObjectPtr<UPCGManagedResource>>&
 	Resources = GeneratedResources;
 }
 
+void UPCGComponent::SetEditingMode(EPCGEditorDirtyMode InEditingMode, EPCGEditorDirtyMode InSerializedEditingMode)
+{
+	CurrentEditingMode = InEditingMode;
+	SerializedEditingMode = InSerializedEditingMode;
+}
+
+#if WITH_EDITOR
+bool UPCGComponent::DeletePreviewResources()
+{
+	bool bDeletedSomething = false;
+
+	TSet<TSoftObjectPtr<AActor>> ActorsToDelete;
+	// Make sure to release fully the resources that were loaded
+	for (TObjectPtr<UPCGManagedResource> ResourceToRelease : LoadedPreviewResources)
+	{
+		// Changing the transient state will clear the "marked transient on load" flag
+		ResourceToRelease->ChangeTransientState(/*bNowTransient=*/false);
+		ResourceToRelease->Release(/*bHardRelease=*/true, ActorsToDelete);
+		bDeletedSomething = true;
+	}
+
+	LoadedPreviewResources.Empty();
+
+	if (!ActorsToDelete.IsEmpty())
+	{
+		UPCGActorHelpers::DeleteActors(GetWorld(), ActorsToDelete.Array());
+	}
+
+	return bDeletedSomething;
+}
+
+void UPCGComponent::MarkResourcesAsTransientOnLoad()
+{
+	for (TObjectPtr<UPCGManagedResource>& GeneratedResource : GeneratedResources)
+	{
+		if (GeneratedResource)
+		{
+			GeneratedResource->MarkTransientOnLoad();
+		}
+	}
+}
+
+void UPCGComponent::ChangeTransientState(EPCGEditorDirtyMode NewEditingMode)
+{
+	bool bShouldMarkDirty = false;
+
+	// Affect all resources
+	{
+		FScopeLock ResourcesLock(&GeneratedResourcesLock);
+		check(!GeneratedResourcesInaccessible);
+
+		for (TObjectPtr<UPCGManagedResource>& GeneratedResource : GeneratedResources)
+		{
+			if (GeneratedResource)
+			{
+				GeneratedResource->ChangeTransientState(NewEditingMode == EPCGEditorDirtyMode::Preview);
+				bShouldMarkDirty = true;
+			}
+		}
+
+		// If switching from preview mode to normal or preview-on-load,
+		// we must materialize any kind of change we've done on the packages that had a different behavior on load (e.g. actor packages)
+		if (NewEditingMode != EPCGEditorDirtyMode::Preview)
+		{
+			bShouldMarkDirty |= DeletePreviewResources();
+		}
+	}
+
+	if (IsLocalComponent())
+	{
+		bShouldMarkDirty = true;
+
+		if (NewEditingMode == EPCGEditorDirtyMode::Preview)
+		{
+			SetFlags(RF_Transient);
+		}
+		else
+		{
+			ClearFlags(RF_Transient);
+		}
+
+		ForEachObjectWithOuter(this, [NewEditingMode](UObject* Object)
+		{
+			if (NewEditingMode == EPCGEditorDirtyMode::Preview)
+			{
+				Object->SetFlags(RF_Transient);
+			}
+			else
+			{
+				Object->ClearFlags(RF_Transient);
+			}
+		});
+	}
+
+	if (bShouldMarkDirty)
+	{
+		MarkPackageDirty();
+	}
+
+	// Un-transient PAs if needed and propagate the call
+	if (IsPartitioned())
+	{
+		if (UPCGSubsystem* Subsystem = GetSubsystem())
+		{
+			Subsystem->PropagateEditingModeToLocalComponents(this, NewEditingMode);
+		}
+	}
+
+	// Changing the transient state can and will play with packages and is not meant to be undoable
+	if (GEditor->Trans)
+	{
+		GEditor->Trans->Reset(LOCTEXT("ChangeEditingMode", "Changing Editing Mode"));
+	}
+}
+#endif // WITH_EDITOR
+
 FPCGComponentInstanceData::FPCGComponentInstanceData(const UPCGComponent* InSourceComponent)
 	: FActorComponentInstanceData(InSourceComponent)
 	, SourceComponent(InSourceComponent)
@@ -2381,6 +2576,9 @@ FPCGComponentInstanceData::FPCGComponentInstanceData(const UPCGComponent* InSour
 	if (SourceComponent)
 	{
 		SourceComponent->GetManagedResources(GeneratedResources);
+#if WITH_EDITOR
+		LoadedPreviewResources = SourceComponent->LoadedPreviewResources;
+#endif
 	}
 }
 
@@ -2421,6 +2619,8 @@ void FPCGComponentInstanceData::ApplyToComponent(UActorComponent* Component, con
 
 			// While this is serialized, it is not properly copied because it is not visible. This is needed otherwise a refresh can retoggle from not generated to generated
 			PCGComponent->bForceGenerateOnBPAddedToWorld = SourceComponent->bForceGenerateOnBPAddedToWorld;
+
+			PCGComponent->CurrentEditingMode = SourceComponent->CurrentEditingMode;
 #endif // WITH_EDITOR
 
 			// Non-critical but should be done: transient data, tracked actors cache, landscape tracking
@@ -2445,6 +2645,22 @@ void FPCGComponentInstanceData::ApplyToComponent(UActorComponent* Component, con
 		}
 
 #if WITH_EDITOR
+		TArray<TObjectPtr<UPCGManagedResource>> DuplicateLoadedPreviewResources;
+		for (const TObjectPtr<UPCGManagedResource>& Resource : LoadedPreviewResources)
+		{
+			if (Resource)
+			{
+				UPCGManagedResource* DuplicatedResource = CastChecked<UPCGManagedResource>(StaticDuplicateObject(Resource, PCGComponent, FName()));
+				DuplicatedResource->PostApplyToComponent();
+				DuplicateLoadedPreviewResources.Add(DuplicatedResource);
+			}
+		}
+
+		if (DuplicateLoadedPreviewResources.Num() > 0)
+		{
+			PCGComponent->LoadedPreviewResources = DuplicateLoadedPreviewResources;
+		}
+
 		// Reconnect callbacks
 		if (PCGComponent->GraphInstance)
 		{
diff --git a/Engine/Plugins/Experimental/PCG/Source/PCG/Private/PCGManagedResource.cpp b/Engine/Plugins/Experimental/PCG/Source/PCG/Private/PCGManagedResource.cpp
--- a/Engine/Plugins/Experimental/PCG/Source/PCG/Private/PCGManagedResource.cpp
+++ b/Engine/Plugins/Experimental/PCG/Source/PCG/Private/PCGManagedResource.cpp
@@ -7,8 +7,13 @@
 #include "Helpers/PCGHelpers.h"
 
 #include "Components/InstancedStaticMeshComponent.h"
+#include "Engine/Level.h"
 #include "Utils/PCGGeneratedResourcesLogging.h"
 
+#if WITH_EDITOR
+#include "ObjectTools.h"
+#endif
+
 #include UE_INLINE_GENERATED_CPP_BY_NAME(PCGManagedResource)
 
 static TAutoConsoleVariable<bool> CVarForceReleaseResourcesOnGenerate(
@@ -45,11 +50,28 @@ bool UPCGManagedResource::ReleaseIfUnused(TSet<TSoftObjectPtr<AActor>>& OutActor
 	return false;
 }
 
+bool UPCGManagedResource::CanBeUsed() const
+{
+#if WITH_EDITOR
+	return !bMarkedTransientOnLoad;
+#else
+	return true;
+#endif
+}
+
 bool UPCGManagedResource::DebugForcePurgeAllResourcesOnGenerate()
 {
 	return CVarForceReleaseResourcesOnGenerate.GetValueOnAnyThread();
 }
 
+#if WITH_EDITOR
+void UPCGManagedResource::ChangeTransientState(bool /*bNowTransient*/)
+{
+	// Any change in the transient state resets the transient state that was set on load, regardless of the bNowTransient flag
+	bMarkedTransientOnLoad = false;
+}
+#endif // WITH_EDITOR
+
 void UPCGManagedActors::PostEditImport()
 {
 	// In this case, the managed actors won't be copied along the actor/component,
@@ -83,7 +105,26 @@ bool UPCGManagedActors::Release(bool bHardRelease, TSet<TSoftObjectPtr<AActor>>&
 		return false;
 	}
 
-	OutActorsToDelete.Append(GeneratedActors);
+#if WITH_EDITOR
+	if (bMarkedTransientOnLoad)
+	{
+		// Here, instead of adding the actors to be deleted (which has the side effect of potentially emptying the package, which leads to its deletion, we will hide the actors instead
+		for (TSoftObjectPtr<AActor> GeneratedActor : GeneratedActors)
+		{
+			// In rare cases where the actor wouldn't be currently loaded (in WP) we must make sure it is done here
+			if (AActor* Actor = GeneratedActor.LoadSynchronous())
+			{
+				Actor->SetIsTemporarilyHiddenInEditor(true);
+				Actor->SetHidden(true);
+				Actor->SetActorEnableCollision(false);
+			}
+		}
+	}
+	else
+#endif // WITH_EDITOR
+	{
+		OutActorsToDelete.Append(GeneratedActors);
+	}
 
 	PCGGeneratedResourcesLogging::LogManagedActorsHardRelease(GeneratedActors);
 
@@ -107,7 +148,13 @@ bool UPCGManagedActors::Release(bool bHardRelease, TSet<TSoftObjectPtr<AActor>>&
 		}
 	}
 
-	GeneratedActors.Reset();
+#if WITH_EDITOR
+	if (!bMarkedTransientOnLoad)
+#endif
+	{
+		GeneratedActors.Reset();
+	}
+	
 	return true;
 }
 
@@ -162,6 +209,66 @@ void UPCGManagedActors::MarkAsReused()
 	}
 }
 
+#if WITH_EDITOR
+void UPCGManagedActors::ChangeTransientState(bool bNowTransient)
+{
+	TSet<UPackage*> PackagesToCleanup;
+
+	for (TSoftObjectPtr<AActor> GeneratedActor : GeneratedActors)
+	{
+		// Make sure to load if needed because we need to affect the actors regardless of the current WP state
+		if (GeneratedActor.LoadSynchronous() != nullptr)
+		{
+			// If this is changed during loading, we shouldn't set the actor to be transient if we're using external packages,
+			// otherwise the package will get deleted spuriously.
+			if (bNowTransient)
+			{
+				if (!GeneratedActor->IsPackageExternal())
+				{
+					GeneratedActor->SetFlags(RF_Transient);
+				}
+			}
+			else
+			{
+				GeneratedActor->ClearFlags(RF_Transient);
+			}
+
+			if (GeneratedActor->GetLevel() && GeneratedActor->GetLevel()->IsUsingExternalActors())
+			{
+				if (bNowTransient)
+				{
+					PackagesToCleanup.Add(GeneratedActor->GetExternalPackage());
+				}
+
+				GeneratedActor->SetPackageExternal(/*bExternal=*/!bNowTransient, /*bShouldDirty=*/false);
+			}
+
+			if (!bNowTransient || GeneratedActor->IsPackageExternal())
+			{
+				ForEachObjectWithOuter(GeneratedActor.Get(), [bNowTransient](UObject* Object)
+				{
+					if (bNowTransient)
+					{
+						Object->SetFlags(RF_Transient);
+					}
+					else
+					{
+						Object->ClearFlags(RF_Transient);
+					}
+				});
+			}
+		}
+	}
+
+	if (!PackagesToCleanup.IsEmpty())
+	{
+		ObjectTools::CleanupAfterSuccessfulDelete(PackagesToCleanup.Array(), /*bPerformanceReferenceCheck=*/true);
+	}
+
+	Super::ChangeTransientState(bNowTransient);
+}
+#endif // WITH_EDITOR
+
 void UPCGManagedComponent::PostEditImport()
 {
 	Super::PostEditImport();
@@ -201,6 +308,16 @@ void UPCGManagedComponent::PostEditImport()
 	}
 }
 
+#if WITH_EDITOR
+void UPCGManagedComponent::HideComponent()
+{
+	if (GeneratedComponent.IsValid())
+	{
+		GeneratedComponent->UnregisterComponent();
+	}
+}
+#endif // WITH_EDITOR
+
 bool UPCGManagedComponent::Release(bool bHardRelease, TSet<TSoftObjectPtr<AActor>>& /*OutActorsToDelete*/)
 {
 	TRACE_CPUPROFILER_EVENT_SCOPE(UPCGManagedComponent::Release);
@@ -210,16 +327,26 @@ bool UPCGManagedComponent::Release(bool bHardRelease, TSet<TSoftObjectPtr<AActor
 
 	if (GeneratedComponent.IsValid())
 	{
-		if (bDeleteComponent)
+#if WITH_EDITOR
+		if (bMarkedTransientOnLoad)
 		{
-			GeneratedComponent->DestroyComponent();
-			ForgetComponent();
+			HideComponent();
+			bIsMarkedUnused = true;
 		}
 		else
+#endif // WITH_EDITOR
 		{
-			// We can only mark it unused if we can reset the component.
-			bIsMarkedUnused = true;
-			GeneratedComponent->ComponentTags.Add(PCGHelpers::MarkedForCleanupPCGTag);
+			if (bDeleteComponent)
+			{
+				GeneratedComponent->DestroyComponent();
+				ForgetComponent();
+			}
+			else
+			{
+				// We can only mark it unused if we can reset the component.
+				bIsMarkedUnused = true;
+				GeneratedComponent->ComponentTags.Add(PCGHelpers::MarkedForCleanupPCGTag);
+			}
 		}
 	}
 	else
@@ -318,6 +445,44 @@ void UPCGManagedComponent::MarkAsReused()
 	}
 }
 
+#if WITH_EDITOR
+void UPCGManagedComponent::ChangeTransientState(bool bNowTransient)
+{
+	if (GeneratedComponent.Get())
+	{
+		const bool bWasTransient = GeneratedComponent->HasAnyFlags(RF_Transient);
+
+		if (bWasTransient != bNowTransient)
+		{
+			if (bNowTransient)
+			{
+				GeneratedComponent->SetFlags(RF_Transient);
+			}
+			else
+			{
+				GeneratedComponent->ClearFlags(RF_Transient);
+			}
+
+			ForEachObjectWithOuter(GeneratedComponent.Get(), [bNowTransient](UObject* Object)
+			{
+				if (bNowTransient)
+				{
+					Object->SetFlags(RF_Transient);
+				}
+				else
+				{
+					Object->ClearFlags(RF_Transient);
+				}
+			});
+
+			GeneratedComponent->MarkPackageDirty(); // should dirty actor this component is attached to
+		}
+	}
+
+	Super::ChangeTransientState(bNowTransient);
+}
+#endif // WITH_EDITOR
+
 void UPCGManagedISMComponent::PostLoad()
 {
 	Super::PostLoad();
diff --git a/Engine/Plugins/Experimental/PCG/Source/PCG/Private/PCGSubsystem.cpp b/Engine/Plugins/Experimental/PCG/Source/PCG/Private/PCGSubsystem.cpp
--- a/Engine/Plugins/Experimental/PCG/Source/PCG/Private/PCGSubsystem.cpp
+++ b/Engine/Plugins/Experimental/PCG/Source/PCG/Private/PCGSubsystem.cpp
@@ -1338,6 +1338,23 @@ void UPCGSubsystem::NotifyGraphChanged(UPCGGraph* InGraph)
 	}
 }
 
+void UPCGSubsystem::PropagateEditingModeToLocalComponents(UPCGComponent* InComponent, EPCGEditorDirtyMode EditingMode)
+{
+	if (ensure(InComponent && InComponent->IsPartitioned()))
+	{
+		FBox Bounds = ActorAndComponentMapping.PartitionedOctree.GetBounds(InComponent);
+		if (!Bounds.IsValid)
+		{
+			return;
+		}
+
+		ActorAndComponentMapping.ForAllIntersectingPartitionActors(Bounds, [InComponent, EditingMode](APCGPartitionActor* Actor)
+		{
+			Actor->ChangeTransientState(InComponent, EditingMode);
+		});
+	}
+}
+
 void UPCGSubsystem::CleanFromCache(const IPCGElement* InElement, const UPCGSettings* InSettings/*= nullptr*/)
 {
 	if (GraphExecutor)
diff --git a/Engine/Plugins/Experimental/PCG/Source/PCG/Public/PCGCommon.h b/Engine/Plugins/Experimental/PCG/Source/PCG/Public/PCGCommon.h
--- a/Engine/Plugins/Experimental/PCG/Source/PCG/Public/PCGCommon.h
+++ b/Engine/Plugins/Experimental/PCG/Source/PCG/Public/PCGCommon.h
@@ -182,4 +182,12 @@ enum class EPCGAttachOptions : uint32
 	NotAttached UMETA(Tooltip="Actor will not be attached to the target actor nor placed in an actor folder"),
 	Attached UMETA(Tooltip="Actor will be attached to the target actor in the given node"),
 	InFolder UMETA(Tooltip="Actor will be placed in an actor folder containing the name of the target actor.")
+};
+
+UENUM()
+enum class EPCGEditorDirtyMode : uint8
+{
+	Normal UMETA(Tooltip="Normal editing mode where generation changes (generation, cleanup) dirty the component and its resources."),
+	Preview UMETA(Tooltip="Editing mode where generation changes (generation, cleanup, resources) on the component will not trigger any dirty state, but will also not save any of the generated resources."),
+	LoadAsPreview UMETA(Tooltip="Acts as the normal editing mode until the next load of the component, at which state it acts as-if-transient, namely that any further generation changes will not dirty the component.")
 };
\ No newline at end of file
diff --git a/Engine/Plugins/Experimental/PCG/Source/PCG/Public/PCGComponent.h b/Engine/Plugins/Experimental/PCG/Source/PCG/Public/PCGComponent.h
--- a/Engine/Plugins/Experimental/PCG/Source/PCG/Public/PCGComponent.h
+++ b/Engine/Plugins/Experimental/PCG/Source/PCG/Public/PCGComponent.h
@@ -2,6 +2,7 @@
 
 #pragma once
 
+#include "PCGCommon.h"
 #include "PCGNode.h"
 #include "PCGSettings.h"
 #include "Graph/PCGStackContext.h"
@@ -76,6 +77,7 @@ class PCG_API UPCGComponent : public UActorComponent
 
 public:
 	/** ~Begin UObject interface */
+	virtual void Serialize(FArchive& Ar) override;
 	virtual void PostLoad() override;
 	virtual void PostInitProperties() override;
 	virtual void BeginDestroy() override;
@@ -262,6 +264,9 @@ public:
 
 	/** Know if we need to force a generation, in case of BP added to the world in editor */
 	bool ShouldGenerateBPPCGAddedToWorld() const;
+
+	/** Changes the transient state (preview, normal, load on preview) - public only because it needs to be accessed by APCGPartitionActor */
+	void ChangeTransientState(EPCGEditorDirtyMode NewEditingMode);
 #endif
 
 	/** Utility function (mostly for tests) to properly set the value of bIsComponentPartitioned.
@@ -279,6 +284,19 @@ public:
 	/** Updates internal properties from other component, dirties as required but does not trigger Refresh */
 	void SetPropertiesFromOriginal(const UPCGComponent* Original);
 
+	/** Returns whether the component (or resources) should be marked as dirty following interaction/refresh based on the current editing mode */
+	bool IsInPreviewMode() const { return CurrentEditingMode == EPCGEditorDirtyMode::Preview; }
+
+	UFUNCTION(BlueprintCallable, Category="PCG|Advanced")
+	void SetEditingMode(EPCGEditorDirtyMode InEditingMode, EPCGEditorDirtyMode InSerializedEditingMode);
+
+	/** Returns the current editing mode */
+	UFUNCTION(BlueprintCallable, Category="PCG|Advanced")
+	EPCGEditorDirtyMode GetEditingMode() const { return CurrentEditingMode; }
+
+	UFUNCTION(BlueprintCallable, Category = "PCG|Advanced")
+	EPCGEditorDirtyMode GetSerializedEditingMode() const { return SerializedEditingMode; }
+
 	UPCGSubsystem* GetSubsystem() const;
 
 	FBox GetGridBounds() const;
@@ -297,6 +315,13 @@ protected:
 	UPROPERTY(Transient, VisibleAnywhere, AdvancedDisplay, Category = Properties)
 	uint32 GenerationGridSize = PCGHiGenGrid::UnboundedGridSize();
 
+	// Current editing mode that depends on the serialized editing mode and loading
+	UPROPERTY(Transient, EditAnywhere, AdvancedDisplay, Category = Properties, meta = (DisplayName = "Editing Mode", EditCondition = "!bIsComponentLocal"))
+	EPCGEditorDirtyMode CurrentEditingMode = EPCGEditorDirtyMode::Normal;
+
+	UPROPERTY(VisibleAnywhere, AdvancedDisplay, Category = Properties, meta = (NoResetToDefault))
+	EPCGEditorDirtyMode SerializedEditingMode = EPCGEditorDirtyMode::Normal;
+
 #if WITH_EDITORONLY_DATA
 	UPROPERTY()
 	TObjectPtr<UPCGGraph> Graph_DEPRECATED;
@@ -362,6 +387,9 @@ private:
 	TArray<const UPCGSettings*> GatherSettingsTrackingActor(const AActor* InActor, const bool bIntersect, const TSet<FName>& InRemovedTags, const UObject* InOriginatingChangeObject) const;
 
 	bool ShouldTrackLandscape() const;
+
+	void MarkResourcesAsTransientOnLoad();
+	bool DeletePreviewResources();
 #endif
 
 	FBox GetGridBounds(const AActor* InActor) const;
@@ -389,6 +417,11 @@ private:
 	UPROPERTY()
 	TArray<TObjectPtr<UPCGManagedResource>> GeneratedResources;
 
+#if WITH_EDITORONLY_DATA
+	UPROPERTY(Transient)
+	TArray<TObjectPtr<UPCGManagedResource>> LoadedPreviewResources;
+#endif
+
 	// When doing a cleanup, locking resource modification. Used as sentinel.
 	bool GeneratedResourcesInaccessible = false;
 
@@ -459,6 +492,11 @@ protected:
 	UPROPERTY()
 	TArray<TObjectPtr<UPCGManagedResource>> GeneratedResources;
 
+#if WITH_EDITORONLY_DATA
+	UPROPERTY()
+	TArray<TObjectPtr<UPCGManagedResource>> LoadedPreviewResources;
+#endif
+
 	UPROPERTY()
 	TObjectPtr<const UPCGComponent> SourceComponent;
 };
diff --git a/Engine/Plugins/Experimental/PCG/Source/PCG/Public/PCGManagedResource.h b/Engine/Plugins/Experimental/PCG/Source/PCG/Public/PCGManagedResource.h
--- a/Engine/Plugins/Experimental/PCG/Source/PCG/Public/PCGManagedResource.h
+++ b/Engine/Plugins/Experimental/PCG/Source/PCG/Public/PCGManagedResource.h
@@ -2,6 +2,7 @@
 
 #pragma once
 
+#include "PCGCommon.h"
 #include "PCGCrc.h"
 
 #include "ISMPartition/ISMComponentDescriptor.h"
@@ -31,7 +32,12 @@ public:
 	/** Releases resource if empty or unused. Returns true if the resource can be removed from the PCG component */
 	virtual bool ReleaseIfUnused(TSet<TSoftObjectPtr<AActor>>& OutActorsToDelete);
 
-	virtual void MarkAsUsed() { bIsMarkedUnused = false; }
+	/** Returns whether a resource can be used - generally true except for resources marked as transient (from loading) */
+	virtual bool CanBeUsed() const;
+
+	/** Marks the resources as being kept and changed through generation */
+	virtual void MarkAsUsed() { ensure(CanBeUsed()); bIsMarkedUnused = false; }
+	/** Marks the resource as being reused as-is during the generation */
 	// Ensure may fire if multiple executions of the graph are happening in parallel/overlapping
 	virtual void MarkAsReused() { ensure(bIsMarkedUnused); bIsMarkedUnused = false; }
 	bool IsMarkedUnused() const { return bIsMarkedUnused; }
@@ -45,12 +51,28 @@ public:
 	const FPCGCrc& GetCrc() const { return Crc; }
 	void SetCrc(const FPCGCrc& InCrc) { Crc = InCrc; }
 
+#if WITH_EDITOR
+	virtual void ChangeTransientState(bool bNowTransient);
+	virtual void MarkTransientOnLoad() { bMarkedTransientOnLoad = true; }
+
+	bool IsMarkedTransientOnLoad() const { return bMarkedTransientOnLoad; }
+#endif // WITH_EDITOR
+
 protected:
 	UPROPERTY(VisibleAnywhere, Category = GeneratedData)
 	FPCGCrc Crc;
 
 	UPROPERTY(Transient, VisibleAnywhere, Category = GeneratedData)
 	bool bIsMarkedUnused = false;
+
+#if WITH_EDITORONLY_DATA
+	// Resources on a Load-as-preview component are marked as 'transient on load'; these resources must not be affected in any
+	//  permanent way in order to make sure they are not serialized in a different state if their outer is saved.
+	// These resources will generally have a different Release path, and will be managed differently from the PCG component as well.
+	// Note that this flag will be reset if there is a transient state change originating from the component, which might trigger resource deletion, flags change, etc.
+	UPROPERTY(VisibleAnywhere, Category = GeneratedData, meta = (NoResetToDefault))
+	bool bMarkedTransientOnLoad = false;
+#endif
 };
 
 UCLASS(BlueprintType)
@@ -70,6 +92,10 @@ public:
 	virtual bool MoveResourceToNewActor(AActor* NewActor) override;
 	virtual void MarkAsUsed() override;
 	virtual void MarkAsReused() override;
+
+#if WITH_EDITOR
+	virtual void ChangeTransientState(bool bNowTransient) override;
+#endif
 	//~End UPCGManagedResource interface
 
 	UPROPERTY(BlueprintReadWrite, EditAnywhere, Category = GeneratedData)
@@ -90,6 +116,12 @@ public:
 	virtual bool Release(bool bHardRelease, TSet<TSoftObjectPtr<AActor>>& OutActorsToDelete) override;
 	virtual bool ReleaseIfUnused(TSet<TSoftObjectPtr<AActor>>& OutActorsToDelete) override;
 	virtual bool MoveResourceToNewActor(AActor* NewActor, const AActor* ExpectedPreviousOwner) override;
+
+#if WITH_EDITOR
+	virtual void ChangeTransientState(bool bNowTransient) override;
+	/** Hides the content of the component in a transient way (such as unregistering) */
+	virtual void HideComponent();
+#endif
 	//~End UPCGManagedResource interface
 
 	virtual void ResetComponent() { check(0); }
diff --git a/Engine/Plugins/Experimental/PCG/Source/PCG/Public/PCGSubsystem.h b/Engine/Plugins/Experimental/PCG/Source/PCG/Public/PCGSubsystem.h
--- a/Engine/Plugins/Experimental/PCG/Source/PCG/Public/PCGSubsystem.h
+++ b/Engine/Plugins/Experimental/PCG/Source/PCG/Public/PCGSubsystem.h
@@ -181,6 +181,9 @@ public:
 	/** Update the tracking on a given component. */
 	void UpdateComponentTracking(UPCGComponent* InComponent, bool bShouldDirtyActors, const TArray<FPCGActorSelectionKey>* OptionalChangedKeys = nullptr) { ActorAndComponentMapping.UpdateTracking(InComponent, bShouldDirtyActors, OptionalChangedKeys); }
 
+	/** Propagates transient state change from an original component to the relevant partition actors */
+	void PropagateEditingModeToLocalComponents(UPCGComponent* InOriginalComponent, EPCGEditorDirtyMode EditingMode);
+
 	/** Cleans up the graph cache on an element basis. InSettings is used for debugging and is optional. */
 	void CleanFromCache(const IPCGElement* InElement, const UPCGSettings* InSettings = nullptr);
 
