Adjusted variable naming for clarity

#rnx

diff --git a/Engine/Plugins/Experimental/PCG/Source/PCG/Private/PCGComponent.cpp b/Engine/Plugins/Experimental/PCG/Source/PCG/Private/PCGComponent.cpp
--- a/Engine/Plugins/Experimental/PCG/Source/PCG/Private/PCGComponent.cpp
+++ b/Engine/Plugins/Experimental/PCG/Source/PCG/Private/PCGComponent.cpp
@@ -2462,7 +2462,7 @@ void UPCGComponent::SetEditingMode(EPCGEditorDirtyMode InEditingMode, EPCGEditor
 #if WITH_EDITOR
 bool UPCGComponent::DeletePreviewResources()
 {
-	bool bDeletedSomething = false;
+	bool bResourceWasReleased = false;
 
 	TSet<TSoftObjectPtr<AActor>> ActorsToDelete;
 	// Make sure to release fully the resources that were loaded
@@ -2471,7 +2471,7 @@ bool UPCGComponent::DeletePreviewResources()
 		// Changing the transient state will clear the "marked transient on load" flag
 		ResourceToRelease->ChangeTransientState(/*bNowTransient=*/false);
 		ResourceToRelease->Release(/*bHardRelease=*/true, ActorsToDelete);
-		bDeletedSomething = true;
+		bResourceWasReleased = true;
 	}
 
 	LoadedPreviewResources.Empty();
@@ -2481,7 +2481,7 @@ bool UPCGComponent::DeletePreviewResources()
 		UPCGActorHelpers::DeleteActors(GetWorld(), ActorsToDelete.Array());
 	}
 
-	return bDeletedSomething;
+	return bResourceWasReleased;
 }
 
 void UPCGComponent::MarkResourcesAsTransientOnLoad()
