Made the 'keep zero density points' option available in non-editor builds and non-transient where it makes sense (in data operations, not filtering elements)

diff --git a/Engine/Plugins/Experimental/PCG/Source/PCG/Private/Data/PCGDifferenceData.cpp b/Engine/Plugins/Experimental/PCG/Source/PCG/Private/Data/PCGDifferenceData.cpp
--- a/Engine/Plugins/Experimental/PCG/Source/PCG/Private/Data/PCGDifferenceData.cpp
+++ b/Engine/Plugins/Experimental/PCG/Source/PCG/Private/Data/PCGDifferenceData.cpp
@@ -273,11 +273,7 @@ const UPCGPointData* UPCGDifferenceData::CreatePointData(FPCGContext* Context) c
 				OutMetadata->MergePointAttributesSubset(Point, SourceMetadata, SourceMetadata, PointFromDiff, TempDiffMetadata, TempDiffMetadata, OutPoint, EPCGMetadataOp::Sub);
 			}
 
-#if WITH_EDITOR
 			return OutPoint.Density > 0 || bKeepZeroDensityPoints;
-#else
-			return OutPoint.Density > 0;
-#endif
 		}
 		else
 		{
diff --git a/Engine/Plugins/Experimental/PCG/Source/PCG/Private/Data/PCGIntersectionData.cpp b/Engine/Plugins/Experimental/PCG/Source/PCG/Private/Data/PCGIntersectionData.cpp
--- a/Engine/Plugins/Experimental/PCG/Source/PCG/Private/Data/PCGIntersectionData.cpp
+++ b/Engine/Plugins/Experimental/PCG/Source/PCG/Private/Data/PCGIntersectionData.cpp
@@ -237,13 +237,19 @@ UPCGPointData* UPCGIntersectionData::CreateAndFilterPointData(FPCGContext* Conte
 		const FPCGPoint& Point = SourcePoints[Index];
 
 		FPCGPoint PointFromY;
-#if WITH_EDITOR
-		if (!Y->SamplePoint(Point.Transform, Point.GetLocalBounds(), PointFromY, TempYMetadata) && !bKeepZeroDensityPoints)
-#else
 		if (!Y->SamplePoint(Point.Transform, Point.GetLocalBounds(), PointFromY, TempYMetadata))
-#endif
 		{
-			return false;
+			if (!bKeepZeroDensityPoints)
+			{
+				return false;
+			}
+			else
+			{
+				// Point is rejected, mark its density to zero
+				PointFromY.Density = 0;
+				PointFromY.MetadataEntry = PCGInvalidEntryKey;
+				PointFromY.Color = FVector4::One();
+			}
 		}
 
 		OutPoint = Point;
diff --git a/Engine/Plugins/Experimental/PCG/Source/PCG/Private/Data/PCGProjectionData.cpp b/Engine/Plugins/Experimental/PCG/Source/PCG/Private/Data/PCGProjectionData.cpp
--- a/Engine/Plugins/Experimental/PCG/Source/PCG/Private/Data/PCGProjectionData.cpp
+++ b/Engine/Plugins/Experimental/PCG/Source/PCG/Private/Data/PCGProjectionData.cpp
@@ -234,13 +234,20 @@ const UPCGPointData* UPCGProjectionData::CreatePointData(FPCGContext* Context) c
 		const FPCGPoint& SourcePoint = SourcePoints[Index];
 
 		FPCGPoint PointFromTarget;
-#if WITH_EDITOR
-		if (!Target->ProjectPoint(SourcePoint.Transform, SourcePoint.GetLocalBounds(), ProjectionParams, PointFromTarget, TempTargetMetadata) && !bKeepZeroDensityPoints)
-#else
 		if (!Target->ProjectPoint(SourcePoint.Transform, SourcePoint.GetLocalBounds(), ProjectionParams, PointFromTarget, TempTargetMetadata))
-#endif
 		{
-			return false;
+			if (!bKeepZeroDensityPoints)
+			{
+				return false;
+			}
+			else
+			{
+				// Point is rejected, mark its density to zero, put it in a state where we won't affect the output point
+				PointFromTarget.Transform = SourcePoint.Transform;
+				PointFromTarget.Color = FVector4::One();
+				PointFromTarget.Density = 0;
+				PointFromTarget.MetadataEntry = PCGInvalidEntryKey;
+			}
 		}
 
 		// Merge points into a single point
diff --git a/Engine/Plugins/Experimental/PCG/Source/PCG/Private/Data/PCGTextureData.cpp b/Engine/Plugins/Experimental/PCG/Source/PCG/Private/Data/PCGTextureData.cpp
--- a/Engine/Plugins/Experimental/PCG/Source/PCG/Private/Data/PCGTextureData.cpp
+++ b/Engine/Plugins/Experimental/PCG/Source/PCG/Private/Data/PCGTextureData.cpp
@@ -220,11 +220,7 @@ const UPCGPointData* UPCGBaseTextureData::CreatePointData(FPCGContext* Context)
 		if (PCGTextureSampling::Sample<FLinearColor>(LocalCoordinate, Surface, this, Width, Height, Color, [this](int32 Index) { return ColorData[Index]; }))
 		{
 			const float Density = ((DensityFunction == EPCGTextureDensityFunction::Ignore) ? 1.0f : PCGTextureSampling::SampleFloatChannel(Color, ColorChannel));
-#if WITH_EDITOR
 			if (Density > 0 || bKeepZeroDensityPoints)
-#else
-			if (Density > 0)
-#endif
 			{
 				FVector LocalPosition(LocalCoordinate, 0);
 				OutPoint = FPCGPoint(FTransform(Transform.TransformPosition(LocalPosition)),
diff --git a/Engine/Plugins/Experimental/PCG/Source/PCG/Private/Elements/PCGDifferenceElement.cpp b/Engine/Plugins/Experimental/PCG/Source/PCG/Private/Elements/PCGDifferenceElement.cpp
--- a/Engine/Plugins/Experimental/PCG/Source/PCG/Private/Elements/PCGDifferenceElement.cpp
+++ b/Engine/Plugins/Experimental/PCG/Source/PCG/Private/Elements/PCGDifferenceElement.cpp
@@ -181,9 +181,7 @@ bool FPCGDifferenceElement::ExecuteInternal(FPCGContext* Context) const
 
 		DifferenceData->SetDensityFunction(Settings->DensityFunction);
 		DifferenceData->bDiffMetadata = Settings->bDiffMetadata;
-#if WITH_EDITOR
 		DifferenceData->bKeepZeroDensityPoints = Settings->bKeepZeroDensityPoints;
-#endif
 
 		FPCGTaggedData& Output = Outputs.Add_GetRef(Source);
 		Output.Data = DifferenceData;
diff --git a/Engine/Plugins/Experimental/PCG/Source/PCG/Private/Elements/PCGInnerIntersectionElement.cpp b/Engine/Plugins/Experimental/PCG/Source/PCG/Private/Elements/PCGInnerIntersectionElement.cpp
--- a/Engine/Plugins/Experimental/PCG/Source/PCG/Private/Elements/PCGInnerIntersectionElement.cpp
+++ b/Engine/Plugins/Experimental/PCG/Source/PCG/Private/Elements/PCGInnerIntersectionElement.cpp
@@ -55,11 +55,7 @@ bool FPCGInnerIntersectionElement::ExecuteInternal(FPCGContext* Context) const
 	TArray<FPCGTaggedData> Inputs = Context->InputData.GetInputs();
 
 	const EPCGIntersectionDensityFunction DensityFunction = Settings->DensityFunction;
-#if WITH_EDITOR
 	const bool bKeepZeroDensityPoints = Settings->bKeepZeroDensityPoints;
-#else
-	const bool bKeepZeroDensityPoints = false;
-#endif
 
 	TArray<FPCGTaggedData>& Outputs = Context->OutputData.TaggedData;
 
@@ -90,9 +86,7 @@ bool FPCGInnerIntersectionElement::ExecuteInternal(FPCGContext* Context) const
 		IntersectionData = (IntersectionData ? IntersectionData : FirstSpatialData)->IntersectWith(SpatialData);
 		// Propagate settings
 		IntersectionData->DensityFunction = DensityFunction;
-#if WITH_EDITOR
 		IntersectionData->bKeepZeroDensityPoints = bKeepZeroDensityPoints;
-#endif
 
 		// Update tagged data
 		FPCGTaggedData& IntersectionTaggedData = Outputs[IntersectionTaggedDataIndex];
diff --git a/Engine/Plugins/Experimental/PCG/Source/PCG/Private/Elements/PCGOuterIntersectionElement.cpp b/Engine/Plugins/Experimental/PCG/Source/PCG/Private/Elements/PCGOuterIntersectionElement.cpp
--- a/Engine/Plugins/Experimental/PCG/Source/PCG/Private/Elements/PCGOuterIntersectionElement.cpp
+++ b/Engine/Plugins/Experimental/PCG/Source/PCG/Private/Elements/PCGOuterIntersectionElement.cpp
@@ -155,9 +155,8 @@ bool FPCGOuterIntersectionElement::ExecuteInternal(FPCGContext* Context) const
 			UPCGIntersectionData* TempIntersectionData = IntersectionData->IntersectWith(SecondarySourceUnionData);
 			// Propagate settings
 			TempIntersectionData->DensityFunction = Settings->DensityFunction;
-#if WITH_EDITOR
 			TempIntersectionData->bKeepZeroDensityPoints = Settings->bKeepZeroDensityPoints;
-#endif // WITH_EDITOR
+
 			IntersectionData = TempIntersectionData;
 		}
 
diff --git a/Engine/Plugins/Experimental/PCG/Source/PCG/Private/Elements/PCGProjectionElement.cpp b/Engine/Plugins/Experimental/PCG/Source/PCG/Private/Elements/PCGProjectionElement.cpp
--- a/Engine/Plugins/Experimental/PCG/Source/PCG/Private/Elements/PCGProjectionElement.cpp
+++ b/Engine/Plugins/Experimental/PCG/Source/PCG/Private/Elements/PCGProjectionElement.cpp
@@ -62,11 +62,7 @@ bool FPCGProjectionElement::ExecuteInternal(FPCGContext* Context) const
 
 	TArray<FPCGTaggedData>& Outputs = Context->OutputData.TaggedData;
 
-#if WITH_EDITOR
 	const bool bKeepZeroDensityPoints = Settings->bKeepZeroDensityPoints;
-#else
-	const bool bKeepZeroDensityPoints = false;
-#endif
 
 	for (FPCGTaggedData& Source : Sources)
 	{
@@ -79,9 +75,7 @@ bool FPCGProjectionElement::ExecuteInternal(FPCGContext* Context) const
 		}
 
 		UPCGSpatialData* ProjectionData = ProjectionSource->ProjectOn(ProjectionTarget, ProjectionParams);
-#if WITH_EDITOR
 		ProjectionData->bKeepZeroDensityPoints = bKeepZeroDensityPoints;
-#endif
 
 		if (ProjectionData->RequiresCollapseToSample())
 		{
diff --git a/Engine/Plugins/Experimental/PCG/Source/PCG/Public/Data/PCGSpatialData.h b/Engine/Plugins/Experimental/PCG/Source/PCG/Public/Data/PCGSpatialData.h
--- a/Engine/Plugins/Experimental/PCG/Source/PCG/Public/Data/PCGSpatialData.h
+++ b/Engine/Plugins/Experimental/PCG/Source/PCG/Public/Data/PCGSpatialData.h
@@ -139,10 +139,8 @@ public:
 	UPROPERTY(Transient, BlueprintReadWrite, EditAnywhere, Category = Data)
 	TWeakObjectPtr<AActor> TargetActor = nullptr;
 
-#if WITH_EDITORONLY_DATA
-	UPROPERTY(Transient, BlueprintReadWrite, EditAnywhere, Category = Debug)
+	UPROPERTY(BlueprintReadWrite, EditAnywhere, Category = Data)
 	bool bKeepZeroDensityPoints = false;
-#endif
 
 	// Not accessible through blueprint to make sure the constness is preserved
 	UPROPERTY(VisibleAnywhere, Category = Metadata)
diff --git a/Engine/Plugins/Experimental/PCG/Source/PCG/Public/Elements/PCGDifferenceElement.h b/Engine/Plugins/Experimental/PCG/Source/PCG/Public/Elements/PCGDifferenceElement.h
--- a/Engine/Plugins/Experimental/PCG/Source/PCG/Public/Elements/PCGDifferenceElement.h
+++ b/Engine/Plugins/Experimental/PCG/Source/PCG/Public/Elements/PCGDifferenceElement.h
@@ -52,11 +52,9 @@ public:
 	UPROPERTY(BlueprintReadWrite, EditAnywhere, Category = Settings)
 	bool bDiffMetadata = true;
 
-#if WITH_EDITORONLY_DATA
 	/** If enabled, the output will not automatically filter out points with zero density. */
-	UPROPERTY(Transient, BlueprintReadWrite, EditAnywhere, Category = "Settings|Debug", meta = (PCG_Overridable))
+	UPROPERTY(BlueprintReadWrite, EditAnywhere, Category = Settings, meta = (PCG_Overridable))
 	bool bKeepZeroDensityPoints = false;
-#endif
 };
 
 class FPCGDifferenceElement : public IPCGElement
diff --git a/Engine/Plugins/Experimental/PCG/Source/PCG/Public/Elements/PCGInnerIntersectionElement.h b/Engine/Plugins/Experimental/PCG/Source/PCG/Public/Elements/PCGInnerIntersectionElement.h
--- a/Engine/Plugins/Experimental/PCG/Source/PCG/Public/Elements/PCGInnerIntersectionElement.h
+++ b/Engine/Plugins/Experimental/PCG/Source/PCG/Public/Elements/PCGInnerIntersectionElement.h
@@ -35,11 +35,9 @@ public:
 	UPROPERTY(BlueprintReadWrite, EditAnywhere, Category = Settings, meta = (PCG_Overridable))
 	EPCGIntersectionDensityFunction DensityFunction = EPCGIntersectionDensityFunction::Multiply;
 
-#if WITH_EDITORONLY_DATA
-	/** [EDITOR ONLY] If enabled, output points with a density value of 0 will NOT be automatically filtered out. */
-	UPROPERTY(Transient, BlueprintReadWrite, EditAnywhere, Category = "Settings|Debug", meta = (PCG_Overridable))
+	/** If enabled, output points with a density value of 0 will NOT be automatically filtered out. */
+	UPROPERTY(BlueprintReadWrite, EditAnywhere, Category = Settings, meta = (PCG_Overridable))
 	bool bKeepZeroDensityPoints = false;
-#endif
 };
 
 class FPCGInnerIntersectionElement : public IPCGElement
diff --git a/Engine/Plugins/Experimental/PCG/Source/PCG/Public/Elements/PCGOuterIntersectionElement.h b/Engine/Plugins/Experimental/PCG/Source/PCG/Public/Elements/PCGOuterIntersectionElement.h
--- a/Engine/Plugins/Experimental/PCG/Source/PCG/Public/Elements/PCGOuterIntersectionElement.h
+++ b/Engine/Plugins/Experimental/PCG/Source/PCG/Public/Elements/PCGOuterIntersectionElement.h
@@ -49,11 +49,9 @@ public:
 	UPROPERTY(BlueprintReadWrite, EditAnywhere, Category = Settings, meta = (PCG_Overridable))
 	bool bIgnorePinsWithNoInput = false;
 
-#if WITH_EDITORONLY_DATA
-	/** [EDITOR ONLY] If enabled, output points with a density value of 0 will NOT be automatically filtered out. */
-	UPROPERTY(Transient, BlueprintReadWrite, EditAnywhere, Category = "Settings|Debug", meta = (PCG_Overridable))
+	/** If enabled, output points with a density value of 0 will NOT be automatically filtered out. */
+	UPROPERTY(BlueprintReadWrite, EditAnywhere, Category = Settings, meta = (PCG_Overridable))
 	bool bKeepZeroDensityPoints = false;
-#endif // WITH_EDITORONLY_DATA
 };
 
 class FPCGOuterIntersectionElement : public IPCGElement
diff --git a/Engine/Plugins/Experimental/PCG/Source/PCG/Public/Elements/PCGProjectionElement.h b/Engine/Plugins/Experimental/PCG/Source/PCG/Public/Elements/PCGProjectionElement.h
--- a/Engine/Plugins/Experimental/PCG/Source/PCG/Public/Elements/PCGProjectionElement.h
+++ b/Engine/Plugins/Experimental/PCG/Source/PCG/Public/Elements/PCGProjectionElement.h
@@ -31,10 +31,8 @@ public:
 	UPROPERTY(BlueprintReadWrite, EditAnywhere, Category = Settings, meta = (ShowOnlyInnerProperties, PCG_Overridable))
 	FPCGProjectionParams ProjectionParams;
 
-#if WITH_EDITORONLY_DATA
-	UPROPERTY(Transient, BlueprintReadWrite, EditAnywhere, Category = "Settings|Debug", meta = (PCG_Overridable))
+	UPROPERTY(BlueprintReadWrite, EditAnywhere, Category = Settings, meta = (PCG_Overridable))
 	bool bKeepZeroDensityPoints = false;
-#endif
 
 #if WITH_EDITOR
 	virtual void ApplyDeprecation(UPCGNode* InOutNode) override;
