[PCG] Updated the Volume Sampler to the new Time Slice Context

* Minor update to Surface Sampler to better utilize the iteration state of the Time Slice Context

diff --git a/Engine/Plugins/Experimental/PCG/Source/PCG/Private/Data/PCGCollisionShapeData.cpp b/Engine/Plugins/Experimental/PCG/Source/PCG/Private/Data/PCGCollisionShapeData.cpp
--- a/Engine/Plugins/Experimental/PCG/Source/PCG/Private/Data/PCGCollisionShapeData.cpp
+++ b/Engine/Plugins/Experimental/PCG/Source/PCG/Private/Data/PCGCollisionShapeData.cpp
@@ -94,12 +94,9 @@ const UPCGPointData* UPCGCollisionShapeData::CreatePointData(FPCGContext* Contex
 {
 	TRACE_CPUPROFILER_EVENT_SCOPE(UPCGShapeData::CreatePointData);
 
-	const FVector DefaultVoxelSize = FVector(100.0, 100.0, 100.0);
+	const PCGVolumeSampler::FVolumeSamplerParams SamplerParams;
 
-	PCGVolumeSampler::FVolumeSamplerSettings SamplerSettings;
-	SamplerSettings.VoxelSize = DefaultVoxelSize;
-
-	const UPCGPointData* Data = PCGVolumeSampler::SampleVolume(Context, this, SamplerSettings);
+	const UPCGPointData* Data = PCGVolumeSampler::SampleVolume(Context, SamplerParams, this);
 
 	if (Data)
 	{
diff --git a/Engine/Plugins/Experimental/PCG/Source/PCG/Private/Data/PCGPrimitiveData.cpp b/Engine/Plugins/Experimental/PCG/Source/PCG/Private/Data/PCGPrimitiveData.cpp
--- a/Engine/Plugins/Experimental/PCG/Source/PCG/Private/Data/PCGPrimitiveData.cpp
+++ b/Engine/Plugins/Experimental/PCG/Source/PCG/Private/Data/PCGPrimitiveData.cpp
@@ -52,10 +52,11 @@ const UPCGPointData* UPCGPrimitiveData::CreatePointData(FPCGContext* Context) co
 {
 	TRACE_CPUPROFILER_EVENT_SCOPE(UPCGPrimitiveData::CreatePointData);
 
-	PCGVolumeSampler::FVolumeSamplerSettings SamplerSettings;
-	SamplerSettings.VoxelSize = VoxelSize;
+	PCGVolumeSampler::FVolumeSamplerParams SamplerParams;
+	SamplerParams.VoxelSize = VoxelSize;
+	SamplerParams.Bounds = GetBounds();
 
-	UPCGPointData* Data = PCGVolumeSampler::SampleVolume(Context, this, SamplerSettings);
+	const UPCGPointData* Data = PCGVolumeSampler::SampleVolume(Context, SamplerParams, this);
 
 	if (Data && ensure(Primitive.Get()))
 	{
diff --git a/Engine/Plugins/Experimental/PCG/Source/PCG/Private/Data/PCGVolumeData.cpp b/Engine/Plugins/Experimental/PCG/Source/PCG/Private/Data/PCGVolumeData.cpp
--- a/Engine/Plugins/Experimental/PCG/Source/PCG/Private/Data/PCGVolumeData.cpp
+++ b/Engine/Plugins/Experimental/PCG/Source/PCG/Private/Data/PCGVolumeData.cpp
@@ -61,10 +61,11 @@ const UPCGPointData* UPCGVolumeData::CreatePointData(FPCGContext* Context) const
 {
 	TRACE_CPUPROFILER_EVENT_SCOPE(UPCGVolumeData::CreatePointData);
 
-	PCGVolumeSampler::FVolumeSamplerSettings SamplerSettings;
-	SamplerSettings.VoxelSize = VoxelSize;
+	PCGVolumeSampler::FVolumeSamplerParams SamplerParams;
+	SamplerParams.VoxelSize = VoxelSize;
+	SamplerParams.Bounds = GetBounds();
 
-	UPCGPointData* Data = PCGVolumeSampler::SampleVolume(Context, this, SamplerSettings);
+	const UPCGPointData* Data = PCGVolumeSampler::SampleVolume(Context, SamplerParams, this);
 
 	if (Data)
 	{
diff --git a/Engine/Plugins/Experimental/PCG/Source/PCG/Private/Data/PCGWorldData.cpp b/Engine/Plugins/Experimental/PCG/Source/PCG/Private/Data/PCGWorldData.cpp
--- a/Engine/Plugins/Experimental/PCG/Source/PCG/Private/Data/PCGWorldData.cpp
+++ b/Engine/Plugins/Experimental/PCG/Source/PCG/Private/Data/PCGWorldData.cpp
@@ -199,13 +199,14 @@ const UPCGPointData* UPCGWorldVolumetricData::CreatePointData(FPCGContext* Conte
 		
 		return Data;
 	}
-	
-	PCGVolumeSampler::FVolumeSamplerSettings SamplerSettings;
-	SamplerSettings.VoxelSize = VoxelSize;
 
-	PCGVolumeSampler::SampleVolume(Context, this, nullptr, EffectiveBounds, SamplerSettings, Data);
+	PCGVolumeSampler::FVolumeSamplerParams SamplerParams;
+	SamplerParams.VoxelSize = VoxelSize;
+	SamplerParams.Bounds = EffectiveBounds;
+
+	Data = PCGVolumeSampler::SampleVolume(Context, SamplerParams, this);
 	UE_LOG(LogPCG, Verbose, TEXT("Volumetric world extracted %d points"), Data->GetPoints().Num());
-	
+
 	return Data;
 }
 
@@ -399,10 +400,10 @@ const UPCGPointData* UPCGWorldRayHitData::CreatePointData(FPCGContext* Context,
 		return Data;
 	}
 
-	PCGSurfaceSampler::FSurfaceSamplerExecutionSettings SamplerSettings;
-	if (SamplerSettings.Initialize(nullptr, Context, EffectiveBounds))
+	PCGSurfaceSampler::FSurfaceSamplerParams SamplerParams;
+	if (SamplerParams.Initialize(nullptr, Context, EffectiveBounds))
 	{
-		Data = PCGSurfaceSampler::SampleSurface(Context, this, nullptr, SamplerSettings);
+		Data = PCGSurfaceSampler::SampleSurface(Context, this, nullptr, SamplerParams);
 	}
 
 	return Data;
diff --git a/Engine/Plugins/Experimental/PCG/Source/PCG/Private/Elements/PCGSurfaceSampler.cpp b/Engine/Plugins/Experimental/PCG/Source/PCG/Private/Elements/PCGSurfaceSampler.cpp
--- a/Engine/Plugins/Experimental/PCG/Source/PCG/Private/Elements/PCGSurfaceSampler.cpp
+++ b/Engine/Plugins/Experimental/PCG/Source/PCG/Private/Elements/PCGSurfaceSampler.cpp
@@ -30,7 +30,7 @@ namespace PCGSurfaceSamplerConstants
 
 namespace PCGSurfaceSampler
 {
-	bool FSurfaceSamplerExecutionSettings::Initialize(const UPCGSurfaceSamplerSettings* Settings, const FPCGContext* Context, const FBox& InputBounds)
+	bool FSurfaceSamplerParams::Initialize(const UPCGSurfaceSamplerSettings* Settings, const FPCGContext* Context, const FBox& InputBounds)
 	{
 		if (!Context)
 		{
@@ -115,7 +115,7 @@ namespace PCGSurfaceSampler
 		return true;
 	}
 
-	FIntVector2 FSurfaceSamplerExecutionSettings::ComputeCellIndices(int32 Index) const
+	FIntVector2 FSurfaceSamplerParams::ComputeCellIndices(int32 Index) const
 	{
 		check(Index >= 0 && Index < CellCount);
 		const int32 CellCountX = 1 + CellMaxX - CellMinX;
@@ -123,7 +123,7 @@ namespace PCGSurfaceSampler
 		return FIntVector2(CellMinX + (Index % CellCountX), CellMinY + (Index / CellCountX));
 	}
 
-	UPCGPointData* SampleSurface(FPCGContext* Context, const UPCGSpatialData* InSurface, const UPCGSpatialData* InBoundingShape, const FSurfaceSamplerExecutionSettings& ExecutionSettings)
+	UPCGPointData* SampleSurface(FPCGContext* Context, const UPCGSpatialData* InSurface, const UPCGSpatialData* InBoundingShape, const FSurfaceSamplerParams& ExecutionSettings)
 	{
 		UPCGPointData* SampledData = NewObject<UPCGPointData>();
 		SampledData->InitializeFromData(InSurface);
@@ -134,7 +134,7 @@ namespace PCGSurfaceSampler
 		return SampledData;
 	}
 
-	bool SampleSurface(FPCGContext* Context, const FSurfaceSamplerExecutionSettings& Settings, const UPCGSpatialData* InSurface, const UPCGSpatialData* InBoundingShape, UPCGPointData* SampledData, const bool bTimeSlicingIsEnabled)
+	bool SampleSurface(FPCGContext* Context, const FSurfaceSamplerParams& Settings, const UPCGSpatialData* InSurface, const UPCGSpatialData* InBoundingShape, UPCGPointData* SampledData, const bool bTimeSlicingIsEnabled)
 	{
 		check(InSurface && SampledData);
 
@@ -218,8 +218,7 @@ namespace PCGSurfaceSampler
 		};
 
 		FPCGAsyncState* AsyncState = Context ? &Context->AsyncState : nullptr;
-		const bool bEnableTimeSlicing = Context ? bTimeSlicingIsEnabled : false;
-		return FPCGAsync::AsyncProcessing<FPCGPoint>(AsyncState, Settings.CellCount, SampledPoints, AsyncProcessFunc, bEnableTimeSlicing);
+		return FPCGAsync::AsyncProcessing<FPCGPoint>(AsyncState, Settings.CellCount, SampledPoints, AsyncProcessFunc, /*bEnableTimeSlicing=*/Context && bTimeSlicingIsEnabled);
 	}
 
 #if WITH_EDITOR
@@ -347,7 +346,7 @@ namespace PCGSurfaceSamplerHelpers
 
 		// Grab the Bounding Shape input if there is one.
 		const TArray<FPCGTaggedData> BoundingShapeInputs = Context->InputData.GetInputsByPin(PCGSurfaceSamplerConstants::BoundingShapeLabel);
-		FBox BoundingShapeBounds(EForceInit::ForceInit);
+		FBox& BoundingShapeBounds = OutState.BoundingShapeBounds;
 		if (!Settings->bUnbounded)
 		{
 			bool bUnionWasCreated;
@@ -396,16 +395,42 @@ namespace PCGSurfaceSamplerHelpers
 			return EPCGTimeSliceInitResult::AbortExecution;
 		}
 
-		// Iterate over the generating shapes to compute the bounds
-		int32 GeneratingShapeIndex = 0;
-		while (GeneratingShapeIndex < GeneratingShapes.Num())
+		return EPCGTimeSliceInitResult::Success;
+	}
+}
+
+bool FPCGSurfaceSamplerElement::PrepareDataInternal(FPCGContext* InContext) const
+{
+	TRACE_CPUPROFILER_EVENT_SCOPE(FPCGSurfaceSamplerElement::PrepareDataInternal);
+	ContextType* Context = static_cast<ContextType*>(InContext);
+	check(Context);
+
+	const UPCGSurfaceSamplerSettings* Settings = Context->GetInputSettings<UPCGSurfaceSamplerSettings>();
+
+	// Initialize the per-execution state data that won't change over the duration of the time slicing
+	if (Context->InitializePerExecutionState(PCGSurfaceSamplerHelpers::InitializePerExecutionData) == EPCGTimeSliceInitResult::AbortExecution)
+	{
+		PCGE_LOG(Warning, GraphAndLog, LOCTEXT("CouldNotInitializeExecutionState", "Could not initialize per-execution timeslice state data"));
+		return true;
+	}
+
+	TArray<const UPCGSpatialData*>& GeneratingShapes = Context->GetPerExecutionState().GeneratingShapes;
+	TArray<FPCGTaggedData>& Outputs = Context->OutputData.TaggedData;
+
+	// Initialize the per-iteration data, using the generating shapes as the source of iteration
+	Context->InitializePerIterationStates(GeneratingShapes.Num(),
+		[&GeneratingShapes, &Outputs, &Settings, &Context](IterStateType& OutState, const ExecStateType& ExecState, const uint32 IterationIndex)
 		{
 			// If we have generating shape inputs, use them
-			const UPCGSpatialData* GeneratingShape = GeneratingShapes[GeneratingShapeIndex];
+			const UPCGSpatialData* GeneratingShape = GeneratingShapes[IterationIndex];
 			check(GeneratingShape);
 
+			OutState.OutputPoints = NewObject<UPCGPointData>();
+			OutState.OutputPoints->InitializeFromData(GeneratingShape);
+
 			// Calculate the intersection of bounds of the provided inputs
 			FBox InputBounds = FBox(EForceInit::ForceInit);
+			const FBox& BoundingShapeBounds = ExecState.BoundingShapeBounds;
 
 			if (GeneratingShape->IsBounded())
 			{
@@ -430,50 +455,12 @@ namespace PCGSurfaceSamplerHelpers
 				}
 				else if(!InputBounds.IsValid)
 				{
-					// TODO [REVIEWERS]: Should this not be at least a warning, since the generating shape is being culled here?
 					PCGE_LOG_C(Verbose, LogOnly, Context, LOCTEXT("InvalidSamplingBounds", "Final sampling bounds is invalid/zero-sized."));
 				}
 
-				Outputs.RemoveAt(GeneratingShapeIndex);
-				GeneratingShapes.RemoveAt(GeneratingShapeIndex);
-				continue;
+				return EPCGTimeSliceInitResult::NoOperation;
 			}
 
-			++GeneratingShapeIndex;
-		}
-
-		// If there are still no generating shapes
-		if (GeneratingShapes.IsEmpty())
-		{
-			PCGE_LOG_C(Warning, GraphAndLog, Context, LOCTEXT("InvalidSurface", "No surfaces from which to generate are valid"));
-		}
-
-		return EPCGTimeSliceInitResult::Success;
-	}
-}
-
-bool FPCGSurfaceSamplerElement::PrepareDataInternal(FPCGContext* InContext) const
-{
-	TRACE_CPUPROFILER_EVENT_SCOPE(FPCGSurfaceSamplerElement::PrepareDataInternal);
-	ContextType* Context = static_cast<ContextType*>(InContext);
-	check(Context);
-
-	// Initialize the per-execution state data that won't change over the duration of the time slicing
-	if (Context->InitializePerExecutionState(PCGSurfaceSamplerHelpers::InitializePerExecutionData) == EPCGTimeSliceInitResult::AbortExecution)
-	{
-		PCGE_LOG(Warning, GraphAndLog, LOCTEXT("CouldNotInitializeExecutionState", "Could not initialize per-execution timeslice state data"));
-		return true;
-	}
-
-	TArray<const UPCGSpatialData*>& GeneratingShapes = Context->GetPerExecutionState().GeneratingShapes;
-	TArray<FPCGTaggedData>& Outputs = Context->OutputData.TaggedData;
-
-	// Initialize the per-iteration data, using the generating shapes as the source of iteration
-	Context->InitializePerIterationStates(GeneratingShapes.Num(),
-		[&GeneratingShapes, &Outputs](IterStateType& OutState, const uint32 IterationIndex)
-		{
-			OutState.OutputPoints = NewObject<UPCGPointData>();
-			OutState.OutputPoints->InitializeFromData(GeneratingShapes[IterationIndex]);
 			// Assigning this here prevents the need to root
 			Outputs[IterationIndex].Data = OutState.OutputPoints;
 			return EPCGTimeSliceInitResult::Success;
@@ -516,10 +503,11 @@ bool FPCGSurfaceSamplerElement::ExecuteInternal(FPCGContext* InContext) const
 	}
 
 	// The context will iterate over per-iteration states and execute the lambda until it returns true
-	return ExecuteSlice(TimeSlicedContext, [](ContextType* Context, const ExecStateType& ExecutionState, const IterStateType& IterationState, const uint32 IterationIndex)->bool
+	return ExecuteSlice(TimeSlicedContext, [](ContextType* Context, const ExecStateType& ExecState, const IterStateType& IterState, const uint32 IterationIndex)->bool
 	{
-		// This iteration resulted in an early out for no sampling operation. Early out with empty point data.
 		const EPCGTimeSliceInitResult InitResult = Context->GetIterationStateResult(IterationIndex);
+
+		// This iteration resulted in an early out for no sampling operation. Early out with empty point data.
 		if (InitResult == EPCGTimeSliceInitResult::NoOperation)
 		{
 			Context->OutputData.TaggedData[IterationIndex].Data = NewObject<UPCGPointData>();
@@ -531,11 +519,11 @@ bool FPCGSurfaceSamplerElement::ExecuteInternal(FPCGContext* InContext) const
 		check(InitResult == EPCGTimeSliceInitResult::Success);
 
 		// Run the execution until the time slice is finished
-		const bool bAsyncDone = PCGSurfaceSampler::SampleSurface(Context, ExecutionState.Settings, ExecutionState.GeneratingShapes[IterationIndex], ExecutionState.BoundingShape, IterationState.OutputPoints, Context->TimeSliceIsEnabled());
+		const bool bAsyncDone = PCGSurfaceSampler::SampleSurface(Context, IterState.Settings, ExecState.GeneratingShapes[IterationIndex], ExecState.BoundingShape, IterState.OutputPoints, Context->TimeSliceIsEnabled());
 
 		if (Context && bAsyncDone)
 		{
-			PCGE_LOG_C(Verbose, LogOnly, Context, FText::Format(LOCTEXT("GenerationInfo", "Generated {0} points in {1} cells"), IterationState.OutputPoints->GetPoints().Num(), ExecutionState.Settings.CellCount));
+			PCGE_LOG_C(Verbose, LogOnly, Context, FText::Format(LOCTEXT("GenerationInfo", "Generated {0} points in {1} cells"), IterState.OutputPoints->GetPoints().Num(), IterState.Settings.CellCount));
 		}
 
 		return bAsyncDone;
diff --git a/Engine/Plugins/Experimental/PCG/Source/PCG/Private/Elements/PCGVolumeSampler.cpp b/Engine/Plugins/Experimental/PCG/Source/PCG/Private/Elements/PCGVolumeSampler.cpp
--- a/Engine/Plugins/Experimental/PCG/Source/PCG/Private/Elements/PCGVolumeSampler.cpp
+++ b/Engine/Plugins/Experimental/PCG/Source/PCG/Private/Elements/PCGVolumeSampler.cpp
@@ -8,6 +8,7 @@
 #include "PCGCustomVersion.h"
 #include "Data/PCGPointData.h"
 #include "Data/PCGSpatialData.h"
+#include "Elements/PCGTimeSlicedElementBase.h"
 #include "Helpers/PCGAsync.h"
 #include "Helpers/PCGHelpers.h"
 
@@ -25,41 +26,40 @@ namespace PCGVolumeSamplerConstants
 
 namespace PCGVolumeSampler
 {
-	UPCGPointData* SampleVolume(FPCGContext* InContext, const UPCGSpatialData* InVolume, const FVolumeSamplerSettings& InSamplerSettings)
-	{
-		const FBox Bounds = InVolume->GetBounds();
-		return SampleVolume(InContext, InVolume, nullptr, Bounds, InSamplerSettings);
-	}
-
-	UPCGPointData* SampleVolume(FPCGContext* InContext, const UPCGSpatialData* InVolume, const UPCGSpatialData* InBoundingShape, const FBox& InBounds, const FVolumeSamplerSettings& InSamplerSettings)
+	UPCGPointData* SampleVolume(FPCGContext* Context, const FVolumeSamplerParams& SamplerSettings, const UPCGSpatialData* Volume, const UPCGSpatialData* BoundingShape)
 	{
 		UPCGPointData* Data = NewObject<UPCGPointData>();
-		Data->InitializeFromData(InVolume);
+		Data->InitializeFromData(Volume);
 
-		SampleVolume(InContext, InVolume, InBoundingShape, InBounds, InSamplerSettings, Data);
+		SampleVolume(Context, SamplerSettings, Volume, BoundingShape, Data, Context->TimeSliceIsEnabled());
 
 		return Data;
 	}
 
-	void SampleVolume(FPCGContext* InContext, const UPCGSpatialData* InVolume, const UPCGSpatialData* InBoundingShape, const FBox& InBounds, const FVolumeSamplerSettings& InSamplerSettings, UPCGPointData* OutputData)
+	bool SampleVolume(FPCGContext* Context, const FVolumeSamplerParams& SamplerSettings, const UPCGSpatialData* Volume, const UPCGSpatialData* BoundingShape, UPCGPointData* OutputData, const bool bTimeSlicingIsEnabled)
 	{
-		check(InVolume && OutputData);
+		check(Volume && OutputData);
 
-		// Early out
-		if (!InBounds.IsValid)
+		FBox Bounds(SamplerSettings.Bounds);
+		if (!Bounds.IsValid)
 		{
-			return;
+			Bounds = Volume->GetBounds();
+			// Early out
+			if (!Bounds.IsValid)
+			{
+				return true;
+			}
 		}
 
 		TArray<FPCGPoint>& Points = OutputData->GetMutablePoints();
-		const FVector& VoxelSize = InSamplerSettings.VoxelSize;
+		const FVector& VoxelSize = SamplerSettings.VoxelSize;
 
-		const int32 MinX = FMath::CeilToInt(InBounds.Min.X / VoxelSize.X);
-		const int32 MaxX = FMath::FloorToInt(InBounds.Max.X / VoxelSize.X);
-		const int32 MinY = FMath::CeilToInt(InBounds.Min.Y / VoxelSize.Y);
-		const int32 MaxY = FMath::FloorToInt(InBounds.Max.Y / VoxelSize.Y);
-		const int32 MinZ = FMath::CeilToInt(InBounds.Min.Z / VoxelSize.Z);
-		const int32 MaxZ = FMath::FloorToInt(InBounds.Max.Z / VoxelSize.Z);
+		const int32 MinX = FMath::CeilToInt(Bounds.Min.X / VoxelSize.X);
+		const int32 MaxX = FMath::FloorToInt(Bounds.Max.X / VoxelSize.X);
+		const int32 MinY = FMath::CeilToInt(Bounds.Min.Y / VoxelSize.Y);
+		const int32 MaxY = FMath::FloorToInt(Bounds.Max.Y / VoxelSize.Y);
+		const int32 MinZ = FMath::CeilToInt(Bounds.Min.Z / VoxelSize.Z);
+		const int32 MaxZ = FMath::FloorToInt(Bounds.Max.Z / VoxelSize.Z);
 
 		// Set uninitialized, then carefully initialize step by step with overflow checks
 		int32 NumIterations = -1;
@@ -73,12 +73,12 @@ namespace PCGVolumeSampler
 
 			if (NumX <= 0 || NumY <= 0 || NumZ <= 0)
 			{
-				if (InContext)
+				if (Context)
 				{
-					PCGE_LOG_C(Verbose, LogOnly, InContext, FText::Format(FText::FromString(TEXT("Skipped - invalid cell bounds ({0} x {1} x {2})")), NumX, NumY, NumZ));
+					PCGE_LOG_C(Verbose, LogOnly, Context, FText::Format(FText::FromString(TEXT("Skipped - invalid cell bounds ({0} x {1} x {2})")), NumX, NumY, NumZ));
 				}
 
-				return;
+				return true;
 			}
 
 			if (NumIterationsXY64 > 0 && 
@@ -91,16 +91,16 @@ namespace PCGVolumeSampler
 			}
 			else
 			{
-				if (InContext)
+				if (Context)
 				{
-					PCGE_LOG_C(Error, GraphAndLog, InContext, FText::Format(FText::FromString(TEXT("Skipped - tried to generate too many points ({0} x {1} x {2})")), NumX, NumY, NumZ));
+					PCGE_LOG_C(Error, GraphAndLog, Context, FText::Format(FText::FromString(TEXT("Skipped - tried to generate too many points ({0} x {1} x {2})")), NumX, NumY, NumZ));
 				}
 
-				return;
+				return true;
 			}
 		}
 
-		FPCGAsync::AsyncPointProcessing(InContext, NumIterations, Points, [InVolume, InBoundingShape, PointSteepness = InSamplerSettings.PointSteepness, VoxelSize, MinX, MaxX, MinY, MaxY, MinZ](int32 Index, FPCGPoint& OutPoint)
+		auto AsyncProcessingFunc = [Volume, BoundingShape, PointSteepness = SamplerSettings.PointSteepness, VoxelSize, MinX, MaxX, MinY, MaxY, MinZ](int32 Index, FPCGPoint& OutPoint)
 		{
 			const int X = MinX + (Index % (MaxX - MinX));
 			const int Y = MinY + (Index / (MaxX - MinX) % (MaxY - MinY));
@@ -114,12 +114,12 @@ namespace PCGVolumeSampler
 			// The OutPoint has not been initialized, so do it now
 			OutPoint = FPCGPoint();
 
-			if (InVolume->SamplePoint(SampleTransform, VoxelBox, OutPoint, nullptr))
+			if (Volume->SamplePoint(SampleTransform, VoxelBox, OutPoint, nullptr))
 			{
-				if (InBoundingShape)
+				if (BoundingShape)
 				{
 					FPCGPoint BoundingShapeSample;
-					if (!InBoundingShape->SamplePoint(SampleTransform, VoxelBox, BoundingShapeSample, nullptr))
+					if (!BoundingShape->SamplePoint(SampleTransform, VoxelBox, BoundingShapeSample, nullptr))
 					{
 						return false;
 					}
@@ -133,7 +133,10 @@ namespace PCGVolumeSampler
 			{
 				return false;
 			}
-		});
+		};
+
+		FPCGAsyncState* AsyncState = Context ? &Context->AsyncState : nullptr;
+		return FPCGAsync::AsyncProcessing<FPCGPoint>(AsyncState, NumIterations, Points, AsyncProcessingFunc, /*bEnableTimeSlicing=*/Context && bTimeSlicingIsEnabled);
 	}
 }
 
@@ -164,140 +167,221 @@ FPCGElementPtr UPCGVolumeSamplerSettings::CreateElement() const
 	return MakeShared<FPCGVolumeSamplerElement>();
 }
 
-bool FPCGVolumeSamplerElement::ExecuteInternal(FPCGContext* Context) const
+namespace PCGVolumeSamplerHelpers
 {
-	TRACE_CPUPROFILER_EVENT_SCOPE(FPCGVolumeSamplerElement::Execute);
-	// TODO: time-sliced implementation
-	const UPCGVolumeSamplerSettings* Settings = Context->GetInputSettings<UPCGVolumeSamplerSettings>();
-	check(Settings);
+	using ContextType = FPCGVolumeSamplerElement::ContextType;
+	using ExecStateType = FPCGVolumeSamplerElement::ExecStateType;
 
-	TArray<FPCGTaggedData> VolumeInputs = Context->InputData.GetInputsByPin(PCGVolumeSamplerConstants::VolumeLabel);
-
-	const FVector& VoxelSize = Settings->VoxelSize;
-	if (VoxelSize.X <= 0 || VoxelSize.Y <= 0 || VoxelSize.Z <= 0)
+	EPCGTimeSliceInitResult InitializePerExecutionData(ContextType* Context, ExecStateType& OutState)
 	{
-		PCGE_LOG(Warning, GraphAndLog, LOCTEXT("InvalidVoxelSize", "Skipped - Invalid voxel size"));
-		return true;
-	}
-
-	PCGVolumeSampler::FVolumeSamplerSettings SamplerSettings;
-	SamplerSettings.VoxelSize = VoxelSize;
-	SamplerSettings.PointSteepness = Settings->PointSteepness;
-
-	TArray<FPCGTaggedData>& Outputs = Context->OutputData.TaggedData;
+		TRACE_CPUPROFILER_EVENT_SCOPE(FPCGVolumeSamplerElement::InitializePerExecutionData);
 
-	// Grab the Bounding Shape input if there is one.
-	TArray<FPCGTaggedData> BoundingShapeInputs = Context->InputData.GetInputsByPin(PCGVolumeSamplerConstants::BoundingShapeLabel);
-	const UPCGSpatialData* BoundingShape = nullptr;
-	bool bUsedDefaultBoundingShape = false;
+		check(Context);
+		const UPCGVolumeSamplerSettings* Settings = Context->GetInputSettings<UPCGVolumeSamplerSettings>();
+		check(Settings);
 
-	if (!Settings->bUnbounded)
-	{
-		// TODO: Once we support time-slicing, put this in the context and root (see FPCGSurfaceSamplerContext)
-		bool bUnionCreated = false;
-		BoundingShape = Context->InputData.GetSpatialUnionOfInputsByPin(PCGVolumeSamplerConstants::BoundingShapeLabel, bUnionCreated);
-		if (!BoundingShape && Context->SourceComponent.IsValid())
+		const FVector& VoxelSize = Settings->VoxelSize;
+		if (VoxelSize.X <= 0 || VoxelSize.Y <= 0 || VoxelSize.Z <= 0)
 		{
-			BoundingShape = Cast<UPCGSpatialData>(Context->SourceComponent->GetActorPCGData());
-			bUsedDefaultBoundingShape = true;
+			PCGE_LOG_C(Warning, GraphAndLog, Context, LOCTEXT("InvalidVoxelSize", "Skipped - Invalid voxel size"));
+			return EPCGTimeSliceInitResult::AbortExecution;
 		}
-	}
-	else if (BoundingShapeInputs.Num() > 0)
-	{
-		PCGE_LOG_C(Verbose, LogOnly, Context, LOCTEXT("BoundsIgnored", "The bounds of the Bounding Shape input pin will be ignored because the Unbounded option is enabled"));
-	}	
 
-	// Compute bounds of bounding shape input
-	FBox BoundingShapeBounds(EForceInit::ForceInit);
-	if (BoundingShape)
-	{
-		BoundingShapeBounds = BoundingShape->GetBounds();
-	}
+		const TArray<FPCGTaggedData> VolumeInputs = Context->InputData.GetInputsByPin(PCGVolumeSamplerConstants::VolumeLabel);
+		// Grab the Bounding Shape input if there is one.
+		TArray<FPCGTaggedData> BoundingShapeInputs = Context->InputData.GetInputsByPin(PCGVolumeSamplerConstants::BoundingShapeLabel);
 
-	// Construct a list of shapes to generate samples from. Prefer to get these directly from the first input pin.
-	TArray<const UPCGSpatialData*, TInlineAllocator<16>> GeneratingShapes;
-	for (FPCGTaggedData& TaggedData : VolumeInputs)
-	{
-		if (const UPCGSpatialData* SpatialData = Cast<UPCGSpatialData>(TaggedData.Data))
+		bool bUsedDefaultBoundingShape = false;
+		if (!Settings->bUnbounded)
 		{
-			GeneratingShapes.Add(SpatialData);
-			Outputs.Add(TaggedData);
+			bool bUnionWasCreated;
+			// Get a union of inputs and if successful, add it to the root. Will be removed and marked for GC in the state destructor
+			OutState.BoundingShape = Context->InputData.GetSpatialUnionOfInputsByPin(PCGVolumeSamplerConstants::BoundingShapeLabel, bUnionWasCreated);
+			if (OutState.BoundingShape && bUnionWasCreated)
+			{
+				Context->RootAndTrackObject(const_cast<UPCGSpatialData*>(OutState.BoundingShape));
+			}
+
+			if (!OutState.BoundingShape && Context->SourceComponent.IsValid())
+			{
+				// Create a bounding shape from the actor data
+				OutState.BoundingShape = Cast<UPCGSpatialData>(Context->SourceComponent->GetActorPCGData());
+				bUsedDefaultBoundingShape = true;
+			}
+		}
+		else if (BoundingShapeInputs.Num() > 0)
+		{
+			PCGE_LOG_C(Verbose, LogOnly, Context, LOCTEXT("BoundsIgnored", "The bounds of the Bounding Shape input pin will be ignored because the Unbounded option is enabled"));
 		}
-	}
 
-	// If no shapes were obtained from the first input pin, try to find a shape to sample from nodes connected to the second pin.
-	if (GeneratingShapes.Num() == 0 && BoundingShape && !bUsedDefaultBoundingShape)
-	{
-		GeneratingShapes.Add(BoundingShape);
+		FBox& BoundingShapeBounds = OutState.BoundingShapeBounds;
+		// Compute bounds of bounding shape input
+		if (OutState.BoundingShape)
+		{
+			BoundingShapeBounds = OutState.BoundingShape->GetBounds();
+		}
 
-		// If there was a bounding shape input, use it as the starting point to get the tags
-		if (BoundingShapeInputs.Num() > 0)
+		TArray<FPCGTaggedData>& Outputs = Context->OutputData.TaggedData;
+		// Construct a list of shapes to generate samples from. Prefer to get these directly from the first input pin.
+		TArray<const UPCGSpatialData*>& GeneratingShapes = OutState.GeneratingShapes;
+		GeneratingShapes.Reserve(VolumeInputs.Num());
+		for (const FPCGTaggedData& TaggedData : VolumeInputs)
 		{
-			Outputs.Add(BoundingShapeInputs[0]);
+			if (const UPCGSpatialData* SpatialData = Cast<UPCGSpatialData>(TaggedData.Data))
+			{
+				GeneratingShapes.Add(SpatialData);
+				Outputs.Add(TaggedData);
+			}
 		}
-		else
+
+		// Warn if something is connected but no spatial data could be obtained for sampling
+		if (GeneratingShapes.IsEmpty() && (BoundingShapeInputs.Num() > 0 || VolumeInputs.Num() > 0))
 		{
-			Outputs.Emplace();
+			PCGE_LOG_C(Warning, GraphAndLog, Context, LOCTEXT("NoShapeToSample", "No Spatial data shape was provided for sampling. No points will be sampled."));
+			return EPCGTimeSliceInitResult::NoOperation;
 		}
+
+		return EPCGTimeSliceInitResult::Success;
 	}
+}
+
+bool FPCGVolumeSamplerElement::PrepareDataInternal(FPCGContext* Context) const
+{
+	TRACE_CPUPROFILER_EVENT_SCOPE(FPCGVolumeSamplerElement::PrepareDataInternal);
+	ContextType* TimeSlicedContext = static_cast<ContextType*>(Context);
+	check(TimeSlicedContext);
+
+	const UPCGVolumeSamplerSettings* Settings = TimeSlicedContext->GetInputSettings<UPCGVolumeSamplerSettings>();
+	check(Settings);
 
-	// Warn if something is connected but no spatial data could be obtained for sampling
-	if (GeneratingShapes.Num() == 0 && (BoundingShapeInputs.Num() > 0 || VolumeInputs.Num() > 0))
+	if (TimeSlicedContext->InitializePerExecutionState(PCGVolumeSamplerHelpers::InitializePerExecutionData) == EPCGTimeSliceInitResult::AbortExecution)
 	{
-		PCGE_LOG(Warning, GraphAndLog, LOCTEXT("NoShapeToSample", "No Spatial data shape was provided for sampling, no points will be produced"));
+		PCGE_LOG(Warning, GraphAndLog, LOCTEXT("CouldNotInitializeExecutionState", "Could not initialize per-execution timeslice state data"));
+		return true;
 	}
 
-	// TODO: embarassingly parallel loop
-	for (int GenerationIndex = 0; GenerationIndex < GeneratingShapes.Num(); ++GenerationIndex)
-	{
-		const UPCGSpatialData* GeneratingShape = GeneratingShapes[GenerationIndex];
-		check(GeneratingShape);
+	// The generating shapes will be used for the time slicing iterations
+	TArray<const UPCGSpatialData*>& GeneratingShapes = TimeSlicedContext->GetPerExecutionState().GeneratingShapes;
+	TArray<FPCGTaggedData>& Outputs = Context->OutputData.TaggedData;
 
-		// Calculate the intersection of bounds of the provided inputs
-		FBox InputBounds = FBox(EForceInit::ForceInit);
-		if (GeneratingShape->IsBounded())
+	TimeSlicedContext->InitializePerIterationStates(GeneratingShapes.Num(),
+		[&GeneratingShapes, &Settings, &Outputs, &Context](IterStateType& OutState, const ExecStateType& ExecState, const uint32 IterationIndex)
 		{
-			InputBounds = GeneratingShape->GetBounds();
+			OutState.Settings.VoxelSize = Settings->VoxelSize;
+			OutState.Settings.PointSteepness = Settings->PointSteepness;
 
-			if (BoundingShapeBounds.IsValid)
-			{
-				InputBounds = PCGHelpers::OverlapBounds(InputBounds, BoundingShapeBounds);
-			}
-		}
-		else
-		{
-			InputBounds = BoundingShapeBounds;
-		}
+			const UPCGSpatialData* GeneratingShape = GeneratingShapes[IterationIndex];
+			check(GeneratingShape);
 
-		if (!InputBounds.IsValid)
-		{
-			if (!GeneratingShape->IsBounded())
+			OutState.Volume = GeneratingShape;
+			OutState.OutputData = NewObject<UPCGPointData>();
+			OutState.OutputData->InitializeFromData(OutState.Volume);
+			Outputs[IterationIndex].Data = OutState.OutputData;
+
+			FBox& InputBounds = OutState.Settings.Bounds;
+
+			// Get the bounding shape bounds from the execution state
+			const FBox& BoundingShapeBounds = ExecState.BoundingShapeBounds;
+
+			// Calculate the intersection of bounds of the provided inputs
+			if (GeneratingShape->IsBounded())
 			{
-				// Some inputs are unable to provide bounds, like the WorldVolumetricQuery, in which case the user must provide bounds.
-				PCGE_LOG(Warning, GraphAndLog, LOCTEXT("CouldNotObtainInputBounds", "Input data is not bounded, so bounds must be provided for sampling. Consider providing a Bounding Shape input."));
+				InputBounds = GeneratingShape->GetBounds();
+
+				if (BoundingShapeBounds.IsValid)
+				{
+					InputBounds = PCGHelpers::OverlapBounds(InputBounds, BoundingShapeBounds);
+				}
 			}
 			else
 			{
-				PCGE_LOG(Verbose, LogOnly, LOCTEXT("InvalidSamplingBounds", "Final sampling bounds is invalid/zero-sized."));
+				InputBounds = BoundingShapeBounds;
 			}
 
-			Outputs.RemoveAt(GenerationIndex);
-			GeneratingShapes.RemoveAt(GenerationIndex);
-			--GenerationIndex;
-			continue;
-		}
+			if (!InputBounds.IsValid)
+			{
+				if (!GeneratingShape->IsBounded())
+				{
+					// Some inputs are unable to provide bounds, like the WorldVolumetricQuery, in which case the user must provide bounds.
+					PCGE_LOG_C(Warning, GraphAndLog, Context, LOCTEXT("CouldNotObtainInputBounds", "Input data is not bounded, so bounds must be provided for sampling. Consider providing a Bounding Shape input."));
+				}
+				else
+				{
+					PCGE_LOG_C(Verbose, LogOnly, Context, LOCTEXT("InvalidSamplingBounds", "Final sampling bounds is invalid/zero-sized."));
+				}
+
+				return EPCGTimeSliceInitResult::NoOperation;
+			}
+
+			return EPCGTimeSliceInitResult::Success;
+		});
+
+	if (!TimeSlicedContext->DataIsPreparedForExecution())
+	{
+		PCGE_LOG(Warning, GraphAndLog, LOCTEXT("CouldNotInitializeStateData", "Could not initialize timeslice state data"));
+		return true;
+	}
+
+	return true;
+}
 
-		// Sample volume
-		const UPCGPointData* SampledData = PCGVolumeSampler::SampleVolume(Context, GeneratingShape, BoundingShape, InputBounds, SamplerSettings);
-		Outputs[GenerationIndex].Data = SampledData;
+bool FPCGVolumeSamplerElement::ExecuteInternal(FPCGContext* Context) const
+{
+	TRACE_CPUPROFILER_EVENT_SCOPE(FPCGVolumeSamplerElement::Execute);
+	ContextType* TimeSlicedContext = static_cast<ContextType*>(Context);
+	check(TimeSlicedContext);
+
+	// Abort execution was called at some point during the state initialization
+	if (!TimeSlicedContext->DataIsPreparedForExecution())
+	{
+		TimeSlicedContext->OutputData.TaggedData.Empty();
 
-		if (SampledData)
+		return true;
+	}
+
+	// The execution would have resulted in an empty set of points for all iterations
+	if (TimeSlicedContext->GetExecutionStateResult() == EPCGTimeSliceInitResult::NoOperation)
+	{
+		for (FPCGTaggedData& Input : TimeSlicedContext->InputData.GetInputs())
 		{
-			PCGE_LOG(Verbose, LogOnly, FText::Format(LOCTEXT("GenerationInfo", "Generated {0} points in volume"), SampledData->GetPoints().Num()));
+			// TODO: Empty point data (to preserve previous behavior). Eventually, should be replaced with no output at all
+			FPCGTaggedData& Output = TimeSlicedContext->OutputData.TaggedData.Emplace_GetRef();
+			UPCGPointData* PointData = NewObject<UPCGPointData>();
+			PointData->InitializeFromData(Cast<UPCGSpatialData>(Input.Data));
+			Output.Data = PointData;
 		}
+
+		return true;
 	}
 
-	return true;
+	return ExecuteSlice(TimeSlicedContext, [](ContextType* Context, const ExecStateType& ExecState, const IterStateType& IterState, const uint32 IterationIndex)->bool
+	{
+		const EPCGTimeSliceInitResult InitResult = Context->GetIterationStateResult(IterationIndex);
+
+		if (InitResult == EPCGTimeSliceInitResult::NoOperation)
+		{
+			Context->OutputData.TaggedData[IterationIndex].Data = NewObject<UPCGPointData>();
+
+			return true;
+		}
+
+		check(InitResult == EPCGTimeSliceInitResult::Success);
+
+		const bool bAsyncDone = PCGVolumeSampler::SampleVolume(
+				Context,
+				IterState.Settings,
+				ExecState.GeneratingShapes[IterationIndex],
+				ExecState.BoundingShape,
+				IterState.OutputData,
+				Context->TimeSliceIsEnabled());
+
+		if (Context)
+		{
+			PCGE_LOG_C(Verbose, LogOnly, Context, FText::Format(LOCTEXT("GenerationInfo", "Generated {0} points in volume"), IterState.OutputData->GetPoints().Num()));
+		}
+
+		return bAsyncDone;
+	});
 }
 
 #if WITH_EDITOR
diff --git a/Engine/Plugins/Experimental/PCG/Source/PCG/Public/Elements/PCGSurfaceSampler.h b/Engine/Plugins/Experimental/PCG/Source/PCG/Public/Elements/PCGSurfaceSampler.h
--- a/Engine/Plugins/Experimental/PCG/Source/PCG/Public/Elements/PCGSurfaceSampler.h
+++ b/Engine/Plugins/Experimental/PCG/Source/PCG/Public/Elements/PCGSurfaceSampler.h
@@ -16,7 +16,8 @@ class UPCGSurfaceSamplerSettings;
 
 namespace PCGSurfaceSampler
 {
-	struct FSurfaceSamplerExecutionSettings
+	// TODO: Factor out "computed values" that shouldn't be exposed to the external user into a more appropriate state
+	struct FSurfaceSamplerParams
 	{
 		float PointsPerSquaredMeter = 1.0f;
 		FVector PointExtents = FVector::One() * 0.5f;
@@ -49,22 +50,22 @@ namespace PCGSurfaceSampler
 
 	struct FSurfaceSamplerExecutionState
 	{
-		FSurfaceSamplerExecutionSettings Settings;
-
 		const UPCGSpatialData* BoundingShape = nullptr;
+		FBox BoundingShapeBounds = FBox(EForceInit::ForceInit);
 		TArray<const UPCGSpatialData*> GeneratingShapes;
 	};
 
 	struct FSurfaceSamplerIterationState
 	{
+		FSurfaceSamplerParams Settings;
 		UPCGPointData* OutputPoints = nullptr;
 	};
 
 	/** Sample a surface and returns the resulting point data. Can't be timesliced. */
-	UPCGPointData* SampleSurface(FPCGContext* Context, const UPCGSpatialData* InSurface, const UPCGSpatialData* InBoundingShape, const FSurfaceSamplerExecutionSettings& ExecutionSettings);
+	UPCGPointData* SampleSurface(FPCGContext* Context, const UPCGSpatialData* InSurface, const UPCGSpatialData* InBoundingShape, const FSurfaceSamplerParams& ExecutionSettings);
 
 	/** Sample a surface and write the results in the given point data. Can be timesliced, and will return false if the processing is not done, true otherwise. */
-	bool SampleSurface(FPCGContext* Context, const FSurfaceSamplerExecutionSettings& Settings, const UPCGSpatialData* InSurface, const UPCGSpatialData* InBoundingShape, UPCGPointData* SampledData, const bool bTimeSlicingIsEnabled = false);
+	bool SampleSurface(FPCGContext* Context, const FSurfaceSamplerParams& Settings, const UPCGSpatialData* InSurface, const UPCGSpatialData* InBoundingShape, UPCGPointData* SampledData, const bool bTimeSlicingIsEnabled = false);
 }
 
 UCLASS(BlueprintType, ClassGroup = (Procedural))
diff --git a/Engine/Plugins/Experimental/PCG/Source/PCG/Public/Elements/PCGTimeSlicedElementBase.h b/Engine/Plugins/Experimental/PCG/Source/PCG/Public/Elements/PCGTimeSlicedElementBase.h
--- a/Engine/Plugins/Experimental/PCG/Source/PCG/Public/Elements/PCGTimeSlicedElementBase.h
+++ b/Engine/Plugins/Experimental/PCG/Source/PCG/Public/Elements/PCGTimeSlicedElementBase.h
@@ -59,7 +59,7 @@ struct TPCGTimeSlicedContext : public FPCGContext
 	/** Initializes per execution state data if required. */
 	EPCGTimeSliceInitResult InitializePerExecutionState(TFunctionRef<InitExecSignature> InitFunc = []{ return true; });
 
-	using InitIterSignature = EPCGTimeSliceInitResult(PerIterationStateT& OutState, const uint32 IterationIndex);
+	using InitIterSignature = EPCGTimeSliceInitResult(PerIterationStateT& OutState, const PerExecutionStateT& ExecState, const uint32 IterationIndex);
 
 	/** Initializes per execution state data if required. An array will be created with a state element for every execution iteration in the context. Returns the Init Result for each iteration's initialization. */
 	const TArray<EPCGTimeSliceInitResult>& InitializePerIterationStates(int32 NumIterations = 1, TFunctionRef<InitIterSignature> IterFunc = []{ return true; });
@@ -206,7 +206,7 @@ const TArray<EPCGTimeSliceInitResult>& TPCGTimeSlicedContext<PerExecutionStateT,
 
 	// Should be guaranteed to be uninitialized
 	check(PerIterationStateResultArray.IsEmpty() && PerIterationStateArray.IsEmpty());
-	
+
 	// An empty iteration state is still valid
 	if (NumIterations < 1)
 	{
@@ -218,7 +218,7 @@ const TArray<EPCGTimeSliceInitResult>& TPCGTimeSlicedContext<PerExecutionStateT,
 
 	for (int32 I = 0; I < NumIterations; ++I)
 	{
-		PerIterationStateResultArray[I] = IterFunc(PerIterationStateArray.Emplace_GetRef(), I);
+		PerIterationStateResultArray[I] = IterFunc(PerIterationStateArray.Emplace_GetRef(), PerExecutionStateData, I);
 
 		if (PerIterationStateResultArray[I] == EPCGTimeSliceInitResult::AbortExecution)
 		{
diff --git a/Engine/Plugins/Experimental/PCG/Source/PCG/Public/Elements/PCGVolumeSampler.h b/Engine/Plugins/Experimental/PCG/Source/PCG/Public/Elements/PCGVolumeSampler.h
--- a/Engine/Plugins/Experimental/PCG/Source/PCG/Public/Elements/PCGVolumeSampler.h
+++ b/Engine/Plugins/Experimental/PCG/Source/PCG/Public/Elements/PCGVolumeSampler.h
@@ -2,8 +2,10 @@
 
 #pragma once
 
-#include "Math/Box.h"
 #include "PCGSettings.h"
+#include "PCGTimeSlicedElementBase.h"
+
+#include "Math/Box.h"
 
 #include "PCGVolumeSampler.generated.h"
 
@@ -12,16 +14,35 @@ class UPCGSpatialData;
 
 namespace PCGVolumeSampler
 {
-	struct FVolumeSamplerSettings
+	const FVector DefaultVoxelSize = FVector(100.0, 100.0, 100.0);
+
+	struct FVolumeSamplerParams
 	{
-		FVector VoxelSize;
+		FVector VoxelSize = DefaultVoxelSize;
 		float PointSteepness = 0.0f;
+		FBox Bounds{EForceInit::ForceInit};
+	};
+
+	struct FVolumeSamplerExecutionState
+	{
+		const UPCGSpatialData* BoundingShape = nullptr;
+		FBox BoundingShapeBounds = FBox(EForceInit::ForceInit);
+		TArray<const UPCGSpatialData*> GeneratingShapes;
+	};
+
+	struct FVolumeSamplerIterationState
+	{
+		FVolumeSamplerParams Settings;
+
+		const UPCGSpatialData* Volume = nullptr;
+		UPCGPointData* OutputData = nullptr;
 	};
 
-	UPCGPointData* SampleVolume(FPCGContext* InContext, const UPCGSpatialData* InVolume, const FVolumeSamplerSettings& InSamplerSettings);
-	UPCGPointData* SampleVolume(FPCGContext* InContext, const UPCGSpatialData* InVolume, const UPCGSpatialData* InBoundingShape, const FBox& InBounds, const FVolumeSamplerSettings& InSamplerSettings);
+	/** Sample a volume and returns the resulting point data. */
+	UPCGPointData* SampleVolume(FPCGContext* Context, const FVolumeSamplerParams& SamplerSettings, const UPCGSpatialData* Volume, const UPCGSpatialData* BoundingShape = nullptr);
 
-	void SampleVolume(FPCGContext* InContext, const UPCGSpatialData* InVolume, const UPCGSpatialData* InBoundingShape, const FBox& InBounds, const FVolumeSamplerSettings& InSamplerSettings, UPCGPointData* InOutputData);
+	/** Sample a volume and write the results in the given point data. Can be timesliced and will return false if the processing is not done, true otherwise. */
+	bool SampleVolume(FPCGContext* Context, const FVolumeSamplerParams& SamplerSettings, const UPCGSpatialData* Volume, const UPCGSpatialData* BoundingShape, UPCGPointData* OutputData, const bool bTimeSlicingIsEnabled = false);
 }
 
 UCLASS(BlueprintType, ClassGroup = (Procedural))
@@ -31,7 +52,7 @@ class PCG_API UPCGVolumeSamplerSettings : public UPCGSettings
 
 public:
 	UPROPERTY(BlueprintReadWrite, EditAnywhere, Category = "Data", meta = (PCG_Overridable))
-	FVector VoxelSize = FVector(100.0, 100.0, 100.0);
+	FVector VoxelSize = PCGVolumeSampler::DefaultVoxelSize;
 
 	/** If no Bounding Shape input is provided, the actor bounds are used to limit the sample generation domain.
 	* This option allows ignoring the actor bounds and generating over the entire volume. Use with caution as this
@@ -67,12 +88,13 @@ protected:
 	//~End UPCGSettings interface
 };
 
-class FPCGVolumeSamplerElement : public IPCGElement
+class FPCGVolumeSamplerElement : public TPCGTimeSlicedElementBase<PCGVolumeSampler::FVolumeSamplerExecutionState, PCGVolumeSampler::FVolumeSamplerIterationState>
 {
 public:
 	// Might be sampling external data like brush, worth computing a full CRC in case we can halt change propagation/re-executions
 	virtual bool ShouldComputeFullOutputDataCrc(FPCGContext* Context) const override { return true; }
 
 protected:
+	virtual bool PrepareDataInternal(FPCGContext* Context) const override;
 	virtual bool ExecuteInternal(FPCGContext* Context) const override;
 };
