Fix: AddTags node no longer outputs overrides

#rnx

diff --git a/Engine/Plugins/Experimental/PCG/Source/PCG/Private/Elements/PCGAddTag.cpp b/Engine/Plugins/Experimental/PCG/Source/PCG/Private/Elements/PCGAddTag.cpp
--- a/Engine/Plugins/Experimental/PCG/Source/PCG/Private/Elements/PCGAddTag.cpp
+++ b/Engine/Plugins/Experimental/PCG/Source/PCG/Private/Elements/PCGAddTag.cpp
@@ -51,7 +51,8 @@ bool FPCGAddTagElement::ExecuteInternal(FPCGContext* Context) const
 	const UPCGAddTagSettings* Settings = Context->GetInputSettings<UPCGAddTagSettings>();
 	check(Settings);
 	
-	Context->OutputData = Context->InputData;
+	Context->OutputData.TaggedData = Context->InputData.GetInputsByPin(PCGPinConstants::DefaultInputLabel);
+
 	const TArray<FString> TagsArray = PCGHelpers::GetStringArrayFromCommaSeparatedString(Settings->TagsToAdd);
 	
 	for (const FString& Tag : TagsArray)
