Fixed copy-paste error wrt editing mode in the partition actors

#rnx

diff --git a/Engine/Plugins/Experimental/PCG/Source/PCG/Private/Grid/PCGPartitionActor.cpp b/Engine/Plugins/Experimental/PCG/Source/PCG/Private/Grid/PCGPartitionActor.cpp
--- a/Engine/Plugins/Experimental/PCG/Source/PCG/Private/Grid/PCGPartitionActor.cpp
+++ b/Engine/Plugins/Experimental/PCG/Source/PCG/Private/Grid/PCGPartitionActor.cpp
@@ -118,7 +118,7 @@ void APCGPartitionActor::Serialize(FArchive& Ar)
 		{
 			if (UPCGComponent* OriginalComponent = It.Value.Get())
 			{
-				if (OriginalComponent->GetSerializedEditingMode() == EPCGEditorDirtyMode::LoadAsPreview && OriginalComponent->GetSerializedEditingMode() == EPCGEditorDirtyMode::Preview)
+				if (OriginalComponent->GetSerializedEditingMode() == EPCGEditorDirtyMode::LoadAsPreview && OriginalComponent->GetEditingMode() == EPCGEditorDirtyMode::Preview)
 				{
 					LocalToOriginal.Add(It);
 				}
