[PCG] Minor fix around transient state for spawn actor

#rnx

diff --git a/Engine/Plugins/Experimental/PCG/Source/PCG/Private/Elements/PCGSpawnActor.cpp b/Engine/Plugins/Experimental/PCG/Source/PCG/Private/Elements/PCGSpawnActor.cpp
--- a/Engine/Plugins/Experimental/PCG/Source/PCG/Private/Elements/PCGSpawnActor.cpp
+++ b/Engine/Plugins/Experimental/PCG/Source/PCG/Private/Elements/PCGSpawnActor.cpp
@@ -858,7 +858,7 @@ void FPCGSpawnActorElement::SpawnActors(FPCGSubgraphContext* Context, AActor* Ta
 	SpawnParams.SpawnCollisionHandlingOverride = ESpawnActorCollisionHandlingMethod::AlwaysSpawn;
 	SpawnParams.OverrideLevel = TargetActor->GetLevel();
 
-	if (PCGHelpers::IsRuntimeOrPIE())
+	if (PCGHelpers::IsRuntimeOrPIE() || (Context->SourceComponent.IsValid() && Context->SourceComponent->IsInPreviewMode()))
 	{
 		SpawnParams.ObjectFlags |= RF_Transient;
 	}
