[PCG] Added sorting to the debug objects tree

diff --git a/Engine/Plugins/Experimental/PCG/Source/PCGEditor/Private/SPCGEditorGraphDebugObjectTree.cpp b/Engine/Plugins/Experimental/PCG/Source/PCGEditor/Private/SPCGEditorGraphDebugObjectTree.cpp
--- a/Engine/Plugins/Experimental/PCG/Source/PCGEditor/Private/SPCGEditorGraphDebugObjectTree.cpp
+++ b/Engine/Plugins/Experimental/PCG/Source/PCGEditor/Private/SPCGEditorGraphDebugObjectTree.cpp
@@ -38,6 +38,22 @@ FPCGEditorGraphDebugObjectItemPtr FPCGEditorGraphDebugObjectItem::GetParent() co
 	return Parent.Pin();
 }
 
+void FPCGEditorGraphDebugObjectItem::SortChildren(bool bIsAscending, bool bIsRecursive)
+{
+	Children.Sort([bIsAscending](const FPCGEditorGraphDebugObjectItemPtr& InLHS, const FPCGEditorGraphDebugObjectItemPtr& InRHS)
+	{
+		return (InLHS->GetLabel() < InRHS->GetLabel()) == bIsAscending;
+	});
+
+	if (bIsRecursive)
+	{
+		for (const FPCGEditorGraphDebugObjectItemPtr& Child : Children)
+		{
+			Child->SortChildren(bIsAscending, bIsRecursive);
+		}
+	}
+}
+
 FString FPCGEditorGraphDebugObjectItem_Actor::GetLabel() const
 {
 	return Actor.IsValid() ? Actor->GetActorNameOrLabel() : FString();
@@ -59,7 +75,7 @@ FString FPCGEditorGraphDebugObjectItem_PCGSubgraph::GetLabel() const
 	{
 		return PCGNode->GetName() + FString(TEXT(" - ") + PCGGraph->GetName());
 	}
-	
+
 	return FString();
 }
 
@@ -76,7 +92,7 @@ FString FPCGEditorGraphDebugObjectItem_PCGLoopIndex::GetLabel() const
 void SPCGEditorGraphDebugObjectItemRow::Construct(const FArguments& InArgs, const TSharedRef<STableViewBase>& InOwnerTableView, FPCGEditorGraphDebugObjectItemPtr InItem)
 {
 	Item = InItem;
-	
+
 	ChildSlot
 	[
 		SNew(SHorizontalBox)
@@ -115,11 +131,11 @@ void SPCGEditorGraphDebugObjectTree::Construct(const FArguments& InArgs, TShared
 
 	UPCGGraph* PCGGraph = GetPCGGraph();
 	check(PCGGraph);
-	
+
 	FCoreUObjectDelegates::OnPreObjectPropertyChanged.AddSP(this, &SPCGEditorGraphDebugObjectTree::OnPreObjectPropertyChanged);
 	FCoreUObjectDelegates::OnObjectPropertyChanged.AddSP(this, &SPCGEditorGraphDebugObjectTree::OnObjectPropertyChanged);
 	FCoreUObjectDelegates::OnObjectConstructed.AddSP(this, &SPCGEditorGraphDebugObjectTree::OnObjectConstructed);
-	
+
 	const TSharedRef<SScrollBar> HorizontalScrollBar = SNew(SScrollBar)
 		.Orientation(Orient_Horizontal)
 		.Thickness(FVector2D(12.0f, 12.0f));
@@ -127,7 +143,7 @@ void SPCGEditorGraphDebugObjectTree::Construct(const FArguments& InArgs, TShared
 	const TSharedRef<SScrollBar> VerticalScrollBar = SNew(SScrollBar)
 		.Orientation(Orient_Vertical)
 		.Thickness(FVector2D(12.0f, 12.0f));
-	
+
 	DebugObjectTreeView = SNew(STreeView<FPCGEditorGraphDebugObjectItemPtr>)
 				.TreeItemsSource(&RootItems)
 				.OnGenerateRow(this, &SPCGEditorGraphDebugObjectTree::MakeTreeRowWidget)
@@ -149,7 +165,7 @@ void SPCGEditorGraphDebugObjectTree::Construct(const FArguments& InArgs, TShared
 		LOCTEXT("DebugSelectActor", "Select and frame the debug actor in the Level Editor."),
 		TAttribute<bool>::Create(TAttribute<bool>::FGetter::CreateSP(this, &SPCGEditorGraphDebugObjectTree::IsSelectDebugObjectButtonEnabled))
 );
-	
+
 	ChildSlot
 	[
 		SNew(SVerticalBox)
@@ -273,7 +289,7 @@ void SPCGEditorGraphDebugObjectTree::SetDebugObjectFromSelection_OnClicked()
 		{
 			continue;
 		}
-		
+
 		TArray<const UPCGComponent*> PCGComponents;
 		SelectedActor->GetComponents<const UPCGComponent>(PCGComponents, /*bIncludeFromChildActors=*/true);
 
@@ -292,23 +308,23 @@ void SPCGEditorGraphDebugObjectTree::SetDebugObjectFromSelection_OnClicked()
 				{
 					return Item->GetPCGComponent() == PCGComponent && *Item->GetPCGStack() == Stack;
 				});
-			
+
 				if (DebugObjectItem)
 				{
 					DebugObjectTreeView->SetSingleExpandedItem(*DebugObjectItem);
-				
+
 					FPCGEditorGraphDebugObjectItemPtr Parent = DebugObjectItem->Get()->GetParent();
 					while (Parent)
 					{
 						DebugObjectTreeView->SetItemExpansion(Parent, true);
 						Parent = Parent->GetParent();
 					}
-				
+
 					DebugObjectTreeView->SetSelection(*DebugObjectItem);
 					DebugObjectTreeView->RequestScrollIntoView(*DebugObjectItem);
 					break;
 				}
-			}		
+			}
 		}
 	}
 }
@@ -334,7 +350,7 @@ bool SPCGEditorGraphDebugObjectTree::IsSetDebugObjectFromSelectionButtonEnabled(
 		{
 			continue;
 		}
-		
+
 		TArray<const UPCGComponent*> PCGComponents;
 		SelectedActor->GetComponents<const UPCGComponent>(PCGComponents, /*bIncludeFromChildActors=*/true);
 
@@ -344,14 +360,14 @@ bool SPCGEditorGraphDebugObjectTree::IsSetDebugObjectFromSelectionButtonEnabled(
 			{
 				continue;
 			}
-			
+
 			FPCGStackContext StackContext = FPCGStackContext::CreateStackContextFromGraph(PCGComponent->GetGraph());
 
 			const bool bGraphFound = Algo::AnyOf(StackContext.GetStacks(), [&PCGGraph](const FPCGStack& InStack)
 			{
 				return Cast<const UPCGGraph>(InStack.GetStackFrames().Top().Object) == PCGGraph;
 			});
-			
+
 			if(bGraphFound)
 			{
 				return true;
@@ -366,7 +382,7 @@ void SPCGEditorGraphDebugObjectTree::RefreshTree()
 {
 	RootItems.Empty();
 	AllGraphItems.Empty();
-	DebugObjectTreeView->RequestTreeRefresh();	
+	DebugObjectTreeView->RequestTreeRefresh();
 
 	const UPCGGraph* PCGGraph = GetPCGGraph();
 	if (!PCGGraph)
@@ -379,7 +395,7 @@ void SPCGEditorGraphDebugObjectTree::RefreshTree()
 
 	TMap<AActor*, TSharedPtr<FPCGEditorGraphDebugObjectItem_Actor>> ActorItems;
 	TMap<UPCGComponent*, TSharedPtr<FPCGEditorGraphDebugObjectItem_PCGComponent>> ComponentItems;
-	
+
 	for (UObject* PCGComponentObject : PCGComponents)
 	{
 		if (!IsValid(PCGComponentObject) || PCGComponentObject->HasAnyFlags(RF_Transient))
@@ -494,7 +510,7 @@ void SPCGEditorGraphDebugObjectTree::RefreshTree()
 								{
 									CurrentItem = NodeItems.Add_GetRef(MakeShared<FPCGEditorGraphDebugObjectItem_PCGSubgraph>(StackNode, NextStackGraph, NextStackGraph == TopStackGraph ? Stack : FPCGStack()));
 								}
-								
+
 								StackIndex++;
 							}
 						}
@@ -614,11 +630,29 @@ void SPCGEditorGraphDebugObjectTree::RefreshTree()
 			}
 		}
 	}
-	
+
 	for (TPair<AActor*, TSharedPtr<FPCGEditorGraphDebugObjectItem_Actor>>& ActorItem : ActorItems)
-	{	
+	{
 		RootItems.Add(MoveTemp(ActorItem.Value));
 	}
+
+	SortTreeItems();
+}
+
+void SPCGEditorGraphDebugObjectTree::SortTreeItems(bool bIsAscending, bool bIsRecursive)
+{
+	RootItems.Sort([bIsAscending](const FPCGEditorGraphDebugObjectItemPtr& InLHS, const FPCGEditorGraphDebugObjectItemPtr& InRHS)
+	{
+		return (InLHS->GetLabel() < InRHS->GetLabel()) == bIsAscending;
+	});
+
+	if (bIsRecursive)
+	{
+		for (const FPCGEditorGraphDebugObjectItemPtr& Item : RootItems)
+		{
+			Item->SortChildren(bIsAscending, bIsRecursive);
+		}
+	}
 }
 
 void SPCGEditorGraphDebugObjectTree::OnPreObjectPropertyChanged(UObject* InObject, const FEditPropertyChain& /*InPropertyChain*/)
@@ -634,7 +668,7 @@ void SPCGEditorGraphDebugObjectTree::OnPreObjectPropertyChanged(UObject* InObjec
 }
 
 void SPCGEditorGraphDebugObjectTree::OnObjectPropertyChanged(UObject* InObject, FPropertyChangedEvent& InPropertyChangedEvent)
-{	
+{
 	if (const UPCGGraphInterface* PCGGraphInterface = Cast<UPCGGraphInterface>(InObject))
 	{
 		const UPCGGraph* PCGGraph = PCGGraphInterface->GetGraph();
@@ -688,7 +722,7 @@ void SPCGEditorGraphDebugObjectTree::OnSelectionChanged(FPCGEditorGraphDebugObje
 		PCGEditor.Pin()->SetComponentAndStackBeingInspected(nullptr, FPCGStack());
 		return;
 	}
-	
+
 	if (UPCGComponent* PCGComponent = InItem->GetPCGComponent())
 	{
 		if (const FPCGStack* PCGStack = InItem->GetPCGStack())
diff --git a/Engine/Plugins/Experimental/PCG/Source/PCGEditor/Private/SPCGEditorGraphDebugObjectTree.h b/Engine/Plugins/Experimental/PCG/Source/PCGEditor/Private/SPCGEditorGraphDebugObjectTree.h
--- a/Engine/Plugins/Experimental/PCG/Source/PCGEditor/Private/SPCGEditorGraphDebugObjectTree.h
+++ b/Engine/Plugins/Experimental/PCG/Source/PCGEditor/Private/SPCGEditorGraphDebugObjectTree.h
@@ -26,12 +26,13 @@ public:
 	void AddChild(TSharedRef<FPCGEditorGraphDebugObjectItem> InChild);
 	const TSet<TSharedPtr<FPCGEditorGraphDebugObjectItem>>& GetChildren() const;
 	FPCGEditorGraphDebugObjectItemPtr GetParent() const;
+	void SortChildren(bool bIsAscending, bool bIsRecursive);
 
-	virtual FString GetLabel() const = 0;	
+	virtual FString GetLabel() const = 0;
 	virtual UPCGComponent* GetPCGComponent() const = 0;
 	virtual const FPCGStack* GetPCGStack() const { return nullptr; }
 	virtual const UObject* GetObject() const  = 0;
-	
+
 protected:
 	TWeakPtr<FPCGEditorGraphDebugObjectItem> Parent;
 	TSet<TSharedPtr<FPCGEditorGraphDebugObjectItem>> Children;
@@ -47,7 +48,7 @@ public:
 	virtual FString GetLabel() const override;
 	virtual UPCGComponent* GetPCGComponent() const override { return nullptr; }
 	virtual const UObject* GetObject() const override { return Actor.Get(); };
-	
+
 protected:
 	TWeakObjectPtr<AActor> Actor = nullptr;
 };
@@ -62,7 +63,7 @@ public:
 	virtual FString GetLabel() const override;
 	virtual UPCGComponent* GetPCGComponent() const override { return PCGComponent.Get(); }
 	virtual const UObject* GetObject() const override { return PCGComponent.Get(); };
-	
+
 protected:
 	TWeakObjectPtr<UPCGComponent> PCGComponent = nullptr;
 };
@@ -79,7 +80,7 @@ public:
 	virtual UPCGComponent* GetPCGComponent() const override { return GetParent() ? GetParent()->GetPCGComponent() : nullptr; }
 	virtual const UObject* GetObject() const override { return PCGGraph.Get(); };
 	virtual const FPCGStack* GetPCGStack() const override { return &PCGStack; }
-	
+
 protected:
 	TWeakObjectPtr<const UPCGGraph> PCGGraph = nullptr;
 	FPCGStack PCGStack;
@@ -98,7 +99,7 @@ public:
 	virtual UPCGComponent* GetPCGComponent() const override { return GetParent() ? GetParent()->GetPCGComponent() : nullptr; }
 	virtual const UObject* GetObject() const override { return PCGNode.Get(); };
 	virtual const FPCGStack* GetPCGStack() const override { return &PCGStack; }
-	
+
 protected:
 	TWeakObjectPtr<const UPCGNode> PCGNode = nullptr;
 	TWeakObjectPtr<const UPCGGraph> PCGGraph = nullptr;
@@ -119,7 +120,7 @@ public:
 	virtual const UObject* GetObject() const override { return PCGNode.Get(); };
 	virtual const FPCGStack* GetPCGStack() const override { return &PCGStack; }
 	virtual int32 GetLoopIndex() const { return LoopIndex; }
-	
+
 protected:
 	TWeakObjectPtr<const UPCGNode> PCGNode = nullptr;
 	int32 LoopIndex = INDEX_NONE;
@@ -133,7 +134,7 @@ public:
 	SLATE_END_ARGS()
 
 	void Construct(const FArguments& InArgs, const TSharedRef<STableViewBase>& InOwnerTableView, FPCGEditorGraphDebugObjectItemPtr InItem);
-	
+
 private:
 	FPCGEditorGraphDebugObjectItemPtr Item;
 };
@@ -145,35 +146,36 @@ public:
 	SLATE_END_ARGS()
 
 	virtual ~SPCGEditorGraphDebugObjectTree();
-	
+
 	void Construct(const FArguments& InArgs, TSharedPtr<FPCGEditor> InPCGEditor);
 
 	virtual void Tick(const FGeometry& AllottedGeometry, const double InCurrentTime, const float InDeltaTime) override;
-	
+
 	void RequestRefresh() { bNeedsRefresh = true; }
 
 	void AddDynamicStack(TWeakObjectPtr<UPCGComponent> InComponent, const FPCGStack& InvocationStack);
-	
+
 private:
 	void SelectedDebugObject_OnClicked() const;
 	bool IsSelectDebugObjectButtonEnabled() const;
-	
+
 	void SetDebugObjectFromSelection_OnClicked();
 	bool IsSetDebugObjectFromSelectionButtonEnabled() const;
-	
+
 	void RefreshTree();
-	
+	void SortTreeItems(bool bIsAscending = true, bool bIsRecursive = true);
+
 	void OnPreObjectPropertyChanged(UObject* InObject, const FEditPropertyChain& InPropertyChain);
 	void OnObjectPropertyChanged(UObject* InObject, FPropertyChangedEvent& InPropertyChangedEvent);
 	void OnObjectConstructed(UObject* InObject);
 
 	UPCGGraph* GetPCGGraph() const;
-	
+
 	TSharedRef<ITableRow> MakeTreeRowWidget(FPCGEditorGraphDebugObjectItemPtr InItem, const TSharedRef<STableViewBase>& InOwnerTable) const;
 	void OnGetChildren(FPCGEditorGraphDebugObjectItemPtr InItem, TArray<FPCGEditorGraphDebugObjectItemPtr>& OutChildren) const;
 	void OnSelectionChanged(FPCGEditorGraphDebugObjectItemPtr InItem, ESelectInfo::Type InSelectInfo) const;
 	void OnSetExpansionRecursive(FPCGEditorGraphDebugObjectItemPtr InItem, bool bInExpand);
-	
+
 	TWeakPtr<FPCGEditor> PCGEditor;
 
 	TSharedPtr<STreeView<FPCGEditorGraphDebugObjectItemPtr>> DebugObjectTreeView;
