[PCG] Use the new set function to set number of custom data floats for ISMC's to make sure the data manger is updated

#rnx

diff --git a/Engine/Plugins/Experimental/PCG/Source/PCG/Private/Helpers/PCGActorHelpers.cpp b/Engine/Plugins/Experimental/PCG/Source/PCG/Private/Helpers/PCGActorHelpers.cpp
--- a/Engine/Plugins/Experimental/PCG/Source/PCG/Private/Helpers/PCGActorHelpers.cpp
+++ b/Engine/Plugins/Experimental/PCG/Source/PCG/Private/Helpers/PCGActorHelpers.cpp
@@ -125,7 +125,7 @@ UPCGManagedISMComponent* UPCGActorHelpers::GetOrCreateManagedISMC(AActor* InTarg
 	const EObjectFlags ObjectFlags = (InSourceComponent->IsInPreviewMode() ? RF_Transient : RF_NoFlags);
 	UInstancedStaticMeshComponent* ISMC = NewObject<UInstancedStaticMeshComponent>(InTargetActor, ComponentClass, MakeUniqueObjectName(InTargetActor, ComponentClass, FName(ComponentName)), ObjectFlags);
 	InParams.Descriptor.InitComponent(ISMC);
-	ISMC->NumCustomDataFloats = InParams.NumCustomDataFloats;
+	ISMC->SetNumCustomDataFloats(InParams.NumCustomDataFloats);
 
 	ISMC->RegisterComponent();
 	InTargetActor->AddInstanceComponent(ISMC);
