[PCG] Minor - fix format string in graph validation log message to include graph name

#rnx

diff --git a/Engine/Plugins/Experimental/PCG/Source/PCG/Private/PCGGraph.cpp b/Engine/Plugins/Experimental/PCG/Source/PCG/Private/PCGGraph.cpp
--- a/Engine/Plugins/Experimental/PCG/Source/PCG/Private/PCGGraph.cpp
+++ b/Engine/Plugins/Experimental/PCG/Source/PCG/Private/PCGGraph.cpp
@@ -1127,7 +1127,7 @@ void UPCGGraph::FixInvalidEdges()
 				else if (!ConnectedNode || (GetInputNode() != ConnectedNode && GetOutputNode() != ConnectedNode && !Nodes.Contains(ConnectedNode)))
 				{
 					// Remove edges to nodes that are not present in the graph.
-					UE_LOG(LogPCG, Error, TEXT("Removed edge to a node '%s' that is not registered in the graph."),
+					UE_LOG(LogPCG, Error, TEXT("Removed edge to a node '%s' that is not registered in graph '%s'."),
 						ConnectedNode ? *ConnectedNode->GetFName().ToString() : TEXT("NULL"),
 						*GetFName().ToString());
 
