[PCG] Minor hotfix for a bug causing a crash in sampling volume

#rnx

diff --git a/Engine/Plugins/Experimental/PCG/Source/PCG/Private/Elements/PCGVolumeSampler.cpp b/Engine/Plugins/Experimental/PCG/Source/PCG/Private/Elements/PCGVolumeSampler.cpp
--- a/Engine/Plugins/Experimental/PCG/Source/PCG/Private/Elements/PCGVolumeSampler.cpp
+++ b/Engine/Plugins/Experimental/PCG/Source/PCG/Private/Elements/PCGVolumeSampler.cpp
@@ -31,7 +31,8 @@ namespace PCGVolumeSampler
 		UPCGPointData* Data = NewObject<UPCGPointData>();
 		Data->InitializeFromData(Volume);
 
-		SampleVolume(Context, SamplerSettings, Volume, BoundingShape, Data, Context->TimeSliceIsEnabled());
+		const bool bTimeSliceIsEnabled = Context ? Context->TimeSliceIsEnabled() : false;
+		SampleVolume(Context, SamplerSettings, Volume, BoundingShape, Data, bTimeSliceIsEnabled);
 
 		return Data;
 	}
