[PCG] Runtime Generation Scheduling

Added FPCGRuntimeGenScheduler. Gets ticked by the Subsystem to efficiently schedule components for generation/cleanup every frame. To be considered, a component must set its generation trigger to GenerateAtRuntime. Note that external Generate/Cleanup are disabled on these components, since they are now managed exclusively by the scheduler.

1. Queue nearby components for generation.
2. Schedule cleanup on components that become out of range
3. Schedule generation on the components we queued in step 1

"Nearby" means that a component is within some distance of a UPCGGenerationSource. Generation sources represent objects in the world that should provoke runtime generation. Examples are PlayerControllers, StreamingSources, GenerationSource Components, and EditorViewport cameras.

Components are scheduled in an efficient order because the generation queue is sorted by "Priority". Priority is calculated by a UPCGSchedulingPolicy, which lives on the Component. Scheduling Policies use a Generation Source to compute their priority. The default Scheduling Policy uses distance and view direction. If multiple Generation Sources are near a component, the scheduler will choose the max priority.

The Scheduler keeps track of GeneratedKeys. GeneratedKeys represent a grid size, grid coordinate, and Original Component. For partitioned and non-partitioned components, the grid size should always be EPCGHiGenGrid::Unbounded, and the grid coordinates should be (0, 0, 0). For Local Components, the grid size and coordinates should reflect the local component's actual grid cell.

Note that the GeneratedKey for a Local Component actually holds the *Original Component*, and not the Local Component itself. This is because at scheduling time, the local component does not exist yet!

When the Scheduler decides to schedule a grid cell/Local Component for generation/cleanup, it first needs to make sure a PartitionActor and LocalComponent actually exist. There are two ways this can happen.

1. The PA is created or destroyed, respectively
2. The PA is popped or pushed into a pool, respectively

In both cases, the LocalComponent is created/destroyed on demand and applied to the PartitionActor. Pooling is enabled by default, but can be disabled with a CVar.

diff --git a/Engine/Plugins/Experimental/PCG/Source/PCG/Private/Elements/PCGDataFromActor.cpp b/Engine/Plugins/Experimental/PCG/Source/PCG/Private/Elements/PCGDataFromActor.cpp
--- a/Engine/Plugins/Experimental/PCG/Source/PCG/Private/Elements/PCGDataFromActor.cpp
+++ b/Engine/Plugins/Experimental/PCG/Source/PCG/Private/Elements/PCGDataFromActor.cpp
@@ -155,7 +155,7 @@ bool FPCGDataFromActorElement::ExecuteInternal(FPCGContext* InContext) const
 		if (Self && Settings->ActorSelector.bMustOverlapSelf)
 		{
 			// Capture ActorBounds by value because it goes out of scope
-			const FBox ActorBounds = PCGHelpers::GetActorBounds(Self);
+			FBox ActorBounds = PCGHelpers::GetGridBounds(Self, PCGComponent);
 			BoundsCheck = [Settings, ActorBounds, PCGComponent](const AActor* OtherActor) -> bool
 			{
 				const FBox OtherActorBounds = OtherActor ? PCGHelpers::GetGridBounds(OtherActor, PCGComponent) : FBox(EForceInit::ForceInit);
diff --git a/Engine/Plugins/Experimental/PCG/Source/PCG/Private/Graph/PCGGraphExecutor.cpp b/Engine/Plugins/Experimental/PCG/Source/PCG/Private/Graph/PCGGraphExecutor.cpp
--- a/Engine/Plugins/Experimental/PCG/Source/PCG/Private/Graph/PCGGraphExecutor.cpp
+++ b/Engine/Plugins/Experimental/PCG/Source/PCG/Private/Graph/PCGGraphExecutor.cpp
@@ -1548,7 +1548,10 @@ namespace PCGGraphExecutor
 				{
 					// Get grid coords using the parent grid (FromGridSize).
 					const FIntVector CellCoords = UPCGActorHelpers::GetCellCoord(ComponentActor->GetActorLocation(), FromGridSize, PCGWorldActor->bUse2DGrid);
-					ComponentWithData = Subsystem->GetLocalComponent(FromGridSize, CellCoords, InContext->SourceComponent->GetOriginalComponent());
+
+					// Search for a transient local component if the source component is runtime managed.
+					const bool bTransientComponent = InContext->SourceComponent->GenerationTrigger == EPCGComponentGenerationTrigger::GenerateAtRuntime;
+					ComponentWithData = Subsystem->GetLocalComponent(FromGridSize, CellCoords, InContext->SourceComponent->GetOriginalComponent(), bTransientComponent);
 				}
 			}
 
@@ -1607,8 +1610,11 @@ namespace PCGGraphExecutor
 			{
 				PCGGraphExecutionLogging::LogGridLinkageTaskExecuteRetrieveScheduleGraph(InContext, ComponentWithData, InResourceKey);
 
+				EPCGComponentGenerationTrigger GenTrigger = (InContext->SourceComponent->GenerationTrigger == EPCGComponentGenerationTrigger::GenerateAtRuntime) ?
+					EPCGComponentGenerationTrigger::GenerateAtRuntime : EPCGComponentGenerationTrigger::GenerateOnDemand;
+
 				// Wake up this task after graph has generated.
-				const FPCGTaskId GraphTaskId = ComponentWithData->GenerateLocalGetTaskId(/*bForce=*/true);
+				const FPCGTaskId GraphTaskId = ComponentWithData->GenerateLocalGetTaskId(GenTrigger, /*bForce=*/true);
 				Subsystem->ScheduleGeneric(WakeUpLambda, InContext->SourceComponent.Get(), { GraphTaskId });
 
 				// Update state and go to sleep.
diff --git a/Engine/Plugins/Experimental/PCG/Source/PCG/Private/Grid/PCGPartitionActor.cpp b/Engine/Plugins/Experimental/PCG/Source/PCG/Private/Grid/PCGPartitionActor.cpp
--- a/Engine/Plugins/Experimental/PCG/Source/PCG/Private/Grid/PCGPartitionActor.cpp
+++ b/Engine/Plugins/Experimental/PCG/Source/PCG/Private/Grid/PCGPartitionActor.cpp
@@ -1,7 +1,7 @@
 // Copyright Epic Games, Inc. All Rights Reserved.
 
 #include "PCGPartitionActor.h"
-#include "Engine/World.h"
+
 #include "PCGComponent.h"
 #include "PCGModule.h"
 #include "PCGSubsystem.h"
@@ -10,6 +10,7 @@
 #include "Helpers/PCGHelpers.h"
 
 #include "Components/BoxComponent.h"
+#include "Engine/World.h"
 
 #include UE_INLINE_GENERATED_CPP_BY_NAME(PCGPartitionActor)
 
@@ -32,11 +33,6 @@ APCGPartitionActor::APCGPartitionActor(const FObjectInitializer& ObjectInitializ
 #endif // WITH_EDITOR
 }
 
-UPCGSubsystem* APCGPartitionActor::GetSubsystem() const 
-{
-	return UPCGSubsystem::GetInstance(GetWorld());
-}
-
 void APCGPartitionActor::PostLoad()
 {
 	Super::PostLoad();
@@ -69,7 +65,6 @@ void APCGPartitionActor::PostLoad()
 	// Make sure that we don't track objects that do not exist anymore
 	CleanupDeadGraphInstances(/*bRemoveNullOnly=*/true);
 
-#if WITH_EDITOR
 	// Mark all our local components as local
 	for (UPCGComponent* LocalComponent : GetAllLocalPCGComponents())
 	{
@@ -78,17 +73,14 @@ void APCGPartitionActor::PostLoad()
 	}
 
 	bWasPostCreatedLoaded = true;
-#endif // WITH_EDITOR
 }
 
 void APCGPartitionActor::BeginDestroy()
 {
-#if WITH_EDITOR
-	if (!PCGHelpers::IsRuntimeOrPIE() && GetSubsystem() && PCGGridSize > 0)
+	if (!PCGHelpers::IsRuntimeOrPIE())
 	{
-		GetSubsystem()->UnregisterPartitionActor(this);
+		UnregisterPCG();
 	}
-#endif // WITH_EDITOR
 
 	Super::BeginDestroy();
 }
@@ -139,12 +131,10 @@ void APCGPartitionActor::Serialize(FArchive& Ar)
 
 void APCGPartitionActor::Destroyed()
 {
-#if WITH_EDITOR
-	if (!PCGHelpers::IsRuntimeOrPIE() && GetSubsystem() && PCGGridSize > 0)
+	if (!PCGHelpers::IsRuntimeOrPIE())
 	{
-		GetSubsystem()->UnregisterPartitionActor(this);
+		UnregisterPCG();
 	}
-#endif // WITH_EDITOR
 
 	Super::Destroyed();
 }
@@ -171,39 +161,34 @@ void APCGPartitionActor::PostRegisterAllComponents()
 	// Make the Partition actor register itself to the PCG Subsystem
 	// Always do it at runtime, wait for the post load/creation in editor
 	// Only do the mapping if we are at runtime.
-
-	// Make sure the PCGGridSize is not 0, otherwise it will break everything
-
+	// Skip if RuntimeGen, RG PAs only register in PostCreate.
 	const bool bIsRuntimeOrPIE = PCGHelpers::IsRuntimeOrPIE();
-	if ((bIsRuntimeOrPIE 
-#if WITH_EDITOR
-		|| bWasPostCreatedLoaded
-#endif // WITH_EDITOR
-		) && ensure(PCGGridSize > 0))
+	if (!IsRuntimeGenerated() && !HasAnyFlags(RF_Transient) && (bIsRuntimeOrPIE || bWasPostCreatedLoaded))
 	{
-		if (UPCGSubsystem* Subsystem = GetSubsystem())
-		{
-			Subsystem->RegisterPartitionActor(this, /*bDoComponentMapping*/ bIsRuntimeOrPIE);
-		}
+		RegisterPCG(bIsRuntimeOrPIE);
 	}
 }
 
 void APCGPartitionActor::BeginPlay()
 {
-	// Pass through all the pcg components, to verify if we need to generate them
-	for (auto& It : OriginalToLocal)
+	if (!IsRuntimeGenerated() && !HasAnyFlags(RF_Transient))
 	{
-		if (It.Key && It.Value)
+		// Pass through all the pcg components, to verify if we need to generate them.
+		for (auto& It : OriginalToLocal)
 		{
-			// If we have an original component that is generated (or generating), this one is automatically generated => GenerateOnLoad
-			if (It.Key->bGenerated || It.Key->IsGenerating())
+			if (It.Key && It.Value)
 			{
-				It.Value->GenerationTrigger = EPCGComponentGenerationTrigger::GenerateOnLoad;
-			}
-			// Otherwise, make them match
-			else
-			{
-				It.Value->GenerationTrigger = It.Key->GenerationTrigger;
+				// If we have an original component that is generated (or generating), this one is automatically generated => GenerateOnLoad.
+				// But if its runtime generated then it's handled by the runtime generation scheduler.
+				if ((It.Key->bGenerated || It.Key->IsGenerating()) && ensure(It.Key->GenerationTrigger != EPCGComponentGenerationTrigger::GenerateAtRuntime))
+				{
+					It.Value->GenerationTrigger = EPCGComponentGenerationTrigger::GenerateOnLoad;
+				}
+				// Otherwise, make them match
+				else
+				{
+					It.Value->GenerationTrigger = It.Key->GenerationTrigger;
+				}
 			}
 		}
 	}
@@ -213,10 +198,7 @@ void APCGPartitionActor::BeginPlay()
 
 void APCGPartitionActor::EndPlay(const EEndPlayReason::Type EndPlayReason)
 {
-	if (UPCGSubsystem* Subsystem = GetSubsystem())
-	{
-		Subsystem->UnregisterPartitionActor(this);
-	}
+	UnregisterPCG();
 
 	Super::EndPlay(EndPlayReason);
 }
@@ -256,6 +238,42 @@ FIntVector APCGPartitionActor::GetGridCoord() const
 	return UPCGActorHelpers::GetCellCoord(Center, PCGGridSize, bUse2DGrid);
 }
 
+bool APCGPartitionActor::Teleport(const FVector& NewLocation)
+{
+	// We should not be teleporting a PA that is in use. We only teleport empty RuntimeGen PAs to their grid cell before initialization.
+	check(bIsRuntimeGenerated && OriginalToLocal.IsEmpty() && LocalToOriginal.IsEmpty());
+
+	RootComponent->SetMobility(EComponentMobility::Movable);
+	bool bResult = Super::SetActorLocation(NewLocation, /*bSweep=*/false, nullptr, ETeleportType::TeleportPhysics);
+	RootComponent->SetMobility(EComponentMobility::Static);
+
+	return bResult;
+}
+
+void APCGPartitionActor::RegisterPCG(bool bDoComponentMapping)
+{
+	if (UPCGSubsystem* Subsystem = GetSubsystem())
+	{
+		if (!bIsRegistered)
+		{
+			Subsystem->RegisterPartitionActor(this, bDoComponentMapping);
+			bIsRegistered = true;
+		}
+	}
+}
+
+void APCGPartitionActor::UnregisterPCG()
+{
+	if (UPCGSubsystem* Subsystem = GetSubsystem())
+	{
+		if (bIsRegistered)
+		{
+			Subsystem->UnregisterPartitionActor(this);
+			bIsRegistered = false;
+		}
+	}
+}
+
 void APCGPartitionActor::GetActorBounds(bool bOnlyCollidingComponents, FVector& Origin, FVector& BoxExtent, bool bIncludeFromChildActors) const
 {
 	Super::GetActorBounds(bOnlyCollidingComponents, Origin, BoxExtent, bIncludeFromChildActors);
@@ -476,10 +494,18 @@ AActor* APCGPartitionActor::GetSceneOutlinerParent() const
 		return Super::GetSceneOutlinerParent();
 	}	
 }
+#endif // WITH_EDITOR
 
 void APCGPartitionActor::PostCreation(const FGuid& InGridGUID)
 {
-	PCGGridSize = Super::GetGridSize();
+#if WITH_EDITOR
+	// For RuntimeGen we do not care about the grid size from WP, we write it ourselves.
+	if (!IsRuntimeGenerated())
+	{
+		PCGGridSize = Super::GetGridSize();
+	}
+#endif
+
 	PCGGuid = InGridGUID;
 
 	// Put in cache if we use the 2D grid or not.
@@ -498,15 +524,15 @@ void APCGPartitionActor::PostCreation(const FGuid& InGridGUID)
 	{
 		BoundsComponent->SetBoxExtent(GetFixedBounds().GetExtent());
 	}
+#endif // WITH_EDITOR
 
-	// Make sure PCGGrid size if greater than 0, otherwise it will break everything
-	if (!PCGHelpers::IsRuntimeOrPIE() && GetSubsystem() && ensure(PCGGridSize > 0))
+	// RuntimeGen PAs always register in PostCreate. Non-RG should only register here if in-editor
+	if (IsRuntimeGenerated() || !PCGHelpers::IsRuntimeOrPIE())
 	{
-		GetSubsystem()->RegisterPartitionActor(this, /*bDoComponentMapping=*/ false);
+		RegisterPCG(/*bDoComponentMapping=*/false);
 	}
 
 	bWasPostCreatedLoaded = true;
-#endif // WITH_EDITOR
 }
 
 bool APCGPartitionActor::IsSafeForDeletion() const
@@ -542,6 +568,12 @@ TSet<TObjectPtr<UPCGComponent>> APCGPartitionActor::GetAllOriginalPCGComponents(
 	return ResultComponents;
 }
 
+UPCGSubsystem* APCGPartitionActor::GetSubsystem() const 
+{
+	return UPCGSubsystem::GetInstance(GetWorld());
+}
+
+#if WITH_EDITOR
 bool APCGPartitionActor::ChangeTransientState(UPCGComponent* OriginalComponent, EPCGEditorDirtyMode EditingMode)
 {
 	check(OriginalComponent);
@@ -580,5 +612,14 @@ bool APCGPartitionActor::ChangeTransientState(UPCGComponent* OriginalComponent,
 
 	return OriginalToLocal.IsEmpty() && LoadedPreviewComponents.IsEmpty();
 }
-
 #endif // WITH_EDITOR
+
+FString APCGPartitionActor::GetRuntimeGenActorName(uint32 GridSize, const FIntVector& GridCoords)
+{
+	TStringBuilderWithBuffer<TCHAR, NAME_SIZE> ActorNameBuilder;
+	ActorNameBuilder += TEXT("PCGRuntimeGenPartitionActor_");
+	ActorNameBuilder += FString::Printf(TEXT("%d_"), GridSize);
+	ActorNameBuilder += FString::Printf(TEXT("%d_%d_%d"), GridCoords.X, GridCoords.Y, GridCoords.Z);
+
+	return ActorNameBuilder.ToString();
+}
diff --git a/Engine/Plugins/Experimental/PCG/Source/PCG/Private/Grid/PCGPartitionActor.h b/Engine/Plugins/Experimental/PCG/Source/PCG/Private/Grid/PCGPartitionActor.h
--- a/Engine/Plugins/Experimental/PCG/Source/PCG/Private/Grid/PCGPartitionActor.h
+++ b/Engine/Plugins/Experimental/PCG/Source/PCG/Private/Grid/PCGPartitionActor.h
@@ -55,9 +55,24 @@ public:
 	FBox GetFixedBounds() const;
 	FIntVector GetGridCoord() const;
 
+	void SetPCGGridSize(uint32 InPCGGridSize) { PCGGridSize = InPCGGridSize; }
 	uint32 GetPCGGridSize() const { return PCGGridSize; }
+
 	bool IsUsing2DGrid() const { return bUse2DGrid; }
 
+	/** Marks this PartitionActor as managed by the runtime generation system. */
+	void SetToRuntimeGenerated() { bIsRuntimeGenerated = true; }
+	bool IsRuntimeGenerated() const { return bIsRuntimeGenerated; }
+
+	/** Forces the actor location to change even if its mobility is static. */
+	bool Teleport(const FVector& NewLocation);
+
+	/** Register with the PCG Subsystem. */
+	void RegisterPCG(bool bDoComponentMapping);
+
+	/** Unregister with the PCG Subsystem. */
+	void UnregisterPCG();
+
 	void AddGraphInstance(UPCGComponent* OriginalComponent);
 	void RemapGraphInstance(const UPCGComponent* OldOriginalComponent, UPCGComponent* NewOriginalComponent);
 	bool RemoveGraphInstance(UPCGComponent* OriginalComponent);
@@ -66,7 +81,6 @@ public:
 	// When a local component is destroyed. It calls this function. We make sure we don't keep mappings that are dead.
 	void RemoveLocalComponent(UPCGComponent* LocalComponent);
 
-#if WITH_EDITOR
 	/** To be called after the creation of a new actor to copy the GridSize property (Editor only) into the PCGGridSize property */
 	void PostCreation(const FGuid& InGridGUID);
 
@@ -81,7 +95,6 @@ public:
 
 	/** Changes transient state for the local component matching the given original component. Returns true if PA becomes empty */
 	bool ChangeTransientState(UPCGComponent* OriginalComponent, EPCGEditorDirtyMode EditingMode);
-#endif
 
 	// TODO: Make this in-editor only; during runtime, we should keep a map of component to bounds/volume only
 	// and preferably precompute the intersection, so this would make it easier/possible to not have the original actor in game version.
@@ -122,8 +135,20 @@ private:
 	/** Box component to draw the Partition actor bounds in the Editor viewport */
 	UPROPERTY(Transient)
 	TObjectPtr<UBoxComponent> BoundsComponent;
+#endif // WITH_EDITORONLY_DATA
+
+	/** Tracks the registration status of this PA with the ActorAndComponentMapping system. Helps us avoid invalid (un)registers. */
+	bool bIsRegistered = false;
+
+	/** Tracks if this actor was created by the Runtime Generation system. */
+	bool bIsRuntimeGenerated = false;
 
 	/** Utility bool to check if PostCreation/PostLoad was called. */
 	bool bWasPostCreatedLoaded = false;
-#endif // WITH_EDITORONLY_DATA
+
+public:
+	/** Gets the name this actor would have if it were Runtime Generated.
+	* This does not respect traditional PA name contents like GridGuid, ShouldIncludeGridSizeInName, or ContextHash.
+	*/
+	static FString GetRuntimeGenActorName(uint32 GridSize, const FIntVector& GridCoords);
 };
diff --git a/Engine/Plugins/Experimental/PCG/Source/PCG/Private/PCGActorAndComponentMapping.cpp b/Engine/Plugins/Experimental/PCG/Source/PCG/Private/PCGActorAndComponentMapping.cpp
--- a/Engine/Plugins/Experimental/PCG/Source/PCG/Private/PCGActorAndComponentMapping.cpp
+++ b/Engine/Plugins/Experimental/PCG/Source/PCG/Private/PCGActorAndComponentMapping.cpp
@@ -105,10 +105,18 @@ void FPCGActorAndComponentMapping::Tick()
 TArray<FPCGTaskId> FPCGActorAndComponentMapping::DispatchToRegisteredLocalComponents(UPCGComponent* OriginalComponent, const TFunction<FPCGTaskId(UPCGComponent*)>& InFunc) const
 {
 	TRACE_CPUPROFILER_EVENT_SCOPE(FPCGActorAndComponentMapping::DispatchToRegisteredLocalComponents);
+	if (!ensure(OriginalComponent))
+	{
+		return {};
+	}
+
+	const bool bIsRuntimeGenerated = OriginalComponent->GenerationTrigger == EPCGComponentGenerationTrigger::GenerateAtRuntime;
+
+	const TMap<const UPCGComponent*, TSet<TObjectPtr<APCGPartitionActor>>>& Map = bIsRuntimeGenerated ? ComponentToRuntimeGenPartitionActorsMap : ComponentToPartitionActorsMap;
 
 	// TODO: Might be more interesting to copy the set and release the lock.
-	FReadScopeLock ReadLock(ComponentToPartitionActorsMapLock);
-	const TSet<TObjectPtr<APCGPartitionActor>>* PartitionActorsPtr = ComponentToPartitionActorsMap.Find(OriginalComponent);
+	FReadScopeLock ReadLock(bIsRuntimeGenerated ? ComponentToRuntimeGenPartitionActorsMapLock : ComponentToPartitionActorsMapLock);
+	const TSet<TObjectPtr<APCGPartitionActor>>* PartitionActorsPtr = Map.Find(OriginalComponent);
 
 	if (!PartitionActorsPtr)
 	{
@@ -211,10 +219,10 @@ bool FPCGActorAndComponentMapping::RegisterOrUpdatePartitionedPCGComponent(UPCGC
 	PartitionedOctree.AddOrUpdateComponent(InComponent, Bounds, bComponentHasChanged, bComponentWasAdded);
 
 #if WITH_EDITOR
-	// In Editor only, we will create new partition actors depending on the new bounds
+	// In Editor only, we will create new partition actors depending on the new bounds and generation trigger. Runtime managed components should not create PAs here
 	// TODO: For now it will always create the PA. But if we want to create them only when we generate, we need to make
 	// sure to update the runtime flow, for them to also create PA if they need to.
-	if (bComponentHasChanged || bComponentWasAdded)
+	if ((bComponentHasChanged || bComponentWasAdded) && InComponent->GenerationTrigger != EPCGComponentGenerationTrigger::GenerateAtRuntime)
 	{
 		bool bHasUnbounded = false;
 		PCGHiGenGrid::FSizeArray GridSizes;
@@ -287,23 +295,26 @@ bool FPCGActorAndComponentMapping::RemapPCGComponent(const UPCGComponent* OldCom
 	}
 
 	// Remap all previous instances
+	auto RemapPreviousInstances = [OldComponent, NewComponent](TMap<const UPCGComponent*, TSet<TObjectPtr<APCGPartitionActor>>>& Map, FRWLock& Lock)
 	{
-		FWriteScopeLock WriteLock(ComponentToPartitionActorsMapLock);
-		TSet<TObjectPtr<APCGPartitionActor>>* PartitionActorsPtr = ComponentToPartitionActorsMap.Find(OldComponent);
+		FWriteScopeLock WriteLock(Lock);
 
-		if (PartitionActorsPtr)
+		if (TSet<TObjectPtr<APCGPartitionActor>>* PartitionActorsPtr = Map.Find(OldComponent))
 		{
 			TSet<TObjectPtr<APCGPartitionActor>> PartitionActorsToRemap = MoveTemp(*PartitionActorsPtr);
-			ComponentToPartitionActorsMap.Remove(OldComponent);
+			Map.Remove(OldComponent);
 
 			for (APCGPartitionActor* Actor : PartitionActorsToRemap)
 			{
 				Actor->RemapGraphInstance(OldComponent, NewComponent);
 			}
 
-			ComponentToPartitionActorsMap.Add(NewComponent, MoveTemp(PartitionActorsToRemap));
+			Map.Add(NewComponent, MoveTemp(PartitionActorsToRemap));
 		}
-	}
+	};
+
+	RemapPreviousInstances(ComponentToPartitionActorsMap, ComponentToPartitionActorsMapLock);
+	RemapPreviousInstances(ComponentToRuntimeGenPartitionActorsMap, ComponentToRuntimeGenPartitionActorsMapLock);
 
 	// And update the mapping if bounds changed and we want to do actor mapping
 	if (bBoundsChanged && NewComponent->IsPartitioned() && bDoActorMapping)
@@ -362,7 +373,14 @@ void FPCGActorAndComponentMapping::UnregisterPCGComponent(UPCGComponent* InCompo
 
 void FPCGActorAndComponentMapping::UnregisterPartitionedPCGComponent(UPCGComponent* InComponent)
 {
-	if (!PartitionedOctree.RemoveComponent(InComponent))
+	check(InComponent);
+
+	if (!InComponent->GetOwner()->IsA<APCGPartitionActor>())
+	{
+		PCGSubsystem->OnOriginalComponentUnregistered(InComponent);
+	}
+
+	if (!PartitionedOctree.RemoveComponent(InComponent) || InComponent->GenerationTrigger == EPCGComponentGenerationTrigger::GenerateAtRuntime)
 	{
 		return;
 	}
@@ -389,6 +407,11 @@ void FPCGActorAndComponentMapping::UnregisterPartitionedPCGComponent(UPCGCompone
 
 void FPCGActorAndComponentMapping::UnregisterNonPartitionedPCGComponent(UPCGComponent* InComponent)
 {
+	if (!InComponent->GetOwner()->IsA<APCGPartitionActor>())
+	{
+		PCGSubsystem->OnOriginalComponentUnregistered(InComponent);
+	}
+
 	NonPartitionedOctree.RemoveComponent(InComponent);
 }
 
@@ -400,37 +423,45 @@ void FPCGActorAndComponentMapping::ForAllIntersectingComponents(const FBoxCenter
 	});
 }
 
-void FPCGActorAndComponentMapping::RegisterPartitionActor(APCGPartitionActor* Actor, bool bDoComponentMapping)
+void FPCGActorAndComponentMapping::RegisterPartitionActor(APCGPartitionActor* InActor, bool bDoComponentMapping)
 {
-	check(Actor);
+	check(InActor);
+
+	const uint32 GridSize = InActor->GetPCGGridSize();
+	const FIntVector GridCoord = InActor->GetGridCoord();
+
+	check(GridSize > 0);
+
+	const bool bIsRuntimeGenerated = InActor->IsRuntimeGenerated();
+	TMap<uint32, TMap<FIntVector, TObjectPtr<APCGPartitionActor>>>& ActorsMap = bIsRuntimeGenerated ? RuntimeGenPartitionActorsMap : PartitionActorsMap;
 
-	FIntVector GridCoord = Actor->GetGridCoord();
 	{
-		FWriteScopeLock WriteLock(PartitionActorsMapLock);
+		FWriteScopeLock WriteLock(bIsRuntimeGenerated ? RuntimeGenPartitionActorsMapLock : PartitionActorsMapLock);
 
-		TMap<FIntVector, TObjectPtr<APCGPartitionActor>>& PartitionActorsMapGrid = PartitionActorsMap.FindOrAdd(Actor->GetPCGGridSize());
+		TMap<FIntVector, TObjectPtr<APCGPartitionActor>>& PartitionActorsMapGrid = ActorsMap.FindOrAdd(GridSize);
 		if (PartitionActorsMapGrid.Contains(GridCoord))
 		{
 			return;
 		}
 
-		PartitionActorsMapGrid.Add(GridCoord, Actor);
+		PartitionActorsMapGrid.Add(GridCoord, InActor);
 	}
 
 	// For deprecration: bUse2DGrid is now true by default. But if we already have Partition Actors that were created when the flag was false by default,
 	// we keep this flag
 	if (APCGWorldActor* WorldActor = PCGSubsystem->GetPCGWorldActor())
 	{
-		if (WorldActor->bUse2DGrid != Actor->IsUsing2DGrid())
+		if (WorldActor->bUse2DGrid != InActor->IsUsing2DGrid())
 		{
-			WorldActor->bUse2DGrid = Actor->IsUsing2DGrid();
+			WorldActor->bUse2DGrid = InActor->IsUsing2DGrid();
 		}
 	}
 
-	// And then register itself to all the components that intersect with it
+	// Register to all the components that intersect with the PA. Ignore for runtime generated, it is handled manually
+	if (!bIsRuntimeGenerated)
 	{
 		FWriteScopeLock WriteLock(ComponentToPartitionActorsMapLock);
-		ForAllIntersectingComponents(FBoxCenterAndExtent(Actor->GetFixedBounds()), [this, Actor, bDoComponentMapping](UPCGComponent* Component)
+		ForAllIntersectingComponents(FBoxCenterAndExtent(InActor->GetFixedBounds()), [this, InActor, bDoComponentMapping](UPCGComponent* Component)
 		{
 			// For each component, do the mapping if we ask it explicitly, or if the component is generated
 			if (bDoComponentMapping || Component->bGenerated)
@@ -440,8 +471,8 @@ void FPCGActorAndComponentMapping::RegisterPartitionActor(APCGPartitionActor* Ac
 				// the mapping might not already exists, even if the component is marked generated.
 				if (PartitionActorsPtr)
 				{
-					Actor->AddGraphInstance(Component);
-					PartitionActorsPtr->Add(Actor);
+					InActor->AddGraphInstance(Component);
+					PartitionActorsPtr->Add(InActor);
 				}
 			}
 		});
@@ -452,24 +483,33 @@ void FPCGActorAndComponentMapping::UnregisterPartitionActor(APCGPartitionActor*
 {
 	check(Actor);
 
-	FIntVector GridCoord = Actor->GetGridCoord();
+	const FIntVector GridCoord = Actor->GetGridCoord();
+	const uint32 GridSize = Actor->GetPCGGridSize();
+	if (!ensure(GridSize > 0))
+	{
+		return;
+	}
+
+	const bool bIsRuntimeGenerated = Actor->IsRuntimeGenerated();
+	TMap<uint32, TMap<FIntVector, TObjectPtr<APCGPartitionActor>>>& ActorsMap = bIsRuntimeGenerated ? RuntimeGenPartitionActorsMap : PartitionActorsMap;
 
-	if (TMap<FIntVector, TObjectPtr<APCGPartitionActor>>* PartitionActorsMapGrid = PartitionActorsMap.Find(Actor->GetPCGGridSize()))
+	if (TMap<FIntVector, TObjectPtr<APCGPartitionActor>>* PartitionActorsMapGrid = ActorsMap.Find(GridSize))
 	{
-		FWriteScopeLock WriteLock(PartitionActorsMapLock);
+		FWriteScopeLock WriteLock(bIsRuntimeGenerated ? RuntimeGenPartitionActorsMapLock : PartitionActorsMapLock);
 		PartitionActorsMapGrid->Remove(GridCoord);
 	}
 
-	// And then unregister itself to all the components that intersect with it
+	// Unregister from all intersecting components. Ignore for runtime generated, it is handled manually
+	if (!bIsRuntimeGenerated)
 	{
 		FWriteScopeLock WriteLock(ComponentToPartitionActorsMapLock);
 		ForAllIntersectingComponents(FBoxCenterAndExtent(Actor->GetFixedBounds()), [this, Actor](UPCGComponent* Component)
 		{
 			TSet<TObjectPtr<APCGPartitionActor>>* PartitionActorsPtr = ComponentToPartitionActorsMap.Find(Component);
-		if (PartitionActorsPtr)
-		{
-			PartitionActorsPtr->Remove(Actor);
-		}
+			if (PartitionActorsPtr)
+			{
+				PartitionActorsPtr->Remove(Actor);
+			}
 		});
 	}
 }
@@ -541,21 +581,21 @@ void FPCGActorAndComponentMapping::UpdateMappingPCGComponentPartitionActor(UPCGC
 
 	TSet<TObjectPtr<APCGPartitionActor>> RemovedActors;
 
+	if (const APCGWorldActor* WorldActor = PCGSubsystem->GetPCGWorldActor())
 	{
-		FWriteScopeLock WriteLock(ComponentToPartitionActorsMapLock);
+		const bool bIsHiGenEnabled = InComponent->GetGraph() && InComponent->GetGraph()->IsHierarchicalGenerationEnabled();
 
-		TSet<TObjectPtr<APCGPartitionActor>>* PartitionActorsPtr = ComponentToPartitionActorsMap.Find(InComponent);
-
-		if (!PartitionActorsPtr)
+		auto UpdateMapping = [this, InComponent, &Bounds, &RemovedActors, WorldActor, bIsHiGenEnabled](TMap<const UPCGComponent*, TSet<TObjectPtr<APCGPartitionActor>>>& Map, FRWLock& Lock)
 		{
-			// Does not yet exists, add it
-			PartitionActorsPtr = &ComponentToPartitionActorsMap.Emplace(InComponent);
-			check(PartitionActorsPtr);
-		}
+			FWriteScopeLock WriteLock(Lock);
+			TSet<TObjectPtr<APCGPartitionActor>>* PartitionActorsPtr = Map.Find(InComponent);
 
-		if (const APCGWorldActor* WorldActor = PCGSubsystem->GetPCGWorldActor())
-		{
-			const bool bIsHiGenEnabled = InComponent->GetGraph() && InComponent->GetGraph()->IsHierarchicalGenerationEnabled();
+			if (!PartitionActorsPtr)
+			{
+				// Does not yet exists, add it
+				PartitionActorsPtr = &Map.Emplace(InComponent);
+				check(PartitionActorsPtr);
+			}
 
 			TSet<TObjectPtr<APCGPartitionActor>> NewMapping;
 			ForAllIntersectingPartitionActors(Bounds, [&NewMapping, InComponent, WorldActor, bIsHiGenEnabled](APCGPartitionActor* Actor)
@@ -573,7 +613,10 @@ void FPCGActorAndComponentMapping::UpdateMappingPCGComponentPartitionActor(UPCGC
 			RemovedActors = PartitionActorsPtr->Difference(NewMapping);
 
 			*PartitionActorsPtr = MoveTemp(NewMapping);
-		}
+		};
+
+		UpdateMapping(ComponentToPartitionActorsMap, ComponentToPartitionActorsMapLock);
+		UpdateMapping(ComponentToRuntimeGenPartitionActorsMap, ComponentToRuntimeGenPartitionActorsMapLock);
 	}
 
 	// No need to be locked to do this.
@@ -592,18 +635,23 @@ void FPCGActorAndComponentMapping::DeleteMappingPCGComponentPartitionActor(UPCGC
 		return;
 	}
 
-	FWriteScopeLock WriteLock(ComponentToPartitionActorsMapLock);
-
-	TSet<TObjectPtr<APCGPartitionActor>>* PartitionActorsPtr = ComponentToPartitionActorsMap.Find(InComponent);
-	if (PartitionActorsPtr)
+	auto DeleteMapping = [this, InComponent](TMap<const UPCGComponent*, TSet<TObjectPtr<APCGPartitionActor>>>& Map, FRWLock& Lock)
 	{
-		for (APCGPartitionActor* Actor : *PartitionActorsPtr)
+		FWriteScopeLock WriteLock(Lock);
+
+		if (TSet<TObjectPtr<APCGPartitionActor>>* PartitionActorsPtr = Map.Find(InComponent))
 		{
-			Actor->RemoveGraphInstance(InComponent);
+			for (APCGPartitionActor* Actor : *PartitionActorsPtr)
+			{
+				Actor->RemoveGraphInstance(InComponent);
+			}
+
+			PartitionActorsPtr->Empty();
 		}
+	};
 
-		PartitionActorsPtr->Empty();
-	}
+	DeleteMapping(ComponentToPartitionActorsMap, ComponentToPartitionActorsMapLock);
+	DeleteMapping(ComponentToRuntimeGenPartitionActorsMap, ComponentToRuntimeGenPartitionActorsMapLock);
 }
 
 bool FPCGActorAndComponentMapping::IsComponentRegistered(const UPCGComponent* InComponent) const
@@ -628,11 +676,12 @@ TSet<UPCGComponent*> FPCGActorAndComponentMapping::GetAllRegisteredComponents()
 	return Res;
 }
 
-UPCGComponent* FPCGActorAndComponentMapping::GetLocalComponent(uint32 GridSize, const FIntVector& CellCoords, const UPCGComponent* InOriginalComponent)
+UPCGComponent* FPCGActorAndComponentMapping::GetLocalComponent(uint32 GridSize, const FIntVector& CellCoords, const UPCGComponent* InOriginalComponent, bool bRuntimeGenerated)
 {
-	FReadScopeLock ReadLock(PartitionActorsMapLock);
+	TMap<uint32, TMap<FIntVector, TObjectPtr<APCGPartitionActor>>>& ActorsMap = bRuntimeGenerated ? RuntimeGenPartitionActorsMap : PartitionActorsMap;
+	FReadScopeLock ReadLock(bRuntimeGenerated ? RuntimeGenPartitionActorsMapLock : PartitionActorsMapLock);
 
-	if (const TMap<FIntVector, TObjectPtr<APCGPartitionActor>>* PartitionActorsOnGrid = PartitionActorsMap.Find(GridSize))
+	if (const TMap<FIntVector, TObjectPtr<APCGPartitionActor>>* PartitionActorsOnGrid = ActorsMap.Find(GridSize))
 	{
 		const TObjectPtr<APCGPartitionActor>* PartitionActor = PartitionActorsOnGrid->Find(CellCoords);
 		if (PartitionActor && *PartitionActor)
@@ -652,6 +701,20 @@ UPCGComponent* FPCGActorAndComponentMapping::GetLocalComponent(uint32 GridSize,
 	return nullptr;
 }
 
+APCGPartitionActor* FPCGActorAndComponentMapping::GetPartitionActor(uint32 GridSize, const FIntVector& CellCoords, bool bRuntimeGenerated) const
+{
+	const TMap<uint32, TMap<FIntVector, TObjectPtr<APCGPartitionActor>>>& Map = bRuntimeGenerated ? RuntimeGenPartitionActorsMap : PartitionActorsMap;
+	FReadScopeLock ReadLock(bRuntimeGenerated ? RuntimeGenPartitionActorsMapLock : PartitionActorsMapLock);
+
+	if (const TMap<FIntVector, TObjectPtr<APCGPartitionActor>>* PartitionActorsOnGrid = Map.Find(GridSize))
+	{
+		const TObjectPtr<APCGPartitionActor>* PartitionActor = PartitionActorsOnGrid->Find(CellCoords);
+		return PartitionActor ? *PartitionActor : nullptr;
+	}
+
+	return nullptr;
+}
+
 #if WITH_EDITOR
 void FPCGActorAndComponentMapping::RegisterTracking(UPCGComponent* InComponent)
 {
diff --git a/Engine/Plugins/Experimental/PCG/Source/PCG/Private/PCGCommon.cpp b/Engine/Plugins/Experimental/PCG/Source/PCG/Private/PCGCommon.cpp
--- a/Engine/Plugins/Experimental/PCG/Source/PCG/Private/PCGCommon.cpp
+++ b/Engine/Plugins/Experimental/PCG/Source/PCG/Private/PCGCommon.cpp
@@ -30,6 +30,11 @@ namespace PCGHiGenGrid
 		return InGrid >= EPCGHiGenGrid::GridMin && static_cast<uint32>(InGrid) < 2 * static_cast<uint32>(EPCGHiGenGrid::GridMax);
 	}
 
+	bool IsValidGridOrUninitialized(EPCGHiGenGrid InGrid)
+	{
+		return IsValidGrid(InGrid) || InGrid == EPCGHiGenGrid::Uninitialized;
+	}
+
 	uint32 GridToGridSize(EPCGHiGenGrid InGrid)
 	{
 		const uint32 GridAsUint = static_cast<uint32>(InGrid);
@@ -59,3 +64,30 @@ namespace PCGHiGenGrid
 		return 100 * static_cast<uint32>(EPCGHiGenGrid::Unbounded);
 	}
 }
+
+double FPCGRuntimeGenerationRadii::GetGenerationRadiusFromGrid(EPCGHiGenGrid Grid) const
+{
+	switch (Grid)
+	{
+		case EPCGHiGenGrid::Grid4: return GenerationRadius400;
+		case EPCGHiGenGrid::Grid8: return GenerationRadius800;
+		case EPCGHiGenGrid::Grid16: return GenerationRadius1600;
+		case EPCGHiGenGrid::Grid32: return GenerationRadius3200;
+		case EPCGHiGenGrid::Grid64: return GenerationRadius6400;
+		case EPCGHiGenGrid::Grid128: return GenerationRadius12800;
+		case EPCGHiGenGrid::Grid256: return GenerationRadius25600;
+		case EPCGHiGenGrid::Grid512: return GenerationRadius51200;
+		case EPCGHiGenGrid::Grid1024: return GenerationRadius102400;
+		case EPCGHiGenGrid::Grid2048: return GenerationRadius204800;
+		case EPCGHiGenGrid::Unbounded: return GenerationRadius;
+	}
+
+	ensure(false);
+	return 0;
+}
+
+double FPCGRuntimeGenerationRadii::GetCleanupRadiusFromGrid(EPCGHiGenGrid Grid) const
+{
+	return GetGenerationRadiusFromGrid(Grid) * CleanupRadiusMultiplier;
+}
+
diff --git a/Engine/Plugins/Experimental/PCG/Source/PCG/Private/PCGComponent.cpp b/Engine/Plugins/Experimental/PCG/Source/PCG/Private/PCGComponent.cpp
--- a/Engine/Plugins/Experimental/PCG/Source/PCG/Private/PCGComponent.cpp
+++ b/Engine/Plugins/Experimental/PCG/Source/PCG/Private/PCGComponent.cpp
@@ -26,6 +26,9 @@
 #include "Grid/PCGPartitionActor.h"
 #include "Helpers/PCGActorHelpers.h"
 #include "Helpers/PCGHelpers.h"
+#include "RuntimeGen/GenSources/PCGGenSourceBase.h"
+#include "RuntimeGen/SchedulingPolicies/PCGSchedulingPolicyBase.h"
+#include "RuntimeGen/SchedulingPolicies/PCGSchedulingPolicyDistanceAndDirection.h"
 #include "Utils/PCGGeneratedResourcesLogging.h"
 
 #include "LandscapeComponent.h"
@@ -42,6 +45,7 @@
 #include "GameFramework/Volume.h"
 #include "Kismet/GameplayStatics.h"
 #include "LandscapeSplinesComponent.h"
+#include "UObject/Package.h"
 
 #include UE_INLINE_GENERATED_CPP_BY_NAME(PCGComponent)
 
@@ -63,6 +67,7 @@ UPCGComponent::UPCGComponent(const FObjectInitializer& InObjectInitializer)
 	: Super(InObjectInitializer)
 {
 	GraphInstance = InObjectInitializer.CreateDefaultSubobject<UPCGGraphInstance>(this, TEXT("PCGGraphInstance"));
+	SchedulingPolicyClass = UPCGSchedulingPolicyDistanceAndDirection::StaticClass();
 
 #if WITH_EDITOR
 	// If we are in Editor, and we are a BP template (no owner), we will mark this component to force a generate when added to world.
@@ -75,8 +80,10 @@ UPCGComponent::UPCGComponent(const FObjectInitializer& InObjectInitializer)
 
 bool UPCGComponent::CanPartition() const
 {
-	// Support/Force partitioning on non-PCG partition actors in WP worlds.
-	return GetOwner() && GetOwner()->GetWorld() && GetOwner()->GetWorld()->GetWorldPartition() != nullptr && Cast<APCGPartitionActor>(GetOwner()) == nullptr;
+	// Support/Force partitioning on non-PCG partition actors in WP worlds. GenerateAtRuntime components can be partitioned even if WorldPartition is not enabled.
+	return ((GetOwner() && GetOwner()->GetWorld() && GetOwner()->GetWorld()->GetWorldPartition() != nullptr)
+		|| GenerationTrigger == EPCGComponentGenerationTrigger::GenerateAtRuntime)
+		&& Cast<APCGPartitionActor>(GetOwner()) == nullptr;
 }
 
 bool UPCGComponent::IsPartitioned() const
@@ -166,6 +173,18 @@ bool UPCGComponent::ShouldGenerate(bool bForce, EPCGComponentGenerationTrigger R
 		return false;
 	}
 
+	if (GenerationTrigger == EPCGComponentGenerationTrigger::GenerateAtRuntime)
+	{
+		// If we're runtime generated, turn down other requests.
+		const bool bShouldGenerate = RequestedGenerationTrigger == EPCGComponentGenerationTrigger::GenerateAtRuntime;
+		if (!bShouldGenerate)
+		{
+			UE_LOG(LogPCG, Warning, TEXT("Generation request with trigger %d denied as this component is managed by the runtime generation scheduler."), (int)RequestedGenerationTrigger);
+		}
+
+		return bShouldGenerate;
+	}
+
 #if WITH_EDITOR
 	// Always run Generate if we are in editor and partitioned since the original component doesn't know the state of the local one.
 	if (IsPartitioned() && !PCGHelpers::IsRuntimeOrPIE())
@@ -212,6 +231,11 @@ void UPCGComponent::SetPropertiesFromOriginal(const UPCGComponent* Original)
 	Seed = Original->Seed;
 	GenerationTrigger = Original->GenerationTrigger;
 
+	bOverrideGenerationRadii = Original->bOverrideGenerationRadii;
+	GenerationRadii = Original->GenerationRadii;
+	SchedulingPolicyClass = Original->SchedulingPolicyClass;
+	SchedulingPolicy = Original->SchedulingPolicy;
+
 	UPCGGraph* OriginalGraph = Original->GraphInstance ? Original->GraphInstance->GetGraph() : nullptr;
 	if (OriginalGraph != GraphInstance->GetGraph())
 	{
@@ -254,15 +278,25 @@ void UPCGComponent::Generate_Implementation(bool bForce)
 
 void UPCGComponent::GenerateLocal(bool bForce)
 {
-	GenerateInternal(bForce, EPCGComponentGenerationTrigger::GenerateOnDemand, {});
+	GenerateInternal(bForce, EPCGHiGenGrid::Uninitialized, EPCGComponentGenerationTrigger::GenerateOnDemand, {});
+}
+
+void UPCGComponent::GenerateLocal(EPCGComponentGenerationTrigger RequestedGenerationTrigger, bool bForce, EPCGHiGenGrid Grid)
+{
+	GenerateInternal(bForce, Grid, RequestedGenerationTrigger, {});
 }
 
 FPCGTaskId UPCGComponent::GenerateLocalGetTaskId(bool bForce)
 {
-	return GenerateInternal(bForce, EPCGComponentGenerationTrigger::GenerateOnDemand, {});
+	return GenerateInternal(bForce, EPCGHiGenGrid::Uninitialized, EPCGComponentGenerationTrigger::GenerateOnDemand, {});
 }
 
-FPCGTaskId UPCGComponent::GenerateInternal(bool bForce, EPCGComponentGenerationTrigger RequestedGenerationTrigger, const TArray<FPCGTaskId>& Dependencies)
+FPCGTaskId UPCGComponent::GenerateLocalGetTaskId(EPCGComponentGenerationTrigger RequestedGenerationTrigger, bool bForce, EPCGHiGenGrid Grid)
+{
+	return GenerateInternal(bForce, Grid, RequestedGenerationTrigger, {});
+}
+
+FPCGTaskId UPCGComponent::GenerateInternal(bool bForce, EPCGHiGenGrid Grid, EPCGComponentGenerationTrigger RequestedGenerationTrigger, const TArray<FPCGTaskId>& Dependencies)
 {
 	if (IsGenerating() || !GetSubsystem() || !ShouldGenerate(bForce, RequestedGenerationTrigger))
 	{
@@ -271,7 +305,7 @@ FPCGTaskId UPCGComponent::GenerateInternal(bool bForce, EPCGComponentGenerationT
 
 	Modify(!IsInPreviewMode());
 
-	CurrentGenerationTask = GetSubsystem()->ScheduleComponent(this, /*bSave=*/bForce, Dependencies);
+	CurrentGenerationTask = GetSubsystem()->ScheduleComponent(this, Grid, /*bSave=*/bForce, Dependencies);
 
 	return CurrentGenerationTask;
 }
@@ -504,6 +538,12 @@ void UPCGComponent::Cleanup()
 		return;
 	}
 
+	if (GenerationTrigger == EPCGComponentGenerationTrigger::GenerateAtRuntime)
+	{
+		UE_LOG(LogPCG, Warning, TEXT("Cleanup request denied as this component is managed by the runtime generation scheduler."));
+		return;
+	}
+
 #if WITH_EDITOR
 	FScopedTransaction Transaction(LOCTEXT("PCGCleanup", "Clean up PCG component"));
 #endif
@@ -666,6 +706,51 @@ void UPCGComponent::ClearPerPinGeneratedOutput()
 	PerPinGeneratedOutput.Reset();
 }
 
+void UPCGComponent::SetSchedulingPolicyClass(TSubclassOf<UPCGSchedulingPolicyBase> InSchedulingPolicyClass) 
+{
+	if (!SchedulingPolicy || InSchedulingPolicyClass != SchedulingPolicyClass)
+	{
+		if (InSchedulingPolicyClass != SchedulingPolicyClass)
+		{
+			SchedulingPolicyClass = InSchedulingPolicyClass;
+		}
+		
+		RefreshSchedulingPolicy();
+	}
+}
+
+double UPCGComponent::GetGenerationRadiusFromGrid(EPCGHiGenGrid Grid) const
+{
+	if (bOverrideGenerationRadii)
+	{
+		return GenerationRadii.GetGenerationRadiusFromGrid(Grid);
+	}
+
+	const UPCGGraph* Graph = GetGraph();
+	if (ensure(Graph))
+	{
+		return Graph->GenerationRadii.GetGenerationRadiusFromGrid(Grid);
+	}
+
+	return 0;
+}
+
+double UPCGComponent::GetCleanupRadiusFromGrid(EPCGHiGenGrid Grid) const
+{
+	if (bOverrideGenerationRadii)
+	{
+		return GenerationRadii.GetCleanupRadiusFromGrid(Grid);
+	}
+
+	const UPCGGraph* Graph = GetGraph();
+	if (ensure(Graph))
+	{
+		return Graph->GenerationRadii.GetCleanupRadiusFromGrid(Grid);
+	}
+
+	return 0;
+}
+
 bool UPCGComponent::MoveResourcesToNewActor(AActor* InNewActor, bool bCreateChild)
 {
 	// Don't move resources if we are generating or cleaning up
@@ -731,6 +816,9 @@ void UPCGComponent::CleanupLocalImmediate(bool bRemoveComponents)
 {
 	PCGGeneratedResourcesLogging::LogCleanupLocalImmediate(bRemoveComponents, GeneratedResources);
 
+	// Cancels generation of this component if there is an ongoing generation in progress.
+	CancelGeneration();
+
 	TSet<TSoftObjectPtr<AActor>> ActorsToDelete;
 
 	if (!bRemoveComponents && UPCGManagedResource::DebugForcePurgeAllResourcesOnGenerate())
@@ -1033,7 +1121,7 @@ void UPCGComponent::BeginPlay()
 		}
 		else
 		{
-			GenerateInternal(/*bForce=*/false, EPCGComponentGenerationTrigger::GenerateOnLoad, {});
+			GenerateInternal(/*bForce=*/false, EPCGHiGenGrid::Uninitialized, EPCGComponentGenerationTrigger::GenerateOnLoad, {});
 			bRuntimeGenerated = true;
 		}
 	}
@@ -1186,6 +1274,19 @@ void UPCGComponent::PostLoad()
 		bDirtyGenerated = true;
 	}
 #endif
+
+	if (!SchedulingPolicy)
+	{
+		RefreshSchedulingPolicy();
+	}
+	else
+	{
+		const EObjectFlags Flags = GetMaskedFlags(RF_PropagateToSubObjects) | RF_Transactional;
+		SchedulingPolicy->SetFlags(Flags);
+#if WITH_EDITOR
+		SchedulingPolicy->SetShouldDisplayProperties(GenerationTrigger == EPCGComponentGenerationTrigger::GenerateAtRuntime);
+#endif
+	}
 }
 
 #if WITH_EDITOR
@@ -1287,10 +1388,19 @@ TStructOnScope<FActorComponentInstanceData> UPCGComponent::GetComponentInstanceD
 
 void UPCGComponent::OnGraphChanged(UPCGGraphInterface* InGraph, EPCGChangeType ChangeType)
 {
-	const bool bIsStructural = ((ChangeType & (EPCGChangeType::Edge | EPCGChangeType::Structural)) != EPCGChangeType::None);
-	const bool bDirtyInputs = bIsStructural || ((ChangeType & EPCGChangeType::Input) != EPCGChangeType::None);
+	if (!!(ChangeType & EPCGChangeType::Debug) && GenerationTrigger == EPCGComponentGenerationTrigger::GenerateAtRuntime)
+	{
+		// TODO: Handle RuntimeGen business in the RuntimeGenScheduler, subscribe to OnGraphChanged there instead.
+		CleanupLocalImmediate(/*bRemoveComponents=*/true);
+		GenerateLocal(EPCGComponentGenerationTrigger::GenerateAtRuntime, /*bForce=*/true);
+	}
+	else
+	{
+		const bool bIsStructural = ((ChangeType & (EPCGChangeType::Edge | EPCGChangeType::Structural)) != EPCGChangeType::None);
+		const bool bDirtyInputs = bIsStructural || ((ChangeType & EPCGChangeType::Input) != EPCGChangeType::None);
 
-	RefreshAfterGraphChanged(InGraph, bIsStructural, bDirtyInputs);
+		RefreshAfterGraphChanged(InGraph, bIsStructural, bDirtyInputs);
+	}
 }
 
 void UPCGComponent::RefreshAfterGraphChanged(UPCGGraphInterface* InGraph, bool bIsStructural, bool bDirtyInputs)
@@ -1421,6 +1531,31 @@ void UPCGComponent::PostEditChangeProperty(FPropertyChangedEvent& PropertyChange
 		DirtyGenerated();
 		Refresh();
 	}
+	else if (PropName == GET_MEMBER_NAME_CHECKED(UPCGComponent, SchedulingPolicyClass))
+	{
+		// We don't need to refresh the component here because this does not effect generation behavior, only scheduling behavior.
+		RefreshSchedulingPolicy();
+	} 
+	else if (PropName == GET_MEMBER_NAME_CHECKED(UPCGComponent, GenerationTrigger))
+	{
+		if (GenerationTrigger == EPCGComponentGenerationTrigger::GenerateAtRuntime)
+		{
+			CleanupLocalImmediate(/*bRemoveComponents=*/true);
+		}
+
+		if (!SchedulingPolicy)
+		{
+			RefreshSchedulingPolicy();
+		}
+
+		if (SchedulingPolicy)
+		{
+			// We should only display scheduling policy parameters when in runtime generation mode.
+			SchedulingPolicy->SetShouldDisplayProperties(GenerationTrigger == EPCGComponentGenerationTrigger::GenerateAtRuntime);
+		}
+
+		Refresh();
+	}
 	// General properties that don't affect behavior
 	else
 	{
@@ -1687,6 +1822,12 @@ void UPCGComponent::Refresh(bool bStructural)
 		return;
 	}
 
+	// Do not allow automatic refreshing for runtime generation components.
+	if (GenerationTrigger == EPCGComponentGenerationTrigger::GenerateAtRuntime)
+	{
+		return;
+	}
+
 	// If the component is tagged as not to regenerate in the editor, only exceptional cases should trigger a refresh
 	// namely: the component is deactivated.
 	// Note that the component changing its IsPartitioned state is already covered in the PostEditChangeProperty
@@ -1810,6 +1951,8 @@ bool UPCGComponent::IsActorTracked(AActor* InActor, bool& bOutIsCulled) const
 
 void UPCGComponent::OnRefresh(bool bForceRefresh)
 {
+	check(GenerationTrigger != EPCGComponentGenerationTrigger::GenerateAtRuntime);
+
 	// Mark the refresh task invalid to allow re-triggering refreshes
 	CurrentRefreshTask = InvalidPCGTaskId;
 
@@ -1824,7 +1967,7 @@ void UPCGComponent::OnRefresh(bool bForceRefresh)
 		// If we are partitioned but we have resources, we need to force a cleanup
 		if (!GeneratedResources.IsEmpty())
 		{
-			CleanupLocalImmediate(true);
+			CleanupLocalImmediate(/*bRemoveComponents=*/true);
 		}
 	}
 
@@ -2214,6 +2357,27 @@ FPCGDataCollection UPCGComponent::CreateActorPCGDataCollection(AActor* Actor, co
 	return Collection;
 }
 
+void UPCGComponent::RefreshSchedulingPolicy()
+{
+	if (SchedulingPolicy)
+	{
+		SchedulingPolicy->Rename(nullptr, GetTransientPackage(), REN_DontCreateRedirectors | REN_ForceNoResetLoaders);
+		SchedulingPolicy->MarkAsGarbage();
+		SchedulingPolicy = nullptr;
+	}
+
+	if (SchedulingPolicyClass)
+	{
+		const EObjectFlags Flags = GetMaskedFlags(RF_PropagateToSubObjects);
+		SchedulingPolicy = NewObject<UPCGSchedulingPolicyBase>(this, SchedulingPolicyClass, NAME_None, Flags);
+
+#if WITH_EDITOR
+		// We should only display scheduling policy parameters when in runtime generation mode.
+		SchedulingPolicy->SetShouldDisplayProperties(GenerationTrigger == EPCGComponentGenerationTrigger::GenerateAtRuntime);
+#endif
+	}
+}
+
 UPCGData* UPCGComponent::CreatePCGData()
 {
 	TRACE_CPUPROFILER_EVENT_SCOPE(UPCGComponent::CreatePCGData);
diff --git a/Engine/Plugins/Experimental/PCG/Source/PCG/Private/PCGSubsystem.cpp b/Engine/Plugins/Experimental/PCG/Source/PCG/Private/PCGSubsystem.cpp
--- a/Engine/Plugins/Experimental/PCG/Source/PCG/Private/PCGSubsystem.cpp
+++ b/Engine/Plugins/Experimental/PCG/Source/PCG/Private/PCGSubsystem.cpp
@@ -14,6 +14,8 @@
 #include "Grid/PCGPartitionActor.h"
 #include "Helpers/PCGActorHelpers.h"
 #include "Helpers/PCGHelpers.h"
+#include "RuntimeGen/PCGGenSourceManager.h"
+#include "RuntimeGen/PCGRuntimeGenScheduler.h"
 
 #include "Engine/Level.h"
 
@@ -149,6 +151,9 @@ void UPCGSubsystem::Deinitialize()
 	delete GraphExecutor;
 	GraphExecutor = nullptr;
 
+	delete RuntimeGenScheduler;
+	RuntimeGenScheduler = nullptr;
+
 	PCGWorldActor = nullptr;
 	bHasTickedOnce = false;
 
@@ -166,7 +171,11 @@ void UPCGSubsystem::PostInitialize()
 	// Initialize graph executor
 	check(!GraphExecutor);
 	GraphExecutor = new FPCGGraphExecutor(this);
-	
+
+	// Initialize runtime generation scheduler
+	check(!RuntimeGenScheduler);
+	RuntimeGenScheduler = new FPCGRuntimeGenScheduler(GetWorld(), &ActorAndComponentMapping);
+
 	// Gather world pcg actor if it exists
 	if (!PCGWorldActor)
 	{
@@ -232,6 +241,11 @@ void UPCGSubsystem::Tick(float DeltaSeconds)
 	}
 
 	ActorAndComponentMapping.Tick();
+
+	if (PCGWorldActor)
+	{
+		RuntimeGenScheduler->Tick(PCGWorldActor);
+	}
 }
 
 APCGWorldActor* UPCGSubsystem::GetPCGWorldActor()
@@ -314,13 +328,21 @@ void UPCGSubsystem::UnregisterPCGWorldActor(APCGWorldActor* InActor)
 	}
 }
 
+void UPCGSubsystem::OnOriginalComponentUnregistered(UPCGComponent* InComponent)
+{
+	if (RuntimeGenScheduler)
+	{
+		RuntimeGenScheduler->OnOriginalComponentUnregistered(InComponent);
+	}
+}
+
 UPCGLandscapeCache* UPCGSubsystem::GetLandscapeCache()
 {
 	APCGWorldActor* LandscapeCacheOwner = GetPCGWorldActor();
 	return LandscapeCacheOwner ? LandscapeCacheOwner->LandscapeCacheObject.Get() : nullptr;
 }
 
-FPCGTaskId UPCGSubsystem::ScheduleComponent(UPCGComponent* PCGComponent, bool bSave, const TArray<FPCGTaskId>& InDependencies)
+FPCGTaskId UPCGSubsystem::ScheduleComponent(UPCGComponent* PCGComponent, EPCGHiGenGrid Grid, bool bSave, const TArray<FPCGTaskId>& InDependencies)
 {
 	check(GraphExecutor);
 
@@ -334,7 +356,8 @@ FPCGTaskId UPCGSubsystem::ScheduleComponent(UPCGComponent* PCGComponent, bool bS
 	ensure(PCGHelpers::GetGenerationGridSizes(PCGComponent->GetGraph(), GetPCGWorldActor(), GridSizes, bHasUnbounded));
 
 #if WITH_EDITOR
-	if (PCGComponent->IsPartitioned() && !PCGHelpers::IsRuntimeOrPIE())
+	// Create the PartitionActors if necessary. Skip if this is a runtime managed component, PAs are handled manually by the RuntimeGenScheduler.
+	if (PCGComponent->IsPartitioned() && !PCGHelpers::IsRuntimeOrPIE() && PCGComponent->GenerationTrigger != EPCGComponentGenerationTrigger::GenerateAtRuntime)
 	{
 		if (!GridSizes.IsEmpty())
 		{
@@ -353,7 +376,20 @@ FPCGTaskId UPCGSubsystem::ScheduleComponent(UPCGComponent* PCGComponent, bool bS
 
 	// Schedule generation of original component if is is non-partitioned, or if it has nodes that will execute at the Unbounded level.
 	FPCGTaskId OriginalComponentTask = InvalidPCGTaskId;
-	if (!PCGComponent->IsPartitioned() || bHasUnbounded)
+
+	bool bGeneratePCGComponent = false;
+	if (!PCGComponent->IsPartitioned())
+	{
+		// This component is either an unpartitioned original component or a local component. Generate if grid size matches preference (if provided).
+		bGeneratePCGComponent = (Grid == EPCGHiGenGrid::Uninitialized) || !!(Grid & PCGComponent->GetGenerationGrid());
+	}
+	else
+	{
+		// This component is a partitioned original component. Generate if the graph has unbounded nodes and if this grid matches preference (if provided).
+		bGeneratePCGComponent = bHasUnbounded && (!!(Grid & EPCGHiGenGrid::Unbounded) || Grid == EPCGHiGenGrid::Uninitialized);
+	}
+
+	if (bGeneratePCGComponent)
 	{
 		OriginalComponentTask = PCGComponent->CreateGenerateTask(/*bForce=*/bSave, InDependencies);
 		if (OriginalComponentTask != InvalidPCGTaskId)
@@ -363,7 +399,7 @@ FPCGTaskId UPCGSubsystem::ScheduleComponent(UPCGComponent* PCGComponent, bool bS
 	}
 
 	// If the component is partitioned, we will forward the calls to its registered PCG Partition actors
-	if (PCGComponent->IsPartitioned())
+	if (PCGComponent->IsPartitioned() && PCGHiGenGrid::IsValidGridOrUninitialized(Grid))
 	{
 		// Local components depend on the original component (to ensure any data is available).
 		TArray<FPCGTaskId> Dependencies = InDependencies;
@@ -372,9 +408,14 @@ FPCGTaskId UPCGSubsystem::ScheduleComponent(UPCGComponent* PCGComponent, bool bS
 			Dependencies.Add(OriginalComponentTask);
 		}
 
-		auto LocalGenerateTask = [OriginalComponent = PCGComponent, &Dependencies, bSave, &GridSizes](UPCGComponent* LocalComponent)
+		auto LocalGenerateTask = [OriginalComponent = PCGComponent, Grid, &Dependencies, bSave, &GridSizes](UPCGComponent* LocalComponent)
 		{
-			if (!GridSizes.Contains(LocalComponent->GetGenerationGridSize()))
+			const uint32 LocalComponentGridSize = LocalComponent->GetGenerationGridSize();
+			if (!GridSizes.Contains(LocalComponentGridSize))
+			{
+				return InvalidPCGTaskId;
+			}
+			else if (Grid != EPCGHiGenGrid::Uninitialized && !(Grid & LocalComponent->GetGenerationGrid()))
 			{
 				return InvalidPCGTaskId;
 			}
@@ -396,7 +437,7 @@ FPCGTaskId UPCGSubsystem::ScheduleComponent(UPCGComponent* PCGComponent, bool bS
 				LocalComponent->CleanupLocalImmediate(true);
 			}
 
-			return LocalComponent->GenerateInternal(/*bForce=*/bSave, EPCGComponentGenerationTrigger::GenerateOnDemand, Dependencies);
+			return LocalComponent->GenerateInternal(/*bForce=*/bSave, LocalComponent->GetGenerationGrid(), EPCGComponentGenerationTrigger::GenerateOnDemand, Dependencies);
 		};
 
 		AllTasks.Append(ActorAndComponentMapping.DispatchToRegisteredLocalComponents(PCGComponent, LocalGenerateTask));
@@ -658,9 +699,9 @@ void UPCGSubsystem::ForAllOverlappingComponentsInHierarchy(UPCGComponent* InComp
 	});
 }
 
-UPCGComponent* UPCGSubsystem::GetLocalComponent(uint32 GridSize, const FIntVector& CellCoords, const UPCGComponent* InOriginalComponent)
+UPCGComponent* UPCGSubsystem::GetLocalComponent(uint32 GridSize, const FIntVector& CellCoords, const UPCGComponent* InOriginalComponent, bool bTransient)
 {
-	return ActorAndComponentMapping.GetLocalComponent(GridSize, CellCoords, InOriginalComponent);
+	return ActorAndComponentMapping.GetLocalComponent(GridSize, CellCoords, InOriginalComponent, bTransient);
 }
 
 bool UPCGSubsystem::IsGraphCacheDebuggingEnabled() const
@@ -668,6 +709,11 @@ bool UPCGSubsystem::IsGraphCacheDebuggingEnabled() const
 	return GraphExecutor && GraphExecutor->IsGraphCacheDebuggingEnabled();
 }
 
+FPCGGenSourceManager* UPCGSubsystem::GetGenSourceManager() const
+{
+	return RuntimeGenScheduler ? RuntimeGenScheduler->GenSourceManager : nullptr;
+}
+
 bool UPCGSubsystem::GetOutputData(FPCGTaskId TaskId, FPCGDataCollection& OutData)
 {
 	check(GraphExecutor);
@@ -949,7 +995,7 @@ FPCGTaskId UPCGSubsystem::ProcessGraph(UPCGComponent* Component, const FBox& InP
 						return InvalidPCGTaskId;
 					}
 
-					return LocalComponent->GenerateInternal(bSave, EPCGComponentGenerationTrigger::GenerateOnDemand, TaskDependencies);
+					return LocalComponent->GenerateInternal(bSave, EPCGHiGenGrid::Uninitialized, EPCGComponentGenerationTrigger::GenerateOnDemand, TaskDependencies);
 				}
 				else
 				{
@@ -1015,6 +1061,8 @@ FPCGTaskId UPCGSubsystem::ProcessGraph(UPCGComponent* Component, const FBox& InP
 
 FPCGTaskId UPCGSubsystem::ScheduleRefresh(UPCGComponent* Component, bool bForceRegen)
 {
+	check(Component && Component->GenerationTrigger != EPCGComponentGenerationTrigger::GenerateAtRuntime);
+
 	TWeakObjectPtr<UPCGComponent> ComponentPtr(Component);
 
 	auto RefreshTask = [ComponentPtr, bForceRegen]() {
diff --git a/Engine/Plugins/Experimental/PCG/Source/PCG/Private/RuntimeGen/GenSources/PCGGenSourceComponent.cpp b/Engine/Plugins/Experimental/PCG/Source/PCG/Private/RuntimeGen/GenSources/PCGGenSourceComponent.cpp
new file mode 100644
--- /dev/null
+++ b/Engine/Plugins/Experimental/PCG/Source/PCG/Private/RuntimeGen/GenSources/PCGGenSourceComponent.cpp
@@ -0,0 +1,158 @@
+// Copyright Epic Games, Inc. All Rights Reserved.
+
+#include "RuntimeGen/GenSources/PCGGenSourceComponent.h"
+
+#include "PCGSubsystem.h"
+#include "RuntimeGen/PCGGenSourceManager.h"
+
+#include "SceneManagement.h"
+#include "Engine/World.h"
+#include "GameFramework/Actor.h"
+
+#include UE_INLINE_GENERATED_CPP_BY_NAME(PCGGenSourceComponent)
+
+#define LOCTEXT_NAMESPACE "UPCGGenSourceComponent"
+
+UPCGGenSourceComponent::~UPCGGenSourceComponent()
+{
+#if WITH_EDITOR
+	// For the special case where a component is part of a reconstruction script (from a BP),
+	// but gets destroyed immediately, we need to force the unregistering. 
+	if (UPCGSubsystem* PCGSubsystem = UPCGSubsystem::GetSubsystemForCurrentWorld())
+	{
+		if (FPCGGenSourceManager* GenSourceManager = PCGSubsystem->GetGenSourceManager())
+		{
+			GenSourceManager->UnregisterGenSource(this);
+		}
+	}
+#endif // WITH_EDITOR
+}
+
+void UPCGGenSourceComponent::BeginPlay()
+{
+	Super::BeginPlay();
+
+#if WITH_EDITOR
+	// Disable registration on preview actors.
+	if (GetOwner() && GetOwner()->bIsEditorPreviewActor)
+	{
+		return;
+	}
+#endif
+
+	if (FPCGGenSourceManager* GenSourceManager = GetGenSourceManager())
+	{
+		GenSourceManager->RegisterGenSource(this);
+	}
+}
+
+void UPCGGenSourceComponent::EndPlay(const EEndPlayReason::Type EndPlayReason)
+{
+	// Always try to unregister itself, if it doesn't exist, it will early out. 
+	// Just making sure that we don't left some resources registered while dead.
+	if (FPCGGenSourceManager* GenSourceManager = GetGenSourceManager())
+	{
+		GenSourceManager->UnregisterGenSource(this);
+	}
+
+	Super::EndPlay(EndPlayReason);
+}
+
+void UPCGGenSourceComponent::OnRegister()
+{
+	Super::OnRegister();
+
+#if WITH_EDITOR
+	// Disable registration on preview actors.
+	if (GetOwner() && GetOwner()->bIsEditorPreviewActor)
+	{
+		return;
+	}
+
+	if (UWorld* World = GetWorld())
+	{
+		// We won't be able to spawn any actors if we are currently running a construction script.
+		if (!World->bIsRunningConstructionScript)
+		{
+			if (FPCGGenSourceManager* GenSourceManager = GetGenSourceManager())
+			{
+				GenSourceManager->RegisterGenSource(this);
+			}
+		}
+	}
+#endif //WITH_EDITOR
+}
+
+void UPCGGenSourceComponent::OnUnregister()
+{
+	Super::OnUnregister();
+
+#if WITH_EDITOR
+	// Disable registration on preview actors.
+	if (GetOwner() && GetOwner()->bIsEditorPreviewActor)
+	{
+		return;
+	}
+
+	if (FPCGGenSourceManager* GenSourceManager = GetGenSourceManager())
+	{
+		GenSourceManager->UnregisterGenSource(this);
+	}
+#endif // WITH_EDITOR
+
+}
+
+void UPCGGenSourceComponent::OnComponentCreated()
+{
+	Super::OnComponentCreated();
+
+	if (FPCGGenSourceManager* GenSourceManager = GetGenSourceManager())
+	{
+		GenSourceManager->RegisterGenSource(this);
+	}
+}
+
+void UPCGGenSourceComponent::OnComponentDestroyed(bool bDestroyingHierarchy)
+{
+	Super::OnComponentDestroyed(bDestroyingHierarchy);
+
+	if (FPCGGenSourceManager* GenSourceManager = GetGenSourceManager())
+	{
+		GenSourceManager->UnregisterGenSource(this);
+	}
+}
+
+TOptional<FVector> UPCGGenSourceComponent::GetPosition() const
+{
+	if (AActor* Actor = GetOwner())
+	{
+		return Actor->GetActorLocation();
+	}
+
+	return TOptional<FVector>();
+}
+
+TOptional<FVector> UPCGGenSourceComponent::GetDirection() const
+{
+	if (AActor* Actor = GetOwner())
+	{
+		return Actor->GetActorForwardVector();
+	}
+
+	return TOptional<FVector>();
+}
+
+FPCGGenSourceManager* UPCGGenSourceComponent::GetGenSourceManager() const
+{
+	if (AActor* Actor = GetOwner())
+	{
+		if (UPCGSubsystem* Subsystem = UPCGSubsystem::GetInstance(Actor->GetWorld()))
+		{
+			return Subsystem->GetGenSourceManager();
+		}
+	}
+
+	return nullptr;
+}
+
+#undef LOCTEXT_NAMESPACE
diff --git a/Engine/Plugins/Experimental/PCG/Source/PCG/Private/RuntimeGen/GenSources/PCGGenSourceEditorCamera.cpp b/Engine/Plugins/Experimental/PCG/Source/PCG/Private/RuntimeGen/GenSources/PCGGenSourceEditorCamera.cpp
new file mode 100644
--- /dev/null
+++ b/Engine/Plugins/Experimental/PCG/Source/PCG/Private/RuntimeGen/GenSources/PCGGenSourceEditorCamera.cpp
@@ -0,0 +1,33 @@
+// Copyright Epic Games, Inc. All Rights Reserved.
+
+#include "RuntimeGen/GenSources/PCGGenSourceEditorCamera.h"
+
+#if WITH_EDITOR
+#include "EditorViewportClient.h"
+#endif
+
+#include UE_INLINE_GENERATED_CPP_BY_NAME(PCGGenSourceEditorCamera)
+
+TOptional<FVector> UPCGGenSourceEditorCamera::GetPosition() const
+{
+#if WITH_EDITOR
+	if (EditorViewportClient)
+	{
+		return EditorViewportClient->GetViewLocation();
+	}
+#endif
+
+	return TOptional<FVector>();
+}
+
+TOptional<FVector> UPCGGenSourceEditorCamera::GetDirection() const
+{
+#if WITH_EDITOR
+	if (EditorViewportClient)
+	{
+		return EditorViewportClient->GetViewRotation().Vector();
+	}
+#endif
+
+	return TOptional<FVector>();
+}
diff --git a/Engine/Plugins/Experimental/PCG/Source/PCG/Private/RuntimeGen/GenSources/PCGGenSourcePlayer.cpp b/Engine/Plugins/Experimental/PCG/Source/PCG/Private/RuntimeGen/GenSources/PCGGenSourcePlayer.cpp
new file mode 100644
--- /dev/null
+++ b/Engine/Plugins/Experimental/PCG/Source/PCG/Private/RuntimeGen/GenSources/PCGGenSourcePlayer.cpp
@@ -0,0 +1,39 @@
+// Copyright Epic Games, Inc. All Rights Reserved.
+
+#include "RuntimeGen/GenSources/PCGGenSourcePlayer.h"
+
+#include "GameFramework/Pawn.h"
+#include "GameFramework/PlayerController.h"
+
+#include UE_INLINE_GENERATED_CPP_BY_NAME(PCGGenSourcePlayer)
+
+TOptional<FVector> UPCGGenSourcePlayer::GetPosition() const
+{
+	if (PlayerController.IsValid())
+	{
+		if (APawn* PlayerPawn = PlayerController->GetPawn())
+		{
+			return PlayerPawn->GetActorLocation();
+		}
+	}
+
+	return TOptional<FVector>();
+}
+
+TOptional<FVector> UPCGGenSourcePlayer::GetDirection() const
+{
+	if (PlayerController.IsValid())
+	{
+		if (APawn* PlayerPawn = PlayerController->GetPawn())
+		{
+			return PlayerPawn->GetActorForwardVector();
+		}
+	}
+
+	return TOptional<FVector>();
+}
+
+void UPCGGenSourcePlayer::SetPlayerController(const APlayerController* InPlayerController)
+{
+	PlayerController = InPlayerController;
+}
diff --git a/Engine/Plugins/Experimental/PCG/Source/PCG/Private/RuntimeGen/GenSources/PCGGenSourceWPStreamingSource.cpp b/Engine/Plugins/Experimental/PCG/Source/PCG/Private/RuntimeGen/GenSources/PCGGenSourceWPStreamingSource.cpp
new file mode 100644
--- /dev/null
+++ b/Engine/Plugins/Experimental/PCG/Source/PCG/Private/RuntimeGen/GenSources/PCGGenSourceWPStreamingSource.cpp
@@ -0,0 +1,17 @@
+// Copyright Epic Games, Inc. All Rights Reserved.
+
+#include "RuntimeGen/GenSources/PCGGenSourceWPStreamingSource.h"
+
+#include "WorldPartition/WorldPartitionStreamingSource.h"
+
+#include UE_INLINE_GENERATED_CPP_BY_NAME(PCGGenSourceWPStreamingSource)
+
+TOptional<FVector> UPCGGenSourceWPStreamingSource::GetPosition() const
+{
+	return StreamingSource ? StreamingSource->Location : TOptional<FVector>();
+}
+
+TOptional<FVector> UPCGGenSourceWPStreamingSource::GetDirection() const
+{
+	return StreamingSource ? StreamingSource->Rotation.Vector() : TOptional<FVector>();
+}
diff --git a/Engine/Plugins/Experimental/PCG/Source/PCG/Private/RuntimeGen/PCGGenSourceManager.cpp b/Engine/Plugins/Experimental/PCG/Source/PCG/Private/RuntimeGen/PCGGenSourceManager.cpp
new file mode 100644
--- /dev/null
+++ b/Engine/Plugins/Experimental/PCG/Source/PCG/Private/RuntimeGen/PCGGenSourceManager.cpp
@@ -0,0 +1,147 @@
+// Copyright Epic Games, Inc. All Rights Reserved.
+
+#include "RuntimeGen/PCGGenSourceManager.h"
+
+#include "PCGWorldActor.h"
+#include "RuntimeGen/GenSources/PCGGenSourceBase.h"
+#include "RuntimeGen/GenSources/PCGGenSourceComponent.h"
+#include "RuntimeGen/GenSources/PCGGenSourceEditorCamera.h"
+#include "RuntimeGen/GenSources/PCGGenSourcePlayer.h"
+#include "RuntimeGen/GenSources/PCGGenSourceWPStreamingSource.h"
+
+#include "Engine/GameViewportClient.h"
+#include "Engine/World.h"
+#include "GameFramework/GameModeBase.h"
+#include "WorldPartition/WorldPartition.h"
+
+#if WITH_EDITOR
+#include "EditorViewportClient.h"
+#endif
+
+FPCGGenSourceManager::FPCGGenSourceManager(const UWorld* InWorld)
+{
+	// We capture the World so we can differentiate between editor-world and PIE-world Generation Sources.
+	World = InWorld;
+
+	// Collect PlayerControllers.
+	FGameModeEvents::GameModePostLoginEvent.AddRaw(this, &FPCGGenSourceManager::OnGameModePostLogin);
+	FGameModeEvents::GameModeLogoutEvent.AddRaw(this, &FPCGGenSourceManager::OnGameModePostLogout);
+}
+
+FPCGGenSourceManager::~FPCGGenSourceManager()
+{
+	FGameModeEvents::GameModePostLoginEvent.RemoveAll(this);
+	FGameModeEvents::GameModeLogoutEvent.RemoveAll(this);
+}
+
+TSet<IPCGGenSourceBase*> FPCGGenSourceManager::GetGenSources(const APCGWorldActor* InPCGWorldActor) const
+{
+	// Some the the generation sources need to be queried each frame, so build the 'current' set now.
+	TSet<IPCGGenSourceBase*> CurrentGenSources = GenSources;
+
+	if (const UWorld* InWorld = InPCGWorldActor->GetWorld())
+	{
+		if (const UWorldPartition* WorldPartition = InWorld->GetWorldPartition())
+		{
+			// TODO: Grab StreamingSourceProviders instead of StreamingSources? This could allow us to avoid calling NewObject every tick.
+			// Note: GetStreamingSources only works in GameWorld, so StreamingSources do not act as GenSources in editor.
+			for (const FWorldPartitionStreamingSource& Source : WorldPartition->GetStreamingSources())
+			{
+				// StreamingSources live only for this tick, so they don't need to be rooted.
+				UPCGGenSourceWPStreamingSource* GenSource = NewObject<UPCGGenSourceWPStreamingSource>();
+				GenSource->StreamingSource = &Source;
+
+				// TODO: Is it possible to avoid adding a StreamingSource for the Player, which we already capture in OnGameModePostLogin?
+				CurrentGenSources.Add(GenSource);
+			}
+		}
+	}
+
+#if WITH_EDITOR
+	if (InPCGWorldActor->bTreatEditorViewportAsGenerationSource && !World->IsGameWorld())
+	{
+		for (FEditorViewportClient* EditorViewportClient : GEditor->GetAllViewportClients())
+		{
+			if (EditorViewportClient->IsVisible())
+			{
+				// TODO: grab all 4 of these when GenSourceManager is created, instead of creating a new UObject for each every frame?
+				UPCGGenSourceEditorCamera* GenSource = NewObject<UPCGGenSourceEditorCamera>();
+				GenSource->EditorViewportClient = EditorViewportClient;
+
+				CurrentGenSources.Add(GenSource);
+			}
+		}
+	}
+#endif
+
+	return CurrentGenSources;
+}
+
+bool FPCGGenSourceManager::RegisterGenSource(IPCGGenSourceBase* InGenSource)
+{
+	if (UPCGGenSourceComponent* GenSourceComponent = Cast<UPCGGenSourceComponent>(InGenSource))
+	{
+		if (GenSourceComponent->GetWorld() != World)
+		{
+			return false;
+		}
+	}
+
+	return GenSources.Add(InGenSource).IsValidId();
+}
+
+bool FPCGGenSourceManager::UnregisterGenSource(const IPCGGenSourceBase* InGenSource)
+{
+	if (const UPCGGenSourceComponent* GenSourceComponent = Cast<UPCGGenSourceComponent>(InGenSource))
+	{
+		if (GenSourceComponent->GetWorld() != World)
+		{
+			return false;
+		}
+	}
+
+	return GenSources.Remove(InGenSource) > 0;
+}
+
+void FPCGGenSourceManager::OnGameModePostLogin(AGameModeBase* InGameMode, APlayerController* InPlayerController)
+{
+	if (InPlayerController->GetWorld() != World)
+	{
+		return;
+	}
+	ensure(IsInGameThread());
+
+	if (InPlayerController && InPlayerController->GetPawn())
+	{
+		UPCGGenSourcePlayer* GenSource = NewObject<UPCGGenSourcePlayer>();
+		GenSource->SetPlayerController(InPlayerController);
+
+		FSetElementId ElementId = GenSources.Add(GenSource);
+		if (ElementId.IsValidId())
+		{
+			GenSource->AddToRoot();
+		}
+	}
+}
+
+void FPCGGenSourceManager::OnGameModePostLogout(AGameModeBase* InGameMode, AController* InController)
+{
+	if (InController->GetWorld() != World)
+	{
+		return;
+	}
+	ensure(IsInGameThread());
+
+	for (IPCGGenSourceBase* GenSource : GenSources)
+	{
+		if (UPCGGenSourcePlayer* GenSourcePlayer = Cast<UPCGGenSourcePlayer>(GenSource))
+		{
+			if (GenSourcePlayer->GetPlayerController() == InController || !GenSourcePlayer->IsValid())
+			{
+				GenSourcePlayer->RemoveFromRoot();
+				GenSources.Remove(GenSourcePlayer);
+				break;
+			}
+		}
+	}
+}
diff --git a/Engine/Plugins/Experimental/PCG/Source/PCG/Private/RuntimeGen/PCGRuntimeGenScheduler.cpp b/Engine/Plugins/Experimental/PCG/Source/PCG/Private/RuntimeGen/PCGRuntimeGenScheduler.cpp
new file mode 100644
--- /dev/null
+++ b/Engine/Plugins/Experimental/PCG/Source/PCG/Private/RuntimeGen/PCGRuntimeGenScheduler.cpp
@@ -0,0 +1,954 @@
+// Copyright Epic Games, Inc. All Rights Reserved.
+
+#include "RuntimeGen/PCGRuntimeGenScheduler.h"
+#include "RuntimeGen/SchedulingPolicies/PCGSchedulingPolicyBase.h"
+
+#include "PCGActorAndComponentMapping.h"
+#include "PCGCommon.h"
+#include "PCGComponent.h"
+#include "PCGModule.h"
+#include "PCGSubsystem.h"
+#include "PCGWorldActor.h"
+#include "Grid/PCGPartitionActor.h"
+#include "Helpers/PCGActorHelpers.h"
+#include "Helpers/PCGHelpers.h"
+#include "RuntimeGen/GenSources/PCGGenSourceBase.h"
+#include "RuntimeGen/GenSources/PCGGenSourceComponent.h"
+#include "RuntimeGen/PCGGenSourceManager.h"
+
+#include "UObject/UObjectGlobals.h"
+
+#if WITH_EDITOR
+#include "EditorViewportClient.h"
+#endif
+
+namespace PCGRuntimeGenSchedulerConstants
+{
+	const FString PooledPartitionActorName = TEXT("PCGRuntimeGenPartitionActor_POOLED");
+}
+
+namespace PCGRuntimeGenSchedulerHelpers
+{
+	static TAutoConsoleVariable<bool> CVarRuntimeGenerationEnable(
+		TEXT("pcg.RuntimeGeneration.Enable"),
+		true,
+		TEXT("Enable the RuntimeGeneration system."));
+
+	static TAutoConsoleVariable<int32> CVarFramesBetweenGraphSchedules(
+		TEXT("pcg.RuntimeGeneration.FramesBetweenGraphSchedules"),
+		0,
+		TEXT("Defines the minimum number of frames/ticks between any two Generation schedules in the RuntimeGenScheduler."));
+
+	static TAutoConsoleVariable<bool> CVarRuntimeGenerationEnableDebugging(
+		TEXT("pcg.RuntimeGeneration.EnableDebugging"),
+		false,
+		TEXT("Enable verbose debug logging for the RuntimeGeneration system."));
+
+	static TAutoConsoleVariable<bool> CVarRuntimeGenerationEnablePooling(
+		TEXT("pcg.RuntimeGeneration.EnablePooling"),
+		true,
+		TEXT("Enable PartitionActor pooling for the RuntimeGeneration system."));
+
+	static TAutoConsoleVariable<int32> CVarRuntimeGenerationBasePoolSize(
+		TEXT("pcg.RuntimeGeneration.BasePoolSize"),
+		100,
+		TEXT("Defines the base PartitionActor pool size for the RuntimeGeneration system. Cannot be less than 1."));
+}
+
+FPCGRuntimeGenScheduler::FPCGRuntimeGenScheduler(UWorld* InWorld, FPCGActorAndComponentMapping* InActorAndComponentMapping)
+{
+	check(InWorld && InActorAndComponentMapping);
+
+	World = InWorld;
+	ActorAndComponentMapping = InActorAndComponentMapping;
+	GenSourceManager = new FPCGGenSourceManager(InWorld);
+	bPoolingWasEnabledLastFrame = PCGRuntimeGenSchedulerHelpers::CVarRuntimeGenerationEnablePooling.GetValueOnAnyThread();
+	BasePoolSizeLastFrame = PCGRuntimeGenSchedulerHelpers::CVarRuntimeGenerationBasePoolSize.GetValueOnAnyThread();
+}
+
+FPCGRuntimeGenScheduler::~FPCGRuntimeGenScheduler()
+{
+	delete GenSourceManager;
+	GenSourceManager = nullptr;
+}
+
+void FPCGRuntimeGenScheduler::Tick(const APCGWorldActor* InPCGWorldActor)
+{
+	TRACE_CPUPROFILER_EVENT_SCOPE(FPCGRuntimeGenScheduler::Tick);
+	check(InPCGWorldActor && GenSourceManager);
+
+	// 0. Preamble - check if we should be active in this world and do lazy initialization.
+
+	if (!ShouldTick())
+	{
+		return;
+	}
+
+	TickCVars(InPCGWorldActor);
+
+	const TSet<IPCGGenSourceBase*> GenSources = GenSourceManager->GetGenSources(InPCGWorldActor);
+
+	// Initialize RuntimeGen PA pool if necessary. If PoolSize is 0, then we have not initialized the pool yet.
+	if (!GenSources.IsEmpty() || !GeneratedComponents.IsEmpty())
+	{
+		if (PartitionActorPoolSize == 0 && PCGRuntimeGenSchedulerHelpers::CVarRuntimeGenerationEnablePooling.GetValueOnAnyThread())
+		{
+			AddPartitionActorPoolCount(PCGRuntimeGenSchedulerHelpers::CVarRuntimeGenerationBasePoolSize.GetValueOnAnyThread());
+		}
+	}
+
+	// 1. Queue nearby components for generation.
+	
+	// Mapping of component + coordinates to priorities - needed to compute max priority over all gen sources.
+	TMap<FGridGenerationKey, double> ComponentsToGenerate;
+	TickQueueComponentsForGeneration(GenSources, InPCGWorldActor, ComponentsToGenerate);
+
+	// 2. Schedule cleanup on components that become out of range.
+
+	TickCleanup(GenSources, InPCGWorldActor);
+
+	// 3. Schedule generation on components in priority order.
+
+	if (ScheduleFrameCounter <= 0)
+	{
+		// Sort components by priority (will be generated in descending order).
+		ComponentsToGenerate.ValueSort([](double PrioA, double PrioB)->bool { return PrioA > PrioB; });
+
+		TickScheduleGeneration(ComponentsToGenerate);
+	}
+	else
+	{
+		--ScheduleFrameCounter;
+	}
+}
+
+bool FPCGRuntimeGenScheduler::ShouldTick()
+{
+	check(World);
+
+	if (!PCGRuntimeGenSchedulerHelpers::CVarRuntimeGenerationEnable.GetValueOnAnyThread())
+	{
+		return false;
+	}
+
+	// Disable tick of editor scheduling if in runtime or PIE.
+	if (PCGHelpers::IsRuntimeOrPIE() && !World->IsGameWorld())
+	{
+		return false;
+	}
+
+#if WITH_EDITOR
+	// If we're in an editor world, stop updating preview if the editor window/viewport is not active (follows
+	// same behaviour as other things).
+	if (!World->IsGameWorld())
+	{
+		bool bAnyVisible = false;
+		for (FEditorViewportClient* EditorViewportClient : GEditor->GetAllViewportClients())
+		{
+			bAnyVisible |= EditorViewportClient->IsVisible();
+		}
+
+		if (!bAnyVisible)
+		{
+			return false;
+		}
+	}
+#endif
+
+	return true;
+}
+
+void FPCGRuntimeGenScheduler::TickQueueComponentsForGeneration(
+	const TSet<IPCGGenSourceBase*>& InGenSources,
+	const APCGWorldActor* InPCGWorldActor,
+	TMap<FGridGenerationKey, double>& OutComponentsToGenerate)
+{
+	TRACE_CPUPROFILER_EVENT_SCOPE(FPCGRuntimeGenScheduler::TickQueueComponentsForGeneration);
+
+	check(ActorAndComponentMapping && InPCGWorldActor);
+
+	// TODO: Thought - it would be possible to maintain a global maximum generation distance across all components
+	// perhaps in the actor&comp mapping system, and then do a spatial query to get the components here.
+
+	auto AddComponentToGenerate = [&OutComponentsToGenerate](FGridGenerationKey& InKey, const IPCGGenSourceBase* InGenSource, const UPCGSchedulingPolicyBase* InPolicy, const FBox& InComponentBounds, bool bInUse2DGrid)
+	{
+		const double PolicyPriority = InPolicy ? InPolicy->CalculatePriority(InGenSource, InComponentBounds, bInUse2DGrid) : 0.0;
+		double Priority = FMath::Clamp(PolicyPriority, 0.0, 1.0);
+		if (PolicyPriority != Priority)
+		{
+			UE_LOG(LogPCG, Warning, TEXT("Priority from runtime generation policy (%lf) outside [0.0, 1.0] range, clamped."), PolicyPriority);
+		}
+
+		// Generate largest grid to smallest (and unbounded is larger than any grid).
+		const uint32 GridSize = InKey.GetGridSize();
+		Priority += GridSize;
+
+		double* ExistingPriority = OutComponentsToGenerate.Find(InKey);
+		if (!ExistingPriority)
+		{
+			OutComponentsToGenerate.Add(InKey, Priority);
+		}
+		else if (Priority > *ExistingPriority)
+		{
+			// If this generation source prioritizes this grid cell higher, then bump the priority.
+			*ExistingPriority = Priority;
+		}
+	};
+
+	// Collect local components from all partitioned components.
+	for (UPCGComponent* OriginalComponent : ActorAndComponentMapping->PartitionedOctree.GetAllComponents())
+	{
+		TRACE_CPUPROFILER_EVENT_SCOPE(FPCGRuntimeGenScheduler::CollectLocalComponents);
+
+		if (!ensure(OriginalComponent))
+		{
+			continue;
+		}
+
+		const UPCGSchedulingPolicyBase* Policy = OriginalComponent->GetRuntimeGenSchedulingPolicy();
+
+		// TODO: For each execution domain (for now only GenAtRuntime/dynamic), assuming we run Preview through this scheduler, which it seems like we will.
+		if (OriginalComponent->GenerationTrigger == EPCGComponentGenerationTrigger::GenerateAtRuntime)
+		{
+			bool bHasUnbounded = false;
+			PCGHiGenGrid::FSizeArray GridSizes;
+			ensure(PCGHelpers::GetGenerationGridSizes(OriginalComponent->GetGraph(), InPCGWorldActor, GridSizes, bHasUnbounded));
+
+			if (GridSizes.IsEmpty())
+			{
+				continue;
+			}
+
+			const double MaxGenerationRadius = OriginalComponent->GetGenerationRadiusFromGrid(PCGHiGenGrid::GridSizeToGrid(GridSizes[0]));
+
+			for (const IPCGGenSourceBase* GenSource : InGenSources)
+			{
+				const TOptional<FVector> GenSourcePositionOptional = GenSource->GetPosition();
+				if (!ensure(GenSourcePositionOptional.IsSet()))
+				{
+					// If the gen source fails to return a valid result, skip it, to avoid mistakenly considering the origin as a valid gen source.
+					continue;
+				}
+
+				const FVector GenSourcePosition = GenSourcePositionOptional.GetValue();
+				const FBox OriginalComponentBounds = OriginalComponent->GetGridBounds();
+
+				FVector ModifiedGenSourcePosition = GenSourcePosition;
+				if (InPCGWorldActor->bUse2DGrid)
+				{
+					ModifiedGenSourcePosition.Z = OriginalComponentBounds.Min.Z;
+				}
+
+				const double DistanceSquared = OriginalComponentBounds.ComputeSquaredDistanceToPoint(ModifiedGenSourcePosition);
+
+				if (DistanceSquared > MaxGenerationRadius * MaxGenerationRadius)
+				{
+					// GenSource is not within range of the component, skip!
+					continue;
+				}
+
+				if (bHasUnbounded)
+				{
+					// Ignore components that have already been generated or marked for generation. Unbounded grid size means not-partitioned.
+					FGridGenerationKey Key(PCGHiGenGrid::UnboundedGridSize(), FIntVector(0), OriginalComponent);
+					if (!GeneratedComponents.Contains(Key))
+					{
+						AddComponentToGenerate(Key, GenSource, Policy, OriginalComponentBounds, InPCGWorldActor->bUse2DGrid);
+					}
+				}
+
+				// TODO: once one of the larger grid sizes is out of range, we can forego checking any smaller grid sizes. they can't possibly be closer!
+				// This assumes we enforce generation radii to increase monotonically.
+				for (const uint32 GridSize : GridSizes)
+				{
+					ensure(PCGHiGenGrid::IsValidGridSize(GridSize));
+
+					const FIntVector GenSourceGridPosition = UPCGActorHelpers::GetCellCoord(GenSourcePosition, GridSize, InPCGWorldActor->bUse2DGrid);
+					const double GenerationRadius = OriginalComponent->GetGenerationRadiusFromGrid(PCGHiGenGrid::GridSizeToGrid(GridSize));
+					const int32 GridRadius = FMath::CeilToInt32(GenerationRadius / GridSize); // Radius discretized to # of grid cells.
+					const int32 VerticalGridRadius = InPCGWorldActor->bUse2DGrid ? 0 : GridRadius; // Flatten the vertical grid radius in the 2D case.
+
+					const double HalfGridSize = GridSize / 2.0f;
+					FVector HalfExtent(HalfGridSize, HalfGridSize, HalfGridSize);
+
+					FBox ModifiedBounds = OriginalComponent->GetGridBounds();
+					if (InPCGWorldActor->bUse2DGrid)
+					{
+						FVector MinBounds = ModifiedBounds.Min;
+						FVector MaxBounds = ModifiedBounds.Max;
+
+						MinBounds.Z = 0;
+						MaxBounds.Z = FMath::Max<double>(0.0, GridSize - 1.0); // -1 to be just below the GridSize.
+						ModifiedBounds = FBox(MinBounds, MaxBounds);
+
+						// In case of 2D grid, it's like the actor has infinite bounds on the Z axis.
+						HalfExtent.Z = HALF_WORLD_MAX1;
+					}
+
+					// TODO: Perhaps rasterize sphere instead of walking a naive cube. although maybe the perf on that isn't worthwhile.
+					for (int32 Z = GenSourceGridPosition.Z - VerticalGridRadius; Z <= GenSourceGridPosition.Z + VerticalGridRadius; ++Z)
+					{
+						for (int32 Y = GenSourceGridPosition.Y - GridRadius; Y <= GenSourceGridPosition.Y + GridRadius; ++Y)
+						{
+							for (int32 X = GenSourceGridPosition.X - GridRadius; X <= GenSourceGridPosition.X + GridRadius; ++X)
+							{
+								FIntVector GridCoords(X, Y, Z);
+								FGridGenerationKey Key(GridSize, GridCoords, OriginalComponent);
+
+								// Ignore components that have already been generated or marked for generation.
+								if (GeneratedComponents.Find({ GridSize, GridCoords, OriginalComponent }))
+								{
+									continue;
+								}
+
+								const FVector Center = FVector(GridCoords.X + 0.5, GridCoords.Y + 0.5, GridCoords.Z + 0.5) * GridSize;
+								const FBox CellBounds(Center - HalfExtent, Center + HalfExtent);
+
+								// Overlap cell with the partitioned component.
+								const FBox IntersectedBounds = ModifiedBounds.Overlap(CellBounds);
+								if (!IntersectedBounds.IsValid || IntersectedBounds.GetVolume() <= UE_DOUBLE_SMALL_NUMBER)
+								{
+									continue;
+								}
+
+								if (InPCGWorldActor->bUse2DGrid)
+								{
+									ModifiedGenSourcePosition.Z = CellBounds.Min.Z;
+								}
+
+								// Verify the grid cell actually lies within the generation radius.
+								// TODO: this is no longer necessary if we rasterize the sphere instead.
+								const double LocalDistanceSquared = CellBounds.ComputeSquaredDistanceToPoint(ModifiedGenSourcePosition);
+								if (LocalDistanceSquared <= GenerationRadius * GenerationRadius)
+								{
+									AddComponentToGenerate(Key, GenSource, Policy, IntersectedBounds, InPCGWorldActor->bUse2DGrid);
+								}
+							}
+						}
+					}
+				}
+			}
+		}
+	}
+
+	// Collect all non-partitioned components.
+	for (UPCGComponent* OriginalComponent : ActorAndComponentMapping->NonPartitionedOctree.GetAllComponents())
+	{
+		TRACE_CPUPROFILER_EVENT_SCOPE(FPCGRuntimeGenScheduler::CollectNonPartitionedComponents);
+
+		if (!ensure(OriginalComponent))
+		{
+			continue;
+		}
+
+		// The generation key for a non-partitioned component should always have unbounded grid size and 0,0,0 cell coord.
+		if (GeneratedComponents.Contains({ PCGHiGenGrid::UnboundedGridSize(), FIntVector(0), OriginalComponent }))
+		{
+			continue;
+		}
+
+		const UPCGSchedulingPolicyBase* Policy = OriginalComponent->GetRuntimeGenSchedulingPolicy();
+
+		// TODO: For each execution domain (for now only GenAtRuntime/dynamic), assuming we run Preview through this scheduler, which it seems like we will.
+		if (OriginalComponent->GenerationTrigger == EPCGComponentGenerationTrigger::GenerateAtRuntime)
+		{
+			// Unbounded will grab the base GenerationRadius used for non-partitioned and unbounded.
+			const double MaxGenerationRadius = OriginalComponent->GetGenerationRadiusFromGrid(EPCGHiGenGrid::Unbounded);
+
+			for (const IPCGGenSourceBase* GenSource : InGenSources)
+			{
+				const TOptional<FVector> GenSourcePositionOptional = GenSource->GetPosition();
+				if (!ensure(GenSourcePositionOptional.IsSet()))
+				{
+					// If the gen source fails to return a valid result, skip it, to avoid mistakenly considering the origin as a valid gen source.
+					continue;
+				}
+
+				const FVector GenSourcePosition = GenSourcePositionOptional.GetValue();
+				const FBox OriginalComponentBounds = OriginalComponent->GetGridBounds();
+
+				FVector ModifiedGenSourcePosition = GenSourcePosition;
+				if (InPCGWorldActor->bUse2DGrid)
+				{
+					ModifiedGenSourcePosition.Z = OriginalComponentBounds.Min.Z;
+				}
+
+				const double DistanceSquared = OriginalComponentBounds.ComputeSquaredDistanceToPoint(ModifiedGenSourcePosition);
+
+				// Max radius for a non-partitioned component is just the base GenerationRadius.
+				if (DistanceSquared <= MaxGenerationRadius * MaxGenerationRadius)
+				{
+					// Unbounded grid size means not-partitioned.
+					FGridGenerationKey Key(PCGHiGenGrid::UnboundedGridSize(), FIntVector(0), OriginalComponent);
+					AddComponentToGenerate(Key, GenSource, Policy, OriginalComponentBounds, /*bUse2DGrid=*/false);
+				}
+			}
+		}
+	}
+}
+
+void FPCGRuntimeGenScheduler::TickCleanup(const TSet<IPCGGenSourceBase*>& InGenSources, const APCGWorldActor* InPCGWorldActor)
+{
+	TRACE_CPUPROFILER_EVENT_SCOPE(FPCGRuntimeGenScheduler::TickCleanup);
+
+	check(ActorAndComponentMapping && InPCGWorldActor);
+
+	// Generated Entry Key, Local/Generated Component
+	using PCGComponentToClean = TTuple<FGridGenerationKey, UPCGComponent*>;
+	TArray<PCGComponentToClean> ComponentsToClean;
+
+	// Find any generated components which should be cleaned.
+	for (FGridGenerationKey GenerationKey : GeneratedComponents)
+	{
+		TRACE_CPUPROFILER_EVENT_SCOPE(FPCGRuntimeGenScheduler::SelectComponentsToCleanup);
+
+		const uint32 GridSize = GenerationKey.GetGridSize();
+		const FIntVector& GridCoords = GenerationKey.GetGridCoords();
+		UPCGComponent* OriginalComponent = GenerationKey.GetOriginalComponent();
+		check(OriginalComponent);
+
+		const EPCGHiGenGrid Grid = PCGHiGenGrid::GridSizeToGrid(GridSize);
+
+		// If the Grid is unbounded, we have a non-partitioned or unbounded component.
+		if (Grid == EPCGHiGenGrid::Unbounded)
+		{
+			const FBox Bounds = OriginalComponent->GetGridBounds();
+
+			double MinSquaredDistanceToGenSource = UE_DOUBLE_BIG_NUMBER;
+			for (const IPCGGenSourceBase* GenSource : InGenSources)
+			{
+				if (!ensure(GenSource))
+				{
+					continue;
+				}
+
+				const TOptional<FVector> GenSourcePositionOptional = GenSource->GetPosition();
+				if (!ensure(GenSourcePositionOptional.IsSet()))
+				{
+					// If the gen source fails to return a valid result, skip it, to avoid mistakenly considering the origin as a valid gen source.
+					continue;
+				}
+
+				FVector GenSourcePosition = GenSourcePositionOptional.GetValue();
+
+				// Only consider 2D distance when using a 2D grid.
+				if (InPCGWorldActor->bUse2DGrid)
+				{
+					GenSourcePosition.Z = Bounds.Min.Z;
+				}
+
+				MinSquaredDistanceToGenSource = FMath::Min(MinSquaredDistanceToGenSource, Bounds.ComputeSquaredDistanceToPoint(GenSourcePosition));
+			}
+
+			const double CleanupRadius = OriginalComponent->GetCleanupRadiusFromGrid(Grid);
+			if (MinSquaredDistanceToGenSource > CleanupRadius * CleanupRadius)
+			{
+				ComponentsToClean.Add({ GenerationKey, OriginalComponent });
+			}
+		}
+		// Otherwise, we have a local component.
+		else
+		{
+			UPCGComponent* LocalComponent = ActorAndComponentMapping->GetLocalComponent(GridSize, GridCoords, OriginalComponent, /*bRuntimeGenerated=*/true);
+			if (!ensure(LocalComponent))
+			{
+				UE_LOG(LogPCG, Error, TEXT("[RUNTIMEGEN] Generated local component could not be retrieved: %u, %d_%d_%d"), GridSize, GridCoords.X, GridCoords.Y, GridCoords.Z);
+				continue;
+			}
+
+			const FBox GridBounds = LocalComponent->GetGridBounds();
+
+			double MinSquaredDistanceToGenSource = UE_DOUBLE_BIG_NUMBER;
+			for (const IPCGGenSourceBase* GenSource : InGenSources)
+			{
+				if (!ensure(GenSource))
+				{
+					continue;
+				}
+
+				const TOptional<FVector> GenSourcePositionOptional = GenSource->GetPosition();
+				if (!ensure(GenSourcePositionOptional.IsSet()))
+				{
+					// If the gen source fails to return a valid result, skip it, to avoid mistakenly considering the origin as a valid gen source.
+					continue;
+				}
+
+				FVector GenSourcePosition = GenSourcePositionOptional.GetValue();
+
+				// Only consider 2D distance when using a 2D grid.
+				if (InPCGWorldActor->bUse2DGrid)
+				{
+					GenSourcePosition.Z = GridBounds.Min.Z;
+				}
+
+				MinSquaredDistanceToGenSource = FMath::Min(MinSquaredDistanceToGenSource, GridBounds.ComputeSquaredDistanceToPoint(GenSourcePosition));
+			}
+
+			const double CleanupRadius = OriginalComponent->GetCleanupRadiusFromGrid(Grid);
+			if (MinSquaredDistanceToGenSource > CleanupRadius * CleanupRadius)
+			{
+				ComponentsToClean.Add({ GenerationKey, LocalComponent });
+			}
+		}
+	}
+
+	for (const PCGComponentToClean& ComponentToClean : ComponentsToClean)
+	{
+		CleanupComponent(ComponentToClean.Get<0>(), ComponentToClean.Get<1>());
+	}
+}
+
+void FPCGRuntimeGenScheduler::TickScheduleGeneration(TMap<FGridGenerationKey, double>& ComponentsToGenerate)
+{
+	TRACE_CPUPROFILER_EVENT_SCOPE(FPCGRuntimeGenScheduler::TickScheduleGeneration);
+
+	check(ActorAndComponentMapping);
+
+	for (const TPair<FGridGenerationKey, double>& GenerateEntry : ComponentsToGenerate)
+	{
+		const FGridGenerationKey& Key = GenerateEntry.Key;
+		const double Priority = GenerateEntry.Value;
+
+		const uint32 GridSize = GenerateEntry.Key.GetGridSize();
+		const FIntVector GridCoords = GenerateEntry.Key.GetGridCoords();
+		UPCGComponent* OriginalComponent = GenerateEntry.Key.GetOriginalComponent();
+		check(OriginalComponent);
+
+		const EPCGHiGenGrid Grid = PCGHiGenGrid::GridSizeToGrid(GridSize);
+
+		// If the Grid is unbounded, we have a non-partitioned or unbounded component.
+		if (Grid == EPCGHiGenGrid::Unbounded)
+		{
+			if (ensure(OriginalComponent && !OriginalComponent->bGenerated) && !OriginalComponent->IsGenerating())
+			{
+				if (PCGRuntimeGenSchedulerHelpers::CVarRuntimeGenerationEnableDebugging.GetValueOnAnyThread() && OriginalComponent->GetOwner())
+				{
+					UE_LOG(LogPCG, Warning, TEXT("[RUNTIMEGEN] GENERATE: '%s' (priority %lf)"), *OriginalComponent->GetOwner()->GetActorNameOrLabel(), Priority);
+				}
+
+				OriginalComponent->GenerateLocal(EPCGComponentGenerationTrigger::GenerateAtRuntime, /*bForce=*/false, Grid);
+			}
+		}
+		// Otherwise we have a local component.
+		else
+		{
+			TObjectPtr<APCGPartitionActor> PartitionActor = nullptr;
+
+			if (PCGRuntimeGenSchedulerHelpers::CVarRuntimeGenerationEnablePooling.GetValueOnAnyThread())
+			{
+				if (PCGRuntimeGenSchedulerHelpers::CVarRuntimeGenerationEnableDebugging.GetValueOnAnyThread())
+				{
+					UE_LOG(LogPCG, Warning, TEXT("[RUNTIMEGEN] UNPOOL PARTITION ACTOR: '%s' (priority %lf, %d remaining out of %d)"),
+						*APCGPartitionActor::GetRuntimeGenActorName(GridSize, GridCoords),
+						Priority,
+						PartitionActorPool.Num(),
+						PartitionActorPoolSize);
+				}
+
+				// Get RuntimeGenPA from pool.
+				PartitionActor = GetPartitionActorFromPool(GridSize, GridCoords);
+			}
+			else
+			{
+				if (PCGRuntimeGenSchedulerHelpers::CVarRuntimeGenerationEnableDebugging.GetValueOnAnyThread())
+				{
+					UE_LOG(LogPCG, Warning, TEXT("[RUNTIMEGEN] CREATE PARTITION ACTOR: '%s' (priority %lf)"),
+						*APCGPartitionActor::GetRuntimeGenActorName(GridSize, GridCoords),
+						Priority);
+				}
+
+				// Find or Create RuntimeGenPA.
+				PartitionActor = FindOrCreatePartitionActor(GridSize, GridCoords);
+			}
+
+			if (!ensure(PartitionActor))
+			{
+				continue;
+			}
+
+			// Update component mapping for this PA (add local component).
+			{
+				FWriteScopeLock WriteLock(ActorAndComponentMapping->ComponentToRuntimeGenPartitionActorsMapLock);
+				TSet<TObjectPtr<APCGPartitionActor>>* PartitionActorsPtr = ActorAndComponentMapping->ComponentToRuntimeGenPartitionActorsMap.Find(OriginalComponent);
+
+				if (!PartitionActorsPtr)
+				{
+					PartitionActorsPtr = &ActorAndComponentMapping->ComponentToRuntimeGenPartitionActorsMap.Emplace(OriginalComponent);
+				}
+
+				PartitionActor->AddGraphInstance(OriginalComponent);
+				PartitionActorsPtr->Add(PartitionActor);
+			}
+
+			// Generate local component.
+			UPCGComponent* LocalComponent = PartitionActor->GetLocalComponent(OriginalComponent);
+
+			if (ensure(LocalComponent && !LocalComponent->IsGenerating()))
+			{
+				if (PCGRuntimeGenSchedulerHelpers::CVarRuntimeGenerationEnableDebugging.GetValueOnAnyThread())
+				{
+					UE_LOG(LogPCG, Warning, TEXT("[RUNTIMEGEN] GENERATE: '%s' (priority %lf)"), *PartitionActor->GetActorNameOrLabel(), Priority);
+				}
+
+				LocalComponent->GenerateLocal(EPCGComponentGenerationTrigger::GenerateAtRuntime, /*bForce=*/false, LocalComponent->GetGenerationGrid());
+			}
+		}
+
+		GeneratedComponents.Add(Key);
+
+		ScheduleFrameCounter += PCGRuntimeGenSchedulerHelpers::CVarFramesBetweenGraphSchedules.GetValueOnAnyThread();
+		if (ScheduleFrameCounter > 0)
+		{
+			break;
+		}
+	}
+}
+
+void FPCGRuntimeGenScheduler::TickCVars(const APCGWorldActor* InPCGWorldActor)
+{
+	// If pooling has been disabled since last frame, we should destroy the pool.
+	const bool bPoolingEnabled = PCGRuntimeGenSchedulerHelpers::CVarRuntimeGenerationEnablePooling.GetValueOnAnyThread();
+
+	if (bPoolingWasEnabledLastFrame && !bPoolingEnabled)
+	{
+		CleanupLocalComponents(InPCGWorldActor);
+		ResetPartitionActorPoolToSize(/*NewPoolSize=*/0);
+	}
+
+	bPoolingWasEnabledLastFrame = bPoolingEnabled;
+
+	// Handle when the base PA PoolSize is modified. Cleanup all local components and reset the pool with the correct number of PAs.
+	if (PCGRuntimeGenSchedulerHelpers::CVarRuntimeGenerationEnablePooling.GetValueOnAnyThread())
+	{
+		// Don't allow a pool size <= 0
+		const uint32 BasePoolSize = FMath::Max(1, PCGRuntimeGenSchedulerHelpers::CVarRuntimeGenerationBasePoolSize.GetValueOnAnyThread());
+
+		if (BasePoolSizeLastFrame != BasePoolSize)
+		{
+			BasePoolSizeLastFrame = BasePoolSize;
+
+			CleanupLocalComponents(InPCGWorldActor);
+			ResetPartitionActorPoolToSize(BasePoolSize);
+		}
+	}
+}
+
+void FPCGRuntimeGenScheduler::OnOriginalComponentUnregistered(UPCGComponent* InOriginalComponent)
+{
+	check(ActorAndComponentMapping);
+
+	// Safety checks.
+	if (!InOriginalComponent || Cast<APCGPartitionActor>(InOriginalComponent->GetOwner()))
+	{
+		ensure(false);
+		return;
+	}
+
+	// Gather all generated components which originated from this original component.
+	TSet<FGridGenerationKey> KeysToCleanup;
+	for (FGridGenerationKey GenerationKey : GeneratedComponents)
+	{
+		if (GenerationKey.GetOriginalComponent() == InOriginalComponent)
+		{
+			KeysToCleanup.Add(GenerationKey);
+		}
+	}
+
+	for (FGridGenerationKey GenerationKey : KeysToCleanup)
+	{
+		const uint32 GridSize = GenerationKey.GetGridSize();
+		const bool bIsOriginalComponent = GridSize == PCGHiGenGrid::UnboundedGridSize();
+
+		// Get the generated component for this key (might be a local component).
+		UPCGComponent* ComponentToCleanup = bIsOriginalComponent ? InOriginalComponent : ActorAndComponentMapping->GetLocalComponent(GridSize, GenerationKey.GetGridCoords(), InOriginalComponent, /*bRuntimeGenerated=*/true);
+
+		if (ensure(ComponentToCleanup))
+		{
+			CleanupComponent(GenerationKey, ComponentToCleanup);
+		}
+	}
+}
+
+void FPCGRuntimeGenScheduler::CleanupLocalComponents(const APCGWorldActor* InPCGWorldActor)
+{
+	TRACE_CPUPROFILER_EVENT_SCOPE(FPCGRuntimeGenScheduler::CleanupLocalComponents);
+
+	check(ActorAndComponentMapping && InPCGWorldActor);
+
+	// Generated Entry Key, Local/Generated Component
+	using PCGComponentToClean = TTuple<FGridGenerationKey, UPCGComponent*>;
+	TArray<PCGComponentToClean> ComponentsToClean;
+
+	// Find all generated local components.
+	for (FGridGenerationKey GenerationKey : GeneratedComponents)
+	{
+		const uint32 GridSize = GenerationKey.GetGridSize();
+		const FIntVector& GridCoords = GenerationKey.GetGridCoords();
+		UPCGComponent* OriginalComponent = GenerationKey.GetOriginalComponent();
+		check(OriginalComponent);
+
+		const EPCGHiGenGrid Grid = PCGHiGenGrid::GridSizeToGrid(GridSize);
+
+		// Only operate on LocalComponents.
+		if (Grid != EPCGHiGenGrid::Unbounded)
+		{
+			UPCGComponent* LocalComponent = ActorAndComponentMapping->GetLocalComponent(GridSize, GridCoords, OriginalComponent, /*bRuntimeGenerated=*/true);
+			if (!ensure(LocalComponent))
+			{
+				UE_LOG(LogPCG, Error, TEXT("[RUNTIMEGEN] Generated local component could not be retrieved: %u, %d_%d_%d"), GridSize, GridCoords.X, GridCoords.Y, GridCoords.Z);
+				continue;
+			}
+
+			ComponentsToClean.Add({ GenerationKey, LocalComponent });
+		}
+	}
+
+	for (const PCGComponentToClean& ComponentToClean : ComponentsToClean)
+	{
+		CleanupComponent(ComponentToClean.Get<0>(), ComponentToClean.Get<1>());
+	}
+}
+
+void FPCGRuntimeGenScheduler::CleanupComponent(const FGridGenerationKey& GenerationKey, UPCGComponent* GeneratedComponent)
+{
+	TRACE_CPUPROFILER_EVENT_SCOPE(FPCGRuntimeGenScheduler::CleanupComponent);
+
+	check(ActorAndComponentMapping);
+
+	const uint32 GridSize = GenerationKey.GetGridSize();
+	const FIntVector GridCoords = GenerationKey.GetGridCoords();
+
+	APCGPartitionActor* PartitionActor = nullptr;
+
+	if (!ensure(GeneratedComponent))
+	{
+		UE_LOG(LogPCG, Error, TEXT("[RUNTIMEGEN] GENERATED COMPONENT COULD NOT BE RECOVERED ON GRID %d AT (%d, %d, %d)"), GridSize, GridCoords.X, GridCoords.Y, GridCoords.Z);
+
+		// If the GeneratedComponent has been lost for some reason, get the PA directly from the ActorAndComponentMapping.
+		PartitionActor = ActorAndComponentMapping->GetPartitionActor(GridSize, GridCoords, /*bRuntimeGenerated=*/true);
+	}
+	else // If the GeneratedComponent does exist, we can clean it up!.
+	{
+		PartitionActor = Cast<APCGPartitionActor>(GeneratedComponent->GetOwner());
+		const EPCGHiGenGrid Grid = PCGHiGenGrid::GridSizeToGrid(GridSize);
+
+		// If the Grid is unbounded, we have a non-partitioned or unbounded component.
+		if (Grid == EPCGHiGenGrid::Unbounded)
+		{
+			if (PCGRuntimeGenSchedulerHelpers::CVarRuntimeGenerationEnableDebugging.GetValueOnAnyThread() && GeneratedComponent->GetOwner())
+			{
+				UE_LOG(LogPCG, Warning, TEXT("[RUNTIMEGEN] CLEANUP: '%s'"), *GeneratedComponent->GetOwner()->GetActorNameOrLabel());
+			}
+
+			GeneratedComponent->CleanupLocalImmediate(/*bRemoveComponents=*/true);
+		}
+		// Otherwise we have a local component.
+		else if (ensure(PartitionActor))
+		{
+			UPCGComponent* OriginalComponent = PartitionActor->GetOriginalComponent(GeneratedComponent);
+
+			if (PCGRuntimeGenSchedulerHelpers::CVarRuntimeGenerationEnableDebugging.GetValueOnAnyThread())
+			{
+				UE_LOG(LogPCG, Warning, TEXT("[RUNTIMEGEN] CLEANUP: '%s'"), *PartitionActor->GetActorNameOrLabel());
+			}
+
+			// This performs a CleanupLocalImmediate for us, no need to clean up ourselves.
+			PartitionActor->RemoveGraphInstance(OriginalComponent);
+
+			// Remove component mapping.
+			{
+				FWriteScopeLock WriteLock(ActorAndComponentMapping->ComponentToRuntimeGenPartitionActorsMapLock);
+				TSet<TObjectPtr<APCGPartitionActor>>* PartitionActorsPtr = ActorAndComponentMapping->ComponentToRuntimeGenPartitionActorsMap.Find(OriginalComponent);
+
+				if (PartitionActorsPtr)
+				{
+					PartitionActorsPtr->Remove(PartitionActor);
+
+					if (PartitionActorsPtr->IsEmpty())
+					{
+						ActorAndComponentMapping->ComponentToRuntimeGenPartitionActorsMap.Remove(OriginalComponent);
+					}
+				}
+			}
+		}
+	}
+
+	// Cleanup the PA if it no longer has any components (return to pool or destroy).
+	if (PartitionActor && PartitionActor->GetAllLocalPCGComponents().IsEmpty())
+	{
+		PartitionActor->UnregisterPCG();
+
+		if (PCGRuntimeGenSchedulerHelpers::CVarRuntimeGenerationEnablePooling.GetValueOnAnyThread())
+		{
+			if (PCGRuntimeGenSchedulerHelpers::CVarRuntimeGenerationEnableDebugging.GetValueOnAnyThread())
+			{
+				UE_LOG(LogPCG, Warning, TEXT("[RUNTIMEGEN] RETURNING PARTITION ACTOR TO POOL: '%s' (%d remaining out of %d)"), *PartitionActor->GetActorNameOrLabel(), PartitionActorPool.Num() + 1, PartitionActorPoolSize);
+			}
+
+#if WITH_EDITOR
+			PartitionActor->Rename(nullptr, PartitionActor->GetOuter(), REN_NonTransactional);
+			PartitionActor->SetActorLabel(*PCGRuntimeGenSchedulerConstants::PooledPartitionActorName);
+#endif
+			PartitionActorPool.Push(PartitionActor);
+		}
+		else
+		{
+			if (PCGRuntimeGenSchedulerHelpers::CVarRuntimeGenerationEnableDebugging.GetValueOnAnyThread())
+			{
+				UE_LOG(LogPCG, Warning, TEXT("[RUNTIMEGEN] DESTROY PARTITION ACTOR: '%s'"), *PartitionActor->GetActorNameOrLabel());
+			}
+
+#if WITH_EDITOR
+			PartitionActor->Rename(nullptr, PartitionActor->GetOuter(), REN_NonTransactional);
+#endif
+			World->DestroyActor(PartitionActor);
+		}
+	}
+
+	GeneratedComponents.Remove(GenerationKey);
+}
+
+APCGPartitionActor* FPCGRuntimeGenScheduler::FindOrCreatePartitionActor(uint32 GridSize, const FIntVector& GridCoords)
+{
+	TRACE_CPUPROFILER_EVENT_SCOPE(FPCGRuntimeGenScheduler::FindOrCreatePartitionActor);
+
+	check(ActorAndComponentMapping);
+
+	if (!World)
+	{
+		UE_LOG(LogPCG, Error, TEXT("[FindOrCreatePartitionActor] World is null"));
+		return nullptr;
+	}
+
+	// Attempt to find an existing RuntimeGen PA.
+	if (APCGPartitionActor* ExistingActor = ActorAndComponentMapping->GetPartitionActor(GridSize, GridCoords, /*bRuntimeGenerated=*/true))
+	{
+		return ExistingActor;
+	}
+
+	FActorSpawnParameters SpawnParams;
+#if WITH_EDITOR
+	SpawnParams.Name = *APCGPartitionActor::GetRuntimeGenActorName(GridSize, GridCoords);
+	SpawnParams.NameMode = FActorSpawnParameters::ESpawnActorNameMode::Requested;
+#endif
+	SpawnParams.ObjectFlags |= RF_Transient;
+	SpawnParams.ObjectFlags &= ~RF_Transactional;
+
+	const FVector CellCenter(FVector(GridCoords.X + 0.5, GridCoords.Y + 0.5, GridCoords.Z + 0.5) * GridSize);
+	APCGPartitionActor* NewActor = CastChecked<APCGPartitionActor>(World->SpawnActor(APCGPartitionActor::StaticClass(), &CellCenter, nullptr, SpawnParams));
+	NewActor->SetToRuntimeGenerated();
+	NewActor->SetPCGGridSize(GridSize);
+
+#if WITH_EDITOR
+	NewActor->SetLockLocation(true);
+	NewActor->SetActorLabel(SpawnParams.Name.ToString());
+#endif
+
+	// Empty GUID, RuntimeGen PAs don't need one.
+	NewActor->PostCreation(FGuid());
+
+	return NewActor;
+}
+
+APCGPartitionActor* FPCGRuntimeGenScheduler::GetPartitionActorFromPool(uint32 GridSize, const FIntVector& GridCoords)
+{
+	TRACE_CPUPROFILER_EVENT_SCOPE(FPCGRuntimeGenScheduler::FindOrCreatePartitionActor);
+
+	check(ActorAndComponentMapping);
+
+	if (!World)
+	{
+		UE_LOG(LogPCG, Error, TEXT("[FindOrCreatePartitionActor] World is null"));
+		return nullptr;
+	}
+
+	// Attempt to find an existing RuntimeGen PA.
+	if (APCGPartitionActor* ExistingActor = ActorAndComponentMapping->GetPartitionActor(GridSize, GridCoords, /*bRuntimeGenerated=*/true))
+	{
+		return ExistingActor;
+	}
+
+	// Double size of the pool if it is empty.
+	if (PartitionActorPool.IsEmpty())
+	{
+		// If PartitionActorPoolSize is zero, then we should use the CVarBasePoolSize instead. Result must always at least be >= 1.
+		const uint32 CurrentPoolSize = FMath::Max(1, PartitionActorPoolSize > 0 ? PartitionActorPoolSize : PCGRuntimeGenSchedulerHelpers::CVarRuntimeGenerationBasePoolSize.GetValueOnAnyThread());
+
+		if (PCGRuntimeGenSchedulerHelpers::CVarRuntimeGenerationEnableDebugging.GetValueOnAnyThread())
+		{
+			UE_LOG(LogPCG, Warning, TEXT("[RUNTIMEGEN] INCREASING TRANSIENT PARTITION ACTOR POOL SIZE BY (%d)"), CurrentPoolSize);
+		}
+
+		// If pooling was enabled late, the editor world RuntimeGenScheduler will not have created the initial pool, so we should create it now.
+		AddPartitionActorPoolCount(CurrentPoolSize);
+	}
+
+	check(!PartitionActorPool.IsEmpty());
+	APCGPartitionActor* PartitionActor = PartitionActorPool.Pop();
+	PartitionActor->SetPCGGridSize(GridSize);
+
+#if WITH_EDITOR
+	const FName ActorName = *APCGPartitionActor::GetRuntimeGenActorName(GridSize, GridCoords);
+
+	PartitionActor->Rename(*ActorName.ToString(), PartitionActor->GetOuter(), REN_NonTransactional);
+	PartitionActor->SetActorLabel(ActorName.ToString());
+#endif
+
+	const FVector CellCenter(FVector(GridCoords.X + 0.5, GridCoords.Y + 0.5, GridCoords.Z + 0.5) * GridSize);
+	if (!PartitionActor->Teleport(CellCenter))
+	{
+		UE_LOG(LogPCG, Error, TEXT("[RUNTIMEGEN] Could not set the location of RuntimeGen partition actor '%s'"), *PartitionActor->GetActorNameOrLabel());
+	}
+
+#if WITH_EDITOR
+	PartitionActor->SetLockLocation(true);
+#endif
+
+	// Empty GUID, RuntimeGen PAs don't need one.
+	PartitionActor->PostCreation(FGuid());
+
+	return PartitionActor;
+}
+
+void FPCGRuntimeGenScheduler::AddPartitionActorPoolCount(int32 Count)
+{
+	TRACE_CPUPROFILER_EVENT_SCOPE(FPCGRuntimeGenScheduler::AddPartitionActorPoolCount);
+
+	PartitionActorPoolSize += Count;
+
+	FActorSpawnParameters SpawnParams;
+#if WITH_EDITOR
+	SpawnParams.NameMode = FActorSpawnParameters::ESpawnActorNameMode::Requested;
+	SpawnParams.Name = *PCGRuntimeGenSchedulerConstants::PooledPartitionActorName;
+#endif
+	SpawnParams.ObjectFlags |= RF_Transient;
+	SpawnParams.ObjectFlags &= ~RF_Transactional;
+
+	// Create RuntimeGen PA pool.
+	for (int32 I = 0; I < Count; ++I)
+	{
+		// TODO: do these actors get networked automatically? do we want that or not?
+		APCGPartitionActor* NewActor = World->SpawnActor<APCGPartitionActor>(SpawnParams);
+		check(NewActor);
+		NewActor->SetToRuntimeGenerated();
+		PartitionActorPool.Add(NewActor);
+#if WITH_EDITOR
+		NewActor->SetActorLabel(PCGRuntimeGenSchedulerConstants::PooledPartitionActorName);
+#endif
+	}
+}
+
+void FPCGRuntimeGenScheduler::ResetPartitionActorPoolToSize(uint32 NewPoolSize)
+{
+	for (APCGPartitionActor* PartitionActor : PartitionActorPool)
+	{
+#if WITH_EDITOR
+		PartitionActor->Rename(nullptr, PartitionActor->GetOuter(), REN_NonTransactional);
+#endif
+		World->DestroyActor(PartitionActor);
+	}
+
+	PartitionActorPool.Empty();
+	PartitionActorPoolSize = 0;
+	AddPartitionActorPoolCount(NewPoolSize);
+}
diff --git a/Engine/Plugins/Experimental/PCG/Source/PCG/Private/RuntimeGen/SchedulingPolicies/PCGSchedulingPolicyDistanceAndDirection.cpp b/Engine/Plugins/Experimental/PCG/Source/PCG/Private/RuntimeGen/SchedulingPolicies/PCGSchedulingPolicyDistanceAndDirection.cpp
new file mode 100644
--- /dev/null
+++ b/Engine/Plugins/Experimental/PCG/Source/PCG/Private/RuntimeGen/SchedulingPolicies/PCGSchedulingPolicyDistanceAndDirection.cpp
@@ -0,0 +1,57 @@
+// Copyright Epic Games, Inc. All Rights Reserved.
+
+#include "RuntimeGen/SchedulingPolicies/PCGSchedulingPolicyDistanceAndDirection.h"
+
+#include "PCGCommon.h"
+#include "RuntimeGen/GenSources/PCGGenSourceBase.h"
+
+#include "GameFramework/Actor.h"
+
+#include UE_INLINE_GENERATED_CPP_BY_NAME(PCGSchedulingPolicyDistanceAndDirection)
+
+double UPCGSchedulingPolicyDistanceAndDirection::CalculatePriority(const IPCGGenSourceBase* InGenSource, const FBox& GenerationBounds, bool bUse2DGrid) const
+{
+	check(InGenSource);
+
+	double Priority = 0.0;
+
+	if (bUseDistance || bUseDirection)
+	{
+		TOptional<FVector> GenSourcePositionOptional = InGenSource->GetPosition();
+
+		if (ensure(GenSourcePositionOptional.IsSet()))
+		{
+			FVector GenSourcePosition = GenSourcePositionOptional.GetValue();
+			if (bUse2DGrid)
+			{
+				GenSourcePosition.Z = 0.0;
+			}
+
+			const FVector GenSourceToComponent = GenerationBounds.GetClosestPointTo(GenSourcePosition) - GenSourcePosition;
+
+			if (bUseDistance)
+			{
+				const double MaxRadius = 2.0 * PCGHiGenGrid::GridToGridSize(EPCGHiGenGrid::GridMax);
+				const double Distance = 1.0 - FMath::Clamp(GenSourceToComponent.Length() / MaxRadius, 0.0, 1.0);
+				Priority += Distance * FMath::Clamp(DistanceWeight, 0.0f, 1.0f);
+			}
+
+			if (bUseDirection)
+			{
+				TOptional<FVector> GenSourceDirectionOptional = InGenSource->GetDirection();
+
+				if ((GenSourceDirectionOptional.IsSet()))
+				{
+					const double DirectionDotProd = FVector::DotProduct(GenSourceToComponent.GetSafeNormal(), GenSourceDirectionOptional.GetValue());
+					Priority += DirectionDotProd * FMath::Clamp(DirectionWeight, 0.0f, 1.0f);
+				}
+			}
+		}
+	}
+
+	// Remap range from [-1, 2] to [0, 1].
+	Priority = (Priority + 1.0) / 3.0;
+	ensure(Priority >= 0.0 && Priority <= 1.0);
+
+	return Priority;
+}
diff --git a/Engine/Plugins/Experimental/PCG/Source/PCG/Public/PCGActorAndComponentMapping.h b/Engine/Plugins/Experimental/PCG/Source/PCG/Public/PCGActorAndComponentMapping.h
--- a/Engine/Plugins/Experimental/PCG/Source/PCG/Public/PCGActorAndComponentMapping.h
+++ b/Engine/Plugins/Experimental/PCG/Source/PCG/Public/PCGActorAndComponentMapping.h
@@ -5,6 +5,7 @@
 #include "PCGCommon.h"
 #include "Elements/PCGActorSelector.h"
 #include "Grid/PCGComponentOctree.h"
+#include "RuntimeGen/PCGRuntimeGenScheduler.h"
 
 #include "Containers/Array.h"
 #include "Containers/Map.h"
@@ -36,6 +37,7 @@ class FPCGActorAndComponentMapping
 {
 public:
 	friend UPCGSubsystem;
+	friend FPCGRuntimeGenScheduler;
 
 	~FPCGActorAndComponentMapping() = default;
 	
@@ -79,8 +81,11 @@ public:
 	/** Return a copy of all the registered components. Thread safe */
 	TSet<UPCGComponent*> GetAllRegisteredComponents() const;
 
-	/** Retrieves a local component using grid size and grid coordinates, returns nullptr if no such component found. */
-	UPCGComponent* GetLocalComponent(uint32 GridSize, const FIntVector& CellCoords, const UPCGComponent* InOriginalComponent);
+	/** Retrieves a local component using grid size and grid coordinates, returns nullptr if no such component is found. */
+	UPCGComponent* GetLocalComponent(uint32 GridSize, const FIntVector& CellCoords, const UPCGComponent* InOriginalComponent, bool bRuntimeGenerated = false);
+
+	/** Retrieves a partition actor using grid size and grid coordinates, returns nullptr if no such partition actor is found. */
+	APCGPartitionActor* GetPartitionActor(uint32 GridSize, const FIntVector& CellCoords, bool bRuntimeGenerated = false) const;
 
 private:
 	// This class is only meant to be used as part of the PCG Subsytem and owned by it.
@@ -174,10 +179,18 @@ private:
 	TMap<uint32, TMap<FIntVector, TObjectPtr<APCGPartitionActor>>> PartitionActorsMap;
 	mutable FRWLock PartitionActorsMapLock;
 
+	/** Mapping from grid size and grid coords to RuntimeGen partition actor. We can only have 1 RuntimeGen partition actor per grid cell. */
+	TMap<uint32, TMap<FIntVector, TObjectPtr<APCGPartitionActor>>> RuntimeGenPartitionActorsMap;
+	mutable FRWLock RuntimeGenPartitionActorsMapLock;
+
 	/** Mapping between original components and its overlapping partition actors. */
 	TMap<const UPCGComponent*, TSet<TObjectPtr<APCGPartitionActor>>> ComponentToPartitionActorsMap;
 	mutable FRWLock ComponentToPartitionActorsMapLock;
 
+	/** Mapping between original components and its overlapping RuntimeGen partition actors. */
+	TMap<const UPCGComponent*, TSet<TObjectPtr<APCGPartitionActor>>> ComponentToRuntimeGenPartitionActorsMap;
+	mutable FRWLock ComponentToRuntimeGenPartitionActorsMapLock;
+
 	/** Components to be unregister at the next frame. cf. UnregisterComponent for a better understanding on why it is needed. */
 	TSet<UPCGComponent*> DelayedComponentToUnregister;
 	mutable FCriticalSection DelayedComponentToUnregisterLock;
diff --git a/Engine/Plugins/Experimental/PCG/Source/PCG/Public/PCGCommon.h b/Engine/Plugins/Experimental/PCG/Source/PCG/Public/PCGCommon.h
--- a/Engine/Plugins/Experimental/PCG/Source/PCG/Public/PCGCommon.h
+++ b/Engine/Plugins/Experimental/PCG/Source/PCG/Public/PCGCommon.h
@@ -18,7 +18,8 @@ enum class EPCGChangeType : uint8
 	Input = 1 << 2,
 	Edge = 1 << 3,
 	Node = 1 << 4,
-	Structural = 1 << 5
+	Structural = 1 << 5,
+	Debug = 1 << 6, // The Debug change type propagates when debug is toggled on a node.
 };
 ENUM_CLASS_FLAGS(EPCGChangeType);
 
@@ -139,7 +140,10 @@ enum class EPCGHiGenGrid : uint32
 {
 	Uninitialized = 0 UMETA(Hidden),
 
-	// NOTE: When adding new grids, increment PCGHiGenGrid::NumGridValues below.
+	// NOTE: When adding new grids, increment PCGHiGenGrid::NumGridValues below
+	Grid4 = 4 UMETA(DisplayName = "400"),
+	Grid8 = 8 UMETA(DisplayName = "800"),
+	Grid16 = 16 UMETA(DisplayName = "1600"),
 	Grid32 = 32 UMETA(DisplayName = "3200"),
 	Grid64 = 64 UMETA(DisplayName = "6400"),
 	Grid128 = 128 UMETA(DisplayName = "12800"),
@@ -148,7 +152,7 @@ enum class EPCGHiGenGrid : uint32
 	Grid1024 = 1024 UMETA(DisplayName = "102400"),
 	Grid2048 = 2048 UMETA(DisplayName = "204800"),
 	
-	GridMin = Grid32 UMETA(Hidden),
+	GridMin = Grid4 UMETA(Hidden),
 	GridMax = Grid2048 UMETA(Hidden),
 
 	// Should execute once rather than executing on any grid
@@ -159,7 +163,7 @@ ENUM_CLASS_FLAGS(EPCGHiGenGrid);
 namespace PCGHiGenGrid
 {
 	// Number of unique values of EPCGHiGenGrid, const so it can be used for the inline allocator below.
-	constexpr uint32 NumGridValues = 10;
+	constexpr uint32 NumGridValues = 13;
 
 	// Alias for array which is allocated on the stack (we have a strong idea of the max required elements).
 	using FSizeArray = TArray<uint32, TInlineAllocator<PCGHiGenGrid::NumGridValues>>;
@@ -169,6 +173,7 @@ namespace PCGHiGenGrid
 
 	PCG_API bool IsValidGridSize(uint32 InGridSize);
 	PCG_API bool IsValidGrid(EPCGHiGenGrid InGrid);
+	PCG_API bool IsValidGridOrUninitialized(EPCGHiGenGrid InGrid);
 	PCG_API uint32 GridToGridSize(EPCGHiGenGrid InGrid);
 	PCG_API EPCGHiGenGrid GridSizeToGrid(uint32 InGridSize);
 
@@ -190,4 +195,58 @@ enum class EPCGEditorDirtyMode : uint8
 	Normal UMETA(Tooltip="Normal editing mode where generation changes (generation, cleanup) dirty the component and its resources."),
 	Preview UMETA(Tooltip="Editing mode where generation changes (generation, cleanup, resources) on the component will not trigger any dirty state, but will also not save any of the generated resources."),
 	LoadAsPreview UMETA(Tooltip="Acts as the normal editing mode until the next load of the component, at which state it acts as-if-transient, namely that any further generation changes will not dirty the component.")
-};
\ No newline at end of file
+};
+
+USTRUCT(BlueprintType)
+struct FPCGRuntimeGenerationRadii
+{
+	GENERATED_BODY()
+
+public:
+	/** Get the runtime generation radius for the given grid size. */
+	double GetGenerationRadiusFromGrid(EPCGHiGenGrid Grid) const;
+
+	/** Compute the runtime cleanup radius for the given grid size. */
+	double GetCleanupRadiusFromGrid(EPCGHiGenGrid Grid) const;
+
+	static constexpr double DefaultGenerationRadiusMultiplier = 2.0;
+	static constexpr double DefaultCleanupRadiusMultiplier = 1.1;
+
+	/** The distance (in centimeters) at which the component will be considered for generation by the RuntimeGenerationScheduler. For partitioned components, this also acts as the unbounded generation radius. */
+	UPROPERTY(BlueprintReadWrite, EditAnywhere, Category = "Runtime Generation")
+	double GenerationRadius = PCGHiGenGrid::UnboundedGridSize() * DefaultGenerationRadiusMultiplier;
+
+	UPROPERTY(BlueprintReadWrite, EditAnywhere, Category = "Runtime Generation")
+	double GenerationRadius400 = PCGHiGenGrid::GridToGridSize(EPCGHiGenGrid::Grid4) * DefaultGenerationRadiusMultiplier;
+
+	UPROPERTY(BlueprintReadWrite, EditAnywhere, Category = "Runtime Generation")
+	double GenerationRadius800 = PCGHiGenGrid::GridToGridSize(EPCGHiGenGrid::Grid8) * DefaultGenerationRadiusMultiplier;
+
+	UPROPERTY(BlueprintReadWrite, EditAnywhere, Category = "Runtime Generation")
+	double GenerationRadius1600 = PCGHiGenGrid::GridToGridSize(EPCGHiGenGrid::Grid16) * DefaultGenerationRadiusMultiplier;
+
+	UPROPERTY(BlueprintReadWrite, EditAnywhere, Category = "Runtime Generation")
+	double GenerationRadius3200 = PCGHiGenGrid::GridToGridSize(EPCGHiGenGrid::Grid32) * DefaultGenerationRadiusMultiplier;
+
+	UPROPERTY(BlueprintReadWrite, EditAnywhere, Category = "Runtime Generation")
+	double GenerationRadius6400 = PCGHiGenGrid::GridToGridSize(EPCGHiGenGrid::Grid64) * DefaultGenerationRadiusMultiplier;
+
+	UPROPERTY(BlueprintReadWrite, EditAnywhere, Category = "Runtime Generation")
+	double GenerationRadius12800 = PCGHiGenGrid::GridToGridSize(EPCGHiGenGrid::Grid128) * DefaultGenerationRadiusMultiplier;
+
+	UPROPERTY(BlueprintReadWrite, EditAnywhere, Category = "Runtime Generation")
+	double GenerationRadius25600 = PCGHiGenGrid::GridToGridSize(EPCGHiGenGrid::Grid256) * DefaultGenerationRadiusMultiplier;
+
+	UPROPERTY(BlueprintReadWrite, EditAnywhere, Category = "Runtime Generation")
+	double GenerationRadius51200 = PCGHiGenGrid::GridToGridSize(EPCGHiGenGrid::Grid512) * DefaultGenerationRadiusMultiplier;
+
+	UPROPERTY(BlueprintReadWrite, EditAnywhere, Category = "Runtime Generation")
+	double GenerationRadius102400 = PCGHiGenGrid::GridToGridSize(EPCGHiGenGrid::Grid1024) * DefaultGenerationRadiusMultiplier;
+
+	UPROPERTY(BlueprintReadWrite, EditAnywhere, Category = "Runtime Generation")
+	double GenerationRadius204800 = PCGHiGenGrid::GridToGridSize(EPCGHiGenGrid::Grid2048) * DefaultGenerationRadiusMultiplier;
+
+	/** Multiplier on the GenerationRadius to control the distance at which runtime generated components will be cleaned up. Applied per grid size. */
+	UPROPERTY(BlueprintReadWrite, EditAnywhere, Category = "Runtime Generation", meta = (UIMin = "1.0", ClampMin = "1.0"))
+	double CleanupRadiusMultiplier = DefaultCleanupRadiusMultiplier;
+};
diff --git a/Engine/Plugins/Experimental/PCG/Source/PCG/Public/PCGComponent.h b/Engine/Plugins/Experimental/PCG/Source/PCG/Public/PCGComponent.h
--- a/Engine/Plugins/Experimental/PCG/Source/PCG/Public/PCGComponent.h
+++ b/Engine/Plugins/Experimental/PCG/Source/PCG/Public/PCGComponent.h
@@ -20,16 +20,19 @@ class APCGPartitionActor;
 struct FPCGContext;
 class FPCGActorAndComponentMapping;
 class UPCGComponent;
+class UPCGData;
+class IPCGGenSourceBase;
 class UPCGGraph;
-class UPCGGraphInterface;
 class UPCGGraphInstance;
+class UPCGGraphInterface;
 class UPCGManagedResource;
-class UPCGData;
+class UPCGSchedulingPolicyBase;
 class UPCGSubsystem;
 class ALandscapeProxy;
 class FLandscapeProxyComponentDataChangedParams;
 class UClass;
 
+
 #if WITH_EDITOR
 	DECLARE_MULTICAST_DELEGATE_OneParam(FOnPCGGraphGenerated, UPCGComponent*);
 	DECLARE_MULTICAST_DELEGATE_OneParam(FOnPCGGraphCleaned, UPCGComponent*);
@@ -48,8 +51,9 @@ enum class EPCGComponentInput : uint8
 UENUM(Blueprintable)
 enum class EPCGComponentGenerationTrigger : uint8
 {
-	GenerateOnLoad,
-	GenerateOnDemand
+	GenerateOnLoad    UMETA(ToolTip = "Generates only when the component is loaded into the level."),
+	GenerateOnDemand  UMETA(ToolTip = "Generates only when requested (e.g. via Blueprint)."),
+	GenerateAtRuntime UMETA(ToolTip = "Generates only when scheduled by the Runtime Generation Scheduler.")
 };
 
 UENUM(meta = (Bitflags))
@@ -112,6 +116,8 @@ public:
 	/** If this is a local component returns self, otherwise returns the original component. */
 	UPCGComponent* GetOriginalComponent();
 
+	UPCGSchedulingPolicyBase* GetRuntimeGenSchedulingPolicy() const { return SchedulingPolicy; }
+
 	bool CanPartition() const;
 
 	UPCGGraph* GetGraph() const;
@@ -136,7 +142,11 @@ public:
 	UFUNCTION(BlueprintCallable, Category = PCG)
 	void GenerateLocal(bool bForce);
 
+	/** Requests the component to generate only on the specified grid level (all grid levels if EPCGHiGenGrid::Uninitialized). */
+	void GenerateLocal(EPCGComponentGenerationTrigger RequestedGenerationTrigger, bool bForce, EPCGHiGenGrid Grid = EPCGHiGenGrid::Uninitialized);
+
 	FPCGTaskId GenerateLocalGetTaskId(bool bForce);
+	FPCGTaskId GenerateLocalGetTaskId(EPCGComponentGenerationTrigger RequestedGenerationTrigger, bool bForce, EPCGHiGenGrid Grid = EPCGHiGenGrid::Uninitialized);
 
 	/** Cleans up the generation from a local (vs. remote) standpoint. Will not be replicated. Will be delayed. */
 	UFUNCTION(BlueprintCallable, Category = PCG)
@@ -181,6 +191,15 @@ public:
 	/** Clear any data stored for any pins. */
 	void ClearPerPinGeneratedOutput();
 
+	/** Set the runtime generation scheduling policy type. */
+	void SetSchedulingPolicyClass(TSubclassOf<UPCGSchedulingPolicyBase> InSchedulingPolicyClass);
+
+	/** Get the runtime generation radius for the given grid size. */
+	double GetGenerationRadiusFromGrid(EPCGHiGenGrid Grid) const;
+
+	/** Compute the runtime cleanup radius for the given grid size. */
+	double GetCleanupRadiusFromGrid(EPCGHiGenGrid Grid) const;
+
 	UPROPERTY(BlueprintReadWrite, EditAnywhere, Category = Settings)
 	EPCGComponentInput InputType = EPCGComponentInput::Actor;
 
@@ -200,8 +219,23 @@ public:
 	UPROPERTY(BlueprintReadOnly, EditAnywhere, Category = Properties, AdvancedDisplay, meta = (EditCondition = "!bIsComponentLocal", EditConditionHides))
 	EPCGComponentGenerationTrigger GenerationTrigger = EPCGComponentGenerationTrigger::GenerateOnLoad;
 
+	/** Manual overrides for the graph generation radii and cleanup radius multiplier. */
+	UPROPERTY(BlueprintReadWrite, EditAnywhere, Category = RuntimeGeneration, meta = (EditCondition = "GenerationTrigger == EPCGComponentGenerationTrigger::GenerateAtRuntime", EditConditionHides))
+	bool bOverrideGenerationRadii = false;
+
+	UPROPERTY(BlueprintReadWrite, EditAnywhere, Category = RuntimeGeneration, meta = (EditCondition = "GenerationTrigger == EPCGComponentGenerationTrigger::GenerateAtRuntime && bOverrideGenerationRadii", EditConditionHides))
+	FPCGRuntimeGenerationRadii GenerationRadii;
+
+	/** A Scheduling Policy dictates the order in which instances of this component will be scheduled. */
+	UPROPERTY(BlueprintReadOnly, EditAnywhere, NoClear, Category = RuntimeGeneration, meta = (EditCondition = "GenerationTrigger == EPCGComponentGenerationTrigger::GenerateAtRuntime", EditConditionHides))
+	TSubclassOf<UPCGSchedulingPolicyBase> SchedulingPolicyClass = nullptr;
+
+	/** This is the instanced UPCGSchedulingPolicy object which holds scheduling parameters and calculates priorities. */
+	UPROPERTY(BlueprintReadOnly, VisibleAnywhere, Instanced, Category = RuntimeGeneration, meta = (EditCondition = "GenerationTrigger == EPCGComponentGenerationTrigger::GenerateAtRuntime", EditConditionHides))
+	TObjectPtr<UPCGSchedulingPolicyBase> SchedulingPolicy;
+
 	/** Flag to indicate whether this component has run in the editor. Note that for partitionable actors, this will always be false. */
-	UPROPERTY(BlueprintReadOnly, VisibleAnywhere, AdvancedDisplay, Category = Properties, meta = (NoResetToDefault))
+	UPROPERTY(BlueprintReadOnly, VisibleAnywhere, AdvancedDisplay, Category = Properties, NonTransactional, meta = (NoResetToDefault))
 	bool bGenerated = false;
 
 	UPROPERTY(NonPIEDuplicateTransient)
@@ -308,6 +342,9 @@ public:
 	/** Builds the canonical PCG data from a given actor and its PCG component if any. */
 	static UPCGData* CreateActorPCGData(AActor* Actor, const UPCGComponent* Component, bool bParseActor = true);
 
+protected:
+	void RefreshSchedulingPolicy();
+
 protected:
 	UPROPERTY(BlueprintReadOnly, VisibleAnywhere, Category = PCG, Instanced, meta = (NoResetToDefault))
 	TObjectPtr<UPCGGraphInstance> GraphInstance;
@@ -343,7 +380,7 @@ private:
 	bool ShouldGenerate(bool bForce, EPCGComponentGenerationTrigger RequestedGenerationTrigger) const;
 
 	/* Internal call that allows to delay a Generate/Cleanup call, chain with dependencies and keep track of the task id created. This task id is also returned. */
-	FPCGTaskId GenerateInternal(bool bForce, EPCGComponentGenerationTrigger RequestedGenerationTrigger, const TArray<FPCGTaskId>& Dependencies);
+	FPCGTaskId GenerateInternal(bool bForce, EPCGHiGenGrid Grid, EPCGComponentGenerationTrigger RequestedGenerationTrigger, const TArray<FPCGTaskId>& Dependencies);
 	FPCGTaskId CleanupInternal(bool bRemoveComponents, bool bSave, const TArray<FPCGTaskId>& Dependencies);
 
 	/* Internal call to create tasks to generate the component. If there is nothing to do, an invalid task id will be returned. Should only be used by the subsystem. */
diff --git a/Engine/Plugins/Experimental/PCG/Source/PCG/Public/PCGGraph.h b/Engine/Plugins/Experimental/PCG/Source/PCG/Public/PCGGraph.h
--- a/Engine/Plugins/Experimental/PCG/Source/PCG/Public/PCGGraph.h
+++ b/Engine/Plugins/Experimental/PCG/Source/PCG/Public/PCGGraph.h
@@ -384,6 +384,9 @@ public:
 	// Best used on a brand new PCG Graph.
 	void AddUserParameters(const TArray<FPropertyBagPropertyDesc>& InDescs, const UPCGGraph* InOptionalOriginalGraph = nullptr);
 
+	UPROPERTY(BlueprintReadWrite, EditAnywhere, Category = "Runtime Generation")
+	FPCGRuntimeGenerationRadii GenerationRadii;
+
 protected:
 	virtual void OnGraphParametersChanged(EPCGGraphParameterEvent InChangeType, FName InChangedPropertyName) override;
 
diff --git a/Engine/Plugins/Experimental/PCG/Source/PCG/Public/PCGSubsystem.h b/Engine/Plugins/Experimental/PCG/Source/PCG/Public/PCGSubsystem.h
--- a/Engine/Plugins/Experimental/PCG/Source/PCG/Public/PCGSubsystem.h
+++ b/Engine/Plugins/Experimental/PCG/Source/PCG/Public/PCGSubsystem.h
@@ -11,11 +11,11 @@
 
 #include "PCGSubsystem.generated.h"
 
-
 class APCGPartitionActor;
 class APCGWorldActor;
 class UPCGGraph;
 class UPCGLandscapeCache;
+class FPCGRuntimeGenScheduler;
 
 enum class EPCGComponentDirtyFlag : uint8;
 enum class ETickableTickType : uint8;
@@ -84,10 +84,12 @@ public:
 	void RegisterPCGWorldActor(APCGWorldActor* InActor);
 	void UnregisterPCGWorldActor(APCGWorldActor* InActor);
 
+	void OnOriginalComponentUnregistered(UPCGComponent* InComponent);
+	
 	UPCGLandscapeCache* GetLandscapeCache();
 
 	// Schedule graph (owner -> graph)
-	FPCGTaskId ScheduleComponent(UPCGComponent* PCGComponent, bool bSave, const TArray<FPCGTaskId>& InDependencies);
+	FPCGTaskId ScheduleComponent(UPCGComponent* PCGComponent, EPCGHiGenGrid Grid, bool bSave, const TArray<FPCGTaskId>& InDependencies);
 
 	/** Schedule cleanup(owner->graph). Note that in non-partitioned mode, cleanup is immediate. */
 	FPCGTaskId ScheduleCleanup(UPCGComponent* PCGComponent, bool bRemoveComponents, bool bSave, const TArray<FPCGTaskId>& Dependencies);
@@ -155,11 +157,13 @@ public:
 	void ForAllOverlappingComponentsInHierarchy(UPCGComponent* InComponent, const TFunction<void(UPCGComponent*)>& InFunc) const;
 
 	/** Retrieves a local component using grid size and grid coordinates, returns nullptr if no such component found. */
-	UPCGComponent* GetLocalComponent(uint32 GridSize, const FIntVector& CellCoords, const UPCGComponent* InOriginalComponent);
+	UPCGComponent* GetLocalComponent(uint32 GridSize, const FIntVector& CellCoords, const UPCGComponent* InOriginalComponent, bool bTransient = false);
 
 	/** True if graph cache debugging is enabled. */
 	bool IsGraphCacheDebuggingEnabled() const;
 
+	FPCGGenSourceManager* GetGenSourceManager() const;
+
 #if WITH_EDITOR
 public:
 	/** Schedule refresh on the current or next frame */
@@ -226,6 +230,7 @@ private:
 private:
 	APCGWorldActor* PCGWorldActor = nullptr;
 	FPCGGraphExecutor* GraphExecutor = nullptr;
+	FPCGRuntimeGenScheduler* RuntimeGenScheduler = nullptr;
 	bool bHasTickedOnce = false;
 	FPCGActorAndComponentMapping ActorAndComponentMapping;
 
diff --git a/Engine/Plugins/Experimental/PCG/Source/PCG/Public/PCGWorldActor.h b/Engine/Plugins/Experimental/PCG/Source/PCG/Public/PCGWorldActor.h
--- a/Engine/Plugins/Experimental/PCG/Source/PCG/Public/PCGWorldActor.h
+++ b/Engine/Plugins/Experimental/PCG/Source/PCG/Public/PCGWorldActor.h
@@ -64,6 +64,12 @@ public:
 	UPROPERTY(config, EditAnywhere, Category = GenerationSettings)
 	bool bUse2DGrid = true;
 
+#if WITH_EDITORONLY_DATA
+	/** Allows any currently active editor viewport to act as a Runtime Generation Source. */
+	UPROPERTY(EditAnywhere, Category = RuntimeGeneration)
+	bool bTreatEditorViewportAsGenerationSource = false;
+#endif
+
 private:
 	void RegisterToSubsystem();
 	void UnregisterFromSubsystem();
diff --git a/Engine/Plugins/Experimental/PCG/Source/PCG/Public/RuntimeGen/GenSources/PCGGenSourceBase.h b/Engine/Plugins/Experimental/PCG/Source/PCG/Public/RuntimeGen/GenSources/PCGGenSourceBase.h
new file mode 100644
--- /dev/null
+++ b/Engine/Plugins/Experimental/PCG/Source/PCG/Public/RuntimeGen/GenSources/PCGGenSourceBase.h
@@ -0,0 +1,31 @@
+// Copyright Epic Games, Inc. All Rights Reserved.
+
+#pragma once
+
+#include "UObject/Interface.h"
+
+#include "PCGGenSourceBase.generated.h"
+
+UINTERFACE(BlueprintType)
+class PCG_API UPCGGenSourceBase : public UInterface
+{
+	GENERATED_BODY()
+};
+
+/**
+ * A PCG Generation Source represents an object in the world that provokes nearby 
+ * PCG Components to generate through the Runtime Generation Scheduler.
+ */
+class IPCGGenSourceBase
+{
+	GENERATED_BODY()
+
+public:
+	/** Returns the world space position of this gen source. */
+	virtual TOptional<FVector> GetPosition() const PURE_VIRTUAL(UPCGGenSourceBase::GetPosition, return TOptional<FVector>(););
+
+	/** Returns the normalized forward vector of this gen source. */
+	virtual TOptional<FVector> GetDirection() const PURE_VIRTUAL(UPCGGenSourceBase::GetDirection, return TOptional<FVector>(););
+
+	// TODO: Functions to get the horizontal and vertical FOV would be useful for policies based around frustum/zooming.
+};
diff --git a/Engine/Plugins/Experimental/PCG/Source/PCG/Public/RuntimeGen/GenSources/PCGGenSourceComponent.h b/Engine/Plugins/Experimental/PCG/Source/PCG/Public/RuntimeGen/GenSources/PCGGenSourceComponent.h
new file mode 100644
--- /dev/null
+++ b/Engine/Plugins/Experimental/PCG/Source/PCG/Public/RuntimeGen/GenSources/PCGGenSourceComponent.h
@@ -0,0 +1,46 @@
+// Copyright Epic Games, Inc. All Rights Reserved.
+
+#pragma once
+
+#include "RuntimeGen/GenSources/PCGGenSourceBase.h"
+
+#include "Components/ActorComponent.h"
+
+#include "PCGGenSourceComponent.generated.h"
+
+namespace EEndPlayReason { enum Type : int; }
+
+class FPCGGenSourceManager;
+
+/**
+ * UPCGGenSourceComponent makes the actor this is attached to act as a PCG runtime generation source.
+ */
+UCLASS(BlueprintType, ClassGroup = (Procedural), DisplayName = "PCG Generation Source", meta = (BlueprintSpawnableComponent, PrioritizeCategories = "PCG"))
+class PCG_API UPCGGenSourceComponent : public UActorComponent, public IPCGGenSourceBase
+{
+	GENERATED_BODY()
+
+	UPCGGenSourceComponent(const FObjectInitializer& InObjectInitializer) : Super(InObjectInitializer) {}
+	~UPCGGenSourceComponent();
+
+public:
+	//~Begin UActorComponent Interface
+	virtual void BeginPlay() override;
+	virtual void EndPlay(const EEndPlayReason::Type EndPlayReason) override;
+	virtual void OnRegister() override;
+	virtual void OnUnregister() override;
+	virtual void OnComponentCreated() override;
+	virtual void OnComponentDestroyed(bool bDestroyingHierarchy) override;
+	//~End UActorComponent Interface
+
+	//~Begin IPCGGenSourceInterface
+	/** Returns the world space position of this gen source. */
+	virtual TOptional<FVector> GetPosition() const override;
+
+	/** Returns the normalized forward vector of this gen source. */
+	virtual TOptional<FVector> GetDirection() const override;
+	//~End IPCGGenSourceInterface
+
+protected:
+	FPCGGenSourceManager* GetGenSourceManager() const;
+};
diff --git a/Engine/Plugins/Experimental/PCG/Source/PCG/Public/RuntimeGen/GenSources/PCGGenSourceEditorCamera.h b/Engine/Plugins/Experimental/PCG/Source/PCG/Public/RuntimeGen/GenSources/PCGGenSourceEditorCamera.h
new file mode 100644
--- /dev/null
+++ b/Engine/Plugins/Experimental/PCG/Source/PCG/Public/RuntimeGen/GenSources/PCGGenSourceEditorCamera.h
@@ -0,0 +1,33 @@
+// Copyright Epic Games, Inc. All Rights Reserved.
+
+#pragma once
+
+#include "PCGGenSourceBase.h"
+
+#include "PCGGenSourceEditorCamera.generated.h"
+
+#if WITH_EDITOR
+class FEditorViewportClient;
+#endif
+
+/**
+ * This GenerationSource captures active Editor Viewports per tick to provoke RuntimeGeneration. Editor Viewports
+ * are not captured by default, but can be enabled on the PCGWorldActor via bTreatEditorViewportAsGenerationSource.
+ */
+UCLASS(BlueprintType, ClassGroup = (Procedural))
+class PCG_API UPCGGenSourceEditorCamera : public UObject, public IPCGGenSourceBase
+{
+	GENERATED_BODY()
+
+public:
+	/** Returns the world space position of this gen source. */
+	virtual TOptional<FVector> GetPosition() const override;
+
+	/** Returns the normalized forward vector of this gen source. */
+	virtual TOptional<FVector> GetDirection() const override;
+
+public:
+#if WITH_EDITORONLY_DATA
+	FEditorViewportClient* EditorViewportClient = nullptr;
+#endif
+};
diff --git a/Engine/Plugins/Experimental/PCG/Source/PCG/Public/RuntimeGen/GenSources/PCGGenSourcePlayer.h b/Engine/Plugins/Experimental/PCG/Source/PCG/Public/RuntimeGen/GenSources/PCGGenSourcePlayer.h
new file mode 100644
--- /dev/null
+++ b/Engine/Plugins/Experimental/PCG/Source/PCG/Public/RuntimeGen/GenSources/PCGGenSourcePlayer.h
@@ -0,0 +1,33 @@
+// Copyright Epic Games, Inc. All Rights Reserved.
+
+#pragma once
+
+#include "PCGGenSourceBase.h"
+
+#include "PCGGenSourcePlayer.generated.h"
+
+class APlayerController;
+
+/**
+ * A UPCGGenSourcePlayer is automatically captured for all PlayerControllers in the level on Login/Logout.
+ */
+UCLASS(BlueprintType, ClassGroup = (Procedural))
+class PCG_API UPCGGenSourcePlayer : public UObject, public IPCGGenSourceBase
+{
+	GENERATED_BODY()
+
+public:
+	/** Returns the world space position of this gen source. */
+	virtual TOptional<FVector> GetPosition() const override;
+
+	/** Returns the normalized forward vector of this gen source. */
+	virtual TOptional<FVector> GetDirection() const override;
+
+	TWeakObjectPtr<const APlayerController> GetPlayerController() const { return PlayerController; }
+	void SetPlayerController(const APlayerController* InPlayerController);
+
+	bool IsValid() const { return PlayerController.IsValid(); }
+
+protected:
+	TWeakObjectPtr<const APlayerController> PlayerController;
+};
diff --git a/Engine/Plugins/Experimental/PCG/Source/PCG/Public/RuntimeGen/GenSources/PCGGenSourceWPStreamingSource.h b/Engine/Plugins/Experimental/PCG/Source/PCG/Public/RuntimeGen/GenSources/PCGGenSourceWPStreamingSource.h
new file mode 100644
--- /dev/null
+++ b/Engine/Plugins/Experimental/PCG/Source/PCG/Public/RuntimeGen/GenSources/PCGGenSourceWPStreamingSource.h
@@ -0,0 +1,30 @@
+// Copyright Epic Games, Inc. All Rights Reserved.
+
+#pragma once
+
+#include "PCGGenSourceBase.h"
+
+#include "PCGGenSourceWPStreamingSource.generated.h"
+
+struct FWorldPartitionStreamingSource;
+
+/**
+ * This GenerationSource captures WorldPartitionStreamingSources for RuntimeGeneration.
+ * 
+ * UPCGGenSourceWPStreamingSources are created per tick and live only until the end of the tick.
+ */
+UCLASS(BlueprintType, ClassGroup = (Procedural))
+class PCG_API UPCGGenSourceWPStreamingSource : public UObject, public IPCGGenSourceBase
+{
+	GENERATED_BODY()
+
+public:
+	/** Returns the world space position of this gen source. */
+	virtual TOptional<FVector> GetPosition() const override;
+
+	/** Returns the normalized forward vector of this gen source. */
+	virtual TOptional<FVector> GetDirection() const override;
+
+public:
+	const FWorldPartitionStreamingSource* StreamingSource;
+};
diff --git a/Engine/Plugins/Experimental/PCG/Source/PCG/Public/RuntimeGen/PCGGenSourceManager.h b/Engine/Plugins/Experimental/PCG/Source/PCG/Public/RuntimeGen/PCGGenSourceManager.h
new file mode 100644
--- /dev/null
+++ b/Engine/Plugins/Experimental/PCG/Source/PCG/Public/RuntimeGen/PCGGenSourceManager.h
@@ -0,0 +1,41 @@
+// Copyright Epic Games, Inc. All Rights Reserved.
+
+#pragma once
+
+#include "RuntimeGen/GenSources/PCGGenSourceBase.h"
+
+class APCGWorldActor;
+class UPCGGenSourceComponent;
+
+class AController;
+class AGameModeBase;
+class APlayerController;
+class UWorld;
+
+/**
+ * The runtime Generation Source Manager tracks generation sources in the world for use by the Runtime Generation Scheduler.
+ */
+class FPCGGenSourceManager
+{
+public:
+	FPCGGenSourceManager(const UWorld* InWorld);
+	~FPCGGenSourceManager();
+
+	/** Creates the set of all GenerationSources in the level based on the contents of the GenSourceManager's world. */
+	TSet<IPCGGenSourceBase*> GetGenSources(const APCGWorldActor* InPCGWorldActor) const;
+
+	/** Adds a UPCGGenSource to be tracked by the GenSourceManager. */
+	bool RegisterGenSource(IPCGGenSourceBase* InGenSource);
+
+	/** Removes a UPCGGenSource from being tracked by the GenSourceManager. */
+	bool UnregisterGenSource(const IPCGGenSourceBase* InGenSource);
+
+protected:
+	void OnGameModePostLogin(AGameModeBase* InGameMode, APlayerController* InPlayerController);
+	void OnGameModePostLogout(AGameModeBase* InGameMode, AController* InController);
+
+protected:
+	TSet<IPCGGenSourceBase*> GenSources;
+
+	const UWorld* World = nullptr;
+};
diff --git a/Engine/Plugins/Experimental/PCG/Source/PCG/Public/RuntimeGen/PCGRuntimeGenScheduler.h b/Engine/Plugins/Experimental/PCG/Source/PCG/Public/RuntimeGen/PCGRuntimeGenScheduler.h
new file mode 100644
--- /dev/null
+++ b/Engine/Plugins/Experimental/PCG/Source/PCG/Public/RuntimeGen/PCGRuntimeGenScheduler.h
@@ -0,0 +1,113 @@
+// Copyright Epic Games, Inc. All Rights Reserved.
+
+#pragma once
+
+#include "PCGCommon.h"
+#include "UObject/WeakObjectPtr.h"
+
+class FPCGActorAndComponentMapping;
+class UPCGComponent;
+class IPCGGenSourceBase;
+class APCGPartitionActor;
+class APCGWorldActor;
+
+/**
+ * The Runtime Generation Scheduler system handles the scheduling of PCG Components marked as GenerateAtRuntime.
+ * It searches the level for Partitioned and Non-Partitioned components in range of the currently active
+ * UPCGGenSources in the level, and schedules them efficiently based on their UPCGSchedulingPolicy, creating 
+ * APCGPartitionActors as necessary to support hierarchical generation.
+ *
+ * APCGPartitionActors can be created/destroyed on demand or provided by a dynamically growing pool of actors. If
+ * enabled, the pool will double in capacity anytime the number of available PAs reaches zero.
+ * 
+ * Components and PartitionActors created by the Runtime Generation Scheduler should be managed exclusively by the
+ * runtime gen scheduling system.
+ */
+class FPCGRuntimeGenScheduler
+{
+	friend class UPCGSubsystem;
+
+public:
+	FPCGRuntimeGenScheduler(UWorld* InWorld, FPCGActorAndComponentMapping* InActorAndComponentMapping);
+	~FPCGRuntimeGenScheduler();
+
+	FPCGRuntimeGenScheduler(const FPCGRuntimeGenScheduler&) = delete;
+	FPCGRuntimeGenScheduler(FPCGRuntimeGenScheduler&& other) = delete;
+	FPCGRuntimeGenScheduler& operator=(const FPCGRuntimeGenScheduler& other) = delete;
+	FPCGRuntimeGenScheduler& operator=(FPCGRuntimeGenScheduler&& other) = delete; 
+
+	void Tick(const APCGWorldActor* InPCGWorldActor);
+
+	void OnOriginalComponentUnregistered(UPCGComponent* InComponent);
+
+protected:
+	struct FGridGenerationKey : TTuple<uint32, FIntVector, UPCGComponent*>
+	{
+		FGridGenerationKey(uint32 InGridSize, const FIntVector& InGridCoords, UPCGComponent* InOriginalComponent)
+			: TTuple<uint32, FIntVector, UPCGComponent*>(InGridSize, InGridCoords, InOriginalComponent) {}
+
+		uint32 GetGridSize() const { return Get<0>(); }
+		FIntVector GetGridCoords() const { return Get<1>(); }
+		UPCGComponent* GetOriginalComponent() const { return Get<2>(); }
+	};
+
+	/** Returns true if the scheduler should tick this frame. */
+	bool ShouldTick();
+
+	/** Queue nearby components for generation. */
+	void TickQueueComponentsForGeneration(
+		const TSet<IPCGGenSourceBase*>& GenSources,
+		const APCGWorldActor* InPCGWorldActor,
+		TMap<FGridGenerationKey, double>& OutComponentsToGenerate);
+
+	/** Perform immediate cleanup on components that become out of range. */
+	void TickCleanup(const TSet<IPCGGenSourceBase*>& GenSources, const APCGWorldActor* InPCGWorldActor);
+
+	/** Schedule generation on components in priority order. */
+	void TickScheduleGeneration(TMap<FGridGenerationKey, double>& ComponentsToGenerate);
+
+	/** Detects changes in RuntimeGen CVars to keep the PA pool in a valid state. */
+	void TickCVars(const APCGWorldActor* InPCGWorldActor);
+
+	/** Cleanup all local components in the GeneratedComponents set. */
+	void CleanupLocalComponents(const APCGWorldActor* InPCGWorldActor);
+
+	/** Cleanup a component and remove it from the GeneratedComponents set. */
+	void CleanupComponent(const FGridGenerationKey& GenerationKey, UPCGComponent* GeneratedComponent);
+
+	/** Creates an empty RuntimeGen PA on demand if one cannot already be found in the level. PA will be created for the given GridSize and GridCoords. */
+	APCGPartitionActor* FindOrCreatePartitionActor(uint32 GridSize, const FIntVector& GridCoords);
+	
+	/** Grabs an empty RuntimeGen PA from the PartitionActorPool and initializes it at the given GridSize and GridCoords. If no PAs are available in the pool,
+	* the pool capacity will double and new PAs will be created.
+	*/
+	APCGPartitionActor* GetPartitionActorFromPool(uint32 GridSize, const FIntVector& GridCoords);
+
+	/** Adds Count new RuntimeGen PAs to the Runtime PA pool. */
+	void AddPartitionActorPoolCount(int32 Count);
+
+	/** Destroy all pooled partition actors and rebuild with the NewPoolSize. */
+	void ResetPartitionActorPoolToSize(uint32 NewPoolSize);
+
+private:
+	/** Tracks the generated components managed by the RuntimeGenScheduler. For local components, this generation key will hold the original component.
+	* For non-partitioned components, the generation key should have unbounded grid size and (0, 0, 0) grid coordinates.
+	*/
+	TSet<FGridGenerationKey> GeneratedComponents;
+
+	/** Pool of RuntimeGen PartitionActors used for hierarchical generation. */
+	TArray<APCGPartitionActor*> PartitionActorPool;
+
+	/** PartitionActorPoolSize represents the current maximum capacity of the PartitionActorPool. */
+	int32 PartitionActorPoolSize = 0;
+
+	/** Used to track the number of frames until we can schedule another generation. */
+	uint32 ScheduleFrameCounter = 0;
+
+	FPCGActorAndComponentMapping* ActorAndComponentMapping = nullptr;
+	class FPCGGenSourceManager* GenSourceManager = nullptr;
+	class UWorld* World = nullptr;
+
+	bool bPoolingWasEnabledLastFrame = true;
+	uint32 BasePoolSizeLastFrame = 0;
+};
diff --git a/Engine/Plugins/Experimental/PCG/Source/PCG/Public/RuntimeGen/SchedulingPolicies/PCGSchedulingPolicyBase.h b/Engine/Plugins/Experimental/PCG/Source/PCG/Public/RuntimeGen/SchedulingPolicies/PCGSchedulingPolicyBase.h
new file mode 100644
--- /dev/null
+++ b/Engine/Plugins/Experimental/PCG/Source/PCG/Public/RuntimeGen/SchedulingPolicies/PCGSchedulingPolicyBase.h
@@ -0,0 +1,37 @@
+// Copyright Epic Games, Inc. All Rights Reserved.
+
+#pragma once
+
+#include "PCGSchedulingPolicyBase.generated.h"
+
+class IPCGGenSourceBase;
+
+/** 
+ * Scheduling Policies provide custom logic to efficiently schedule work for the Runtime Generation Scheduling system.
+ * A higher priority value means the work will be scheduled sooner, and larger grid sizes will always have a higher
+ * priority than lower grid sizes.
+ *
+ * If multiple Generation Sources overlap a component, the highest priority value will be used for scheduling.
+ */
+UCLASS(Abstract, BlueprintType, ClassGroup = (Procedural))
+class PCG_API UPCGSchedulingPolicyBase : public UObject
+{
+	GENERATED_BODY()
+
+public:
+	/** Calculate the runtime scheduling priority with respect to a Generation Source. Should return a value in the range [0, 1], where higher values will be scheduled sooner. */
+	virtual double CalculatePriority(const IPCGGenSourceBase* InGenSource, const FBox& GenerationBounds, bool bUse2DGrid) const 
+		PURE_VIRTUAL(UPCGSchedulingPolicyBase::CalculatePriority, return 0.0;);
+
+#if WITH_EDITOR
+	/** Sets whether or not properties should be displayed in the editor.Used to hide instanced SchedulingPolicy properties when runtime generation is not enabled. */
+	void SetShouldDisplayProperties(bool bInShouldDisplayProperties) { bShouldDisplayProperties = bInShouldDisplayProperties; }
+#endif // WITH_EDITOR
+
+#if WITH_EDITORONLY_DATA
+private:
+	/** Hidden property to control display of SchedulingPolicy properties. */
+	UPROPERTY(Transient, meta = (EditCondition = false, EditConditionHides))
+	bool bShouldDisplayProperties = true;
+#endif // WITH_EDITORONLY_DATA 
+};
diff --git a/Engine/Plugins/Experimental/PCG/Source/PCG/Public/RuntimeGen/SchedulingPolicies/PCGSchedulingPolicyDistanceAndDirection.h b/Engine/Plugins/Experimental/PCG/Source/PCG/Public/RuntimeGen/SchedulingPolicies/PCGSchedulingPolicyDistanceAndDirection.h
new file mode 100644
--- /dev/null
+++ b/Engine/Plugins/Experimental/PCG/Source/PCG/Public/RuntimeGen/SchedulingPolicies/PCGSchedulingPolicyDistanceAndDirection.h
@@ -0,0 +1,41 @@
+// Copyright Epic Games, Inc. All Rights Reserved.
+
+#pragma once
+
+#include "PCGSchedulingPolicyBase.h"
+
+#include "PCGSchedulingPolicyDistanceAndDirection.generated.h"
+
+class IPCGGenSourceBase;
+
+/**
+ * SchedulingPolicyDistanceAndDirection uses distance from the generating volume 
+ * and alignment with view direction to choose the most important volumes to generate.
+ *
+ * Distance and Direction are calculated with respect to the Generation Source.
+ */
+UCLASS(BlueprintType, ClassGroup = (Procedural))
+class PCG_API UPCGSchedulingPolicyDistanceAndDirection : public UPCGSchedulingPolicyBase
+{
+	GENERATED_BODY()
+
+public:
+	virtual double CalculatePriority(const IPCGGenSourceBase* InGenSource, const FBox& GenerationBounds, bool bUse2DGrid) const override;
+
+public:
+	/** Toggle whether or not distance is used to calculate the scheduling priority. */
+	UPROPERTY(BlueprintReadWrite, EditAnywhere, Category = "RuntimeGeneration|Scheduling Policy Parameters", meta = (EditCondition = "bShouldDisplayProperties", EditConditionHides))
+	bool bUseDistance = true;
+
+	/** Scalar value used to increase/decrease the impact of distance in the scheduling priority. */
+	UPROPERTY(BlueprintReadWrite, EditAnywhere, Category = "RuntimeGeneration|Scheduling Policy Parameters", meta = (ClampMin = 0.0f, ClampMax = 1.0f, EditCondition = "bShouldDisplayProperties && bUseDistance", EditConditionHides))
+	float DistanceWeight = 1.0f;
+
+	/** Toggle whether or not direction is used to calculate the scheduling priority. */
+	UPROPERTY(BlueprintReadWrite, EditAnywhere, Category = "RuntimeGeneration|Scheduling Policy Parameters", meta = (EditCondition = "bShouldDisplayProperties", EditConditionHides))
+	bool bUseDirection = true;
+
+	/** Scalar value used to increase/decrease the impact of direction in the scheduling priority. */
+	UPROPERTY(BlueprintReadWrite, EditAnywhere, Category = "RuntimeGeneration|Scheduling Policy Parameters", meta = (ClampMin = 0.0f, ClampMax = 1.0f, EditCondition = "bShouldDisplayProperties && bUseDirection", EditConditionHides))
+	float DirectionWeight = 0.0025f;
+};
diff --git a/Engine/Plugins/Experimental/PCG/Source/PCGEditor/Private/Details/PCGComponentDetails.cpp b/Engine/Plugins/Experimental/PCG/Source/PCGEditor/Private/Details/PCGComponentDetails.cpp
--- a/Engine/Plugins/Experimental/PCG/Source/PCGEditor/Private/Details/PCGComponentDetails.cpp
+++ b/Engine/Plugins/Experimental/PCG/Source/PCGEditor/Private/Details/PCGComponentDetails.cpp
@@ -95,9 +95,10 @@ void FPCGComponentDetails::CustomizeDetails(IDetailLayoutBuilder& DetailBuilder)
 			.Padding(4.0f)
 			[
 				SNew(SButton)
-				.OnClicked(this, &FPCGComponentDetails::OnCleanupClicked)				
+				.OnClicked(this, &FPCGComponentDetails::OnCleanupClicked)
+				.Visibility(this, &FPCGComponentDetails::CleanupButtonVisible)
 				[
-					SNew(STextBlock)					
+					SNew(STextBlock)
 					.Font(IDetailLayoutBuilder::GetDetailFont())
 					.Text(LOCTEXT("CleanupButton", "Cleanup"))
 				]
@@ -144,7 +145,8 @@ EVisibility FPCGComponentDetails::GenerateButtonVisible() const
 {
 	for (const TWeakObjectPtr<UPCGComponent>& Component : SelectedComponents)
 	{
-		if (Component.IsValid() && !Component->IsGenerating())
+		// If component is runtime generated then generate/cleanup is managed by the scheduler.
+		if (Component.IsValid() && !Component->IsGenerating() && Component->GenerationTrigger != EPCGComponentGenerationTrigger::GenerateAtRuntime)
 		{
 			return EVisibility::Visible;
 		}
@@ -157,7 +159,22 @@ EVisibility FPCGComponentDetails::CancelButtonVisible() const
 {
 	for (const TWeakObjectPtr<UPCGComponent>& Component : SelectedComponents)
 	{
-		if (Component.IsValid() && Component->IsGenerating())
+		// If component is runtime generated then generate/cleanup is managed by the scheduler.
+		if (Component.IsValid() && Component->IsGenerating() && Component->GenerationTrigger != EPCGComponentGenerationTrigger::GenerateAtRuntime)
+		{
+			return EVisibility::Visible;
+		}
+	}
+
+	return EVisibility::Collapsed;
+}
+
+EVisibility FPCGComponentDetails::CleanupButtonVisible() const
+{
+	for (const TWeakObjectPtr<UPCGComponent>& Component : SelectedComponents)
+	{
+		// If component is runtime generated then generate/cleanup is managed by the scheduler.
+		if (Component.IsValid() && Component->GenerationTrigger != EPCGComponentGenerationTrigger::GenerateAtRuntime)
 		{
 			return EVisibility::Visible;
 		}
diff --git a/Engine/Plugins/Experimental/PCG/Source/PCGEditor/Private/PCGEditor.cpp b/Engine/Plugins/Experimental/PCG/Source/PCGEditor/Private/PCGEditor.cpp
--- a/Engine/Plugins/Experimental/PCG/Source/PCGEditor/Private/PCGEditor.cpp
+++ b/Engine/Plugins/Experimental/PCG/Source/PCGEditor/Private/PCGEditor.cpp
@@ -273,12 +273,19 @@ void FPCGEditor::UpdateDebugAfterComponentSelection(UPCGComponent* InOldComponen
 		return;
 	}
 
+	// GenerateOnDemand requests will bounce off if the component is GenerateAtRuntime, so we need to select the correct GenTrigger for the component to be updated.
+	EPCGComponentGenerationTrigger GenerationTrigger = EPCGComponentGenerationTrigger::GenerateOnDemand;
+	if (InNewComponent && InNewComponent->GenerationTrigger == EPCGComponentGenerationTrigger::GenerateAtRuntime)
+	{
+		GenerationTrigger = EPCGComponentGenerationTrigger::GenerateAtRuntime;
+	}
+
 	// If individual component debugging is disabled, just generate the new component if required.
 	if (!PCGGraphBeingEdited->DebugFlagAppliesToIndividualComponents())
 	{
 		if (InNewComponent && bInNewComponentStartedInspecting)
 		{
-			InNewComponent->GenerateLocal(/*bForce=*/true);
+			InNewComponent->GenerateLocal(GenerationTrigger, /*bForce=*/true);
 		}
 
 		return;
@@ -291,7 +298,7 @@ void FPCGEditor::UpdateDebugAfterComponentSelection(UPCGComponent* InOldComponen
 		{
 			// Transition from 'null' to 'any component not already inspecting' - generate to create debug/inspection info.
 			// If we have null selected, all components are displaying debug. Go to Original component so that all refresh.
-			InNewComponent->GetOriginalComponent()->GenerateLocal(/*bForce=*/true);
+			InNewComponent->GetOriginalComponent()->GenerateLocal(GenerationTrigger, /*bForce=*/true);
 		}
 	}
 	else
@@ -304,15 +311,19 @@ void FPCGEditor::UpdateDebugAfterComponentSelection(UPCGComponent* InOldComponen
 		// Regenerate to clear debug info if switching components, or if changing from a component to null.
 		if (InNewComponent || bDebugFlagSetOnAnyNode)
 		{
-			// Use original component - debug can be displayed both by the local component and parent local components.
-			InOldComponent->GetOriginalComponent()->GenerateLocal(/*bForce=*/true);
+			// If InNewComponent is nullptr in local instance debug mode, then RuntimeGen does not want to regenerate (this is to avoid regenerating all PAs).
+			if (InNewComponent || InOldComponent->GenerationTrigger != EPCGComponentGenerationTrigger::GenerateAtRuntime)
+			{
+				// Use original component - debug can be displayed both by the local component and parent local components.
+				InOldComponent->GetOriginalComponent()->GenerateLocal(GenerationTrigger, /*bForce=*/true);
+			}
 		}
 
 		// Debug new component if it wasn't already
 		if (InNewComponent && bInNewComponentStartedInspecting)
 		{
 			// Use original component - debug can be displayed both by the local component and parent local components.
-			InNewComponent->GetOriginalComponent()->GenerateLocal(/*bForce=*/true);
+			InNewComponent->GetOriginalComponent()->GenerateLocal(GenerationTrigger, /*bForce=*/true);
 		}
 	}
 }
@@ -1327,7 +1338,7 @@ void FPCGEditor::OnToggleDebug()
 			if (PCGSettingsInterface->bDebug != bNewCheckState)
 			{
 				PCGSettingsInterface->bDebug = bNewCheckState;
-				PCGNode->OnNodeChangedDelegate.Broadcast(PCGNode, EPCGChangeType::Settings);
+				PCGNode->OnNodeChangedDelegate.Broadcast(PCGNode, EPCGChangeType::Debug);
 			}
 		}
 	}
diff --git a/Engine/Plugins/Experimental/PCG/Source/PCGEditor/Public/Details/PCGComponentDetails.h b/Engine/Plugins/Experimental/PCG/Source/PCGEditor/Public/Details/PCGComponentDetails.h
--- a/Engine/Plugins/Experimental/PCG/Source/PCGEditor/Public/Details/PCGComponentDetails.h
+++ b/Engine/Plugins/Experimental/PCG/Source/PCGEditor/Public/Details/PCGComponentDetails.h
@@ -27,6 +27,7 @@ protected:
 private:
 	EVisibility GenerateButtonVisible() const;
 	EVisibility CancelButtonVisible() const;
+	EVisibility CleanupButtonVisible() const;
 	FReply OnGenerateClicked();
 	FReply OnCancelClicked();
 	FReply OnCleanupClicked();
