[Backout] - CL28597192
Original CL Desc
-----------------------------------------------------------------
[PCG] Runtime Generation Scheduling

Added FPCGRuntimeGenScheduler. Gets ticked by the Subsystem to efficiently schedule components for generation/cleanup every frame. To be considered, a component must set its generation trigger to GenerateAtRuntime. Note that external Generate/Cleanup are disabled on these components, since they are now managed exclusively by the scheduler.

1. Queue nearby components for generation.
2. Schedule cleanup on components that become out of range
3. Schedule generation on the components we queued in step 1

"Nearby" means that a component is within some distance of a UPCGGenerationSource. Generation sources represent objects in the world that should provoke runtime generation. Examples are PlayerControllers, StreamingSources, GenerationSource Components, and EditorViewport cameras.

Components are scheduled in an efficient order because the generation queue is sorted by "Priority". Priority is calculated by a UPCGSchedulingPolicy, which lives on the Component. Scheduling Policies use a Generation Source to compute their priority. The default Scheduling Policy uses distance and view direction. If multiple Generation Sources are near a component, the scheduler will choose the max priority.

The Scheduler keeps track of GeneratedKeys. GeneratedKeys represent a grid size, grid coordinate, and Original Component. For partitioned and non-partitioned components, the grid size should always be EPCGHiGenGrid::Unbounded, and the grid coordinates should be (0, 0, 0). For Local Components, the grid size and coordinates should reflect the local component's actual grid cell.

Note that the GeneratedKey for a Local Component actually holds the *Original Component*, and not the Local Component itself. This is because at scheduling time, the local component does not exist yet!

When the Scheduler decides to schedule a grid cell/Local Component for generation/cleanup, it first needs to make sure a PartitionActor and LocalComponent actually exist. There are two ways this can happen.

1. The PA is created or destroyed, respectively
2. The PA is popped or pushed into a pool, respectively

In both cases, the LocalComponent is created/destroyed on demand and applied to the PartitionActor. Pooling is enabled by default, but can be disabled with a CVar.

diff --git a/Engine/Plugins/Experimental/PCG/Source/PCG/Private/Elements/PCGDataFromActor.cpp b/Engine/Plugins/Experimental/PCG/Source/PCG/Private/Elements/PCGDataFromActor.cpp
--- a/Engine/Plugins/Experimental/PCG/Source/PCG/Private/Elements/PCGDataFromActor.cpp
+++ b/Engine/Plugins/Experimental/PCG/Source/PCG/Private/Elements/PCGDataFromActor.cpp
@@ -155,7 +155,7 @@ bool FPCGDataFromActorElement::ExecuteInternal(FPCGContext* InContext) const
 		if (Self && Settings->ActorSelector.bMustOverlapSelf)
 		{
 			// Capture ActorBounds by value because it goes out of scope
-			FBox ActorBounds = PCGHelpers::GetGridBounds(Self, PCGComponent);
+			const FBox ActorBounds = PCGHelpers::GetActorBounds(Self);
 			BoundsCheck = [Settings, ActorBounds, PCGComponent](const AActor* OtherActor) -> bool
 			{
 				const FBox OtherActorBounds = OtherActor ? PCGHelpers::GetGridBounds(OtherActor, PCGComponent) : FBox(EForceInit::ForceInit);
diff --git a/Engine/Plugins/Experimental/PCG/Source/PCG/Private/Graph/PCGGraphExecutor.cpp b/Engine/Plugins/Experimental/PCG/Source/PCG/Private/Graph/PCGGraphExecutor.cpp
--- a/Engine/Plugins/Experimental/PCG/Source/PCG/Private/Graph/PCGGraphExecutor.cpp
+++ b/Engine/Plugins/Experimental/PCG/Source/PCG/Private/Graph/PCGGraphExecutor.cpp
@@ -1548,10 +1548,7 @@ namespace PCGGraphExecutor
 				{
 					// Get grid coords using the parent grid (FromGridSize).
 					const FIntVector CellCoords = UPCGActorHelpers::GetCellCoord(ComponentActor->GetActorLocation(), FromGridSize, PCGWorldActor->bUse2DGrid);
-
-					// Search for a transient local component if the source component is runtime managed.
-					const bool bTransientComponent = InContext->SourceComponent->GenerationTrigger == EPCGComponentGenerationTrigger::GenerateAtRuntime;
-					ComponentWithData = Subsystem->GetLocalComponent(FromGridSize, CellCoords, InContext->SourceComponent->GetOriginalComponent(), bTransientComponent);
+					ComponentWithData = Subsystem->GetLocalComponent(FromGridSize, CellCoords, InContext->SourceComponent->GetOriginalComponent());
 				}
 			}
 
@@ -1610,11 +1607,8 @@ namespace PCGGraphExecutor
 			{
 				PCGGraphExecutionLogging::LogGridLinkageTaskExecuteRetrieveScheduleGraph(InContext, ComponentWithData, InResourceKey);
 
-				EPCGComponentGenerationTrigger GenTrigger = (InContext->SourceComponent->GenerationTrigger == EPCGComponentGenerationTrigger::GenerateAtRuntime) ?
-					EPCGComponentGenerationTrigger::GenerateAtRuntime : EPCGComponentGenerationTrigger::GenerateOnDemand;
-
 				// Wake up this task after graph has generated.
-				const FPCGTaskId GraphTaskId = ComponentWithData->GenerateLocalGetTaskId(GenTrigger, /*bForce=*/true);
+				const FPCGTaskId GraphTaskId = ComponentWithData->GenerateLocalGetTaskId(/*bForce=*/true);
 				Subsystem->ScheduleGeneric(WakeUpLambda, InContext->SourceComponent.Get(), { GraphTaskId });
 
 				// Update state and go to sleep.
diff --git a/Engine/Plugins/Experimental/PCG/Source/PCG/Private/Grid/PCGPartitionActor.cpp b/Engine/Plugins/Experimental/PCG/Source/PCG/Private/Grid/PCGPartitionActor.cpp
--- a/Engine/Plugins/Experimental/PCG/Source/PCG/Private/Grid/PCGPartitionActor.cpp
+++ b/Engine/Plugins/Experimental/PCG/Source/PCG/Private/Grid/PCGPartitionActor.cpp
@@ -1,7 +1,7 @@
 // Copyright Epic Games, Inc. All Rights Reserved.
 
 #include "PCGPartitionActor.h"
-
+#include "Engine/World.h"
 #include "PCGComponent.h"
 #include "PCGModule.h"
 #include "PCGSubsystem.h"
@@ -10,7 +10,6 @@
 #include "Helpers/PCGHelpers.h"
 
 #include "Components/BoxComponent.h"
-#include "Engine/World.h"
 
 #include UE_INLINE_GENERATED_CPP_BY_NAME(PCGPartitionActor)
 
@@ -33,6 +32,11 @@ APCGPartitionActor::APCGPartitionActor(const FObjectInitializer& ObjectInitializ
 #endif // WITH_EDITOR
 }
 
+UPCGSubsystem* APCGPartitionActor::GetSubsystem() const 
+{
+	return UPCGSubsystem::GetInstance(GetWorld());
+}
+
 void APCGPartitionActor::PostLoad()
 {
 	Super::PostLoad();
@@ -65,6 +69,7 @@ void APCGPartitionActor::PostLoad()
 	// Make sure that we don't track objects that do not exist anymore
 	CleanupDeadGraphInstances(/*bRemoveNullOnly=*/true);
 
+#if WITH_EDITOR
 	// Mark all our local components as local
 	for (UPCGComponent* LocalComponent : GetAllLocalPCGComponents())
 	{
@@ -73,14 +78,17 @@ void APCGPartitionActor::PostLoad()
 	}
 
 	bWasPostCreatedLoaded = true;
+#endif // WITH_EDITOR
 }
 
 void APCGPartitionActor::BeginDestroy()
 {
-	if (!PCGHelpers::IsRuntimeOrPIE())
+#if WITH_EDITOR
+	if (!PCGHelpers::IsRuntimeOrPIE() && GetSubsystem() && PCGGridSize > 0)
 	{
-		UnregisterPCG();
+		GetSubsystem()->UnregisterPartitionActor(this);
 	}
+#endif // WITH_EDITOR
 
 	Super::BeginDestroy();
 }
@@ -131,10 +139,12 @@ void APCGPartitionActor::Serialize(FArchive& Ar)
 
 void APCGPartitionActor::Destroyed()
 {
-	if (!PCGHelpers::IsRuntimeOrPIE())
+#if WITH_EDITOR
+	if (!PCGHelpers::IsRuntimeOrPIE() && GetSubsystem() && PCGGridSize > 0)
 	{
-		UnregisterPCG();
+		GetSubsystem()->UnregisterPartitionActor(this);
 	}
+#endif // WITH_EDITOR
 
 	Super::Destroyed();
 }
@@ -161,34 +171,39 @@ void APCGPartitionActor::PostRegisterAllComponents()
 	// Make the Partition actor register itself to the PCG Subsystem
 	// Always do it at runtime, wait for the post load/creation in editor
 	// Only do the mapping if we are at runtime.
-	// Skip if RuntimeGen, RG PAs only register in PostCreate.
+
+	// Make sure the PCGGridSize is not 0, otherwise it will break everything
+
 	const bool bIsRuntimeOrPIE = PCGHelpers::IsRuntimeOrPIE();
-	if (!IsRuntimeGenerated() && !HasAnyFlags(RF_Transient) && (bIsRuntimeOrPIE || bWasPostCreatedLoaded))
+	if ((bIsRuntimeOrPIE 
+#if WITH_EDITOR
+		|| bWasPostCreatedLoaded
+#endif // WITH_EDITOR
+		) && ensure(PCGGridSize > 0))
 	{
-		RegisterPCG(bIsRuntimeOrPIE);
+		if (UPCGSubsystem* Subsystem = GetSubsystem())
+		{
+			Subsystem->RegisterPartitionActor(this, /*bDoComponentMapping*/ bIsRuntimeOrPIE);
+		}
 	}
 }
 
 void APCGPartitionActor::BeginPlay()
 {
-	if (!IsRuntimeGenerated() && !HasAnyFlags(RF_Transient))
+	// Pass through all the pcg components, to verify if we need to generate them
+	for (auto& It : OriginalToLocal)
 	{
-		// Pass through all the pcg components, to verify if we need to generate them.
-		for (auto& It : OriginalToLocal)
+		if (It.Key && It.Value)
 		{
-			if (It.Key && It.Value)
+			// If we have an original component that is generated (or generating), this one is automatically generated => GenerateOnLoad
+			if (It.Key->bGenerated || It.Key->IsGenerating())
 			{
-				// If we have an original component that is generated (or generating), this one is automatically generated => GenerateOnLoad.
-				// But if its runtime generated then it's handled by the runtime generation scheduler.
-				if ((It.Key->bGenerated || It.Key->IsGenerating()) && ensure(It.Key->GenerationTrigger != EPCGComponentGenerationTrigger::GenerateAtRuntime))
-				{
-					It.Value->GenerationTrigger = EPCGComponentGenerationTrigger::GenerateOnLoad;
-				}
-				// Otherwise, make them match
-				else
-				{
-					It.Value->GenerationTrigger = It.Key->GenerationTrigger;
-				}
+				It.Value->GenerationTrigger = EPCGComponentGenerationTrigger::GenerateOnLoad;
+			}
+			// Otherwise, make them match
+			else
+			{
+				It.Value->GenerationTrigger = It.Key->GenerationTrigger;
 			}
 		}
 	}
@@ -198,7 +213,10 @@ void APCGPartitionActor::BeginPlay()
 
 void APCGPartitionActor::EndPlay(const EEndPlayReason::Type EndPlayReason)
 {
-	UnregisterPCG();
+	if (UPCGSubsystem* Subsystem = GetSubsystem())
+	{
+		Subsystem->UnregisterPartitionActor(this);
+	}
 
 	Super::EndPlay(EndPlayReason);
 }
@@ -238,42 +256,6 @@ FIntVector APCGPartitionActor::GetGridCoord() const
 	return UPCGActorHelpers::GetCellCoord(Center, PCGGridSize, bUse2DGrid);
 }
 
-bool APCGPartitionActor::Teleport(const FVector& NewLocation)
-{
-	// We should not be teleporting a PA that is in use. We only teleport empty RuntimeGen PAs to their grid cell before initialization.
-	check(bIsRuntimeGenerated && OriginalToLocal.IsEmpty() && LocalToOriginal.IsEmpty());
-
-	RootComponent->SetMobility(EComponentMobility::Movable);
-	bool bResult = Super::SetActorLocation(NewLocation, /*bSweep=*/false, nullptr, ETeleportType::TeleportPhysics);
-	RootComponent->SetMobility(EComponentMobility::Static);
-
-	return bResult;
-}
-
-void APCGPartitionActor::RegisterPCG(bool bDoComponentMapping)
-{
-	if (UPCGSubsystem* Subsystem = GetSubsystem())
-	{
-		if (!bIsRegistered)
-		{
-			Subsystem->RegisterPartitionActor(this, bDoComponentMapping);
-			bIsRegistered = true;
-		}
-	}
-}
-
-void APCGPartitionActor::UnregisterPCG()
-{
-	if (UPCGSubsystem* Subsystem = GetSubsystem())
-	{
-		if (bIsRegistered)
-		{
-			Subsystem->UnregisterPartitionActor(this);
-			bIsRegistered = false;
-		}
-	}
-}
-
 void APCGPartitionActor::GetActorBounds(bool bOnlyCollidingComponents, FVector& Origin, FVector& BoxExtent, bool bIncludeFromChildActors) const
 {
 	Super::GetActorBounds(bOnlyCollidingComponents, Origin, BoxExtent, bIncludeFromChildActors);
@@ -494,18 +476,10 @@ AActor* APCGPartitionActor::GetSceneOutlinerParent() const
 		return Super::GetSceneOutlinerParent();
 	}	
 }
-#endif // WITH_EDITOR
 
 void APCGPartitionActor::PostCreation(const FGuid& InGridGUID)
 {
-#if WITH_EDITOR
-	// For RuntimeGen we do not care about the grid size from WP, we write it ourselves.
-	if (!IsRuntimeGenerated())
-	{
-		PCGGridSize = Super::GetGridSize();
-	}
-#endif
-
+	PCGGridSize = Super::GetGridSize();
 	PCGGuid = InGridGUID;
 
 	// Put in cache if we use the 2D grid or not.
@@ -524,15 +498,15 @@ void APCGPartitionActor::PostCreation(const FGuid& InGridGUID)
 	{
 		BoundsComponent->SetBoxExtent(GetFixedBounds().GetExtent());
 	}
-#endif // WITH_EDITOR
 
-	// RuntimeGen PAs always register in PostCreate. Non-RG should only register here if in-editor
-	if (IsRuntimeGenerated() || !PCGHelpers::IsRuntimeOrPIE())
+	// Make sure PCGGrid size if greater than 0, otherwise it will break everything
+	if (!PCGHelpers::IsRuntimeOrPIE() && GetSubsystem() && ensure(PCGGridSize > 0))
 	{
-		RegisterPCG(/*bDoComponentMapping=*/false);
+		GetSubsystem()->RegisterPartitionActor(this, /*bDoComponentMapping=*/ false);
 	}
 
 	bWasPostCreatedLoaded = true;
+#endif // WITH_EDITOR
 }
 
 bool APCGPartitionActor::IsSafeForDeletion() const
@@ -568,12 +542,6 @@ TSet<TObjectPtr<UPCGComponent>> APCGPartitionActor::GetAllOriginalPCGComponents(
 	return ResultComponents;
 }
 
-UPCGSubsystem* APCGPartitionActor::GetSubsystem() const 
-{
-	return UPCGSubsystem::GetInstance(GetWorld());
-}
-
-#if WITH_EDITOR
 bool APCGPartitionActor::ChangeTransientState(UPCGComponent* OriginalComponent, EPCGEditorDirtyMode EditingMode)
 {
 	check(OriginalComponent);
@@ -612,14 +580,5 @@ bool APCGPartitionActor::ChangeTransientState(UPCGComponent* OriginalComponent,
 
 	return OriginalToLocal.IsEmpty() && LoadedPreviewComponents.IsEmpty();
 }
-#endif // WITH_EDITOR
 
-FString APCGPartitionActor::GetRuntimeGenActorName(uint32 GridSize, const FIntVector& GridCoords)
-{
-	TStringBuilderWithBuffer<TCHAR, NAME_SIZE> ActorNameBuilder;
-	ActorNameBuilder += TEXT("PCGRuntimeGenPartitionActor_");
-	ActorNameBuilder += FString::Printf(TEXT("%d_"), GridSize);
-	ActorNameBuilder += FString::Printf(TEXT("%d_%d_%d"), GridCoords.X, GridCoords.Y, GridCoords.Z);
-
-	return ActorNameBuilder.ToString();
-}
+#endif // WITH_EDITOR
diff --git a/Engine/Plugins/Experimental/PCG/Source/PCG/Private/Grid/PCGPartitionActor.h b/Engine/Plugins/Experimental/PCG/Source/PCG/Private/Grid/PCGPartitionActor.h
--- a/Engine/Plugins/Experimental/PCG/Source/PCG/Private/Grid/PCGPartitionActor.h
+++ b/Engine/Plugins/Experimental/PCG/Source/PCG/Private/Grid/PCGPartitionActor.h
@@ -55,24 +55,9 @@ public:
 	FBox GetFixedBounds() const;
 	FIntVector GetGridCoord() const;
 
-	void SetPCGGridSize(uint32 InPCGGridSize) { PCGGridSize = InPCGGridSize; }
 	uint32 GetPCGGridSize() const { return PCGGridSize; }
-
 	bool IsUsing2DGrid() const { return bUse2DGrid; }
 
-	/** Marks this PartitionActor as managed by the runtime generation system. */
-	void SetToRuntimeGenerated() { bIsRuntimeGenerated = true; }
-	bool IsRuntimeGenerated() const { return bIsRuntimeGenerated; }
-
-	/** Forces the actor location to change even if its mobility is static. */
-	bool Teleport(const FVector& NewLocation);
-
-	/** Register with the PCG Subsystem. */
-	void RegisterPCG(bool bDoComponentMapping);
-
-	/** Unregister with the PCG Subsystem. */
-	void UnregisterPCG();
-
 	void AddGraphInstance(UPCGComponent* OriginalComponent);
 	void RemapGraphInstance(const UPCGComponent* OldOriginalComponent, UPCGComponent* NewOriginalComponent);
 	bool RemoveGraphInstance(UPCGComponent* OriginalComponent);
@@ -81,6 +66,7 @@ public:
 	// When a local component is destroyed. It calls this function. We make sure we don't keep mappings that are dead.
 	void RemoveLocalComponent(UPCGComponent* LocalComponent);
 
+#if WITH_EDITOR
 	/** To be called after the creation of a new actor to copy the GridSize property (Editor only) into the PCGGridSize property */
 	void PostCreation(const FGuid& InGridGUID);
 
@@ -95,6 +81,7 @@ public:
 
 	/** Changes transient state for the local component matching the given original component. Returns true if PA becomes empty */
 	bool ChangeTransientState(UPCGComponent* OriginalComponent, EPCGEditorDirtyMode EditingMode);
+#endif
 
 	// TODO: Make this in-editor only; during runtime, we should keep a map of component to bounds/volume only
 	// and preferably precompute the intersection, so this would make it easier/possible to not have the original actor in game version.
@@ -135,20 +122,8 @@ private:
 	/** Box component to draw the Partition actor bounds in the Editor viewport */
 	UPROPERTY(Transient)
 	TObjectPtr<UBoxComponent> BoundsComponent;
-#endif // WITH_EDITORONLY_DATA
-
-	/** Tracks the registration status of this PA with the ActorAndComponentMapping system. Helps us avoid invalid (un)registers. */
-	bool bIsRegistered = false;
-
-	/** Tracks if this actor was created by the Runtime Generation system. */
-	bool bIsRuntimeGenerated = false;
 
 	/** Utility bool to check if PostCreation/PostLoad was called. */
 	bool bWasPostCreatedLoaded = false;
-
-public:
-	/** Gets the name this actor would have if it were Runtime Generated.
-	* This does not respect traditional PA name contents like GridGuid, ShouldIncludeGridSizeInName, or ContextHash.
-	*/
-	static FString GetRuntimeGenActorName(uint32 GridSize, const FIntVector& GridCoords);
+#endif // WITH_EDITORONLY_DATA
 };
diff --git a/Engine/Plugins/Experimental/PCG/Source/PCG/Private/PCGActorAndComponentMapping.cpp b/Engine/Plugins/Experimental/PCG/Source/PCG/Private/PCGActorAndComponentMapping.cpp
--- a/Engine/Plugins/Experimental/PCG/Source/PCG/Private/PCGActorAndComponentMapping.cpp
+++ b/Engine/Plugins/Experimental/PCG/Source/PCG/Private/PCGActorAndComponentMapping.cpp
@@ -105,18 +105,10 @@ void FPCGActorAndComponentMapping::Tick()
 TArray<FPCGTaskId> FPCGActorAndComponentMapping::DispatchToRegisteredLocalComponents(UPCGComponent* OriginalComponent, const TFunction<FPCGTaskId(UPCGComponent*)>& InFunc) const
 {
 	TRACE_CPUPROFILER_EVENT_SCOPE(FPCGActorAndComponentMapping::DispatchToRegisteredLocalComponents);
-	if (!ensure(OriginalComponent))
-	{
-		return {};
-	}
-
-	const bool bIsRuntimeGenerated = OriginalComponent->GenerationTrigger == EPCGComponentGenerationTrigger::GenerateAtRuntime;
-
-	const TMap<const UPCGComponent*, TSet<TObjectPtr<APCGPartitionActor>>>& Map = bIsRuntimeGenerated ? ComponentToRuntimeGenPartitionActorsMap : ComponentToPartitionActorsMap;
 
 	// TODO: Might be more interesting to copy the set and release the lock.
-	FReadScopeLock ReadLock(bIsRuntimeGenerated ? ComponentToRuntimeGenPartitionActorsMapLock : ComponentToPartitionActorsMapLock);
-	const TSet<TObjectPtr<APCGPartitionActor>>* PartitionActorsPtr = Map.Find(OriginalComponent);
+	FReadScopeLock ReadLock(ComponentToPartitionActorsMapLock);
+	const TSet<TObjectPtr<APCGPartitionActor>>* PartitionActorsPtr = ComponentToPartitionActorsMap.Find(OriginalComponent);
 
 	if (!PartitionActorsPtr)
 	{
@@ -219,10 +211,10 @@ bool FPCGActorAndComponentMapping::RegisterOrUpdatePartitionedPCGComponent(UPCGC
 	PartitionedOctree.AddOrUpdateComponent(InComponent, Bounds, bComponentHasChanged, bComponentWasAdded);
 
 #if WITH_EDITOR
-	// In Editor only, we will create new partition actors depending on the new bounds and generation trigger. Runtime managed components should not create PAs here
+	// In Editor only, we will create new partition actors depending on the new bounds
 	// TODO: For now it will always create the PA. But if we want to create them only when we generate, we need to make
 	// sure to update the runtime flow, for them to also create PA if they need to.
-	if ((bComponentHasChanged || bComponentWasAdded) && InComponent->GenerationTrigger != EPCGComponentGenerationTrigger::GenerateAtRuntime)
+	if (bComponentHasChanged || bComponentWasAdded)
 	{
 		bool bHasUnbounded = false;
 		PCGHiGenGrid::FSizeArray GridSizes;
@@ -295,26 +287,23 @@ bool FPCGActorAndComponentMapping::RemapPCGComponent(const UPCGComponent* OldCom
 	}
 
 	// Remap all previous instances
-	auto RemapPreviousInstances = [OldComponent, NewComponent](TMap<const UPCGComponent*, TSet<TObjectPtr<APCGPartitionActor>>>& Map, FRWLock& Lock)
 	{
-		FWriteScopeLock WriteLock(Lock);
+		FWriteScopeLock WriteLock(ComponentToPartitionActorsMapLock);
+		TSet<TObjectPtr<APCGPartitionActor>>* PartitionActorsPtr = ComponentToPartitionActorsMap.Find(OldComponent);
 
-		if (TSet<TObjectPtr<APCGPartitionActor>>* PartitionActorsPtr = Map.Find(OldComponent))
+		if (PartitionActorsPtr)
 		{
 			TSet<TObjectPtr<APCGPartitionActor>> PartitionActorsToRemap = MoveTemp(*PartitionActorsPtr);
-			Map.Remove(OldComponent);
+			ComponentToPartitionActorsMap.Remove(OldComponent);
 
 			for (APCGPartitionActor* Actor : PartitionActorsToRemap)
 			{
 				Actor->RemapGraphInstance(OldComponent, NewComponent);
 			}
 
-			Map.Add(NewComponent, MoveTemp(PartitionActorsToRemap));
+			ComponentToPartitionActorsMap.Add(NewComponent, MoveTemp(PartitionActorsToRemap));
 		}
-	};
-
-	RemapPreviousInstances(ComponentToPartitionActorsMap, ComponentToPartitionActorsMapLock);
-	RemapPreviousInstances(ComponentToRuntimeGenPartitionActorsMap, ComponentToRuntimeGenPartitionActorsMapLock);
+	}
 
 	// And update the mapping if bounds changed and we want to do actor mapping
 	if (bBoundsChanged && NewComponent->IsPartitioned() && bDoActorMapping)
@@ -373,14 +362,7 @@ void FPCGActorAndComponentMapping::UnregisterPCGComponent(UPCGComponent* InCompo
 
 void FPCGActorAndComponentMapping::UnregisterPartitionedPCGComponent(UPCGComponent* InComponent)
 {
-	check(InComponent);
-
-	if (!InComponent->GetOwner()->IsA<APCGPartitionActor>())
-	{
-		PCGSubsystem->OnOriginalComponentUnregistered(InComponent);
-	}
-
-	if (!PartitionedOctree.RemoveComponent(InComponent) || InComponent->GenerationTrigger == EPCGComponentGenerationTrigger::GenerateAtRuntime)
+	if (!PartitionedOctree.RemoveComponent(InComponent))
 	{
 		return;
 	}
@@ -407,11 +389,6 @@ void FPCGActorAndComponentMapping::UnregisterPartitionedPCGComponent(UPCGCompone
 
 void FPCGActorAndComponentMapping::UnregisterNonPartitionedPCGComponent(UPCGComponent* InComponent)
 {
-	if (!InComponent->GetOwner()->IsA<APCGPartitionActor>())
-	{
-		PCGSubsystem->OnOriginalComponentUnregistered(InComponent);
-	}
-
 	NonPartitionedOctree.RemoveComponent(InComponent);
 }
 
@@ -423,45 +400,37 @@ void FPCGActorAndComponentMapping::ForAllIntersectingComponents(const FBoxCenter
 	});
 }
 
-void FPCGActorAndComponentMapping::RegisterPartitionActor(APCGPartitionActor* InActor, bool bDoComponentMapping)
+void FPCGActorAndComponentMapping::RegisterPartitionActor(APCGPartitionActor* Actor, bool bDoComponentMapping)
 {
-	check(InActor);
-
-	const uint32 GridSize = InActor->GetPCGGridSize();
-	const FIntVector GridCoord = InActor->GetGridCoord();
-
-	check(GridSize > 0);
-
-	const bool bIsRuntimeGenerated = InActor->IsRuntimeGenerated();
-	TMap<uint32, TMap<FIntVector, TObjectPtr<APCGPartitionActor>>>& ActorsMap = bIsRuntimeGenerated ? RuntimeGenPartitionActorsMap : PartitionActorsMap;
+	check(Actor);
 
+	FIntVector GridCoord = Actor->GetGridCoord();
 	{
-		FWriteScopeLock WriteLock(bIsRuntimeGenerated ? RuntimeGenPartitionActorsMapLock : PartitionActorsMapLock);
+		FWriteScopeLock WriteLock(PartitionActorsMapLock);
 
-		TMap<FIntVector, TObjectPtr<APCGPartitionActor>>& PartitionActorsMapGrid = ActorsMap.FindOrAdd(GridSize);
+		TMap<FIntVector, TObjectPtr<APCGPartitionActor>>& PartitionActorsMapGrid = PartitionActorsMap.FindOrAdd(Actor->GetPCGGridSize());
 		if (PartitionActorsMapGrid.Contains(GridCoord))
 		{
 			return;
 		}
 
-		PartitionActorsMapGrid.Add(GridCoord, InActor);
+		PartitionActorsMapGrid.Add(GridCoord, Actor);
 	}
 
 	// For deprecration: bUse2DGrid is now true by default. But if we already have Partition Actors that were created when the flag was false by default,
 	// we keep this flag
 	if (APCGWorldActor* WorldActor = PCGSubsystem->GetPCGWorldActor())
 	{
-		if (WorldActor->bUse2DGrid != InActor->IsUsing2DGrid())
+		if (WorldActor->bUse2DGrid != Actor->IsUsing2DGrid())
 		{
-			WorldActor->bUse2DGrid = InActor->IsUsing2DGrid();
+			WorldActor->bUse2DGrid = Actor->IsUsing2DGrid();
 		}
 	}
 
-	// Register to all the components that intersect with the PA. Ignore for runtime generated, it is handled manually
-	if (!bIsRuntimeGenerated)
+	// And then register itself to all the components that intersect with it
 	{
 		FWriteScopeLock WriteLock(ComponentToPartitionActorsMapLock);
-		ForAllIntersectingComponents(FBoxCenterAndExtent(InActor->GetFixedBounds()), [this, InActor, bDoComponentMapping](UPCGComponent* Component)
+		ForAllIntersectingComponents(FBoxCenterAndExtent(Actor->GetFixedBounds()), [this, Actor, bDoComponentMapping](UPCGComponent* Component)
 		{
 			// For each component, do the mapping if we ask it explicitly, or if the component is generated
 			if (bDoComponentMapping || Component->bGenerated)
@@ -471,8 +440,8 @@ void FPCGActorAndComponentMapping::RegisterPartitionActor(APCGPartitionActor* In
 				// the mapping might not already exists, even if the component is marked generated.
 				if (PartitionActorsPtr)
 				{
-					InActor->AddGraphInstance(Component);
-					PartitionActorsPtr->Add(InActor);
+					Actor->AddGraphInstance(Component);
+					PartitionActorsPtr->Add(Actor);
 				}
 			}
 		});
@@ -483,33 +452,24 @@ void FPCGActorAndComponentMapping::UnregisterPartitionActor(APCGPartitionActor*
 {
 	check(Actor);
 
-	const FIntVector GridCoord = Actor->GetGridCoord();
-	const uint32 GridSize = Actor->GetPCGGridSize();
-	if (!ensure(GridSize > 0))
-	{
-		return;
-	}
-
-	const bool bIsRuntimeGenerated = Actor->IsRuntimeGenerated();
-	TMap<uint32, TMap<FIntVector, TObjectPtr<APCGPartitionActor>>>& ActorsMap = bIsRuntimeGenerated ? RuntimeGenPartitionActorsMap : PartitionActorsMap;
+	FIntVector GridCoord = Actor->GetGridCoord();
 
-	if (TMap<FIntVector, TObjectPtr<APCGPartitionActor>>* PartitionActorsMapGrid = ActorsMap.Find(GridSize))
+	if (TMap<FIntVector, TObjectPtr<APCGPartitionActor>>* PartitionActorsMapGrid = PartitionActorsMap.Find(Actor->GetPCGGridSize()))
 	{
-		FWriteScopeLock WriteLock(bIsRuntimeGenerated ? RuntimeGenPartitionActorsMapLock : PartitionActorsMapLock);
+		FWriteScopeLock WriteLock(PartitionActorsMapLock);
 		PartitionActorsMapGrid->Remove(GridCoord);
 	}
 
-	// Unregister from all intersecting components. Ignore for runtime generated, it is handled manually
-	if (!bIsRuntimeGenerated)
+	// And then unregister itself to all the components that intersect with it
 	{
 		FWriteScopeLock WriteLock(ComponentToPartitionActorsMapLock);
 		ForAllIntersectingComponents(FBoxCenterAndExtent(Actor->GetFixedBounds()), [this, Actor](UPCGComponent* Component)
 		{
 			TSet<TObjectPtr<APCGPartitionActor>>* PartitionActorsPtr = ComponentToPartitionActorsMap.Find(Component);
-			if (PartitionActorsPtr)
-			{
-				PartitionActorsPtr->Remove(Actor);
-			}
+		if (PartitionActorsPtr)
+		{
+			PartitionActorsPtr->Remove(Actor);
+		}
 		});
 	}
 }
@@ -581,21 +541,21 @@ void FPCGActorAndComponentMapping::UpdateMappingPCGComponentPartitionActor(UPCGC
 
 	TSet<TObjectPtr<APCGPartitionActor>> RemovedActors;
 
-	if (const APCGWorldActor* WorldActor = PCGSubsystem->GetPCGWorldActor())
 	{
-		const bool bIsHiGenEnabled = InComponent->GetGraph() && InComponent->GetGraph()->IsHierarchicalGenerationEnabled();
+		FWriteScopeLock WriteLock(ComponentToPartitionActorsMapLock);
 
-		auto UpdateMapping = [this, InComponent, &Bounds, &RemovedActors, WorldActor, bIsHiGenEnabled](TMap<const UPCGComponent*, TSet<TObjectPtr<APCGPartitionActor>>>& Map, FRWLock& Lock)
+		TSet<TObjectPtr<APCGPartitionActor>>* PartitionActorsPtr = ComponentToPartitionActorsMap.Find(InComponent);
+
+		if (!PartitionActorsPtr)
 		{
-			FWriteScopeLock WriteLock(Lock);
-			TSet<TObjectPtr<APCGPartitionActor>>* PartitionActorsPtr = Map.Find(InComponent);
+			// Does not yet exists, add it
+			PartitionActorsPtr = &ComponentToPartitionActorsMap.Emplace(InComponent);
+			check(PartitionActorsPtr);
+		}
 
-			if (!PartitionActorsPtr)
-			{
-				// Does not yet exists, add it
-				PartitionActorsPtr = &Map.Emplace(InComponent);
-				check(PartitionActorsPtr);
-			}
+		if (const APCGWorldActor* WorldActor = PCGSubsystem->GetPCGWorldActor())
+		{
+			const bool bIsHiGenEnabled = InComponent->GetGraph() && InComponent->GetGraph()->IsHierarchicalGenerationEnabled();
 
 			TSet<TObjectPtr<APCGPartitionActor>> NewMapping;
 			ForAllIntersectingPartitionActors(Bounds, [&NewMapping, InComponent, WorldActor, bIsHiGenEnabled](APCGPartitionActor* Actor)
@@ -613,10 +573,7 @@ void FPCGActorAndComponentMapping::UpdateMappingPCGComponentPartitionActor(UPCGC
 			RemovedActors = PartitionActorsPtr->Difference(NewMapping);
 
 			*PartitionActorsPtr = MoveTemp(NewMapping);
-		};
-
-		UpdateMapping(ComponentToPartitionActorsMap, ComponentToPartitionActorsMapLock);
-		UpdateMapping(ComponentToRuntimeGenPartitionActorsMap, ComponentToRuntimeGenPartitionActorsMapLock);
+		}
 	}
 
 	// No need to be locked to do this.
@@ -635,23 +592,18 @@ void FPCGActorAndComponentMapping::DeleteMappingPCGComponentPartitionActor(UPCGC
 		return;
 	}
 
-	auto DeleteMapping = [this, InComponent](TMap<const UPCGComponent*, TSet<TObjectPtr<APCGPartitionActor>>>& Map, FRWLock& Lock)
-	{
-		FWriteScopeLock WriteLock(Lock);
+	FWriteScopeLock WriteLock(ComponentToPartitionActorsMapLock);
 
-		if (TSet<TObjectPtr<APCGPartitionActor>>* PartitionActorsPtr = Map.Find(InComponent))
+	TSet<TObjectPtr<APCGPartitionActor>>* PartitionActorsPtr = ComponentToPartitionActorsMap.Find(InComponent);
+	if (PartitionActorsPtr)
+	{
+		for (APCGPartitionActor* Actor : *PartitionActorsPtr)
 		{
-			for (APCGPartitionActor* Actor : *PartitionActorsPtr)
-			{
-				Actor->RemoveGraphInstance(InComponent);
-			}
-
-			PartitionActorsPtr->Empty();
+			Actor->RemoveGraphInstance(InComponent);
 		}
-	};
 
-	DeleteMapping(ComponentToPartitionActorsMap, ComponentToPartitionActorsMapLock);
-	DeleteMapping(ComponentToRuntimeGenPartitionActorsMap, ComponentToRuntimeGenPartitionActorsMapLock);
+		PartitionActorsPtr->Empty();
+	}
 }
 
 bool FPCGActorAndComponentMapping::IsComponentRegistered(const UPCGComponent* InComponent) const
@@ -676,12 +628,11 @@ TSet<UPCGComponent*> FPCGActorAndComponentMapping::GetAllRegisteredComponents()
 	return Res;
 }
 
-UPCGComponent* FPCGActorAndComponentMapping::GetLocalComponent(uint32 GridSize, const FIntVector& CellCoords, const UPCGComponent* InOriginalComponent, bool bRuntimeGenerated)
+UPCGComponent* FPCGActorAndComponentMapping::GetLocalComponent(uint32 GridSize, const FIntVector& CellCoords, const UPCGComponent* InOriginalComponent)
 {
-	TMap<uint32, TMap<FIntVector, TObjectPtr<APCGPartitionActor>>>& ActorsMap = bRuntimeGenerated ? RuntimeGenPartitionActorsMap : PartitionActorsMap;
-	FReadScopeLock ReadLock(bRuntimeGenerated ? RuntimeGenPartitionActorsMapLock : PartitionActorsMapLock);
+	FReadScopeLock ReadLock(PartitionActorsMapLock);
 
-	if (const TMap<FIntVector, TObjectPtr<APCGPartitionActor>>* PartitionActorsOnGrid = ActorsMap.Find(GridSize))
+	if (const TMap<FIntVector, TObjectPtr<APCGPartitionActor>>* PartitionActorsOnGrid = PartitionActorsMap.Find(GridSize))
 	{
 		const TObjectPtr<APCGPartitionActor>* PartitionActor = PartitionActorsOnGrid->Find(CellCoords);
 		if (PartitionActor && *PartitionActor)
@@ -701,20 +652,6 @@ UPCGComponent* FPCGActorAndComponentMapping::GetLocalComponent(uint32 GridSize,
 	return nullptr;
 }
 
-APCGPartitionActor* FPCGActorAndComponentMapping::GetPartitionActor(uint32 GridSize, const FIntVector& CellCoords, bool bRuntimeGenerated) const
-{
-	const TMap<uint32, TMap<FIntVector, TObjectPtr<APCGPartitionActor>>>& Map = bRuntimeGenerated ? RuntimeGenPartitionActorsMap : PartitionActorsMap;
-	FReadScopeLock ReadLock(bRuntimeGenerated ? RuntimeGenPartitionActorsMapLock : PartitionActorsMapLock);
-
-	if (const TMap<FIntVector, TObjectPtr<APCGPartitionActor>>* PartitionActorsOnGrid = Map.Find(GridSize))
-	{
-		const TObjectPtr<APCGPartitionActor>* PartitionActor = PartitionActorsOnGrid->Find(CellCoords);
-		return PartitionActor ? *PartitionActor : nullptr;
-	}
-
-	return nullptr;
-}
-
 #if WITH_EDITOR
 void FPCGActorAndComponentMapping::RegisterTracking(UPCGComponent* InComponent)
 {
diff --git a/Engine/Plugins/Experimental/PCG/Source/PCG/Private/PCGCommon.cpp b/Engine/Plugins/Experimental/PCG/Source/PCG/Private/PCGCommon.cpp
--- a/Engine/Plugins/Experimental/PCG/Source/PCG/Private/PCGCommon.cpp
+++ b/Engine/Plugins/Experimental/PCG/Source/PCG/Private/PCGCommon.cpp
@@ -30,11 +30,6 @@ namespace PCGHiGenGrid
 		return InGrid >= EPCGHiGenGrid::GridMin && static_cast<uint32>(InGrid) < 2 * static_cast<uint32>(EPCGHiGenGrid::GridMax);
 	}
 
-	bool IsValidGridOrUninitialized(EPCGHiGenGrid InGrid)
-	{
-		return IsValidGrid(InGrid) || InGrid == EPCGHiGenGrid::Uninitialized;
-	}
-
 	uint32 GridToGridSize(EPCGHiGenGrid InGrid)
 	{
 		const uint32 GridAsUint = static_cast<uint32>(InGrid);
@@ -64,30 +59,3 @@ namespace PCGHiGenGrid
 		return 100 * static_cast<uint32>(EPCGHiGenGrid::Unbounded);
 	}
 }
-
-double FPCGRuntimeGenerationRadii::GetGenerationRadiusFromGrid(EPCGHiGenGrid Grid) const
-{
-	switch (Grid)
-	{
-		case EPCGHiGenGrid::Grid4: return GenerationRadius400;
-		case EPCGHiGenGrid::Grid8: return GenerationRadius800;
-		case EPCGHiGenGrid::Grid16: return GenerationRadius1600;
-		case EPCGHiGenGrid::Grid32: return GenerationRadius3200;
-		case EPCGHiGenGrid::Grid64: return GenerationRadius6400;
-		case EPCGHiGenGrid::Grid128: return GenerationRadius12800;
-		case EPCGHiGenGrid::Grid256: return GenerationRadius25600;
-		case EPCGHiGenGrid::Grid512: return GenerationRadius51200;
-		case EPCGHiGenGrid::Grid1024: return GenerationRadius102400;
-		case EPCGHiGenGrid::Grid2048: return GenerationRadius204800;
-		case EPCGHiGenGrid::Unbounded: return GenerationRadius;
-	}
-
-	ensure(false);
-	return 0;
-}
-
-double FPCGRuntimeGenerationRadii::GetCleanupRadiusFromGrid(EPCGHiGenGrid Grid) const
-{
-	return GetGenerationRadiusFromGrid(Grid) * CleanupRadiusMultiplier;
-}
-
diff --git a/Engine/Plugins/Experimental/PCG/Source/PCG/Private/PCGComponent.cpp b/Engine/Plugins/Experimental/PCG/Source/PCG/Private/PCGComponent.cpp
--- a/Engine/Plugins/Experimental/PCG/Source/PCG/Private/PCGComponent.cpp
+++ b/Engine/Plugins/Experimental/PCG/Source/PCG/Private/PCGComponent.cpp
@@ -26,9 +26,6 @@
 #include "Grid/PCGPartitionActor.h"
 #include "Helpers/PCGActorHelpers.h"
 #include "Helpers/PCGHelpers.h"
-#include "RuntimeGen/GenSources/PCGGenSourceBase.h"
-#include "RuntimeGen/SchedulingPolicies/PCGSchedulingPolicyBase.h"
-#include "RuntimeGen/SchedulingPolicies/PCGSchedulingPolicyDistanceAndDirection.h"
 #include "Utils/PCGGeneratedResourcesLogging.h"
 
 #include "LandscapeComponent.h"
@@ -45,7 +42,6 @@
 #include "GameFramework/Volume.h"
 #include "Kismet/GameplayStatics.h"
 #include "LandscapeSplinesComponent.h"
-#include "UObject/Package.h"
 
 #include UE_INLINE_GENERATED_CPP_BY_NAME(PCGComponent)
 
@@ -67,7 +63,6 @@ UPCGComponent::UPCGComponent(const FObjectInitializer& InObjectInitializer)
 	: Super(InObjectInitializer)
 {
 	GraphInstance = InObjectInitializer.CreateDefaultSubobject<UPCGGraphInstance>(this, TEXT("PCGGraphInstance"));
-	SchedulingPolicyClass = UPCGSchedulingPolicyDistanceAndDirection::StaticClass();
 
 #if WITH_EDITOR
 	// If we are in Editor, and we are a BP template (no owner), we will mark this component to force a generate when added to world.
@@ -80,10 +75,8 @@ UPCGComponent::UPCGComponent(const FObjectInitializer& InObjectInitializer)
 
 bool UPCGComponent::CanPartition() const
 {
-	// Support/Force partitioning on non-PCG partition actors in WP worlds. GenerateAtRuntime components can be partitioned even if WorldPartition is not enabled.
-	return ((GetOwner() && GetOwner()->GetWorld() && GetOwner()->GetWorld()->GetWorldPartition() != nullptr)
-		|| GenerationTrigger == EPCGComponentGenerationTrigger::GenerateAtRuntime)
-		&& Cast<APCGPartitionActor>(GetOwner()) == nullptr;
+	// Support/Force partitioning on non-PCG partition actors in WP worlds.
+	return GetOwner() && GetOwner()->GetWorld() && GetOwner()->GetWorld()->GetWorldPartition() != nullptr && Cast<APCGPartitionActor>(GetOwner()) == nullptr;
 }
 
 bool UPCGComponent::IsPartitioned() const
@@ -173,18 +166,6 @@ bool UPCGComponent::ShouldGenerate(bool bForce, EPCGComponentGenerationTrigger R
 		return false;
 	}
 
-	if (GenerationTrigger == EPCGComponentGenerationTrigger::GenerateAtRuntime)
-	{
-		// If we're runtime generated, turn down other requests.
-		const bool bShouldGenerate = RequestedGenerationTrigger == EPCGComponentGenerationTrigger::GenerateAtRuntime;
-		if (!bShouldGenerate)
-		{
-			UE_LOG(LogPCG, Warning, TEXT("Generation request with trigger %d denied as this component is managed by the runtime generation scheduler."), (int)RequestedGenerationTrigger);
-		}
-
-		return bShouldGenerate;
-	}
-
 #if WITH_EDITOR
 	// Always run Generate if we are in editor and partitioned since the original component doesn't know the state of the local one.
 	if (IsPartitioned() && !PCGHelpers::IsRuntimeOrPIE())
@@ -231,11 +212,6 @@ void UPCGComponent::SetPropertiesFromOriginal(const UPCGComponent* Original)
 	Seed = Original->Seed;
 	GenerationTrigger = Original->GenerationTrigger;
 
-	bOverrideGenerationRadii = Original->bOverrideGenerationRadii;
-	GenerationRadii = Original->GenerationRadii;
-	SchedulingPolicyClass = Original->SchedulingPolicyClass;
-	SchedulingPolicy = Original->SchedulingPolicy;
-
 	UPCGGraph* OriginalGraph = Original->GraphInstance ? Original->GraphInstance->GetGraph() : nullptr;
 	if (OriginalGraph != GraphInstance->GetGraph())
 	{
@@ -278,25 +254,15 @@ void UPCGComponent::Generate_Implementation(bool bForce)
 
 void UPCGComponent::GenerateLocal(bool bForce)
 {
-	GenerateInternal(bForce, EPCGHiGenGrid::Uninitialized, EPCGComponentGenerationTrigger::GenerateOnDemand, {});
-}
-
-void UPCGComponent::GenerateLocal(EPCGComponentGenerationTrigger RequestedGenerationTrigger, bool bForce, EPCGHiGenGrid Grid)
-{
-	GenerateInternal(bForce, Grid, RequestedGenerationTrigger, {});
+	GenerateInternal(bForce, EPCGComponentGenerationTrigger::GenerateOnDemand, {});
 }
 
 FPCGTaskId UPCGComponent::GenerateLocalGetTaskId(bool bForce)
 {
-	return GenerateInternal(bForce, EPCGHiGenGrid::Uninitialized, EPCGComponentGenerationTrigger::GenerateOnDemand, {});
+	return GenerateInternal(bForce, EPCGComponentGenerationTrigger::GenerateOnDemand, {});
 }
 
-FPCGTaskId UPCGComponent::GenerateLocalGetTaskId(EPCGComponentGenerationTrigger RequestedGenerationTrigger, bool bForce, EPCGHiGenGrid Grid)
-{
-	return GenerateInternal(bForce, Grid, RequestedGenerationTrigger, {});
-}
-
-FPCGTaskId UPCGComponent::GenerateInternal(bool bForce, EPCGHiGenGrid Grid, EPCGComponentGenerationTrigger RequestedGenerationTrigger, const TArray<FPCGTaskId>& Dependencies)
+FPCGTaskId UPCGComponent::GenerateInternal(bool bForce, EPCGComponentGenerationTrigger RequestedGenerationTrigger, const TArray<FPCGTaskId>& Dependencies)
 {
 	if (IsGenerating() || !GetSubsystem() || !ShouldGenerate(bForce, RequestedGenerationTrigger))
 	{
@@ -305,7 +271,7 @@ FPCGTaskId UPCGComponent::GenerateInternal(bool bForce, EPCGHiGenGrid Grid, EPCG
 
 	Modify(!IsInPreviewMode());
 
-	CurrentGenerationTask = GetSubsystem()->ScheduleComponent(this, Grid, /*bSave=*/bForce, Dependencies);
+	CurrentGenerationTask = GetSubsystem()->ScheduleComponent(this, /*bSave=*/bForce, Dependencies);
 
 	return CurrentGenerationTask;
 }
@@ -538,12 +504,6 @@ void UPCGComponent::Cleanup()
 		return;
 	}
 
-	if (GenerationTrigger == EPCGComponentGenerationTrigger::GenerateAtRuntime)
-	{
-		UE_LOG(LogPCG, Warning, TEXT("Cleanup request denied as this component is managed by the runtime generation scheduler."));
-		return;
-	}
-
 #if WITH_EDITOR
 	FScopedTransaction Transaction(LOCTEXT("PCGCleanup", "Clean up PCG component"));
 #endif
@@ -706,51 +666,6 @@ void UPCGComponent::ClearPerPinGeneratedOutput()
 	PerPinGeneratedOutput.Reset();
 }
 
-void UPCGComponent::SetSchedulingPolicyClass(TSubclassOf<UPCGSchedulingPolicyBase> InSchedulingPolicyClass) 
-{
-	if (!SchedulingPolicy || InSchedulingPolicyClass != SchedulingPolicyClass)
-	{
-		if (InSchedulingPolicyClass != SchedulingPolicyClass)
-		{
-			SchedulingPolicyClass = InSchedulingPolicyClass;
-		}
-		
-		RefreshSchedulingPolicy();
-	}
-}
-
-double UPCGComponent::GetGenerationRadiusFromGrid(EPCGHiGenGrid Grid) const
-{
-	if (bOverrideGenerationRadii)
-	{
-		return GenerationRadii.GetGenerationRadiusFromGrid(Grid);
-	}
-
-	const UPCGGraph* Graph = GetGraph();
-	if (ensure(Graph))
-	{
-		return Graph->GenerationRadii.GetGenerationRadiusFromGrid(Grid);
-	}
-
-	return 0;
-}
-
-double UPCGComponent::GetCleanupRadiusFromGrid(EPCGHiGenGrid Grid) const
-{
-	if (bOverrideGenerationRadii)
-	{
-		return GenerationRadii.GetCleanupRadiusFromGrid(Grid);
-	}
-
-	const UPCGGraph* Graph = GetGraph();
-	if (ensure(Graph))
-	{
-		return Graph->GenerationRadii.GetCleanupRadiusFromGrid(Grid);
-	}
-
-	return 0;
-}
-
 bool UPCGComponent::MoveResourcesToNewActor(AActor* InNewActor, bool bCreateChild)
 {
 	// Don't move resources if we are generating or cleaning up
@@ -816,9 +731,6 @@ void UPCGComponent::CleanupLocalImmediate(bool bRemoveComponents)
 {
 	PCGGeneratedResourcesLogging::LogCleanupLocalImmediate(bRemoveComponents, GeneratedResources);
 
-	// Cancels generation of this component if there is an ongoing generation in progress.
-	CancelGeneration();
-
 	TSet<TSoftObjectPtr<AActor>> ActorsToDelete;
 
 	if (!bRemoveComponents && UPCGManagedResource::DebugForcePurgeAllResourcesOnGenerate())
@@ -1121,7 +1033,7 @@ void UPCGComponent::BeginPlay()
 		}
 		else
 		{
-			GenerateInternal(/*bForce=*/false, EPCGHiGenGrid::Uninitialized, EPCGComponentGenerationTrigger::GenerateOnLoad, {});
+			GenerateInternal(/*bForce=*/false, EPCGComponentGenerationTrigger::GenerateOnLoad, {});
 			bRuntimeGenerated = true;
 		}
 	}
@@ -1274,19 +1186,6 @@ void UPCGComponent::PostLoad()
 		bDirtyGenerated = true;
 	}
 #endif
-
-	if (!SchedulingPolicy)
-	{
-		RefreshSchedulingPolicy();
-	}
-	else
-	{
-		const EObjectFlags Flags = GetMaskedFlags(RF_PropagateToSubObjects) | RF_Transactional;
-		SchedulingPolicy->SetFlags(Flags);
-#if WITH_EDITOR
-		SchedulingPolicy->SetShouldDisplayProperties(GenerationTrigger == EPCGComponentGenerationTrigger::GenerateAtRuntime);
-#endif
-	}
 }
 
 #if WITH_EDITOR
@@ -1388,19 +1287,10 @@ TStructOnScope<FActorComponentInstanceData> UPCGComponent::GetComponentInstanceD
 
 void UPCGComponent::OnGraphChanged(UPCGGraphInterface* InGraph, EPCGChangeType ChangeType)
 {
-	if (!!(ChangeType & EPCGChangeType::Debug) && GenerationTrigger == EPCGComponentGenerationTrigger::GenerateAtRuntime)
-	{
-		// TODO: Handle RuntimeGen business in the RuntimeGenScheduler, subscribe to OnGraphChanged there instead.
-		CleanupLocalImmediate(/*bRemoveComponents=*/true);
-		GenerateLocal(EPCGComponentGenerationTrigger::GenerateAtRuntime, /*bForce=*/true);
-	}
-	else
-	{
-		const bool bIsStructural = ((ChangeType & (EPCGChangeType::Edge | EPCGChangeType::Structural)) != EPCGChangeType::None);
-		const bool bDirtyInputs = bIsStructural || ((ChangeType & EPCGChangeType::Input) != EPCGChangeType::None);
+	const bool bIsStructural = ((ChangeType & (EPCGChangeType::Edge | EPCGChangeType::Structural)) != EPCGChangeType::None);
+	const bool bDirtyInputs = bIsStructural || ((ChangeType & EPCGChangeType::Input) != EPCGChangeType::None);
 
-		RefreshAfterGraphChanged(InGraph, bIsStructural, bDirtyInputs);
-	}
+	RefreshAfterGraphChanged(InGraph, bIsStructural, bDirtyInputs);
 }
 
 void UPCGComponent::RefreshAfterGraphChanged(UPCGGraphInterface* InGraph, bool bIsStructural, bool bDirtyInputs)
@@ -1531,31 +1421,6 @@ void UPCGComponent::PostEditChangeProperty(FPropertyChangedEvent& PropertyChange
 		DirtyGenerated();
 		Refresh();
 	}
-	else if (PropName == GET_MEMBER_NAME_CHECKED(UPCGComponent, SchedulingPolicyClass))
-	{
-		// We don't need to refresh the component here because this does not effect generation behavior, only scheduling behavior.
-		RefreshSchedulingPolicy();
-	} 
-	else if (PropName == GET_MEMBER_NAME_CHECKED(UPCGComponent, GenerationTrigger))
-	{
-		if (GenerationTrigger == EPCGComponentGenerationTrigger::GenerateAtRuntime)
-		{
-			CleanupLocalImmediate(/*bRemoveComponents=*/true);
-		}
-
-		if (!SchedulingPolicy)
-		{
-			RefreshSchedulingPolicy();
-		}
-
-		if (SchedulingPolicy)
-		{
-			// We should only display scheduling policy parameters when in runtime generation mode.
-			SchedulingPolicy->SetShouldDisplayProperties(GenerationTrigger == EPCGComponentGenerationTrigger::GenerateAtRuntime);
-		}
-
-		Refresh();
-	}
 	// General properties that don't affect behavior
 	else
 	{
@@ -1822,12 +1687,6 @@ void UPCGComponent::Refresh(bool bStructural)
 		return;
 	}
 
-	// Do not allow automatic refreshing for runtime generation components.
-	if (GenerationTrigger == EPCGComponentGenerationTrigger::GenerateAtRuntime)
-	{
-		return;
-	}
-
 	// If the component is tagged as not to regenerate in the editor, only exceptional cases should trigger a refresh
 	// namely: the component is deactivated.
 	// Note that the component changing its IsPartitioned state is already covered in the PostEditChangeProperty
@@ -1951,8 +1810,6 @@ bool UPCGComponent::IsActorTracked(AActor* InActor, bool& bOutIsCulled) const
 
 void UPCGComponent::OnRefresh(bool bForceRefresh)
 {
-	check(GenerationTrigger != EPCGComponentGenerationTrigger::GenerateAtRuntime);
-
 	// Mark the refresh task invalid to allow re-triggering refreshes
 	CurrentRefreshTask = InvalidPCGTaskId;
 
@@ -1967,7 +1824,7 @@ void UPCGComponent::OnRefresh(bool bForceRefresh)
 		// If we are partitioned but we have resources, we need to force a cleanup
 		if (!GeneratedResources.IsEmpty())
 		{
-			CleanupLocalImmediate(/*bRemoveComponents=*/true);
+			CleanupLocalImmediate(true);
 		}
 	}
 
@@ -2357,27 +2214,6 @@ FPCGDataCollection UPCGComponent::CreateActorPCGDataCollection(AActor* Actor, co
 	return Collection;
 }
 
-void UPCGComponent::RefreshSchedulingPolicy()
-{
-	if (SchedulingPolicy)
-	{
-		SchedulingPolicy->Rename(nullptr, GetTransientPackage(), REN_DontCreateRedirectors | REN_ForceNoResetLoaders);
-		SchedulingPolicy->MarkAsGarbage();
-		SchedulingPolicy = nullptr;
-	}
-
-	if (SchedulingPolicyClass)
-	{
-		const EObjectFlags Flags = GetMaskedFlags(RF_PropagateToSubObjects);
-		SchedulingPolicy = NewObject<UPCGSchedulingPolicyBase>(this, SchedulingPolicyClass, NAME_None, Flags);
-
-#if WITH_EDITOR
-		// We should only display scheduling policy parameters when in runtime generation mode.
-		SchedulingPolicy->SetShouldDisplayProperties(GenerationTrigger == EPCGComponentGenerationTrigger::GenerateAtRuntime);
-#endif
-	}
-}
-
 UPCGData* UPCGComponent::CreatePCGData()
 {
 	TRACE_CPUPROFILER_EVENT_SCOPE(UPCGComponent::CreatePCGData);
diff --git a/Engine/Plugins/Experimental/PCG/Source/PCG/Private/PCGSubsystem.cpp b/Engine/Plugins/Experimental/PCG/Source/PCG/Private/PCGSubsystem.cpp
--- a/Engine/Plugins/Experimental/PCG/Source/PCG/Private/PCGSubsystem.cpp
+++ b/Engine/Plugins/Experimental/PCG/Source/PCG/Private/PCGSubsystem.cpp
@@ -14,8 +14,6 @@
 #include "Grid/PCGPartitionActor.h"
 #include "Helpers/PCGActorHelpers.h"
 #include "Helpers/PCGHelpers.h"
-#include "RuntimeGen/PCGGenSourceManager.h"
-#include "RuntimeGen/PCGRuntimeGenScheduler.h"
 
 #include "Engine/Level.h"
 
@@ -151,9 +149,6 @@ void UPCGSubsystem::Deinitialize()
 	delete GraphExecutor;
 	GraphExecutor = nullptr;
 
-	delete RuntimeGenScheduler;
-	RuntimeGenScheduler = nullptr;
-
 	PCGWorldActor = nullptr;
 	bHasTickedOnce = false;
 
@@ -171,11 +166,7 @@ void UPCGSubsystem::PostInitialize()
 	// Initialize graph executor
 	check(!GraphExecutor);
 	GraphExecutor = new FPCGGraphExecutor(this);
-
-	// Initialize runtime generation scheduler
-	check(!RuntimeGenScheduler);
-	RuntimeGenScheduler = new FPCGRuntimeGenScheduler(GetWorld(), &ActorAndComponentMapping);
-
+	
 	// Gather world pcg actor if it exists
 	if (!PCGWorldActor)
 	{
@@ -241,11 +232,6 @@ void UPCGSubsystem::Tick(float DeltaSeconds)
 	}
 
 	ActorAndComponentMapping.Tick();
-
-	if (PCGWorldActor)
-	{
-		RuntimeGenScheduler->Tick(PCGWorldActor);
-	}
 }
 
 APCGWorldActor* UPCGSubsystem::GetPCGWorldActor()
@@ -328,21 +314,13 @@ void UPCGSubsystem::UnregisterPCGWorldActor(APCGWorldActor* InActor)
 	}
 }
 
-void UPCGSubsystem::OnOriginalComponentUnregistered(UPCGComponent* InComponent)
-{
-	if (RuntimeGenScheduler)
-	{
-		RuntimeGenScheduler->OnOriginalComponentUnregistered(InComponent);
-	}
-}
-
 UPCGLandscapeCache* UPCGSubsystem::GetLandscapeCache()
 {
 	APCGWorldActor* LandscapeCacheOwner = GetPCGWorldActor();
 	return LandscapeCacheOwner ? LandscapeCacheOwner->LandscapeCacheObject.Get() : nullptr;
 }
 
-FPCGTaskId UPCGSubsystem::ScheduleComponent(UPCGComponent* PCGComponent, EPCGHiGenGrid Grid, bool bSave, const TArray<FPCGTaskId>& InDependencies)
+FPCGTaskId UPCGSubsystem::ScheduleComponent(UPCGComponent* PCGComponent, bool bSave, const TArray<FPCGTaskId>& InDependencies)
 {
 	check(GraphExecutor);
 
@@ -356,8 +334,7 @@ FPCGTaskId UPCGSubsystem::ScheduleComponent(UPCGComponent* PCGComponent, EPCGHiG
 	ensure(PCGHelpers::GetGenerationGridSizes(PCGComponent->GetGraph(), GetPCGWorldActor(), GridSizes, bHasUnbounded));
 
 #if WITH_EDITOR
-	// Create the PartitionActors if necessary. Skip if this is a runtime managed component, PAs are handled manually by the RuntimeGenScheduler.
-	if (PCGComponent->IsPartitioned() && !PCGHelpers::IsRuntimeOrPIE() && PCGComponent->GenerationTrigger != EPCGComponentGenerationTrigger::GenerateAtRuntime)
+	if (PCGComponent->IsPartitioned() && !PCGHelpers::IsRuntimeOrPIE())
 	{
 		if (!GridSizes.IsEmpty())
 		{
@@ -376,20 +353,7 @@ FPCGTaskId UPCGSubsystem::ScheduleComponent(UPCGComponent* PCGComponent, EPCGHiG
 
 	// Schedule generation of original component if is is non-partitioned, or if it has nodes that will execute at the Unbounded level.
 	FPCGTaskId OriginalComponentTask = InvalidPCGTaskId;
-
-	bool bGeneratePCGComponent = false;
-	if (!PCGComponent->IsPartitioned())
-	{
-		// This component is either an unpartitioned original component or a local component. Generate if grid size matches preference (if provided).
-		bGeneratePCGComponent = (Grid == EPCGHiGenGrid::Uninitialized) || !!(Grid & PCGComponent->GetGenerationGrid());
-	}
-	else
-	{
-		// This component is a partitioned original component. Generate if the graph has unbounded nodes and if this grid matches preference (if provided).
-		bGeneratePCGComponent = bHasUnbounded && (!!(Grid & EPCGHiGenGrid::Unbounded) || Grid == EPCGHiGenGrid::Uninitialized);
-	}
-
-	if (bGeneratePCGComponent)
+	if (!PCGComponent->IsPartitioned() || bHasUnbounded)
 	{
 		OriginalComponentTask = PCGComponent->CreateGenerateTask(/*bForce=*/bSave, InDependencies);
 		if (OriginalComponentTask != InvalidPCGTaskId)
@@ -399,7 +363,7 @@ FPCGTaskId UPCGSubsystem::ScheduleComponent(UPCGComponent* PCGComponent, EPCGHiG
 	}
 
 	// If the component is partitioned, we will forward the calls to its registered PCG Partition actors
-	if (PCGComponent->IsPartitioned() && PCGHiGenGrid::IsValidGridOrUninitialized(Grid))
+	if (PCGComponent->IsPartitioned())
 	{
 		// Local components depend on the original component (to ensure any data is available).
 		TArray<FPCGTaskId> Dependencies = InDependencies;
@@ -408,14 +372,9 @@ FPCGTaskId UPCGSubsystem::ScheduleComponent(UPCGComponent* PCGComponent, EPCGHiG
 			Dependencies.Add(OriginalComponentTask);
 		}
 
-		auto LocalGenerateTask = [OriginalComponent = PCGComponent, Grid, &Dependencies, bSave, &GridSizes](UPCGComponent* LocalComponent)
+		auto LocalGenerateTask = [OriginalComponent = PCGComponent, &Dependencies, bSave, &GridSizes](UPCGComponent* LocalComponent)
 		{
-			const uint32 LocalComponentGridSize = LocalComponent->GetGenerationGridSize();
-			if (!GridSizes.Contains(LocalComponentGridSize))
-			{
-				return InvalidPCGTaskId;
-			}
-			else if (Grid != EPCGHiGenGrid::Uninitialized && !(Grid & LocalComponent->GetGenerationGrid()))
+			if (!GridSizes.Contains(LocalComponent->GetGenerationGridSize()))
 			{
 				return InvalidPCGTaskId;
 			}
@@ -437,7 +396,7 @@ FPCGTaskId UPCGSubsystem::ScheduleComponent(UPCGComponent* PCGComponent, EPCGHiG
 				LocalComponent->CleanupLocalImmediate(true);
 			}
 
-			return LocalComponent->GenerateInternal(/*bForce=*/bSave, LocalComponent->GetGenerationGrid(), EPCGComponentGenerationTrigger::GenerateOnDemand, Dependencies);
+			return LocalComponent->GenerateInternal(/*bForce=*/bSave, EPCGComponentGenerationTrigger::GenerateOnDemand, Dependencies);
 		};
 
 		AllTasks.Append(ActorAndComponentMapping.DispatchToRegisteredLocalComponents(PCGComponent, LocalGenerateTask));
@@ -699,9 +658,9 @@ void UPCGSubsystem::ForAllOverlappingComponentsInHierarchy(UPCGComponent* InComp
 	});
 }
 
-UPCGComponent* UPCGSubsystem::GetLocalComponent(uint32 GridSize, const FIntVector& CellCoords, const UPCGComponent* InOriginalComponent, bool bTransient)
+UPCGComponent* UPCGSubsystem::GetLocalComponent(uint32 GridSize, const FIntVector& CellCoords, const UPCGComponent* InOriginalComponent)
 {
-	return ActorAndComponentMapping.GetLocalComponent(GridSize, CellCoords, InOriginalComponent, bTransient);
+	return ActorAndComponentMapping.GetLocalComponent(GridSize, CellCoords, InOriginalComponent);
 }
 
 bool UPCGSubsystem::IsGraphCacheDebuggingEnabled() const
@@ -709,11 +668,6 @@ bool UPCGSubsystem::IsGraphCacheDebuggingEnabled() const
 	return GraphExecutor && GraphExecutor->IsGraphCacheDebuggingEnabled();
 }
 
-FPCGGenSourceManager* UPCGSubsystem::GetGenSourceManager() const
-{
-	return RuntimeGenScheduler ? RuntimeGenScheduler->GenSourceManager : nullptr;
-}
-
 bool UPCGSubsystem::GetOutputData(FPCGTaskId TaskId, FPCGDataCollection& OutData)
 {
 	check(GraphExecutor);
@@ -995,7 +949,7 @@ FPCGTaskId UPCGSubsystem::ProcessGraph(UPCGComponent* Component, const FBox& InP
 						return InvalidPCGTaskId;
 					}
 
-					return LocalComponent->GenerateInternal(bSave, EPCGHiGenGrid::Uninitialized, EPCGComponentGenerationTrigger::GenerateOnDemand, TaskDependencies);
+					return LocalComponent->GenerateInternal(bSave, EPCGComponentGenerationTrigger::GenerateOnDemand, TaskDependencies);
 				}
 				else
 				{
@@ -1061,8 +1015,6 @@ FPCGTaskId UPCGSubsystem::ProcessGraph(UPCGComponent* Component, const FBox& InP
 
 FPCGTaskId UPCGSubsystem::ScheduleRefresh(UPCGComponent* Component, bool bForceRegen)
 {
-	check(Component && Component->GenerationTrigger != EPCGComponentGenerationTrigger::GenerateAtRuntime);
-
 	TWeakObjectPtr<UPCGComponent> ComponentPtr(Component);
 
 	auto RefreshTask = [ComponentPtr, bForceRegen]() {
diff --git a/Engine/Plugins/Experimental/PCG/Source/PCG/Private/RuntimeGen/GenSources/PCGGenSourceComponent.cpp b/Engine/Plugins/Experimental/PCG/Source/PCG/Private/RuntimeGen/GenSources/PCGGenSourceComponent.cpp
deleted file mode 100644
diff --git a/Engine/Plugins/Experimental/PCG/Source/PCG/Private/RuntimeGen/GenSources/PCGGenSourceEditorCamera.cpp b/Engine/Plugins/Experimental/PCG/Source/PCG/Private/RuntimeGen/GenSources/PCGGenSourceEditorCamera.cpp
deleted file mode 100644
diff --git a/Engine/Plugins/Experimental/PCG/Source/PCG/Private/RuntimeGen/GenSources/PCGGenSourcePlayer.cpp b/Engine/Plugins/Experimental/PCG/Source/PCG/Private/RuntimeGen/GenSources/PCGGenSourcePlayer.cpp
deleted file mode 100644
diff --git a/Engine/Plugins/Experimental/PCG/Source/PCG/Private/RuntimeGen/GenSources/PCGGenSourceWPStreamingSource.cpp b/Engine/Plugins/Experimental/PCG/Source/PCG/Private/RuntimeGen/GenSources/PCGGenSourceWPStreamingSource.cpp
deleted file mode 100644
diff --git a/Engine/Plugins/Experimental/PCG/Source/PCG/Private/RuntimeGen/PCGGenSourceManager.cpp b/Engine/Plugins/Experimental/PCG/Source/PCG/Private/RuntimeGen/PCGGenSourceManager.cpp
deleted file mode 100644
diff --git a/Engine/Plugins/Experimental/PCG/Source/PCG/Private/RuntimeGen/PCGRuntimeGenScheduler.cpp b/Engine/Plugins/Experimental/PCG/Source/PCG/Private/RuntimeGen/PCGRuntimeGenScheduler.cpp
deleted file mode 100644
diff --git a/Engine/Plugins/Experimental/PCG/Source/PCG/Private/RuntimeGen/SchedulingPolicies/PCGSchedulingPolicyDistanceAndDirection.cpp b/Engine/Plugins/Experimental/PCG/Source/PCG/Private/RuntimeGen/SchedulingPolicies/PCGSchedulingPolicyDistanceAndDirection.cpp
deleted file mode 100644
diff --git a/Engine/Plugins/Experimental/PCG/Source/PCG/Public/PCGActorAndComponentMapping.h b/Engine/Plugins/Experimental/PCG/Source/PCG/Public/PCGActorAndComponentMapping.h
--- a/Engine/Plugins/Experimental/PCG/Source/PCG/Public/PCGActorAndComponentMapping.h
+++ b/Engine/Plugins/Experimental/PCG/Source/PCG/Public/PCGActorAndComponentMapping.h
@@ -5,7 +5,6 @@
 #include "PCGCommon.h"
 #include "Elements/PCGActorSelector.h"
 #include "Grid/PCGComponentOctree.h"
-#include "RuntimeGen/PCGRuntimeGenScheduler.h"
 
 #include "Containers/Array.h"
 #include "Containers/Map.h"
@@ -37,7 +36,6 @@ class FPCGActorAndComponentMapping
 {
 public:
 	friend UPCGSubsystem;
-	friend FPCGRuntimeGenScheduler;
 
 	~FPCGActorAndComponentMapping() = default;
 	
@@ -81,11 +79,8 @@ public:
 	/** Return a copy of all the registered components. Thread safe */
 	TSet<UPCGComponent*> GetAllRegisteredComponents() const;
 
-	/** Retrieves a local component using grid size and grid coordinates, returns nullptr if no such component is found. */
-	UPCGComponent* GetLocalComponent(uint32 GridSize, const FIntVector& CellCoords, const UPCGComponent* InOriginalComponent, bool bRuntimeGenerated = false);
-
-	/** Retrieves a partition actor using grid size and grid coordinates, returns nullptr if no such partition actor is found. */
-	APCGPartitionActor* GetPartitionActor(uint32 GridSize, const FIntVector& CellCoords, bool bRuntimeGenerated = false) const;
+	/** Retrieves a local component using grid size and grid coordinates, returns nullptr if no such component found. */
+	UPCGComponent* GetLocalComponent(uint32 GridSize, const FIntVector& CellCoords, const UPCGComponent* InOriginalComponent);
 
 private:
 	// This class is only meant to be used as part of the PCG Subsytem and owned by it.
@@ -179,18 +174,10 @@ private:
 	TMap<uint32, TMap<FIntVector, TObjectPtr<APCGPartitionActor>>> PartitionActorsMap;
 	mutable FRWLock PartitionActorsMapLock;
 
-	/** Mapping from grid size and grid coords to RuntimeGen partition actor. We can only have 1 RuntimeGen partition actor per grid cell. */
-	TMap<uint32, TMap<FIntVector, TObjectPtr<APCGPartitionActor>>> RuntimeGenPartitionActorsMap;
-	mutable FRWLock RuntimeGenPartitionActorsMapLock;
-
 	/** Mapping between original components and its overlapping partition actors. */
 	TMap<const UPCGComponent*, TSet<TObjectPtr<APCGPartitionActor>>> ComponentToPartitionActorsMap;
 	mutable FRWLock ComponentToPartitionActorsMapLock;
 
-	/** Mapping between original components and its overlapping RuntimeGen partition actors. */
-	TMap<const UPCGComponent*, TSet<TObjectPtr<APCGPartitionActor>>> ComponentToRuntimeGenPartitionActorsMap;
-	mutable FRWLock ComponentToRuntimeGenPartitionActorsMapLock;
-
 	/** Components to be unregister at the next frame. cf. UnregisterComponent for a better understanding on why it is needed. */
 	TSet<UPCGComponent*> DelayedComponentToUnregister;
 	mutable FCriticalSection DelayedComponentToUnregisterLock;
diff --git a/Engine/Plugins/Experimental/PCG/Source/PCG/Public/PCGCommon.h b/Engine/Plugins/Experimental/PCG/Source/PCG/Public/PCGCommon.h
--- a/Engine/Plugins/Experimental/PCG/Source/PCG/Public/PCGCommon.h
+++ b/Engine/Plugins/Experimental/PCG/Source/PCG/Public/PCGCommon.h
@@ -18,8 +18,7 @@ enum class EPCGChangeType : uint8
 	Input = 1 << 2,
 	Edge = 1 << 3,
 	Node = 1 << 4,
-	Structural = 1 << 5,
-	Debug = 1 << 6, // The Debug change type propagates when debug is toggled on a node.
+	Structural = 1 << 5
 };
 ENUM_CLASS_FLAGS(EPCGChangeType);
 
@@ -140,10 +139,7 @@ enum class EPCGHiGenGrid : uint32
 {
 	Uninitialized = 0 UMETA(Hidden),
 
-	// NOTE: When adding new grids, increment PCGHiGenGrid::NumGridValues below
-	Grid4 = 4 UMETA(DisplayName = "400"),
-	Grid8 = 8 UMETA(DisplayName = "800"),
-	Grid16 = 16 UMETA(DisplayName = "1600"),
+	// NOTE: When adding new grids, increment PCGHiGenGrid::NumGridValues below.
 	Grid32 = 32 UMETA(DisplayName = "3200"),
 	Grid64 = 64 UMETA(DisplayName = "6400"),
 	Grid128 = 128 UMETA(DisplayName = "12800"),
@@ -152,7 +148,7 @@ enum class EPCGHiGenGrid : uint32
 	Grid1024 = 1024 UMETA(DisplayName = "102400"),
 	Grid2048 = 2048 UMETA(DisplayName = "204800"),
 	
-	GridMin = Grid4 UMETA(Hidden),
+	GridMin = Grid32 UMETA(Hidden),
 	GridMax = Grid2048 UMETA(Hidden),
 
 	// Should execute once rather than executing on any grid
@@ -163,7 +159,7 @@ ENUM_CLASS_FLAGS(EPCGHiGenGrid);
 namespace PCGHiGenGrid
 {
 	// Number of unique values of EPCGHiGenGrid, const so it can be used for the inline allocator below.
-	constexpr uint32 NumGridValues = 13;
+	constexpr uint32 NumGridValues = 10;
 
 	// Alias for array which is allocated on the stack (we have a strong idea of the max required elements).
 	using FSizeArray = TArray<uint32, TInlineAllocator<PCGHiGenGrid::NumGridValues>>;
@@ -173,7 +169,6 @@ namespace PCGHiGenGrid
 
 	PCG_API bool IsValidGridSize(uint32 InGridSize);
 	PCG_API bool IsValidGrid(EPCGHiGenGrid InGrid);
-	PCG_API bool IsValidGridOrUninitialized(EPCGHiGenGrid InGrid);
 	PCG_API uint32 GridToGridSize(EPCGHiGenGrid InGrid);
 	PCG_API EPCGHiGenGrid GridSizeToGrid(uint32 InGridSize);
 
@@ -195,58 +190,4 @@ enum class EPCGEditorDirtyMode : uint8
 	Normal UMETA(Tooltip="Normal editing mode where generation changes (generation, cleanup) dirty the component and its resources."),
 	Preview UMETA(Tooltip="Editing mode where generation changes (generation, cleanup, resources) on the component will not trigger any dirty state, but will also not save any of the generated resources."),
 	LoadAsPreview UMETA(Tooltip="Acts as the normal editing mode until the next load of the component, at which state it acts as-if-transient, namely that any further generation changes will not dirty the component.")
-};
-
-USTRUCT(BlueprintType)
-struct FPCGRuntimeGenerationRadii
-{
-	GENERATED_BODY()
-
-public:
-	/** Get the runtime generation radius for the given grid size. */
-	double GetGenerationRadiusFromGrid(EPCGHiGenGrid Grid) const;
-
-	/** Compute the runtime cleanup radius for the given grid size. */
-	double GetCleanupRadiusFromGrid(EPCGHiGenGrid Grid) const;
-
-	static constexpr double DefaultGenerationRadiusMultiplier = 2.0;
-	static constexpr double DefaultCleanupRadiusMultiplier = 1.1;
-
-	/** The distance (in centimeters) at which the component will be considered for generation by the RuntimeGenerationScheduler. For partitioned components, this also acts as the unbounded generation radius. */
-	UPROPERTY(BlueprintReadWrite, EditAnywhere, Category = "Runtime Generation")
-	double GenerationRadius = PCGHiGenGrid::UnboundedGridSize() * DefaultGenerationRadiusMultiplier;
-
-	UPROPERTY(BlueprintReadWrite, EditAnywhere, Category = "Runtime Generation")
-	double GenerationRadius400 = PCGHiGenGrid::GridToGridSize(EPCGHiGenGrid::Grid4) * DefaultGenerationRadiusMultiplier;
-
-	UPROPERTY(BlueprintReadWrite, EditAnywhere, Category = "Runtime Generation")
-	double GenerationRadius800 = PCGHiGenGrid::GridToGridSize(EPCGHiGenGrid::Grid8) * DefaultGenerationRadiusMultiplier;
-
-	UPROPERTY(BlueprintReadWrite, EditAnywhere, Category = "Runtime Generation")
-	double GenerationRadius1600 = PCGHiGenGrid::GridToGridSize(EPCGHiGenGrid::Grid16) * DefaultGenerationRadiusMultiplier;
-
-	UPROPERTY(BlueprintReadWrite, EditAnywhere, Category = "Runtime Generation")
-	double GenerationRadius3200 = PCGHiGenGrid::GridToGridSize(EPCGHiGenGrid::Grid32) * DefaultGenerationRadiusMultiplier;
-
-	UPROPERTY(BlueprintReadWrite, EditAnywhere, Category = "Runtime Generation")
-	double GenerationRadius6400 = PCGHiGenGrid::GridToGridSize(EPCGHiGenGrid::Grid64) * DefaultGenerationRadiusMultiplier;
-
-	UPROPERTY(BlueprintReadWrite, EditAnywhere, Category = "Runtime Generation")
-	double GenerationRadius12800 = PCGHiGenGrid::GridToGridSize(EPCGHiGenGrid::Grid128) * DefaultGenerationRadiusMultiplier;
-
-	UPROPERTY(BlueprintReadWrite, EditAnywhere, Category = "Runtime Generation")
-	double GenerationRadius25600 = PCGHiGenGrid::GridToGridSize(EPCGHiGenGrid::Grid256) * DefaultGenerationRadiusMultiplier;
-
-	UPROPERTY(BlueprintReadWrite, EditAnywhere, Category = "Runtime Generation")
-	double GenerationRadius51200 = PCGHiGenGrid::GridToGridSize(EPCGHiGenGrid::Grid512) * DefaultGenerationRadiusMultiplier;
-
-	UPROPERTY(BlueprintReadWrite, EditAnywhere, Category = "Runtime Generation")
-	double GenerationRadius102400 = PCGHiGenGrid::GridToGridSize(EPCGHiGenGrid::Grid1024) * DefaultGenerationRadiusMultiplier;
-
-	UPROPERTY(BlueprintReadWrite, EditAnywhere, Category = "Runtime Generation")
-	double GenerationRadius204800 = PCGHiGenGrid::GridToGridSize(EPCGHiGenGrid::Grid2048) * DefaultGenerationRadiusMultiplier;
-
-	/** Multiplier on the GenerationRadius to control the distance at which runtime generated components will be cleaned up. Applied per grid size. */
-	UPROPERTY(BlueprintReadWrite, EditAnywhere, Category = "Runtime Generation", meta = (UIMin = "1.0", ClampMin = "1.0"))
-	double CleanupRadiusMultiplier = DefaultCleanupRadiusMultiplier;
-};
+};
\ No newline at end of file
diff --git a/Engine/Plugins/Experimental/PCG/Source/PCG/Public/PCGComponent.h b/Engine/Plugins/Experimental/PCG/Source/PCG/Public/PCGComponent.h
--- a/Engine/Plugins/Experimental/PCG/Source/PCG/Public/PCGComponent.h
+++ b/Engine/Plugins/Experimental/PCG/Source/PCG/Public/PCGComponent.h
@@ -20,19 +20,16 @@ class APCGPartitionActor;
 struct FPCGContext;
 class FPCGActorAndComponentMapping;
 class UPCGComponent;
-class UPCGData;
-class IPCGGenSourceBase;
 class UPCGGraph;
-class UPCGGraphInstance;
 class UPCGGraphInterface;
+class UPCGGraphInstance;
 class UPCGManagedResource;
-class UPCGSchedulingPolicyBase;
+class UPCGData;
 class UPCGSubsystem;
 class ALandscapeProxy;
 class FLandscapeProxyComponentDataChangedParams;
 class UClass;
 
-
 #if WITH_EDITOR
 	DECLARE_MULTICAST_DELEGATE_OneParam(FOnPCGGraphGenerated, UPCGComponent*);
 	DECLARE_MULTICAST_DELEGATE_OneParam(FOnPCGGraphCleaned, UPCGComponent*);
@@ -51,9 +48,8 @@ enum class EPCGComponentInput : uint8
 UENUM(Blueprintable)
 enum class EPCGComponentGenerationTrigger : uint8
 {
-	GenerateOnLoad    UMETA(ToolTip = "Generates only when the component is loaded into the level."),
-	GenerateOnDemand  UMETA(ToolTip = "Generates only when requested (e.g. via Blueprint)."),
-	GenerateAtRuntime UMETA(ToolTip = "Generates only when scheduled by the Runtime Generation Scheduler.")
+	GenerateOnLoad,
+	GenerateOnDemand
 };
 
 UENUM(meta = (Bitflags))
@@ -116,8 +112,6 @@ public:
 	/** If this is a local component returns self, otherwise returns the original component. */
 	UPCGComponent* GetOriginalComponent();
 
-	UPCGSchedulingPolicyBase* GetRuntimeGenSchedulingPolicy() const { return SchedulingPolicy; }
-
 	bool CanPartition() const;
 
 	UPCGGraph* GetGraph() const;
@@ -142,11 +136,7 @@ public:
 	UFUNCTION(BlueprintCallable, Category = PCG)
 	void GenerateLocal(bool bForce);
 
-	/** Requests the component to generate only on the specified grid level (all grid levels if EPCGHiGenGrid::Uninitialized). */
-	void GenerateLocal(EPCGComponentGenerationTrigger RequestedGenerationTrigger, bool bForce, EPCGHiGenGrid Grid = EPCGHiGenGrid::Uninitialized);
-
 	FPCGTaskId GenerateLocalGetTaskId(bool bForce);
-	FPCGTaskId GenerateLocalGetTaskId(EPCGComponentGenerationTrigger RequestedGenerationTrigger, bool bForce, EPCGHiGenGrid Grid = EPCGHiGenGrid::Uninitialized);
 
 	/** Cleans up the generation from a local (vs. remote) standpoint. Will not be replicated. Will be delayed. */
 	UFUNCTION(BlueprintCallable, Category = PCG)
@@ -191,15 +181,6 @@ public:
 	/** Clear any data stored for any pins. */
 	void ClearPerPinGeneratedOutput();
 
-	/** Set the runtime generation scheduling policy type. */
-	void SetSchedulingPolicyClass(TSubclassOf<UPCGSchedulingPolicyBase> InSchedulingPolicyClass);
-
-	/** Get the runtime generation radius for the given grid size. */
-	double GetGenerationRadiusFromGrid(EPCGHiGenGrid Grid) const;
-
-	/** Compute the runtime cleanup radius for the given grid size. */
-	double GetCleanupRadiusFromGrid(EPCGHiGenGrid Grid) const;
-
 	UPROPERTY(BlueprintReadWrite, EditAnywhere, Category = Settings)
 	EPCGComponentInput InputType = EPCGComponentInput::Actor;
 
@@ -219,23 +200,8 @@ public:
 	UPROPERTY(BlueprintReadOnly, EditAnywhere, Category = Properties, AdvancedDisplay, meta = (EditCondition = "!bIsComponentLocal", EditConditionHides))
 	EPCGComponentGenerationTrigger GenerationTrigger = EPCGComponentGenerationTrigger::GenerateOnLoad;
 
-	/** Manual overrides for the graph generation radii and cleanup radius multiplier. */
-	UPROPERTY(BlueprintReadWrite, EditAnywhere, Category = RuntimeGeneration, meta = (EditCondition = "GenerationTrigger == EPCGComponentGenerationTrigger::GenerateAtRuntime", EditConditionHides))
-	bool bOverrideGenerationRadii = false;
-
-	UPROPERTY(BlueprintReadWrite, EditAnywhere, Category = RuntimeGeneration, meta = (EditCondition = "GenerationTrigger == EPCGComponentGenerationTrigger::GenerateAtRuntime && bOverrideGenerationRadii", EditConditionHides))
-	FPCGRuntimeGenerationRadii GenerationRadii;
-
-	/** A Scheduling Policy dictates the order in which instances of this component will be scheduled. */
-	UPROPERTY(BlueprintReadOnly, EditAnywhere, NoClear, Category = RuntimeGeneration, meta = (EditCondition = "GenerationTrigger == EPCGComponentGenerationTrigger::GenerateAtRuntime", EditConditionHides))
-	TSubclassOf<UPCGSchedulingPolicyBase> SchedulingPolicyClass = nullptr;
-
-	/** This is the instanced UPCGSchedulingPolicy object which holds scheduling parameters and calculates priorities. */
-	UPROPERTY(BlueprintReadOnly, VisibleAnywhere, Instanced, Category = RuntimeGeneration, meta = (EditCondition = "GenerationTrigger == EPCGComponentGenerationTrigger::GenerateAtRuntime", EditConditionHides))
-	TObjectPtr<UPCGSchedulingPolicyBase> SchedulingPolicy;
-
 	/** Flag to indicate whether this component has run in the editor. Note that for partitionable actors, this will always be false. */
-	UPROPERTY(BlueprintReadOnly, VisibleAnywhere, AdvancedDisplay, Category = Properties, NonTransactional, meta = (NoResetToDefault))
+	UPROPERTY(BlueprintReadOnly, VisibleAnywhere, AdvancedDisplay, Category = Properties, meta = (NoResetToDefault))
 	bool bGenerated = false;
 
 	UPROPERTY(NonPIEDuplicateTransient)
@@ -342,9 +308,6 @@ public:
 	/** Builds the canonical PCG data from a given actor and its PCG component if any. */
 	static UPCGData* CreateActorPCGData(AActor* Actor, const UPCGComponent* Component, bool bParseActor = true);
 
-protected:
-	void RefreshSchedulingPolicy();
-
 protected:
 	UPROPERTY(BlueprintReadOnly, VisibleAnywhere, Category = PCG, Instanced, meta = (NoResetToDefault))
 	TObjectPtr<UPCGGraphInstance> GraphInstance;
@@ -380,7 +343,7 @@ private:
 	bool ShouldGenerate(bool bForce, EPCGComponentGenerationTrigger RequestedGenerationTrigger) const;
 
 	/* Internal call that allows to delay a Generate/Cleanup call, chain with dependencies and keep track of the task id created. This task id is also returned. */
-	FPCGTaskId GenerateInternal(bool bForce, EPCGHiGenGrid Grid, EPCGComponentGenerationTrigger RequestedGenerationTrigger, const TArray<FPCGTaskId>& Dependencies);
+	FPCGTaskId GenerateInternal(bool bForce, EPCGComponentGenerationTrigger RequestedGenerationTrigger, const TArray<FPCGTaskId>& Dependencies);
 	FPCGTaskId CleanupInternal(bool bRemoveComponents, bool bSave, const TArray<FPCGTaskId>& Dependencies);
 
 	/* Internal call to create tasks to generate the component. If there is nothing to do, an invalid task id will be returned. Should only be used by the subsystem. */
diff --git a/Engine/Plugins/Experimental/PCG/Source/PCG/Public/PCGGraph.h b/Engine/Plugins/Experimental/PCG/Source/PCG/Public/PCGGraph.h
--- a/Engine/Plugins/Experimental/PCG/Source/PCG/Public/PCGGraph.h
+++ b/Engine/Plugins/Experimental/PCG/Source/PCG/Public/PCGGraph.h
@@ -384,9 +384,6 @@ public:
 	// Best used on a brand new PCG Graph.
 	void AddUserParameters(const TArray<FPropertyBagPropertyDesc>& InDescs, const UPCGGraph* InOptionalOriginalGraph = nullptr);
 
-	UPROPERTY(BlueprintReadWrite, EditAnywhere, Category = "Runtime Generation")
-	FPCGRuntimeGenerationRadii GenerationRadii;
-
 protected:
 	virtual void OnGraphParametersChanged(EPCGGraphParameterEvent InChangeType, FName InChangedPropertyName) override;
 
diff --git a/Engine/Plugins/Experimental/PCG/Source/PCG/Public/PCGSubsystem.h b/Engine/Plugins/Experimental/PCG/Source/PCG/Public/PCGSubsystem.h
--- a/Engine/Plugins/Experimental/PCG/Source/PCG/Public/PCGSubsystem.h
+++ b/Engine/Plugins/Experimental/PCG/Source/PCG/Public/PCGSubsystem.h
@@ -11,11 +11,11 @@
 
 #include "PCGSubsystem.generated.h"
 
+
 class APCGPartitionActor;
 class APCGWorldActor;
 class UPCGGraph;
 class UPCGLandscapeCache;
-class FPCGRuntimeGenScheduler;
 
 enum class EPCGComponentDirtyFlag : uint8;
 enum class ETickableTickType : uint8;
@@ -84,12 +84,10 @@ public:
 	void RegisterPCGWorldActor(APCGWorldActor* InActor);
 	void UnregisterPCGWorldActor(APCGWorldActor* InActor);
 
-	void OnOriginalComponentUnregistered(UPCGComponent* InComponent);
-	
 	UPCGLandscapeCache* GetLandscapeCache();
 
 	// Schedule graph (owner -> graph)
-	FPCGTaskId ScheduleComponent(UPCGComponent* PCGComponent, EPCGHiGenGrid Grid, bool bSave, const TArray<FPCGTaskId>& InDependencies);
+	FPCGTaskId ScheduleComponent(UPCGComponent* PCGComponent, bool bSave, const TArray<FPCGTaskId>& InDependencies);
 
 	/** Schedule cleanup(owner->graph). Note that in non-partitioned mode, cleanup is immediate. */
 	FPCGTaskId ScheduleCleanup(UPCGComponent* PCGComponent, bool bRemoveComponents, bool bSave, const TArray<FPCGTaskId>& Dependencies);
@@ -157,13 +155,11 @@ public:
 	void ForAllOverlappingComponentsInHierarchy(UPCGComponent* InComponent, const TFunction<void(UPCGComponent*)>& InFunc) const;
 
 	/** Retrieves a local component using grid size and grid coordinates, returns nullptr if no such component found. */
-	UPCGComponent* GetLocalComponent(uint32 GridSize, const FIntVector& CellCoords, const UPCGComponent* InOriginalComponent, bool bTransient = false);
+	UPCGComponent* GetLocalComponent(uint32 GridSize, const FIntVector& CellCoords, const UPCGComponent* InOriginalComponent);
 
 	/** True if graph cache debugging is enabled. */
 	bool IsGraphCacheDebuggingEnabled() const;
 
-	FPCGGenSourceManager* GetGenSourceManager() const;
-
 #if WITH_EDITOR
 public:
 	/** Schedule refresh on the current or next frame */
@@ -230,7 +226,6 @@ private:
 private:
 	APCGWorldActor* PCGWorldActor = nullptr;
 	FPCGGraphExecutor* GraphExecutor = nullptr;
-	FPCGRuntimeGenScheduler* RuntimeGenScheduler = nullptr;
 	bool bHasTickedOnce = false;
 	FPCGActorAndComponentMapping ActorAndComponentMapping;
 
diff --git a/Engine/Plugins/Experimental/PCG/Source/PCG/Public/PCGWorldActor.h b/Engine/Plugins/Experimental/PCG/Source/PCG/Public/PCGWorldActor.h
--- a/Engine/Plugins/Experimental/PCG/Source/PCG/Public/PCGWorldActor.h
+++ b/Engine/Plugins/Experimental/PCG/Source/PCG/Public/PCGWorldActor.h
@@ -64,12 +64,6 @@ public:
 	UPROPERTY(config, EditAnywhere, Category = GenerationSettings)
 	bool bUse2DGrid = true;
 
-#if WITH_EDITORONLY_DATA
-	/** Allows any currently active editor viewport to act as a Runtime Generation Source. */
-	UPROPERTY(EditAnywhere, Category = RuntimeGeneration)
-	bool bTreatEditorViewportAsGenerationSource = false;
-#endif
-
 private:
 	void RegisterToSubsystem();
 	void UnregisterFromSubsystem();
diff --git a/Engine/Plugins/Experimental/PCG/Source/PCG/Public/RuntimeGen/GenSources/PCGGenSourceBase.h b/Engine/Plugins/Experimental/PCG/Source/PCG/Public/RuntimeGen/GenSources/PCGGenSourceBase.h
deleted file mode 100644
diff --git a/Engine/Plugins/Experimental/PCG/Source/PCG/Public/RuntimeGen/GenSources/PCGGenSourceComponent.h b/Engine/Plugins/Experimental/PCG/Source/PCG/Public/RuntimeGen/GenSources/PCGGenSourceComponent.h
deleted file mode 100644
diff --git a/Engine/Plugins/Experimental/PCG/Source/PCG/Public/RuntimeGen/GenSources/PCGGenSourceEditorCamera.h b/Engine/Plugins/Experimental/PCG/Source/PCG/Public/RuntimeGen/GenSources/PCGGenSourceEditorCamera.h
deleted file mode 100644
diff --git a/Engine/Plugins/Experimental/PCG/Source/PCG/Public/RuntimeGen/GenSources/PCGGenSourcePlayer.h b/Engine/Plugins/Experimental/PCG/Source/PCG/Public/RuntimeGen/GenSources/PCGGenSourcePlayer.h
deleted file mode 100644
diff --git a/Engine/Plugins/Experimental/PCG/Source/PCG/Public/RuntimeGen/GenSources/PCGGenSourceWPStreamingSource.h b/Engine/Plugins/Experimental/PCG/Source/PCG/Public/RuntimeGen/GenSources/PCGGenSourceWPStreamingSource.h
deleted file mode 100644
diff --git a/Engine/Plugins/Experimental/PCG/Source/PCG/Public/RuntimeGen/PCGGenSourceManager.h b/Engine/Plugins/Experimental/PCG/Source/PCG/Public/RuntimeGen/PCGGenSourceManager.h
deleted file mode 100644
diff --git a/Engine/Plugins/Experimental/PCG/Source/PCG/Public/RuntimeGen/PCGRuntimeGenScheduler.h b/Engine/Plugins/Experimental/PCG/Source/PCG/Public/RuntimeGen/PCGRuntimeGenScheduler.h
deleted file mode 100644
diff --git a/Engine/Plugins/Experimental/PCG/Source/PCG/Public/RuntimeGen/SchedulingPolicies/PCGSchedulingPolicyBase.h b/Engine/Plugins/Experimental/PCG/Source/PCG/Public/RuntimeGen/SchedulingPolicies/PCGSchedulingPolicyBase.h
deleted file mode 100644
diff --git a/Engine/Plugins/Experimental/PCG/Source/PCG/Public/RuntimeGen/SchedulingPolicies/PCGSchedulingPolicyDistanceAndDirection.h b/Engine/Plugins/Experimental/PCG/Source/PCG/Public/RuntimeGen/SchedulingPolicies/PCGSchedulingPolicyDistanceAndDirection.h
deleted file mode 100644
diff --git a/Engine/Plugins/Experimental/PCG/Source/PCGEditor/Private/Details/PCGComponentDetails.cpp b/Engine/Plugins/Experimental/PCG/Source/PCGEditor/Private/Details/PCGComponentDetails.cpp
--- a/Engine/Plugins/Experimental/PCG/Source/PCGEditor/Private/Details/PCGComponentDetails.cpp
+++ b/Engine/Plugins/Experimental/PCG/Source/PCGEditor/Private/Details/PCGComponentDetails.cpp
@@ -95,10 +95,9 @@ void FPCGComponentDetails::CustomizeDetails(IDetailLayoutBuilder& DetailBuilder)
 			.Padding(4.0f)
 			[
 				SNew(SButton)
-				.OnClicked(this, &FPCGComponentDetails::OnCleanupClicked)
-				.Visibility(this, &FPCGComponentDetails::CleanupButtonVisible)
+				.OnClicked(this, &FPCGComponentDetails::OnCleanupClicked)				
 				[
-					SNew(STextBlock)
+					SNew(STextBlock)					
 					.Font(IDetailLayoutBuilder::GetDetailFont())
 					.Text(LOCTEXT("CleanupButton", "Cleanup"))
 				]
@@ -145,8 +144,7 @@ EVisibility FPCGComponentDetails::GenerateButtonVisible() const
 {
 	for (const TWeakObjectPtr<UPCGComponent>& Component : SelectedComponents)
 	{
-		// If component is runtime generated then generate/cleanup is managed by the scheduler.
-		if (Component.IsValid() && !Component->IsGenerating() && Component->GenerationTrigger != EPCGComponentGenerationTrigger::GenerateAtRuntime)
+		if (Component.IsValid() && !Component->IsGenerating())
 		{
 			return EVisibility::Visible;
 		}
@@ -159,22 +157,7 @@ EVisibility FPCGComponentDetails::CancelButtonVisible() const
 {
 	for (const TWeakObjectPtr<UPCGComponent>& Component : SelectedComponents)
 	{
-		// If component is runtime generated then generate/cleanup is managed by the scheduler.
-		if (Component.IsValid() && Component->IsGenerating() && Component->GenerationTrigger != EPCGComponentGenerationTrigger::GenerateAtRuntime)
-		{
-			return EVisibility::Visible;
-		}
-	}
-
-	return EVisibility::Collapsed;
-}
-
-EVisibility FPCGComponentDetails::CleanupButtonVisible() const
-{
-	for (const TWeakObjectPtr<UPCGComponent>& Component : SelectedComponents)
-	{
-		// If component is runtime generated then generate/cleanup is managed by the scheduler.
-		if (Component.IsValid() && Component->GenerationTrigger != EPCGComponentGenerationTrigger::GenerateAtRuntime)
+		if (Component.IsValid() && Component->IsGenerating())
 		{
 			return EVisibility::Visible;
 		}
diff --git a/Engine/Plugins/Experimental/PCG/Source/PCGEditor/Private/PCGEditor.cpp b/Engine/Plugins/Experimental/PCG/Source/PCGEditor/Private/PCGEditor.cpp
--- a/Engine/Plugins/Experimental/PCG/Source/PCGEditor/Private/PCGEditor.cpp
+++ b/Engine/Plugins/Experimental/PCG/Source/PCGEditor/Private/PCGEditor.cpp
@@ -273,19 +273,12 @@ void FPCGEditor::UpdateDebugAfterComponentSelection(UPCGComponent* InOldComponen
 		return;
 	}
 
-	// GenerateOnDemand requests will bounce off if the component is GenerateAtRuntime, so we need to select the correct GenTrigger for the component to be updated.
-	EPCGComponentGenerationTrigger GenerationTrigger = EPCGComponentGenerationTrigger::GenerateOnDemand;
-	if (InNewComponent && InNewComponent->GenerationTrigger == EPCGComponentGenerationTrigger::GenerateAtRuntime)
-	{
-		GenerationTrigger = EPCGComponentGenerationTrigger::GenerateAtRuntime;
-	}
-
 	// If individual component debugging is disabled, just generate the new component if required.
 	if (!PCGGraphBeingEdited->DebugFlagAppliesToIndividualComponents())
 	{
 		if (InNewComponent && bInNewComponentStartedInspecting)
 		{
-			InNewComponent->GenerateLocal(GenerationTrigger, /*bForce=*/true);
+			InNewComponent->GenerateLocal(/*bForce=*/true);
 		}
 
 		return;
@@ -298,7 +291,7 @@ void FPCGEditor::UpdateDebugAfterComponentSelection(UPCGComponent* InOldComponen
 		{
 			// Transition from 'null' to 'any component not already inspecting' - generate to create debug/inspection info.
 			// If we have null selected, all components are displaying debug. Go to Original component so that all refresh.
-			InNewComponent->GetOriginalComponent()->GenerateLocal(GenerationTrigger, /*bForce=*/true);
+			InNewComponent->GetOriginalComponent()->GenerateLocal(/*bForce=*/true);
 		}
 	}
 	else
@@ -311,19 +304,15 @@ void FPCGEditor::UpdateDebugAfterComponentSelection(UPCGComponent* InOldComponen
 		// Regenerate to clear debug info if switching components, or if changing from a component to null.
 		if (InNewComponent || bDebugFlagSetOnAnyNode)
 		{
-			// If InNewComponent is nullptr in local instance debug mode, then RuntimeGen does not want to regenerate (this is to avoid regenerating all PAs).
-			if (InNewComponent || InOldComponent->GenerationTrigger != EPCGComponentGenerationTrigger::GenerateAtRuntime)
-			{
-				// Use original component - debug can be displayed both by the local component and parent local components.
-				InOldComponent->GetOriginalComponent()->GenerateLocal(GenerationTrigger, /*bForce=*/true);
-			}
+			// Use original component - debug can be displayed both by the local component and parent local components.
+			InOldComponent->GetOriginalComponent()->GenerateLocal(/*bForce=*/true);
 		}
 
 		// Debug new component if it wasn't already
 		if (InNewComponent && bInNewComponentStartedInspecting)
 		{
 			// Use original component - debug can be displayed both by the local component and parent local components.
-			InNewComponent->GetOriginalComponent()->GenerateLocal(GenerationTrigger, /*bForce=*/true);
+			InNewComponent->GetOriginalComponent()->GenerateLocal(/*bForce=*/true);
 		}
 	}
 }
@@ -1338,7 +1327,7 @@ void FPCGEditor::OnToggleDebug()
 			if (PCGSettingsInterface->bDebug != bNewCheckState)
 			{
 				PCGSettingsInterface->bDebug = bNewCheckState;
-				PCGNode->OnNodeChangedDelegate.Broadcast(PCGNode, EPCGChangeType::Debug);
+				PCGNode->OnNodeChangedDelegate.Broadcast(PCGNode, EPCGChangeType::Settings);
 			}
 		}
 	}
diff --git a/Engine/Plugins/Experimental/PCG/Source/PCGEditor/Public/Details/PCGComponentDetails.h b/Engine/Plugins/Experimental/PCG/Source/PCGEditor/Public/Details/PCGComponentDetails.h
--- a/Engine/Plugins/Experimental/PCG/Source/PCGEditor/Public/Details/PCGComponentDetails.h
+++ b/Engine/Plugins/Experimental/PCG/Source/PCGEditor/Public/Details/PCGComponentDetails.h
@@ -27,7 +27,6 @@ protected:
 private:
 	EVisibility GenerateButtonVisible() const;
 	EVisibility CancelButtonVisible() const;
-	EVisibility CleanupButtonVisible() const;
 	FReply OnGenerateClicked();
 	FReply OnCancelClicked();
 	FReply OnCleanupClicked();
