[PCG] Remove UPCGSubsystem::ProcessGraph which is no longer used.

#rnx

diff --git a/Engine/Plugins/Experimental/PCG/Source/PCG/Private/PCGSubsystem.cpp b/Engine/Plugins/Experimental/PCG/Source/PCG/Private/PCGSubsystem.cpp
--- a/Engine/Plugins/Experimental/PCG/Source/PCG/Private/PCGSubsystem.cpp
+++ b/Engine/Plugins/Experimental/PCG/Source/PCG/Private/PCGSubsystem.cpp
@@ -892,127 +892,6 @@ void UPCGSubsystem::CreatePartitionActorsWithinBounds(const FBox& InBounds, cons
 	}
 }
 
-FPCGTaskId UPCGSubsystem::ProcessGraph(UPCGComponent* Component, const FBox& InPreviousBounds, const FBox& InNewBounds, EOperation InOperation, bool bSave)
-{
-	check(Component);
-	TWeakObjectPtr<UPCGComponent> ComponentPtr(Component);
-
-	// TODO: optimal implementation would find the difference between the previous bounds and the new bounds
-	// and process these only. This is esp. important because of the CreateActor parameter.
-	auto ScheduleTask = [this, ComponentPtr, InOperation, bSave](APCGPartitionActor* PCGActor, const FBox& InBounds, const TArray<FPCGTaskId>& TaskDependencies){
-		TWeakObjectPtr<APCGPartitionActor> PCGActorPtr(PCGActor);
-
-		auto UnpartitionTask = [ComponentPtr, PCGActorPtr]() {
-			// TODO: PCG actors that become empty could be deleted, but we also need to keep track
-			// of packages that would need to be deleted from SCC.
-			if (APCGPartitionActor* PCGActor = PCGActorPtr.Get())
-			{
-				if (UPCGComponent* Component = ComponentPtr.Get())
-				{
-					PCGActor->RemoveGraphInstance(Component);
-				}
-				else
-				{
-#if WITH_EDITOR
-					PCGActor->CleanupDeadGraphInstances();
-#endif
-				}
-			}
-			return true;
-		};
-
-		auto PartitionTask = [ComponentPtr, PCGActorPtr]() {
-			if (APCGPartitionActor* PCGActor = PCGActorPtr.Get())
-			{
-				if (UPCGComponent* Component = ComponentPtr.Get())
-				{
-					PCGActor->AddGraphInstance(Component);
-				}
-			}
-
-			return true;
-		};
-
-		auto ScheduleGraph = [this, PCGActorPtr, ComponentPtr, bSave, TaskDependencies]() {
-			if (APCGPartitionActor* PCGActor = PCGActorPtr.Get())
-			{
-				if (UPCGComponent* Component = ComponentPtr.Get())
-				{
-					// Ensure that the PCG actor has a matching local component.
-					// This is done immediately, but technically we could add it as a task
-					PCGActor->AddGraphInstance(Component);
-
-					UPCGComponent* LocalComponent = PCGActor->GetLocalComponent(Component);
-
-					if (!LocalComponent)
-					{
-						return InvalidPCGTaskId;
-					}
-
-					return LocalComponent->GenerateInternal(bSave, EPCGComponentGenerationTrigger::GenerateOnDemand, TaskDependencies);
-				}
-				else
-				{
-					UE_LOG(LogPCG, Error, TEXT("[ProcessGraph] PCG Component on PCG Actor is null"));
-					return InvalidPCGTaskId;
-				}
-			}
-			
-			return InvalidPCGTaskId;
-		};
-
-		switch (InOperation)
-		{
-		case EOperation::Unpartition:
-			return GraphExecutor->ScheduleGeneric(UnpartitionTask, ComponentPtr.Get(), TaskDependencies);
-		case EOperation::Partition:
-			return GraphExecutor->ScheduleGeneric(PartitionTask, ComponentPtr.Get(), TaskDependencies);
-		case EOperation::Generate:
-		default:
-			return ScheduleGraph();
-		}
-	};
-
-	FBox UnionBounds = InPreviousBounds + InNewBounds;
-	const bool bGenerate = InOperation == EOperation::Generate;
-	const bool bCreateActors = InOperation != EOperation::Unpartition;
-	const bool bLoadCell = (bGenerate && bSave); // TODO: review this
-	const bool bSaveActors = (bSave);
-
-	FPCGTaskId ProcessAllCellsTaskId = InvalidPCGTaskId;
-
-	if (UnionBounds.IsValid)
-	{
-		bool bHasUnbounded = false;
-		PCGHiGenGrid::FSizeArray GridSizes;
-		ensure(PCGHelpers::GetGenerationGridSizes(Component->GetGraph(), GetPCGWorldActor(), GridSizes, bHasUnbounded));
-		ProcessAllCellsTaskId = PCGSubsystem::ForEachIntersectingCell(GraphExecutor, Component->GetWorld(), UnionBounds, GridSizes, bCreateActors, bLoadCell, bSaveActors, ScheduleTask);
-	}
-
-	// Finally, call PostProcessGraph if something happened
-	if (ProcessAllCellsTaskId != InvalidPCGTaskId)
-	{
-		auto PostProcessGraphTask = [ComponentPtr, InNewBounds, bGenerate](FPCGContext* Context) {
-			if (UPCGComponent* Component = ComponentPtr.Get())
-			{
-				Component->PostProcessGraph(InNewBounds, bGenerate, Context);
-			}
-			return true;
-		};
-
-		return GraphExecutor->ScheduleGenericWithContext(PostProcessGraphTask, Component, { ProcessAllCellsTaskId });
-	}
-	else
-	{
-		if (Component)
-		{
-			Component->OnProcessGraphAborted();
-		}
-
-		return InvalidPCGTaskId;
-	}
-}
-
 FPCGTaskId UPCGSubsystem::ScheduleRefresh(UPCGComponent* Component, bool bForceRegen)
 {
 	TWeakObjectPtr<UPCGComponent> ComponentPtr(Component);
diff --git a/Engine/Plugins/Experimental/PCG/Source/PCG/Public/PCGSubsystem.h b/Engine/Plugins/Experimental/PCG/Source/PCG/Public/PCGSubsystem.h
--- a/Engine/Plugins/Experimental/PCG/Source/PCG/Public/PCGSubsystem.h
+++ b/Engine/Plugins/Experimental/PCG/Source/PCG/Public/PCGSubsystem.h
@@ -217,7 +217,6 @@ private:
 		Generate
 	};
 
-	FPCGTaskId ProcessGraph(UPCGComponent* Component, const FBox& InPreviousBounds, const FBox& InNewBounds, EOperation InOperation, bool bSave);
 	void CreatePartitionActorsWithinBounds(const FBox& InBounds, const PCGHiGenGrid::FSizeArray& InGridSizes);
 
 	FPCGNodeVisualLogs NodeVisualLogs;
