[PCG] Demote reroute nodes to be graph-only connection helpers

* Reroute tasks culled in compile. This required adding support to rewire to bypass the culled task.
* Disabled inspect/debug/disable functionality

#rnx

diff --git a/Engine/Plugins/Experimental/PCG/Source/PCG/Private/Elements/PCGReroute.cpp b/Engine/Plugins/Experimental/PCG/Source/PCG/Private/Elements/PCGReroute.cpp
--- a/Engine/Plugins/Experimental/PCG/Source/PCG/Private/Elements/PCGReroute.cpp
+++ b/Engine/Plugins/Experimental/PCG/Source/PCG/Private/Elements/PCGReroute.cpp
@@ -11,6 +11,12 @@ UPCGRerouteSettings::UPCGRerouteSettings()
 {
 #if WITH_EDITORONLY_DATA
 	bExposeToLibrary = false;
+
+	// Reroutes don't support disabling or debugging.
+	bDisplayDebuggingProperties = false;
+	DebugSettings.bDisplayProperties = false;
+	bEnabled = true;
+	bDebug = false;
 #endif
 }
 
@@ -64,6 +70,13 @@ FPCGElementPtr UPCGRerouteSettings::CreateElement() const
 
 bool FPCGRerouteElement::ExecuteInternal(FPCGContext* Context) const
 {
+	ensureMsgf(false, TEXT("Reroute elements are not supposed to execute - reroutes should be culled during graph compilation."));
+	
 	Context->OutputData = Context->InputData;
+	for (FPCGTaggedData& Output : Context->OutputData.TaggedData)
+	{
+		Output.Pin = PCGPinConstants::DefaultOutputLabel;
+	}
+
 	return true;
 }
diff --git a/Engine/Plugins/Experimental/PCG/Source/PCG/Private/Graph/PCGGraphCompiler.cpp b/Engine/Plugins/Experimental/PCG/Source/PCG/Private/Graph/PCGGraphCompiler.cpp
--- a/Engine/Plugins/Experimental/PCG/Source/PCG/Private/Graph/PCGGraphCompiler.cpp
+++ b/Engine/Plugins/Experimental/PCG/Source/PCG/Private/Graph/PCGGraphCompiler.cpp
@@ -8,6 +8,7 @@
 #include "PCGPin.h"
 #include "PCGSubgraph.h"
 #include "Elements/PCGHiGenGridSize.h"
+#include "Elements/PCGReroute.h"
 #include "Graph/PCGGraphExecutor.h"
 
 #include "Misc/ScopeRWLock.h"
@@ -452,7 +453,7 @@ EPCGHiGenGrid FPCGGraphCompiler::CalculateGridRecursive(
 	return Grid;
 }
 
-void FPCGGraphCompiler::CullTasks(TArray<FPCGGraphTask>& InOutCompiledTasks, TFunctionRef<bool(const FPCGGraphTask&)> CullTask)
+void FPCGGraphCompiler::CullTasks(TArray<FPCGGraphTask>& InOutCompiledTasks, bool bAddPassthroughWires, TFunctionRef<bool(const FPCGGraphTask&)> CullTask)
 {
 	TRACE_CPUPROFILER_EVENT_SCOPE(FPCGGraphCompiler::CullTasks);
 
@@ -465,14 +466,57 @@ void FPCGGraphCompiler::CullTasks(TArray<FPCGGraphTask>& InOutCompiledTasks, TFu
 	TArray<int32> TaskRemapping;
 	TaskRemapping.SetNumUninitialized(InOutCompiledTasks.Num());
 
-	// Never cull first (Input) task.
-	int32 WriteIndex = 1;
-	int32 ReadIndex = 1;
+	// Mark culled tasks by remapping to INDEX_NONE. First task is input task and is never culled.
 	TaskRemapping[0] = 0;
+	for (int32 TaskIndex = 1; TaskIndex < InOutCompiledTasks.Num(); ++TaskIndex)
+	{
+		TaskRemapping[TaskIndex] = CullTask(InOutCompiledTasks[TaskIndex]) ? INDEX_NONE : 0;
+	}
+
+	// Optionally add wires that bypass culled nodes.
+	if (bAddPassthroughWires)
+	{
+		for (int32 TaskIndex = 1; TaskIndex < InOutCompiledTasks.Num(); ++TaskIndex)
+		{
+			FPCGGraphTask& Task = InOutCompiledTasks[TaskIndex];
+			if (!Task.Node || Task.Node->GetInputPins().IsEmpty())
+			{
+				continue;
+			}
+
+			const int32 InputNumBefore = Task.Inputs.Num();
+			for (int32 InputIndex = 0; InputIndex < InputNumBefore; ++InputIndex)
+			{
+				const int32 InputTaskId = Task.Inputs[InputIndex].TaskId;
+
+				// Is node culled?
+				if (TaskRemapping[InputTaskId] == INDEX_NONE)
+				{
+					const FPCGGraphTask& InputTask = InOutCompiledTasks[InputTaskId];
+					if (InputTask.Node && InputTask.Node->GetInputPins().Num() > 1)
+					{
+						ensureMsgf(false, TEXT("Task culling currently only supports nodes with a single input pin which are trivial to unwire."));
+						continue;
+					}
 
+					// Upstream node was culled. Wire up the inputs of the culled node to this node.
+					for (int32 InputInputIndex = 0; InputInputIndex < InputTask.Inputs.Num(); ++InputInputIndex)
+					{
+						FPCGGraphTaskInput& NewInput = Task.Inputs.Add_GetRef(InputTask.Inputs[InputInputIndex]);
+						NewInput.OutPin = Task.Inputs[InputIndex].OutPin;
+					}
+				}
+			}
+		}
+	}
+
+	// Remove all culled tasks by compacting the task array. Never cull first (Input) task.
+	int32 WriteIndex = 1;
+	int32 ReadIndex = 1;
 	while (ReadIndex < InOutCompiledTasks.Num())
 	{
-		if (!CullTask(InOutCompiledTasks[ReadIndex]))
+		// If not culled, then move the task to it's final remapped position in the task array.
+		if (TaskRemapping[ReadIndex] != INDEX_NONE)
 		{
 			if (WriteIndex != ReadIndex)
 			{
@@ -483,17 +527,13 @@ void FPCGGraphCompiler::CullTasks(TArray<FPCGGraphTask>& InOutCompiledTasks, TFu
 			TaskRemapping[ReadIndex] = WriteIndex;
 			++WriteIndex;
 		}
-		else
-		{
-			// Flag as culled.
-			TaskRemapping[ReadIndex] = INDEX_NONE;
-		}
 
 		++ReadIndex;
 	}
 
 	InOutCompiledTasks.SetNum(WriteIndex);
 
+	// Remap input task IDs, and remove edges that connect to culled nodes.
 	for (FPCGGraphTask& Task : InOutCompiledTasks)
 	{
 		for (int32 InputIndex = Task.Inputs.Num() - 1; InputIndex >= 0; --InputIndex)
@@ -592,6 +632,9 @@ void FPCGGraphCompiler::CompileTopGraph(UPCGGraph* InGraph, uint32 GenerationGri
 		return;
 	}
 
+	// Remove reroute nodes before execution grid setup, as grid linkages need final nodes to connect from/to.
+	CullTasks(CompiledTasks, /*bAddPassthroughWires=*/true, [](const FPCGGraphTask& InTask) { return InTask.Node && Cast<UPCGRerouteSettings>(InTask.Node->GetSettings()); });
+
 	// For hierarchical generation resolve the execution grid for each task and cull any tasks that won't execute.
 	if (InGraph->IsHierarchicalGenerationEnabled() && GenerationGridSize != PCGHiGenGrid::UninitializedGridSize())
 	{
@@ -607,7 +650,7 @@ void FPCGGraphCompiler::CompileTopGraph(UPCGGraph* InGraph, uint32 GenerationGri
 		CreateGridLinkages(GenerationGrid, TaskGenerationGrid, CompiledTasks, StackContext);
 
 		// Cull any task that should not execute on the current grid.
-		CullTasks(CompiledTasks, [GenerationGrid, &TaskGenerationGrid](const FPCGGraphTask& InTask)
+		CullTasks(CompiledTasks, /*bAddPassthroughWires=*/false, [GenerationGrid, &TaskGenerationGrid](const FPCGGraphTask& InTask)
 		{
 			const EPCGHiGenGrid TaskGrid = TaskGenerationGrid[InTask.NodeId];
 			return TaskGrid != EPCGHiGenGrid::Uninitialized && !(TaskGrid & GenerationGrid);
diff --git a/Engine/Plugins/Experimental/PCG/Source/PCG/Private/Graph/PCGGraphCompiler.h b/Engine/Plugins/Experimental/PCG/Source/PCG/Private/Graph/PCGGraphCompiler.h
--- a/Engine/Plugins/Experimental/PCG/Source/PCG/Private/Graph/PCGGraphCompiler.h
+++ b/Engine/Plugins/Experimental/PCG/Source/PCG/Private/Graph/PCGGraphCompiler.h
@@ -56,7 +56,7 @@ private:
 		const FPCGStackContext& InStackContext);
 
 	/** Culls tasks based on a given lambda. Never culls the first (input) task in the array. */
-	static void CullTasks(TArray<FPCGGraphTask>& InOutCompiledTasks, TFunctionRef<bool(const FPCGGraphTask&)> CullTask);
+	static void CullTasks(TArray<FPCGGraphTask>& InOutCompiledTasks, bool bAddPassthroughWires, TFunctionRef<bool(const FPCGGraphTask&)> CullTask);
 
 	mutable FRWLock GraphToTaskMapLock;
 	TMap<UPCGGraph*, TArray<FPCGGraphTask>> GraphToTaskMap;
diff --git a/Engine/Plugins/Experimental/PCG/Source/PCG/Public/Elements/PCGReroute.h b/Engine/Plugins/Experimental/PCG/Source/PCG/Public/Elements/PCGReroute.h
--- a/Engine/Plugins/Experimental/PCG/Source/PCG/Public/Elements/PCGReroute.h
+++ b/Engine/Plugins/Experimental/PCG/Source/PCG/Public/Elements/PCGReroute.h
@@ -16,26 +16,30 @@ class PCG_API UPCGRerouteSettings : public UPCGSettings
 
 public:
 	UPCGRerouteSettings();
-	
+
+	//~Begin UPCGSettingsInterface interface
+	virtual bool CanBeDisabled() const override { return false; }
+	virtual bool CanBeDebugged() const override { return false; }
+	//~End UPCGSettingsInterface interface
+
+	//~Begin UPCGSettings interface
 	virtual bool HasDynamicPins() const override { return true; }
 	virtual EPCGDataType GetCurrentPinTypes(const UPCGPin* InPin) const override;
 
+#if WITH_EDITOR
+	virtual FName GetDefaultNodeName() const override { return FName("Reroute"); }
+	virtual FText GetDefaultNodeTitle() const override { return NSLOCTEXT("PCGRerouteElement", "NodeTitle", "Reroute"); }
+#endif
+
 protected:
 	virtual TArray<FPCGPinProperties> InputPinProperties() const override;
 	virtual TArray<FPCGPinProperties> OutputPinProperties() const override;
 	virtual FPCGElementPtr CreateElement() const override;
-	
-#if WITH_EDITOR
-	virtual FName GetDefaultNodeName() const override { return FName("Reroute"); }
-	virtual FText GetDefaultNodeTitle() const override { return NSLOCTEXT("PCGRerouteSettings", "NodeTitle", "Reroute"); }
-#endif
+	//~End UPCGSettings interface
 };
 
 class PCG_API FPCGRerouteElement : public IPCGElement
 {
-public:
-	virtual bool IsCacheable(const UPCGSettings* InSettings) const override { return true; }
-
 protected:
 	virtual bool ExecuteInternal(FPCGContext* Context) const override;
 };
diff --git a/Engine/Plugins/Experimental/PCG/Source/PCG/Public/PCGDebug.h b/Engine/Plugins/Experimental/PCG/Source/PCG/Public/PCGDebug.h
--- a/Engine/Plugins/Experimental/PCG/Source/PCG/Public/PCGDebug.h
+++ b/Engine/Plugins/Experimental/PCG/Source/PCG/Public/PCGDebug.h
@@ -24,23 +24,29 @@ struct PCG_API FPCGDebugVisualizationSettings
 public:
 	FPCGDebugVisualizationSettings();
 
-	UPROPERTY(BlueprintReadWrite, EditAnywhere, Category = Debug, meta=(ClampMin="0", EditCondition = "ScaleMethod != EPCGDebugVisScaleMethod::Extents"))
+	UPROPERTY(BlueprintReadWrite, EditAnywhere, Category = Debug, meta = (ClampMin="0", EditCondition = "ScaleMethod != EPCGDebugVisScaleMethod::Extents && bDisplayProperties", EditConditionHides, HideEditConditionToggle))
 	float PointScale = 1.0f;
 
-	UPROPERTY(BlueprintReadWrite, EditAnywhere, Category = Debug)
+	UPROPERTY(BlueprintReadWrite, EditAnywhere, Category = Debug, meta = (EditCondition = "bDisplayProperties", EditConditionHides, HideEditConditionToggle))
 	EPCGDebugVisScaleMethod ScaleMethod = EPCGDebugVisScaleMethod::Extents;
 
-	UPROPERTY(BlueprintReadWrite, EditAnywhere, Category = Debug)
+	UPROPERTY(BlueprintReadWrite, EditAnywhere, Category = Debug, meta = (EditCondition = "bDisplayProperties", EditConditionHides, HideEditConditionToggle))
 	TSoftObjectPtr<UStaticMesh> PointMesh;
 
-	UPROPERTY(BlueprintReadWrite, EditAnywhere, Category = Debug)
+	UPROPERTY(BlueprintReadWrite, EditAnywhere, Category = Debug, meta = (EditCondition = "bDisplayProperties", EditConditionHides, HideEditConditionToggle))
 	TSoftObjectPtr<UMaterialInterface> MaterialOverride;
 
 	/** Warning: enabling this flag will have severe performance impact */
-	UPROPERTY(BlueprintReadWrite, EditAnywhere, Category = Debug)
+	UPROPERTY(BlueprintReadWrite, EditAnywhere, Category = Debug, meta = (EditCondition = "bDisplayProperties", EditConditionHides, HideEditConditionToggle))
 	bool bCheckForDuplicates = false;
 
 	TSoftObjectPtr<UMaterialInterface> GetMaterial() const;
+
+#if WITH_EDITORONLY_DATA
+	// This can be set false to hide the debugging properties.
+	UPROPERTY(Transient, meta = (EditCondition = false, EditConditionHides))
+	bool bDisplayProperties = true;
+#endif
 };
 
 #if UE_ENABLE_INCLUDE_ORDER_DEPRECATED_IN_5_2
diff --git a/Engine/Plugins/Experimental/PCG/Source/PCG/Public/PCGSettings.h b/Engine/Plugins/Experimental/PCG/Source/PCG/Public/PCGSettings.h
--- a/Engine/Plugins/Experimental/PCG/Source/PCG/Public/PCGSettings.h
+++ b/Engine/Plugins/Experimental/PCG/Source/PCG/Public/PCGSettings.h
@@ -154,14 +154,20 @@ public:
 	/** Dedicated method to change enable state because some nodes have more complex behavior on enable/disable (such as subgraphs) */
 	void SetEnabled(bool bInEnabled);
 
+	/** Whether this element can be disabled. */
+	virtual bool CanBeDisabled() const { return true; }
+
+	/** Whether this element supports Debug and Inspect features. */
+	virtual bool CanBeDebugged() const { return true; }
+
 #if WITH_EDITOR
 	FOnPCGSettingsChanged OnSettingsChangedDelegate;
 #endif
 
-	UPROPERTY(EditAnywhere, BlueprintReadWrite, Category = Debug)
+	UPROPERTY(EditAnywhere, BlueprintReadWrite, Category = Debug, meta = (EditCondition = bDisplayDebuggingProperties, EditConditionHides, HideEditConditionToggle))
 	bool bEnabled = true;
 
-	UPROPERTY(Transient, EditAnywhere, BlueprintReadWrite, Category = Debug)
+	UPROPERTY(Transient, EditAnywhere, BlueprintReadWrite, Category = Debug, meta = (EditCondition = bDisplayDebuggingProperties, EditConditionHides, HideEditConditionToggle))
 	bool bDebug = false;
 
 #if WITH_EDITORONLY_DATA
@@ -169,8 +175,12 @@ public:
 	FPCGDebugVisualizationSettings DebugSettings;
 
 	/** If a debugger is attached, triggers a breakpoint inside IPCGElement::Execute(). Editor only. Transient. */
-	UPROPERTY(Transient, DuplicateTransient, EditAnywhere, BlueprintReadWrite, Category = Debug, AdvancedDisplay)
+	UPROPERTY(Transient, DuplicateTransient, EditAnywhere, BlueprintReadWrite, Category = Debug, AdvancedDisplay, meta = (EditCondition = bDisplayDebuggingProperties, EditConditionHides))
 	bool bBreakDebugger = false;
+
+	// This can be set false by inheriting nodes to hide the debugging properties.
+	UPROPERTY(Transient, meta = (EditCondition = false, EditConditionHides))
+	bool bDisplayDebuggingProperties = true;
 #endif
 };
 
diff --git a/Engine/Plugins/Experimental/PCG/Source/PCGEditor/Private/PCGEditor.cpp b/Engine/Plugins/Experimental/PCG/Source/PCGEditor/Private/PCGEditor.cpp
--- a/Engine/Plugins/Experimental/PCG/Source/PCGEditor/Private/PCGEditor.cpp
+++ b/Engine/Plugins/Experimental/PCG/Source/PCGEditor/Private/PCGEditor.cpp
@@ -579,22 +579,24 @@ void FPCGEditor::BindCommands()
 	GraphEditorCommands->MapAction(
 		PCGEditorCommands.ToggleEnabled,
 		FExecuteAction::CreateSP(this, &FPCGEditor::OnToggleEnabled),
-		FCanExecuteAction(),
+		FCanExecuteAction::CreateSP(this, &FPCGEditor::CanToggleEnabled),
 		FGetActionCheckState::CreateSP(this, &FPCGEditor::GetEnabledCheckState));
 
 	GraphEditorCommands->MapAction(
 		PCGEditorCommands.ToggleDebug,
 		FExecuteAction::CreateSP(this, &FPCGEditor::OnToggleDebug),
-		FCanExecuteAction(),
+		FCanExecuteAction::CreateSP(this, &FPCGEditor::CanToggleDebug),
 		FGetActionCheckState::CreateSP(this, &FPCGEditor::GetDebugCheckState));
 
 	GraphEditorCommands->MapAction(
 		PCGEditorCommands.DebugOnlySelected,
-		FExecuteAction::CreateSP(this, &FPCGEditor::OnDebugOnlySelected));
+		FExecuteAction::CreateSP(this, &FPCGEditor::OnDebugOnlySelected),
+		FCanExecuteAction::CreateSP(this, &FPCGEditor::CanToggleDebug));
 
 	GraphEditorCommands->MapAction(
 		PCGEditorCommands.DisableDebugOnAllNodes,
-		FExecuteAction::CreateSP(this, &FPCGEditor::OnDisableDebugOnAllNodes));
+		FExecuteAction::CreateSP(this, &FPCGEditor::OnDisableDebugOnAllNodes),
+		FCanExecuteAction::CreateSP(this, &FPCGEditor::CanToggleDebug));
 
 	GraphEditorCommands->MapAction(
 		PCGEditorCommands.AddSourcePin,
@@ -1137,6 +1139,15 @@ void FPCGEditor::OnToggleInspected()
 
 	UEdGraphNode* GraphNode = GraphEditorWidget->GetSingleSelectedNode();
 	UPCGEditorGraphNodeBase* PCGGraphNodeBase = Cast<UPCGEditorGraphNodeBase>(GraphNode);
+
+	const UPCGNode* PCGNode = PCGGraphNodeBase ? PCGGraphNodeBase->GetPCGNode() : nullptr;
+	const UPCGSettingsInterface* PCGSettingsInterface = PCGNode ? PCGNode->GetSettingsInterface() : nullptr;
+
+	if (PCGSettingsInterface && !PCGSettingsInterface->CanBeDebugged())
+	{
+		return;
+	}
+
 	if (PCGGraphNodeBase && PCGGraphNodeBase != PCGGraphNodeBeingInspected)
 	{
 		PCGGraphNodeBeingInspected = PCGGraphNodeBase;
@@ -1163,9 +1174,29 @@ bool FPCGEditor::CanToggleInspected() const
 		return false;
 	}
 
-	const UEdGraphNode* GraphNode = GraphEditorWidget->GetSingleSelectedNode();
+	const FGraphPanelSelectionSet& SelectedNodes = GraphEditorWidget->GetSelectedNodes();
+	if (SelectedNodes.Num() != 1)
+	{
+		// Can only inspect one node.
+		return false;
+	}
+
+	for (const UObject* Object : GraphEditorWidget->GetSelectedNodes())
+	{
+		const UPCGEditorGraphNodeBase* PCGEditorGraphNode = Cast<UPCGEditorGraphNodeBase>(Object);
+		if (!PCGEditorGraphNode)
+		{
+			return false;
+		}
+
+		const UPCGSettingsInterface* PCGSettingsInterface = PCGEditorGraphNode->GetPCGNode() ? PCGEditorGraphNode->GetPCGNode()->GetSettingsInterface() : nullptr;
+		if (PCGSettingsInterface && PCGSettingsInterface->CanBeDebugged())
+		{
+			return true;
+		}
+	}
 
-	return GraphNode && GraphNode->IsA<UPCGEditorGraphNodeBase>();
+	return false;
 }
 
 ECheckBoxState FPCGEditor::GetInspectedCheckState() const
@@ -1219,19 +1250,10 @@ void FPCGEditor::OnToggleEnabled()
 		for (UObject* Object : GraphEditorWidget->GetSelectedNodes())
 		{
 			UPCGEditorGraphNodeBase* PCGEditorGraphNode = Cast<UPCGEditorGraphNodeBase>(Object);
-			if (!PCGEditorGraphNode)
-			{
-				continue;
-			}
-
-			UPCGNode* PCGNode = PCGEditorGraphNode->GetPCGNode();
-			if (!PCGNode)
-			{
-				continue;
-			}
+			UPCGNode* PCGNode = PCGEditorGraphNode ? PCGEditorGraphNode->GetPCGNode() : nullptr;
+			UPCGSettingsInterface* PCGSettingsInterface = PCGNode ? PCGNode->GetSettingsInterface() : nullptr;
 
-			UPCGSettingsInterface* PCGSettingsInterface = PCGNode->GetSettingsInterface();
-			if (!PCGSettingsInterface)
+			if (!PCGSettingsInterface || !PCGSettingsInterface->CanBeDisabled())
 			{
 				continue;
 			}
@@ -1251,6 +1273,32 @@ void FPCGEditor::OnToggleEnabled()
 	}
 }
 
+bool FPCGEditor::CanToggleEnabled() const
+{
+	if (!GraphEditorWidget.IsValid())
+	{
+		return false;
+	}
+
+	for (const UObject* Object : GraphEditorWidget->GetSelectedNodes())
+	{
+		const UPCGEditorGraphNodeBase* PCGEditorGraphNode = Cast<UPCGEditorGraphNodeBase>(Object);
+		const UPCGNode* PCGNode = PCGEditorGraphNode ? PCGEditorGraphNode->GetPCGNode() : nullptr;
+		if (!PCGNode)
+		{
+			continue;
+		}
+
+		if (PCGNode->GetSettingsInterface() && PCGNode->GetSettingsInterface()->CanBeDisabled())
+		{
+			return true;
+		}
+	}
+
+	// Could not toggle enabled on anything in selection.
+	return false;
+}
+
 ECheckBoxState FPCGEditor::GetEnabledCheckState() const
 {
 	if (GraphEditorWidget.IsValid())
@@ -1258,22 +1306,13 @@ ECheckBoxState FPCGEditor::GetEnabledCheckState() const
 		bool bAllEnabled = true;
 		bool bAnyEnabled = false;
 
-		for (UObject* Object : GraphEditorWidget->GetSelectedNodes())
+		for (const UObject* Object : GraphEditorWidget->GetSelectedNodes())
 		{
 			const UPCGEditorGraphNodeBase* PCGEditorGraphNode = Cast<UPCGEditorGraphNodeBase>(Object);
-			if (!PCGEditorGraphNode)
-			{
-				continue;
-			}
-
-			const UPCGNode* PCGNode = PCGEditorGraphNode->GetPCGNode();
-			if (!PCGNode)
-			{
-				continue;
-			}
+			const UPCGNode* PCGNode = PCGEditorGraphNode ? PCGEditorGraphNode->GetPCGNode() : nullptr;
+			const UPCGSettingsInterface* PCGSettingsInterface = PCGNode ? PCGNode->GetSettingsInterface() : nullptr;
 
-			const UPCGSettingsInterface* PCGSettingsInterface = PCGNode->GetSettingsInterface();
-			if (!PCGSettingsInterface)
+			if (!PCGSettingsInterface || !PCGSettingsInterface->CanBeDisabled())
 			{
 				continue;
 			}
@@ -1307,19 +1346,10 @@ void FPCGEditor::OnToggleDebug()
 		for (UObject* Object : GraphEditorWidget->GetSelectedNodes())
 		{
 			UPCGEditorGraphNodeBase* PCGEditorGraphNode = Cast<UPCGEditorGraphNodeBase>(Object);
-			if (!PCGEditorGraphNode)
-			{
-				continue;
-			}
-
-			UPCGNode* PCGNode = PCGEditorGraphNode->GetPCGNode();
-			if (!PCGNode)
-			{
-				continue;
-			}
+			UPCGNode* PCGNode = PCGEditorGraphNode ? PCGEditorGraphNode->GetPCGNode() : nullptr;
+			UPCGSettingsInterface* PCGSettingsInterface = PCGNode ? PCGNode->GetSettingsInterface() : nullptr;
 
-			UPCGSettingsInterface* PCGSettingsInterface = PCGNode->GetSettingsInterface();
-			if (!PCGSettingsInterface)
+			if (!PCGSettingsInterface || !PCGSettingsInterface->CanBeDebugged())
 			{
 				continue;
 			}
@@ -1333,6 +1363,28 @@ void FPCGEditor::OnToggleDebug()
 	}
 }
 
+bool FPCGEditor::CanToggleDebug() const
+{
+	if (!GraphEditorWidget.IsValid())
+	{
+		return false;
+	}
+
+	for (const UObject* Object : GraphEditorWidget->GetSelectedNodes())
+	{
+		const UPCGEditorGraphNodeBase* PCGEditorGraphNode = Cast<UPCGEditorGraphNodeBase>(Object);
+		const UPCGNode* PCGNode = PCGEditorGraphNode ? PCGEditorGraphNode->GetPCGNode() : nullptr;
+
+		if (PCGNode && PCGNode->GetSettingsInterface()->CanBeDebugged())
+		{
+			return true;
+		}
+	}
+
+	// Could not toggle debug on anything in selection.
+	return false;
+}
+
 void FPCGEditor::OnDebugOnlySelected()
 {
 	if (GraphEditorWidget.IsValid() && PCGEditorGraph)
@@ -1347,13 +1399,12 @@ void FPCGEditor::OnDebugOnlySelected()
 		bool bAllSelectedNodesDebugged = true;
 
 		// Initial pass - inspect state of selected and non-selected nodes.
-		for (UEdGraphNode* Node : PCGEditorGraph->Nodes)
+		for (const UEdGraphNode* Node : PCGEditorGraph->Nodes)
 		{
-			UPCGEditorGraphNodeBase* PCGEditorGraphNode = Cast<UPCGEditorGraphNodeBase>(Node);
-			UPCGNode* PCGNode = PCGEditorGraphNode ? PCGEditorGraphNode->GetPCGNode() : nullptr;
-			UPCGSettingsInterface* PCGSettingsInterface = PCGNode ? PCGNode->GetSettingsInterface() : nullptr;
-
-			if (!PCGEditorGraphNode || !PCGNode || !PCGSettingsInterface)
+			const UPCGEditorGraphNodeBase* PCGEditorGraphNode = Cast<UPCGEditorGraphNodeBase>(Node);
+			const UPCGNode* PCGNode = PCGEditorGraphNode ? PCGEditorGraphNode->GetPCGNode() : nullptr;
+			const UPCGSettingsInterface* PCGSettingsInterface = PCGNode ? PCGNode->GetSettingsInterface() : nullptr;
+			if (!PCGSettingsInterface)
 			{
 				continue;
 			}
@@ -1378,7 +1429,7 @@ void FPCGEditor::OnDebugOnlySelected()
 			UPCGNode* PCGNode = PCGEditorGraphNode ? PCGEditorGraphNode->GetPCGNode() : nullptr;
 			UPCGSettingsInterface* PCGSettingsInterface = PCGNode ? PCGNode->GetSettingsInterface() : nullptr;
 
-			if (!PCGEditorGraphNode || !PCGNode || !PCGSettingsInterface)
+			if (!PCGSettingsInterface || !PCGSettingsInterface->CanBeDebugged())
 			{
 				continue;
 			}
@@ -1414,8 +1465,7 @@ void FPCGEditor::OnDisableDebugOnAllNodes()
 			UPCGEditorGraphNodeBase* PCGEditorGraphNode = Cast<UPCGEditorGraphNodeBase>(Node);
 			UPCGNode* PCGNode = PCGEditorGraphNode ? PCGEditorGraphNode->GetPCGNode() : nullptr;
 			UPCGSettingsInterface* PCGSettingsInterface = PCGNode ? PCGNode->GetSettingsInterface() : nullptr;
-
-			if (!PCGEditorGraphNode || !PCGNode || !PCGSettingsInterface)
+			if (!PCGSettingsInterface)
 			{
 				continue;
 			}
@@ -1443,22 +1493,13 @@ ECheckBoxState FPCGEditor::GetDebugCheckState() const
 		bool bAllDebug = true;
 		bool bAnyDebug = false;
 
-		for (UObject* Object : GraphEditorWidget->GetSelectedNodes())
+		for (const UObject* Object : GraphEditorWidget->GetSelectedNodes())
 		{
 			const UPCGEditorGraphNodeBase* PCGEditorGraphNode = Cast<UPCGEditorGraphNodeBase>(Object);
-			if (!PCGEditorGraphNode)
-			{
-				continue;
-			}
+			const UPCGNode* PCGNode = PCGEditorGraphNode ? PCGEditorGraphNode->GetPCGNode() : nullptr;
+			const UPCGSettingsInterface* PCGSettingsInterface = PCGNode ? PCGNode->GetSettingsInterface() : nullptr;
 
-			const UPCGNode* PCGNode = PCGEditorGraphNode->GetPCGNode();
-			if (!PCGNode)
-			{
-				continue;
-			}
-
-			const UPCGSettingsInterface* PCGSettingsInterface = PCGNode->GetSettingsInterface();
-			if (!PCGSettingsInterface)
+			if (!PCGSettingsInterface || !PCGSettingsInterface->CanBeDebugged())
 			{
 				continue;
 			}
diff --git a/Engine/Plugins/Experimental/PCG/Source/PCGEditor/Private/PCGEditor.h b/Engine/Plugins/Experimental/PCG/Source/PCGEditor/Private/PCGEditor.h
--- a/Engine/Plugins/Experimental/PCG/Source/PCGEditor/Private/PCGEditor.h
+++ b/Engine/Plugins/Experimental/PCG/Source/PCGEditor/Private/PCGEditor.h
@@ -124,11 +124,15 @@ private:
 	
 	/** Toggle node enabled state for selected nodes */
 	void OnToggleEnabled();
+	/** Whether we can toggle enabled state of selected nodes */
+	bool CanToggleEnabled() const;
 	/** Whether selected nodes are enabled or not */
 	ECheckBoxState GetEnabledCheckState() const;
 	
 	/** Toggle node debug state for selected nodes */
 	void OnToggleDebug();
+	/** Whether we can toggle debug state of selected nodes */
+	bool CanToggleDebug() const;
 	/** Whether selected nodes are being debugged or not */
 	ECheckBoxState GetDebugCheckState() const;
 
