[PCG] Expose DensityFunction settings for GetTextureData

#rnx

diff --git a/Engine/Plugins/Experimental/PCG/Source/PCG/Public/Elements/PCGTextureSampler.h b/Engine/Plugins/Experimental/PCG/Source/PCG/Public/Elements/PCGTextureSampler.h
--- a/Engine/Plugins/Experimental/PCG/Source/PCG/Public/Elements/PCGTextureSampler.h
+++ b/Engine/Plugins/Experimental/PCG/Source/PCG/Public/Elements/PCGTextureSampler.h
@@ -64,7 +64,7 @@ public:
 #endif
 
 	// Common members in BaseTextureData
-	UPROPERTY(BlueprintReadWrite, VisibleAnywhere, Category = SpatialData, meta = (PCG_Overridable))
+	UPROPERTY(BlueprintReadWrite, EditAnywhere, Category = SpatialData, meta = (PCG_Overridable))
 	EPCGTextureDensityFunction DensityFunction = EPCGTextureDensityFunction::Multiply;
 
 	UPROPERTY(BlueprintReadWrite, EditAnywhere, Category = Settings, meta = (PCG_Overridable))
