[PCG] Replace PCGPointProcessingElement with new PCGPointOperationElement

- Beginning deprecation of PCGPointProcessingElement
+ Inherits from PCGTimeSlicedElement
* Mutate Seed Element updated
* Bounds Modifier Element updated

diff --git a/Engine/Plugins/Experimental/PCG/Source/PCG/Private/Elements/PCGBoundsModifier.cpp b/Engine/Plugins/Experimental/PCG/Source/PCG/Private/Elements/PCGBoundsModifier.cpp
--- a/Engine/Plugins/Experimental/PCG/Source/PCG/Private/Elements/PCGBoundsModifier.cpp
+++ b/Engine/Plugins/Experimental/PCG/Source/PCG/Private/Elements/PCGBoundsModifier.cpp
@@ -9,6 +9,13 @@
 
 #define LOCTEXT_NAMESPACE "PCGBoundsModifier"
 
+namespace PCGBoundsModifier
+{
+	// TODO: Evaluate this value for optimization
+	// An evolving best guess for the most optimized number of points to operate per thread per slice
+	static constexpr int32 PointsPerChunk = 65536;
+}
+
 FPCGElementPtr UPCGBoundsModifierSettings::CreateElement() const
 {
 	return MakeShared<FPCGBoundsModifier>();
@@ -25,96 +32,90 @@ bool FPCGBoundsModifier::ExecuteInternal(FPCGContext* Context) const
 {
 	TRACE_CPUPROFILER_EVENT_SCOPE(FPCGBoundsModifier::Execute);
 
+	check(Context);
+	FPCGBoundsModifier::ContextType* BoundsModifierContext = static_cast<FPCGBoundsModifier::ContextType*>(Context);
+
 	const UPCGBoundsModifierSettings* Settings = Context->GetInputSettings<UPCGBoundsModifierSettings>();
 	check(Settings);
 
-	TArray<FPCGTaggedData> Inputs = Context->InputData.GetInputsByPin(PCGPinConstants::DefaultInputLabel);
-	TArray<FPCGTaggedData>& Outputs = Context->OutputData.TaggedData;
-
-	const EPCGBoundsModifierMode Mode = Settings->Mode;
-	const FVector& BoundsMin = Settings->BoundsMin;
-	const FVector& BoundsMax = Settings->BoundsMax;
-	const bool bAffectSteepness = Settings->bAffectSteepness;
-	const float Steepness = Settings->Steepness;
-
-	const FBox Bounds(BoundsMin, BoundsMax);
-
-	switch (Mode)
+	switch (Settings->Mode)
 	{
-	case EPCGBoundsModifierMode::Intersect:
-		ProcessPoints(Context, Inputs, Outputs, [&Bounds, bAffectSteepness, Steepness](const FPCGPoint& InPoint, FPCGPoint& OutPoint){
+		case EPCGBoundsModifierMode::Intersect:
+		return ExecutePointOperation(BoundsModifierContext, [Settings](const FPCGPoint& InPoint, FPCGPoint& OutPoint)
+		{
 			OutPoint = InPoint;
-			OutPoint.SetLocalBounds(InPoint.GetLocalBounds().Overlap(Bounds));
-			
-			if (bAffectSteepness)
+			OutPoint.SetLocalBounds(InPoint.GetLocalBounds().Overlap(FBox(Settings->BoundsMin, Settings->BoundsMax)));
+
+			if (Settings->bAffectSteepness)
 			{
-				OutPoint.Steepness = FMath::Min(InPoint.Steepness, Steepness);
+				OutPoint.Steepness = FMath::Min(InPoint.Steepness, Settings->Steepness);
 			}
 
 			return true;
-		});
-		break;
+		}, PCGBoundsModifier::PointsPerChunk);
 
 	case EPCGBoundsModifierMode::Include:
-		ProcessPoints(Context, Inputs, Outputs, [&Bounds, bAffectSteepness, Steepness](const FPCGPoint& InPoint, FPCGPoint& OutPoint){
+		return ExecutePointOperation(BoundsModifierContext, [Settings](const FPCGPoint& InPoint, FPCGPoint& OutPoint)
+		{
 			OutPoint = InPoint;
-			OutPoint.SetLocalBounds(InPoint.GetLocalBounds() + Bounds);
+			OutPoint.SetLocalBounds(InPoint.GetLocalBounds() + FBox(Settings->BoundsMin, Settings->BoundsMax));
 
-			if (bAffectSteepness)
+			if (Settings->bAffectSteepness)
 			{
-				OutPoint.Steepness = FMath::Max(InPoint.Steepness, Steepness);
+				OutPoint.Steepness = FMath::Max(InPoint.Steepness, Settings->Steepness);
 			}
 
 			return true;
-		});
-		break;
+		}, PCGBoundsModifier::PointsPerChunk);
 
 	case EPCGBoundsModifierMode::Translate:
-		ProcessPoints(Context, Inputs, Outputs, [&BoundsMin, &BoundsMax, bAffectSteepness, Steepness](const FPCGPoint& InPoint, FPCGPoint& OutPoint){
+		return ExecutePointOperation(BoundsModifierContext, [Settings](const FPCGPoint& InPoint, FPCGPoint& OutPoint)
+		{
 			OutPoint = InPoint;
-			OutPoint.BoundsMin += BoundsMin;
-			OutPoint.BoundsMax += BoundsMax;
+			OutPoint.BoundsMin += Settings->BoundsMin;
+			OutPoint.BoundsMax += Settings->BoundsMax;
 
-			if (bAffectSteepness)
+			if (Settings->bAffectSteepness)
 			{
-				OutPoint.Steepness = FMath::Clamp(InPoint.Steepness + Steepness, 0.0f, 1.0f);
+				OutPoint.Steepness = FMath::Clamp(InPoint.Steepness + Settings->Steepness, 0.0f, 1.0f);
 			}
 
 			return true;
-		});
-		break;
+		}, PCGBoundsModifier::PointsPerChunk);
 
 	case EPCGBoundsModifierMode::Scale:
-		ProcessPoints(Context, Inputs, Outputs, [&BoundsMin, &BoundsMax, bAffectSteepness, Steepness](const FPCGPoint& InPoint, FPCGPoint& OutPoint){
+		return ExecutePointOperation(BoundsModifierContext, [Settings](const FPCGPoint& InPoint, FPCGPoint& OutPoint)
+		{
 			OutPoint = InPoint;
-			OutPoint.BoundsMin *= BoundsMin;
-			OutPoint.BoundsMax *= BoundsMax;
+			OutPoint.BoundsMin *= Settings->BoundsMin;
+			OutPoint.BoundsMax *= Settings->BoundsMax;
 
-			if (bAffectSteepness)
+			if (Settings->bAffectSteepness)
 			{
-				OutPoint.Steepness = FMath::Clamp(InPoint.Steepness * Steepness, 0.0f, 1.0f);
+				OutPoint.Steepness = FMath::Clamp(InPoint.Steepness * Settings->Steepness, 0.0f, 1.0f);
 			}
 
 			return true;
-		});
-		break;
+		}, PCGBoundsModifier::PointsPerChunk);
 
 	case EPCGBoundsModifierMode::Set:
-		ProcessPoints(Context, Inputs, Outputs, [&Bounds, bAffectSteepness, Steepness](const FPCGPoint& InPoint, FPCGPoint& OutPoint){
+		return ExecutePointOperation(BoundsModifierContext, [Settings](const FPCGPoint& InPoint, FPCGPoint& OutPoint)
+		{
 			OutPoint = InPoint;
-			OutPoint.SetLocalBounds(Bounds);
+			OutPoint.SetLocalBounds(FBox(Settings->BoundsMin, Settings->BoundsMax));
 
-			if (bAffectSteepness)
+			if (Settings->bAffectSteepness)
 			{
-				OutPoint.Steepness = Steepness;
+				OutPoint.Steepness = Settings->Steepness;
 			}
 
 			return true;
-		});
-		break;
-	}
+		}, PCGBoundsModifier::PointsPerChunk);
 
-	return true;
+		default:
+			checkNoEntry();
+			return true;
+	}
 }
 
 #undef LOCTEXT_NAMESPACE
diff --git a/Engine/Plugins/Experimental/PCG/Source/PCG/Private/Elements/PCGMutateSeed.cpp b/Engine/Plugins/Experimental/PCG/Source/PCG/Private/Elements/PCGMutateSeed.cpp
--- a/Engine/Plugins/Experimental/PCG/Source/PCG/Private/Elements/PCGMutateSeed.cpp
+++ b/Engine/Plugins/Experimental/PCG/Source/PCG/Private/Elements/PCGMutateSeed.cpp
@@ -3,15 +3,18 @@
 #include "Elements/PCGMutateSeed.h"
 
 #include "PCGContext.h"
-#include "PCGParamData.h"
-#include "Data/PCGPointData.h"
-#include "Helpers/PCGBlueprintHelpers.h"
 #include "Helpers/PCGHelpers.h"
-#include "Metadata/PCGMetadata.h"
 
 #define LOCTEXT_NAMESPACE "PCGMutateSeedSettings"
 
-UPCGMutateSeedSettings::UPCGMutateSeedSettings() 
+namespace PCGMutateSeedConstants
+{
+	// TODO: Evaluate this value for optimization
+	// An evolving best guess for the most optimized number of points to operate per thread per slice
+	static constexpr int32 PointsPerChunk = 98304;
+}
+
+UPCGMutateSeedSettings::UPCGMutateSeedSettings()
 {
 	bUseSeed = true;
 }
@@ -24,25 +27,16 @@ FPCGElementPtr UPCGMutateSeedSettings::CreateElement() const
 bool FPCGMutateSeedElement::ExecuteInternal(FPCGContext* Context) const
 {
 	TRACE_CPUPROFILER_EVENT_SCOPE(FPCGMutateSeedElement::Execute);
-	
-	check(Context);
 
-	const UPCGMutateSeedSettings* Settings = Context->GetInputSettings<UPCGMutateSeedSettings>();
-	check(Settings);
-
-	TArray<FPCGTaggedData> Inputs = Context->InputData.GetInputs();
-	TArray<FPCGTaggedData>& Outputs = Context->OutputData.TaggedData;
-
-	const int Seed = Context->GetSeed();
+	check(Context);
+	ContextType* MutateSeedContext = static_cast<ContextType*>(Context);
 
-	ProcessPoints(Context, Inputs, Outputs, [Seed](const FPCGPoint& InPoint, FPCGPoint& OutPoint)
+	return ExecutePointOperation(MutateSeedContext, [Seed = Context->GetSeed()](const FPCGPoint& InPoint, FPCGPoint& OutPoint)
 	{
 		OutPoint = InPoint;
-		OutPoint.Seed = PCGHelpers::ComputeSeed(UPCGBlueprintHelpers::ComputeSeedFromPosition(OutPoint.Transform.GetLocation()), Seed, OutPoint.Seed);
+		OutPoint.Seed = PCGHelpers::ComputeSeed(PCGHelpers::ComputeSeedFromPosition(OutPoint.Transform.GetLocation()), Seed, OutPoint.Seed);
 		return true;
-	});
-	
-	return true;
+	}, PCGMutateSeedConstants::PointsPerChunk);
 }
 
 #undef LOCTEXT_NAMESPACE
diff --git a/Engine/Plugins/Experimental/PCG/Source/PCG/Private/Elements/PCGPointOperationElementBase.cpp b/Engine/Plugins/Experimental/PCG/Source/PCG/Private/Elements/PCGPointOperationElementBase.cpp
new file mode 100644
--- /dev/null
+++ b/Engine/Plugins/Experimental/PCG/Source/PCG/Private/Elements/PCGPointOperationElementBase.cpp
@@ -0,0 +1,117 @@
+// Copyright Epic Games, Inc. All Rights Reserved.
+
+#include "Elements/PCGPointOperationElementBase.h"
+
+#include "Data/PCGPointData.h"
+#include "Helpers/PCGAsync.h"
+
+#define LOCTEXT_NAMESPACE "PCGPointOperationElementBase"
+
+bool FPCGPointOperationElementBase::PrepareDataInternal(FPCGContext* Context) const
+{
+	check(Context);
+	ContextType* PointProcessContext = static_cast<ContextType*>(Context);
+	check(PointProcessContext);
+
+	// Prepares the context for time slicing
+	return PreparePointOperationData(PointProcessContext);
+}
+
+bool FPCGPointOperationElementBase::PreparePointOperationData(ContextType* Context) const
+{
+	TRACE_CPUPROFILER_EVENT_SCOPE(FPCGPointOperationElementBase::PreparePointProcessing);
+	check(Context);
+
+	const TArray<FPCGTaggedData>& Inputs = Context->InputData.GetInputs();
+	TArray<FPCGTaggedData>& Outputs = Context->OutputData.TaggedData;
+
+	// There is no execution state, so this just flags that its okay to continue
+	Context->InitializePerExecutionState();
+
+	// Prepare the 'per iteration' time slice context state and allocate output point data
+	Context->InitializePerIterationStates(Inputs.Num(),
+		[this, &Context, &Inputs, &Outputs](IterStateType& OutState, const ExecStateType& ExecState, const uint32 IterationIndex)
+		{
+			TRACE_CPUPROFILER_EVENT_SCOPE(FPCGPointOperationElementBase::InitializePerIterationStates);
+
+			FPCGTaggedData& Output = Outputs.Add_GetRef(Inputs[IterationIndex]);
+
+			const UPCGSpatialData* SpatialData = Cast<UPCGSpatialData>(Inputs[IterationIndex].Data);
+
+			if (!SpatialData)
+			{
+				PCGE_LOG(Error, GraphAndLog, LOCTEXT("InputMissingSpatialData", "Unable to get Spatial data from input"));
+				return EPCGTimeSliceInitResult::NoOperation;
+			}
+
+			OutState.InputPointData = SpatialData->ToPointData(Context);
+
+			if (!OutState.InputPointData)
+			{
+				PCGE_LOG(Error, GraphAndLog, LOCTEXT("InputMissingPointData", "Unable to get Point data from input"));
+				return EPCGTimeSliceInitResult::NoOperation;
+			}
+
+			OutState.NumPoints = OutState.InputPointData->GetPoints().Num();
+
+			// Create and initialize the output points
+			OutState.OutputPointData = NewObject<UPCGPointData>();
+			OutState.OutputPointData->InitializeFromData(OutState.InputPointData);
+			OutState.OutputPointData->GetMutablePoints().SetNumUninitialized(OutState.NumPoints);
+
+			Output.Data = OutState.OutputPointData;
+
+			return EPCGTimeSliceInitResult::Success;
+		});
+
+	return true;
+}
+
+bool FPCGPointOperationElementBase::ExecutePointOperation(ContextType* Context, TFunctionRef<PointExecSignature> PointFunction, int32 PointsPerChunk) const
+{
+	TRACE_CPUPROFILER_EVENT_SCOPE(FPCGPointOperationElementBase::ExecutePointProcessing);
+	check(Context);
+
+	// Standard check that the time slice state has been prepared. If the result is NoOp or Failure, result in no output
+	if (!Context->DataIsPreparedForExecution() || Context->GetExecutionStateResult() != EPCGTimeSliceInitResult::Success)
+	{
+		Context->OutputData.TaggedData.Empty();
+		return true;
+	}
+
+	return ExecuteSlice(Context, [this, &PointFunction, &PointsPerChunk](ContextType* Context, const ExecStateType& ExecState, const IterStateType& IterState, const uint32 IterIndex)
+	{
+		const TArray<FPCGPoint>& InPoints = IterState.InputPointData->GetPoints();
+		TArray<FPCGPoint>& OutPoints = IterState.OutputPointData->GetMutablePoints();
+
+		// If this input created an error, result in no output
+		if (Context->GetIterationStateResult(IterIndex) != EPCGTimeSliceInitResult::Success)
+		{
+			OutPoints.Empty();
+			return true;
+		}
+
+		// Conversion lambda from index to point ref for ease of use
+		auto InternalPointFunction = [&PointFunction, &InPoints, &OutPoints](int32 ReadIndex, int32 WriteIndex)->bool
+		{
+			return PointFunction(InPoints[ReadIndex], OutPoints[WriteIndex]);
+		};
+
+		const bool bAsyncDone = FPCGAsync::AsyncProcessingOneToOneEx(
+			&Context->AsyncState,
+			IterState.NumPoints,
+			/*InitializeFunc=*/[] {}, // Not useful for this context, since its preferred to initialize in PrepareDataInternal, so empty lambda
+			std::move(InternalPointFunction),
+			Context->TimeSliceIsEnabled(),
+			PointsPerChunk);
+
+		if (bAsyncDone && Context)
+		{
+			PCGE_LOG_C(Verbose, LogOnly, Context, FText::Format(LOCTEXT("PointProcessInfo", "Processed {0} points"), IterState.NumPoints));
+		}
+
+		return bAsyncDone;
+	});
+}
+
+#undef LOCTEXT_NAMESPACE
diff --git a/Engine/Plugins/Experimental/PCG/Source/PCG/Private/Helpers/PCGBlueprintHelpers.cpp b/Engine/Plugins/Experimental/PCG/Source/PCG/Private/Helpers/PCGBlueprintHelpers.cpp
--- a/Engine/Plugins/Experimental/PCG/Source/PCG/Private/Helpers/PCGBlueprintHelpers.cpp
+++ b/Engine/Plugins/Experimental/PCG/Source/PCG/Private/Helpers/PCGBlueprintHelpers.cpp
@@ -19,8 +19,7 @@
 
 int UPCGBlueprintHelpers::ComputeSeedFromPosition(const FVector& InPosition)
 {
-	// TODO: should have a config to drive this
-	return PCGHelpers::ComputeSeed((int)InPosition.X, (int)InPosition.Y, (int)InPosition.Z);
+	return PCGHelpers::ComputeSeedFromPosition(InPosition);
 }
 
 void UPCGBlueprintHelpers::SetSeedFromPosition(FPCGPoint& InPoint)
diff --git a/Engine/Plugins/Experimental/PCG/Source/PCG/Private/Helpers/PCGHelpers.cpp b/Engine/Plugins/Experimental/PCG/Source/PCG/Private/Helpers/PCGHelpers.cpp
--- a/Engine/Plugins/Experimental/PCG/Source/PCG/Private/Helpers/PCGHelpers.cpp
+++ b/Engine/Plugins/Experimental/PCG/Source/PCG/Private/Helpers/PCGHelpers.cpp
@@ -36,6 +36,11 @@ namespace PCGHelpers
 		return ((A * 196314165U) + 907633515U) ^ ((B * 73148459U) + 453816763U) ^ ((C * 34731343U) + 453816743U);
 	}
 
+	int ComputeSeedFromPosition(const FVector& InPosition)
+	{
+		return ComputeSeed(static_cast<int>(InPosition.X), static_cast<int>(InPosition.Y), static_cast<int>(InPosition.Z));
+	}
+
 	bool IsInsideBounds(const FBox& InBox, const FVector& InPosition)
 	{
 		return (InPosition.X >= InBox.Min.X) && (InPosition.X < InBox.Max.X) &&
diff --git a/Engine/Plugins/Experimental/PCG/Source/PCG/Public/Elements/PCGBoundsModifier.h b/Engine/Plugins/Experimental/PCG/Source/PCG/Public/Elements/PCGBoundsModifier.h
--- a/Engine/Plugins/Experimental/PCG/Source/PCG/Public/Elements/PCGBoundsModifier.h
+++ b/Engine/Plugins/Experimental/PCG/Source/PCG/Public/Elements/PCGBoundsModifier.h
@@ -2,10 +2,9 @@
 
 #pragma once
 
+#include "PCGPointOperationElementBase.h"
 #include "PCGSettings.h"
 
-#include "Elements/PCGPointProcessingElementBase.h"
-
 #include "PCGBoundsModifier.generated.h"
 
 UENUM()
@@ -59,8 +58,8 @@ public:
 	float Steepness = 1.0f;
 };
 
-class FPCGBoundsModifier : public FPCGPointProcessingElementBase
+class FPCGBoundsModifier : public FPCGPointOperationElementBase
 {
 protected:
-	virtual bool ExecuteInternal(FPCGContext* Context) const;
+	virtual bool ExecuteInternal(FPCGContext* Context) const override;
 };
\ No newline at end of file
diff --git a/Engine/Plugins/Experimental/PCG/Source/PCG/Public/Elements/PCGMutateSeed.h b/Engine/Plugins/Experimental/PCG/Source/PCG/Public/Elements/PCGMutateSeed.h
--- a/Engine/Plugins/Experimental/PCG/Source/PCG/Public/Elements/PCGMutateSeed.h
+++ b/Engine/Plugins/Experimental/PCG/Source/PCG/Public/Elements/PCGMutateSeed.h
@@ -2,12 +2,9 @@
 
 #pragma once
 
+#include "PCGPointOperationElementBase.h"
 #include "PCGSettings.h"
 
-#include "Elements/PCGPointProcessingElementBase.h"
-
-#include "UObject/NameTypes.h"
-
 #include "PCGMutateSeed.generated.h"
 
 /**
@@ -30,13 +27,13 @@ public:
 #endif
 
 protected:
-	virtual TArray<FPCGPinProperties> InputPinProperties() const override { return Super::DefaultPointInputPinProperties(); };
-	virtual TArray<FPCGPinProperties> OutputPinProperties() const override { return Super::DefaultPointOutputPinProperties(); };
+	virtual TArray<FPCGPinProperties> InputPinProperties() const override { return Super::DefaultPointInputPinProperties(); }
+	virtual TArray<FPCGPinProperties> OutputPinProperties() const override { return Super::DefaultPointOutputPinProperties(); }
 	virtual FPCGElementPtr CreateElement() const override;
 	//~End UPCGSettings interface
 };
 
-class FPCGMutateSeedElement : public FPCGPointProcessingElementBase
+class FPCGMutateSeedElement : public FPCGPointOperationElementBase
 {
 protected:
 	virtual bool ExecuteInternal(FPCGContext* Context) const override;
diff --git a/Engine/Plugins/Experimental/PCG/Source/PCG/Public/Elements/PCGPointOperationElementBase.h b/Engine/Plugins/Experimental/PCG/Source/PCG/Public/Elements/PCGPointOperationElementBase.h
new file mode 100644
--- /dev/null
+++ b/Engine/Plugins/Experimental/PCG/Source/PCG/Public/Elements/PCGPointOperationElementBase.h
@@ -0,0 +1,46 @@
+// Copyright Epic Games, Inc. All Rights Reserved.
+
+#pragma once
+
+#include "PCGTimeSlicedElementBase.h"
+
+class UPCGPointData;
+struct FPCGPoint;
+
+namespace PCGPointOperation
+{
+	namespace Constants
+	{
+		/** The default minimum number of points to execute per async slice */
+		static constexpr int32 PointsPerChunk = 4096;
+	}
+
+	/** Stores the input and output data as the state of the time sliced execution */
+	struct IterationState
+	{
+		const UPCGPointData* InputPointData = nullptr;
+		UPCGPointData* OutputPointData = nullptr;
+		int32 NumPoints = 0;
+	};
+}
+
+/** Simplified, time-sliced, and point by point operation class. A function or lambda may be passed into the `ExecutePointOperation` at execution time to
+ * invoke a customized update operation on all incoming points, individually.
+ */
+class PCG_API FPCGPointOperationElementBase : public TPCGTimeSlicedElementBase<PCGTimeSlice::FEmptyStruct, PCGPointOperation::IterationState>
+{
+	using PointExecSignature = bool(const FPCGPoint& InPoint, FPCGPoint& OutPoint);
+
+protected:
+	//~Begin IPCGElement interface
+	/** Conveniently calls PreparePointOperationData to prepare the time sliced element for execution. May be overridden, but PreparePointOperationData must be called. */
+	virtual bool PrepareDataInternal(FPCGContext* Context) const override;
+	//~End IPCGElement interface
+
+	/** Executes the PointFunction function/lambda for every point copied from PreparePointOperationData */
+	bool ExecutePointOperation(ContextType* Context, TFunctionRef<PointExecSignature> PointFunction, int32 PointsPerChunk = PCGPointOperation::Constants::PointsPerChunk) const;
+
+private:
+	/** Mandatory call. Using the context, prepares the state data for time slice execution */
+	bool PreparePointOperationData(ContextType* Context) const;
+};
\ No newline at end of file
diff --git a/Engine/Plugins/Experimental/PCG/Source/PCG/Public/Elements/PCGPointProcessingElementBase.h b/Engine/Plugins/Experimental/PCG/Source/PCG/Public/Elements/PCGPointProcessingElementBase.h
--- a/Engine/Plugins/Experimental/PCG/Source/PCG/Public/Elements/PCGPointProcessingElementBase.h
+++ b/Engine/Plugins/Experimental/PCG/Source/PCG/Public/Elements/PCGPointProcessingElementBase.h
@@ -8,6 +8,8 @@
 struct FPCGPoint;
 struct FPCGTaggedData;
 
+// TODO: Deprecate this class once all of the internal inheritors have been cleaned up
+// class UE_DEPRECATED(5.4, "FPCGPointProcessingElementBase is deprecated. Please use 'FPCGPointOperationElementBase' instead.") PCG_API FPCGPointProcessingElementBase : public IPCGElement
 class PCG_API FPCGPointProcessingElementBase : public IPCGElement
 {
 protected:
diff --git a/Engine/Plugins/Experimental/PCG/Source/PCG/Public/Elements/PCGTimeSlicedElementBase.h b/Engine/Plugins/Experimental/PCG/Source/PCG/Public/Elements/PCGTimeSlicedElementBase.h
--- a/Engine/Plugins/Experimental/PCG/Source/PCG/Public/Elements/PCGTimeSlicedElementBase.h
+++ b/Engine/Plugins/Experimental/PCG/Source/PCG/Public/Elements/PCGTimeSlicedElementBase.h
@@ -57,12 +57,12 @@ struct TPCGTimeSlicedContext : public FPCGContext
 	using InitExecSignature = EPCGTimeSliceInitResult(TPCGTimeSlicedContext* Context, PerExecutionStateT& OutState);
 
 	/** Initializes per execution state data if required. */
-	EPCGTimeSliceInitResult InitializePerExecutionState(TFunctionRef<InitExecSignature> InitFunc = []{ return true; });
+	EPCGTimeSliceInitResult InitializePerExecutionState(TFunctionRef<InitExecSignature> InitFunc = [](TPCGTimeSlicedContext*, PerExecutionStateT&) -> EPCGTimeSliceInitResult { return EPCGTimeSliceInitResult::Success; });
 
 	using InitIterSignature = EPCGTimeSliceInitResult(PerIterationStateT& OutState, const PerExecutionStateT& ExecState, const uint32 IterationIndex);
 
 	/** Initializes per execution state data if required. An array will be created with a state element for every execution iteration in the context. Returns the Init Result for each iteration's initialization. */
-	const TArray<EPCGTimeSliceInitResult>& InitializePerIterationStates(int32 NumIterations = 1, TFunctionRef<InitIterSignature> IterFunc = []{ return true; });
+	const TArray<EPCGTimeSliceInitResult>& InitializePerIterationStates(int32 NumIterations = 1, TFunctionRef<InitIterSignature> IterFunc = [](PerIterationStateT&, const uint32){ return TArray({EPCGTimeSliceInitResult::Success}); });
 
 	/** Will return the result of the attempt to initialize the per execution state */
 	EPCGTimeSliceInitResult GetExecutionStateResult() const { return ExecutionStateResult; }
diff --git a/Engine/Plugins/Experimental/PCG/Source/PCG/Public/Helpers/PCGHelpers.h b/Engine/Plugins/Experimental/PCG/Source/PCG/Public/Helpers/PCGHelpers.h
--- a/Engine/Plugins/Experimental/PCG/Source/PCG/Public/Helpers/PCGHelpers.h
+++ b/Engine/Plugins/Experimental/PCG/Source/PCG/Public/Helpers/PCGHelpers.h
@@ -25,6 +25,7 @@ namespace PCGHelpers
 	PCG_API int ComputeSeed(int A);
 	PCG_API int ComputeSeed(int A, int B);
 	PCG_API int ComputeSeed(int A, int B, int C);
+	PCG_API int ComputeSeedFromPosition(const FVector& InPosition);
 
 	PCG_API bool IsInsideBounds(const FBox& InBox, const FVector& InPosition);
 	PCG_API bool IsInsideBoundsXY(const FBox& InBox, const FVector& InPosition);
