[PCG] Fix a missed ID remapping during task culling.

This was triggering an infinite loop in downstream code - also added detection and an ensure there.

#rnx

diff --git a/Engine/Plugins/Experimental/PCG/Source/PCG/Private/Graph/PCGGraphCompiler.cpp b/Engine/Plugins/Experimental/PCG/Source/PCG/Private/Graph/PCGGraphCompiler.cpp
--- a/Engine/Plugins/Experimental/PCG/Source/PCG/Private/Graph/PCGGraphCompiler.cpp
+++ b/Engine/Plugins/Experimental/PCG/Source/PCG/Private/Graph/PCGGraphCompiler.cpp
@@ -533,9 +533,9 @@ void FPCGGraphCompiler::CullTasks(TArray<FPCGGraphTask>& InOutCompiledTasks, boo
 
 	InOutCompiledTasks.SetNum(WriteIndex);
 
-	// Remap input task IDs, and remove edges that connect to culled nodes.
 	for (FPCGGraphTask& Task : InOutCompiledTasks)
 	{
+		// Remap input task IDs, and remove edges that connect to culled nodes.
 		for (int32 InputIndex = Task.Inputs.Num() - 1; InputIndex >= 0; --InputIndex)
 		{
 			const int32 InputTaskId = Task.Inputs[InputIndex].TaskId;
@@ -549,6 +549,17 @@ void FPCGGraphCompiler::CullTasks(TArray<FPCGGraphTask>& InOutCompiledTasks, boo
 				Task.Inputs.RemoveAt(InputIndex);
 			}
 		}
+
+		// Remap parent ID if there tasks is in a child scope.
+		if (Task.ParentId != InvalidPCGTaskId)
+		{
+			const int32 RemappedParentId = TaskRemapping[Task.ParentId];
+			// Parent task should not have been culled.
+			ensure(RemappedParentId != INDEX_NONE);
+
+			// Write the remapped ID - even if it's invalid/INDEX_NONE. Hanging parent IDs can cause issues elsewhere.
+			Task.ParentId = RemappedParentId;
+		}
 	}
 }
 
@@ -594,7 +605,7 @@ void FPCGGraphCompiler::OffsetNodeIds(TArray<FPCGGraphTask>& Tasks, FPCGTaskId O
 	{
 		Task.NodeId += Offset;
 
-		if  (Task.ParentId == InvalidPCGTaskId)
+		if (Task.ParentId == InvalidPCGTaskId)
 		{
 			Task.ParentId = ParentId;
 		}
@@ -603,7 +614,7 @@ void FPCGGraphCompiler::OffsetNodeIds(TArray<FPCGGraphTask>& Tasks, FPCGTaskId O
 			Task.ParentId += Offset;
 		}
 
-		for(FPCGGraphTaskInput& Input : Task.Inputs)
+		for (FPCGGraphTaskInput& Input : Task.Inputs)
 		{
 			Input.TaskId += Offset;
 		}
diff --git a/Engine/Plugins/Experimental/PCG/Source/PCG/Private/Utils/PCGExtraCapture.cpp b/Engine/Plugins/Experimental/PCG/Source/PCG/Private/Utils/PCGExtraCapture.cpp
--- a/Engine/Plugins/Experimental/PCG/Source/PCG/Private/Utils/PCGExtraCapture.cpp
+++ b/Engine/Plugins/Experimental/PCG/Source/PCG/Private/Utils/PCGExtraCapture.cpp
@@ -101,6 +101,12 @@ namespace PCGUtils
 					break;
 				}
 
+				if (Task.ParentId == PathIds.Last())
+				{
+					ensureMsgf(false, TEXT("Parent ID is same as current task ID which will result in infinite loop."));
+					break;
+				}
+
 				PathIds.Add(Task.ParentId);
 			}
 
