[PCG] Add extractor capability on GetActorProperty

* Use AttributePropertySelector for extraction
* Add convenient functions for Selector creation
* Add unit test cases and fix existing broken test

#rnx

diff --git a/Engine/Plugins/Experimental/PCG/Source/PCG/Private/Elements/PCGPropertyToParamData.cpp b/Engine/Plugins/Experimental/PCG/Source/PCG/Private/Elements/PCGPropertyToParamData.cpp
--- a/Engine/Plugins/Experimental/PCG/Source/PCG/Private/Elements/PCGPropertyToParamData.cpp
+++ b/Engine/Plugins/Experimental/PCG/Source/PCG/Private/Elements/PCGPropertyToParamData.cpp
@@ -9,6 +9,7 @@
 #include "Helpers/PCGBlueprintHelpers.h"
 #include "Helpers/PCGPropertyHelpers.h"
 #include "Helpers/PCGHelpers.h"
+#include "Metadata/PCGAttributePropertySelector.h"
 #include "Metadata/Accessors/PCGAttributeAccessorHelpers.h"
 
 #include UE_INLINE_GENERATED_CPP_BY_NAME(PCGPropertyToParamData)
@@ -154,7 +155,9 @@ bool FPCGPropertyToParamDataElement::ExecuteInternal(FPCGContext* Context) const
 		}
 	}
 
-	PCGPropertyHelpers::FExtractorParameters Parameters{ ObjectToInspect, ObjectToInspect->GetClass(), Settings->PropertyName, Settings->OutputAttributeName, Settings->bForceObjectAndStructExtraction, /*bPropertyNeedsToBeVisible=*/true};
+	FPCGAttributePropertySelector Selector = FPCGAttributePropertySelector::CreateSelectorFromString(Settings->PropertyName.ToString());
+
+	PCGPropertyHelpers::FExtractorParameters Parameters{ ObjectToInspect, ObjectToInspect->GetClass(), Selector, Settings->OutputAttributeName, Settings->bForceObjectAndStructExtraction, /*bPropertyNeedsToBeVisible=*/true};
 	if (UPCGParamData* ParamData = PCGPropertyHelpers::ExtractPropertyAsAttributeSet(Parameters, Context))
 	{
 		TArray<FPCGTaggedData>& Outputs = Context->OutputData.TaggedData;
diff --git a/Engine/Plugins/Experimental/PCG/Source/PCG/Private/Elements/PCGUserParameterGet.cpp b/Engine/Plugins/Experimental/PCG/Source/PCG/Private/Elements/PCGUserParameterGet.cpp
--- a/Engine/Plugins/Experimental/PCG/Source/PCG/Private/Elements/PCGUserParameterGet.cpp
+++ b/Engine/Plugins/Experimental/PCG/Source/PCG/Private/Elements/PCGUserParameterGet.cpp
@@ -10,6 +10,7 @@
 #include "PCGSubgraph.h"
 #include "Data/PCGUserParametersData.h"
 #include "Helpers/PCGPropertyHelpers.h"
+#include "Metadata/PCGAttributePropertySelector.h"
 
 #include "PropertyBag.h"
 #include "StructView.h"
@@ -158,7 +159,7 @@ bool FPCGUserParameterGetElement::ExecuteInternal(FPCGContext* Context) const
 		return true;
 	}
 
-	PCGPropertyHelpers::FExtractorParameters ExtractorParameters{ Parameters.GetMemory(), PropertyBag, PropertyName, PropertyName, Settings->bForceObjectAndStructExtraction, /*bPropertyNeedsToBeVisible=*/false };
+	PCGPropertyHelpers::FExtractorParameters ExtractorParameters{ Parameters.GetMemory(), PropertyBag, FPCGAttributePropertySelector::CreateAttributeSelector(PropertyName), PropertyName, Settings->bForceObjectAndStructExtraction, /*bPropertyNeedsToBeVisible=*/false};
 
 	if (UPCGParamData* ParamData = PCGPropertyHelpers::ExtractPropertyAsAttributeSet(ExtractorParameters, Context))
 	{
diff --git a/Engine/Plugins/Experimental/PCG/Source/PCG/Private/Helpers/PCGPropertyHelpers.cpp b/Engine/Plugins/Experimental/PCG/Source/PCG/Private/Helpers/PCGPropertyHelpers.cpp
--- a/Engine/Plugins/Experimental/PCG/Source/PCG/Private/Helpers/PCGPropertyHelpers.cpp
+++ b/Engine/Plugins/Experimental/PCG/Source/PCG/Private/Helpers/PCGPropertyHelpers.cpp
@@ -14,6 +14,81 @@
 
 #define LOCTEXT_NAMESPACE "PCGPropertyHelpers"
 
+namespace PCGPropertyHelpers
+{
+	static constexpr uint64 ExcludePropertyFlags = CPF_DisableEditOnInstance;
+	static constexpr uint64 IncludePropertyFlags = CPF_BlueprintVisible;
+
+	void LogError(const FText& ErrorMessage, FPCGContext* InOptionalContext)
+	{
+		if (InOptionalContext)
+		{
+			PCGE_LOG_C(Error, GraphAndLog, InOptionalContext, ErrorMessage);
+		}
+		else
+		{
+			UE_LOG(LogPCG, Error, TEXT("%s"), *ErrorMessage.ToString());
+		}
+	}
+
+	/**
+	* Recursive function to go down the property chain to find the property and its container address.
+	* @param CurrentClass      Struct/Class for the current container
+	* @param CurrentName       Property name to look for in the container class.
+	* @param NextNames         List of property names to continue extracting at a deeper level.
+	* @param bNeedsToBeVisible Discard properties that are not visibile in Blueprint
+	* @param OutContainer      Raw address for the current container. Will be write to at each recursive call.
+	* @param OptionalContext   Optional context used for logging.
+	* @returns                 The last property of the chain (and its container address is in OutContainer)
+	*/
+	FProperty* ExtractPropertyChain(const UStruct* CurrentClass, const FName CurrentName, TArrayView<const FString> NextNames, const bool bNeedsToBeVisible, const void*& OutContainer, FPCGContext* OptionalContext)
+	{
+		check(CurrentClass);
+
+		// Try to get the property
+		FProperty* Property = FindFProperty<FProperty>(CurrentClass, CurrentName);
+		if (!Property)
+		{
+			LogError(FText::Format(LOCTEXT("PropertyDoesNotExist", "Property '{0}' does not exist in {1}."), FText::FromName(CurrentName), FText::FromName(CurrentClass->GetFName())), OptionalContext);
+			return nullptr;
+		}
+
+		// Make sure the property is visible, if requested
+		if (bNeedsToBeVisible && (Property->HasAnyPropertyFlags(ExcludePropertyFlags) || !Property->HasAnyPropertyFlags(IncludePropertyFlags)))
+		{
+			LogError(FText::Format(LOCTEXT("PropertyExistsButNotVisible", "Property '{0}' does exist in {1}, but is not visible."), FText::FromName(CurrentName), FText::FromName(CurrentClass->GetFName())), OptionalContext);
+			return nullptr;
+		}
+
+		if (!NextNames.IsEmpty())
+		{
+			UStruct* NextClass = nullptr;
+
+			if (FStructProperty* StructProperty = CastField<FStructProperty>(Property))
+			{
+				NextClass = StructProperty->Struct;
+				OutContainer = StructProperty->ContainerPtrToValuePtr<void>(OutContainer);
+			}
+			else if (FObjectProperty* ObjectProperty = CastField<FObjectProperty>(Property))
+			{
+				NextClass = ObjectProperty->PropertyClass;
+				OutContainer = ObjectProperty->GetObjectPropertyValue_InContainer(OutContainer);
+			}
+			else
+			{
+				LogError(FText::Format(LOCTEXT("PropertyIsNotExtractable", "Property '{0}' does exist in {1}, but is not extractable."), FText::FromName(CurrentName), FText::FromName(CurrentClass->GetFName())), OptionalContext);
+				return nullptr;
+			}
+
+			return ExtractPropertyChain(NextClass, FName(NextNames[0]), NextNames.RightChop(1), bNeedsToBeVisible, OutContainer, OptionalContext);
+		}
+		else
+		{
+			return Property;
+		}
+	}
+}
+
 EPCGMetadataTypes PCGPropertyHelpers::GetMetadataTypeFromProperty(const FProperty* InProperty)
 {
 	if (!InProperty)
@@ -36,38 +111,14 @@ UPCGParamData* PCGPropertyHelpers::ExtractPropertyAsAttributeSet(const PCGProper
 {
 	check(Parameters.Container && Parameters.Class);
 
-	auto LogError = [InOptionalContext](const FText& ErrorMessage)
-	{
-		if (InOptionalContext)
-		{
-			PCGE_LOG_C(Error, GraphAndLog, InOptionalContext, ErrorMessage);
-		}
-		else
-		{
-			UE_LOG(LogPCG, Error, TEXT("%s"), *ErrorMessage.ToString());
-		}
-	};
+	const void* Container = Parameters.Container;
+	FProperty* Property = ExtractPropertyChain(Parameters.Class, Parameters.PropertySelector.GetName(), Parameters.PropertySelector.GetExtraNames(), Parameters.bPropertyNeedsToBeVisible, Container, InOptionalContext);
 
-	// Try to get the property
-	FProperty* Property = FindFProperty<FProperty>(Parameters.Class, Parameters.PropertyName);
 	if (!Property)
 	{
-		LogError(FText::Format(LOCTEXT("PropertyDoesNotExist", "Property '{0}' does not exist."), FText::FromName(Parameters.PropertyName)));
 		return nullptr;
 	}
 
-	// Make sure the property is visible, if requested
-	const uint64 ExcludePropertyFlags = CPF_DisableEditOnInstance;
-	const uint64 IncludePropertyFlags = CPF_BlueprintVisible;
-	if (Parameters.bPropertyNeedsToBeVisible)
-	{
-		if (Property->HasAnyPropertyFlags(ExcludePropertyFlags) || !Property->HasAnyPropertyFlags(IncludePropertyFlags))
-		{
-			LogError(FText::Format(LOCTEXT("PropertyExistsButNotVisible", "Property '{0}' does exist, but is not visible."), FText::FromName(Parameters.PropertyName)));
-			return nullptr;
-		}
-	}
-
 	// If the property is an array, we will work on the underlying property, and extract each element as an entry in the param data
 	FArrayProperty* ArrayProperty = CastField<FArrayProperty>(Property);
 	if (ArrayProperty)
@@ -75,10 +126,10 @@ UPCGParamData* PCGPropertyHelpers::ExtractPropertyAsAttributeSet(const PCGProper
 		Property = ArrayProperty->Inner;
 	}
 
-	using GetAddressFunc = TFunction<const void* (const void*)>;
 	using ExtractablePropertyTuple = TTuple<FName, const FProperty*>;
 	TArray<ExtractablePropertyTuple> ExtractableProperties;
 
+	using GetAddressFunc = TFunction<const void* (const void*)>;
 	GetAddressFunc AddressFunc;
 
 	// Force extraction if the property is not supported by accessors.
@@ -144,7 +195,7 @@ UPCGParamData* PCGPropertyHelpers::ExtractPropertyAsAttributeSet(const PCGProper
 
 	if (ExtractableProperties.IsEmpty())
 	{
-		LogError(LOCTEXT("NoPropertiesFound", "No properties found to extract"));
+		LogError(LOCTEXT("NoPropertiesFound", "No properties found to extract"), InOptionalContext);
 		return nullptr;
 	}
 
@@ -152,7 +203,7 @@ UPCGParamData* PCGPropertyHelpers::ExtractPropertyAsAttributeSet(const PCGProper
 	TArray<const void*, TInlineAllocator<16>> ElementAddresses;
 	if (ArrayProperty)
 	{
-		FScriptArrayHelper_InContainer Helper(ArrayProperty, Parameters.Container);
+		FScriptArrayHelper_InContainer Helper(ArrayProperty, Container);
 		ElementAddresses.Reserve(Helper.Num());
 		for (int32 DynamicIndex = 0; DynamicIndex < Helper.Num(); ++DynamicIndex)
 		{
@@ -161,7 +212,7 @@ UPCGParamData* PCGPropertyHelpers::ExtractPropertyAsAttributeSet(const PCGProper
 	}
 	else
 	{
-		ElementAddresses.Add(Parameters.Container);
+		ElementAddresses.Add(Container);
 	}
 
 	// From there, we should be able to create the data.
@@ -191,7 +242,7 @@ UPCGParamData* PCGPropertyHelpers::ExtractPropertyAsAttributeSet(const PCGProper
 
 			if (!Metadata->SetAttributeFromDataProperty(AttributeName, EntryKey, ContainerPtr, FinalProperty, /*bCreate=*/ true))
 			{
-				LogError(FText::Format(LOCTEXT("ErrorCreatingAttribute", "Error while creating an attribute for property '{0}'. Either the property type is not supported by PCG or attribute creation failed."), FText::FromString(FinalProperty->GetName())));
+				LogError(FText::Format(LOCTEXT("ErrorCreatingAttribute", "Error while creating an attribute for property '{0}'. Either the property type is not supported by PCG or attribute creation failed."), FText::FromString(FinalProperty->GetName())), InOptionalContext);
 				bValidOperation = false;
 				break;
 			}
diff --git a/Engine/Plugins/Experimental/PCG/Source/PCG/Private/Metadata/PCGAttributePropertySelector.cpp b/Engine/Plugins/Experimental/PCG/Source/PCG/Private/Metadata/PCGAttributePropertySelector.cpp
--- a/Engine/Plugins/Experimental/PCG/Source/PCG/Private/Metadata/PCGAttributePropertySelector.cpp
+++ b/Engine/Plugins/Experimental/PCG/Source/PCG/Private/Metadata/PCGAttributePropertySelector.cpp
@@ -268,6 +268,39 @@ void FPCGAttributePropertySelector::AddToCrc(FArchiveCrc32& Ar) const
 	}
 }
 
+bool FPCGAttributePropertySelector::IsBasicAttribute() const
+{
+	return Selection == EPCGAttributePropertySelection::Attribute && ExtraNames.IsEmpty();
+}
+
+FPCGAttributePropertySelector FPCGAttributePropertySelector::CreateAttributeSelector(const FName AttributeName)
+{
+	FPCGAttributePropertySelector Selector;
+	Selector.SetAttributeName(AttributeName);
+	return Selector;
+}
+
+FPCGAttributePropertySelector FPCGAttributePropertySelector::CreatePointPropertySelector(EPCGPointProperties PointProperty)
+{
+	FPCGAttributePropertySelector Selector;
+	Selector.SetPointProperty(PointProperty);
+	return Selector;
+}
+
+FPCGAttributePropertySelector FPCGAttributePropertySelector::CreateExtraPropertySelector(EPCGExtraProperties ExtraProperty)
+{
+	FPCGAttributePropertySelector Selector;
+	Selector.SetExtraProperty(ExtraProperty);
+	return Selector;
+}
+
+FPCGAttributePropertySelector FPCGAttributePropertySelector::CreateSelectorFromString(const FString& String)
+{
+	FPCGAttributePropertySelector Selector;
+	Selector.Update(String);
+	return Selector;
+}
+
 ///////////////////////////////////////////////////////////////////////
 
 FPCGAttributePropertyInputSelector::FPCGAttributePropertyInputSelector()
diff --git a/Engine/Plugins/Experimental/PCG/Source/PCG/Private/Tests/Elements/PCGPropertyToParamDataElementTest.cpp b/Engine/Plugins/Experimental/PCG/Source/PCG/Private/Tests/Elements/PCGPropertyToParamDataElementTest.cpp
--- a/Engine/Plugins/Experimental/PCG/Source/PCG/Private/Tests/Elements/PCGPropertyToParamDataElementTest.cpp
+++ b/Engine/Plugins/Experimental/PCG/Source/PCG/Private/Tests/Elements/PCGPropertyToParamDataElementTest.cpp
@@ -94,14 +94,27 @@ bool VerifyAttributeValue(FPCGTestBaseClass* TestInstance, PCGTestsCommon::FTest
 
 		if (bSuccess)
 		{
-			const UPCGParamData* ParamData = Cast<UPCGParamData>(Context->OutputData.GetAllParams()[0].Data);
+			const UPCGParamData* ParamData = CastChecked<UPCGParamData>(Context->OutputData.GetAllParams()[0].Data);
+			const UPCGMetadata* Metadata = ParamData->ConstMetadata();
+			check(Metadata);
+
 			int32 Index = 0;
 			int32 ItemKey = 0;
 
+			// Make sure that we have exactly the same number of attributes as expected
+			if (!TestInstance->TestEqual(FormatWithPropertyName("Attribute number is matching"), Metadata->GetAttributeCount(), AttributeNames.Num()))
+			{
+				// Cleanup behind ourselves.
+				TestData.TestPCGComponent->GetGraph()->RemoveNode(TrivialNode);
+				TestData.TestPCGComponent->GetGraph()->RemoveNode(TestNode);
+
+				return false;
+			}
+
 			// Fold expression that will iterate on all Expected Values, to check their types and their values.
 			// We keep track of the current Index, to know which attribute to check, and an ItemKey to know where to look at in the value array of the attribute.
 			// cf. comment at the top to know in which order ExpectedValues should be.
-			([&Index, &AttributeNames, ParamData, &FormatWithPropertyNameAndSubNames, TestInstance, &bSuccess, &ItemKey, &ExpectedValues]()
+			([&Index, &AttributeNames, Metadata, &FormatWithPropertyNameAndSubNames, TestInstance, &bSuccess, &ItemKey, &ExpectedValues]()
 			{
 				if (bSuccess)
 				{
@@ -110,7 +123,7 @@ bool VerifyAttributeValue(FPCGTestBaseClass* TestInstance, PCGTestsCommon::FTest
 
 					// We use Index to track the current attribute to check
 					const FName AttributeName = AttributeNames[Index];
-					const FPCGMetadataAttributeBase* Attribute = ParamData->ConstMetadata()->GetConstAttribute(AttributeName);
+					const FPCGMetadataAttributeBase* Attribute = Metadata->GetConstAttribute(AttributeName);
 
 					bSuccess = TestInstance->TestNotNull(FormatWithPropertyNameAndSubNames("Attribute should not be null", AttributeName, ItemKey), Attribute);
 
@@ -233,6 +246,9 @@ bool FPCGPropertyToParamDataPropertyTypeTest::RunTest(const FString& Parameters)
 	Actor->ArrayOfVectorsProperty = { VectorValue, SecondVectorValue };
 	Actor->ArrayOfStructsProperty = { ColorValue, SecondColorValue };
 	Actor->ArrayOfObjectsProperty = { ObjectValue, SecondObjectValue };
+	Actor->DummyStruct.FloatProperty = 1.2f;
+	Actor->DummyStruct.IntArrayProperty = { 5, 6, 7 };
+	Actor->DummyStruct.Level2Struct.DoubleArrayProperty = { 0.1, 0.2, 0.3 };
 
 	// Basic properties
 	bSuccess &= VerifyAttributeValueValid(this, TestData, GET_MEMBER_NAME_CHECKED(APCGUnitTestDummyActor, IntProperty), 42ll, ExtraTestWhat);
@@ -281,9 +297,13 @@ bool FPCGPropertyToParamDataPropertyTypeTest::RunTest(const FString& Parameters)
 		(int64)ColorValue.R, (int64)ColorValue.G, (int64)ColorValue.B, (int64)ColorValue.A, (int64)SecondColorValue.R, (int64)SecondColorValue.G, (int64)SecondColorValue.B, (int64)SecondColorValue.A);
 	bSuccess &= VerifyAttributeValuesValid(this, TestData, GET_MEMBER_NAME_CHECKED(APCGUnitTestDummyActor, ArrayOfObjectsProperty), ObjectPropertyNames, ExtraTestWhat, ObjectValue->Int64Property, ObjectValue->DoubleProperty, SecondObjectValue->Int64Property, SecondObjectValue->DoubleProperty);
 
-	// Unsupported struct properties
-	AddExpectedError(TEXT("Error while creating an attribute for property"), EAutomationExpectedErrorFlags::Contains, 1);
-	bSuccess &= VerifyAttributeValueInvalid(this, TestData, GET_MEMBER_NAME_CHECKED(APCGUnitTestDummyActor, ColorProperty), ColorValue, ExtraTestWhat);
+	// Extractors
+	bSuccess &= VerifyAttributeValuesValid(this, TestData, TEXT("DummyStruct.FloatProperty"), { GET_MEMBER_NAME_CHECKED(FPCGDummyGetPropertyStruct, FloatProperty) }, ExtraTestWhat, 1.2);
+	bSuccess &= VerifyAttributeValuesValid(this, TestData, TEXT("DummyStruct.IntArrayProperty"), { GET_MEMBER_NAME_CHECKED(FPCGDummyGetPropertyStruct, IntArrayProperty) }, ExtraTestWhat, 5ll, 6ll, 7ll);
+	bSuccess &= VerifyAttributeValuesValid(this, TestData, TEXT("DummyStruct.Level2Struct.DoubleArrayProperty"), { GET_MEMBER_NAME_CHECKED(FPCGDummyGetPropertyLevel2Struct, DoubleArrayProperty) }, ExtraTestWhat, 0.1, 0.2, 0.3);
+
+	// Extracting the DummyStruct should only extract the float, as arrays and deeper structs are discarded
+	bSuccess &= VerifyAttributeValuesValid(this, TestData, TEXT("DummyStruct"), { GET_MEMBER_NAME_CHECKED(FPCGDummyGetPropertyStruct, FloatProperty) }, ExtraTestWhat, 1.2);
 
 	// Unknown property
 	AddExpectedError(TEXT("Property 'DummyMissingProperty' does not exist"), EAutomationExpectedErrorFlags::Contains, 1);
diff --git a/Engine/Plugins/Experimental/PCG/Source/PCG/Private/Tests/Elements/PCGPropertyToParamDataElementTest.h b/Engine/Plugins/Experimental/PCG/Source/PCG/Private/Tests/Elements/PCGPropertyToParamDataElementTest.h
--- a/Engine/Plugins/Experimental/PCG/Source/PCG/Private/Tests/Elements/PCGPropertyToParamDataElementTest.h
+++ b/Engine/Plugins/Experimental/PCG/Source/PCG/Private/Tests/Elements/PCGPropertyToParamDataElementTest.h
@@ -14,6 +14,30 @@ enum class EPCGUnitTestDummyEnum : int64
 	Three
 };
 
+USTRUCT(BlueprintType, meta = (Hidden))
+struct FPCGDummyGetPropertyLevel2Struct
+{
+	GENERATED_BODY()
+
+	UPROPERTY(BlueprintReadOnly, Category = "")
+	TArray<double> DoubleArrayProperty;
+};
+
+USTRUCT(BlueprintType, meta = (Hidden))
+struct FPCGDummyGetPropertyStruct
+{
+	GENERATED_BODY()
+
+	UPROPERTY(BlueprintReadOnly, Category = "")
+	TArray<int> IntArrayProperty;
+
+	UPROPERTY(BlueprintReadOnly, Category = "")
+	float FloatProperty = 0.0f;
+
+	UPROPERTY(BlueprintReadOnly, Category = "")
+	FPCGDummyGetPropertyLevel2Struct Level2Struct;
+};
+
 UCLASS(MinimalAPI, NotBlueprintable, NotPlaceable, NotBlueprintType, Transient, HideDropdown, meta = (Hidden))
 class UPCGDummyGetPropertyTest : public UObject
 {
@@ -101,6 +125,9 @@ public:
 
 	UPROPERTY(BlueprintReadOnly, Category = "")
 	TArray<TObjectPtr<UPCGDummyGetPropertyTest>> ArrayOfObjectsProperty;
+
+	UPROPERTY(BlueprintReadOnly, Category = "")
+	FPCGDummyGetPropertyStruct DummyStruct;
 };
 
 UCLASS(MinimalAPI, NotBlueprintable, NotPlaceable, NotBlueprintType, Transient, HideDropdown, meta = (Hidden))
diff --git a/Engine/Plugins/Experimental/PCG/Source/PCG/Public/Elements/PCGPropertyToParamData.h b/Engine/Plugins/Experimental/PCG/Source/PCG/Public/Elements/PCGPropertyToParamData.h
--- a/Engine/Plugins/Experimental/PCG/Source/PCG/Public/Elements/PCGPropertyToParamData.h
+++ b/Engine/Plugins/Experimental/PCG/Source/PCG/Public/Elements/PCGPropertyToParamData.h
@@ -62,7 +62,7 @@ public:
 	UPROPERTY(BlueprintReadWrite, EditAnywhere, Category = Settings)
 	bool bForceObjectAndStructExtraction = false;
 
-	/** By default, attribute name will be None, but it can be overridden by this name. */
+	/** By default, attribute name will be None, but it can be overridden by this name. Use @SourceName to use the property name (only works when not extracting). */
 	UPROPERTY(BlueprintReadWrite, EditAnywhere, Category = Settings, meta = (EditCondition = "!bExtractObjectAndStruct", EditConditionHides))
 	FName OutputAttributeName = NAME_None;
 
diff --git a/Engine/Plugins/Experimental/PCG/Source/PCG/Public/Helpers/PCGPropertyHelpers.h b/Engine/Plugins/Experimental/PCG/Source/PCG/Public/Helpers/PCGPropertyHelpers.h
--- a/Engine/Plugins/Experimental/PCG/Source/PCG/Public/Helpers/PCGPropertyHelpers.h
+++ b/Engine/Plugins/Experimental/PCG/Source/PCG/Public/Helpers/PCGPropertyHelpers.h
@@ -2,6 +2,7 @@
 
 #pragma once
 
+#include "Metadata/PCGAttributePropertySelector.h"
 #include "Metadata/PCGMetadataAttributeTpl.h"
 #include "Metadata/Accessors/IPCGAttributeAccessor.h"
 #include "Metadata/Accessors/PCGAttributeAccessorHelpers.h"
@@ -53,8 +54,8 @@ namespace PCGPropertyHelpers
 		// Class or ScriptStruct of the object/struct we want to extract the property from.
 		const UStruct* Class = nullptr;
 
-		// Name of the property we want to extract
-		FName PropertyName = NAME_None;
+		// Selector of the property we want to extract
+		FPCGAttributePropertySelector PropertySelector;
 
 		// Optional name of the attribute that will receive the extracted property. If None, will take the property name. 
 		// Also not used for Structs/Object extraction, as we will create multiple attributes, and they will be the name of all the extracted members.
diff --git a/Engine/Plugins/Experimental/PCG/Source/PCG/Public/Metadata/PCGAttributePropertySelector.h b/Engine/Plugins/Experimental/PCG/Source/PCG/Public/Metadata/PCGAttributePropertySelector.h
--- a/Engine/Plugins/Experimental/PCG/Source/PCG/Public/Metadata/PCGAttributePropertySelector.h
+++ b/Engine/Plugins/Experimental/PCG/Source/PCG/Public/Metadata/PCGAttributePropertySelector.h
@@ -58,6 +58,9 @@ public:
 	EPCGPointProperties GetPointProperty() const { return PointProperty; }
 	EPCGExtraProperties GetExtraProperty() const { return ExtraProperty; }
 
+	// Convenient function to know if it is a basic attribute (attribute and no extra names)
+	bool IsBasicAttribute() const;
+
 	// Return the name of the selector.
 	FName GetName() const;
 
@@ -78,6 +81,12 @@ public:
 		return OutSelector;
 	}
 
+	// Convenience static constructors
+	static FPCGAttributePropertySelector CreateAttributeSelector(const FName AttributeName);
+	static FPCGAttributePropertySelector CreatePointPropertySelector(EPCGPointProperties PointProperty);
+	static FPCGAttributePropertySelector CreateExtraPropertySelector(EPCGExtraProperties ExtraProperty);
+	static FPCGAttributePropertySelector CreateSelectorFromString(const FString& String);
+
 	void ImportFromOtherSelector(const FPCGAttributePropertySelector& InOther);
 
 	virtual void AddToCrc(FArchiveCrc32& Ar) const;
