[PCG] Clean up AddTag PCG Unit Test failures.

#rnx

diff --git a/Engine/Plugins/Experimental/PCG/Source/PCG/Private/Tests/Elements/PCGAddTagTest.cpp b/Engine/Plugins/Experimental/PCG/Source/PCG/Private/Tests/Elements/PCGAddTagTest.cpp
--- a/Engine/Plugins/Experimental/PCG/Source/PCG/Private/Tests/Elements/PCGAddTagTest.cpp
+++ b/Engine/Plugins/Experimental/PCG/Source/PCG/Private/Tests/Elements/PCGAddTagTest.cpp
@@ -4,6 +4,7 @@
 
 #include "PCGComponent.h"
 #include "PCGContext.h"
+#include "Data/PCGPointData.h"
 
 #include "Elements/PCGAddTag.h"
 
@@ -20,12 +21,17 @@ bool FPCGAddTagTest_ZeroTags::RunTest(const FString& Parameters)
 
 	TUniquePtr<FPCGContext> Context = TestData.InitializeTestContext();
 
+	// At least one input data is required to add tags and get an output
+	FPCGTaggedData& NewTaggedData = Context->InputData.TaggedData.Emplace_GetRef();
+	NewTaggedData.Data = PCGTestsCommon::CreatePointData();
+	NewTaggedData.Pin = PCGPinConstants::DefaultInputLabel;
+
 	while (!TestElement->Execute(Context.Get())) {}
 
 	//testing inputs and outputs
 	const TArray<FPCGTaggedData>& Inputs = Context->InputData.TaggedData;
 
-	UTEST_EQUAL("Number of Inputs", Inputs.Num(), 1);
+	UTEST_EQUAL("Number of Inputs", Inputs.Num(), 2);
 
 	const TArray<FPCGTaggedData>& Outputs = Context->OutputData.TaggedData;
 
@@ -55,12 +61,17 @@ bool FPCGAddTagTest_OneTag::RunTest(const FString& Parameters)
 
 	TUniquePtr<FPCGContext> Context = TestData.InitializeTestContext();
 
+	// At least one input data is required to add tags and get an output
+	FPCGTaggedData& NewTaggedData = Context->InputData.TaggedData.Emplace_GetRef();
+	NewTaggedData.Data = PCGTestsCommon::CreatePointData();
+	NewTaggedData.Pin = PCGPinConstants::DefaultInputLabel;
+
 	while (!TestElement->Execute(Context.Get())) {}
 
 	//testing inputs and outputs
 	const TArray<FPCGTaggedData>& Inputs = Context->InputData.TaggedData;
 
-	UTEST_EQUAL("Number of Inputs", Inputs.Num(), 1);
+	UTEST_EQUAL("Number of Inputs", Inputs.Num(), 2);
 
 	const TArray<FPCGTaggedData>& Outputs = Context->OutputData.TaggedData;
 
@@ -97,17 +108,22 @@ bool FPCGAddTagTest_MultipleTags::RunTest(const FString& Parameters)
 
 	TUniquePtr<FPCGContext> Context = TestData.InitializeTestContext();
 
+	// At least one input data is required to add tags and get an output
+	FPCGTaggedData& NewTaggedData = Context->InputData.TaggedData.Emplace_GetRef();
+	NewTaggedData.Data = PCGTestsCommon::CreatePointData();
+	NewTaggedData.Pin = PCGPinConstants::DefaultInputLabel;
+
 	while (!TestElement->Execute(Context.Get())) {}
 
 	//testing inputs and outputs
 	const TArray<FPCGTaggedData>& Inputs = Context->InputData.TaggedData;
 
-	UTEST_EQUAL("Number of Inputs", Inputs.Num(), 1);
+	UTEST_EQUAL("Number of Inputs", Inputs.Num(), 2);
 
 	const TArray<FPCGTaggedData>& Outputs = Context->OutputData.TaggedData;
 
 	UTEST_EQUAL("Number of Outputs", Outputs.Num(), 1);
-	
+
 	//testing tag info
 	const TSet<FString> OutputTagSet = Outputs[0].Tags;
 
