[PCG] Adds a node that creates points in a 2d or 3d grid
--Includes change to allow PCGVolumes for testing in PCGTestsCommon
--Support for partitioning (including CreatePoints)

diff --git a/Engine/Plugins/Experimental/PCG/Source/PCG/Private/Elements/PCGCreatePoints.cpp b/Engine/Plugins/Experimental/PCG/Source/PCG/Private/Elements/PCGCreatePoints.cpp
--- a/Engine/Plugins/Experimental/PCG/Source/PCG/Private/Elements/PCGCreatePoints.cpp
+++ b/Engine/Plugins/Experimental/PCG/Source/PCG/Private/Elements/PCGCreatePoints.cpp
@@ -37,11 +37,30 @@ bool FPCGCreatePointsElement::ExecuteInternal(FPCGContext* Context) const
 	const UPCGCreatePointsSettings* Settings = Context->GetInputSettings<UPCGCreatePointsSettings>();
 	check(Settings);
 
-	UPCGComponent* OriginalComponentContext = Context->SourceComponent.Get();
-	check(OriginalComponentContext);
+	UPCGComponent* PCGComponent = nullptr;
+	check(Context->SourceComponent.Get());
 
-	const FTransform OriginalComponentTransform = OriginalComponentContext->GetOwner()->GetActorTransform();
-	const FTransform ComponentTransformScaleOne = FTransform(OriginalComponentTransform.Rotator(), OriginalComponentTransform.GetLocation(), FVector::One());
+	if (Settings->GridPivot == EPCGLocalGridPivot::OriginalComponent)
+	{
+		PCGComponent = Context->SourceComponent->GetOriginalComponent();
+	}
+	else if (Settings->GridPivot == EPCGLocalGridPivot::LocalComponent)
+	{
+		PCGComponent = Context->SourceComponent.Get();
+	}
+
+	FTransform OriginalComponentTransform = FTransform();
+	FTransform ComponentTransformScaleOne = FTransform();
+	const UPCGSpatialData* Target = nullptr;
+
+	if (PCGComponent)
+	{
+		check(PCGComponent->GetOwner());
+
+		OriginalComponentTransform = PCGComponent->GetOwner()->GetActorTransform();
+		ComponentTransformScaleOne = FTransform(OriginalComponentTransform.Rotator(), OriginalComponentTransform.GetLocation(), FVector::One());
+		Target = Settings->bCullPointsOutsideVolume ? Cast<UPCGSpatialData>(PCGComponent->GetActorPCGData()) : nullptr;
+	}
 
 	TArray<FPCGPoint> PointsToLoopOn = Settings->PointsToCreate;
 	
@@ -53,9 +72,8 @@ bool FPCGCreatePointsElement::ExecuteInternal(FPCGContext* Context) const
 
 	TArray<FPCGPoint>& OutputPoints = PtData->GetMutablePoints();
 	Output.Data = PtData;
-	
 
-	if (!Settings->bLocal && !Settings->bCullPointsOutsideVolume)
+	if (Settings->GridPivot == EPCGLocalGridPivot::Global && !Settings->bCullPointsOutsideVolume)
 	{
 		for (auto& Points : PointsToLoopOn)
 		{
@@ -70,14 +88,12 @@ bool FPCGCreatePointsElement::ExecuteInternal(FPCGContext* Context) const
 	}
 	else
 	{
-		const UPCGSpatialData* Target = Settings->bCullPointsOutsideVolume ? Cast<UPCGSpatialData>(OriginalComponentContext->GetActorPCGData()) : nullptr;
-
 		FPCGAsync::AsyncPointProcessing(Context, PointsToLoopOn.Num(), OutputPoints, [&PointsToLoopOn, Settings, &ComponentTransformScaleOne, Target](int32 Index, FPCGPoint& OutPoint)
 		{
 			const FPCGPoint& InPoint = PointsToLoopOn[Index];
 			OutPoint = InPoint;
 
-			if (Settings->bLocal)
+			if (Settings->GridPivot == EPCGLocalGridPivot::LocalComponent || Settings->GridPivot == EPCGLocalGridPivot::OriginalComponent)
 			{
 				OutPoint.Transform *= ComponentTransformScaleOne;
 			}
@@ -96,7 +112,7 @@ bool FPCGCreatePointsElement::IsCacheable(const UPCGSettings* InSettings) const
 {
 	const UPCGCreatePointsSettings* Settings = Cast<const UPCGCreatePointsSettings>(InSettings);
 
-	return Settings && !Settings->bLocal;
+	return Settings && Settings->GridPivot == EPCGLocalGridPivot::Global;
 }
 
 #undef LOCTEXT_NAMESPACE
diff --git a/Engine/Plugins/Experimental/PCG/Source/PCG/Private/Elements/PCGCreatePointsGrid.cpp b/Engine/Plugins/Experimental/PCG/Source/PCG/Private/Elements/PCGCreatePointsGrid.cpp
new file mode 100644
--- /dev/null
+++ b/Engine/Plugins/Experimental/PCG/Source/PCG/Private/Elements/PCGCreatePointsGrid.cpp
@@ -0,0 +1,234 @@
+// Copyright Epic Games, Inc. All Rights Reserved.
+
+#include "Elements/PCGCreatePointsGrid.h"
+
+#include "PCGComponent.h"
+#include "PCGContext.h"
+#include "Helpers/PCGAsync.h"
+#include "Helpers/PCGHelpers.h"
+#include "Helpers/PCGSettingsHelpers.h"
+
+#define LOCTEXT_NAMESPACE "PCGCreatePointsGridElement"
+
+TArray<FPCGPinProperties> UPCGCreatePointsGridSettings::InputPinProperties() const
+{
+	return TArray<FPCGPinProperties>();
+}
+
+FPCGElementPtr UPCGCreatePointsGridSettings::CreateElement() const
+{
+	return MakeShared<FPCGCreatePointsGridElement>();
+}
+
+bool FPCGCreatePointsGridElement::ExecuteInternal(FPCGContext* Context) const
+{
+	TRACE_CPUPROFILER_EVENT_SCOPE(FPCGCreatePointsGridElement::Execute);
+	
+	check(Context);
+
+	const UPCGCreatePointsGridSettings* Settings = Context->GetInputSettings<UPCGCreatePointsGridSettings>();
+	check(Settings);
+
+	UPCGComponent* PCGComponent = nullptr;
+	check(Context->SourceComponent.Get());
+
+	if (Settings->GridPivot == EPCGGridPivot::OriginalComponent)
+	{
+		PCGComponent = Context->SourceComponent->GetOriginalComponent();
+	}
+	else if (Settings->GridPivot == EPCGGridPivot::LocalComponent)
+	{
+		PCGComponent = Context->SourceComponent.Get();
+	}
+
+	FTransform OriginalComponentTransform = FTransform();
+	FTransform ComponentTransformScaleOne = FTransform();
+	const UPCGSpatialData* Target = nullptr;
+
+	if (PCGComponent)
+	{
+		check(PCGComponent->GetOwner());
+
+		OriginalComponentTransform = PCGComponent->GetOwner()->GetActorTransform();
+		ComponentTransformScaleOne = FTransform(OriginalComponentTransform.Rotator(), OriginalComponentTransform.GetLocation(), FVector::One());
+		Target = Settings->bCullPointsOutsideVolume ? Cast<UPCGSpatialData>(PCGComponent->GetActorPCGData()) : nullptr;
+	}
+
+	if (Settings->CellSize.X <= 0.0 || Settings->CellSize.Y <= 0.0 || Settings->CellSize.Z <= 0.0)
+	{
+		PCGE_LOG(Warning, GraphAndLog, LOCTEXT("InvalidDataInput", "CellSize must not be less than 0"));
+		return true;
+	}
+
+	if (Settings->GridExtents.X < 0.0 || Settings->GridExtents.Y < 0.0 || Settings->GridExtents.Z < 0.0)
+	{
+		PCGE_LOG(Warning, GraphAndLog, LOCTEXT("InvalidDataInput", "GridExtents must not be less than 0"));
+		return true;
+	}
+
+	FVector GridExtents = Settings->GridExtents;
+	const FVector& CellSize = Settings->CellSize;
+
+	int32 PointCountX = FMath::TruncToInt((2 * GridExtents.X) / CellSize.X);
+	int32 PointCountY = FMath::TruncToInt((2 * GridExtents.Y) / CellSize.Y);
+	int32 PointCountZ = FMath::TruncToInt((2 * GridExtents.Z) / CellSize.Z);
+
+	if (Settings->PointPosition == EPCGPointPosition::CellCorners)
+	{ 
+		PointCountX++;
+		PointCountY++;
+		PointCountZ++;
+
+		if (GridExtents.X < (CellSize.X / 2))
+		{
+			GridExtents.X = 0.0;
+		}
+
+		if (GridExtents.Y < (CellSize.Y / 2))
+		{
+			GridExtents.Y = 0.0;
+		}
+
+		if (GridExtents.Z < (CellSize.Z / 2))
+		{
+			GridExtents.Z = 0.0;
+		}
+	}
+
+	// If the GridExtent would produce an off center result, snap it to the center of the grid
+	if (Settings->PointPosition == EPCGPointPosition::CellCenter)
+	{
+		if (GridExtents.X < (CellSize.X / 2))
+		{
+			PointCountX++;
+		}
+
+		if (GridExtents.Y < (CellSize.Y / 2))
+		{
+			PointCountY++;
+		}
+
+		if (GridExtents.Z < (CellSize.Z / 2))
+		{
+			PointCountZ++;
+		}
+
+		GridExtents.X = GridExtents.X - FMath::Fmod(GridExtents.X, CellSize.X / 2);
+		GridExtents.Y = GridExtents.Y - FMath::Fmod(GridExtents.Y, CellSize.Y / 2);
+		GridExtents.Z = GridExtents.Z - FMath::Fmod(GridExtents.Z, CellSize.Z / 2);
+	}
+
+	const int64 NumIterations64 = static_cast<int64>(PointCountX) * static_cast<int64>(PointCountY) * static_cast<int64>(PointCountZ);
+
+	if (NumIterations64 <= 0)
+	{
+		PCGE_LOG(Error, GraphAndLog, LOCTEXT("InvalidNumberOfIterations", "The number of iterations produced cannot be less than or equal to 0."));
+		return true;
+	}
+
+	if (NumIterations64 >= MAX_int32)
+	{
+		PCGE_LOG(Error, GraphAndLog, LOCTEXT("Overflow_int32", "The number of iterations produced is larger than what a 32-bit integer can hold."));
+		return true;
+	}
+
+	if (PCGFeatureSwitches::CVarCheckSamplerMemory.GetValueOnAnyThread() && FPlatformMemory::GetStats().AvailablePhysical < sizeof(FPCGPoint) * NumIterations64)
+	{
+		PCGE_LOG(Error, GraphAndLog, LOCTEXT("Overflow_int32", "The number of iterations produced is larger than available memory."));
+		return true;
+	}
+
+	int32 NumIterations = static_cast<int32>(NumIterations64);
+	
+	TArray<FPCGTaggedData>& Outputs = Context->OutputData.TaggedData;
+	FPCGTaggedData& Output = Outputs.Emplace_GetRef();
+
+	UPCGPointData* PtData = NewObject<UPCGPointData>();
+	check(PtData);
+
+	TArray<FPCGPoint>& OutputPoints = PtData->GetMutablePoints();
+	Output.Data = PtData;
+
+	FPCGAsync::AsyncPointProcessing(Context, NumIterations, OutputPoints, [Settings, Target, &PointCountX, &PointCountY, &CellSize, &GridExtents, &ComponentTransformScaleOne](int32 Index, FPCGPoint& OutPoint)
+	{
+		OutPoint = FPCGPoint();
+
+		double XCoordinate = Index % PointCountX;
+		double YCoordinate = (Index / PointCountX) % PointCountY;
+		double ZCoordinate = Index / (PointCountX * PointCountY);
+
+		if (Settings->PointPosition == EPCGPointPosition::CellCenter)
+		{
+			// If extents are smaller than a point in that dimension, place in center, otherwise offset 
+			XCoordinate += 0.5;
+			YCoordinate += 0.5;
+			ZCoordinate += 0.5;
+		}
+
+		// If the extends are smaller than the point, set point to origin
+		if (GridExtents.X < CellSize.X / 2.0)
+		{
+			XCoordinate = 0.0;
+		}
+		if (GridExtents.Y < CellSize.Y / 2.0)
+		{
+			YCoordinate = 0.0;
+		}
+		if (GridExtents.Z < CellSize.Z / 2.0)
+		{
+			ZCoordinate = 0.0;
+		}
+
+		const FVector GridTransformLocation((CellSize.X * XCoordinate) - GridExtents.X, (CellSize.Y * YCoordinate) - GridExtents.Y, (CellSize.Z * ZCoordinate) - GridExtents.Z);
+
+		if (Settings->GridPivot == EPCGGridPivot::LocalComponent || Settings->GridPivot == EPCGGridPivot::OriginalComponent)
+		{
+			OutPoint.Transform = FTransform(FRotator::ZeroRotator, GridTransformLocation, FVector::One()) * ComponentTransformScaleOne;
+		}
+		else
+		{
+			OutPoint.Transform.SetLocation(GridTransformLocation);
+		}
+
+		if (Settings->bSetPointsBounds)
+		{
+			OutPoint.SetExtents(CellSize * 0.5);
+		}
+
+		OutPoint.Seed = PCGHelpers::ComputeSeed(OutPoint.Transform.GetLocation().X, OutPoint.Transform.GetLocation().Y, OutPoint.Transform.GetLocation().Z);
+
+		// Discards points outside of the volume
+		return !Target || (Target->GetDensityAtPosition(OutPoint.Transform.GetLocation()) > 0.0f);
+	});
+
+	return true;
+}
+
+void FPCGCreatePointsGridElement::GetDependenciesCrc(const FPCGDataCollection& InInput, const UPCGSettings* InSettings, UPCGComponent* InComponent, FPCGCrc& OutCrc) const
+{
+	FPCGCrc Crc;
+	IPCGElement::GetDependenciesCrc(InInput, InSettings, InComponent, Crc);
+	
+	if (const UPCGCreatePointsGridSettings* Settings = Cast<UPCGCreatePointsGridSettings>(InSettings))
+	{
+		int GridPivot = static_cast<int>(EPCGGridPivot::Global);
+		bool bCullPointsOutsideVolume = false;
+		PCGSettingsHelpers::GetOverrideValue(InInput, Settings, GET_MEMBER_NAME_CHECKED(UPCGCreatePointsGridSettings, GridPivot), static_cast<int>(Settings->GridPivot), GridPivot);
+		PCGSettingsHelpers::GetOverrideValue(InInput, Settings, GET_MEMBER_NAME_CHECKED(UPCGCreatePointsGridSettings, bCullPointsOutsideVolume), Settings->bCullPointsOutsideVolume, bCullPointsOutsideVolume);
+		
+		// We're using the bounds of the pcg volume, so we extract the actor data here
+		EPCGGridPivot EnumGridPivot = static_cast<EPCGGridPivot>(GridPivot);
+
+		if ((EnumGridPivot == EPCGGridPivot::OriginalComponent || EnumGridPivot == EPCGGridPivot::LocalComponent || bCullPointsOutsideVolume) && InComponent)
+		{
+			if (const UPCGData* Data = InComponent->GetActorPCGData())
+			{
+				Crc.Combine(Data->GetOrComputeCrc(/*bFullDataCrc=*/false));
+			}
+		}
+	}
+
+	OutCrc = Crc;
+}
+
+#undef LOCTEXT_NAMESPACE
diff --git a/Engine/Plugins/Experimental/PCG/Source/PCG/Private/Tests/Elements/PCGCreatePointsGridTest.cpp b/Engine/Plugins/Experimental/PCG/Source/PCG/Private/Tests/Elements/PCGCreatePointsGridTest.cpp
new file mode 100644
--- /dev/null
+++ b/Engine/Plugins/Experimental/PCG/Source/PCG/Private/Tests/Elements/PCGCreatePointsGridTest.cpp
@@ -0,0 +1,410 @@
+// Copyright Epic Games, Inc. All Rights Reserved.
+
+#include "Tests/PCGTestsCommon.h"
+
+#include "PCGComponent.h"
+#include "PCGContext.h"
+#include "PCGVolume.h"
+#include "Data/PCGPointData.h"
+
+#include "Elements/PCGCreatePointsGrid.h"
+
+namespace CreatePointsGridCommonTests
+{
+	TUniquePtr<FPCGContext> GenerateTestDataAndRunGrid(EPCGPointPosition Position, const FVector GridSize, const FVector Cell)
+	{
+		PCGTestsCommon::FTestData TestData;
+		PCGTestsCommon::GenerateSettings<UPCGCreatePointsGridSettings>(TestData);
+		UPCGCreatePointsGridSettings* Settings = CastChecked<UPCGCreatePointsGridSettings>(TestData.Settings);
+		Settings->PointPosition = Position;
+		Settings->GridPivot = EPCGGridPivot::Global;
+		Settings->GridExtents = GridSize;
+		Settings->CellSize = Cell;
+
+		FPCGElementPtr TestElement = TestData.Settings->GetElement();
+
+		TUniquePtr<FPCGContext> Context = TestData.InitializeTestContext();
+
+		while (!TestElement->Execute(Context.Get())) {}
+
+		return Context;
+	}
+
+	PCGTestsCommon::FTestData GenerateTestData(EPCGPointPosition Position, EPCGGridPivot Pivot)
+	{
+		PCGTestsCommon::FTestData TestData(42, nullptr, APCGVolume::StaticClass());
+		PCGTestsCommon::GenerateSettings<UPCGCreatePointsGridSettings>(TestData);
+		UPCGCreatePointsGridSettings* Settings = CastChecked<UPCGCreatePointsGridSettings>(TestData.Settings);
+		Settings->PointPosition = Position;
+		Settings->GridPivot = Pivot;
+		Settings->GridExtents = FVector(100.0, 100.0, 50.0);
+		Settings->CellSize = FVector(100.0);
+
+		return TestData;
+	}
+}
+
+IMPLEMENT_CUSTOM_SIMPLE_AUTOMATION_TEST(FPCGCreatePointsGridTest_CellCenter, FPCGTestBaseClass, "pcg.tests.CreatePointsGrid.CellCenter", PCGTestsCommon::TestFlags)
+
+bool FPCGCreatePointsGridTest_CellCenter::RunTest(const FString& Parameters)
+{
+	TUniquePtr<FPCGContext> Context = CreatePointsGridCommonTests::GenerateTestDataAndRunGrid(EPCGPointPosition::CellCenter, FVector(500.0, 500.0, 50.0), FVector(100.0));
+
+	// Testing using the default settings
+	const TArray<FPCGTaggedData>& Outputs = Context->OutputData.TaggedData;
+
+	UTEST_EQUAL("Output count", Outputs.Num(), 1);
+
+	const UPCGPointData* OutPointData = Cast<UPCGPointData>(Outputs[0].Data);
+
+	UTEST_NOT_NULL("Output point data", OutPointData);
+
+	const TArray<FPCGPoint>& OutPoints = OutPointData->GetPoints();
+
+	UTEST_EQUAL("Output point count", OutPoints.Num(), 100);
+
+	UTEST_EQUAL("Point 0 has the correct transform,", OutPoints[0].Transform, FTransform(FRotator::ZeroRotator, FVector(-450.0, -450.0, 0), FVector::One()));
+	UTEST_EQUAL("Point 23 has the correct transform,", OutPoints[23].Transform, FTransform(FRotator::ZeroRotator, FVector(-150.0, -250.0, 0), FVector::One()));
+	UTEST_EQUAL("Point 65 has the correct transform,", OutPoints[65].Transform, FTransform(FRotator::ZeroRotator, FVector(50.0, 150.0, 0), FVector::One()));
+	UTEST_EQUAL("Point 92 has the correct transform,", OutPoints[92].Transform, FTransform(FRotator::ZeroRotator, FVector(-250.0, 450.0, 0), FVector::One()));
+
+	return true;
+}
+
+IMPLEMENT_CUSTOM_SIMPLE_AUTOMATION_TEST(FPCGCreatePointsGridTest_CellCorner, FPCGTestBaseClass, "pcg.tests.CreatePointsGrid.CellCorner", PCGTestsCommon::TestFlags)
+
+bool FPCGCreatePointsGridTest_CellCorner::RunTest(const FString& Parameters)
+{
+	TUniquePtr<FPCGContext> Context = CreatePointsGridCommonTests::GenerateTestDataAndRunGrid(EPCGPointPosition::CellCorners, FVector(500.0, 500.0, 50.0), FVector(100.0));
+
+	// Testing using the default sizes on the corners
+	const TArray<FPCGTaggedData>& Outputs = Context->OutputData.TaggedData;
+
+	UTEST_EQUAL("Output count", Outputs.Num(), 1);
+
+	const UPCGPointData* OutPointData = Cast<UPCGPointData>(Outputs[0].Data);
+
+	UTEST_NOT_NULL("Output point data", OutPointData);
+
+	const TArray<FPCGPoint>& OutPoints = OutPointData->GetPoints();
+
+	// Cell Corners generate (10 + 1) * (10 + 1) * 2
+	UTEST_EQUAL("Output point count", OutPoints.Num(), 242);
+
+	UTEST_EQUAL("Point 0 has the correct transform,", OutPoints[0].Transform, FTransform(FRotator::ZeroRotator, FVector(-500.0, -500.0, -50.0), FVector::One()));
+	UTEST_EQUAL("Point 23 has the correct transform,", OutPoints[23].Transform, FTransform(FRotator::ZeroRotator, FVector(-400.0, -300.0, -50.0), FVector::One()));
+	UTEST_EQUAL("Point 65 has the correct transform,", OutPoints[65].Transform, FTransform(FRotator::ZeroRotator, FVector(500.0, 0.0, -50.0), FVector::One()));
+	UTEST_EQUAL("Point 92 has the correct transform,", OutPoints[92].Transform, FTransform(FRotator::ZeroRotator, FVector(-100.0, 300.0, -50.0), FVector::One()));
+	UTEST_EQUAL("Point 145 has the correct transform,", OutPoints[145].Transform, FTransform(FRotator::ZeroRotator, FVector(-300.0, -300.0, 50.0), FVector::One()));
+
+	return true;
+}
+
+IMPLEMENT_CUSTOM_SIMPLE_AUTOMATION_TEST(FPCGCreatePointsGridTest_XYCenter, FPCGTestBaseClass, "pcg.tests.CreatePointsGrid.XYCenter", PCGTestsCommon::TestFlags)
+
+bool FPCGCreatePointsGridTest_XYCenter::RunTest(const FString& Parameters)
+{
+	TUniquePtr<FPCGContext> Context = CreatePointsGridCommonTests::GenerateTestDataAndRunGrid(EPCGPointPosition::CellCenter, FVector(100.0, 100.0, 50.0), FVector(100.0));
+
+	// Testing using a small grid in "2D"
+	const TArray<FPCGTaggedData>& Outputs = Context->OutputData.TaggedData;
+
+	UTEST_EQUAL("Output count", Outputs.Num(), 1);
+
+	const UPCGPointData* OutPointData = Cast<UPCGPointData>(Outputs[0].Data);
+
+	UTEST_NOT_NULL("Output point data", OutPointData);
+
+	const TArray<FPCGPoint>& OutPoints = OutPointData->GetPoints();
+
+	UTEST_EQUAL("Output point count", OutPoints.Num(), 4);
+
+	UTEST_EQUAL("Point 0 has the correct transform,", OutPoints[0].Transform, FTransform(FRotator::ZeroRotator, FVector(-50.0, -50.0, 0.0), FVector::One()));
+	UTEST_EQUAL("Point 1 has the correct transform,", OutPoints[1].Transform, FTransform(FRotator::ZeroRotator, FVector(50.0, -50.0, 0.0), FVector::One()));
+	UTEST_EQUAL("Point 2 has the correct transform,", OutPoints[2].Transform, FTransform(FRotator::ZeroRotator, FVector(-50.0, 50.0, 0.0), FVector::One()));
+	UTEST_EQUAL("Point 3 has the correct transform,", OutPoints[3].Transform, FTransform(FRotator::ZeroRotator, FVector(50.0, 50.0, 0.0), FVector::One()));
+
+	return true;
+}
+
+IMPLEMENT_CUSTOM_SIMPLE_AUTOMATION_TEST(FPCGCreatePointsGridTest_XYCorner, FPCGTestBaseClass, "pcg.tests.CreatePointsGrid.XYCorner", PCGTestsCommon::TestFlags)
+
+bool FPCGCreatePointsGridTest_XYCorner::RunTest(const FString& Parameters)
+{
+	TUniquePtr<FPCGContext> Context = CreatePointsGridCommonTests::GenerateTestDataAndRunGrid(EPCGPointPosition::CellCorners, FVector(100.0, 100.0, 0.0), FVector(100.0));
+
+	// Testing using a small grid in "2D"
+	const TArray<FPCGTaggedData>& Outputs = Context->OutputData.TaggedData;
+
+	UTEST_EQUAL("Output count", Outputs.Num(), 1);
+
+	const UPCGPointData* OutPointData = Cast<UPCGPointData>(Outputs[0].Data);
+
+	UTEST_NOT_NULL("Output point data", OutPointData);
+
+	const TArray<FPCGPoint>& OutPoints = OutPointData->GetPoints();
+
+	UTEST_EQUAL("Output point count", OutPoints.Num(), 9);
+
+	UTEST_EQUAL("Point 1 has the correct transform,", OutPoints[1].Transform, FTransform(FRotator::ZeroRotator, FVector(0.0, -100.0, 0.0), FVector::One()));
+	UTEST_EQUAL("Point 3 has the correct transform,", OutPoints[3].Transform, FTransform(FRotator::ZeroRotator, FVector(-100.0, 0.0, 0.0), FVector::One()));
+	UTEST_EQUAL("Point 5 has the correct transform,", OutPoints[5].Transform, FTransform(FRotator::ZeroRotator, FVector(100.0, 0.0, 0.0), FVector::One()));
+	UTEST_EQUAL("Point 7 has the correct transform,", OutPoints[7].Transform, FTransform(FRotator::ZeroRotator, FVector(0.0, 100.0, 0.0), FVector::One()));
+
+	return true;
+}
+
+IMPLEMENT_CUSTOM_SIMPLE_AUTOMATION_TEST(FPCGCreatePointsGridTest_XZCenter, FPCGTestBaseClass, "pcg.tests.CreatePointsGrid.XZCenter", PCGTestsCommon::TestFlags)
+
+bool FPCGCreatePointsGridTest_XZCenter::RunTest(const FString& Parameters)
+{
+	TUniquePtr<FPCGContext> Context = CreatePointsGridCommonTests::GenerateTestDataAndRunGrid(EPCGPointPosition::CellCenter, FVector(100.0, 50.0, 100.0), FVector(100.0));
+
+	// Testing using a small grid in "2D"
+	const TArray<FPCGTaggedData>& Outputs = Context->OutputData.TaggedData;
+
+	UTEST_EQUAL("Output count", Outputs.Num(), 1);
+
+	const UPCGPointData* OutPointData = Cast<UPCGPointData>(Outputs[0].Data);
+
+	UTEST_NOT_NULL("Output point data", OutPointData);
+
+	const TArray<FPCGPoint>& OutPoints = OutPointData->GetPoints();
+
+	UTEST_EQUAL("Output point count", OutPoints.Num(), 4);
+
+	UTEST_EQUAL("Point 0 has the correct transform,", OutPoints[0].Transform, FTransform(FRotator::ZeroRotator, FVector(-50.0, 0.0, -50.0), FVector::One()));
+	UTEST_EQUAL("Point 1 has the correct transform,", OutPoints[1].Transform, FTransform(FRotator::ZeroRotator, FVector(50.0, 0.0, -50.0), FVector::One()));
+	UTEST_EQUAL("Point 2 has the correct transform,", OutPoints[2].Transform, FTransform(FRotator::ZeroRotator, FVector(-50.0, 0.0, 50.0), FVector::One()));
+	UTEST_EQUAL("Point 3 has the correct transform,", OutPoints[3].Transform, FTransform(FRotator::ZeroRotator, FVector(50.0, 0.0, 50.0), FVector::One()));
+	
+	return true;
+}
+
+IMPLEMENT_CUSTOM_SIMPLE_AUTOMATION_TEST(FPCGCreatePointsGridTest_XZCorner, FPCGTestBaseClass, "pcg.tests.CreatePointsGrid.XZCorner", PCGTestsCommon::TestFlags)
+
+bool FPCGCreatePointsGridTest_XZCorner::RunTest(const FString& Parameters)
+{
+	TUniquePtr<FPCGContext> Context = CreatePointsGridCommonTests::GenerateTestDataAndRunGrid(EPCGPointPosition::CellCorners, FVector(100.0, 0.0, 100.0), FVector(100.0));
+
+	// Testing using a small grid in "2D"
+	const TArray<FPCGTaggedData>& Outputs = Context->OutputData.TaggedData;
+
+	UTEST_EQUAL("Output count", Outputs.Num(), 1);
+
+	const UPCGPointData* OutPointData = Cast<UPCGPointData>(Outputs[0].Data);
+
+	UTEST_NOT_NULL("Output point data", OutPointData);
+
+	const TArray<FPCGPoint>& OutPoints = OutPointData->GetPoints();
+
+	UTEST_EQUAL("Output point count", OutPoints.Num(), 9);
+
+	UTEST_EQUAL("Point 1 has the correct transform,", OutPoints[1].Transform, FTransform(FRotator::ZeroRotator, FVector(0.0, 0.0, -100.0), FVector::One()));
+	UTEST_EQUAL("Point 3 has the correct transform,", OutPoints[3].Transform, FTransform(FRotator::ZeroRotator, FVector(-100.0, 0.0, 0.0), FVector::One()));
+	UTEST_EQUAL("Point 5 has the correct transform,", OutPoints[5].Transform, FTransform(FRotator::ZeroRotator, FVector(100.0, 0.0, 0.0), FVector::One()));
+	UTEST_EQUAL("Point 7 has the correct transform,", OutPoints[7].Transform, FTransform(FRotator::ZeroRotator, FVector(0.0, 0.0, 100.0), FVector::One()));
+
+	return true;
+}
+
+IMPLEMENT_CUSTOM_SIMPLE_AUTOMATION_TEST(FPCGCreatePointsGridTest_YZCenter, FPCGTestBaseClass, "pcg.tests.CreatePointsGrid.YZCenter", PCGTestsCommon::TestFlags)
+
+bool FPCGCreatePointsGridTest_YZCenter::RunTest(const FString& Parameters)
+{
+	TUniquePtr<FPCGContext> Context = CreatePointsGridCommonTests::GenerateTestDataAndRunGrid(EPCGPointPosition::CellCenter, FVector(50.0, 100.0, 100.0), FVector(100.0));
+
+	// Testing using a small grid in "2D"
+	const TArray<FPCGTaggedData>& Outputs = Context->OutputData.TaggedData;
+
+	UTEST_EQUAL("Output count", Outputs.Num(), 1);
+
+	const UPCGPointData* OutPointData = Cast<UPCGPointData>(Outputs[0].Data);
+
+	UTEST_NOT_NULL("Output point data", OutPointData);
+
+	const TArray<FPCGPoint>& OutPoints = OutPointData->GetPoints();
+
+	UTEST_EQUAL("Output point count", OutPoints.Num(), 4);
+
+	UTEST_EQUAL("Point 0 has the correct transform,", OutPoints[0].Transform, FTransform(FRotator::ZeroRotator, FVector(0.0, -50.0, -50.0), FVector::One()));
+	UTEST_EQUAL("Point 1 has the correct transform,", OutPoints[1].Transform, FTransform(FRotator::ZeroRotator, FVector(0.0, 50.0, -50.0), FVector::One()));
+	UTEST_EQUAL("Point 2 has the correct transform,", OutPoints[2].Transform, FTransform(FRotator::ZeroRotator, FVector(0.0, -50.0, 50.0), FVector::One()));
+	UTEST_EQUAL("Point 3 has the correct transform,", OutPoints[3].Transform, FTransform(FRotator::ZeroRotator, FVector(0.0, 50.0, 50.0), FVector::One()));
+
+	return true;
+}
+
+IMPLEMENT_CUSTOM_SIMPLE_AUTOMATION_TEST(FPCGCreatePointsGridTest_YZCorner, FPCGTestBaseClass, "pcg.tests.CreatePointsGrid.YZCorner", PCGTestsCommon::TestFlags)
+
+bool FPCGCreatePointsGridTest_YZCorner::RunTest(const FString& Parameters)
+{
+	TUniquePtr<FPCGContext> Context = CreatePointsGridCommonTests::GenerateTestDataAndRunGrid(EPCGPointPosition::CellCorners, FVector(0.0, 100.0, 100.0), FVector(100.0));
+
+	// Testing using a small grid in "2D"
+	const TArray<FPCGTaggedData>& Outputs = Context->OutputData.TaggedData;
+
+	UTEST_EQUAL("Output count", Outputs.Num(), 1);
+
+	const UPCGPointData* OutPointData = Cast<UPCGPointData>(Outputs[0].Data);
+
+	UTEST_NOT_NULL("Output point data", OutPointData);
+
+	const TArray<FPCGPoint>& OutPoints = OutPointData->GetPoints();
+
+	UTEST_EQUAL("Output point count", OutPoints.Num(), 9);
+
+	UTEST_EQUAL("Point 1 has the correct transform,", OutPoints[1].Transform, FTransform(FRotator::ZeroRotator, FVector(0.0, 0.0, -100.0), FVector::One()));
+	UTEST_EQUAL("Point 3 has the correct transform,", OutPoints[3].Transform, FTransform(FRotator::ZeroRotator, FVector(0.0, -100.0, 0.0), FVector::One()));
+	UTEST_EQUAL("Point 5 has the correct transform,", OutPoints[5].Transform, FTransform(FRotator::ZeroRotator, FVector(0.0, 100.0, 0.0), FVector::One()));
+	UTEST_EQUAL("Point 7 has the correct transform,", OutPoints[7].Transform, FTransform(FRotator::ZeroRotator, FVector(0.0, 0.0, 100.0), FVector::One()));
+
+	return true;
+}
+
+IMPLEMENT_CUSTOM_SIMPLE_AUTOMATION_TEST(FPCGCreatePointsGridTest_NegativeValuesGrid, FPCGTestBaseClass, "pcg.tests.CreatePointsGrid.NegativeValuesGrid", PCGTestsCommon::TestFlags)
+
+bool FPCGCreatePointsGridTest_NegativeValuesGrid::RunTest(const FString& Parameters)
+{
+	AddExpectedError(TEXT("must not be less than 0"));
+	
+	TUniquePtr<FPCGContext> Context = CreatePointsGridCommonTests::GenerateTestDataAndRunGrid(EPCGPointPosition::CellCenter, FVector(-1.0), FVector(100.0));
+
+	// Testing for negative values, should produce no outputs
+	const TArray<FPCGTaggedData>& Outputs = Context->OutputData.TaggedData;
+
+	UTEST_EQUAL("Output count", Outputs.Num(), 0);
+
+	return true;
+}
+
+IMPLEMENT_CUSTOM_SIMPLE_AUTOMATION_TEST(FPCGCreatePointsGridTest_NegativeValuesCell, FPCGTestBaseClass, "pcg.tests.CreatePointsGrid.NegativeValuesCell", PCGTestsCommon::TestFlags)
+
+bool FPCGCreatePointsGridTest_NegativeValuesCell::RunTest(const FString& Parameters)
+{
+	AddExpectedError(TEXT("must not be less than 0"));
+
+	TUniquePtr<FPCGContext> Context = CreatePointsGridCommonTests::GenerateTestDataAndRunGrid(EPCGPointPosition::CellCenter, FVector(100.0), FVector(-1.0));
+
+	// Testing for negative values, should produce no outputs
+	const TArray<FPCGTaggedData>& Outputs = Context->OutputData.TaggedData;
+
+	UTEST_EQUAL("Output count", Outputs.Num(), 0);
+
+	return true;
+}
+
+IMPLEMENT_CUSTOM_SIMPLE_AUTOMATION_TEST(FPCGCreatePointsGridTest_OversizedCell, FPCGTestBaseClass, "pcg.tests.CreatePointsGrid.OversizedCell", PCGTestsCommon::TestFlags)
+
+bool FPCGCreatePointsGridTest_OversizedCell::RunTest(const FString& Parameters)
+{
+	TUniquePtr<FPCGContext> Context = CreatePointsGridCommonTests::GenerateTestDataAndRunGrid(EPCGPointPosition::CellCenter, FVector(100.0), FVector(400.0));
+
+	// Testing for negative values, should produce no outputs
+	const TArray<FPCGTaggedData>& Outputs = Context->OutputData.TaggedData;
+
+	UTEST_EQUAL("Output count", Outputs.Num(), 1);
+
+	const UPCGPointData* OutPointData = Cast<UPCGPointData>(Outputs[0].Data);
+
+	UTEST_NOT_NULL("Output point data", OutPointData);
+
+	const TArray<FPCGPoint>& OutPoints = OutPointData->GetPoints();
+
+	UTEST_EQUAL("Output point count", OutPoints.Num(), 1);
+
+	UTEST_EQUAL("Point 0 has the correct transform,", OutPoints[0].Transform, FTransform(FRotator::ZeroRotator, FVector(0.0, 0.0, 0.0), FVector::One()));
+
+	return true;
+}
+
+IMPLEMENT_CUSTOM_SIMPLE_AUTOMATION_TEST(FPCGCreatePointsGridTest_LocalCenter, FPCGTestBaseClass, "pcg.tests.CreatePointsGrid.LocalCenter", PCGTestsCommon::TestFlags)
+
+bool FPCGCreatePointsGridTest_LocalCenter::RunTest(const FString& Parameters)
+{
+	PCGTestsCommon::FTestData TestData = CreatePointsGridCommonTests::GenerateTestData(EPCGPointPosition::CellCenter, EPCGGridPivot::LocalComponent);
+
+	// Creating the volume and proper bounds
+	AActor* TestActor = TestData.TestActor;
+	TestActor->SetActorScale3D(FVector(50.0));
+	TestActor->SetActorRelativeLocation(FVector(100.0, 100.0, 0.0));
+	FBox TestBounds = TestActor->GetComponentsBoundingBox().ExpandBy(FVector(200.0));
+
+	FPCGElementPtr TestElement = TestData.Settings->GetElement();
+
+	TUniquePtr<FPCGContext> Context = TestData.InitializeTestContext();
+
+	while (!TestElement->Execute(Context.Get())) {}
+
+	// Testing using a small grid in local mode
+	const TArray<FPCGTaggedData>& Outputs = Context->OutputData.TaggedData;
+
+	UTEST_EQUAL("Output count", Outputs.Num(), 1);
+
+	const UPCGPointData* OutPointData = Cast<UPCGPointData>(Outputs[0].Data);
+
+	UTEST_NOT_NULL("Output point data", OutPointData);
+
+	const TArray<FPCGPoint>& OutPoints = OutPointData->GetPoints();
+
+	UTEST_EQUAL("Output point count", OutPoints.Num(), 4);
+
+	// Check to see if the position of the points are within the bounds of the volume
+	for (int i = 0; i < OutPoints.Num(); ++i)
+	{
+		UTEST_TRUE(FString::Format(TEXT("Point {0}, {1} is within the volume"), { i }), TestBounds.IsInsideOrOn(OutPoints[i].Transform.GetLocation()));
+	}
+
+	UTEST_EQUAL("Point 0 has the correct transform,", OutPoints[0].Transform, FTransform(FRotator::ZeroRotator, FVector(50.0, 50.0, 0.0), FVector::One()));
+	UTEST_EQUAL("Point 1 has the correct transform,", OutPoints[1].Transform, FTransform(FRotator::ZeroRotator, FVector(150.0, 50.0, 0.0), FVector::One()));
+	UTEST_EQUAL("Point 2 has the correct transform,", OutPoints[2].Transform, FTransform(FRotator::ZeroRotator, FVector(50.0, 150.0, 0.0), FVector::One()));
+	UTEST_EQUAL("Point 3 has the correct transform,", OutPoints[3].Transform, FTransform(FRotator::ZeroRotator, FVector(150.0, 150.0, 0.0), FVector::One()));
+	
+	return true;
+}
+
+IMPLEMENT_CUSTOM_SIMPLE_AUTOMATION_TEST(FPCGCreatePointsGridTest_LocalCorner, FPCGTestBaseClass, "pcg.tests.CreatePointsGrid.LocalCorner", PCGTestsCommon::TestFlags)
+
+bool FPCGCreatePointsGridTest_LocalCorner::RunTest(const FString& Parameters)
+{
+	PCGTestsCommon::FTestData TestData = CreatePointsGridCommonTests::GenerateTestData(EPCGPointPosition::CellCorners, EPCGGridPivot::LocalComponent);
+
+	// Creating the volume and proper bounds
+	AActor* TestActor = TestData.TestActor;
+	TestActor->SetActorScale3D(FVector(50.0));
+	TestActor->SetActorRelativeLocation(FVector(100.0, 100.0, 0.0));
+	FBox TestBounds = TestActor->GetComponentsBoundingBox().ExpandBy(FVector(200.0));
+
+	FPCGElementPtr TestElement = TestData.Settings->GetElement();
+
+	TUniquePtr<FPCGContext> Context = TestData.InitializeTestContext();
+
+	while (!TestElement->Execute(Context.Get())) {}
+
+	// Testing using a small grid in local mode
+	const TArray<FPCGTaggedData>& Outputs = Context->OutputData.TaggedData;
+
+	UTEST_EQUAL("Output count", Outputs.Num(), 1);
+
+	const UPCGPointData* OutPointData = Cast<UPCGPointData>(Outputs[0].Data);
+
+	UTEST_NOT_NULL("Output point data", OutPointData);
+
+	const TArray<FPCGPoint>& OutPoints = OutPointData->GetPoints();
+
+	//Cell Corners generate (2 + 1) * (2 + 1) * 2
+	UTEST_EQUAL("Output point count", OutPoints.Num(), 18);
+
+	// Check to see if the position of the points are within the bounds of the volume
+	for (int i = 0; i < OutPoints.Num(); ++i)
+	{
+		UTEST_TRUE(FString::Format(TEXT("Point {0} is within the volume"), { i }), TestBounds.IsInsideOrOn(OutPoints[i].Transform.GetLocation()));
+	}
+
+	UTEST_EQUAL("Point 1 has the correct transform,", OutPoints[1].Transform, FTransform(FRotator::ZeroRotator, FVector(100.0, 0.0, -50.0), FVector::One()));
+	UTEST_EQUAL("Point 3 has the correct transform,", OutPoints[3].Transform, FTransform(FRotator::ZeroRotator, FVector(0.0, 100.0, -50.0), FVector::One()));
+	UTEST_EQUAL("Point 5 has the correct transform,", OutPoints[5].Transform, FTransform(FRotator::ZeroRotator, FVector(200.0, 100.0, -50.0), FVector::One()));
+	UTEST_EQUAL("Point 11 has the correct transform,", OutPoints[11].Transform, FTransform(FRotator::ZeroRotator, FVector(200.0, 0.0, 50.0), FVector::One()));
+
+	return true;
+}
\ No newline at end of file
diff --git a/Engine/Plugins/Experimental/PCG/Source/PCG/Private/Tests/PCGTestsCommon.cpp b/Engine/Plugins/Experimental/PCG/Source/PCG/Private/Tests/PCGTestsCommon.cpp
--- a/Engine/Plugins/Experimental/PCG/Source/PCG/Private/Tests/PCGTestsCommon.cpp
+++ b/Engine/Plugins/Experimental/PCG/Source/PCG/Private/Tests/PCGTestsCommon.cpp
@@ -39,12 +39,20 @@ namespace PCGTestsCommon
 		TestActor = EditorWorld->SpawnActor<AActor>(ActorClass, TransientActorParameters);
 		check(TestActor);
 
-		TestPCGComponent = NewObject<UPCGComponent>(TestActor, FName(TEXT("Test PCG Component")), RF_Transient);
-		check(TestPCGComponent);
+		if (UPCGComponent* PCGComponent = TestActor->GetComponentByClass<UPCGComponent>())
+		{
+			TestPCGComponent = PCGComponent;
+		}
+		else
+		{
+			TestPCGComponent = NewObject<UPCGComponent>(TestActor, FName(TEXT("Test PCG Component")), RF_Transient);
+			check(TestPCGComponent);
+			TestActor->AddInstanceComponent(TestPCGComponent);
+			TestPCGComponent->RegisterComponent();
+		}
+
 		// By default PCG components for tests will be non-partitioned
 		TestPCGComponent->SetIsPartitioned(false);
-		TestActor->AddInstanceComponent(TestPCGComponent);
-		TestPCGComponent->RegisterComponent();
 
 		UPCGGraph* TestGraph = NewObject<UPCGGraph>(TestPCGComponent, FName(TEXT("Test PCG Graph")), RF_Transient);
 		check(TestGraph);
diff --git a/Engine/Plugins/Experimental/PCG/Source/PCG/Public/Elements/PCGCreatePoints.h b/Engine/Plugins/Experimental/PCG/Source/PCG/Public/Elements/PCGCreatePoints.h
--- a/Engine/Plugins/Experimental/PCG/Source/PCG/Public/Elements/PCGCreatePoints.h
+++ b/Engine/Plugins/Experimental/PCG/Source/PCG/Public/Elements/PCGCreatePoints.h
@@ -6,6 +6,14 @@
 #include "Data/PCGPointData.h"
 #include "PCGCreatePoints.generated.h"
 
+UENUM()
+enum class EPCGLocalGridPivot : uint8
+{
+	Global,
+	OriginalComponent,
+	LocalComponent
+};
+
 /**
  * Creates point data from a provided list of points.
  */
@@ -34,10 +42,11 @@ public:
 	UPROPERTY(BlueprintReadWrite, EditAnywhere, Category = Settings, meta = (PCG_Overridable))
 	TArray<FPCGPoint> PointsToCreate;
 
-	/** If true, points are transformed to world space using the PCG component transform */
+	/** Sets the points transform to world or local space*/
 	UPROPERTY(BlueprintReadWrite, EditAnywhere, Category = Settings, meta = (PCG_Overridable))
-	bool bLocal = false;
+	EPCGLocalGridPivot GridPivot = EPCGLocalGridPivot::Global;
 
+	/** If true, points are removed if they are outside of the volume */
 	UPROPERTY(BlueprintReadWrite, EditAnywhere, Category = Settings, meta = (PCG_Overridable))
 	bool bCullPointsOutsideVolume = false;
 };
diff --git a/Engine/Plugins/Experimental/PCG/Source/PCG/Public/Elements/PCGCreatePointsGrid.h b/Engine/Plugins/Experimental/PCG/Source/PCG/Public/Elements/PCGCreatePointsGrid.h
new file mode 100644
--- /dev/null
+++ b/Engine/Plugins/Experimental/PCG/Source/PCG/Public/Elements/PCGCreatePointsGrid.h
@@ -0,0 +1,75 @@
+// Copyright Epic Games, Inc. All Rights Reserved.
+#pragma once
+
+#include "PCGSettings.h"
+
+#include "Data/PCGPointData.h"
+#include "PCGCreatePointsGrid.generated.h"
+
+UENUM()
+enum class EPCGPointPosition : uint8
+{
+	CellCenter,
+	CellCorners
+};
+
+UENUM()
+enum class EPCGGridPivot : uint8
+{
+	Global,
+	OriginalComponent,
+	LocalComponent
+};
+
+/**
+ * Creates a 2D or 3D grid of points.
+ */
+UCLASS(BlueprintType, ClassGroup = (Procedural))
+class UPCGCreatePointsGridSettings : public UPCGSettings
+{
+	GENERATED_BODY()
+
+public:
+	//~Begin UPCGSettings interface
+#if WITH_EDITOR
+	virtual FName GetDefaultNodeName() const override { return FName(TEXT("CreatePointsGrid")); }
+	virtual FText GetDefaultNodeTitle() const override { return NSLOCTEXT("PCGCreatePointsGridElement", "NodeTitle", "Create Points Grid"); }
+	virtual FText GetNodeTooltipText() const override { return NSLOCTEXT("PCGCreatePointsGridElement", "NodeTooltip", "Creates a 2D or 3D grid of points."); }
+	virtual EPCGSettingsType GetType() const override { return EPCGSettingsType::Spatial; }
+#endif
+
+protected:
+	virtual TArray<FPCGPinProperties> InputPinProperties() const override;
+	virtual TArray<FPCGPinProperties> OutputPinProperties() const override { return Super::DefaultPointOutputPinProperties(); }
+	virtual FPCGElementPtr CreateElement() const override;
+	//~End UPCGSettings interface
+
+public:
+	UPROPERTY(BlueprintReadWrite, EditAnywhere, Category = Settings, meta = (PCG_Overridable, ClampMin = "0.0", UIMin = "0.0"))
+	FVector GridExtents = FVector(500.0, 500.0, 50.0);
+	
+	UPROPERTY(BlueprintReadWrite, EditAnywhere, Category = Settings, meta = (PCG_Overridable, ClampMin = "0.0", UIMin = "0.0"))
+	FVector CellSize = FVector(100.0, 100.0, 100.0);
+
+	/** Sets the points transform to world or local space*/
+	UPROPERTY(BlueprintReadWrite, EditAnywhere, Category = Settings, meta = (PCG_Overridable))
+	EPCGGridPivot GridPivot = EPCGGridPivot::Global;
+
+	/** If true, the bounds of the points are set to 50.0, if false, 1.0 */
+	UPROPERTY(BlueprintReadWrite, EditAnywhere, Category = Settings, meta = (PCG_Overridable))
+	bool bSetPointsBounds = true;
+
+	/** If true, points are removed if they are outside of the volume */
+	UPROPERTY(BlueprintReadWrite, EditAnywhere, Category = Settings, meta = (PCG_Overridable))
+	bool bCullPointsOutsideVolume = false;
+	
+	UPROPERTY(BlueprintReadWrite, EditAnywhere, Category = Settings, meta = (PCG_Overridable))
+	EPCGPointPosition PointPosition = EPCGPointPosition::CellCenter;
+};
+
+class FPCGCreatePointsGridElement : public IPCGElement
+{
+protected:
+	virtual bool ExecuteInternal(FPCGContext* Context) const override;
+	virtual void GetDependenciesCrc(const FPCGDataCollection& InInput, const UPCGSettings* InSettings, UPCGComponent* InComponent, FPCGCrc& OutCrc) const override;
+};
\ No newline at end of file
