Tentaive fix for Non-Unity Build

diff --git a/Engine/Plugins/Experimental/PCG/Source/PCG/Private/Elements/PCGCreatePointsGrid.cpp b/Engine/Plugins/Experimental/PCG/Source/PCG/Private/Elements/PCGCreatePointsGrid.cpp
--- a/Engine/Plugins/Experimental/PCG/Source/PCG/Private/Elements/PCGCreatePointsGrid.cpp
+++ b/Engine/Plugins/Experimental/PCG/Source/PCG/Private/Elements/PCGCreatePointsGrid.cpp
@@ -8,6 +8,8 @@
 #include "Helpers/PCGHelpers.h"
 #include "Helpers/PCGSettingsHelpers.h"
 
+#include "GameFramework/Actor.h"
+
 #define LOCTEXT_NAMESPACE "PCGCreatePointsGridElement"
 
 TArray<FPCGPinProperties> UPCGCreatePointsGridSettings::InputPinProperties() const
