[PCG] Make sure we have the same default value than the first entry when doing add attrbiutes and get attribute set by index

#rnx

diff --git a/Engine/Plugins/Experimental/PCG/Source/PCG/Private/Elements/PCGAttributeGetFromIndexElement.cpp b/Engine/Plugins/Experimental/PCG/Source/PCG/Private/Elements/PCGAttributeGetFromIndexElement.cpp
--- a/Engine/Plugins/Experimental/PCG/Source/PCG/Private/Elements/PCGAttributeGetFromIndexElement.cpp
+++ b/Engine/Plugins/Experimental/PCG/Source/PCG/Private/Elements/PCGAttributeGetFromIndexElement.cpp
@@ -5,6 +5,8 @@
 #include "PCGContext.h"
 #include "PCGParamData.h"
 #include "PCGPin.h"
+#include "Metadata/PCGMetadata.h"
+#include "Metadata/PCGMetadataAttribute.h"
 
 #include UE_INLINE_GENERATED_CPP_BY_NAME(PCGAttributeGetFromIndexElement)
 
@@ -86,6 +88,18 @@ bool FPCGAttributeGetFromIndexElement::ExecuteInternal(FPCGContext* Context) con
 		PCGMetadataEntryKey SingleEntryKey = SubParam->Metadata->AddEntry();
 		SubParam->Metadata->SetAttributes(OriginalEntryKey, ParamMetadata, SingleEntryKey);
 
+		TArray<FName> AttributeNames;
+		TArray<EPCGMetadataTypes> AttributeTypes;
+		SubParam->Metadata->GetAttributes(AttributeNames, AttributeTypes);
+
+		for (const FName AttributeName : AttributeNames)
+		{
+			if (FPCGMetadataAttributeBase* Attribute = SubParam->Metadata->GetMutableAttribute(AttributeName))
+			{
+				Attribute->SetDefaultValueToFirstEntry();
+			}
+		}
+
 		FPCGTaggedData& TaggedData = Context->OutputData.TaggedData.Emplace_GetRef(Input);
 		TaggedData.Data = SubParam;
 	}
diff --git a/Engine/Plugins/Experimental/PCG/Source/PCG/Private/Elements/PCGCreateAttribute.cpp b/Engine/Plugins/Experimental/PCG/Source/PCG/Private/Elements/PCGCreateAttribute.cpp
--- a/Engine/Plugins/Experimental/PCG/Source/PCG/Private/Elements/PCGCreateAttribute.cpp
+++ b/Engine/Plugins/Experimental/PCG/Source/PCG/Private/Elements/PCGCreateAttribute.cpp
@@ -355,13 +355,27 @@ bool FPCGCreateAttributeElement::ExecuteInternal(FPCGContext* Context) const
 			const FPCGAttributePropertyInputSelector InputSource = Settings->InputSource.CopyAndFixLast(SourceParamData);
 
 			// We need accessors if we have a multi entry source attribute or we have extractors
-			const bool bNeedAccessors = SourceParamData->Metadata->GetLocalItemCount() > 1 || !InputSource.GetExtraNames().IsEmpty();
+			const bool bIsMultiEntries = SourceParamData->Metadata->GetLocalItemCount() > 1;
+			const bool bNeedAccessors = bIsMultiEntries || !InputSource.GetExtraNames().IsEmpty();
 
 			// If no accessor, copy over the attribute
 			if (!bNeedAccessors)
 			{
 				const FPCGMetadataAttributeBase* SourceAttribute = SourceParamData->Metadata->GetConstAttribute(SourceParamAttributeName);
 				Attribute = Metadata->CopyAttribute(SourceAttribute, OutputAttributeName, /*bKeepParent=*/false, /*bCopyEntries=*/false, /*bCopyValues=*/false);
+
+				// We perhaps need to fix the default value. If the first entry is different from the default value, we override the default value.
+				if (Attribute && Attribute->IsEqualToDefaultValue(Attribute->GetValueKey(0)))
+				{
+					auto FixDefaultValue = [Attribute](auto Dummy)
+					{
+						using AttributeType = decltype(Dummy);
+						FPCGMetadataAttribute<AttributeType>* TypedAttribute = static_cast<FPCGMetadataAttribute<AttributeType>*>(Attribute);
+						TypedAttribute->SetDefaultValue(TypedAttribute->GetValueFromItemKey(0));
+					};
+
+					PCGMetadataAttribute::CallbackWithRightType(Attribute->GetTypeId(), FixDefaultValue);
+				}
 			}
 			else // Create a new attribute of the accessed field's type manually
 			{
@@ -373,12 +387,17 @@ bool FPCGCreateAttributeElement::ExecuteInternal(FPCGContext* Context) const
 					return true;
 				}
 
-				auto CreateOutputAttribute = [Metadata, OutputAttributeName, &InputAccessor, &InputKeys, OutputData, bIsParamData]<typename Type>(Type Dummy) -> FPCGMetadataAttributeBase*
+				auto CreateOutputAttribute = [Metadata, OutputAttributeName, &InputAccessor, &InputKeys, OutputData, bIsParamData, bIsMultiEntries]<typename Type>(Type Dummy) -> FPCGMetadataAttributeBase*
 				{
-					// Get the value from the input accessor default value and pass that as the default value
-					Type Value{};
-					InputAccessor->Get<Type>(Value, FPCGAttributeAccessorKeysEntries(PCGInvalidEntryKey));
-					FPCGMetadataAttribute<Type>* Attribute = PCGMetadataElementCommon::ClearOrCreateAttribute<Type>(Metadata, OutputAttributeName, Value);
+					// If we have multiple entries, use the zero value as a default value.
+					// Otherwise get the first entry of the source param as default value.
+					Type DefaultValue = PCG::Private::MetadataTraits<Type>::ZeroValue();
+					if (!bIsMultiEntries)
+					{
+						InputAccessor->Get<Type>(DefaultValue, FPCGAttributeAccessorKeysEntries(PCGMetadataEntryKey(0)));
+					}
+
+					FPCGMetadataAttribute<Type>* Attribute = PCGMetadataElementCommon::ClearOrCreateAttribute<Type>(Metadata, OutputAttributeName, DefaultValue);
 
 					FPCGAttributePropertySelector OutputSelector;
 					OutputSelector.SetAttributeName(OutputAttributeName);
diff --git a/Engine/Plugins/Experimental/PCG/Source/PCG/Private/Tests/Elements/PCGGetAttributeSetByIndexTest.cpp b/Engine/Plugins/Experimental/PCG/Source/PCG/Private/Tests/Elements/PCGGetAttributeSetByIndexTest.cpp
new file mode 100644
--- /dev/null
+++ b/Engine/Plugins/Experimental/PCG/Source/PCG/Private/Tests/Elements/PCGGetAttributeSetByIndexTest.cpp
@@ -0,0 +1,78 @@
+// Copyright Epic Games, Inc. All Rights Reserved.
+
+#include "Tests/PCGTestsCommon.h"
+
+#include "PCGContext.h"
+#include "PCGParamData.h"
+#include "Elements/PCGAttributeGetFromIndexElement.h"
+#include "Metadata/PCGMetadata.h"
+#include "Metadata/PCGMetadataAttributeTpl.h"
+
+IMPLEMENT_CUSTOM_SIMPLE_AUTOMATION_TEST(FPCGGetAttributeSetByIndexTest_Basic, FPCGTestBaseClass, "pcg.GetAttributeSetByIndex.Basic", PCGTestsCommon::TestFlags)
+
+bool FPCGGetAttributeSetByIndexTest_Basic::RunTest(const FString& Parameters)
+{
+	PCGTestsCommon::FTestData TestData;
+	UPCGAttributeGetFromIndexSettings* Settings = PCGTestsCommon::GenerateSettings<UPCGAttributeGetFromIndexSettings>(TestData);
+	Settings->Index = 1;
+	
+	FPCGElementPtr TestElement = TestData.Settings->GetElement();
+
+	static const FName DoubleAttrName = TEXT("Double");
+	static const FName IntAttrName = TEXT("Int");
+
+	// Data preparation
+	{
+		UPCGParamData* ParamData = NewObject<UPCGParamData>();
+
+		// Static analyser paranoid
+		check(ParamData && ParamData->Metadata);
+		PCGMetadataEntryKey FirstEntry = ParamData->Metadata->AddEntry();
+		PCGMetadataEntryKey SecondEntry = ParamData->Metadata->AddEntry();
+
+		FPCGMetadataAttribute<double>* DoubleAttr = ParamData->Metadata->CreateAttribute(DoubleAttrName, 0.0, true, false);
+		check(DoubleAttr);
+		DoubleAttr->SetValue(FirstEntry, 1.0);
+		DoubleAttr->SetValue(SecondEntry, 2.0);
+
+		FPCGMetadataAttribute<int>* IntAttr = ParamData->Metadata->CreateAttribute(IntAttrName, 0, true, false);
+		check(IntAttr);
+		IntAttr->SetValue(FirstEntry, 4);
+		IntAttr->SetValue(SecondEntry, 3);
+
+		FPCGTaggedData& SourcePin = TestData.InputData.TaggedData.Emplace_GetRef();
+		SourcePin.Pin = PCGPinConstants::DefaultInputLabel;
+		SourcePin.Data = ParamData;
+	}
+
+	TUniquePtr<FPCGContext> Context = TestData.InitializeTestContext();
+
+	while (!TestElement->Execute(Context.Get())) {}
+
+	const TArray<FPCGTaggedData>& Outputs = Context->OutputData.GetInputsByPin(PCGPinConstants::DefaultOutputLabel);
+	UTEST_EQUAL("Output count", Outputs.Num(), 1);
+
+	const UPCGParamData* OutParamData = Cast<const UPCGParamData>(Outputs[0].Data);
+	UTEST_NOT_NULL("Output Param data", OutParamData);
+	UTEST_NOT_NULL("Output Metadata", OutParamData->Metadata.Get());
+
+	check(OutParamData && OutParamData->Metadata);
+
+	UTEST_EQUAL("Metadata has 1 entry", OutParamData->Metadata->GetLocalItemCount(), 1);
+
+	FPCGMetadataAttribute<double>* OutDoubleAttr = OutParamData->Metadata->GetMutableTypedAttribute<double>(DoubleAttrName);
+	UTEST_NOT_NULL("Double attribute exists", OutDoubleAttr);
+
+	FPCGMetadataAttribute<int>* OutIntAttr = OutParamData->Metadata->GetMutableTypedAttribute<int>(IntAttrName);
+	UTEST_NOT_NULL("Int attribute exists", OutIntAttr);
+
+	check(OutDoubleAttr && OutIntAttr);
+
+	UTEST_EQUAL("Double attribute first entry", OutDoubleAttr->GetValueFromItemKey(PCGMetadataEntryKey(0)), 2.0);
+	UTEST_EQUAL("Double attribute default value was set to the first entry", OutDoubleAttr->GetValueFromItemKey(PCGInvalidEntryKey), 2.0);
+
+	UTEST_EQUAL("Int attribute first entry", OutIntAttr->GetValueFromItemKey(PCGMetadataEntryKey(0)), 3);
+	UTEST_EQUAL("Int attribute default value was set to the first entry", OutIntAttr->GetValueFromItemKey(PCGInvalidEntryKey), 3);
+
+	return true;
+}
\ No newline at end of file
diff --git a/Engine/Plugins/Experimental/PCG/Source/PCG/Public/Metadata/PCGMetadataAttribute.h b/Engine/Plugins/Experimental/PCG/Source/PCG/Public/Metadata/PCGMetadataAttribute.h
--- a/Engine/Plugins/Experimental/PCG/Source/PCG/Public/Metadata/PCGMetadataAttribute.h
+++ b/Engine/Plugins/Experimental/PCG/Source/PCG/Public/Metadata/PCGMetadataAttribute.h
@@ -39,6 +39,8 @@ public:
 	virtual void SetWeightedValue(PCGMetadataEntryKey ItemKey, const FPCGMetadataAttributeBase* InAttribute, const TArrayView<TPair<PCGMetadataEntryKey, float>> InWeightedKeys) = 0;
 	virtual void SetValue(PCGMetadataEntryKey ItemKey, const FPCGMetadataAttributeBase* InAttributeA, PCGMetadataEntryKey InEntryKeyA, const FPCGMetadataAttributeBase* InAttributeB, PCGMetadataEntryKey InEntryKeyB, EPCGMetadataOp Op) = 0;
 	virtual bool IsEqualToDefaultValue(PCGMetadataValueKey ValueKey) const = 0;
+	/** In the case of multi entry attribute and after some operations, we might have a single entry attribute with a default value that is different than the first entry. Use this function to fix that. Only valid if there is one and only one value. */
+	virtual void SetDefaultValueToFirstEntry() = 0;
 
 	virtual bool UsesValueKeys() const = 0;
 	virtual bool AreValuesEqualForEntryKeys(PCGMetadataEntryKey EntryKey1, PCGMetadataEntryKey EntryKey2) const = 0;
diff --git a/Engine/Plugins/Experimental/PCG/Source/PCG/Public/Metadata/PCGMetadataAttributeTpl.h b/Engine/Plugins/Experimental/PCG/Source/PCG/Public/Metadata/PCGMetadataAttributeTpl.h
--- a/Engine/Plugins/Experimental/PCG/Source/PCG/Public/Metadata/PCGMetadataAttributeTpl.h
+++ b/Engine/Plugins/Experimental/PCG/Source/PCG/Public/Metadata/PCGMetadataAttributeTpl.h
@@ -331,6 +331,16 @@ public:
 		return PCG::Private::MetadataTraits<T>::Equal(GetValue(ValueKey), DefaultValue);
 	}
 
+	virtual void SetDefaultValueToFirstEntry() override
+	{
+		if (Values.Num() != 1)
+		{
+			return;
+		}
+
+		DefaultValue = Values[0];
+	}
+
 	PCGMetadataValueKey GetValueKeyOffsetForChild() const
 	{
 		FReadScopeLock ScopeLock(ValueLock);
