Tentative Fix for Build Engine Localization

#rnx

diff --git a/Engine/Plugins/Experimental/PCG/Source/PCG/Private/Elements/PCGCreatePointsGrid.cpp b/Engine/Plugins/Experimental/PCG/Source/PCG/Private/Elements/PCGCreatePointsGrid.cpp
--- a/Engine/Plugins/Experimental/PCG/Source/PCG/Private/Elements/PCGCreatePointsGrid.cpp
+++ b/Engine/Plugins/Experimental/PCG/Source/PCG/Private/Elements/PCGCreatePointsGrid.cpp
@@ -58,13 +58,13 @@ bool FPCGCreatePointsGridElement::ExecuteInternal(FPCGContext* Context) const
 
 	if (Settings->CellSize.X <= 0.0 || Settings->CellSize.Y <= 0.0 || Settings->CellSize.Z <= 0.0)
 	{
-		PCGE_LOG(Warning, GraphAndLog, LOCTEXT("InvalidDataInput", "CellSize must not be less than 0"));
+		PCGE_LOG(Warning, GraphAndLog, LOCTEXT("InvalidCellDataInput", "CellSize must not be less than 0"));
 		return true;
 	}
 
 	if (Settings->GridExtents.X < 0.0 || Settings->GridExtents.Y < 0.0 || Settings->GridExtents.Z < 0.0)
 	{
-		PCGE_LOG(Warning, GraphAndLog, LOCTEXT("InvalidDataInput", "GridExtents must not be less than 0"));
+		PCGE_LOG(Warning, GraphAndLog, LOCTEXT("InvalidGridDataInput", "GridExtents must not be less than 0"));
 		return true;
 	}
 
@@ -136,7 +136,7 @@ bool FPCGCreatePointsGridElement::ExecuteInternal(FPCGContext* Context) const
 
 	if (PCGFeatureSwitches::CVarCheckSamplerMemory.GetValueOnAnyThread() && FPlatformMemory::GetStats().AvailablePhysical < sizeof(FPCGPoint) * NumIterations64)
 	{
-		PCGE_LOG(Error, GraphAndLog, LOCTEXT("Overflow_int32", "The number of iterations produced is larger than available memory."));
+		PCGE_LOG(Error, GraphAndLog, LOCTEXT("MemoryOverflow", "The number of iterations produced is larger than available memory."));
 		return true;
 	}
 
