[StructUtils] [PCG] Implement a callback to catch changes in the internal layout to force the UI to update. Only do it on the PCG side, as it might be too intrusive in StructUtils. Generic solution will be pushed to 5.4.

* Without this, the property detail view might read in freed memory, causing a crash for non-POD types and UD for POD types.

diff --git a/Engine/Plugins/Experimental/PCG/Source/PCG/Private/PCGCommon.cpp b/Engine/Plugins/Experimental/PCG/Source/PCG/Private/PCGCommon.cpp
--- a/Engine/Plugins/Experimental/PCG/Source/PCG/Private/PCGCommon.cpp
+++ b/Engine/Plugins/Experimental/PCG/Source/PCG/Private/PCGCommon.cpp
@@ -91,3 +91,10 @@ double FPCGRuntimeGenerationRadii::GetCleanupRadiusFromGrid(EPCGHiGenGrid Grid)
 	return GetGenerationRadiusFromGrid(Grid) * CleanupRadiusMultiplier;
 }
 
+namespace PCGDelegates
+{
+#if WITH_EDITOR
+	FOnInstanceLayoutChanged OnInstancedPropertyBagLayoutChanged;
+#endif
+} // PCGDelegates
+
diff --git a/Engine/Plugins/Experimental/PCG/Source/PCG/Private/PCGGraph.cpp b/Engine/Plugins/Experimental/PCG/Source/PCG/Private/PCGGraph.cpp
--- a/Engine/Plugins/Experimental/PCG/Source/PCG/Private/PCGGraph.cpp
+++ b/Engine/Plugins/Experimental/PCG/Source/PCG/Private/PCGGraph.cpp
@@ -2,6 +2,7 @@
 
 #include "PCGGraph.h"
 
+#include "PCGCommon.h"
 #include "PCGComponent.h"
 #include "PCGEdge.h"
 #include "PCGInputOutputSettings.h"
@@ -1687,6 +1688,10 @@ bool FPCGOverrideInstancedPropertyBag::RefreshParameters(const FInstancedPropert
 		// Copy the parent parameters and reset overriddes
 		Parameters = *ParentUserParameters;
 		PropertiesIDsOverridden.Reset();
+#if WITH_EDITOR
+		// Notify the UI to force refresh
+		PCGDelegates::OnInstancedPropertyBagLayoutChanged.Broadcast(Parameters);
+#endif // WITH_EDITOR
 		break;
 	}
 	case EPCGGraphParameterEvent::Added:
@@ -1702,6 +1707,10 @@ bool FPCGOverrideInstancedPropertyBag::RefreshParameters(const FInstancedPropert
 		}
 
 		MigrateToNewBagInstance(*ParentUserParameters);
+#if WITH_EDITOR
+		// Notify the UI to force refresh
+		PCGDelegates::OnInstancedPropertyBagLayoutChanged.Broadcast(Parameters);
+#endif // WITH_EDITOR
 		break;
 	}
 	case EPCGGraphParameterEvent::ValueModifiedByParent:
diff --git a/Engine/Plugins/Experimental/PCG/Source/PCG/Public/PCGCommon.h b/Engine/Plugins/Experimental/PCG/Source/PCG/Public/PCGCommon.h
--- a/Engine/Plugins/Experimental/PCG/Source/PCG/Public/PCGCommon.h
+++ b/Engine/Plugins/Experimental/PCG/Source/PCG/Public/PCGCommon.h
@@ -250,3 +250,15 @@ public:
 	UPROPERTY(BlueprintReadWrite, EditAnywhere, Category = "Runtime Generation", meta = (UIMin = "1.0", ClampMin = "1.0"))
 	double CleanupRadiusMultiplier = DefaultCleanupRadiusMultiplier;
 };
+
+struct FInstancedPropertyBag;
+
+namespace PCGDelegates
+{
+#if WITH_EDITOR
+	/** Callback to hook in the UI to detect property bag changes, so the UI is reset and does not try to read in garbage memory. */
+	DECLARE_MULTICAST_DELEGATE_OneParam(FOnInstanceLayoutChanged, const FInstancedPropertyBag& /*Instance*/);
+	extern PCG_API FOnInstanceLayoutChanged OnInstancedPropertyBagLayoutChanged;
+#endif
+}
+
diff --git a/Engine/Plugins/Experimental/PCG/Source/PCGEditor/Private/Details/PCGInstancedPropertyBagOverrideDetails.cpp b/Engine/Plugins/Experimental/PCG/Source/PCGEditor/Private/Details/PCGInstancedPropertyBagOverrideDetails.cpp
--- a/Engine/Plugins/Experimental/PCG/Source/PCGEditor/Private/Details/PCGInstancedPropertyBagOverrideDetails.cpp
+++ b/Engine/Plugins/Experimental/PCG/Source/PCGEditor/Private/Details/PCGInstancedPropertyBagOverrideDetails.cpp
@@ -1,6 +1,7 @@
 // Copyright Epic Games, Inc. All Rights Reserved.
 
 #include "Details/PCGInstancedPropertyBagOverrideDetails.h"
+#include "PCGCommon.h"
 #include "PCGGraph.h"
 #include "PCGSettings.h"
 
@@ -96,6 +97,18 @@ FPCGOverrideInstancedPropertyBagDataDetails::FPCGOverrideInstancedPropertyBagDat
 	}
 }
 
+FPCGOverrideInstancedPropertyBagDataDetails::~FPCGOverrideInstancedPropertyBagDataDetails()
+{
+	PCGDelegates::OnInstancedPropertyBagLayoutChanged.RemoveAll(this);
+}
+
+void FPCGOverrideInstancedPropertyBagDataDetails::GenerateHeaderRowContent(FDetailWidgetRow& NodeRow)
+{
+	FInstancedStructDataDetails::GenerateHeaderRowContent(NodeRow);
+
+	PCGDelegates::OnInstancedPropertyBagLayoutChanged.AddSP(this, &FPCGOverrideInstancedPropertyBagDataDetails::OnStructChange);
+}
+
 EVisibility FPCGOverrideInstancedPropertyBagDataDetails::IsResetVisible(TSharedPtr<IPropertyHandle> InPropertyHandle) const
 {
 	return (Owner.IsValid() ? Owner->IsPropertyOverriddenAndNotDefault(InPropertyHandle->GetProperty()) : false) ? EVisibility::Visible : EVisibility::Hidden;
diff --git a/Engine/Plugins/Experimental/PCG/Source/PCGEditor/Private/Details/PCGInstancedPropertyBagOverrideDetails.h b/Engine/Plugins/Experimental/PCG/Source/PCGEditor/Private/Details/PCGInstancedPropertyBagOverrideDetails.h
--- a/Engine/Plugins/Experimental/PCG/Source/PCGEditor/Private/Details/PCGInstancedPropertyBagOverrideDetails.h
+++ b/Engine/Plugins/Experimental/PCG/Source/PCGEditor/Private/Details/PCGInstancedPropertyBagOverrideDetails.h
@@ -12,6 +12,7 @@
 
 class UPCGGraphInstance;
 class UPCGSettings;
+struct FInstancedPropertyBag;
 
 class FPCGOverrideInstancedPropertyBagDetails : public IPropertyTypeCustomization
 {
@@ -42,4 +43,10 @@ private:
 	TWeakObjectPtr<UPCGGraphInstance> Owner;
 	TWeakObjectPtr<UPCGSettings> SettingsOwner;
 	TSharedPtr<IPropertyHandle> PropertiesIDsOverriddenHandle;
+
+public:
+	~FPCGOverrideInstancedPropertyBagDataDetails();
+
+	virtual void GenerateHeaderRowContent(FDetailWidgetRow& NodeRow) override;
+	void OnStructChange(const FInstancedPropertyBag&) { OnStructLayoutChanges(); }
 };
\ No newline at end of file
diff --git a/Engine/Plugins/Experimental/StructUtils/Source/StructUtilsEditor/Private/Customizations/InstancedStructDetails.cpp b/Engine/Plugins/Experimental/StructUtils/Source/StructUtilsEditor/Private/Customizations/InstancedStructDetails.cpp
--- a/Engine/Plugins/Experimental/StructUtils/Source/StructUtilsEditor/Private/Customizations/InstancedStructDetails.cpp
+++ b/Engine/Plugins/Experimental/StructUtils/Source/StructUtilsEditor/Private/Customizations/InstancedStructDetails.cpp
@@ -249,12 +249,7 @@ FInstancedStructDataDetails::~FInstancedStructDataDetails()
 
 void FInstancedStructDataDetails::OnUserDefinedStructReinstancedHandle(const UUserDefinedStruct& Struct)
 {
-	if (StructProvider.IsValid())
-	{
-		// Reset the struct provider immediately, some update functions might get called with the old struct.
-		StructProvider->Reset();
-	}
-	OnRegenerateChildren.ExecuteIfBound();
+	OnStructLayoutChanges();
 }
 
 void FInstancedStructDataDetails::SetOnRebuildChildren(FSimpleDelegate InOnRegenerateChildren)
@@ -283,6 +278,16 @@ TArray<TWeakObjectPtr<const UStruct>> FInstancedStructDataDetails::GetInstanceTy
 	return Result;
 }
 
+void FInstancedStructDataDetails::OnStructLayoutChanges()
+{
+	if (StructProvider.IsValid())
+	{
+		// Reset the struct provider immediately, some update functions might get called with the old struct.
+		StructProvider->Reset();
+	}
+	OnRegenerateChildren.ExecuteIfBound();
+}
+
 void FInstancedStructDataDetails::OnStructHandlePostChange()
 {
 	if (StructProvider.IsValid())
diff --git a/Engine/Plugins/Experimental/StructUtils/Source/StructUtilsEditor/Public/InstancedStructDetails.h b/Engine/Plugins/Experimental/StructUtils/Source/StructUtilsEditor/Public/InstancedStructDetails.h
--- a/Engine/Plugins/Experimental/StructUtils/Source/StructUtilsEditor/Public/InstancedStructDetails.h
+++ b/Engine/Plugins/Experimental/StructUtils/Source/StructUtilsEditor/Public/InstancedStructDetails.h
@@ -106,6 +106,9 @@ private:
 	bool bIsHandlingStructValuePostChange = false;
 	
 	FDelegateHandle UserDefinedStructReinstancedHandle;
+
+protected:
+	void OnStructLayoutChanges();
 };
 
 /**
