[PCG] Add support for component refresh to Runtime Generation system
* Functionality to refresh individual components in the runtime generation scheduler
* Complete support for inspection/debug on GenerateAtRuntime components
* Added Refresh button to GenerateAtRuntime components
* Added optional dependencies to GenerateLocal to support executing HiGen grid levels in the correct order

diff --git a/Engine/Plugins/Experimental/PCG/Source/PCG/Private/Grid/PCGPartitionActor.cpp b/Engine/Plugins/Experimental/PCG/Source/PCG/Private/Grid/PCGPartitionActor.cpp
--- a/Engine/Plugins/Experimental/PCG/Source/PCG/Private/Grid/PCGPartitionActor.cpp
+++ b/Engine/Plugins/Experimental/PCG/Source/PCG/Private/Grid/PCGPartitionActor.cpp
@@ -171,6 +171,7 @@ void APCGPartitionActor::PostRegisterAllComponents()
 
 void APCGPartitionActor::BeginPlay()
 {
+	// bIsRuntimeGenerated is not set yet, so we also need to check if the PA is transient.
 	if (!IsRuntimeGenerated() && !HasAnyFlags(RF_Transient))
 	{
 		// Pass through all the pcg components, to verify if we need to generate them.
diff --git a/Engine/Plugins/Experimental/PCG/Source/PCG/Private/PCGComponent.cpp b/Engine/Plugins/Experimental/PCG/Source/PCG/Private/PCGComponent.cpp
--- a/Engine/Plugins/Experimental/PCG/Source/PCG/Private/PCGComponent.cpp
+++ b/Engine/Plugins/Experimental/PCG/Source/PCG/Private/PCGComponent.cpp
@@ -281,9 +281,9 @@ void UPCGComponent::GenerateLocal(bool bForce)
 	GenerateInternal(bForce, EPCGHiGenGrid::Uninitialized, EPCGComponentGenerationTrigger::GenerateOnDemand, {});
 }
 
-void UPCGComponent::GenerateLocal(EPCGComponentGenerationTrigger RequestedGenerationTrigger, bool bForce, EPCGHiGenGrid Grid)
+void UPCGComponent::GenerateLocal(EPCGComponentGenerationTrigger RequestedGenerationTrigger, bool bForce, EPCGHiGenGrid Grid, const TArray<FPCGTaskId>& Dependencies)
 {
-	GenerateInternal(bForce, Grid, RequestedGenerationTrigger, {});
+	GenerateInternal(bForce, Grid, RequestedGenerationTrigger, Dependencies);
 }
 
 FPCGTaskId UPCGComponent::GenerateLocalGetTaskId(bool bForce)
@@ -1454,6 +1454,23 @@ void UPCGComponent::RefreshAfterGraphChanged(UPCGGraphInterface* InGraph, bool b
 }
 
 #if WITH_EDITOR
+void UPCGComponent::PreEditChange(FProperty* PropertyAboutToChange)
+{
+	Super::PreEditChange(PropertyAboutToChange);
+
+	if (PropertyAboutToChange && PropertyAboutToChange->GetFName() == GET_MEMBER_NAME_CHECKED(UPCGComponent, GenerationTrigger))
+	{
+		if (GenerationTrigger == EPCGComponentGenerationTrigger::GenerateAtRuntime)
+		{
+			if (UPCGSubsystem* Subsystem = GetSubsystem())
+			{
+				// When toggling off of GenerateAtRuntime, we should flush the RuntimeGenScheduler state for this component.
+				Subsystem->RefreshRuntimeGenComponent(this, /*bRemovePartitionActors=*/true);
+			}
+		}
+	}
+}
+
 void UPCGComponent::PostEditChangeProperty(FPropertyChangedEvent& PropertyChangedEvent)
 {
 	if (!PropertyChangedEvent.Property || !IsValid(this))
@@ -1538,11 +1555,6 @@ void UPCGComponent::PostEditChangeProperty(FPropertyChangedEvent& PropertyChange
 	} 
 	else if (PropName == GET_MEMBER_NAME_CHECKED(UPCGComponent, GenerationTrigger))
 	{
-		if (GenerationTrigger == EPCGComponentGenerationTrigger::GenerateAtRuntime)
-		{
-			CleanupLocalImmediate(/*bRemoveComponents=*/true);
-		}
-
 		if (!SchedulingPolicy)
 		{
 			RefreshSchedulingPolicy();
@@ -1554,7 +1566,16 @@ void UPCGComponent::PostEditChangeProperty(FPropertyChangedEvent& PropertyChange
 			SchedulingPolicy->SetShouldDisplayProperties(GenerationTrigger == EPCGComponentGenerationTrigger::GenerateAtRuntime);
 		}
 
-		Refresh();
+		if (GenerationTrigger == EPCGComponentGenerationTrigger::GenerateAtRuntime)
+		{
+			// If we have been set to GenerateAtRuntime, we should cleanup any artifacts.
+			// TODO: Should this include PartitionActors?
+			CleanupLocalImmediate(/*bRemoveComponents=*/true);
+		}
+		else
+		{
+			Refresh();
+		}
 	}
 	// General properties that don't affect behavior
 	else
@@ -1822,9 +1843,15 @@ void UPCGComponent::Refresh(bool bStructural)
 		return;
 	}
 
-	// Do not allow automatic refreshing for runtime generation components.
+	// Runtime component refreshes should go through the runtime scheduler.
 	if (GenerationTrigger == EPCGComponentGenerationTrigger::GenerateAtRuntime)
 	{
+		if (UPCGSubsystem* Subsystem = GetSubsystem())
+		{
+			// TODO: We only need to remove PAs if the grid sizes changed. Is there a reliable way to know that?
+			Subsystem->RefreshRuntimeGenComponent(this, /*bRemovePartitionActors=*/true);
+		}
+
 		return;
 	}
 
diff --git a/Engine/Plugins/Experimental/PCG/Source/PCG/Private/PCGSubsystem.cpp b/Engine/Plugins/Experimental/PCG/Source/PCG/Private/PCGSubsystem.cpp
--- a/Engine/Plugins/Experimental/PCG/Source/PCG/Private/PCGSubsystem.cpp
+++ b/Engine/Plugins/Experimental/PCG/Source/PCG/Private/PCGSubsystem.cpp
@@ -410,13 +410,14 @@ FPCGTaskId UPCGSubsystem::ScheduleComponent(UPCGComponent* PCGComponent, EPCGHiG
 
 		auto LocalGenerateTask = [OriginalComponent = PCGComponent, Grid, &Dependencies, bSave, &GridSizes](UPCGComponent* LocalComponent)
 		{
-			const uint32 LocalComponentGridSize = LocalComponent->GetGenerationGridSize();
-			if (!GridSizes.Contains(LocalComponentGridSize))
+			if (!GridSizes.Contains(LocalComponent->GetGenerationGridSize()))
 			{
-				return InvalidPCGTaskId;
+				// Local component with invalid grid size. Grid sizes may have changed in graph.
+				return LocalComponent->CleanupInternal(/*bRemoveComponents=*/true, /*bSave=*/true, Dependencies);
 			}
 			else if (Grid != EPCGHiGenGrid::Uninitialized && !(Grid & LocalComponent->GetGenerationGrid()))
 			{
+				// Grid size does not match the given target grid, so skip.
 				return InvalidPCGTaskId;
 			}
 
@@ -660,6 +661,19 @@ void UPCGSubsystem::CancelAllGeneration()
 	}
 }
 
+void UPCGSubsystem::RefreshRuntimeGenComponent(UPCGComponent* RuntimeComponent, bool bRemovePartitionActors)
+{
+	if (!ensure(RuntimeComponent && RuntimeComponent->GenerationTrigger == EPCGComponentGenerationTrigger::GenerateAtRuntime))
+	{
+		return;
+	}
+
+	if (ensure(RuntimeGenScheduler))
+	{
+		RuntimeGenScheduler->RefreshComponent(RuntimeComponent, bRemovePartitionActors);
+	}
+}
+
 bool UPCGSubsystem::IsGraphCurrentlyExecuting(UPCGGraph* Graph)
 {
 	check(GraphExecutor);
diff --git a/Engine/Plugins/Experimental/PCG/Source/PCG/Private/RuntimeGen/PCGRuntimeGenScheduler.cpp b/Engine/Plugins/Experimental/PCG/Source/PCG/Private/RuntimeGen/PCGRuntimeGenScheduler.cpp
--- a/Engine/Plugins/Experimental/PCG/Source/PCG/Private/RuntimeGen/PCGRuntimeGenScheduler.cpp
+++ b/Engine/Plugins/Experimental/PCG/Source/PCG/Private/RuntimeGen/PCGRuntimeGenScheduler.cpp
@@ -6,6 +6,7 @@
 #include "PCGActorAndComponentMapping.h"
 #include "PCGCommon.h"
 #include "PCGComponent.h"
+#include "PCGGraph.h"
 #include "PCGModule.h"
 #include "PCGSubsystem.h"
 #include "PCGWorldActor.h"
@@ -111,10 +112,13 @@ void FPCGRuntimeGenScheduler::Tick(const APCGWorldActor* InPCGWorldActor)
 
 	if (ScheduleFrameCounter <= 0)
 	{
-		// Sort components by priority (will be generated in descending order).
-		ComponentsToGenerate.ValueSort([](double PrioA, double PrioB)->bool { return PrioA > PrioB; });
+		if (!ComponentsToGenerate.IsEmpty())
+		{
+			// Sort components by priority (will be generated in descending order).
+			ComponentsToGenerate.ValueSort([](double PrioA, double PrioB)->bool { return PrioA > PrioB; });
 
-		TickScheduleGeneration(ComponentsToGenerate);
+			TickScheduleGeneration(ComponentsToGenerate);
+		}
 	}
 	else
 	{
@@ -590,7 +594,20 @@ void FPCGRuntimeGenScheduler::TickScheduleGeneration(TMap<FGridGenerationKey, do
 					UE_LOG(LogPCG, Warning, TEXT("[RUNTIMEGEN] GENERATE: '%s' (priority %lf)"), *PartitionActor->GetActorNameOrLabel(), Priority);
 				}
 
-				LocalComponent->GenerateLocal(EPCGComponentGenerationTrigger::GenerateAtRuntime, /*bForce=*/false, LocalComponent->GetGenerationGrid());
+				// Higen graphs may have data links from original component to local components. The original component will be given a higher priority than local
+				// components and will start generating first. If it is currently generating, local component needs to take a dependency to ensure execution completes.
+				TArray<FPCGTaskId> Dependencies;
+				if (OriginalComponent->IsGenerating() && OriginalComponent->GetGraph() && OriginalComponent->GetGraph()->IsHierarchicalGenerationEnabled())
+				{
+					const FPCGTaskId TaskId = OriginalComponent->GetGenerationTaskId();
+
+					if (TaskId != InvalidPCGTaskId)
+					{
+						Dependencies.Add(TaskId);
+					}
+				}
+
+				LocalComponent->GenerateLocal(EPCGComponentGenerationTrigger::GenerateAtRuntime, /*bForce=*/false, LocalComponent->GetGenerationGrid(), Dependencies);
 			}
 		}
 
@@ -808,6 +825,113 @@ void FPCGRuntimeGenScheduler::CleanupComponent(const FGridGenerationKey& Generat
 	GeneratedComponents.Remove(GenerationKey);
 }
 
+void FPCGRuntimeGenScheduler::RefreshComponent(UPCGComponent* InComponent, bool bRemovePartitionActors)
+{
+	if (!InComponent || !ensure(IsInGameThread()))
+	{
+		return;
+	}
+
+	const APCGPartitionActor* PartitionActor = Cast<APCGPartitionActor>(InComponent->GetOwner());
+	const bool bIsLocalComponent = PartitionActor != nullptr;
+
+	UPCGComponent* OriginalComponent = bIsLocalComponent ? PartitionActor->GetOriginalComponent(InComponent) : InComponent;
+
+	auto RefreshLocalComponent = [this, OriginalComponent, bRemovePartitionActors](UPCGComponent* LocalComponent)
+	{
+		check(LocalComponent);
+		APCGPartitionActor* PartitionActor = CastChecked<APCGPartitionActor>(LocalComponent->GetOwner());
+
+		// Find the specific generation key for this component, if it exists, cleanup and generate.
+		FGridGenerationKey LocalComponentKey(PartitionActor->GetPCGGridSize(), PartitionActor->GetGridCoord(), OriginalComponent);
+
+		if (GeneratedComponents.Find(LocalComponentKey))
+		{
+			if (PCGRuntimeGenSchedulerHelpers::CVarRuntimeGenerationEnableDebugging.GetValueOnAnyThread())
+			{
+				UE_LOG(LogPCG, Warning, TEXT("[RUNTIMEGEN] REFRESH LOCAL COMPONENT: '%s'"), *PartitionActor->GetActorNameOrLabel());
+			}
+
+			if (bRemovePartitionActors)
+			{
+				CleanupComponent(LocalComponentKey, LocalComponent);
+			}
+			else
+			{
+				LocalComponent->CleanupLocalImmediate(/*bRemoveComponents=*/true);
+				LocalComponent->GenerateLocal(EPCGComponentGenerationTrigger::GenerateAtRuntime, /*bForce=*/false, LocalComponent->GetGenerationGrid());
+			}
+		}
+	};
+
+	if (bIsLocalComponent)
+	{
+		RefreshLocalComponent(InComponent);
+	}
+	else
+	{
+		TArray<FGridGenerationKey> GenerationKeys;
+
+		for (FGridGenerationKey GenerationKey : GeneratedComponents)
+		{
+			if (GenerationKey.GetOriginalComponent() == OriginalComponent)
+			{
+				GenerationKeys.Add(GenerationKey);
+			}
+		}
+
+		if (!bRemovePartitionActors)
+		{
+			// Sort GenerationKeys by GridSize to guarantee the largest grid sizes execute first.
+			// This is only necessary if we are not doing a full cleanup, since full cleanup regeneration is already sorted by the scheduler.
+			GenerationKeys.Sort([](const FGridGenerationKey& Lhs, const FGridGenerationKey& Rhs)
+			{
+				return Lhs.GetGridSize() > Rhs.GetGridSize();
+			});
+		}
+
+		// Gather all generated components which originated from this original component.
+		for (FGridGenerationKey GenerationKey : GenerationKeys)
+		{
+			const uint32 GridSize = GenerationKey.GetGridSize();
+			const EPCGHiGenGrid Grid = PCGHiGenGrid::GridSizeToGrid(GridSize);
+
+			// If the Grid is unbounded, we have a non-partitioned or unbounded component.
+			if (Grid == EPCGHiGenGrid::Unbounded)
+			{
+				if (PCGRuntimeGenSchedulerHelpers::CVarRuntimeGenerationEnableDebugging.GetValueOnAnyThread() && OriginalComponent->GetOwner())
+				{
+					UE_LOG(LogPCG, Warning, TEXT("[RUNTIMEGEN] REFRESH COMPONENT: '%s'"), *OriginalComponent->GetOwner()->GetActorNameOrLabel());
+				}
+
+				if (bRemovePartitionActors)
+				{
+					CleanupComponent(GenerationKey, OriginalComponent);
+				}
+				else
+				{
+					OriginalComponent->CleanupLocalImmediate(/*bRemoveComponents=*/true);
+					OriginalComponent->GenerateLocal(EPCGComponentGenerationTrigger::GenerateAtRuntime, /*bForce=*/false, Grid);
+				}
+			}
+			// Otherwise we have a local component.
+			else
+			{
+				const FIntVector GridCoords = GenerationKey.GetGridCoords();
+
+				UPCGComponent* LocalComponent = ActorAndComponentMapping->GetLocalComponent(GridSize, GridCoords, OriginalComponent, /*bRuntimeGenerated=*/true);
+				if (!ensure(LocalComponent))
+				{
+					UE_LOG(LogPCG, Error, TEXT("[RUNTIMEGEN] Generated local component could not be retrieved: %u, %d_%d_%d"), GridSize, GridCoords.X, GridCoords.Y, GridCoords.Z);
+					continue;
+				}
+
+				RefreshLocalComponent(LocalComponent);
+			}
+		}
+	}
+}
+
 APCGPartitionActor* FPCGRuntimeGenScheduler::FindOrCreatePartitionActor(uint32 GridSize, const FIntVector& GridCoords)
 {
 	TRACE_CPUPROFILER_EVENT_SCOPE(FPCGRuntimeGenScheduler::FindOrCreatePartitionActor);
diff --git a/Engine/Plugins/Experimental/PCG/Source/PCG/Public/PCGComponent.h b/Engine/Plugins/Experimental/PCG/Source/PCG/Public/PCGComponent.h
--- a/Engine/Plugins/Experimental/PCG/Source/PCG/Public/PCGComponent.h
+++ b/Engine/Plugins/Experimental/PCG/Source/PCG/Public/PCGComponent.h
@@ -143,7 +143,7 @@ public:
 	void GenerateLocal(bool bForce);
 
 	/** Requests the component to generate only on the specified grid level (all grid levels if EPCGHiGenGrid::Uninitialized). */
-	void GenerateLocal(EPCGComponentGenerationTrigger RequestedGenerationTrigger, bool bForce, EPCGHiGenGrid Grid = EPCGHiGenGrid::Uninitialized);
+	void GenerateLocal(EPCGComponentGenerationTrigger RequestedGenerationTrigger, bool bForce, EPCGHiGenGrid Grid = EPCGHiGenGrid::Uninitialized, const TArray<FPCGTaskId>& Dependencies = {});
 
 	FPCGTaskId GenerateLocalGetTaskId(bool bForce);
 	FPCGTaskId GenerateLocalGetTaskId(EPCGComponentGenerationTrigger RequestedGenerationTrigger, bool bForce, EPCGHiGenGrid Grid = EPCGHiGenGrid::Uninitialized);
@@ -404,6 +404,7 @@ private:
 	void OnGraphChanged(UPCGGraphInterface* InGraph, EPCGChangeType ChangeType);
 
 #if WITH_EDITOR
+	virtual void PreEditChange(FProperty* PropertyAboutToChange) override;
 	virtual void PostEditChangeProperty(FPropertyChangedEvent& PropertyChangedEvent) override;
 	virtual void PreEditUndo() override;
 	virtual void PostEditUndo() override;
diff --git a/Engine/Plugins/Experimental/PCG/Source/PCG/Public/PCGSubsystem.h b/Engine/Plugins/Experimental/PCG/Source/PCG/Public/PCGSubsystem.h
--- a/Engine/Plugins/Experimental/PCG/Source/PCG/Public/PCGSubsystem.h
+++ b/Engine/Plugins/Experimental/PCG/Source/PCG/Public/PCGSubsystem.h
@@ -121,6 +121,9 @@ public:
 	/** Cancels currently running generation on given graph */
 	void CancelGeneration(UPCGGraph* Graph);
 
+	/** Asks the runtime generation scheduler to refresh a given GenerateAtRuntime component. bRemovePartitionActors will also perform a full cleanup of PAs and local components. */
+	void RefreshRuntimeGenComponent(UPCGComponent* RuntimeComponent, bool bRemovePartitionActors = false);
+
 	/** Returns true if there are any tasks for this graph currently scheduled or executing. */
 	bool IsGraphCurrentlyExecuting(UPCGGraph* Graph);
 
diff --git a/Engine/Plugins/Experimental/PCG/Source/PCG/Public/RuntimeGen/PCGRuntimeGenScheduler.h b/Engine/Plugins/Experimental/PCG/Source/PCG/Public/RuntimeGen/PCGRuntimeGenScheduler.h
--- a/Engine/Plugins/Experimental/PCG/Source/PCG/Public/RuntimeGen/PCGRuntimeGenScheduler.h
+++ b/Engine/Plugins/Experimental/PCG/Source/PCG/Public/RuntimeGen/PCGRuntimeGenScheduler.h
@@ -75,6 +75,9 @@ protected:
 	/** Cleanup a component and remove it from the GeneratedComponents set. */
 	void CleanupComponent(const FGridGenerationKey& GenerationKey, UPCGComponent* GeneratedComponent);
 
+	/** Refresh a generated component. bRemovePartitionActors will also perform a full cleanup of PAs and local components. */
+	void RefreshComponent(UPCGComponent* InComponent, bool bRemovePartitionActors = false);
+
 	/** Creates an empty RuntimeGen PA on demand if one cannot already be found in the level. PA will be created for the given GridSize and GridCoords. */
 	APCGPartitionActor* FindOrCreatePartitionActor(uint32 GridSize, const FIntVector& GridCoords);
 	
diff --git a/Engine/Plugins/Experimental/PCG/Source/PCGEditor/Private/Details/PCGComponentDetails.cpp b/Engine/Plugins/Experimental/PCG/Source/PCGEditor/Private/Details/PCGComponentDetails.cpp
--- a/Engine/Plugins/Experimental/PCG/Source/PCGEditor/Private/Details/PCGComponentDetails.cpp
+++ b/Engine/Plugins/Experimental/PCG/Source/PCGEditor/Private/Details/PCGComponentDetails.cpp
@@ -106,6 +106,19 @@ void FPCGComponentDetails::CustomizeDetails(IDetailLayoutBuilder& DetailBuilder)
 			+ SHorizontalBox::Slot()
 			.AutoWidth()
 			.Padding(4.0f)
+			[
+				SNew(SButton)
+				.OnClicked(this, &FPCGComponentDetails::OnRefreshClicked)
+				.Visibility(this, &FPCGComponentDetails::RefreshButtonVisible)
+				[
+					SNew(STextBlock)
+					.Font(IDetailLayoutBuilder::GetDetailFont())
+					.Text(LOCTEXT("RefreshButton", "Refresh"))
+				]
+			]
+			+ SHorizontalBox::Slot()
+			.AutoWidth()
+			.Padding(4.0f)
 			[
 				SNew(SButton)
 				.OnClicked(this, &FPCGComponentDetails::OnClearPCGLinkClicked)
@@ -183,6 +196,19 @@ EVisibility FPCGComponentDetails::CleanupButtonVisible() const
 	return EVisibility::Collapsed;
 }
 
+EVisibility FPCGComponentDetails::RefreshButtonVisible() const
+{
+	for (const TWeakObjectPtr<UPCGComponent>& Component : SelectedComponents)
+	{
+		if (Component.IsValid() && Component->GenerationTrigger == EPCGComponentGenerationTrigger::GenerateAtRuntime)
+		{
+			return EVisibility::Visible;
+		}
+	}
+
+	return EVisibility::Collapsed;
+}
+
 FReply FPCGComponentDetails::OnGenerateClicked()
 {
 	for (TWeakObjectPtr<UPCGComponent>& Component : SelectedComponents)
@@ -209,6 +235,19 @@ FReply FPCGComponentDetails::OnCancelClicked()
 	return FReply::Handled();
 }
 
+FReply FPCGComponentDetails::OnRefreshClicked()
+{
+	for (TWeakObjectPtr<UPCGComponent>& Component : SelectedComponents)
+	{
+		if (Component.IsValid())
+		{
+			Component.Get()->Refresh();
+		}
+	}
+
+	return FReply::Handled();
+}
+
 FReply FPCGComponentDetails::OnClearPCGLinkClicked()
 {
 	for (TWeakObjectPtr<UPCGComponent>& Component : SelectedComponents)
diff --git a/Engine/Plugins/Experimental/PCG/Source/PCGEditor/Private/PCGEditor.cpp b/Engine/Plugins/Experimental/PCG/Source/PCGEditor/Private/PCGEditor.cpp
--- a/Engine/Plugins/Experimental/PCG/Source/PCGEditor/Private/PCGEditor.cpp
+++ b/Engine/Plugins/Experimental/PCG/Source/PCGEditor/Private/PCGEditor.cpp
@@ -273,19 +273,31 @@ void FPCGEditor::UpdateDebugAfterComponentSelection(UPCGComponent* InOldComponen
 		return;
 	}
 
-	// GenerateOnDemand requests will bounce off if the component is GenerateAtRuntime, so we need to select the correct GenTrigger for the component to be updated.
-	EPCGComponentGenerationTrigger GenerationTrigger = EPCGComponentGenerationTrigger::GenerateOnDemand;
-	if (InNewComponent && InNewComponent->GenerationTrigger == EPCGComponentGenerationTrigger::GenerateAtRuntime)
+	auto RefreshComponent = [](UPCGComponent* Component)
 	{
-		GenerationTrigger = EPCGComponentGenerationTrigger::GenerateAtRuntime;
-	}
+		check(Component);
+
+		// GenerateAtRuntime components should be refreshed through the runtime gen scheduler.
+		if (Component->GenerationTrigger == EPCGComponentGenerationTrigger::GenerateAtRuntime)
+		{
+			if (UPCGSubsystem* Subsystem = GetSubsystem())
+			{
+				// We don't want to do a full cleanup if we're setting the debug object, since full cleanup destroys the component, which is the debug object itself!
+				Subsystem->RefreshRuntimeGenComponent(Component, /*bRemovePartitionActors=*/false);
+			}
+		}
+		else
+		{
+			Component->GenerateLocal(/*bForce=*/true);
+		}
+	};
 
 	// If individual component debugging is disabled, just generate the new component if required.
 	if (!PCGGraphBeingEdited->DebugFlagAppliesToIndividualComponents())
 	{
 		if (InNewComponent && bInNewComponentStartedInspecting)
 		{
-			InNewComponent->GenerateLocal(GenerationTrigger, /*bForce=*/true);
+			RefreshComponent(InNewComponent);
 		}
 
 		return;
@@ -298,7 +310,7 @@ void FPCGEditor::UpdateDebugAfterComponentSelection(UPCGComponent* InOldComponen
 		{
 			// Transition from 'null' to 'any component not already inspecting' - generate to create debug/inspection info.
 			// If we have null selected, all components are displaying debug. Go to Original component so that all refresh.
-			InNewComponent->GetOriginalComponent()->GenerateLocal(GenerationTrigger, /*bForce=*/true);
+			RefreshComponent(InNewComponent->GetOriginalComponent());
 		}
 	}
 	else
@@ -311,19 +323,15 @@ void FPCGEditor::UpdateDebugAfterComponentSelection(UPCGComponent* InOldComponen
 		// Regenerate to clear debug info if switching components, or if changing from a component to null.
 		if (InNewComponent || bDebugFlagSetOnAnyNode)
 		{
-			// If InNewComponent is nullptr in local instance debug mode, then RuntimeGen does not want to regenerate (this is to avoid regenerating all PAs).
-			if (InNewComponent || InOldComponent->GenerationTrigger != EPCGComponentGenerationTrigger::GenerateAtRuntime)
-			{
-				// Use original component - debug can be displayed both by the local component and parent local components.
-				InOldComponent->GetOriginalComponent()->GenerateLocal(GenerationTrigger, /*bForce=*/true);
-			}
+			// Use original component - debug can be displayed both by the local component and parent local components.
+			RefreshComponent(InOldComponent->GetOriginalComponent());
 		}
 
 		// Debug new component if it wasn't already
 		if (InNewComponent && bInNewComponentStartedInspecting)
 		{
 			// Use original component - debug can be displayed both by the local component and parent local components.
-			InNewComponent->GetOriginalComponent()->GenerateLocal(GenerationTrigger, /*bForce=*/true);
+			RefreshComponent(InNewComponent->GetOriginalComponent());
 		}
 	}
 }
diff --git a/Engine/Plugins/Experimental/PCG/Source/PCGEditor/Public/Details/PCGComponentDetails.h b/Engine/Plugins/Experimental/PCG/Source/PCGEditor/Public/Details/PCGComponentDetails.h
--- a/Engine/Plugins/Experimental/PCG/Source/PCGEditor/Public/Details/PCGComponentDetails.h
+++ b/Engine/Plugins/Experimental/PCG/Source/PCGEditor/Public/Details/PCGComponentDetails.h
@@ -28,9 +28,11 @@ private:
 	EVisibility GenerateButtonVisible() const;
 	EVisibility CancelButtonVisible() const;
 	EVisibility CleanupButtonVisible() const;
+	EVisibility RefreshButtonVisible() const;
 	FReply OnGenerateClicked();
 	FReply OnCancelClicked();
 	FReply OnCleanupClicked();
+	FReply OnRefreshClicked();
 	FReply OnClearPCGLinkClicked();
 	void OnGraphChanged(UPCGComponent* InComponent);
 
