[PCG] Fix issue where UserDefinedStructs were not extracted properly in Get Actor Property

* Also made a pass on constness

#rnx

diff --git a/Engine/Plugins/Experimental/PCG/Source/PCG/Private/Helpers/PCGPropertyHelpers.cpp b/Engine/Plugins/Experimental/PCG/Source/PCG/Private/Helpers/PCGPropertyHelpers.cpp
--- a/Engine/Plugins/Experimental/PCG/Source/PCG/Private/Helpers/PCGPropertyHelpers.cpp
+++ b/Engine/Plugins/Experimental/PCG/Source/PCG/Private/Helpers/PCGPropertyHelpers.cpp
@@ -9,6 +9,7 @@
 #include "Helpers/PCGSettingsHelpers.h"
 #include "Metadata/PCGMetadata.h"
 
+#include "Engine/UserDefinedStruct.h"
 #include "UObject/Field.h"
 #include "UObject/UnrealType.h"
 
@@ -41,12 +42,29 @@ namespace PCGPropertyHelpers
 	* @param OptionalContext   Optional context used for logging.
 	* @returns                 The last property of the chain (and its container address is in OutContainer)
 	*/
-	FProperty* ExtractPropertyChain(const UStruct* CurrentClass, const FName CurrentName, TArrayView<const FString> NextNames, const bool bNeedsToBeVisible, const void*& OutContainer, FPCGContext* OptionalContext)
+	const FProperty* ExtractPropertyChain(const UStruct* CurrentClass, const FName CurrentName, TArrayView<const FString> NextNames, const bool bNeedsToBeVisible, const void*& OutContainer, FPCGContext* OptionalContext)
 	{
 		check(CurrentClass);
 
-		// Try to get the property
-		FProperty* Property = FindFProperty<FProperty>(CurrentClass, CurrentName);
+		const FProperty* Property = nullptr;
+		// Try to get the property. If it is coming from a user struct, we need to iterate on all properties because the property name is mangled
+		if (const UUserDefinedStruct* UserDefinedStruct = Cast<UUserDefinedStruct>(CurrentClass))
+		{
+			for (TFieldIterator<const FProperty> PropIt(UserDefinedStruct, EFieldIterationFlags::IncludeSuper); PropIt; ++PropIt)
+			{
+				const FName PropertyName = *UserDefinedStruct->GetAuthoredNameForField(*PropIt);
+				if (PropertyName == CurrentName)
+				{
+					Property = *PropIt;
+					break;
+				}
+			}
+		}
+		else
+		{
+			Property = FindFProperty<FProperty>(CurrentClass, CurrentName);
+		}
+
 		if (!Property)
 		{
 			LogError(FText::Format(LOCTEXT("PropertyDoesNotExist", "Property '{0}' does not exist in {1}."), FText::FromName(CurrentName), FText::FromName(CurrentClass->GetFName())), OptionalContext);
@@ -64,12 +82,12 @@ namespace PCGPropertyHelpers
 		{
 			UStruct* NextClass = nullptr;
 
-			if (FStructProperty* StructProperty = CastField<FStructProperty>(Property))
+			if (const FStructProperty* StructProperty = CastField<FStructProperty>(Property))
 			{
 				NextClass = StructProperty->Struct;
 				OutContainer = StructProperty->ContainerPtrToValuePtr<void>(OutContainer);
 			}
-			else if (FObjectProperty* ObjectProperty = CastField<FObjectProperty>(Property))
+			else if (const FObjectProperty* ObjectProperty = CastField<FObjectProperty>(Property))
 			{
 				NextClass = ObjectProperty->PropertyClass;
 				OutContainer = ObjectProperty->GetObjectPropertyValue_InContainer(OutContainer);
@@ -112,7 +130,7 @@ UPCGParamData* PCGPropertyHelpers::ExtractPropertyAsAttributeSet(const PCGProper
 	check(Parameters.Container && Parameters.Class);
 
 	const void* Container = Parameters.Container;
-	FProperty* Property = ExtractPropertyChain(Parameters.Class, Parameters.PropertySelector.GetName(), Parameters.PropertySelector.GetExtraNames(), Parameters.bPropertyNeedsToBeVisible, Container, InOptionalContext);
+	const FProperty* Property = ExtractPropertyChain(Parameters.Class, Parameters.PropertySelector.GetName(), Parameters.PropertySelector.GetExtraNames(), Parameters.bPropertyNeedsToBeVisible, Container, InOptionalContext);
 
 	if (!Property)
 	{
@@ -120,7 +138,7 @@ UPCGParamData* PCGPropertyHelpers::ExtractPropertyAsAttributeSet(const PCGProper
 	}
 
 	// If the property is an array, we will work on the underlying property, and extract each element as an entry in the param data
-	FArrayProperty* ArrayProperty = CastField<FArrayProperty>(Property);
+	const FArrayProperty* ArrayProperty = CastField<FArrayProperty>(Property);
 	if (ArrayProperty)
 	{
 		Property = ArrayProperty->Inner;
@@ -138,15 +156,15 @@ UPCGParamData* PCGPropertyHelpers::ExtractPropertyAsAttributeSet(const PCGProper
 	// Special case where the property is a struct/object, that is not supported by our metadata, we will try to break it down to multiple attributes in the resulting param data, if asked.
 	if ((Property->IsA<FStructProperty>() || Property->IsA<FObjectProperty>()) && bShouldExtract)
 	{
-		UScriptStruct* UnderlyingStruct = nullptr;
-		UClass* UnderlyingClass = nullptr;
+		const UScriptStruct* UnderlyingStruct = nullptr;
+		const UClass* UnderlyingClass = nullptr;
 
-		if (FStructProperty* StructProperty = CastField<FStructProperty>(Property))
+		if (const FStructProperty* StructProperty = CastField<FStructProperty>(Property))
 		{
 			UnderlyingStruct = StructProperty->Struct;
 			AddressFunc = [StructProperty](const void* InAddress) { return StructProperty->ContainerPtrToValuePtr<void>(InAddress); };
 		}
-		else if (FObjectProperty* ObjectProperty = CastField<FObjectProperty>(Property))
+		else if (const FObjectProperty* ObjectProperty = CastField<FObjectProperty>(Property))
 		{
 			UnderlyingClass = ObjectProperty->PropertyClass;
 			AddressFunc = [ObjectProperty](const void* InAddress) { return ObjectProperty->GetObjectPropertyValue_InContainer(InAddress); };
