Added support for padding actor bounds before culling points in CullPointsOutsideActorBounds Node.

diff --git a/Engine/Plugins/Experimental/PCG/Source/PCG/Private/Elements/PCGCullPointsOutsideActorBounds.cpp b/Engine/Plugins/Experimental/PCG/Source/PCG/Private/Elements/PCGCullPointsOutsideActorBounds.cpp
--- a/Engine/Plugins/Experimental/PCG/Source/PCG/Private/Elements/PCGCullPointsOutsideActorBounds.cpp
+++ b/Engine/Plugins/Experimental/PCG/Source/PCG/Private/Elements/PCGCullPointsOutsideActorBounds.cpp
@@ -48,8 +48,11 @@ bool FPCGCullPointsOutsideActorBoundsElement::ExecuteInternal(FPCGContext* Conte
 	UPCGVolumeData* VolumeData = NewObject<UPCGVolumeData>();
 	check(VolumeData);
 
+	const UPCGCullPointsOutsideActorBoundsSettings* Settings = Context->GetInputSettings<UPCGCullPointsOutsideActorBoundsSettings>();
+	check(Settings);
+
 	// Initialize directly. Could also have gone through PCGComponent::CreateActorPCGDataCollection.
-	VolumeData->Initialize(Context->SourceComponent->GetGridBounds());
+	VolumeData->Initialize(Context->SourceComponent->GetGridBounds().ExpandBy(Settings->BoundsExpansion));
 
 	const TArray<FPCGTaggedData>& Inputs = Context->InputData.TaggedData;
 	for (const FPCGTaggedData& Input : Inputs)
diff --git a/Engine/Plugins/Experimental/PCG/Source/PCG/Public/Elements/PCGCullPointsOutsideActorBounds.h b/Engine/Plugins/Experimental/PCG/Source/PCG/Public/Elements/PCGCullPointsOutsideActorBounds.h
--- a/Engine/Plugins/Experimental/PCG/Source/PCG/Public/Elements/PCGCullPointsOutsideActorBounds.h
+++ b/Engine/Plugins/Experimental/PCG/Source/PCG/Public/Elements/PCGCullPointsOutsideActorBounds.h
@@ -28,6 +28,10 @@ protected:
 	virtual TArray<FPCGPinProperties> OutputPinProperties() const override;
 	virtual FPCGElementPtr CreateElement() const override;
 	//~End UPCGSettings interface
+
+public:
+	UPROPERTY(BlueprintReadWrite, EditAnywhere, Category = Settings, meta = (PCG_Overridable))
+	float BoundsExpansion = 0.0;
 };
 
 class FPCGCullPointsOutsideActorBoundsElement : public IPCGElement
