[PCG] Disable higen within subgraphs - subgraphs run at their invoked grid level.

diff --git a/Engine/Plugins/Experimental/PCG/Source/PCG/Private/Graph/PCGGraphCompiler.cpp b/Engine/Plugins/Experimental/PCG/Source/PCG/Private/Graph/PCGGraphCompiler.cpp
--- a/Engine/Plugins/Experimental/PCG/Source/PCG/Private/Graph/PCGGraphCompiler.cpp
+++ b/Engine/Plugins/Experimental/PCG/Source/PCG/Private/Graph/PCGGraphCompiler.cpp
@@ -425,29 +425,55 @@ EPCGHiGenGrid FPCGGraphCompiler::CalculateGridRecursive(
 		return InOutTaskGenerationGrid[InTaskId];
 	}
 
-	// Default for outside any grid size range.
-	EPCGHiGenGrid Grid = GenerationDefaultGrid;
+	const FPCGStack* Stack = InStackContext.GetStack(InCompiledTasks[InTaskId].StackIndex);
+	check(Stack);
+	const bool bTopLevelGraph = InCompiledTasks[InTaskId].ParentId == InvalidPCGTaskId;
 
 	const UPCGNode* Node = InCompiledTasks[InTaskId].Node;
 	const UPCGSettings* Settings = Node ? Node->GetSettings() : nullptr;
-	const UPCGHiGenGridSizeSettings* Gate = Cast<UPCGHiGenGridSizeSettings>(Settings);
-	if (Gate && Gate->bEnabled)
+	const UPCGHiGenGridSizeSettings* GridSizeSettings = Cast<UPCGHiGenGridSizeSettings>(Settings);
+
+	EPCGHiGenGrid Grid = EPCGHiGenGrid::Uninitialized;
+
+	// Grid Size nodes in the top graph set the execution grid level.
+	if (GridSizeSettings && GridSizeSettings->bEnabled && bTopLevelGraph)
 	{
-		Grid = Gate->GetGrid();
+		Grid = GridSizeSettings->GetGrid();
 	}
 	else
 	{
-		// Grid of this task is minimum of all input grids. We can link in data from a larger grid, but not from a finer grid (this goes against hierarchy).
-		for (FPCGGraphTaskInput InputTask : InCompiledTasks[InTaskId].Inputs)
+		if (InCompiledTasks[InTaskId].Inputs.IsEmpty())
 		{
-			const EPCGHiGenGrid InputGrid = CalculateGridRecursive(InputTask.TaskId, GenerationDefaultGrid, InStackContext, InCompiledTasks, InOutTaskGenerationGrid);
-			if (PCGHiGenGrid::IsValidGrid(InputGrid))
+			if (bTopLevelGraph)
 			{
-				Grid = FMath::Min(InputGrid, Grid);
+				// Tasks with no inputs in top graph get prescribed the default generation grid.
+				Grid = GenerationDefaultGrid;
+			}
+			else
+			{
+				// Tasks with no inputs in a subgraph should execute on the same grid as the subgraph node task.
+				check(InCompiledTasks[InTaskId].ParentId != InvalidPCGTaskId);
+				Grid = CalculateGridRecursive(InCompiledTasks[InTaskId].ParentId, GenerationDefaultGrid, InStackContext, InCompiledTasks, InOutTaskGenerationGrid);
+			}
+		}
+		else
+		{
+			// This task has inputs. Grid of this task is minimum of all input grids. We can link in data from a larger grid, but
+			// not from a finer grid (this goes against hierarchy).
+			Grid = EPCGHiGenGrid::Unbounded;
+			for (FPCGGraphTaskInput InputTask : InCompiledTasks[InTaskId].Inputs)
+			{
+				const EPCGHiGenGrid InputGrid = CalculateGridRecursive(InputTask.TaskId, GenerationDefaultGrid, InStackContext, InCompiledTasks, InOutTaskGenerationGrid);
+				if (PCGHiGenGrid::IsValidGrid(InputGrid))
+				{
+					Grid = FMath::Min(InputGrid, Grid);
+				}
 			}
 		}
 	}
 
+	ensure(Grid != EPCGHiGenGrid::Uninitialized);
+
 	InOutTaskGenerationGrid[InTaskId] = Grid;
 
 	return Grid;
diff --git a/Engine/Plugins/Experimental/PCG/Source/PCG/Private/Graph/PCGStackContext.cpp b/Engine/Plugins/Experimental/PCG/Source/PCG/Private/Graph/PCGStackContext.cpp
--- a/Engine/Plugins/Experimental/PCG/Source/PCG/Private/Graph/PCGStackContext.cpp
+++ b/Engine/Plugins/Experimental/PCG/Source/PCG/Private/Graph/PCGStackContext.cpp
@@ -96,6 +96,17 @@ bool FPCGStack::CreateStackFramePath(FString& OutString, const UPCGNode* InNode,
 	return true;
 }
 
+uint32 FPCGStack::GetNumGraphLevels() const
+{
+	uint32 GraphCount = 0;
+	for (const FPCGStackFrame& Frame : StackFrames)
+	{
+		GraphCount += Frame.Object->IsA<UPCGGraph>() ? 1 : 0;
+	}
+
+	return GraphCount;
+}
+
 bool FPCGStack::operator==(const FPCGStack& Other) const
 {
 	// Stacks are the same if all stack frames are the same
diff --git a/Engine/Plugins/Experimental/PCG/Source/PCG/Public/Graph/PCGStackContext.h b/Engine/Plugins/Experimental/PCG/Source/PCG/Public/Graph/PCGStackContext.h
--- a/Engine/Plugins/Experimental/PCG/Source/PCG/Public/Graph/PCGStackContext.h
+++ b/Engine/Plugins/Experimental/PCG/Source/PCG/Public/Graph/PCGStackContext.h
@@ -64,6 +64,12 @@ public:
 	/** Construct a string version of this stack. Postfixed by optional node/pin if provided. */
 	bool CreateStackFramePath(FString& OutString, const UPCGNode* InNode = nullptr, const UPCGPin* InPin = nullptr) const;
 
+	/** Returns how many graphs the stack contains (top level graph stacks will return 1). */
+	uint32 GetNumGraphLevels() const;
+
+	/** Returns true if this stack in the top level graph, rather than in a subgraph. */
+	bool IsTopLevelGraph() const { return GetNumGraphLevels() == 1; }
+
 	const TArray<FPCGStackFrame>& GetStackFrames() const { return StackFrames; }
 	TArray<FPCGStackFrame>& GetStackFramesMutable() { return StackFrames; }
 
diff --git a/Engine/Plugins/Experimental/PCG/Source/PCGEditor/Private/PCGEditor.cpp b/Engine/Plugins/Experimental/PCG/Source/PCGEditor/Private/PCGEditor.cpp
--- a/Engine/Plugins/Experimental/PCG/Source/PCGEditor/Private/PCGEditor.cpp
+++ b/Engine/Plugins/Experimental/PCG/Source/PCGEditor/Private/PCGEditor.cpp
@@ -249,7 +249,7 @@ void FPCGEditor::SetComponentAndStackBeingInspected(UPCGComponent* InPCGComponen
 			{
 				// Update now that component has changed. Will fire OnNodeChanged if necessary.
 				EPCGChangeType ChangeType = PCGNode->UpdateErrorsAndWarnings();
-				ChangeType |= PCGNode->UpdateGridSizeVisualization(InPCGComponent);
+				ChangeType |= PCGNode->UpdateGridSizeVisualization(InPCGComponent, InPCGStack);
 
 				if (ChangeType != EPCGChangeType::None)
 				{
@@ -2270,7 +2270,7 @@ void FPCGEditor::OnGraphGridSizesChanged(UPCGGraphInterface* InGraph)
 	check(PCGEditorGraph);
 	if (UPCGComponent* PCGComponent = GetPCGComponentBeingInspected())
 	{
-		PCGEditorGraph->UpdateGridSizeVisualization(PCGComponent);
+		PCGEditorGraph->UpdateGridSizeVisualization(PCGComponent, GetStackBeingInspected());
 	}
 }
 
diff --git a/Engine/Plugins/Experimental/PCG/Source/PCGEditor/Private/PCGEditorGraph.cpp b/Engine/Plugins/Experimental/PCG/Source/PCGEditor/Private/PCGEditorGraph.cpp
--- a/Engine/Plugins/Experimental/PCG/Source/PCGEditor/Private/PCGEditorGraph.cpp
+++ b/Engine/Plugins/Experimental/PCG/Source/PCGEditor/Private/PCGEditorGraph.cpp
@@ -134,12 +134,12 @@ void UPCGEditorGraph::CreateLinks(UPCGEditorGraphNodeBase* GraphNode, bool bCrea
 	CreateLinks(GraphNode, bCreateInbound, bCreateOutbound, PCGNodeToPCGEditorNodeMap);
 }
 
-void UPCGEditorGraph::UpdateGridSizeVisualization(UPCGComponent* InPCGComponentBeingInspected)
+void UPCGEditorGraph::UpdateGridSizeVisualization(UPCGComponent* PCGComponentBeingInspected, const FPCGStack& PCGStackBeingInspected)
 {
 	for (UEdGraphNode* EditorNode : Nodes)
 	{
 		UPCGEditorGraphNodeBase* PCGEditorNode = Cast<UPCGEditorGraphNodeBase>(EditorNode);
-		if (PCGEditorNode && PCGEditorNode->UpdateGridSizeVisualization(InPCGComponentBeingInspected) != EPCGChangeType::None)
+		if (PCGEditorNode && PCGEditorNode->UpdateGridSizeVisualization(PCGComponentBeingInspected, PCGStackBeingInspected) != EPCGChangeType::None)
 		{
 			PCGEditorNode->ReconstructNode();
 		}
diff --git a/Engine/Plugins/Experimental/PCG/Source/PCGEditor/Private/PCGEditorGraph.h b/Engine/Plugins/Experimental/PCG/Source/PCGEditor/Private/PCGEditorGraph.h
--- a/Engine/Plugins/Experimental/PCG/Source/PCGEditor/Private/PCGEditorGraph.h
+++ b/Engine/Plugins/Experimental/PCG/Source/PCGEditor/Private/PCGEditorGraph.h
@@ -41,7 +41,7 @@ public:
 	TWeakPtr<FPCGEditor> GetEditor() const { return PCGEditor; }
 
 	/** Updates the grid size visualization in the editor. */
-	void UpdateGridSizeVisualization(UPCGComponent* InPCGComponentBeingInspected);
+	void UpdateGridSizeVisualization(UPCGComponent* PCGComponentBeingInspected, const FPCGStack& PCGStackBeingInspected);
 
 protected:
 	void CreateLinks(UPCGEditorGraphNodeBase* InGraphNode, bool bCreateInbound, bool bCreateOutbound, const TMap<UPCGNode*, UPCGEditorGraphNodeBase*>& InPCGNodeToPCGEditorNodeMap);
diff --git a/Engine/Plugins/Experimental/PCG/Source/PCGEditor/Private/PCGEditorGraphNodeBase.cpp b/Engine/Plugins/Experimental/PCG/Source/PCGEditor/Private/PCGEditorGraphNodeBase.cpp
--- a/Engine/Plugins/Experimental/PCG/Source/PCGEditor/Private/PCGEditorGraphNodeBase.cpp
+++ b/Engine/Plugins/Experimental/PCG/Source/PCGEditor/Private/PCGEditorGraphNodeBase.cpp
@@ -360,7 +360,7 @@ EPCGChangeType UPCGEditorGraphNodeBase::UpdateErrorsAndWarnings()
 	return ChangeType;
 }
 
-EPCGChangeType UPCGEditorGraphNodeBase::UpdateGridSizeVisualization(UPCGComponent* InComponentBeingDebugged)
+EPCGChangeType UPCGEditorGraphNodeBase::UpdateGridSizeVisualization(UPCGComponent* InComponentBeingDebugged, const FPCGStack& InStackBeingInspected)
 {
 	const UPCGGraph* Graph = PCGNode ? PCGNode->GetGraph() : nullptr;
 	if (!Graph)
@@ -372,7 +372,10 @@ EPCGChangeType UPCGEditorGraphNodeBase::UpdateGridSizeVisualization(UPCGComponen
 
 	const bool HiGenEnabled = Graph->IsHierarchicalGenerationEnabled();
 	const uint32 InspectingGridSize = InComponentBeingDebugged ? InComponentBeingDebugged->GetGenerationGridSize() : PCGHiGenGrid::UninitializedGridSize();
-	if (!HiGenEnabled || (InspectingGridSize == PCGHiGenGrid::UninitializedGridSize()))
+
+	// Disable grid size visualization if higen is disabled, or if we're not inspecting a specific grid, or if we're
+	// inspecting a subgraph since subgraphs execute at the invoked grid level.
+	if (!HiGenEnabled || (InspectingGridSize == PCGHiGenGrid::UninitializedGridSize()) || !InStackBeingInspected.IsTopLevelGraph())
 	{
 		if (IsDisplayAsDisabledForced())
 		{
diff --git a/Engine/Plugins/Experimental/PCG/Source/PCGEditor/Private/PCGEditorGraphNodeBase.h b/Engine/Plugins/Experimental/PCG/Source/PCGEditor/Private/PCGEditorGraphNodeBase.h
--- a/Engine/Plugins/Experimental/PCG/Source/PCGEditor/Private/PCGEditorGraphNodeBase.h
+++ b/Engine/Plugins/Experimental/PCG/Source/PCGEditor/Private/PCGEditorGraphNodeBase.h
@@ -9,6 +9,7 @@
 
 enum class EPCGChangeType : uint8;
 
+struct FPCGStack;
 class UPCGComponent;
 class UPCGNode;
 class UPCGPin;
@@ -64,7 +65,7 @@ public:
 	EPCGChangeType UpdateErrorsAndWarnings();
 
 	/** If the currently inspected grid size is smaller than the grid size of this node, display transparent. */
-	EPCGChangeType UpdateGridSizeVisualization(class UPCGComponent* InComponentBeingDebugged);
+	EPCGChangeType UpdateGridSizeVisualization(class UPCGComponent* InComponentBeingDebugged, const FPCGStack& InStackBeingInspected);
 
 	/** Puts node title on node body, reducing overall node size */
 	bool ShouldDrawCompact() const;
