[PCG] Add more control on component refresh when the landscape changes

* Landscape edit now by default trigger a refresh only 1s after the last edit
* Delayed is controlled by a CVar and can be disabled by setting it to 0 or negative value
* Add CVar to disable landscape refresh in editing mode.
* Add CVar to disable landscape refresh globally.
* Still dirty the components in both cases.
* Also remove tracking if the node is disabled

#rnx

diff --git a/Engine/Plugins/Experimental/PCG/Source/PCG/Private/PCGActorAndComponentMapping.cpp b/Engine/Plugins/Experimental/PCG/Source/PCG/Private/PCGActorAndComponentMapping.cpp
--- a/Engine/Plugins/Experimental/PCG/Source/PCG/Private/PCGActorAndComponentMapping.cpp
+++ b/Engine/Plugins/Experimental/PCG/Source/PCG/Private/PCGActorAndComponentMapping.cpp
@@ -11,6 +11,7 @@
 #include "Helpers/PCGActorHelpers.h"
 #include "Helpers/PCGHelpers.h"
 
+#include "Landscape.h"
 #include "LandscapeProxy.h"
 #include "StaticMeshCompiler.h"
 #include "TextureCompiler.h"
@@ -46,6 +47,21 @@ namespace PCGActorAndComponentMapping
 		TEXT("pcg.DisableDelayedActorRegistering"),
 		false,
 		TEXT("If delayed actor registering when their components aren't registered yet is introducing bad behavior, disables it, allowing people to continue working while we investigate."));
+
+	static TAutoConsoleVariable<bool> CVarLandscapeDisableRefreshTracking(
+		TEXT("pcg.LandscapeDisableRefreshTracking"),
+		false,
+		TEXT("Completely disable landscape refresh when it changes."));
+
+	static TAutoConsoleVariable<bool> CVarLandscapeDisableRefreshTrackingInLandscapeEditingMode(
+		TEXT("pcg.LandscapeDisableRefreshTrackingInLandscapeEditingMode"),
+		false,
+		TEXT("Disable landscape refresh when it changes in landscape editing mode."));
+
+	static TAutoConsoleVariable<int> CVarLandscapeRefreshTimeDelay(
+		TEXT("pcg.LandscapeRefreshTimeDelayMS"),
+		1000,
+		TEXT("Time in MS between a landscape change and PCG refresh. Set it to 0 or negative value to disable the delay."));
 #endif // WITH_EDITOR
 
 	static TAutoConsoleVariable<bool> CVarDisableDelayedUnregister(
@@ -69,6 +85,31 @@ namespace PCGActorAndComponentMapping
 			}
 		}
 	}
+
+	bool ShouldDiscardLandscapeRefresh(const ALandscapeProxy* InLandscape)
+	{
+		// If it is not a landscape, we should refresh.
+		if (!InLandscape)
+		{
+			return false;
+		}
+
+		// If refresh is globably disabled, never refresh
+		if (PCGActorAndComponentMapping::CVarLandscapeDisableRefreshTracking.GetValueOnAnyThread())
+		{
+			return true;
+		}
+
+		// If refresh is not disabled in editing, always refresh
+		if (!PCGActorAndComponentMapping::CVarLandscapeDisableRefreshTrackingInLandscapeEditingMode.GetValueOnAnyThread())
+		{
+			return false;
+		}
+
+		// Refresh only if we are not editing.
+		const ALandscape* Landscape = InLandscape->GetLandscapeActor();
+		return Landscape && Landscape->HasLandscapeEdMode();
+	}
 #endif
 }
 
@@ -99,6 +140,17 @@ void FPCGActorAndComponentMapping::Tick()
 
 #if WITH_EDITOR
 	AddDelayedActors();
+
+	if (!DelayedModifiedLandscapes.IsEmpty() && LastLandscapeDirtyTime > 0.0 && ((FApp::GetCurrentTime() - LastLandscapeDirtyTime) * 1000.0) > PCGActorAndComponentMapping::CVarLandscapeRefreshTimeDelay.GetValueOnAnyThread())
+	{
+		LastLandscapeDirtyTime = -1.0;
+		for (TObjectKey<ALandscapeProxy> Landscape : DelayedModifiedLandscapes)
+		{
+			ApplyLandscapeChanges(Landscape.ResolveObjectPtr());
+		}
+
+		DelayedModifiedLandscapes.Empty();
+	}
 #endif // WITH_EDITOR
 }
 
@@ -1693,6 +1745,12 @@ void FPCGActorAndComponentMapping::OnActorChanged(AActor* InActor, bool bInHasMo
 		}
 	}
 
+	// If it is a landscape and we should discard the refresh, early out.
+	if (PCGActorAndComponentMapping::ShouldDiscardLandscapeRefresh(Cast<ALandscapeProxy>(InActor)))
+	{
+		return;
+	}
+
 	ULevelInstanceSubsystem* LevelInstanceSubsystem = (PCGSubsystem && PCGSubsystem->GetWorld()) ? PCGSubsystem->GetWorld()->GetSubsystem<ULevelInstanceSubsystem>() : nullptr;
 
 	// And refresh all dirtied components
@@ -1737,6 +1795,24 @@ void FPCGActorAndComponentMapping::OnLandscapeChanged(ALandscapeProxy* InLandsca
 		return;
 	}
 
+	if (PCGActorAndComponentMapping::CVarLandscapeRefreshTimeDelay.GetValueOnAnyThread() > 0)
+	{
+		LastLandscapeDirtyTime = FApp::GetCurrentTime();
+		DelayedModifiedLandscapes.AddUnique(InLandscape);
+	}
+	else
+	{
+		ApplyLandscapeChanges(InLandscape);
+	}
+}
+
+void FPCGActorAndComponentMapping::ApplyLandscapeChanges(ALandscapeProxy* InLandscape)
+{
+	if (!InLandscape)
+	{
+		return;
+	}
+	
 	// We don't know if the landscape moved, only that it has changed. Since `bHasMoved` is doing a bit more, always assume that the landscape has moved.
 	OnActorChanged(InLandscape, /*bHasMoved=*/true);
 
diff --git a/Engine/Plugins/Experimental/PCG/Source/PCG/Private/PCGGraph.cpp b/Engine/Plugins/Experimental/PCG/Source/PCG/Private/PCGGraph.cpp
--- a/Engine/Plugins/Experimental/PCG/Source/PCG/Private/PCGGraph.cpp
+++ b/Engine/Plugins/Experimental/PCG/Source/PCG/Private/PCGGraph.cpp
@@ -931,11 +931,13 @@ void UPCGGraph::GetTrackedActorKeysToSettings(FPCGActorSelectionKeyToSettingsMap
 
 	OutVisitedGraphs.Emplace(this);
 
-	for (UPCGNode* Node : Nodes)
+	for (const UPCGNode* Node : Nodes)
 	{
-		if (Node && Node->GetSettings())
+		const UPCGSettings* Settings = Node ? Node->GetSettings() : nullptr;
+		// Don't track for disabled nodes.
+		if (Settings && Settings->bEnabled)
 		{
-			Node->GetSettings()->GetTrackedActorKeys(OutTagsToSettings, OutVisitedGraphs);
+			Settings->GetTrackedActorKeys(OutTagsToSettings, OutVisitedGraphs);
 		}
 	}
 
diff --git a/Engine/Plugins/Experimental/PCG/Source/PCG/Public/PCGActorAndComponentMapping.h b/Engine/Plugins/Experimental/PCG/Source/PCG/Public/PCGActorAndComponentMapping.h
--- a/Engine/Plugins/Experimental/PCG/Source/PCG/Public/PCGActorAndComponentMapping.h
+++ b/Engine/Plugins/Experimental/PCG/Source/PCG/Public/PCGActorAndComponentMapping.h
@@ -139,6 +139,7 @@ private:
 	void OnActorMoved(AActor* InActor);
 	void OnActorMoved_Internal(AActor* InActor, int32 LevelInstanceDepth);
 	void OnLandscapeChanged(ALandscapeProxy* InLandscape, const FLandscapeProxyComponentDataChangedParams& InChangeParams);
+	void ApplyLandscapeChanges(ALandscapeProxy* InLandscape);
 	void OnPreObjectPropertyChanged(UObject* InObject, const FEditPropertyChain& InEditPropertyChain);
 	void OnObjectPropertyChanged(UObject* InObject, FPropertyChangedEvent& InEvent);
 	void OnPCGGraphGeneratedOrCleaned(UPCGComponent* InComponent);
@@ -221,4 +222,10 @@ private:
 
 	/** Transient list of tags, kept when there is a tag change on a tracked Actor. */
 	TSet<FName> TempTrackedActorTags;
+
+#if WITH_EDITOR
+	// Part for the delayed landscape change update
+	double LastLandscapeDirtyTime = -1.0;
+	TArray<TObjectKey<ALandscapeProxy>> DelayedModifiedLandscapes;
+#endif // WITH_EDITOR
 };
\ No newline at end of file
