Tentative fix for the landscape cache causing loading errors wrt WP references
Submitted on behalf of julien.lheureux

#rnx

diff --git a/Engine/Plugins/Experimental/PCG/Source/PCG/Private/PCGWorldActor.cpp b/Engine/Plugins/Experimental/PCG/Source/PCG/Private/PCGWorldActor.cpp
--- a/Engine/Plugins/Experimental/PCG/Source/PCG/Private/PCGWorldActor.cpp
+++ b/Engine/Plugins/Experimental/PCG/Source/PCG/Private/PCGWorldActor.cpp
@@ -51,14 +51,17 @@ void APCGWorldActor::BeginCacheForCookedPlatformData(const ITargetPlatform* Targ
 
 		if (UWorldPartition* WorldPartition = World->GetWorldPartition())
 		{
-			/*
 			FWorldPartitionHelpers::ForEachActorDesc<ALandscapeProxy>(WorldPartition, [WorldPartition, &ActorRefs](const FWorldPartitionActorDesc* ActorDesc)
 			{
 				check(ActorDesc);
-				ActorRefs.Add(FWorldPartitionReference(WorldPartition, ActorDesc->GetGuid()));
+				// Create WP references only for actors that aren't currently loaded, otherwise we might end up unloading them
+				// if their actor desc ref count isn't setup properly
+				if (!ActorDesc->GetActor())
+				{
+					ActorRefs.Add(FWorldPartitionReference(WorldPartition, ActorDesc->GetGuid()));
+				}
 				return true;
 			});
-			*/
 		}
 		else
 		{
