[PCG] Fix cook bug where local components hold a hard ref to global component's SchedulingPolicy, preventing saves
#rnx

diff --git a/Engine/Plugins/Experimental/PCG/Source/PCG/Private/PCGComponent.cpp b/Engine/Plugins/Experimental/PCG/Source/PCG/Private/PCGComponent.cpp
--- a/Engine/Plugins/Experimental/PCG/Source/PCG/Private/PCGComponent.cpp
+++ b/Engine/Plugins/Experimental/PCG/Source/PCG/Private/PCGComponent.cpp
@@ -224,17 +224,14 @@ void UPCGComponent::SetPropertiesFromOriginal(const UPCGComponent* Original)
 
 #if WITH_EDITOR
 	const bool bHasDirtyInput = InputType != NewInputType;
-	const bool bIsDirty = bHasDirtyInput || bGraphInstanceIsDifferent;
+	bool bIsDirty = bHasDirtyInput || bGraphInstanceIsDifferent;
 #endif // WITH_EDITOR
 
 	InputType = NewInputType;
 	Seed = Original->Seed;
 	GenerationTrigger = Original->GenerationTrigger;
-
 	bOverrideGenerationRadii = Original->bOverrideGenerationRadii;
 	GenerationRadii = Original->GenerationRadii;
-	SchedulingPolicyClass = Original->SchedulingPolicyClass;
-	SchedulingPolicy = Original->SchedulingPolicy;
 
 	UPCGGraph* OriginalGraph = Original->GraphInstance ? Original->GraphInstance->GetGraph() : nullptr;
 	if (OriginalGraph != GraphInstance->GetGraph())
@@ -247,6 +244,21 @@ void UPCGComponent::SetPropertiesFromOriginal(const UPCGComponent* Original)
 		GraphInstance->CopyParameterOverrides(Original->GraphInstance);
 	}
 
+	if (!SchedulingPolicy || SchedulingPolicyClass != Original->SchedulingPolicyClass)
+	{
+		SchedulingPolicyClass = Original->SchedulingPolicyClass;
+		RefreshSchedulingPolicy();
+	}
+
+	if (SchedulingPolicy && ensure(Original->SchedulingPolicy) && !SchedulingPolicy->IsEquivalent(Original->SchedulingPolicy))
+	{
+		UEngine::CopyPropertiesForUnrelatedObjects(Original->SchedulingPolicy, SchedulingPolicy);
+
+#if WITH_EDITOR
+		bIsDirty = true;
+#endif
+	}
+
 #if WITH_EDITOR
 	// Note that while we dirty here, we won't trigger a refresh since we don't have the required context
 	if (bIsDirty)
diff --git a/Engine/Plugins/Experimental/PCG/Source/PCG/Private/RuntimeGen/SchedulingPolicies/PCGSchedulingPolicyDistanceAndDirection.cpp b/Engine/Plugins/Experimental/PCG/Source/PCG/Private/RuntimeGen/SchedulingPolicies/PCGSchedulingPolicyDistanceAndDirection.cpp
--- a/Engine/Plugins/Experimental/PCG/Source/PCG/Private/RuntimeGen/SchedulingPolicies/PCGSchedulingPolicyDistanceAndDirection.cpp
+++ b/Engine/Plugins/Experimental/PCG/Source/PCG/Private/RuntimeGen/SchedulingPolicies/PCGSchedulingPolicyDistanceAndDirection.cpp
@@ -55,3 +55,28 @@ double UPCGSchedulingPolicyDistanceAndDirection::CalculatePriority(const IPCGGen
 
 	return Priority;
 }
+
+bool UPCGSchedulingPolicyDistanceAndDirection::IsEquivalent(const UPCGSchedulingPolicyBase* OtherSchedulingPolicy) const
+{
+	if (!OtherSchedulingPolicy)
+	{
+		return false;
+	}
+
+	if (this == OtherSchedulingPolicy)
+	{
+		return true;
+	}
+
+	if (const UPCGSchedulingPolicyDistanceAndDirection* Other = Cast<UPCGSchedulingPolicyDistanceAndDirection>(OtherSchedulingPolicy))
+	{
+		return bUseDistance == Other->bUseDistance
+			&& DistanceWeight == Other->DistanceWeight
+			&& bUseDirection == Other->bUseDirection
+			&& DirectionWeight == Other->DirectionWeight;
+	}
+	else
+	{
+		return false;
+	}
+}
diff --git a/Engine/Plugins/Experimental/PCG/Source/PCG/Public/RuntimeGen/SchedulingPolicies/PCGSchedulingPolicyBase.h b/Engine/Plugins/Experimental/PCG/Source/PCG/Public/RuntimeGen/SchedulingPolicies/PCGSchedulingPolicyBase.h
--- a/Engine/Plugins/Experimental/PCG/Source/PCG/Public/RuntimeGen/SchedulingPolicies/PCGSchedulingPolicyBase.h
+++ b/Engine/Plugins/Experimental/PCG/Source/PCG/Public/RuntimeGen/SchedulingPolicies/PCGSchedulingPolicyBase.h
@@ -23,6 +23,9 @@ public:
 	virtual double CalculatePriority(const IPCGGenSourceBase* InGenSource, const FBox& GenerationBounds, bool bUse2DGrid) const 
 		PURE_VIRTUAL(UPCGSchedulingPolicyBase::CalculatePriority, return 0.0;);
 
+	/** A SchedulingPolicy is equivalent to another SchedulingPolicy if they are the same (same ptr), or if they have the same type and parameter values. */
+	virtual bool IsEquivalent(const UPCGSchedulingPolicyBase* OtherSchedulingPolicy) const PURE_VIRTUAL(UPCGSchedulingPolicyBase::IsEquivalent, return false;);
+
 #if WITH_EDITOR
 	/** Sets whether or not properties should be displayed in the editor.Used to hide instanced SchedulingPolicy properties when runtime generation is not enabled. */
 	void SetShouldDisplayProperties(bool bInShouldDisplayProperties) { bShouldDisplayProperties = bInShouldDisplayProperties; }
diff --git a/Engine/Plugins/Experimental/PCG/Source/PCG/Public/RuntimeGen/SchedulingPolicies/PCGSchedulingPolicyDistanceAndDirection.h b/Engine/Plugins/Experimental/PCG/Source/PCG/Public/RuntimeGen/SchedulingPolicies/PCGSchedulingPolicyDistanceAndDirection.h
--- a/Engine/Plugins/Experimental/PCG/Source/PCG/Public/RuntimeGen/SchedulingPolicies/PCGSchedulingPolicyDistanceAndDirection.h
+++ b/Engine/Plugins/Experimental/PCG/Source/PCG/Public/RuntimeGen/SchedulingPolicies/PCGSchedulingPolicyDistanceAndDirection.h
@@ -20,8 +20,12 @@ class PCG_API UPCGSchedulingPolicyDistanceAndDirection : public UPCGSchedulingPo
 	GENERATED_BODY()
 
 public:
+	/** Calculate the runtime scheduling priority with respect to a Generation Source. Should return a value in the range [0, 1], where higher values will be scheduled sooner. */
 	virtual double CalculatePriority(const IPCGGenSourceBase* InGenSource, const FBox& GenerationBounds, bool bUse2DGrid) const override;
 
+	/** A SchedulingPolicy is equivalent to another SchedulingPolicy if they are the same (same ptr), or if they have the same type and parameter values. */
+	virtual bool IsEquivalent(const UPCGSchedulingPolicyBase* OtherSchedulingPolicy) const;
+
 public:
 	/** Toggle whether or not distance is used to calculate the scheduling priority. */
 	UPROPERTY(BlueprintReadWrite, EditAnywhere, Category = "RuntimeGeneration|Scheduling Policy Parameters", meta = (EditCondition = "bShouldDisplayProperties", EditConditionHides))
