[PCG] Add some defensive code in PCGGraphExecutor::Cancel which may help graph execution cancellation issues.

#rnx

diff --git a/Engine/Plugins/Experimental/PCG/Source/PCG/Private/Graph/PCGGraphExecutor.cpp b/Engine/Plugins/Experimental/PCG/Source/PCG/Private/Graph/PCGGraphExecutor.cpp
--- a/Engine/Plugins/Experimental/PCG/Source/PCG/Private/Graph/PCGGraphExecutor.cpp
+++ b/Engine/Plugins/Experimental/PCG/Source/PCG/Private/Graph/PCGGraphExecutor.cpp
@@ -256,6 +256,12 @@ TSet<UPCGComponent*> FPCGGraphExecutor::Cancel(TFunctionRef<bool(TWeakObjectPtr<
 		}
 	}
 
+	// In one instance this function was observed to return nullptr in the CancelledComponents set.
+	// All the cancel filter lambdas check the ptr is valid, so it's not clear why. It seems perhaps
+	// the SourceComponent weak ptr became nullptr between calling CancelFilter and adding the component
+	// to the set.
+	ensure(CancelledComponents.Remove(nullptr) == 0);
+
 	// Early out - nothing to cancel
 	if (CancelledComponents.IsEmpty())
 	{
