[PCG] Remove unnecessary ensure from graph visual logs.

It's not surprising that component weak pointer can go null. Null entries are cleaned out fairly frequenly in ClearLogs(), so just ignore here.

#rnx

diff --git a/Engine/Plugins/Experimental/PCG/Source/PCG/Private/Utils/PCGNodeVisualLogs.cpp b/Engine/Plugins/Experimental/PCG/Source/PCG/Private/Utils/PCGNodeVisualLogs.cpp
--- a/Engine/Plugins/Experimental/PCG/Source/PCG/Private/Utils/PCGNodeVisualLogs.cpp
+++ b/Engine/Plugins/Experimental/PCG/Source/PCG/Private/Utils/PCGNodeVisualLogs.cpp
@@ -134,7 +134,7 @@ FText FPCGNodeVisualLogs::GetLogsSummaryText(TWeakObjectPtr<const UPCGNode> InNo
 
 		if (!InComponent)
 		{
-			FText ActorName = ensure(NodeLogs[i].Component.Get() && NodeLogs[i].Component->GetOwner()) ? FText::FromName(NodeLogs[i].Component->GetOwner()->GetFName()) : FText::FromString(TEXT("MissingComponent"));
+			FText ActorName = (NodeLogs[i].Component.Get() && NodeLogs[i].Component->GetOwner()) ? FText::FromName(NodeLogs[i].Component->GetOwner()->GetFName()) : FText::FromString(TEXT("MissingComponent"));
 			ResultText = FText::Format(LOCTEXT("NodeTooltipLogWithActor", "{0}{1}/{2}: [{3}] {4}: {5}"), ResultText, i + 1, NodeLogs.Num(), ActorName, MessageVerbosity, NodeLogs[i].Message);
 		}
 		else
