[PCG] Make sure we don't have a Scheduling Policy if runtime generation is disabled

#rnx

diff --git a/Engine/Plugins/Experimental/PCG/Source/PCG/Private/PCGComponent.cpp b/Engine/Plugins/Experimental/PCG/Source/PCG/Private/PCGComponent.cpp
--- a/Engine/Plugins/Experimental/PCG/Source/PCG/Private/PCGComponent.cpp
+++ b/Engine/Plugins/Experimental/PCG/Source/PCG/Private/PCGComponent.cpp
@@ -1287,7 +1287,7 @@ void UPCGComponent::PostLoad()
 	}
 #endif
 
-	if (!SchedulingPolicy)
+	if (!SchedulingPolicy || GenerationTrigger != EPCGComponentGenerationTrigger::GenerateAtRuntime)
 	{
 		RefreshSchedulingPolicy();
 	}
@@ -2398,20 +2398,22 @@ FPCGDataCollection UPCGComponent::CreateActorPCGDataCollection(AActor* Actor, co
 
 void UPCGComponent::RefreshSchedulingPolicy()
 {
-	if (SchedulingPolicy)
+	// Only delete it if we are the owner, it's for deprecation where local components had hard ref on original policy.
+	if (IsValid(SchedulingPolicy) && SchedulingPolicy->GetOuter() == this)
 	{
 		SchedulingPolicy->Rename(nullptr, GetTransientPackage(), REN_DontCreateRedirectors | REN_ForceNoResetLoaders);
 		SchedulingPolicy->MarkAsGarbage();
-		SchedulingPolicy = nullptr;
 	}
 
-	if (SchedulingPolicyClass)
+	SchedulingPolicy = nullptr;
+
+	// We should never create the scheduling policy when not in runtime generation mode.
+	if (SchedulingPolicyClass && GenerationTrigger == EPCGComponentGenerationTrigger::GenerateAtRuntime)
 	{
 		const EObjectFlags Flags = GetMaskedFlags(RF_PropagateToSubObjects);
 		SchedulingPolicy = NewObject<UPCGSchedulingPolicyBase>(this, SchedulingPolicyClass, NAME_None, Flags);
 
 #if WITH_EDITOR
-		// We should only display scheduling policy parameters when in runtime generation mode.
 		SchedulingPolicy->SetShouldDisplayProperties(GenerationTrigger == EPCGComponentGenerationTrigger::GenerateAtRuntime);
 #endif
 	}
