[PCG] Minor argument rename to avoid confusion

#rnx

diff --git a/Engine/Plugins/Experimental/PCG/Source/PCG/Private/Metadata/PCGMetadata.cpp b/Engine/Plugins/Experimental/PCG/Source/PCG/Private/Metadata/PCGMetadata.cpp
--- a/Engine/Plugins/Experimental/PCG/Source/PCG/Private/Metadata/PCGMetadata.cpp
+++ b/Engine/Plugins/Experimental/PCG/Source/PCG/Private/Metadata/PCGMetadata.cpp
@@ -807,29 +807,29 @@ void UPCGMetadata::AddDelayedEntries(const TArray<TTuple<int64, int64>>& AllEntr
 	DelayedEntriesIndex.Exchange(0);
 }
 
-bool UPCGMetadata::InitializeOnSet(PCGMetadataEntryKey& InKey, PCGMetadataEntryKey InParentKeyA, const UPCGMetadata* InParentMetadataA, PCGMetadataEntryKey InParentKeyB, const UPCGMetadata* InParentMetadataB)
+bool UPCGMetadata::InitializeOnSet(PCGMetadataEntryKey& InOutKey, PCGMetadataEntryKey InParentKeyA, const UPCGMetadata* InParentMetadataA, PCGMetadataEntryKey InParentKeyB, const UPCGMetadata* InParentMetadataB)
 {
-	if (InKey == PCGInvalidEntryKey)
+	if (InOutKey == PCGInvalidEntryKey)
 	{
 		if (InParentKeyA != PCGInvalidEntryKey && Parent == InParentMetadataA)
 		{
-			InKey = AddEntry(InParentKeyA);
+			InOutKey = AddEntry(InParentKeyA);
 			return true;
 		}
 		else if (InParentKeyB != PCGInvalidEntryKey && Parent == InParentMetadataB)
 		{
-			InKey = AddEntry(InParentKeyB);
+			InOutKey = AddEntry(InParentKeyB);
 			return true;
 		}
 		else
 		{
-			InKey = AddEntry();
+			InOutKey = AddEntry();
 			return false;
 		}
 	}
-	else if(InKey < ItemKeyOffset)
+	else if(InOutKey < ItemKeyOffset)
 	{
-		InKey = AddEntry(InKey);
+		InOutKey = AddEntry(InOutKey);
 		return false;
 	}
 	else
diff --git a/Engine/Plugins/Experimental/PCG/Source/PCG/Public/Metadata/PCGMetadata.h b/Engine/Plugins/Experimental/PCG/Source/PCG/Public/Metadata/PCGMetadata.h
--- a/Engine/Plugins/Experimental/PCG/Source/PCG/Public/Metadata/PCGMetadata.h
+++ b/Engine/Plugins/Experimental/PCG/Source/PCG/Public/Metadata/PCGMetadata.h
@@ -227,7 +227,7 @@ public:
 	void AddDelayedEntries(const TArray<TTuple<int64, int64>>& AllEntries);
 
 	/** Initializes the metadata entry key. Returns true if key set from either parent */
-	bool InitializeOnSet(PCGMetadataEntryKey& InKey, PCGMetadataEntryKey InParentKeyA = PCGInvalidEntryKey, const UPCGMetadata* InParentMetadataA = nullptr, PCGMetadataEntryKey InParentKeyB = PCGInvalidEntryKey, const UPCGMetadata* InParentMetadataB = nullptr);
+	bool InitializeOnSet(PCGMetadataEntryKey& InOutKey, PCGMetadataEntryKey InParentKeyA = PCGInvalidEntryKey, const UPCGMetadata* InParentMetadataA = nullptr, PCGMetadataEntryKey InParentKeyB = PCGInvalidEntryKey, const UPCGMetadata* InParentMetadataB = nullptr);
 
 	/** Metadata chaining mechanism */
 	PCGMetadataEntryKey GetParentKey(PCGMetadataEntryKey LocalItemKey) const;
