Minor fixes for the transient workflow
- Fixed some missing cleanup when going from normal to preview modes wrt actors
- Added transient propagation through PCG components in managed resources
- Made sure it works properly for spawned actors too

Submitted on behalf of julien.lheureux

#rnx

diff --git a/Engine/Plugins/Experimental/PCG/Source/PCG/Private/Elements/PCGSpawnActor.cpp b/Engine/Plugins/Experimental/PCG/Source/PCG/Private/Elements/PCGSpawnActor.cpp
--- a/Engine/Plugins/Experimental/PCG/Source/PCG/Private/Elements/PCGSpawnActor.cpp
+++ b/Engine/Plugins/Experimental/PCG/Source/PCG/Private/Elements/PCGSpawnActor.cpp
@@ -990,6 +990,15 @@ void FPCGSpawnActorElement::SpawnActors(FPCGSubgraphContext* Context, AActor* Ta
 
 		for (UPCGComponent* PCGComponent : PCGComponents)
 		{
+#if WITH_EDITOR
+			// For both pre-existing and new actors, we need to make sure we're inline with loading/generation as needed
+			if (PCGComponent->GetEditingMode() != Context->SourceComponent->GetEditingMode())
+			{
+				PCGComponent->SetEditingMode(/*CurrentEditingMode=*/Context->SourceComponent->GetEditingMode(), /*SerializedEditingMode=*/Context->SourceComponent->GetEditingMode());
+				PCGComponent->ChangeTransientState(Context->SourceComponent->GetEditingMode());
+			}
+#endif // WITH_EDITOR
+
 			if (Settings->Option == EPCGSpawnActorOption::NoMerging)
 			{
 				if (bForceDisableActorParsing)
diff --git a/Engine/Plugins/Experimental/PCG/Source/PCG/Private/PCGComponent.cpp b/Engine/Plugins/Experimental/PCG/Source/PCG/Private/PCGComponent.cpp
--- a/Engine/Plugins/Experimental/PCG/Source/PCG/Private/PCGComponent.cpp
+++ b/Engine/Plugins/Experimental/PCG/Source/PCG/Private/PCGComponent.cpp
@@ -220,6 +220,11 @@ void UPCGComponent::SetPropertiesFromOriginal(const UPCGComponent* Original)
 		}
 	}
 
+	if (!ensure(GraphInstance))
+	{
+		return;
+	}
+
 	const bool bGraphInstanceIsDifferent = !GraphInstance->IsEquivalent(Original->GraphInstance);
 
 #if WITH_EDITOR
@@ -1492,6 +1497,8 @@ void UPCGComponent::PostEditChangeProperty(FPropertyChangedEvent& PropertyChange
 
 	const FName PropName = PropertyChangedEvent.Property->GetFName();
 
+	bool bTransientPropertyChangedThatDoesNotRequireARefresh = false;
+
 	// Implementation note:
 	// Since the current editing mode is a transient variable, if we do not do this transition here before going in the Super call,
 	//  we can end up in a situation where BP actors are reconstructed (... this component included ...) which makes this fall into the !IsValid case just after
@@ -1500,6 +1507,7 @@ void UPCGComponent::PostEditChangeProperty(FPropertyChangedEvent& PropertyChange
 		// When affecting the editing mode from the user's point of view, we need to change both the current & serialized values
 		SetEditingMode(CurrentEditingMode, CurrentEditingMode);
 		ChangeTransientState(CurrentEditingMode);
+		bTransientPropertyChangedThatDoesNotRequireARefresh = true;
 	}
 
 	Super::PostEditChangeProperty(PropertyChangedEvent);
@@ -1590,7 +1598,7 @@ void UPCGComponent::PostEditChangeProperty(FPropertyChangedEvent& PropertyChange
 		}
 	}
 	// General properties that don't affect behavior
-	else
+	else if(!bTransientPropertyChangedThatDoesNotRequireARefresh)
 	{
 		Refresh();
 	}
@@ -2674,7 +2682,7 @@ bool UPCGComponent::DeletePreviewResources()
 	for (TObjectPtr<UPCGManagedResource> ResourceToRelease : LoadedPreviewResources)
 	{
 		// Changing the transient state will clear the "marked transient on load" flag
-		ResourceToRelease->ChangeTransientState(/*bNowTransient=*/false);
+		ResourceToRelease->ChangeTransientState(EPCGEditorDirtyMode::Normal);
 		ResourceToRelease->Release(/*bHardRelease=*/true, ActorsToDelete);
 		bResourceWasReleased = true;
 	}
@@ -2713,7 +2721,7 @@ void UPCGComponent::ChangeTransientState(EPCGEditorDirtyMode NewEditingMode)
 		{
 			if (GeneratedResource)
 			{
-				GeneratedResource->ChangeTransientState(NewEditingMode == EPCGEditorDirtyMode::Preview);
+				GeneratedResource->ChangeTransientState(NewEditingMode);
 				bShouldMarkDirty = true;
 			}
 		}
@@ -2728,7 +2736,10 @@ void UPCGComponent::ChangeTransientState(EPCGEditorDirtyMode NewEditingMode)
 
 	if (IsLocalComponent())
 	{
-		bShouldMarkDirty = true;
+		if (NewEditingMode == EPCGEditorDirtyMode::Preview)
+		{
+			MarkPackageDirty();
+		}
 
 		if (NewEditingMode == EPCGEditorDirtyMode::Preview)
 		{
@@ -2750,9 +2761,13 @@ void UPCGComponent::ChangeTransientState(EPCGEditorDirtyMode NewEditingMode)
 				Object->ClearFlags(RF_Transient);
 			}
 		});
-	}
 
-	if (bShouldMarkDirty)
+		if (NewEditingMode != EPCGEditorDirtyMode::Preview)
+		{
+			MarkPackageDirty();
+		}
+	}
+	else if (bShouldMarkDirty)
 	{
 		MarkPackageDirty();
 	}
diff --git a/Engine/Plugins/Experimental/PCG/Source/PCG/Private/PCGManagedResource.cpp b/Engine/Plugins/Experimental/PCG/Source/PCG/Private/PCGManagedResource.cpp
--- a/Engine/Plugins/Experimental/PCG/Source/PCG/Private/PCGManagedResource.cpp
+++ b/Engine/Plugins/Experimental/PCG/Source/PCG/Private/PCGManagedResource.cpp
@@ -8,12 +8,9 @@
 
 #include "Components/InstancedStaticMeshComponent.h"
 #include "Engine/Level.h"
+#include "UObject/Package.h"
 #include "Utils/PCGGeneratedResourcesLogging.h"
 
-#if WITH_EDITOR
-#include "ObjectTools.h"
-#endif
-
 #include UE_INLINE_GENERATED_CPP_BY_NAME(PCGManagedResource)
 
 static TAutoConsoleVariable<bool> CVarForceReleaseResourcesOnGenerate(
@@ -65,7 +62,7 @@ bool UPCGManagedResource::DebugForcePurgeAllResourcesOnGenerate()
 }
 
 #if WITH_EDITOR
-void UPCGManagedResource::ChangeTransientState(bool /*bNowTransient*/)
+void UPCGManagedResource::ChangeTransientState(EPCGEditorDirtyMode /*NewEditingMode*/)
 {
 	// Any change in the transient state resets the transient state that was set on load, regardless of the bNowTransient flag
 	bMarkedTransientOnLoad = false;
@@ -210,62 +207,79 @@ void UPCGManagedActors::MarkAsReused()
 }
 
 #if WITH_EDITOR
-void UPCGManagedActors::ChangeTransientState(bool bNowTransient)
+void UPCGManagedActors::ChangeTransientState(EPCGEditorDirtyMode NewEditingMode)
 {
-	TSet<UPackage*> PackagesToCleanup;
+	const bool bNowTransient = (NewEditingMode == EPCGEditorDirtyMode::Preview);
 
 	for (TSoftObjectPtr<AActor> GeneratedActor : GeneratedActors)
 	{
 		// Make sure to load if needed because we need to affect the actors regardless of the current WP state
 		if (GeneratedActor.LoadSynchronous() != nullptr)
 		{
-			// If this is changed during loading, we shouldn't set the actor to be transient if we're using external packages,
-			// otherwise the package will get deleted spuriously.
-			if (bNowTransient)
+			const bool bWasTransient = GeneratedActor->HasAnyFlags(RF_Transient);
+
+			if (bNowTransient != bWasTransient)
 			{
-				if (!GeneratedActor->IsPackageExternal())
+				if (bNowTransient)
 				{
 					GeneratedActor->SetFlags(RF_Transient);
 				}
-			}
-			else
-			{
-				GeneratedActor->ClearFlags(RF_Transient);
+				else
+				{
+					GeneratedActor->ClearFlags(RF_Transient);
+				}
 			}
 
-			if (GeneratedActor->GetLevel() && GeneratedActor->GetLevel()->IsUsingExternalActors())
+			// If the actor had PCG components, propagate this downward
 			{
-				if (bNowTransient)
+				TInlineComponentArray<UPCGComponent*, 4> PCGComponents;
+				GeneratedActor->GetComponents(PCGComponents);
+
+				for (UPCGComponent* PCGComponent : PCGComponents)
 				{
-					PackagesToCleanup.Add(GeneratedActor->GetExternalPackage());
+					PCGComponent->SetEditingMode(/*CurrentEditingMode=*/NewEditingMode, /*SerializedEditingMode=*/NewEditingMode);
+					PCGComponent->ChangeTransientState(NewEditingMode);
 				}
-
-				GeneratedActor->SetPackageExternal(/*bExternal=*/!bNowTransient, /*bShouldDirty=*/false);
 			}
 
-			if (!bNowTransient || GeneratedActor->IsPackageExternal())
+			if (bNowTransient != bWasTransient)
 			{
-				ForEachObjectWithOuter(GeneratedActor.Get(), [bNowTransient](UObject* Object)
+				if (GeneratedActor->GetLevel() && GeneratedActor->GetLevel()->IsUsingExternalActors())
 				{
-					if (bNowTransient)
+					UPackage* PreviousExternalPackage = bNowTransient ? GeneratedActor->GetExternalPackage() : nullptr;
+
+					GeneratedActor->SetPackageExternal(/*bExternal=*/!bNowTransient, /*bShouldDirty=*/false);
+
+					if (PreviousExternalPackage)
 					{
-						Object->SetFlags(RF_Transient);
+						PreviousExternalPackage->MarkPackageDirty();
 					}
-					else
+				
+					if (!bNowTransient && GeneratedActor->GetExternalPackage())
 					{
-						Object->ClearFlags(RF_Transient);
+						MarkPackageDirty();
 					}
-				});
+				}
+
+				if (!bNowTransient || GeneratedActor->IsPackageExternal())
+				{
+					ForEachObjectWithOuter(GeneratedActor.Get(), [bNowTransient](UObject* Object)
+					{
+						if (bNowTransient)
+						{
+							Object->SetFlags(RF_Transient);
+						}
+						else
+						{
+							Object->ClearFlags(RF_Transient);
+						}
+					});
+				}
 			}
 		}
 	}
 
-	if (!PackagesToCleanup.IsEmpty())
-	{
-		ObjectTools::CleanupAfterSuccessfulDelete(PackagesToCleanup.Array(), /*bPerformanceReferenceCheck=*/true);
-	}
-
-	Super::ChangeTransientState(bNowTransient);
+	Super::ChangeTransientState(NewEditingMode);
 }
 #endif // WITH_EDITOR
 
@@ -446,8 +460,9 @@ void UPCGManagedComponent::MarkAsReused()
 }
 
 #if WITH_EDITOR
-void UPCGManagedComponent::ChangeTransientState(bool bNowTransient)
+void UPCGManagedComponent::ChangeTransientState(EPCGEditorDirtyMode NewEditingMode)
 {
+	const bool bNowTransient = (NewEditingMode == EPCGEditorDirtyMode::Preview);
 	if (GeneratedComponent.Get())
 	{
 		const bool bWasTransient = GeneratedComponent->HasAnyFlags(RF_Transient);
@@ -479,7 +494,7 @@ void UPCGManagedComponent::ChangeTransientState(bool bNowTransient)
 		}
 	}
 
-	Super::ChangeTransientState(bNowTransient);
+	Super::ChangeTransientState(NewEditingMode);
 }
 #endif // WITH_EDITOR
 
diff --git a/Engine/Plugins/Experimental/PCG/Source/PCG/Public/PCGManagedResource.h b/Engine/Plugins/Experimental/PCG/Source/PCG/Public/PCGManagedResource.h
--- a/Engine/Plugins/Experimental/PCG/Source/PCG/Public/PCGManagedResource.h
+++ b/Engine/Plugins/Experimental/PCG/Source/PCG/Public/PCGManagedResource.h
@@ -52,7 +52,7 @@ public:
 	void SetCrc(const FPCGCrc& InCrc) { Crc = InCrc; }
 
 #if WITH_EDITOR
-	virtual void ChangeTransientState(bool bNowTransient);
+	virtual void ChangeTransientState(EPCGEditorDirtyMode NewEditingMode);
 	virtual void MarkTransientOnLoad() { bMarkedTransientOnLoad = true; }
 
 	bool IsMarkedTransientOnLoad() const { return bMarkedTransientOnLoad; }
@@ -94,7 +94,7 @@ public:
 	virtual void MarkAsReused() override;
 
 #if WITH_EDITOR
-	virtual void ChangeTransientState(bool bNowTransient) override;
+	virtual void ChangeTransientState(EPCGEditorDirtyMode NewEditingMode) override;
 #endif
 	//~End UPCGManagedResource interface
 
@@ -118,7 +118,7 @@ public:
 	virtual bool MoveResourceToNewActor(AActor* NewActor, const AActor* ExpectedPreviousOwner) override;
 
 #if WITH_EDITOR
-	virtual void ChangeTransientState(bool bNowTransient) override;
+	virtual void ChangeTransientState(EPCGEditorDirtyMode NewEditingMode) override;
 	/** Hides the content of the component in a transient way (such as unregistering) */
 	virtual void HideComponent();
 #endif
