[PCG] Add Debug category to PCGComponent properties and expose a few additional bits that we needed while debugging.

#rnx

diff --git a/Engine/Plugins/Experimental/PCG/Source/PCG/Public/PCGComponent.h b/Engine/Plugins/Experimental/PCG/Source/PCG/Public/PCGComponent.h
--- a/Engine/Plugins/Experimental/PCG/Source/PCG/Public/PCGComponent.h
+++ b/Engine/Plugins/Experimental/PCG/Source/PCG/Public/PCGComponent.h
@@ -200,23 +200,17 @@ public:
 	/** Compute the runtime cleanup radius for the given grid size. */
 	double GetCleanupRadiusFromGrid(EPCGHiGenGrid Grid) const;
 
-	UPROPERTY(BlueprintReadWrite, EditAnywhere, Category = Settings)
-	EPCGComponentInput InputType = EPCGComponentInput::Actor;
-
-	UPROPERTY(BlueprintReadWrite, EditAnywhere, Category = Settings)
-	bool bParseActorComponents = true;
-
-	UPROPERTY(BlueprintReadWrite, EditAnywhere, Category = Settings)
+	UPROPERTY(BlueprintReadWrite, EditAnywhere, Category = Settings, meta = (DisplayPriority = 600))
 	int Seed = 42;
 
-	UPROPERTY(BlueprintReadWrite, EditAnywhere, Category = Properties)
+	UPROPERTY(BlueprintReadWrite, EditAnywhere, Category = Settings, meta = (DisplayPriority = 100))
 	bool bActivated = true;
 
 	/* In World Partition map, will partition the component in a grid, according to PCGWorldActor settings, dispatching the generation to multiple local components.*/
-	UPROPERTY(BlueprintReadWrite, EditAnywhere, AdvancedDisplay, Category = Properties, meta = (EditCondition = "!bIsComponentLocal", DisplayName = "Is Partitioned"))
+	UPROPERTY(BlueprintReadWrite, EditAnywhere, Category = Settings, meta = (EditCondition = "!bIsComponentLocal", DisplayName = "Is Partitioned", DisplayPriority = 500))
 	bool bIsComponentPartitioned = false;
 
-	UPROPERTY(BlueprintReadOnly, EditAnywhere, Category = Properties, AdvancedDisplay, meta = (EditCondition = "!bIsComponentLocal", EditConditionHides))
+	UPROPERTY(BlueprintReadOnly, EditAnywhere, Category = Settings, meta = (EditCondition = "!bIsComponentLocal", EditConditionHides, DisplayPriority = 200))
 	EPCGComponentGenerationTrigger GenerationTrigger = EPCGComponentGenerationTrigger::GenerateOnLoad;
 
 	/** Manual overrides for the graph generation radii and cleanup radius multiplier. */
@@ -234,18 +228,11 @@ public:
 	UPROPERTY(BlueprintReadOnly, VisibleAnywhere, Instanced, Category = RuntimeGeneration, meta = (EditCondition = "GenerationTrigger == EPCGComponentGenerationTrigger::GenerateAtRuntime", EditConditionHides))
 	TObjectPtr<UPCGSchedulingPolicyBase> SchedulingPolicy;
 
-	/** Flag to indicate whether this component has run in the editor. Note that for partitionable actors, this will always be false. */
-	UPROPERTY(BlueprintReadOnly, VisibleAnywhere, AdvancedDisplay, Category = Properties, NonTransactional, meta = (NoResetToDefault))
-	bool bGenerated = false;
-
-	UPROPERTY(NonPIEDuplicateTransient)
-	bool bRuntimeGenerated = false;
-
 #if WITH_EDITORONLY_DATA
-	UPROPERTY(BlueprintReadWrite, EditAnywhere, AdvancedDisplay, Category = Properties, meta = (DisplayName = "Regenerate PCG volume in editor"))
+	UPROPERTY(BlueprintReadWrite, EditAnywhere, Category = "Editing Settings", meta = (DisplayName = "Regenerate PCG Volume In Editor", DisplayPriority = 400))
 	bool bRegenerateInEditor = true;
 
-	UPROPERTY(BlueprintReadOnly, VisibleAnywhere, Transient, AdvancedDisplay, Category = Properties, meta = (NoResetToDefault))
+	UPROPERTY(BlueprintReadOnly, VisibleAnywhere, Transient, Category = Debug, meta = (NoResetToDefault))
 	bool bDirtyGenerated = false;
 
 	// Property that will automatically be set on BP templates, to allow for "Generate on add to world" in editor.
@@ -258,10 +245,17 @@ public:
 	FOnPCGGraphCleaned OnPCGGraphCleanedDelegate;
 #endif
 
+	/** Flag to indicate whether this component has run in the editor. Note that for partitionable actors, this will always be false. */
+	UPROPERTY(BlueprintReadOnly, VisibleAnywhere, Category = Debug, NonTransactional, meta = (NoResetToDefault))
+	bool bGenerated = false;
+
+	UPROPERTY(NonPIEDuplicateTransient)
+	bool bRuntimeGenerated = false;
+
 	/** Can specify a list of functions from the owner of this component to be called when generation is done, in order.
 	*   Need to take (and only take) a PCGDataCollection as parameter and with "CallInEditor" flag enabled.
 	*/
-	UPROPERTY(BlueprintReadWrite, EditAnywhere, AdvancedDisplay, Category = Properties)
+	UPROPERTY(BlueprintReadWrite, EditAnywhere, Category = Settings, meta = (DisplayPriority = 700))
 	TArray<FName> PostGenerateFunctionNames;
 
 	/** Return if we are currently generating the graph for this component */
@@ -346,19 +340,27 @@ protected:
 	void RefreshSchedulingPolicy();
 
 protected:
-	UPROPERTY(BlueprintReadOnly, VisibleAnywhere, Category = PCG, Instanced, meta = (NoResetToDefault))
+	UPROPERTY(BlueprintReadOnly, VisibleAnywhere, Category = Settings, Instanced, meta = (NoResetToDefault, DisplayPriority = 100))
 	TObjectPtr<UPCGGraphInstance> GraphInstance;
 
-	UPROPERTY(Transient, VisibleAnywhere, AdvancedDisplay, Category = Properties)
+	UPROPERTY(Transient, VisibleAnywhere, Category = Debug)
 	uint32 GenerationGridSize = PCGHiGenGrid::UnboundedGridSize();
 
 	// Current editing mode that depends on the serialized editing mode and loading
-	UPROPERTY(Transient, EditAnywhere, AdvancedDisplay, Category = Properties, meta = (DisplayName = "Editing Mode", EditCondition = "!bIsComponentLocal"))
+	UPROPERTY(Transient, EditAnywhere, Category = "Editing Settings", meta = (DisplayName = "Editing Mode", EditCondition = "!bIsComponentLocal", DisplayPriority = 300))
 	EPCGEditorDirtyMode CurrentEditingMode = EPCGEditorDirtyMode::Normal;
 
-	UPROPERTY(VisibleAnywhere, AdvancedDisplay, Category = Properties, meta = (NoResetToDefault))
+	UPROPERTY(VisibleAnywhere, Category = Debug, meta = (NoResetToDefault))
 	EPCGEditorDirtyMode SerializedEditingMode = EPCGEditorDirtyMode::Normal;
 
+public:
+	UPROPERTY(BlueprintReadWrite, EditAnywhere, Category = "Input Node Settings (Deprecated)", meta = (DisplayPriority = 800))
+	EPCGComponentInput InputType = EPCGComponentInput::Actor;
+
+	UPROPERTY(BlueprintReadWrite, EditAnywhere, Category = "Input Node Settings (Deprecated)", meta = (DisplayPriority = 900))
+	bool bParseActorComponents = true;
+
+protected:
 #if WITH_EDITORONLY_DATA
 	UPROPERTY()
 	TObjectPtr<UPCGGraph> Graph_DEPRECATED;
@@ -458,19 +460,22 @@ private:
 #if WITH_EDITORONLY_DATA
 	UPROPERTY(Transient)
 	TArray<TObjectPtr<UPCGManagedResource>> LoadedPreviewResources;
+
+	UPROPERTY(Transient, VisibleInstanceOnly, Category = Debug)
+	bool bGenerationInProgress = false;
 #endif
 
 	// When doing a cleanup, locking resource modification. Used as sentinel.
 	bool GeneratedResourcesInaccessible = false;
 
-	UPROPERTY()
+	UPROPERTY(VisibleInstanceOnly, Category = Debug)
 	FBox LastGeneratedBounds = FBox(EForceInit::ForceInit);
 
-	UPROPERTY()
+	UPROPERTY(VisibleInstanceOnly, Category = Debug)
 	FPCGDataCollection GeneratedGraphOutput;
 
 	/** If any graph edges cross execution grid sizes, data on the edge is stored / retrieved from this map. */
-	UPROPERTY(Transient, VisibleAnywhere, Category = Properties, AdvancedDisplay)
+	UPROPERTY(Transient, VisibleAnywhere, Category = Debug)
 	TMap<FString, FPCGDataCollection> PerPinGeneratedOutput;
 
 	mutable FRWLock PerPinGeneratedOutputLock;
@@ -480,11 +485,9 @@ private:
 
 #if WITH_EDITOR
 	FPCGTaskId CurrentRefreshTask = InvalidPCGTaskId;
-
-	bool bGenerationInProgress = false;
 #endif // WITH_EDITOR
 
-	UPROPERTY(VisibleAnywhere, Transient, Category = Properties, meta = (EditCondition = false, EditConditionHides))
+	UPROPERTY(VisibleAnywhere, Transient, Category = Debug, meta = (EditCondition = false, EditConditionHides))
 	bool bIsComponentLocal = false;
 
 #if WITH_EDITOR
