[PCG] Support sampling Texture2DArrays in PCGTextureSampler in standalone builds

Modified PCGTextureReadback to support sampling a Texture2DArray at a provided index.

diff --git a/Engine/Plugins/Experimental/PCG/Shaders/Private/PCGTextureReadback.usf b/Engine/Plugins/Experimental/PCG/Shaders/Private/PCGTextureReadback.usf
--- a/Engine/Plugins/Experimental/PCG/Shaders/Private/PCGTextureReadback.usf
+++ b/Engine/Plugins/Experimental/PCG/Shaders/Private/PCGTextureReadback.usf
@@ -1,15 +1,20 @@
 // Copyright Epic Games, Inc. All Rights Reserved.
 
+#include "/Engine/Private/Common.ush"
 #include "/Engine/Public/Platform.ush"
 
-Texture2D SourceTexture;
+Texture2DArray<float4> SourceTextureArray;
 SamplerState SourceSampler;
 float2 SourceDimensions;
+int SourceTextureIndex;
+
 RWTexture2D<float4> OutputTexture;
 
 [numthreads(THREADGROUPSIZE_X, THREADGROUPSIZE_Y, THREADGROUPSIZE_Z)]
 void PCGTextureReadback_CS(uint3 DispatchThreadId : SV_DispatchThreadID, uint GroupIndex : SV_GroupIndex)
 {
-	float2 TexCoord = DispatchThreadId.xy + 0.5;
-	OutputTexture[TexCoord] = SourceTexture.SampleLevel(SourceSampler, TexCoord / SourceDimensions, /*lod=*/0);
+	float2 ThreadID = DispatchThreadId.xy + 0.5;
+	float2 TexCoord = ThreadID / SourceDimensions;
+
+	OutputTexture[ThreadID] = Texture2DArraySample(SourceTextureArray, SourceSampler, float3(TexCoord.xy, SourceTextureIndex));
 }
diff --git a/Engine/Plugins/Experimental/PCG/Source/PCG/Private/Data/PCGTextureData.cpp b/Engine/Plugins/Experimental/PCG/Source/PCG/Private/Data/PCGTextureData.cpp
--- a/Engine/Plugins/Experimental/PCG/Source/PCG/Private/Data/PCGTextureData.cpp
+++ b/Engine/Plugins/Experimental/PCG/Source/PCG/Private/Data/PCGTextureData.cpp
@@ -7,8 +7,10 @@
 #include "Helpers/PCGAsync.h"
 #include "Helpers/PCGHelpers.h"
 
+#include "RHIStaticStates.h"
 #include "TextureResource.h"
 #include "Engine/Texture2D.h"
+#include "Engine/Texture2DArray.h"
 
 #include UE_INLINE_GENERATED_CPP_BY_NAME(PCGTextureData)
 
@@ -264,11 +266,12 @@ void UPCGBaseTextureData::CopyBaseTextureData(UPCGBaseTextureData* NewTextureDat
 	NewTextureData->Width = Width;
 }
 
-void UPCGTextureData::Initialize(UTexture2D* InTexture, const FTransform& InTransform, const TFunction<void()>& PostInitializeCallback)
+void UPCGTextureData::Initialize(UTexture* InTexture, uint32 InTextureIndex, const FTransform& InTransform, const TFunction<void()>& PostInitializeCallback)
 {
 	TRACE_CPUPROFILER_EVENT_SCOPE(UPCGTextureData::Initialize);
 
 	Texture = InTexture;
+	TextureIndex = InTextureIndex;
 	Transform = InTransform;
 	Width = 0;
 	Height = 0;
@@ -326,7 +329,14 @@ bool UPCGTextureData::InitializeFromCPUTexture()
 		return false;
 	}
 
-	FSharedImageConstRef CPUTextureRef = Texture->GetCPUCopy();
+	// CPU Textures currently only support UTexture2D.
+	UTexture2D* Texture2D = Cast<UTexture2D>(Texture);
+	if (!Texture2D)
+	{
+		return false;
+	}
+
+	FSharedImageConstRef CPUTextureRef = Texture2D->GetCPUCopy();
 	if (!CPUTextureRef.IsValid())
 	{
 		return false;
@@ -433,7 +443,7 @@ bool UPCGTextureData::InitializeFromCPUTexture()
 	}
 	else
 	{
-		UE_LOG(LogPCG, Error, TEXT("PCGTextureReadback has an invalid format (%d) for CPU texture '%s'."), CPUTextureRef->Format, *Texture->GetFName().ToString());
+		UE_LOG(LogPCG, Error, TEXT("PCGTextureReadback has an invalid format (%d) for CPU texture '%s'."), CPUTextureRef->Format, *Texture2D->GetFName().ToString());
 
 		Width = 0;
 		Height = 0;
@@ -462,15 +472,27 @@ bool UPCGTextureData::InitializeFromGPUTexture(const TFunction<void()>& PostInit
 	Texture->UpdateResource();
 	Texture->WaitForPendingInitOrStreaming();
 
-	FTexturePlatformData* PlatformData = Texture->GetPlatformData();
+	FTexturePlatformData* PlatformData = nullptr;
 	FTextureResource* TextureResource = Texture->GetResource();
 
+	if (UTexture2D* Texture2D = Cast<UTexture2D>(Texture))
+	{
+		PlatformData = Texture2D->GetPlatformData();
+	}
+	else if (UTexture2DArray* Texture2DArray = Cast<UTexture2DArray>(Texture))
+	{
+		PlatformData = Texture2DArray->GetPlatformData();
+	}
+
 	if (PlatformData && TextureResource && TextureResource->TextureRHI)
 	{
 		FPCGTextureReadbackDispatchParams Params;
 		Params.SourceTexture = TextureResource->TextureRHI;
-		Params.SourceSampler = TextureResource->SamplerStateRHI;
+
+		// We should always use a point filter sampler since we are trying to get a 1 to 1 copy of the texture. We will do our own filtering later.
+		Params.SourceSampler = TStaticSamplerState<SF_Point, AM_Clamp, AM_Clamp, AM_Clamp>::GetRHI();
 		Params.SourceDimensions = FIntPoint(PlatformData->SizeX, PlatformData->SizeY);
+		Params.SourceTextureIndex = TextureIndex;
 
 		FPCGTextureReadbackInterface::Dispatch(Params, [this, PlatformData, PostInitializeCallback](void* OutBuffer, int32 ReadbackWidth, int32 ReadbackHeight)
 		{
diff --git a/Engine/Plugins/Experimental/PCG/Source/PCG/Private/Elements/PCGTextureSampler.cpp b/Engine/Plugins/Experimental/PCG/Source/PCG/Private/Elements/PCGTextureSampler.cpp
--- a/Engine/Plugins/Experimental/PCG/Source/PCG/Private/Elements/PCGTextureSampler.cpp
+++ b/Engine/Plugins/Experimental/PCG/Source/PCG/Private/Elements/PCGTextureSampler.cpp
@@ -104,41 +104,29 @@ bool FPCGTextureSamplerElement::ExecuteInternal(FPCGContext* InContext) const
 		return true;
 	}
 
-	UTexture2D* Texture2D = Cast<UTexture2D>(Texture);
+	uint32 TextureArrayIndex = 0;
 
-	// If the type is not a UTexture2D, try UTexture2DArray instead
-	if (!Texture2D)
+	if (UTexture2DArray* TextureArray = Cast<UTexture2DArray>(Texture))
 	{
 #if WITH_EDITOR
-		// TODO: support Texture2DArray without editor
-		if (UTexture2DArray* Tex2DArray = Cast<UTexture2DArray>(Texture))
-		{
-			if (Tex2DArray->SourceTextures.IsValidIndex(Settings->TextureArrayIndex))
-			{
-				Texture2D = Tex2DArray->SourceTextures[Settings->TextureArrayIndex];
+		const int32 ArraySize = TextureArray->SourceTextures.Num();
+#else
+		const int32 ArraySize = TextureArray->GetArraySize();
+#endif
 
-				if (!Texture2D)
-				{
-					PCGE_LOG(Error, GraphAndLog, FText::Format(LOCTEXT("CouldNotResolveTextureArray", "Texture index {0} could not be loaded."), Settings->TextureArrayIndex));
-				}
-			}
-			else
-			{
-				PCGE_LOG(Error, GraphAndLog, FText::Format(LOCTEXT("TextureArrayIndexOutOfBounds", "Texture array index {0} was out of bounds. There are only {1} textures in the array."), Settings->TextureArrayIndex, Tex2DArray->SourceTextures.Num()));
-			}
+		if (Settings->TextureArrayIndex >= 0 && Settings->TextureArrayIndex < ArraySize)
+		{
+			TextureArrayIndex = Settings->TextureArrayIndex;
 		}
 		else
 		{
-			PCGE_LOG(Error, GraphAndLog, FText::Format(LOCTEXT("InvalidTextureType", "Texture at path '{0}' was not of valid type. Must be either Texture2D or Texture2DArray."), FText::FromString(Settings->Texture.ToString())));
+			PCGE_LOG(Error, GraphAndLog, FText::Format(LOCTEXT("InvalidTextureIndex", "Array index {0} was out of bounds for TextureArray at path '{1}'."), Settings->TextureArrayIndex, FText::FromString(Settings->Texture.ToString())));
+			return true;
 		}
-#else
-		PCGE_LOG(Error, GraphAndLog, FText::Format(LOCTEXT("InvalidTextureTypeWithoutEditor", "Texture at path '{0}' was not of valid type. Must be a Texture2D when built without editor."), FText::FromString(Settings->Texture.ToString())));
-#endif
 	}
-
-	if (!Texture2D)
+	else if (!Texture->IsA<UTexture2D>())
 	{
-		PCGE_LOG(Error, GraphAndLog, FText::Format(LOCTEXT("InvalidTexture2D", "Texture at path '{0}' could not evaluate to a valid 2D Texture"), FText::FromString(Settings->Texture.ToString())));
+		PCGE_LOG(Error, GraphAndLog, FText::Format(LOCTEXT("InvalidTextureType", "Texture at path '{0}' is not a valid type. Must be UTexture2D or UTexture2DArray."), FText::FromString(Settings->Texture.ToString())));
 		return true;
 	}
 
@@ -186,7 +174,7 @@ bool FPCGTextureSamplerElement::ExecuteInternal(FPCGContext* InContext) const
 		}
 	};
 
-	TextureData->Initialize(Texture2D, FinalTransform, PostInitializeCallback);
+	TextureData->Initialize(Texture, TextureArrayIndex, FinalTransform, PostInitializeCallback);
 
 	TextureData->DensityFunction = DensityFunction;
 	TextureData->ColorChannel = ColorChannel;
diff --git a/Engine/Plugins/Experimental/PCG/Source/PCG/Private/Tests/Data/PCGTextureDataTest.cpp b/Engine/Plugins/Experimental/PCG/Source/PCG/Private/Tests/Data/PCGTextureDataTest.cpp
--- a/Engine/Plugins/Experimental/PCG/Source/PCG/Private/Tests/Data/PCGTextureDataTest.cpp
+++ b/Engine/Plugins/Experimental/PCG/Source/PCG/Private/Tests/Data/PCGTextureDataTest.cpp
@@ -39,7 +39,7 @@ bool FPCGTextureDataOffsetTilingRotation::RunTest(const FString& Parameters)
 
 	bool bTestPassed = true;
 
-	TextureData->Initialize(Texture2D, FTransform(), [this, TextureData, &bTestPassed]() -> void
+	TextureData->Initialize(Texture2D, /*TextureIndex=*/0, FTransform(), [this, TextureData, &bTestPassed]() -> void
 		{
 			FRandomStream RandomStream;
 			FPCGPoint OutPoint;
diff --git a/Engine/Plugins/Experimental/PCG/Source/PCG/Public/Data/PCGTextureData.h b/Engine/Plugins/Experimental/PCG/Source/PCG/Public/Data/PCGTextureData.h
--- a/Engine/Plugins/Experimental/PCG/Source/PCG/Public/Data/PCGTextureData.h
+++ b/Engine/Plugins/Experimental/PCG/Source/PCG/Public/Data/PCGTextureData.h
@@ -7,7 +7,7 @@
 #include "PCGTextureData.generated.h"
 
 class UPCGSpatialData;
-class UTexture2D;
+class UTexture;
 
 UENUM(BlueprintType)
 enum class EPCGTextureColorChannel : uint8
@@ -100,7 +100,7 @@ class PCG_API UPCGTextureData : public UPCGBaseTextureData
 	GENERATED_BODY()
 
 public:
-	void Initialize(UTexture2D* InTexture, const FTransform& InTransform, const TFunction<void()>& PostInitializeCallback);
+	void Initialize(UTexture* InTexture, uint32 InTextureIndex, const FTransform& InTransform, const TFunction<void()>& PostInitializeCallback);
 
 	// ~Begin UPCGData interface
 	virtual EPCGDataType GetDataType() const override { return EPCGDataType::Texture; }
@@ -124,7 +124,10 @@ private:
 
 public:
 	UPROPERTY(BlueprintReadOnly, VisibleAnywhere, Category = Properties)
-	TWeakObjectPtr<UTexture2D> Texture = nullptr;
+	TWeakObjectPtr<UTexture> Texture = nullptr;
+
+	UPROPERTY(BlueprintReadOnly, VisibleAnywhere, Category = Properties)
+	int TextureIndex = 0;
 };
 
 #if UE_ENABLE_INCLUDE_ORDER_DEPRECATED_IN_5_2
diff --git a/Engine/Plugins/Experimental/PCG/Source/PCG/Public/Elements/PCGTextureSampler.h b/Engine/Plugins/Experimental/PCG/Source/PCG/Public/Elements/PCGTextureSampler.h
--- a/Engine/Plugins/Experimental/PCG/Source/PCG/Public/Elements/PCGTextureSampler.h
+++ b/Engine/Plugins/Experimental/PCG/Source/PCG/Public/Elements/PCGTextureSampler.h
@@ -53,16 +53,10 @@ public:
 	UPROPERTY(BlueprintReadWrite, EditAnywhere, Category = Settings, meta = (PCG_Overridable))
 	TSoftObjectPtr<UTexture> Texture = nullptr;
 
-#if WITH_EDITORONLY_DATA
-	/** Index of texture array slice. Only used when built with editor and if the type of Texture is Texture2DArray. */
+	/** Index of texture array slice. Only used when built with editor and if the type of Texture is UTexture2DArray. */
 	UPROPERTY(BlueprintReadWrite, EditAnywhere, Category = Settings, meta = (EditCondition = bDisplayTextureArrayIndex, EditConditionHides, HideEditConditionToggle, ClampMin = '0', PCG_Overridable))
 	int TextureArrayIndex = 0;
 
-	// This can be set false by inheriting nodes to hide the 'TextureArrayIndex' property.
-	UPROPERTY(Transient, meta = (EditCondition = false, EditConditionHides))
-	bool bDisplayTextureArrayIndex = false;
-#endif
-
 	// Common members in BaseTextureData
 	UPROPERTY(BlueprintReadWrite, EditAnywhere, Category = SpatialData, meta = (PCG_Overridable))
 	EPCGTextureDensityFunction DensityFunction = EPCGTextureDensityFunction::Multiply;
@@ -96,6 +90,13 @@ public:
 
 	UPROPERTY(BlueprintReadWrite, EditAnywhere, Category = "Settings|Tiling", meta = (EditCondition = "bUseAdvancedTiling && bUseTileBounds", PCG_Overridable))
 	FVector2D TileBoundsMax = FVector2D(0.5, 0.5);
+
+protected:
+#if WITH_EDITORONLY_DATA
+	// Used to hide the 'TextureArrayIndex' property.
+	UPROPERTY(Transient)
+	bool bDisplayTextureArrayIndex = false;
+#endif
 };
 
 struct PCG_API FPCGTextureSamplerContext : public FPCGContext
diff --git a/Engine/Plugins/Experimental/PCG/Source/PCGCompute/PCGCompute.Build.cs b/Engine/Plugins/Experimental/PCG/Source/PCGCompute/PCGCompute.Build.cs
--- a/Engine/Plugins/Experimental/PCG/Source/PCGCompute/PCGCompute.Build.cs
+++ b/Engine/Plugins/Experimental/PCG/Source/PCGCompute/PCGCompute.Build.cs
@@ -15,6 +15,7 @@ namespace UnrealBuildTool.Rules
 					"Projects",
 					"RHI",
 					"RenderCore",
+					"Renderer",
 				});
 		}
 	}
diff --git a/Engine/Plugins/Experimental/PCG/Source/PCGCompute/Private/PCGTextureReadback.cpp b/Engine/Plugins/Experimental/PCG/Source/PCGCompute/Private/PCGTextureReadback.cpp
--- a/Engine/Plugins/Experimental/PCG/Source/PCGCompute/Private/PCGTextureReadback.cpp
+++ b/Engine/Plugins/Experimental/PCG/Source/PCGCompute/Private/PCGTextureReadback.cpp
@@ -9,8 +9,7 @@
 #include "Async/Async.h"
 
 #define LOCTEXT_NAMESPACE "PCGCompute"
-
-#define NUM_THREADS_PER_GROUP_DIMENSION 32
+#define PCG_NUM_THREADS_PER_GROUP_DIMENSION 8
 
 class PCGCOMPUTE_API FPCGTextureReadbackCS : public FGlobalShader
 {
@@ -19,9 +18,10 @@ public:
 	SHADER_USE_PARAMETER_STRUCT(FPCGTextureReadbackCS, FGlobalShader);
  
 	BEGIN_SHADER_PARAMETER_STRUCT(FParameters, )
-		SHADER_PARAMETER_TEXTURE(Texture2D<float4>, SourceTexture)
+		SHADER_PARAMETER_TEXTURE(Texture2DArray<float4>, SourceTextureArray)
 		SHADER_PARAMETER_SAMPLER(SamplerState, SourceSampler)
 		SHADER_PARAMETER(FVector2f, SourceDimensions)
+		SHADER_PARAMETER(uint32, SourceTextureIndex)
 		SHADER_PARAMETER_UAV(RWTexture2D<float4>, OutputTexture)
 	END_SHADER_PARAMETER_STRUCT()
  
@@ -30,8 +30,8 @@ public:
 	{
 		FGlobalShader::ModifyCompilationEnvironment(Parameters, OutEnvironment);
 
-		OutEnvironment.SetDefine(TEXT("THREADGROUPSIZE_X"), NUM_THREADS_PER_GROUP_DIMENSION);
-		OutEnvironment.SetDefine(TEXT("THREADGROUPSIZE_Y"), NUM_THREADS_PER_GROUP_DIMENSION);
+		OutEnvironment.SetDefine(TEXT("THREADGROUPSIZE_X"), PCG_NUM_THREADS_PER_GROUP_DIMENSION);
+		OutEnvironment.SetDefine(TEXT("THREADGROUPSIZE_Y"), PCG_NUM_THREADS_PER_GROUP_DIMENSION);
 		OutEnvironment.SetDefine(TEXT("THREADGROUPSIZE_Z"), 1);
 	}
 };
@@ -45,9 +45,10 @@ void FPCGTextureReadbackInterface::Dispatch_RenderThread(FRHICommandListImmediat
 	RenderCaptureInterface::FScopedCapture RenderCapture(true, &RHICmdList, TEXT("PCGTextureDataReadback"));
 
 	FPCGTextureReadbackCS::FParameters PassParameters;
-	PassParameters.SourceTexture = Params.SourceTexture;
+	PassParameters.SourceTextureArray = Params.SourceTexture;
 	PassParameters.SourceSampler = Params.SourceSampler;
 	PassParameters.SourceDimensions = { (float)Params.SourceDimensions.X, (float)Params.SourceDimensions.Y };
+	PassParameters.SourceTextureIndex = Params.SourceTextureIndex;
 
 	FRHITextureCreateDesc TargetTextureDesc =
 		FRHITextureCreateDesc::Create2D(TEXT("PCGTexture Readback Compute Target"), Params.SourceDimensions.X, Params.SourceDimensions.Y, EPixelFormat::PF_B8G8R8A8)
@@ -73,8 +74,8 @@ void FPCGTextureReadbackInterface::Dispatch_RenderThread(FRHICommandListImmediat
 
 	TShaderMapRef<FPCGTextureReadbackCS> ComputeShader(GetGlobalShaderMap(GMaxRHIFeatureLevel));
 	FComputeShaderUtils::Dispatch(RHICmdList, ComputeShader, PassParameters,
-		FIntVector(FMath::DivideAndRoundUp(Params.SourceDimensions.X, NUM_THREADS_PER_GROUP_DIMENSION),
-			FMath::DivideAndRoundUp(Params.SourceDimensions.Y, NUM_THREADS_PER_GROUP_DIMENSION), 1));
+		FIntVector(FMath::DivideAndRoundUp(Params.SourceDimensions.X, PCG_NUM_THREADS_PER_GROUP_DIMENSION),
+			FMath::DivideAndRoundUp(Params.SourceDimensions.Y, PCG_NUM_THREADS_PER_GROUP_DIMENSION), 1));
 
 	// Prepare OutputTexture to be copied
 	RHICmdList.Transition(FRHITransitionInfo(OutputTexture, ERHIAccess::UAVCompute, ERHIAccess::CopySrc));
@@ -131,3 +132,4 @@ void FPCGTextureReadbackInterface::Dispatch(const FPCGTextureReadbackDispatchPar
 }
 
 #undef LOCTEXT_NAMESPACE
+#undef PCG_NUM_THREADS_PER_GROUP_DIMENSION
diff --git a/Engine/Plugins/Experimental/PCG/Source/PCGCompute/Public/PCGTextureReadback.h b/Engine/Plugins/Experimental/PCG/Source/PCGCompute/Public/PCGTextureReadback.h
--- a/Engine/Plugins/Experimental/PCG/Source/PCGCompute/Public/PCGTextureReadback.h
+++ b/Engine/Plugins/Experimental/PCG/Source/PCGCompute/Public/PCGTextureReadback.h
@@ -4,19 +4,33 @@
 
 #include "RHI.h"
  
+/**
+ * Parameters to drive execution of a GPU texture readback. These mirror the fields in the texture readback compute shader (PCGTextureReadback.usf).
+ */
 struct PCGCOMPUTE_API FPCGTextureReadbackDispatchParams
 {
+	/** Source texture to sample from. Can be a UTexture2D or UTexture2DArray. */
 	FTextureRHIRef SourceTexture;
+
+	/** Sampler used to sample the SourceTexture. Should use an SF_Point filter for precise per-pixel readback. */
 	FSamplerStateRHIRef SourceSampler;
+
+	/** Width and Height of the SourceTexture. Should match the underlying dimensions exactly for precise per-pixel readback. */
 	FIntPoint SourceDimensions;
+
+	/** Optional texture index. Should be 0 if the SourceTexture is not a texture array. */
+	uint32 SourceTextureIndex = 0;
 };
  
+/**
+ * API for dispatching TextureReadback operations to the GPU.
+ */
 class PCGCOMPUTE_API FPCGTextureReadbackInterface
 {
 public:
 	static void Dispatch_RenderThread(FRHICommandListImmediate& RHICmdList, const FPCGTextureReadbackDispatchParams& Params, const TFunction<void(void* OutBuffer, int32 ReadbackWidth, int32 ReadbackHeight)>& AsyncCallback);
 	static void Dispatch_GameThread(const FPCGTextureReadbackDispatchParams& Params, const TFunction<void(void* OutBuffer, int32 ReadbackWidth, int32 ReadbackHeight)>& AsyncCallback);
 
-	// Dispatches this shader. Can be called from any thread
+	/** Dispatches the texture readback compute shader. Can be called from any thread. */
 	static void Dispatch(const FPCGTextureReadbackDispatchParams& Params, const TFunction<void(void* OutBuffer, int32 ReadbackWidth, int32 ReadbackHeight)>& AsyncCallback);
 };
