[PCG] Improve on the landscape tracking options and enable settings with console variables for PCG

* Use UDeveloperSettings for the PCGEngineSettings
* When landscape refresh is disable in edit mode, force refresh when we exit the edit mode
* Needed to attach a callback in the PCGEditorModule, and change it's load phase

#rnx

diff --git a/Engine/Plugins/Experimental/PCG/PCG.uplugin b/Engine/Plugins/Experimental/PCG/PCG.uplugin
--- a/Engine/Plugins/Experimental/PCG/PCG.uplugin
+++ b/Engine/Plugins/Experimental/PCG/PCG.uplugin
@@ -23,7 +23,7 @@
 		{
 			"Name": "PCGEditor",
 			"Type": "Editor",
-			"LoadingPhase": "Default"
+			"LoadingPhase": "PostEngineInit"
 		},
 		{
 			"Name": "PCGCompute",
diff --git a/Engine/Plugins/Experimental/PCG/Source/PCG/PCG.Build.cs b/Engine/Plugins/Experimental/PCG/Source/PCG/PCG.Build.cs
--- a/Engine/Plugins/Experimental/PCG/Source/PCG/PCG.Build.cs
+++ b/Engine/Plugins/Experimental/PCG/Source/PCG/PCG.Build.cs
@@ -13,6 +13,7 @@ namespace UnrealBuildTool.Rules
 					"Chaos",
 					"Core",
 					"CoreUObject",
+					"DeveloperSettings",
 					"Engine",
 					"Landscape",
 					"Foliage",
diff --git a/Engine/Plugins/Experimental/PCG/Source/PCG/Private/PCGActorAndComponentMapping.cpp b/Engine/Plugins/Experimental/PCG/Source/PCG/Private/PCGActorAndComponentMapping.cpp
--- a/Engine/Plugins/Experimental/PCG/Source/PCG/Private/PCGActorAndComponentMapping.cpp
+++ b/Engine/Plugins/Experimental/PCG/Source/PCG/Private/PCGActorAndComponentMapping.cpp
@@ -86,8 +86,9 @@ namespace PCGActorAndComponentMapping
 		}
 	}
 
-	bool ShouldDiscardLandscapeRefresh(const ALandscapeProxy* InLandscape)
+	bool ShouldDiscardLandscapeRefresh(const ALandscapeProxy* InLandscape, bool& bIsInEditingMode, bool bIsExitingEditingMode)
 	{
+		bIsInEditingMode = false;
 		// If it is not a landscape, we should refresh.
 		if (!InLandscape)
 		{
@@ -108,7 +109,8 @@ namespace PCGActorAndComponentMapping
 
 		// Refresh only if we are not editing.
 		const ALandscape* Landscape = InLandscape->GetLandscapeActor();
-		return Landscape && Landscape->HasLandscapeEdMode();
+		bIsInEditingMode = Landscape && Landscape->HasLandscapeEdMode() && !bIsExitingEditingMode;
+		return bIsInEditingMode;
 	}
 #endif
 }
@@ -1746,8 +1748,16 @@ void FPCGActorAndComponentMapping::OnActorChanged(AActor* InActor, bool bInHasMo
 	}
 
 	// If it is a landscape and we should discard the refresh, early out.
-	if (PCGActorAndComponentMapping::ShouldDiscardLandscapeRefresh(Cast<ALandscapeProxy>(InActor)))
+	bool bIsInEditingMode = false;
+	ALandscapeProxy* Landscape = Cast<ALandscapeProxy>(InActor);
+	if (PCGActorAndComponentMapping::ShouldDiscardLandscapeRefresh(Landscape, bIsInEditingMode, bIsCurrentlyExitingLandscapeEditMode))
 	{
+		// If we are in editing, keep track of all the dirtied landscape to refresh them when we exit.
+		if (bIsInEditingMode)
+		{
+			DirtiedLandscapes.AddUnique(Landscape);
+		}
+
 		return;
 	}
 
@@ -1884,4 +1894,24 @@ bool FPCGActorAndComponentMapping::ClearCacheForActor(const AActor* InActor, con
 	return bShouldDirty;
 }
 
+void FPCGActorAndComponentMapping::NotifyLandscapeEditModeExited()
+{
+	bIsCurrentlyExitingLandscapeEditMode = true;
+	// When the landscape edit mode is exited, force the refresh on all modified/dirtied landscapes.
+	for (TObjectKey<ALandscapeProxy> Landscape : DelayedModifiedLandscapes)
+	{
+		DirtiedLandscapes.AddUnique(Landscape);
+	}
+
+	DelayedModifiedLandscapes.Empty();
+
+	for (TObjectKey<ALandscapeProxy> Landscape : DirtiedLandscapes)
+	{
+		ApplyLandscapeChanges(Landscape.ResolveObjectPtr());
+	}
+
+	DirtiedLandscapes.Empty();
+
+	bIsCurrentlyExitingLandscapeEditMode = false;
+}
 #endif // WITH_EDITOR
\ No newline at end of file
diff --git a/Engine/Plugins/Experimental/PCG/Source/PCG/Private/PCGEngineSettings.cpp b/Engine/Plugins/Experimental/PCG/Source/PCG/Private/PCGEngineSettings.cpp
new file mode 100644
--- /dev/null
+++ b/Engine/Plugins/Experimental/PCG/Source/PCG/Private/PCGEngineSettings.cpp
@@ -0,0 +1,43 @@
+// Copyright Epic Games, Inc. All Rights Reserved.
+
+#include "PCGEngineSettings.h"
+
+#define LOCTEXT_NAMESPACE "PCGEngineSettings"
+
+FName UPCGEngineSettings::GetCategoryName() const
+{
+	return TEXT("Plugins");
+}
+
+#if WITH_EDITOR
+FText UPCGEngineSettings::GetSectionText() const
+{
+	return LOCTEXT("PCGEngineSettingsName", "PCG");
+}
+#endif // WITH_EDITOR
+
+void UPCGEngineSettings::PostInitProperties()
+{
+	Super::PostInitProperties();
+
+#if WITH_EDITOR
+	if (IsTemplate())
+	{
+		ImportConsoleVariableValues();
+	}
+#endif // WITH_EDITOR
+}
+
+#if WITH_EDITOR
+void UPCGEngineSettings::PostEditChangeProperty(FPropertyChangedEvent& PropertyChangedEvent)
+{
+	Super::PostEditChangeProperty(PropertyChangedEvent);
+
+	if (PropertyChangedEvent.Property)
+	{
+		ExportValuesToConsoleVariables(PropertyChangedEvent.Property);
+	}
+}
+#endif // WITH_EDITOR
+
+#undef LOCTEXT_NAMESPACE
\ No newline at end of file
diff --git a/Engine/Plugins/Experimental/PCG/Source/PCG/Private/PCGModule.cpp b/Engine/Plugins/Experimental/PCG/Source/PCG/Private/PCGModule.cpp
--- a/Engine/Plugins/Experimental/PCG/Source/PCG/Private/PCGModule.cpp
+++ b/Engine/Plugins/Experimental/PCG/Source/PCG/Private/PCGModule.cpp
@@ -2,7 +2,8 @@
 
 #include "PCGModule.h"
 
-#include "PCGEngineSettings.h"
+#include "PCGContext.h"
+#include "PCGElement.h"
 
 #include "Modules/ModuleManager.h"
 
@@ -34,9 +35,6 @@ public:
 
 private:
 #if WITH_EDITOR
-	void RegisterSettings();
-	void UnregisterSettings();
-
 	void RegisterNativeElementDeterminismTests();
 	void DeregisterNativeElementDeterminismTests();
 #endif
@@ -45,8 +43,6 @@ private:
 #if WITH_EDITOR
 void FPCGModule::StartupModule()
 {
-	RegisterSettings();
-
 	PCGDeterminismTests::FNativeTestRegistry::Create();
 
 	RegisterNativeElementDeterminismTests();
@@ -54,32 +50,11 @@ void FPCGModule::StartupModule()
 
 void FPCGModule::ShutdownModule()
 {
-	UnregisterSettings();
-
 	DeregisterNativeElementDeterminismTests();
 
 	PCGDeterminismTests::FNativeTestRegistry::Destroy();
 }
 
-void FPCGModule::RegisterSettings()
-{
-	if (ISettingsModule* SettingsModule = FModuleManager::GetModulePtr<ISettingsModule>("Settings"))
-	{
-		SettingsModule->RegisterSettings("Project", "Plugins", "PCG",
-			LOCTEXT("PCGEngineSettingsName", "PCG"),
-			LOCTEXT("PCGEngineSettingsDescription", "Configure PCG."),
-			GetMutableDefault<UPCGEngineSettings>());
-	}
-}
-
-void FPCGModule::UnregisterSettings()
-{
-	if (ISettingsModule* SettingsModule = FModuleManager::GetModulePtr<ISettingsModule>("Settings"))
-	{
-		SettingsModule->UnregisterSettings("Project", "Plugins", "PCG");
-	}
-}
-
 void FPCGModule::RegisterNativeElementDeterminismTests()
 {
 	PCGDeterminismTests::FNativeTestRegistry::RegisterTestFunction(UPCGDifferenceSettings::StaticClass(), PCGDeterminismTests::DifferenceElement::RunTestSuite);
@@ -97,4 +72,28 @@ IMPLEMENT_MODULE(FPCGModule, PCG);
 
 PCG_API DEFINE_LOG_CATEGORY(LogPCG);
 
+void PCGLog::LogErrorOnGraph(const FText& InMsg, const FPCGContext* InContext)
+{
+	if (InContext)
+	{
+		PCGE_LOG_C(Error, GraphAndLog, InContext, InMsg);
+	}
+	else
+	{
+		UE_LOG(LogPCG, Error, TEXT("%s"), *InMsg.ToString());
+	}
+}
+
+void PCGLog::LogWarningOnGraph(const FText& InMsg, const FPCGContext* InContext)
+{
+	if (InContext)
+	{
+		PCGE_LOG_C(Warning, GraphAndLog, InContext, InMsg);
+	}
+	else
+	{
+		UE_LOG(LogPCG, Warning, TEXT("%s"), *InMsg.ToString());
+	}
+}
+
 #undef LOCTEXT_NAMESPACE
diff --git a/Engine/Plugins/Experimental/PCG/Source/PCG/Public/PCGActorAndComponentMapping.h b/Engine/Plugins/Experimental/PCG/Source/PCG/Public/PCGActorAndComponentMapping.h
--- a/Engine/Plugins/Experimental/PCG/Source/PCG/Public/PCGActorAndComponentMapping.h
+++ b/Engine/Plugins/Experimental/PCG/Source/PCG/Public/PCGActorAndComponentMapping.h
@@ -55,6 +55,9 @@ public:
 
 	void RegisterTracking(UPCGComponent* InComponent);
 	void UpdateTracking(UPCGComponent* InComponent, bool bInShouldDirtyActors, const TArray<FPCGActorSelectionKey>* OptionalChangedKeys = nullptr);
+
+	/** Notify that we exited the Landscape edit mode. */
+	PCG_API void NotifyLandscapeEditModeExited();
 #endif // WITH_EDITOR
 
 	/** Register a new PCG Component or update it. Returns true if it was added/updated. Thread safe */
@@ -227,5 +230,11 @@ private:
 	// Part for the delayed landscape change update
 	double LastLandscapeDirtyTime = -1.0;
 	TArray<TObjectKey<ALandscapeProxy>> DelayedModifiedLandscapes;
+
+	// Keep track of all dirtied landscapes and force a refresh on them when exiting edit mode.
+	TArray<TObjectKey<ALandscapeProxy>> DirtiedLandscapes;
+
+	// Temp boolean to indicate we are currently exiting the landscape edit mode
+	bool bIsCurrentlyExitingLandscapeEditMode = false;
 #endif // WITH_EDITOR
 };
\ No newline at end of file
diff --git a/Engine/Plugins/Experimental/PCG/Source/PCG/Public/PCGEngineSettings.h b/Engine/Plugins/Experimental/PCG/Source/PCG/Public/PCGEngineSettings.h
--- a/Engine/Plugins/Experimental/PCG/Source/PCG/Public/PCGEngineSettings.h
+++ b/Engine/Plugins/Experimental/PCG/Source/PCG/Public/PCGEngineSettings.h
@@ -2,12 +2,12 @@
 
 #pragma once
 
-#include "UObject/Object.h"
+#include "Engine/DeveloperSettings.h"
 
 #include "PCGEngineSettings.generated.h"
 
-UCLASS(config = EditorPerProjectUserSettings)
-class PCG_API UPCGEngineSettings : public UObject
+UCLASS(config = Plugins, defaultconfig)
+class PCG_API UPCGEngineSettings : public UDeveloperSettings
 {
 	GENERATED_BODY()
 
@@ -19,4 +19,36 @@ public:
 	/** Whether we want to generate PCG graph/BP with PCG after drag/drop or not */
 	UPROPERTY(EditAnywhere, Config, Category = Workflow)
 	bool bGenerateOnDrop = true;
+
+#if WITH_EDITORONLY_DATA
+	// Console variables defined in PCGActorAndComponentMapping.cpp
+	// Use the variables to get the values, not GetDefault<UPCGEngineSettings>
+
+	/** Completely disable landscape refresh when it changes. */
+	UPROPERTY(EditAnywhere, Config, Category = LandscapeTracking, meta = (ConsoleVariable = "pcg.LandscapeDisableRefreshTracking"))
+	bool bLandscapeDisableRefreshTracking = false;
+
+	/** Completely disable landscape refresh when it changes in edit mode. Will force a refresh when landscape edit mode is exited. */
+	UPROPERTY(EditAnywhere, Config, Category = LandscapeTracking, meta = (ConsoleVariable = "pcg.LandscapeDisableRefreshTrackingInLandscapeEditingMode"))
+	bool bLandscapeDisableRefreshTrackingInLandscapeEditingMode = false;
+
+	/** Time in MS between a landscape change and PCG refresh. Set it to 0 or negative value to disable the delay. */
+	UPROPERTY(EditAnywhere, Config, Category = "LandscapeTracking|Advanced", meta = (ConsoleVariable = "pcg.LandscapeRefreshTimeDelayMS"))
+	int32 LandscapeRefreshTimeDelayMS = 1000;
+#endif // WITH_EDITORONLY_DATA
+
+
+	// Begin UDeveloperSettings Interface
+	virtual FName GetCategoryName() const override;
+#if WITH_EDITOR
+	virtual FText GetSectionText() const override;
+#endif // WITH_EDITOR
+	// End UDeveloperSettings Interface
+
+	//~ Begin UObject Interface
+	virtual void PostInitProperties() override;
+#if WITH_EDITOR
+	virtual void PostEditChangeProperty(FPropertyChangedEvent& PropertyChangedEvent) override;
+#endif
+	//~ End UObject Interface
 };
diff --git a/Engine/Plugins/Experimental/PCG/Source/PCG/Public/PCGModule.h b/Engine/Plugins/Experimental/PCG/Source/PCG/Public/PCGModule.h
--- a/Engine/Plugins/Experimental/PCG/Source/PCG/Public/PCGModule.h
+++ b/Engine/Plugins/Experimental/PCG/Source/PCG/Public/PCGModule.h
@@ -7,6 +7,16 @@
 // Logs
 PCG_API DECLARE_LOG_CATEGORY_EXTERN(LogPCG, Log, All);
 
+struct FPCGContext;
+
+namespace PCGLog
+{
+	/** Convenience function that would either log error on the graph if there is a context, or in the console if not. */
+	PCG_API void LogErrorOnGraph(const FText& InMsg, const FPCGContext* InContext = nullptr);
+	/** Convenience function that would either log warning on the graph if there is a context, or in the console if not. */
+	PCG_API void LogWarningOnGraph(const FText& InMsg, const FPCGContext* InContext = nullptr);
+}
+
 // Stats
 DECLARE_STATS_GROUP(TEXT("PCG"), STATGROUP_PCG, STATCAT_Advanced);
 
diff --git a/Engine/Plugins/Experimental/PCG/Source/PCG/Public/PCGSubsystem.h b/Engine/Plugins/Experimental/PCG/Source/PCG/Public/PCGSubsystem.h
--- a/Engine/Plugins/Experimental/PCG/Source/PCG/Public/PCGSubsystem.h
+++ b/Engine/Plugins/Experimental/PCG/Source/PCG/Public/PCGSubsystem.h
@@ -85,7 +85,7 @@ public:
 	void UnregisterPCGWorldActor(APCGWorldActor* InActor);
 
 	void OnOriginalComponentUnregistered(UPCGComponent* InComponent);
-	
+
 	UPCGLandscapeCache* GetLandscapeCache();
 
 	// Schedule graph (owner -> graph)
@@ -102,14 +102,14 @@ public:
 
 	/** General job scheduling
 	*  @param InOperation:       Callback that returns true if the task is done, false otherwise.
-	*  @param InSourceComponent: PCG component associated with this task. Can be null.
+	*  @param SourceComponent: PCG component associated with this task. Can be null.
 	*  @param TaskDependencies:  List of all the dependencies for this task.
 	*/
 	FPCGTaskId ScheduleGeneric(TFunction<bool()> InOperation, UPCGComponent* SourceComponent, const TArray<FPCGTaskId>& TaskDependencies);
-	
+
 	/** General job scheduling with context
 	*  @param InOperation:       Callback that takes a Context as argument and returns true if the task is done, false otherwise.
-	*  @param InSourceComponent: PCG component associated with this task. Can be null.
+	*  @param SourceComponent: PCG component associated with this task. Can be null.
 	*  @param TaskDependencies:  List of all the dependencies for this task.
 	*  @param bConsumeInputData: If your task need a context, but don't need the input data, set this flag to false. Default is true.
 	*/
@@ -216,6 +216,9 @@ public:
 	const FPCGNodeVisualLogs& GetNodeVisualLogs() const { return NodeVisualLogs; }
 	FPCGNodeVisualLogs& GetNodeVisualLogsMutable() { return NodeVisualLogs; }
 
+	/** Notify that we exited the Landscape edit mode. */
+	void NotifyLandscapeEditModeExited() { ActorAndComponentMapping.NotifyLandscapeEditModeExited(); }
+
 private:
 	enum class EOperation : uint32
 	{
diff --git a/Engine/Plugins/Experimental/PCG/Source/PCGEditor/PCGEditor.Build.cs b/Engine/Plugins/Experimental/PCG/Source/PCGEditor/PCGEditor.Build.cs
--- a/Engine/Plugins/Experimental/PCG/Source/PCGEditor/PCGEditor.Build.cs
+++ b/Engine/Plugins/Experimental/PCG/Source/PCGEditor/PCGEditor.Build.cs
@@ -46,6 +46,7 @@ namespace UnrealBuildTool.Rules
 					"ContentBrowser",
 					"DesktopWidgets",
 					"DetailCustomizations",
+					"EditorFramework",
 					"EditorScriptingUtilities",
 					"EditorStyle",
 					"EditorSubsystem",
diff --git a/Engine/Plugins/Experimental/PCG/Source/PCGEditor/Private/PCGEditorModule.cpp b/Engine/Plugins/Experimental/PCG/Source/PCGEditor/Private/PCGEditorModule.cpp
--- a/Engine/Plugins/Experimental/PCG/Source/PCGEditor/Private/PCGEditorModule.cpp
+++ b/Engine/Plugins/Experimental/PCG/Source/PCGEditor/Private/PCGEditorModule.cpp
@@ -10,6 +10,9 @@
 #include "PCGSubsystem.h"
 #include "PCGVolumeFactory.h"
 
+#include "Editor.h"
+#include "EditorModeManager.h"
+#include "EditorModes.h"
 #include "ISettingsModule.h"
 #include "PropertyEditorModule.h"
 #include "ToolMenus.h"
@@ -39,6 +42,12 @@ void FPCGEditorModule::StartupModule()
 	{
 		GEditor->ActorFactories.Add(NewObject<UPCGVolumeFactory>());
 	}
+
+	// Have a callback that catches changes in the Editor modes, to catch when we exit the landscape edit mode.
+	if (!IsRunningCommandlet() && GLevelEditorModeToolsIsValid())
+	{
+		GLevelEditorModeTools().OnEditorModeIDChanged().AddRaw(this, &FPCGEditorModule::OnEditorModeIDChanged);
+	}
 }
 
 void FPCGEditorModule::ShutdownModule()
@@ -56,6 +65,22 @@ void FPCGEditorModule::ShutdownModule()
 	{
 		GEditor->ActorFactories.RemoveAll([](const UActorFactory* ActorFactory) { return ActorFactory->IsA<UPCGVolumeFactory>(); });
 	}
+
+	if (!IsRunningCommandlet() && GLevelEditorModeToolsIsValid())
+	{
+		GLevelEditorModeTools().OnEditorModeIDChanged().RemoveAll(this);
+	}
+}
+
+void FPCGEditorModule::OnEditorModeIDChanged(const FEditorModeID& EditorModeID, bool bIsEntering)
+{
+	if (EditorModeID == FBuiltinEditorModes::EM_Landscape && !bIsEntering)
+	{
+		if (UPCGSubsystem* PCGSubsystem = UPCGSubsystem::GetSubsystemForCurrentWorld())
+		{
+			PCGSubsystem->NotifyLandscapeEditModeExited();
+		}
+	}
 }
 
 bool FPCGEditorModule::SupportsDynamicReloading()
diff --git a/Engine/Plugins/Experimental/PCG/Source/PCGEditor/Private/PCGEditorModule.h b/Engine/Plugins/Experimental/PCG/Source/PCGEditor/Private/PCGEditorModule.h
--- a/Engine/Plugins/Experimental/PCG/Source/PCGEditor/Private/PCGEditorModule.h
+++ b/Engine/Plugins/Experimental/PCG/Source/PCGEditor/Private/PCGEditorModule.h
@@ -32,6 +32,8 @@ protected:
 	void RegisterSettings();
 	void UnregisterSettings();
 
+	void OnEditorModeIDChanged(const FEditorModeID& EditorModeID, bool bIsEntering);
+
 	TArray<TSharedRef<IAssetTypeActions>> RegisteredAssetTypeActions;
 	static EAssetTypeCategories::Type PCGAssetCategory;
 
