[PCG] Revert change to Gather input pin dynamic typing - input pin types should not dynamically narrow.

#rnx

diff --git a/Engine/Plugins/Experimental/PCG/Source/PCG/Private/Elements/PCGGather.cpp b/Engine/Plugins/Experimental/PCG/Source/PCG/Private/Elements/PCGGather.cpp
--- a/Engine/Plugins/Experimental/PCG/Source/PCG/Private/Elements/PCGGather.cpp
+++ b/Engine/Plugins/Experimental/PCG/Source/PCG/Private/Elements/PCGGather.cpp
@@ -11,7 +11,7 @@ EPCGDataType UPCGGatherSettings::GetCurrentPinTypes(const UPCGPin* InPin) const
 {
 	check(InPin);
 
-	if (InPin->Properties.Label == PCGPinConstants::DefaultDependencyOnlyLabel)
+	if (InPin->Properties.Label == PCGPinConstants::DefaultDependencyOnlyLabel || !InPin->IsOutputPin())
 	{
 		return Super::GetCurrentPinTypes(InPin);
 	}
