[PCG] Fix issue where graphs were not drag and dropable anymore, introduced by changing load order.

* Load order is restored to previous value
* PostEngineInit callback is used instead
* Also fix a small forgotten feature on the volume factory

#rnx

diff --git a/Engine/Plugins/Experimental/PCG/PCG.uplugin b/Engine/Plugins/Experimental/PCG/PCG.uplugin
--- a/Engine/Plugins/Experimental/PCG/PCG.uplugin
+++ b/Engine/Plugins/Experimental/PCG/PCG.uplugin
@@ -23,7 +23,7 @@
 		{
 			"Name": "PCGEditor",
 			"Type": "Editor",
-			"LoadingPhase": "PostEngineInit"
+			"LoadingPhase": "Default"
 		},
 		{
 			"Name": "PCGCompute",
diff --git a/Engine/Plugins/Experimental/PCG/Source/PCGEditor/Private/PCGEditorModule.cpp b/Engine/Plugins/Experimental/PCG/Source/PCGEditor/Private/PCGEditorModule.cpp
--- a/Engine/Plugins/Experimental/PCG/Source/PCGEditor/Private/PCGEditorModule.cpp
+++ b/Engine/Plugins/Experimental/PCG/Source/PCGEditor/Private/PCGEditorModule.cpp
@@ -43,10 +43,9 @@ void FPCGEditorModule::StartupModule()
 		GEditor->ActorFactories.Add(NewObject<UPCGVolumeFactory>());
 	}
 
-	// Have a callback that catches changes in the Editor modes, to catch when we exit the landscape edit mode.
-	if (!IsRunningCommandlet() && GLevelEditorModeToolsIsValid())
+	if (!IsRunningCommandlet())
 	{
-		GLevelEditorModeTools().OnEditorModeIDChanged().AddRaw(this, &FPCGEditorModule::OnEditorModeIDChanged);
+		FCoreDelegates::OnPostEngineInit.AddRaw(this, &FPCGEditorModule::RegisterOnEditorModeChange);
 	}
 }
 
@@ -66,9 +65,13 @@ void FPCGEditorModule::ShutdownModule()
 		GEditor->ActorFactories.RemoveAll([](const UActorFactory* ActorFactory) { return ActorFactory->IsA<UPCGVolumeFactory>(); });
 	}
 
-	if (!IsRunningCommandlet() && GLevelEditorModeToolsIsValid())
+	if (!IsRunningCommandlet())
 	{
-		GLevelEditorModeTools().OnEditorModeIDChanged().RemoveAll(this);
+		FCoreDelegates::OnPostEngineInit.RemoveAll(this);
+		if (GLevelEditorModeToolsIsValid())
+		{
+			GLevelEditorModeTools().OnEditorModeIDChanged().RemoveAll(this);
+		}
 	}
 }
 
@@ -88,6 +91,15 @@ bool FPCGEditorModule::SupportsDynamicReloading()
 	return true;
 }
 
+void FPCGEditorModule::RegisterOnEditorModeChange()
+{
+	// Have a callback that catches changes in the Editor modes, to catch when we exit the landscape edit mode.
+	if (!IsRunningCommandlet() && GLevelEditorModeToolsIsValid())
+	{
+		GLevelEditorModeTools().OnEditorModeIDChanged().AddRaw(this, &FPCGEditorModule::OnEditorModeIDChanged);
+	}
+}
+
 void FPCGEditorModule::RegisterDetailsCustomizations()
 {
 	FPropertyEditorModule& PropertyEditor = FModuleManager::LoadModuleChecked<FPropertyEditorModule>("PropertyEditor");
@@ -131,7 +143,7 @@ void FPCGEditorModule::RegisterMenuExtensions()
 	FToolMenuOwnerScoped OwnerScoped(this);
 	UToolMenu* Menu = UToolMenus::Get()->ExtendMenu("LevelEditor.MainMenu.Tools");
 	FToolMenuSection& Section = Menu->AddSection("PCGToolsSection", LOCTEXT("PCGToolsSection", "Procedural Generation Tools"));
-	
+
 	Section.AddSubMenu(
 		"PCGToolsSubMenu",
 		LOCTEXT("PCGSubMenu", "PCG Framework"),
@@ -140,7 +152,7 @@ void FPCGEditorModule::RegisterMenuExtensions()
 }
 
 void FPCGEditorModule::UnregisterMenuExtensions()
-{	
+{
 	UToolMenus::UnregisterOwner(this);
 }
 
diff --git a/Engine/Plugins/Experimental/PCG/Source/PCGEditor/Private/PCGEditorModule.h b/Engine/Plugins/Experimental/PCG/Source/PCGEditor/Private/PCGEditorModule.h
--- a/Engine/Plugins/Experimental/PCG/Source/PCGEditor/Private/PCGEditorModule.h
+++ b/Engine/Plugins/Experimental/PCG/Source/PCGEditor/Private/PCGEditorModule.h
@@ -32,6 +32,7 @@ protected:
 	void RegisterSettings();
 	void UnregisterSettings();
 
+	void RegisterOnEditorModeChange();
 	void OnEditorModeIDChanged(const FEditorModeID& EditorModeID, bool bIsEntering);
 
 	TArray<TSharedRef<IAssetTypeActions>> RegisteredAssetTypeActions;
diff --git a/Engine/Plugins/Experimental/PCG/Source/PCGEditor/Private/PCGVolumeFactory.cpp b/Engine/Plugins/Experimental/PCG/Source/PCGEditor/Private/PCGVolumeFactory.cpp
--- a/Engine/Plugins/Experimental/PCG/Source/PCGEditor/Private/PCGVolumeFactory.cpp
+++ b/Engine/Plugins/Experimental/PCG/Source/PCGEditor/Private/PCGVolumeFactory.cpp
@@ -49,7 +49,7 @@ bool UPCGVolumeFactory::PreSpawnActor(UObject* Asset, FTransform& InOutLocation)
 
 		return true;
 	}
-	
+
 	return false;
 }
 
@@ -112,4 +112,11 @@ void UPCGVolumeFactory::PostPlaceAsset(TArrayView<const FTypedElementHandle> InE
 	}
 }
 
+UObject* UPCGVolumeFactory::GetAssetFromActorInstance(AActor* ActorInstance)
+{
+	const APCGVolume* PCGVolume = CastChecked<APCGVolume>(ActorInstance);
+	const UPCGComponent* PCGComponent = PCGVolume->GetComponentByClass<UPCGComponent>();
+	return PCGComponent ? PCGComponent->GetGraph() : nullptr;
+}
+
 #undef LOCTEXT_NAMESPACE
diff --git a/Engine/Plugins/Experimental/PCG/Source/PCGEditor/Private/PCGVolumeFactory.h b/Engine/Plugins/Experimental/PCG/Source/PCGEditor/Private/PCGVolumeFactory.h
--- a/Engine/Plugins/Experimental/PCG/Source/PCGEditor/Private/PCGVolumeFactory.h
+++ b/Engine/Plugins/Experimental/PCG/Source/PCGEditor/Private/PCGVolumeFactory.h
@@ -21,5 +21,6 @@ public:
 	virtual bool CanCreateActorFrom(const FAssetData& AssetData, FText& OutErrorMsg) override;
 	virtual bool PreSpawnActor(UObject* Asset, FTransform& InOutLocation) override;
 	virtual void PostPlaceAsset(TArrayView<const FTypedElementHandle> InElementHandles, const FAssetPlacementInfo& InPlacementInfo, const FPlacementOptions& InPlacementOptions) override;
+	virtual UObject* GetAssetFromActorInstance(AActor* ActorInstance) override;
 	//~ End UActorFactory Interface
 };
