[PCG] Improve on Attribute Partition

* Add generic functions to partition an attribute.
* Supports Attribute Set and Selectors
* Make ConstMetadata and MutableMetadata accessible from UPCGData, making the call virtual
* Add unit tests

diff --git a/Engine/Plugins/Experimental/PCG/Source/PCG/Private/Elements/Metadata/PCGMetadataPartition.cpp b/Engine/Plugins/Experimental/PCG/Source/PCG/Private/Elements/Metadata/PCGMetadataPartition.cpp
--- a/Engine/Plugins/Experimental/PCG/Source/PCG/Private/Elements/Metadata/PCGMetadataPartition.cpp
+++ b/Engine/Plugins/Experimental/PCG/Source/PCG/Private/Elements/Metadata/PCGMetadataPartition.cpp
@@ -1,21 +1,61 @@
 // Copyright Epic Games, Inc. All Rights Reserved.
 
 #include "Elements/Metadata/PCGMetadataPartition.h"
-#include "Data/PCGSpatialData.h"
-#include "Data/PCGPointData.h"
 
-#include "Algo/Find.h"
 #include "PCGContext.h"
+#include "Metadata/PCGMetadataPartitionCommon.h"
 
 #include UE_INLINE_GENERATED_CPP_BY_NAME(PCGMetadataPartition)
 
 #define LOCTEXT_NAMESPACE "PCGMetadataPartitionElement"
 
+EPCGDataType UPCGMetadataPartitionSettings::GetCurrentPinTypes(const UPCGPin* InPin) const
+{
+	check(InPin);
+	if (!InPin->IsOutputPin())
+	{
+		return Super::GetCurrentPinTypes(InPin);
+	}
+
+	// Output pin narrows to union of inputs on first pin
+	const EPCGDataType InputTypeUnion = GetTypeUnionOfIncidentEdges(PCGPinConstants::DefaultInputLabel);
+	return InputTypeUnion != EPCGDataType::None ? InputTypeUnion : EPCGDataType::Any;
+}
+
+TArray<FPCGPinProperties> UPCGMetadataPartitionSettings::InputPinProperties() const
+{
+	TArray<FPCGPinProperties> PinProperties;
+	PinProperties.Emplace(PCGPinConstants::DefaultInputLabel, EPCGDataType::Any);
+
+	return PinProperties;
+}
+
+TArray<FPCGPinProperties> UPCGMetadataPartitionSettings::OutputPinProperties() const
+{
+	TArray<FPCGPinProperties> PinProperties;
+	PinProperties.Emplace(PCGPinConstants::DefaultOutputLabel, EPCGDataType::Any);
+
+	return PinProperties;
+}
+
 FPCGElementPtr UPCGMetadataPartitionSettings::CreateElement() const
 {
 	return MakeShared<FPCGMetadataPartitionElement>();
 }
 
+void UPCGMetadataPartitionSettings::PostLoad()
+{
+	Super::PostLoad();
+
+#if WITH_EDITOR
+	if (PartitionAttribute_DEPRECATED != NAME_None)
+	{
+		PartitionAttributeSource.SetAttributeName(PartitionAttribute_DEPRECATED);
+		PartitionAttribute_DEPRECATED = NAME_None;
+	}
+#endif // WITH_EDITOR
+}
+
 bool FPCGMetadataPartitionElement::ExecuteInternal(FPCGContext* Context) const
 {
 	TRACE_CPUPROFILER_EVENT_SCOPE(FPCGMetadataPartitionElement::Execute);
@@ -24,116 +64,15 @@ bool FPCGMetadataPartitionElement::ExecuteInternal(FPCGContext* Context) const
 	const UPCGMetadataPartitionSettings* Settings = Context->GetInputSettings<UPCGMetadataPartitionSettings>();
 	check(Settings);
 
-	TArray<FPCGTaggedData> Inputs = Context->InputData.GetInputs();
+	TArray<FPCGTaggedData> Inputs = Context->InputData.GetInputsByPin(PCGPinConstants::DefaultInputLabel);
 	TArray<FPCGTaggedData>& Outputs = Context->OutputData.TaggedData;
 
-	const FName PartitionAttribute = Settings->PartitionAttribute;
-
 	for (const FPCGTaggedData& Input : Inputs)
 	{
-		const UPCGSpatialData* SpatialInput = Cast<const UPCGSpatialData>(Input.Data);
-		check(SpatialInput);
-
-		if (!SpatialInput->ConstMetadata())
-		{
-			PCGE_LOG(Error, GraphAndLog, LOCTEXT("InputMissingMetadata", "Input does not have metadata"));
-			continue;
-		}
-
-		const FName LocalPartitionAttribute = ((PartitionAttribute != NAME_None) ? PartitionAttribute : SpatialInput->ConstMetadata()->GetLatestAttributeNameOrNone());
-
-		if (!SpatialInput->ConstMetadata()->HasAttribute(LocalPartitionAttribute))
-		{
-			PCGE_LOG(Error, GraphAndLog, FText::Format(LOCTEXT("InputMissingAttribute", "Input does not have the '{0}' attribute"), FText::FromName(LocalPartitionAttribute)));
-			continue;
-		}
-
-		// Currently, we support only to parse points, so we'll do so here
-		const UPCGPointData* InputData = SpatialInput->ToPointData(Context);
-
-		if (!InputData)
-		{
-			PCGE_LOG(Error, GraphAndLog, LOCTEXT("CouldNotObtainPoints", "Unable to get points out of spatial data"));
-			continue;
-		}
-
-		const TArray<FPCGPoint>& SourcePoints = InputData->GetPoints();
-		const UPCGMetadata* SourceMetadata = InputData->ConstMetadata();
-		const FPCGMetadataAttributeBase* AttributeBase = SourceMetadata->GetConstAttribute(LocalPartitionAttribute);
-		check(AttributeBase);
-
-		TMap<PCGMetadataValueKey, UPCGPointData*> ValueToData;
-
-		// Get all value keys (-1 + 0 - N)
-		int64 MetadataValueKeyCount = AttributeBase->GetValueKeyOffsetForChild();
-
-		// For every value key, check if it should be merged with the default value
-		// Keep track of that mapping (should be only one, but that might not be true in uncompressed data)
-		TArray<PCGMetadataValueKey> ValueKeyMapping;
-		ValueKeyMapping.Reserve(MetadataValueKeyCount);
-
-		TArray<PCGMetadataValueKey> UniqueValueKeys;
-
-		const bool bUsesValueKeys = AttributeBase->UsesValueKeys();
-
-		for (PCGMetadataValueKey ValueKey = 0; ValueKey < MetadataValueKeyCount; ++ValueKey)
-		{
-			if (AttributeBase->IsEqualToDefaultValue(ValueKey))
-			{
-				ValueKeyMapping.Add(-1);
-			}
-			else if (!bUsesValueKeys)
-			{
-				PCGMetadataValueKey* MatchingVK = Algo::FindByPredicate(UniqueValueKeys, [ValueKey, AttributeBase](const PCGMetadataValueKey& Key)
-				{
-					return AttributeBase->AreValuesEqual(ValueKey, Key);
-				});
-
-				const PCGMetadataValueKey TentativeKey = MatchingVK ? *MatchingVK : ValueKey;
-
-				if (!MatchingVK)
-				{
-					UniqueValueKeys.Add(TentativeKey);
-				}
-
-				ValueKeyMapping.Add(TentativeKey);
-			}
-			else
-			{
-				ValueKeyMapping.Add(ValueKey);
-			}
-		}
-
-		TArray<UPCGPointData*> PartitionedData;
-		PartitionedData.SetNum(1 + ValueKeyMapping.Num());
-		for (int32 Index = 0; Index < PartitionedData.Num(); ++Index)
-		{
-			PartitionedData[Index] = nullptr;
-		}
-
-		// Loop on points, find matching data pointer, allocate if null, etc.and add point to it
-		for (const FPCGPoint& Point : SourcePoints)
-		{
-			PCGMetadataValueKey ValueKey = AttributeBase->GetValueKey(Point.MetadataEntry);
-			// Remap
-			if (ValueKey != PCGDefaultValueKey)
-			{
-				ValueKey = ValueKeyMapping[ValueKey];
-			}
-
-			const int32 PartitionDataIndex = 1 + ValueKey;
-
-			if (!PartitionedData[PartitionDataIndex])
-			{
-				PartitionedData[PartitionDataIndex] = NewObject<UPCGPointData>();
-				PartitionedData[PartitionDataIndex]->InitializeFromData(InputData);
-			}
-
-			PartitionedData[PartitionDataIndex]->GetMutablePoints().Add(Point);
-		}
+		const UPCGData* InData = Input.Data;
+		const FPCGAttributePropertyInputSelector PartitionAttributeSelector = Settings->PartitionAttributeSource.CopyAndFixLast(InData);
 
-		// Finally, push back the partitioned data to the outputs
-		for (UPCGPointData* PartitionData : PartitionedData)
+		for (UPCGData* PartitionData : PCGMetadataPartitionCommon::AttributePartition(InData, PartitionAttributeSelector, Context))
 		{
 			if (PartitionData)
 			{
diff --git a/Engine/Plugins/Experimental/PCG/Source/PCG/Private/Metadata/PCGMetadataPartitionCommon.cpp b/Engine/Plugins/Experimental/PCG/Source/PCG/Private/Metadata/PCGMetadataPartitionCommon.cpp
new file mode 100644
--- /dev/null
+++ b/Engine/Plugins/Experimental/PCG/Source/PCG/Private/Metadata/PCGMetadataPartitionCommon.cpp
@@ -0,0 +1,343 @@
+// Copyright Epic Games, Inc. All Rights Reserved.
+
+#include "Metadata/PCGMetadataPartitionCommon.h"
+
+#include "PCGContext.h"
+#include "PCGElement.h"
+#include "PCGModule.h"
+#include "PCGParamData.h"
+#include "Data/PCGPointData.h"
+#include "Elements/Metadata/PCGMetadataElementCommon.h"
+#include "Metadata/PCGMetadata.h"
+#include "Metadata/PCGMetadataCommon.h"
+#include "Metadata/PCGMetadataAttribute.h"
+#include "Metadata/Accessors/IPCGAttributeAccessor.h"
+#include "Metadata/Accessors/PCGAttributeAccessorHelpers.h"
+#include "Metadata/Accessors/PCGAttributeAccessorKeys.h"
+
+#include "Algo/IndexOf.h"
+
+#define LOCTEXT_NAMESPACE "PCGMetadataPartitionCommon"
+
+namespace PCGMetadataPartitionCommon
+{
+	/**
+	* Partition a given attribute, by first partitioning all value keys that point to the same value
+	* and then for each unique value key, list of index in the keys that match for this value.
+	*/
+	TArray<TArray<int32>> AttributePartition(const FPCGMetadataAttributeBase* InAttribute, const IPCGAttributeAccessorKeys& InKeys, FPCGContext* InOptionalContext)
+	{
+		check(InAttribute);
+
+		const int32 NumberOfEntries = InKeys.GetNum();
+
+		if (NumberOfEntries <= 0)
+		{
+			return {};
+		}
+
+		TMap<PCGMetadataValueKey, UPCGPointData*> ValueToData;
+
+		// Get all value keys (-1 + 0 - N)
+		const int64 MetadataValueKeyCount = InAttribute->GetValueKeyOffsetForChild();
+
+		// For every value key, check if it should be merged with the default value
+		TMap<PCGMetadataValueKey, int32> ValueKeyMapping;
+		ValueKeyMapping.Reserve(MetadataValueKeyCount);
+
+		int32 NumUniqueValueKeys = 0;
+
+		const bool bUsesValueKeys = InAttribute->UsesValueKeys();
+
+		if (!bUsesValueKeys)
+		{
+			TArray<PCGMetadataValueKey> UniqueValueKeys;
+
+			for (PCGMetadataValueKey ValueKey = 0; ValueKey < MetadataValueKeyCount; ++ValueKey)
+			{
+				if (InAttribute->IsEqualToDefaultValue(ValueKey))
+				{
+					ValueKeyMapping.Add(ValueKey, -1);
+					continue;
+				}
+
+				// TODO: Might want to upgrade to something better wince it can be quadractic and grow quickly.
+				int32 UniqueValueKeyIndex = Algo::IndexOfByPredicate(UniqueValueKeys, [ValueKey, InAttribute](const PCGMetadataValueKey& Key)
+				{
+					return InAttribute->AreValuesEqual(ValueKey, Key);
+				});
+
+				if (UniqueValueKeyIndex == INDEX_NONE)
+				{
+					ValueKeyMapping.Add(ValueKey, UniqueValueKeys.Num());
+					UniqueValueKeys.Add(ValueKey);
+					NumUniqueValueKeys++;
+				}
+				else
+				{
+					ValueKeyMapping.Add(ValueKey, UniqueValueKeyIndex);
+				}
+			}
+		}
+		else
+		{
+			NumUniqueValueKeys = MetadataValueKeyCount;
+		}
+
+		TArray<TArray<int32>> PartitionedData;
+		PartitionedData.SetNum(1 + NumUniqueValueKeys);
+
+		constexpr int32 ChunkSize = 256;
+		TArray<const PCGMetadataEntryKey*, TInlineAllocator<ChunkSize>> TempEntries;
+		TempEntries.SetNum(ChunkSize);
+
+		const int32 NumberOfIterations = (NumberOfEntries + ChunkSize - 1) / ChunkSize;
+
+		for (int32 i = 0; i < NumberOfIterations; ++i)
+		{
+			const int32 StartIndex = i * ChunkSize;
+			const int32 Range = FMath::Min(NumberOfEntries - StartIndex, ChunkSize);
+			TArrayView<const PCGMetadataEntryKey*> View(TempEntries.GetData(), Range);
+			InKeys.GetKeys(StartIndex, View);
+
+			for (int32 j = 0; j < Range; ++j)
+			{
+				PCGMetadataValueKey ValueKey = InAttribute->GetValueKey(*TempEntries[j]);
+				// Remap if not using value keys
+				if (!bUsesValueKeys && ValueKey != PCGDefaultValueKey)
+				{
+					ValueKey = ValueKeyMapping[ValueKey];
+				}
+
+				const int32 PartitionDataIndex = 1 + ValueKey;
+				PartitionedData[PartitionDataIndex].Add(StartIndex + j);
+			}
+		}
+
+		return PartitionedData;
+	}
+
+	/**
+	* Partition a given accessor that iterate on all values, find the identical ones,
+	* and then for each unique value, list of index in the keys that match for this value.
+	*/
+	template <typename T>
+	TArray<TArray<int32>> ValuePartition(const IPCGAttributeAccessor& InAccessor, const IPCGAttributeAccessorKeys& InKeys, FPCGContext* InOptionalContext)
+	{
+		TArray<T> UniqueValues;
+		TArray<TArray<int32>> PartitionedData;
+
+		PCGMetadataElementCommon::ApplyOnAccessor<T>(InKeys, InAccessor, [&PartitionedData, &UniqueValues](const T& InValue, int32 InIndex)
+		{
+			// TODO: Might want to upgrade to something better wince it can be quadractic and grow quickly.
+			int32 UniqueValueIndex = UniqueValues.IndexOfByPredicate([&InValue](const T& OtherValue)
+			{
+				// For consistency with the attribute part, use MetadataTraits::Equal
+				return PCG::Private::MetadataTraits<T>::Equal(InValue, OtherValue);
+			});
+
+			if (UniqueValueIndex == INDEX_NONE)
+			{
+				UniqueValueIndex = UniqueValues.Add(InValue);
+				PartitionedData.Emplace();
+			}
+
+			PartitionedData[UniqueValueIndex].Add(InIndex);
+		});
+
+		return PartitionedData;
+	}
+
+	/**
+	* Dispatch the partition according to the data and selector.
+	*/
+	TArray<TArray<int32>> AttributeGenericPartition(const UPCGData* InData, const FPCGAttributePropertySelector& InSelector, FPCGContext* InOptionalContext)
+	{
+		if (!InData)
+		{
+			return {};
+		}
+
+		TUniquePtr<const IPCGAttributeAccessorKeys> Keys = PCGAttributeAccessorHelpers::CreateConstKeys(InData, InSelector);
+		if (!Keys.IsValid())
+		{
+			PCGLog::LogErrorOnGraph(FText::Format(LOCTEXT("InvalidKeys", "Could not create keys for the input data with selector {0}"), InSelector.GetDisplayText()), InOptionalContext);
+			return {};
+		}
+
+		if (InSelector.IsBasicAttribute())
+		{
+			const UPCGMetadata* Metadata = InData->ConstMetadata();
+			if (!Metadata)
+			{
+				PCGLog::LogErrorOnGraph(FText::Format(LOCTEXT("InvalidMetadata", "Input data does not have metadata, while requesting an attribute {0}"), InSelector.GetDisplayText()), InOptionalContext);
+				return {};
+			}
+
+			const FPCGMetadataAttributeBase* Attribute = Metadata->GetConstAttribute(InSelector.GetName());
+			if (!Attribute)
+			{
+				PCGLog::LogErrorOnGraph(FText::Format(LOCTEXT("InvalidAttribute", "Attribute {0} not found"), InSelector.GetDisplayText()), InOptionalContext);
+				return {};
+			}
+
+			return AttributePartition(Attribute, *Keys, InOptionalContext);
+		}
+		else
+		{
+			TUniquePtr<const IPCGAttributeAccessor> Accessor = PCGAttributeAccessorHelpers::CreateConstAccessor(InData, InSelector);
+			if (!Accessor.IsValid())
+			{
+				PCGLog::LogErrorOnGraph(FText::Format(LOCTEXT("InvalidAccessor", "Attribute {0} not found"), InSelector.GetDisplayText()), InOptionalContext);
+				return {};
+			}
+
+			auto Operation = [&Accessor, &Keys, &InSelector, InOptionalContext](auto Dummy) -> TArray<TArray<int32>>
+			{
+				// Rotators don't have a hash, convert them to Quat
+				using AttributeType = std::conditional_t<std::is_same_v<decltype(Dummy), FRotator>, FQuat, decltype(Dummy)>;
+
+				// Can't partition on a transform.
+				if constexpr (std::is_same_v<AttributeType, FTransform>)
+				{
+					PCGLog::LogErrorOnGraph(FText::Format(LOCTEXT("InvalidType", "Attribute {0} is a transform, partition transforms is not supported"), InSelector.GetDisplayText()), InOptionalContext);
+					return {};
+				}
+				else
+				{
+					return ValuePartition<AttributeType>(*Accessor, *Keys, InOptionalContext);
+				}
+			};
+
+			return PCGMetadataAttribute::CallbackWithRightType(Accessor->GetUnderlyingType(), Operation);
+		}
+	}
+
+	/**
+	* Do a partition on the given point data for the selector
+	*/
+	TArray<UPCGData*> AttributePointPartition(const UPCGPointData* InData, const FPCGAttributePropertySelector& InSelector, FPCGContext* InOptionalContext)
+	{
+		TArray<TArray<int32>> Partition = AttributeGenericPartition(InData, InSelector, InOptionalContext);
+		if (Partition.IsEmpty())
+		{
+			return {};
+		}
+
+		TArray<UPCGData*> PartitionedData;
+		PartitionedData.Reserve(Partition.Num());
+		const TArray<FPCGPoint>& Points = InData->GetPoints();
+
+		for (TArray<int32>& Indices : Partition)
+		{
+			if (Indices.IsEmpty())
+			{
+				continue;
+			}
+
+			UPCGPointData* CurrentPointData = NewObject<UPCGPointData>();
+			PartitionedData.Add(CurrentPointData);
+			CurrentPointData->InitializeFromData(InData);
+
+			TArray<FPCGPoint>& OutPoints = CurrentPointData->GetMutablePoints();
+			OutPoints.Reserve(Indices.Num());
+			for (int32 Index : Indices)
+			{
+				OutPoints.Add(Points[Index]);
+			}
+		}
+
+		return PartitionedData;
+	}
+
+	TArray<UPCGData*> AttributeParamSpatialPartition(const UPCGData* InData, const FPCGAttributePropertySelector& InSelector, FPCGContext* InOptionalContext)
+	{
+		const UPCGSpatialData* InSpatialData = Cast<const UPCGSpatialData>(InData);
+		const UPCGParamData* InParamData = Cast<const UPCGParamData>(InData);
+
+		if (!InSpatialData && !InParamData)
+		{
+			PCGLog::LogErrorOnGraph(LOCTEXT("InvalidDataType", "Input data is not an attribute set nor a spatial data. Operation not supported."), InOptionalContext);
+			return {};
+		}
+
+		TArray<TArray<int32>> Partition = AttributeGenericPartition(InData, InSelector, InOptionalContext);
+		if (Partition.IsEmpty())
+		{
+			return {};
+		}
+
+		TArray<UPCGData*> PartitionedData;
+		PartitionedData.Reserve(Partition.Num());
+
+		TArray<FName> AttributeNames;
+		TArray<EPCGMetadataTypes> AttributeTypes;
+		const UPCGMetadata* OriginalMetadata = InData->ConstMetadata();
+		OriginalMetadata->GetAttributes(AttributeNames, AttributeTypes);
+
+		for (TArray<int32>& Indices : Partition)
+		{
+			if (Indices.IsEmpty())
+			{
+				continue;
+			}
+
+			UPCGMetadata* NewMetadata = nullptr;
+
+			if (InSpatialData)
+			{
+				UPCGSpatialData* NewData = NewObject<UPCGSpatialData>();
+				NewData->InitializeFromData(InSpatialData);
+				NewMetadata = NewData->Metadata;
+				PartitionedData.Add(NewData);
+			}
+			else
+			{
+				UPCGParamData* NewData = NewObject<UPCGParamData>();
+				NewData->Metadata->AddAttributes(OriginalMetadata);
+				NewMetadata = NewData->Metadata;
+				PartitionedData.Add(NewData);
+			}
+
+			TArray<PCGMetadataEntryKey> EntryKeys;
+			EntryKeys.Reserve(Indices.Num());
+			for (int32 i = 0; i < Indices.Num(); ++i)
+			{
+				EntryKeys.Add(NewMetadata->AddEntry());
+			}
+
+			for (const FName AttributeName : AttributeNames)
+			{
+				const FPCGMetadataAttributeBase* OriginalAttribute = OriginalMetadata->GetConstAttribute(AttributeName);
+				FPCGMetadataAttributeBase* NewAttribute = NewMetadata->GetMutableAttribute(AttributeName);
+				check(OriginalAttribute && NewAttribute);
+
+				for (int32 i = 0; i < Indices.Num(); ++i)
+				{
+					NewAttribute->SetValue(EntryKeys[i], OriginalAttribute, Indices[i]);
+				}
+
+				if (Indices.Num() == 1)
+				{
+					NewAttribute->SetDefaultValueToFirstEntry();
+				}
+			}
+		}
+
+		return PartitionedData;
+	}
+
+	TArray<UPCGData*> AttributePartition(const UPCGData* InData, const FPCGAttributePropertySelector& InSelector, FPCGContext* InOptionalContext)
+	{
+		if (const UPCGPointData* InPointData = Cast<UPCGPointData>(InData))
+		{
+			return AttributePointPartition(InPointData, InSelector, InOptionalContext);
+		}
+		else
+		{
+			return AttributeParamSpatialPartition(InData, InSelector, InOptionalContext);
+		}
+	}
+}
+
+#undef LOCTEXT_NAMESPACE
\ No newline at end of file
diff --git a/Engine/Plugins/Experimental/PCG/Source/PCG/Private/Tests/Elements/PCGAttributePartitionTest.cpp b/Engine/Plugins/Experimental/PCG/Source/PCG/Private/Tests/Elements/PCGAttributePartitionTest.cpp
new file mode 100644
--- /dev/null
+++ b/Engine/Plugins/Experimental/PCG/Source/PCG/Private/Tests/Elements/PCGAttributePartitionTest.cpp
@@ -0,0 +1,134 @@
+ // Copyright Epic Games, Inc. All Rights Reserved.
+
+
+// Copyright Epic Games, Inc. All Rights Reserved.
+
+#include "Tests/PCGTestsCommon.h"
+
+#include "PCGComponent.h"
+#include "PCGContext.h"
+#include "PCGParamData.h"
+#include "Data/PCGPointData.h"
+
+#include "Elements/Metadata/PCGMetadataPartition.h"
+
+IMPLEMENT_CUSTOM_SIMPLE_AUTOMATION_TEST(FPCGAttributePartition_Points, FPCGTestBaseClass, "pcg.tests.AttributePartition.Points", PCGTestsCommon::TestFlags)
+IMPLEMENT_CUSTOM_SIMPLE_AUTOMATION_TEST(FPCGAttributePartition_AttributeSet, FPCGTestBaseClass, "pcg.tests.AttributePartition.AttributeSet", PCGTestsCommon::TestFlags)
+
+bool FPCGAttributePartition_Points::RunTest(const FString& Parameters)
+{
+	PCGTestsCommon::FTestData TestData;
+	UPCGMetadataPartitionSettings* Settings = PCGTestsCommon::GenerateSettings<UPCGMetadataPartitionSettings>(TestData);
+	check(Settings);
+
+	Settings->PartitionAttributeSource.SetPointProperty(EPCGPointProperties::Density);
+
+	UPCGPointData* InputPointData = NewObject<UPCGPointData>();
+	TArray<FPCGPoint>& Points = InputPointData->GetMutablePoints();
+	Points.Reserve(100);
+	for (int32 i = 0; i < 100; ++i)
+	{
+		FPCGPoint& Point = Points.Emplace_GetRef();
+		Point.Density = (i % 10) / 10.0f;
+	}
+
+	FPCGTaggedData& InputTaggedData = TestData.InputData.TaggedData.Emplace_GetRef();
+	InputTaggedData.Data = InputPointData;
+	InputTaggedData.Pin = PCGPinConstants::DefaultInputLabel;
+
+	FPCGElementPtr TestElement = TestData.Settings->GetElement();
+	TUniquePtr<FPCGContext> Context = TestData.InitializeTestContext();
+
+	while (!TestElement->Execute(Context.Get())) {}
+
+	UTEST_EQUAL("There are 10 outputs", Context->OutputData.TaggedData.Num(), 10);
+
+	for (int32 i = 0; i < Context->OutputData.TaggedData.Num(); ++i)
+	{
+		const UPCGPointData* OutputPointData = Cast<const UPCGPointData>(Context->OutputData.TaggedData[i].Data);
+		UTEST_NOT_NULL(*FString::Printf(TEXT("Output %d is a point data"), i), OutputPointData);
+
+		const TArray<FPCGPoint>& OutputPoints = OutputPointData->GetPoints();
+
+		UTEST_EQUAL(*FString::Printf(TEXT("Output %d has 10 points"), i), OutputPoints.Num(), 10);
+
+		float TempValue = -1.0f;
+		bool bFirstPoint = true;
+		bool bAllEquals = true;
+		for (const FPCGPoint& Point : OutputPoints)
+		{
+			if (bFirstPoint)
+			{
+				TempValue = Point.Density;
+				bFirstPoint = false;
+			}
+			else
+			{
+				bAllEquals &= TempValue == Point.Density;
+			}
+		}
+
+		UTEST_TRUE(*FString::Printf(TEXT("Output points for output %d have all the same density"), i), bAllEquals);
+	}
+
+	return true;
+}
+
+bool FPCGAttributePartition_AttributeSet::RunTest(const FString& Parameters)
+{
+	PCGTestsCommon::FTestData TestData;
+	UPCGMetadataPartitionSettings* Settings = PCGTestsCommon::GenerateSettings<UPCGMetadataPartitionSettings>(TestData);
+	check(Settings);
+
+	const FName InputAttributeName = TEXT("Double");
+	Settings->PartitionAttributeSource.SetAttributeName(InputAttributeName);
+
+	UPCGParamData* InputParam = NewObject<UPCGParamData>();
+	FPCGMetadataAttribute<double>* Attribute = InputParam->Metadata->CreateAttribute<double>(InputAttributeName, 0.0, true, false);
+	for (int32 i = 0; i < 100; ++i)
+	{
+		Attribute->SetValue(InputParam->Metadata->AddEntry(), (i % 10) / 10.0);
+	}
+
+	FPCGTaggedData& InputTaggedData = TestData.InputData.TaggedData.Emplace_GetRef();
+	InputTaggedData.Data = InputParam;
+	InputTaggedData.Pin = PCGPinConstants::DefaultInputLabel;
+
+	FPCGElementPtr TestElement = TestData.Settings->GetElement();
+	TUniquePtr<FPCGContext> Context = TestData.InitializeTestContext();
+
+	while (!TestElement->Execute(Context.Get())) {}
+
+	UTEST_EQUAL("There are 10 outputs", Context->OutputData.TaggedData.Num(), 10);
+
+	for (int32 i = 0; i < Context->OutputData.TaggedData.Num(); ++i)
+	{
+		const UPCGParamData* OutputParamData = Cast<const UPCGParamData>(Context->OutputData.TaggedData[i].Data);
+		UTEST_NOT_NULL(*FString::Printf(TEXT("Output %d is a param data"), i), OutputParamData);
+
+		UTEST_EQUAL(*FString::Printf(TEXT("Output %d has 10 entries"), i), OutputParamData->Metadata->GetLocalItemCount(), 10ll);
+		const FPCGMetadataAttribute<double>* OutAttribute = OutputParamData->Metadata->GetConstTypedAttribute<double>(InputAttributeName);
+		UTEST_NOT_NULL(*FString::Printf(TEXT("Output %d has the 'Double' attribute"), i), OutAttribute);
+
+
+		double TempValue = -1.0;
+		bool bFirstItem = true;
+		bool bAllEquals = true;
+		for (PCGMetadataEntryKey Key = 0; Key < OutputParamData->Metadata->GetLocalItemCount(); ++Key)
+		{
+			if (bFirstItem)
+			{
+				TempValue = OutAttribute->GetValueFromItemKey(Key);
+				bFirstItem = false;
+			}
+			else
+			{
+				bAllEquals &= TempValue == OutAttribute->GetValueFromItemKey(Key);
+			}
+		}
+
+		UTEST_TRUE(*FString::Printf(TEXT("Output values for output %d are all the same"), i), bAllEquals);
+	}
+
+	return true;
+}
\ No newline at end of file
diff --git a/Engine/Plugins/Experimental/PCG/Source/PCG/Public/Data/PCGSpatialData.h b/Engine/Plugins/Experimental/PCG/Source/PCG/Public/Data/PCGSpatialData.h
--- a/Engine/Plugins/Experimental/PCG/Source/PCG/Public/Data/PCGSpatialData.h
+++ b/Engine/Plugins/Experimental/PCG/Source/PCG/Public/Data/PCGSpatialData.h
@@ -118,10 +118,10 @@ public:
 	virtual UPCGDifferenceData* Subtract(const UPCGSpatialData* InOther) const;
 
 	UFUNCTION(BlueprintCallable, Category = Metadata)
-	const UPCGMetadata* ConstMetadata() const { return Metadata; }
+	virtual const UPCGMetadata* ConstMetadata() const override { return Metadata; }
 
 	UFUNCTION(BlueprintCallable, Category = Metadata)
-	UPCGMetadata* MutableMetadata() { return Metadata; }
+	virtual UPCGMetadata* MutableMetadata() override { return Metadata; }
 
 	UFUNCTION(BlueprintCallable, Category = Metadata, meta=(DeprecatedFunction, DeprecationMessage = "The Create Empty Metadata function is not needed anymore - it can safely be removed"))
 	UPCGMetadata* CreateEmptyMetadata();
diff --git a/Engine/Plugins/Experimental/PCG/Source/PCG/Public/Elements/Metadata/PCGMetadataPartition.h b/Engine/Plugins/Experimental/PCG/Source/PCG/Public/Elements/Metadata/PCGMetadataPartition.h
--- a/Engine/Plugins/Experimental/PCG/Source/PCG/Public/Elements/Metadata/PCGMetadataPartition.h
+++ b/Engine/Plugins/Experimental/PCG/Source/PCG/Public/Elements/Metadata/PCGMetadataPartition.h
@@ -3,6 +3,7 @@
 #pragma once
 
 #include "PCGSettings.h"
+#include "Metadata/PCGAttributePropertySelector.h"
 
 #include "PCGMetadataPartition.generated.h"
 
@@ -12,22 +13,33 @@ class PCG_API UPCGMetadataPartitionSettings : public UPCGSettings
 	GENERATED_BODY()
 
 public:
+	//~Begin UObject interface
+	virtual void PostLoad() override;
+	//~End UObject interface
+
 	//~Begin UPCGSettings interface
 #if WITH_EDITOR
 	virtual FName GetDefaultNodeName() const override { return FName(TEXT("AttributePartition")); }
 	virtual FText GetDefaultNodeTitle() const override { return NSLOCTEXT("PCGMetadataPartitionSettings", "NodeTitle", "Attribute Partition"); }
 	virtual EPCGSettingsType GetType() const override { return EPCGSettingsType::Metadata; }
+	virtual bool HasDynamicPins() const override { return true; }
 #endif
+	virtual EPCGDataType GetCurrentPinTypes(const UPCGPin* InPin) const override;
 
 protected:
-	virtual TArray<FPCGPinProperties> InputPinProperties() const override { return Super::DefaultPointInputPinProperties(); }
-	virtual TArray<FPCGPinProperties> OutputPinProperties() const override { return Super::DefaultPointOutputPinProperties(); }
+	virtual TArray<FPCGPinProperties> InputPinProperties() const override;
+	virtual TArray<FPCGPinProperties> OutputPinProperties() const override;
 	virtual FPCGElementPtr CreateElement() const override;
 	//~End UPCGSettings interface
 
 public:
-	UPROPERTY(BlueprintReadWrite, EditAnywhere, Category = Settings, meta = (PCG_Overridable))
-	FName PartitionAttribute = NAME_None;
+	UPROPERTY(BlueprintReadWrite, EditAnywhere, Category = Settings, meta = (PCG_Overridable, PCG_OverrideAliases = "PartitionAttribute"))
+	FPCGAttributePropertyInputSelector PartitionAttributeSource;
+
+#if WITH_EDITORONLY_DATA
+	UPROPERTY()
+	FName PartitionAttribute_DEPRECATED = NAME_None;
+#endif // WITH_EDITORONLY_DATA
 };
 
 class FPCGMetadataPartitionElement : public IPCGElement
diff --git a/Engine/Plugins/Experimental/PCG/Source/PCG/Public/Metadata/PCGMetadataPartitionCommon.h b/Engine/Plugins/Experimental/PCG/Source/PCG/Public/Metadata/PCGMetadataPartitionCommon.h
new file mode 100644
--- /dev/null
+++ b/Engine/Plugins/Experimental/PCG/Source/PCG/Public/Metadata/PCGMetadataPartitionCommon.h
@@ -0,0 +1,33 @@
+// Copyright Epic Games, Inc. All Rights Reserved.
+
+#pragma once
+
+#include "Containers/Array.h"
+#include "HAL/Platform.h"
+
+struct FPCGAttributePropertySelector;
+struct FPCGContext;
+class UPCGData;
+
+namespace PCGMetadataPartitionCommon
+{
+	/**
+	* Partition the incoming data on the given attribute.
+	* Will return an array of data. Each data will contain the entries (in stable order) for a given partition.
+	* @param InData - Data to partition.
+	* @param InSelector - Selector on the attribute to partition.
+	* @param InOptionalContext - Optional context for logging.
+	* @returns Array of data.
+	*/
+	PCG_API TArray<UPCGData*> AttributePartition(const UPCGData* InData, const FPCGAttributePropertySelector& InSelector, FPCGContext* InOptionalContext = nullptr);
+
+	/**
+	* Generic partition for the incoming data on the given attribute.
+	* Will return an array of bucket indices. Each bucket will contain the indices (in stable order) for a given partition.
+	* @param InData - Data to partition, need to support attributes (spatial data or attribute set).
+	* @param InSelector - Selector on the attribute to partition.
+	* @param InOptionalContext - Optional context for logging.
+	* @returns Array of bucket indices.
+	*/
+	PCG_API TArray<TArray<int32>> AttributeGenericPartition(const UPCGData* InData, const FPCGAttributePropertySelector& InSelector, FPCGContext* InOptionalContext = nullptr);
+}
\ No newline at end of file
diff --git a/Engine/Plugins/Experimental/PCG/Source/PCG/Public/PCGData.h b/Engine/Plugins/Experimental/PCG/Source/PCG/Public/PCGData.h
--- a/Engine/Plugins/Experimental/PCG/Source/PCG/Public/PCGData.h
+++ b/Engine/Plugins/Experimental/PCG/Source/PCG/Public/PCGData.h
@@ -54,6 +54,10 @@ public:
 	virtual void PostDuplicate(bool bDuplicateForPIE) override { InitUID(); }
 	// ~End UObject interface
 
+	// Metadata getters, to be implemented if data supports Metadata
+	virtual UPCGMetadata* MutableMetadata() { return nullptr; }
+	virtual const UPCGMetadata* ConstMetadata() const { return nullptr; }
+
 protected:
 	/** Computes Crc for this and any connected data. */
 	virtual FPCGCrc ComputeCrc(bool bFullDataCrc) const;
diff --git a/Engine/Plugins/Experimental/PCG/Source/PCG/Public/PCGParamData.h b/Engine/Plugins/Experimental/PCG/Source/PCG/Public/PCGParamData.h
--- a/Engine/Plugins/Experimental/PCG/Source/PCG/Public/PCGParamData.h
+++ b/Engine/Plugins/Experimental/PCG/Source/PCG/Public/PCGParamData.h
@@ -28,10 +28,10 @@ public:
 	// ~End UPCGData interface
 
 	UFUNCTION(BlueprintCallable, Category = Metadata)
-	const UPCGMetadata* ConstMetadata() const { return Metadata; }
+	virtual const UPCGMetadata* ConstMetadata() const override{ return Metadata; }
 
 	UFUNCTION(BlueprintCallable, Category = Metadata)
-	UPCGMetadata* MutableMetadata() { return Metadata; }
+	virtual UPCGMetadata* MutableMetadata() override { return Metadata; }
 
 	/** Returns the entry for the given name */
 	UFUNCTION(BlueprintCallable, Category = Params)
