[PCG] Add a Color Projection Mode to the Projection Element to allow for:

* Source
* Target
* Blend (Add, Subtract, Multiply)

diff --git a/Engine/Plugins/Experimental/PCG/Source/PCG/Private/Data/PCGPointData.cpp b/Engine/Plugins/Experimental/PCG/Source/PCG/Private/Data/PCGPointData.cpp
--- a/Engine/Plugins/Experimental/PCG/Source/PCG/Private/Data/PCGPointData.cpp
+++ b/Engine/Plugins/Experimental/PCG/Source/PCG/Private/Data/PCGPointData.cpp
@@ -481,7 +481,8 @@ bool UPCGPointData::SamplePoint(const FTransform& InTransform, const FBox& InBou
 {
 	// Run a projection but don't change the point transform. There is a large overlap in code/functionality so this shares one code path.
 	FPCGProjectionParams Params{};
-	Params.bProjectPositions = Params.bProjectRotations = Params.bProjectScales = Params.bProjectColors = false;
+	Params.bProjectPositions = Params.bProjectRotations = Params.bProjectScales = false;
+	Params.ColorBlendMode = EPCGProjectionColorBlendMode::SourceValue;
 
 	// The ProjectPoint implementation in this class returns true if the query point is overlapping the point data, which is what SamplePoint should return, so forward the return value.
 	return ProjectPoint(InTransform, InBounds, Params, OutPoint, OutMetadata);
diff --git a/Engine/Plugins/Experimental/PCG/Source/PCG/Private/Data/PCGProjectionData.cpp b/Engine/Plugins/Experimental/PCG/Source/PCG/Private/Data/PCGProjectionData.cpp
--- a/Engine/Plugins/Experimental/PCG/Source/PCG/Private/Data/PCGProjectionData.cpp
+++ b/Engine/Plugins/Experimental/PCG/Source/PCG/Private/Data/PCGProjectionData.cpp
@@ -12,6 +12,38 @@
 
 #include UE_INLINE_GENERATED_CPP_BY_NAME(PCGProjectionData)
 
+namespace PCGProjectionPrivate
+{
+	FVector4 ApplyProjectionColorBlend(const FVector4& SourceColor, const FVector4& TargetColor, const EPCGProjectionColorBlendMode BlendMode)
+	{
+		auto ClampVector = [](FVector4 Vector) -> FVector4
+		{
+			Vector[0] = FMath::Clamp(Vector[0], 0.0, 1.0);
+			Vector[1] = FMath::Clamp(Vector[1], 0.0, 1.0);
+			Vector[2] = FMath::Clamp(Vector[2], 0.0, 1.0);
+			Vector[3] = FMath::Clamp(Vector[3], 0.0, 1.0);
+			return Vector;
+		};
+
+		switch (BlendMode)
+		{
+			case EPCGProjectionColorBlendMode::SourceValue:
+				return SourceColor;
+			case EPCGProjectionColorBlendMode::TargetValue:
+				return TargetColor;
+			case EPCGProjectionColorBlendMode::Add:
+				return ClampVector(TargetColor + SourceColor);
+			case EPCGProjectionColorBlendMode::Subtract:
+				return ClampVector(TargetColor - SourceColor);
+			case EPCGProjectionColorBlendMode::Multiply:
+				return SourceColor * TargetColor;
+			default:
+				checkNoEntry();
+				return FVector4::Zero();
+		}
+	}
+}
+
 void UPCGProjectionData::Initialize(const UPCGSpatialData* InSource, const UPCGSpatialData* InTarget, const FPCGProjectionParams& InProjectionParams)
 {
 	check(InSource && InTarget);
@@ -28,6 +60,13 @@ void UPCGProjectionData::Initialize(const UPCGSpatialData* InSource, const UPCGS
 	CachedStrictBounds = ProjectBounds(Source->GetStrictBounds());
 }
 
+void UPCGProjectionData::PostLoad()
+{
+	Super::PostLoad();
+
+	ProjectionParams.ApplyDeprecation();
+}
+
 void UPCGProjectionData::AddToCrc(FArchiveCrc32& Ar, bool bFullDataCrc) const
 {
 	Super::AddToCrc(Ar, bFullDataCrc);
@@ -256,9 +295,11 @@ const UPCGPointData* UPCGProjectionData::CreatePointData(FPCGContext* Context) c
 		// Apply projection result. Some of the params are already used inside ProjectPoint, so we are just applying the remaining bits that ProjectPoint() did not have access to.
 		// TODO this would be cleaner if there was a ProjectPoint that took an FPCGPoint
 		OutPoint.Transform = PointFromTarget.Transform;
-		OutPoint.Color = ProjectionParams.bProjectColors ? (PointFromTarget.Color * SourcePoint.Color) : SourcePoint.Color;
+
+		OutPoint.Color = PCGProjectionPrivate::ApplyProjectionColorBlend(SourcePoint.Color, PointFromTarget.Color, ProjectionParams.ColorBlendMode);
+
 		OutPoint.Density *= PointFromTarget.Density;
-		
+
 		if (OutMetadata && TempTargetMetadata && PointFromTarget.MetadataEntry != PCGInvalidEntryKey)
 		{
 			// Merge metadata to produce final attribute values
@@ -290,10 +331,7 @@ void UPCGProjectionData::ApplyProjectionResult(const FPCGPoint& InTargetPoint, F
 		InOutProjected.Transform.SetScale3D(InTargetPoint.Transform.GetScale3D());
 	}
 
-	if (ProjectionParams.bProjectColors)
-	{
-		InOutProjected.Color *= InTargetPoint.Color;
-	}
+	InOutProjected.Color = PCGProjectionPrivate::ApplyProjectionColorBlend(InOutProjected.Color, InTargetPoint.Color, ProjectionParams.ColorBlendMode);
 
 	InOutProjected.Density *= InTargetPoint.Density;
 }
diff --git a/Engine/Plugins/Experimental/PCG/Source/PCG/Private/Elements/PCGProjectionElement.cpp b/Engine/Plugins/Experimental/PCG/Source/PCG/Private/Elements/PCGProjectionElement.cpp
--- a/Engine/Plugins/Experimental/PCG/Source/PCG/Private/Elements/PCGProjectionElement.cpp
+++ b/Engine/Plugins/Experimental/PCG/Source/PCG/Private/Elements/PCGProjectionElement.cpp
@@ -22,6 +22,13 @@ TArray<FPCGPinProperties> UPCGProjectionSettings::InputPinProperties() const
 	return PinProperties;
 }
 
+void UPCGProjectionSettings::PostLoad()
+{
+	Super::PostLoad();
+
+	ProjectionParams.ApplyDeprecation();
+}
+
 #if WITH_EDITOR
 FText UPCGProjectionSettings::GetNodeTooltipText() const
 {
diff --git a/Engine/Plugins/Experimental/PCG/Source/PCG/Public/Data/PCGProjectionData.h b/Engine/Plugins/Experimental/PCG/Source/PCG/Public/Data/PCGProjectionData.h
--- a/Engine/Plugins/Experimental/PCG/Source/PCG/Public/Data/PCGProjectionData.h
+++ b/Engine/Plugins/Experimental/PCG/Source/PCG/Public/Data/PCGProjectionData.h
@@ -16,6 +16,10 @@ class PCG_API UPCGProjectionData : public UPCGSpatialDataWithPointCache
 public:
 	void Initialize(const UPCGSpatialData* InSource, const UPCGSpatialData* InTarget, const FPCGProjectionParams& InProjectionParams);
 
+	// ~Begin UObject interface
+	virtual void PostLoad() override;
+	// ~End UObject interface
+
 	const FPCGProjectionParams& GetProjectionParams() const { return ProjectionParams; }
 
 	// ~Begin UPCGData interface
diff --git a/Engine/Plugins/Experimental/PCG/Source/PCG/Public/Elements/PCGProjectionElement.h b/Engine/Plugins/Experimental/PCG/Source/PCG/Public/Elements/PCGProjectionElement.h
--- a/Engine/Plugins/Experimental/PCG/Source/PCG/Public/Elements/PCGProjectionElement.h
+++ b/Engine/Plugins/Experimental/PCG/Source/PCG/Public/Elements/PCGProjectionElement.h
@@ -18,6 +18,10 @@ class PCG_API UPCGProjectionSettings : public UPCGSettings
 	GENERATED_BODY()
 
 public:
+	//~Begin UObject interface
+	virtual void PostLoad() override;
+	//~End UObject interface
+
 	//~Begin UPCGSettings interface
 #if WITH_EDITOR
 	virtual FName GetDefaultNodeName() const override { return FName(TEXT("Projection")); }
diff --git a/Engine/Plugins/Experimental/PCG/Source/PCG/Public/Elements/PCGProjectionParams.h b/Engine/Plugins/Experimental/PCG/Source/PCG/Public/Elements/PCGProjectionParams.h
--- a/Engine/Plugins/Experimental/PCG/Source/PCG/Public/Elements/PCGProjectionParams.h
+++ b/Engine/Plugins/Experimental/PCG/Source/PCG/Public/Elements/PCGProjectionParams.h
@@ -4,15 +4,46 @@
 
 #include "Metadata/PCGMetadataCommon.h"
 
-
 #include "PCGProjectionParams.generated.h"
 
+UENUM()
+enum class EPCGProjectionColorBlendMode
+{
+	SourceValue,
+	TargetValue,
+	Add,
+	Subtract,
+	Multiply
+};
+
 /** Parameters that control projection behaviour. */
 USTRUCT(BlueprintType)
 struct PCG_API FPCGProjectionParams
 {
 	GENERATED_BODY()
 
+	// TODO [DEPRECATED_IN_5_4]: This rule of five necessary to avoid compile-time warnings. To be removed in when the deprecated member has expired
+	PRAGMA_DISABLE_DEPRECATION_WARNINGS
+	FPCGProjectionParams() = default;
+	~FPCGProjectionParams() = default;
+	FPCGProjectionParams(const FPCGProjectionParams&) = default;
+	FPCGProjectionParams(FPCGProjectionParams&&) = default;
+	FPCGProjectionParams& operator=(const FPCGProjectionParams&) = default;
+	FPCGProjectionParams& operator=(FPCGProjectionParams&&) = default;
+
+	void ApplyDeprecation()
+	{
+#if WITH_EDITOR
+		// To maintain backwards compatibility, a 'true' here would've multiplied as the blend mode
+		if (bProjectColors_DEPRECATED)
+		{
+			ColorBlendMode = EPCGProjectionColorBlendMode::Multiply;
+			bProjectColors_DEPRECATED = false;
+		}
+#endif // WITH_EDITOR
+	}
+	PRAGMA_ENABLE_DEPRECATION_WARNINGS
+
 	/** Project positions. */
 	UPROPERTY(BlueprintReadWrite, EditAnywhere, Category = "Apply Data")
 	bool bProjectPositions = true;
@@ -25,9 +56,9 @@ struct PCG_API FPCGProjectionParams
 	UPROPERTY(BlueprintReadWrite, EditAnywhere, Category = "Apply Data")
 	bool bProjectScales = false;
 
-	/** Project colors. */
+	/** The blend mode for colors during the projection */
 	UPROPERTY(BlueprintReadWrite, EditAnywhere, Category = "Apply Data")
-	bool bProjectColors = false;
+	EPCGProjectionColorBlendMode ColorBlendMode = EPCGProjectionColorBlendMode::SourceValue;
 
 	/** Attributes to either explicitly exclude or include in the projection operation, depending on the Attribute Mode setting. Leave empty to gather all attributes and their values. Format is comma separated list like: Attribute1,Attribute2 .*/
 	UPROPERTY(BlueprintReadWrite, EditAnywhere, Category = "Apply Data")
@@ -40,6 +71,13 @@ struct PCG_API FPCGProjectionParams
 	/** Operation to use to combine attributes that reside on both source and target data. */
 	UPROPERTY(BlueprintReadWrite, EditAnywhere, Category = "Apply Data")
 	EPCGMetadataOp AttributeMergeOperation = EPCGMetadataOp::TargetValue;
+
+	// TODO [DEPRECATED_IN_5_4]: Remove the 'rule of five' needed above upon deprecation
+#if WITH_EDITORONLY_DATA
+	UE_DEPRECATED(5.4, "This property will be deprecated in 5.4. Please use `ColorBlendMode` instead.")
+	UPROPERTY(meta = (DeprecatedProperty))
+	bool bProjectColors_DEPRECATED = false;
+#endif // WITH_EDITORONLY_DATA
 };
 
 #if UE_ENABLE_INCLUDE_ORDER_DEPRECATED_IN_5_2
