[PCG] Fixing duplicated copyright

#rnx

diff --git a/Engine/Plugins/Experimental/PCG/Source/PCG/Private/Tests/Elements/PCGAttributePartitionTest.cpp b/Engine/Plugins/Experimental/PCG/Source/PCG/Private/Tests/Elements/PCGAttributePartitionTest.cpp
--- a/Engine/Plugins/Experimental/PCG/Source/PCG/Private/Tests/Elements/PCGAttributePartitionTest.cpp
+++ b/Engine/Plugins/Experimental/PCG/Source/PCG/Private/Tests/Elements/PCGAttributePartitionTest.cpp
@@ -1,6 +1,3 @@
- // Copyright Epic Games, Inc. All Rights Reserved.
-
-
 // Copyright Epic Games, Inc. All Rights Reserved.
 
 #include "Tests/PCGTestsCommon.h"
