[PCG] Hotfix for updating a broken BP node asset DistanceToNeighbors

#rnx

diff --git a/Engine/Plugins/Experimental/PCG/Source/PCG/Public/Elements/PCGProjectionParams.h b/Engine/Plugins/Experimental/PCG/Source/PCG/Public/Elements/PCGProjectionParams.h
--- a/Engine/Plugins/Experimental/PCG/Source/PCG/Public/Elements/PCGProjectionParams.h
+++ b/Engine/Plugins/Experimental/PCG/Source/PCG/Public/Elements/PCGProjectionParams.h
@@ -7,7 +7,7 @@
 #include "PCGProjectionParams.generated.h"
 
 UENUM()
-enum class EPCGProjectionColorBlendMode
+enum class EPCGProjectionColorBlendMode : uint8
 {
 	SourceValue,
 	TargetValue,
