[PCG] Fix issue with Create Attribute not adding entries for the newly created attribute in a multi entry attribute set

* Should fix issue with incorrect merge

#rnx

diff --git a/Engine/Plugins/Experimental/PCG/Source/PCG/Private/Elements/PCGCreateAttribute.cpp b/Engine/Plugins/Experimental/PCG/Source/PCG/Private/Elements/PCGCreateAttribute.cpp
--- a/Engine/Plugins/Experimental/PCG/Source/PCG/Private/Elements/PCGCreateAttribute.cpp
+++ b/Engine/Plugins/Experimental/PCG/Source/PCG/Private/Elements/PCGCreateAttribute.cpp
@@ -365,16 +365,9 @@ bool FPCGCreateAttributeElement::ExecuteInternal(FPCGContext* Context) const
 				Attribute = Metadata->CopyAttribute(SourceAttribute, OutputAttributeName, /*bKeepParent=*/false, /*bCopyEntries=*/false, /*bCopyValues=*/false);
 
 				// We perhaps need to fix the default value. If the first entry is different from the default value, we override the default value.
-				if (Attribute && Attribute->IsEqualToDefaultValue(Attribute->GetValueKey(0)))
+				if (Attribute)
 				{
-					auto FixDefaultValue = [Attribute](auto Dummy)
-					{
-						using AttributeType = decltype(Dummy);
-						FPCGMetadataAttribute<AttributeType>* TypedAttribute = static_cast<FPCGMetadataAttribute<AttributeType>*>(Attribute);
-						TypedAttribute->SetDefaultValue(TypedAttribute->GetValueFromItemKey(0));
-					};
-
-					PCGMetadataAttribute::CallbackWithRightType(Attribute->GetTypeId(), FixDefaultValue);
+					Attribute->SetDefaultValueToFirstEntry();
 				}
 			}
 			else // Create a new attribute of the accessed field's type manually
@@ -448,10 +441,23 @@ bool FPCGCreateAttributeElement::ExecuteInternal(FPCGContext* Context) const
 			continue;
 		}
 
-		// Making sure the metadata has at least one entry.
-		if (bIsParamData && Metadata->GetLocalItemCount() == 0)
+		if (bIsParamData)
 		{
-			Metadata->AddEntry();
+			// Making sure the metadata has at least one entry.
+			if (Metadata->GetLocalItemCount() == 0)
+			{
+				Metadata->AddEntry();
+			}
+
+			// Also make sure if the attribute has no entries, to map all the metadata entry to the default value key
+			// Very important to do in the case of multi-entry attribute set.
+			if (Attribute->GetNumberOfEntries() == 0)
+			{
+				for (PCGMetadataEntryKey Key = 0; Key < Metadata->GetLocalItemCount(); ++Key)
+				{
+					Attribute->SetValueFromValueKey(Key, PCGDefaultValueKey);
+				}
+			}
 		}
 
 		TArray<FPCGTaggedData>& Outputs = Context->OutputData.TaggedData;
