[PCG] Add support for multiple actors in GetActorProperty

#rnx

diff --git a/Engine/Plugins/Experimental/PCG/Source/PCG/Private/Elements/PCGPropertyToParamData.cpp b/Engine/Plugins/Experimental/PCG/Source/PCG/Private/Elements/PCGPropertyToParamData.cpp
--- a/Engine/Plugins/Experimental/PCG/Source/PCG/Private/Elements/PCGPropertyToParamData.cpp
+++ b/Engine/Plugins/Experimental/PCG/Source/PCG/Private/Elements/PCGPropertyToParamData.cpp
@@ -135,34 +135,46 @@ bool FPCGPropertyToParamDataElement::ExecuteInternal(FPCGContext* Context) const
 		};
 	}
 
-	AActor* FoundActor = PCGActorSelector::FindActor(Settings->ActorSelector, OriginalComponent, BoundsCheck, NoSelfIgnoreCheck);
+	TArray<AActor*> FoundActors = PCGActorSelector::FindActors(Settings->ActorSelector, OriginalComponent, BoundsCheck, NoSelfIgnoreCheck);
 
-	if (!FoundActor)
+	if (FoundActors.IsEmpty())
 	{
 		PCGE_LOG(Verbose, LogOnly, LOCTEXT("NoActorFound", "No matching actor was found"));
 		return true;
 	}
 
-	// From there, we either check the actor, or the component attached to it.
-	UObject* ObjectToInspect = FoundActor;
-	if (Settings->bSelectComponent)
+	for (AActor* FoundActor : FoundActors)
 	{
-		ObjectToInspect = FoundActor->GetComponentByClass(Settings->ComponentClass);
-		if (!ObjectToInspect)
+		if (!FoundActor)
 		{
-			PCGE_LOG(Error, GraphAndLog, FText::Format(LOCTEXT("ComponentDoesNotExist", "Component class '{0}' does not exist in the found actor"), FText::FromString(Settings->ComponentClass->GetName())));
-			return true;
+			continue;
 		}
-	}
 
-	FPCGAttributePropertySelector Selector = FPCGAttributePropertySelector::CreateSelectorFromString(Settings->PropertyName.ToString());
+		// From there, we either check the actor, or the component attached to it.
+		UObject* ObjectToInspect = FoundActor;
+		if (Settings->bSelectComponent)
+		{
+			ObjectToInspect = FoundActor->GetComponentByClass(Settings->ComponentClass);
+			if (!ObjectToInspect)
+			{
+				PCGE_LOG(Error, GraphAndLog, FText::Format(LOCTEXT("ComponentDoesNotExist", "Component class '{0}' does not exist in the found actor {1}"), FText::FromString(Settings->ComponentClass->GetName()), FText::FromString(FoundActor->GetName())));
+				return true;
+			}
+		}
 
-	PCGPropertyHelpers::FExtractorParameters Parameters{ ObjectToInspect, ObjectToInspect->GetClass(), Selector, Settings->OutputAttributeName, Settings->bForceObjectAndStructExtraction, /*bPropertyNeedsToBeVisible=*/true};
-	if (UPCGParamData* ParamData = PCGPropertyHelpers::ExtractPropertyAsAttributeSet(Parameters, Context))
-	{
-		TArray<FPCGTaggedData>& Outputs = Context->OutputData.TaggedData;
-		FPCGTaggedData& Output = Outputs.Emplace_GetRef();
-		Output.Data = ParamData;
+		const FPCGAttributePropertySelector Selector = FPCGAttributePropertySelector::CreateSelectorFromString(Settings->PropertyName.ToString());
+
+		PCGPropertyHelpers::FExtractorParameters Parameters{ ObjectToInspect, ObjectToInspect->GetClass(), Selector, Settings->OutputAttributeName, Settings->bForceObjectAndStructExtraction, /*bPropertyNeedsToBeVisible=*/true };
+		if (UPCGParamData* ParamData = PCGPropertyHelpers::ExtractPropertyAsAttributeSet(Parameters, Context))
+		{
+			TArray<FPCGTaggedData>& Outputs = Context->OutputData.TaggedData;
+			FPCGTaggedData& Output = Outputs.Emplace_GetRef();
+			Output.Data = ParamData;
+		}
+		else
+		{
+			PCGE_LOG(Error, GraphAndLog, FText::Format(LOCTEXT("FailedToExtract", "Fail to extract the property '{0}' on actor {1}"), Selector.GetDisplayText(), FText::FromString(FoundActor->GetName())));
+		}
 	}
 
 	return true;
