[PCG] Change engine settings from DefaultPlugin to DefaultEngine

#rnx

diff --git a/Engine/Plugins/Experimental/PCG/Source/PCG/Public/PCGEngineSettings.h b/Engine/Plugins/Experimental/PCG/Source/PCG/Public/PCGEngineSettings.h
--- a/Engine/Plugins/Experimental/PCG/Source/PCG/Public/PCGEngineSettings.h
+++ b/Engine/Plugins/Experimental/PCG/Source/PCG/Public/PCGEngineSettings.h
@@ -6,7 +6,7 @@
 
 #include "PCGEngineSettings.generated.h"
 
-UCLASS(config = Plugins, defaultconfig)
+UCLASS(config = Engine, defaultconfig)
 class PCG_API UPCGEngineSettings : public UDeveloperSettings
 {
 	GENERATED_BODY()
