[PCG] Nuke delete of partition actor should also remove children actors.

#rnx

diff --git a/Engine/Plugins/Experimental/PCG/Source/PCG/Private/PCGSubsystem.cpp b/Engine/Plugins/Experimental/PCG/Source/PCG/Private/PCGSubsystem.cpp
--- a/Engine/Plugins/Experimental/PCG/Source/PCG/Private/PCGSubsystem.cpp
+++ b/Engine/Plugins/Experimental/PCG/Source/PCG/Private/PCGSubsystem.cpp
@@ -1140,7 +1140,7 @@ void UPCGSubsystem::ClearPCGLink(UPCGComponent* InComponent, const FBox& InBound
 	}
 }
 
-void UPCGSubsystem::DeletePartitionActors(bool bOnlyDeleteUnused)
+void UPCGSubsystem::DeletePartitionActors(bool bOnlyDeleteUnused, bool bOnlyChildren)
 {
 	TRACE_CPUPROFILER_EVENT_SCOPE(UPCGSubsystem::DeletePartitionActors);
 
@@ -1159,7 +1159,7 @@ void UPCGSubsystem::DeletePartitionActors(bool bOnlyDeleteUnused)
 		return;
 	}
 
-	auto GatherAndDestroyLoadedActors = [&PackagesToCleanup, &PackagesToDeleteFromSCC, World, &ActorsNotSafeToBeDeleted](AActor* Actor) -> bool
+	auto GatherAndDestroyLoadedActors = [&PackagesToCleanup, &PackagesToDeleteFromSCC, World, &ActorsNotSafeToBeDeleted, bOnlyChildren](AActor* Actor) -> bool
 	{
 		// Make sure that this actor was not flagged to not be deleted, or not safe for deletion
 		TObjectPtr<APCGPartitionActor> PartitionActor = CastChecked<APCGPartitionActor>(Actor);
@@ -1178,23 +1178,41 @@ void UPCGSubsystem::DeletePartitionActors(bool bOnlyDeleteUnused)
 				}
 			}
 
-			if (UPackage* ExternalPackage = PartitionActor->GetExternalPackage())
+			// Gather all child actors
+			TArray<AActor*> AttachedActors;
+			PartitionActor->GetAttachedActors(AttachedActors, /*bResetArray=*/true, /*bRecursivelyIncludeAttachedActors=*/ true);
+
+			if (!bOnlyChildren)
 			{
-				PackagesToCleanup.Add(ExternalPackage);
+				AttachedActors.Add(PartitionActor);
 			}
 
-			World->DestroyActor(PartitionActor);
+			for (AActor* ActorToDelete : AttachedActors)
+			{
+				if (UPackage* ExternalPackage = ActorToDelete->GetExternalPackage())
+				{
+					PackagesToCleanup.Add(ExternalPackage);
+				}
+
+				World->DestroyActor(ActorToDelete);
+			}
 		}
 
 		return true;
 	};
 
 	auto GatherAndDestroyActors = [&PackagesToCleanup, &PackagesToDeleteFromSCC, World, &ActorsNotSafeToBeDeleted, &GatherAndDestroyLoadedActors](const FWorldPartitionActorDesc* ActorDesc) {
-		if (ActorDesc->GetActor())
+		AActor* Actor = ActorDesc->GetActor();
+		if (!Actor)
 		{
-			GatherAndDestroyLoadedActors(ActorDesc->GetActor());
+			Actor = ActorDesc->Load();
 		}
-		else
+
+		if (Actor)
+		{
+			GatherAndDestroyLoadedActors(Actor);
+		}
+		else // Couldn't load it
 		{
 			PackagesToDeleteFromSCC.Add(ActorDesc->GetActorPackage().ToString());
 			World->GetWorldPartition()->RemoveActor(ActorDesc->GetGuid());
diff --git a/Engine/Plugins/Experimental/PCG/Source/PCG/Public/PCGSubsystem.h b/Engine/Plugins/Experimental/PCG/Source/PCG/Public/PCGSubsystem.h
--- a/Engine/Plugins/Experimental/PCG/Source/PCG/Public/PCGSubsystem.h
+++ b/Engine/Plugins/Experimental/PCG/Source/PCG/Public/PCGSubsystem.h
@@ -180,7 +180,7 @@ public:
 
 	/** Partition actors methods */
 	void CleanupPartitionActors(const FBox& InBounds);
-	void DeletePartitionActors(bool bOnlyDeleteUnused);
+	void DeletePartitionActors(bool bOnlyDeleteUnused, bool bOnlyChildren = false);
 
 	/** Propagate to the graph compiler graph changes */
 	void NotifyGraphChanged(UPCGGraph* InGraph);
diff --git a/Engine/Plugins/Experimental/PCG/Source/PCGEditor/Private/PCGEditorModule.cpp b/Engine/Plugins/Experimental/PCG/Source/PCGEditor/Private/PCGEditorModule.cpp
--- a/Engine/Plugins/Experimental/PCG/Source/PCGEditor/Private/PCGEditorModule.cpp
+++ b/Engine/Plugins/Experimental/PCG/Source/PCGEditor/Private/PCGEditorModule.cpp
@@ -172,6 +172,22 @@ void FPCGEditorModule::PopulateMenuActions(FMenuBuilder& MenuBuilder)
 					}
 				}
 			})),
+		NAME_None);	
+	
+	MenuBuilder.AddMenuEntry(
+		LOCTEXT("DeletePCGPartitionActors", "Delete all PCG partition actors children"),
+		LOCTEXT("DeletePCGPartitionActors_Tooltip", "Deletes all PCG partition actors children in the current world, but not the Partition Actors themselves"),
+		FSlateIcon(),
+		FUIAction(
+			FExecuteAction::CreateLambda([]() {
+				if (GEditor)
+				{
+					if (UPCGSubsystem* PCGSubsystem = UPCGSubsystem::GetInstance(GEditor->GetEditorWorldContext().World()))
+					{
+						PCGSubsystem->DeletePartitionActors(/*bOnlyDeleteUnused=*/false, /*bOnlyChildren=*/true);
+					}
+				}
+			})),
 		NAME_None);
 
 	MenuBuilder.AddMenuEntry(
