[PCG] Created new CVar to strip empty point data from node outputs

#rnx

diff --git a/Engine/Plugins/Experimental/PCG/Source/PCG/Private/Graph/PCGGraphExecutor.cpp b/Engine/Plugins/Experimental/PCG/Source/PCG/Private/Graph/PCGGraphExecutor.cpp
--- a/Engine/Plugins/Experimental/PCG/Source/PCG/Private/Graph/PCGGraphExecutor.cpp
+++ b/Engine/Plugins/Experimental/PCG/Source/PCG/Private/Graph/PCGGraphExecutor.cpp
@@ -62,6 +62,11 @@ static TAutoConsoleVariable<bool> CVarGraphMultithreading(
 	false,
 	TEXT("Controls whether the graph can dispatch multiple tasks at the same time"));
 
+static TAutoConsoleVariable<bool> CVarStripEmptyPointData(
+	TEXT("pcg.StripEmptyPointData"),
+	false,
+	TEXT("Will strip empty point data from being passed along as output"));
+
 FPCGGraphExecutor::FPCGGraphExecutor(UObject* InOwner)
 	: GraphCompiler(MakeUnique<FPCGGraphCompiler>())
 	, GraphCache(InOwner, &DataRootSet)
@@ -500,6 +505,7 @@ void FPCGGraphExecutor::Execute()
 	const int32 MaxNumThreads = FMath::Max(0, FMath::Min((int32)(FPlatformMisc::NumberOfCoresIncludingHyperthreads() * MaxPercentageOfThreadsToUse), CVarMaxNumTasks.GetValueOnAnyThread() - 1));
 	const bool bAllowMultiDispatch = CVarGraphMultithreading.GetValueOnAnyThread();
 	const bool bGraphCacheDebuggingEnabled = IsGraphCacheDebuggingEnabled();
+	const bool bStripEmptyPointData = CVarStripEmptyPointData.GetValueOnAnyThread();
 
 #if WITH_EDITOR
 	UpdateGenerationNotification();
@@ -762,7 +768,7 @@ void FPCGGraphExecutor::Execute()
 			}
 		}
 
-		auto PostTaskExecute = [this, &bAnyTaskEnded, bGraphCacheDebuggingEnabled](int32 TaskIndex)
+		auto PostTaskExecute = [this, &bAnyTaskEnded, bGraphCacheDebuggingEnabled, bStripEmptyPointData](int32 TaskIndex)
 		{
 			FPCGGraphActiveTask& ActiveTask = ActiveTasks[TaskIndex];
 
@@ -812,6 +818,16 @@ void FPCGGraphExecutor::Execute()
 			}
 #endif
 
+			// TODO: Evaluation code for finding elements that produce empty point data. This is temporary.
+			if (bStripEmptyPointData)
+			{
+				const int32 NumRemoved = ActiveTask.Context->OutputData.StripEmptyPointData();
+				if (NumRemoved > 0)
+				{
+					UE_LOG(LogPCG, Log, TEXT("%d empty point data stripped from node: %s"), NumRemoved, *ActiveTask.Context->Node->GetNodeTitle().ToString());
+				}
+			}
+
 			// Store output in data map
 			StoreResults(ActiveTask.NodeId, ActiveTask.Context->OutputData);
 
diff --git a/Engine/Plugins/Experimental/PCG/Source/PCG/Private/PCGData.cpp b/Engine/Plugins/Experimental/PCG/Source/PCG/Private/PCGData.cpp
--- a/Engine/Plugins/Experimental/PCG/Source/PCG/Private/PCGData.cpp
+++ b/Engine/Plugins/Experimental/PCG/Source/PCG/Private/PCGData.cpp
@@ -479,6 +479,15 @@ void FPCGDataCollection::Reset()
 	bCancelExecution = false;
 }
 
+int32 FPCGDataCollection::StripEmptyPointData()
+{
+	return TaggedData.RemoveAll([](const FPCGTaggedData& Data)
+	{
+		const UPCGPointData* PointData = Cast<UPCGPointData>(Data.Data);
+		return PointData && PointData->GetPoints().IsEmpty();
+	});
+}
+
 TArray<UPCGData*> UPCGDataFunctionLibrary::GetInputsByPredicate(const FPCGDataCollection& InCollection, TArray<FPCGTaggedData>& OutTaggedData, TFunctionRef<bool(const FPCGTaggedData&)> InPredicate)
 {
 	TArray<UPCGData*> Inputs;
diff --git a/Engine/Plugins/Experimental/PCG/Source/PCG/Public/PCGData.h b/Engine/Plugins/Experimental/PCG/Source/PCG/Public/PCGData.h
--- a/Engine/Plugins/Experimental/PCG/Source/PCG/Public/PCGData.h
+++ b/Engine/Plugins/Experimental/PCG/Source/PCG/Public/PCGData.h
@@ -179,6 +179,9 @@ struct PCG_API FPCGDataCollection
 	/** Cleans up the collection, but does not unroot any previously rooted data. */
 	void Reset();
 
+	/** Strips all empty point data from the collection. */
+	int32 StripEmptyPointData();
+
 	UPROPERTY(BlueprintReadWrite, EditAnywhere, Category = Data)
 	TArray<FPCGTaggedData> TaggedData;
 
diff --git a/Engine/Plugins/Experimental/PCG/Source/PCG/Public/PCGEngineSettings.h b/Engine/Plugins/Experimental/PCG/Source/PCG/Public/PCGEngineSettings.h
--- a/Engine/Plugins/Experimental/PCG/Source/PCG/Public/PCGEngineSettings.h
+++ b/Engine/Plugins/Experimental/PCG/Source/PCG/Public/PCGEngineSettings.h
@@ -35,9 +35,11 @@ public:
 	/** Time in MS between a landscape change and PCG refresh. Set it to 0 or negative value to disable the delay. */
 	UPROPERTY(EditAnywhere, Config, Category = "LandscapeTracking|Advanced", meta = (ConsoleVariable = "pcg.LandscapeRefreshTimeDelayMS"))
 	int32 LandscapeRefreshTimeDelayMS = 1000;
-#endif // WITH_EDITORONLY_DATA
-
 
+	/** Strips empty point data from being passed along as output. */
+	UPROPERTY(EditAnywhere, Config, Category = DataFlow, meta = (ConsoleVariable = "pcg.StripEmptyPointData"))
+	bool StripEmptyPointData = false;
+#endif // WITH_EDITORONLY_DATA
 	// Begin UDeveloperSettings Interface
 	virtual FName GetCategoryName() const override;
 #if WITH_EDITOR
