[PCG] Fix de-activate/re-activate leaving bGenerated flag in erroneous state

diff --git a/Engine/Plugins/Experimental/PCG/Source/PCG/Private/PCGComponent.cpp b/Engine/Plugins/Experimental/PCG/Source/PCG/Private/PCGComponent.cpp
--- a/Engine/Plugins/Experimental/PCG/Source/PCG/Private/PCGComponent.cpp
+++ b/Engine/Plugins/Experimental/PCG/Source/PCG/Private/PCGComponent.cpp
@@ -2028,7 +2028,10 @@ void UPCGComponent::OnRefresh(bool bForceRefresh)
 	if (!bActivated)
 	{
 		CleanupLocalImmediate(/*bRemoveComponents=*/true);
+
+		// Retain our generated state when going inactive, and mark bDirtyGenerated so that the component will re-generate upon re-activation (if necessary).
 		bGenerated = bWasGenerated;
+		bDirtyGenerated = bWasGenerated;
 	}
 	else
 	{
