[PCG] Add time slice to Attribute Partition

* Should fix an ensure happening when partition was not done within the allowed time.

#rnx

diff --git a/Engine/Plugins/Experimental/PCG/Source/PCG/Private/Elements/PCGMatchAndSetAttributes.cpp b/Engine/Plugins/Experimental/PCG/Source/PCG/Private/Elements/PCGMatchAndSetAttributes.cpp
--- a/Engine/Plugins/Experimental/PCG/Source/PCG/Private/Elements/PCGMatchAndSetAttributes.cpp
+++ b/Engine/Plugins/Experimental/PCG/Source/PCG/Private/Elements/PCGMatchAndSetAttributes.cpp
@@ -544,58 +544,121 @@ private:
 	TArray<TPair<const FPCGMetadataAttributeBase*, FPCGMetadataAttributeBase*>> AttributesToSet;
 };
 
-bool FPCGMatchAndSetAttributesElement::ExecuteInternal(FPCGContext* Context) const
+FPCGMatchAndSetAttributesExecutionState::~FPCGMatchAndSetAttributesExecutionState()
 {
-	TRACE_CPUPROFILER_EVENT_SCOPE(FPCGMatchAndSetAttributesElement::Execute);
+	if (Partition)
+	{
+		delete Partition;
+	}
+}
+
+bool FPCGMatchAndSetAttributesElement::PrepareDataInternal(FPCGContext* InContext) const
+{
+	TRACE_CPUPROFILER_EVENT_SCOPE(FPCGMatchAndSetAttributesElement::PrepareDataInternal);
 
-	const UPCGMatchAndSetAttributesSettings* Settings = Context->GetInputSettings<UPCGMatchAndSetAttributesSettings>();
+	const UPCGMatchAndSetAttributesSettings* Settings = InContext->GetInputSettings<UPCGMatchAndSetAttributesSettings>();
 	check(Settings);
 
-	TArray<FPCGTaggedData> Inputs = Context->InputData.GetInputsByPin(PCGPinConstants::DefaultInputLabel);
-	TArray<FPCGTaggedData> ParamDataInputs = Context->InputData.GetInputsByPin(PCGMatchAndSetAttributesConstants::MatchDataLabel);
-	TArray<FPCGTaggedData>& Outputs = Context->OutputData.TaggedData;
-	
-	const UPCGParamData* ParamData = nullptr;
+	FPCGMatchAndSetAttributesElement::ContextType* TimeSlicedContext = static_cast<FPCGMatchAndSetAttributesElement::ContextType*>(InContext);
+	check(TimeSlicedContext);
+
+	TArray<FPCGTaggedData> Inputs = InContext->InputData.GetInputsByPin(PCGPinConstants::DefaultInputLabel);
+	TArray<FPCGTaggedData> ParamDataInputs = InContext->InputData.GetInputsByPin(PCGMatchAndSetAttributesConstants::MatchDataLabel);
+
+	EPCGTimeSliceInitResult InitResult = TimeSlicedContext->InitializePerExecutionState([Settings, &ParamDataInputs](FPCGMatchAndSetAttributesElement::ContextType* Context, FPCGMatchAndSetAttributesExecutionState& OutState) -> EPCGTimeSliceInitResult
+	{
+		const UPCGParamData* ParamData = nullptr;
 
-	if (ParamDataInputs.Num() == 1)
+		if (ParamDataInputs.Num() == 1)
+		{
+			ParamData = Cast<const UPCGParamData>(ParamDataInputs[0].Data);
+		}
+
+		if (!ParamData)
+		{
+			PCGE_LOG_C(Warning, GraphAndLog, Context, LOCTEXT("NoMatchData", "Must have exactly one Attribute Set to match against"));
+			return EPCGTimeSliceInitResult::AbortExecution;
+		}
+
+		OutState.Partition = new FPCGMatchAndSetPartition(Context, Settings, Context->SourceComponent.Get(), ParamData);
+		check(OutState.Partition);
+
+		return OutState.Partition->Initialize() ? EPCGTimeSliceInitResult::Success : EPCGTimeSliceInitResult::AbortExecution;
+	});
+
+	if (InitResult == EPCGTimeSliceInitResult::AbortExecution)
 	{
-		ParamData = Cast<const UPCGParamData>(ParamDataInputs[0].Data);
+		PCGE_LOG_C(Warning, GraphAndLog, InContext, LOCTEXT("CouldNotInitializeExecutionState", "Could not initialize per-execution timeslice state data"));
+		return true;
 	}
 
-	if (!ParamData)
+	TimeSlicedContext->InitializePerIterationStates(Inputs.Num(), [&Inputs, InContext](FPCGMatchAndSetAttributesIterationState& OutState, const FPCGMatchAndSetAttributesExecutionState&, int32 Index) -> EPCGTimeSliceInitResult
 	{
-		PCGE_LOG(Warning, GraphAndLog, LOCTEXT("NoMatchData", "Must have exactly one Attribute Set to match against"));
+		OutState.InputData = Inputs[Index];
+		OutState.InPointData = Cast<UPCGPointData>(OutState.InputData.Data);
+		if (!OutState.InPointData)
+		{
+			PCGE_LOG_C(Error, GraphAndLog, InContext, FText::Format(LOCTEXT("InvalidInputDataType", "Input {0}: Input data must be of type Point"), FText::AsNumber(Index)));
+			return EPCGTimeSliceInitResult::NoOperation;
+		}
+
+		OutState.OutPointData = NewObject<UPCGPointData>();
+		OutState.OutPointData->InitializeFromData(OutState.InPointData);
+		OutState.OutPointData->GetMutablePoints().Reserve(OutState.InPointData->GetPoints().Num());
+
+		return EPCGTimeSliceInitResult::Success;
+	});
+
+	if (!TimeSlicedContext->DataIsPreparedForExecution())
+	{
+		PCGE_LOG_C(Warning, GraphAndLog, InContext, LOCTEXT("CouldNotInitializeStateData", "Could not initialize timeslice state data"));
 		return true;
 	}
 
-	FPCGMatchAndSetPartition Partition(Context, Settings, Context->SourceComponent.Get(), ParamData);
+	return true;
+}
+
+bool FPCGMatchAndSetAttributesElement::ExecuteInternal(FPCGContext* InContext) const
+{
+	TRACE_CPUPROFILER_EVENT_SCOPE(FPCGMatchAndSetAttributesElement::ExecuteInternal);
+
+	FPCGMatchAndSetAttributesElement::ContextType* TimeSlicedContext = static_cast<FPCGMatchAndSetAttributesElement::ContextType*>(InContext);
+	check(TimeSlicedContext);
+
+	TArray<FPCGTaggedData> Inputs = InContext->InputData.GetInputsByPin(PCGPinConstants::DefaultInputLabel);
 
-	if (!Partition.Initialize())
+	// Prepare data failed, no need to execute. Return an empty output
+	if (!TimeSlicedContext->DataIsPreparedForExecution())
 	{
+		TimeSlicedContext->OutputData.TaggedData.Empty();
 		return true;
 	}
 
-	for (const FPCGTaggedData& Input : Inputs)
+	// The context will iterate over per-iteration states and execute the lambda until it returns true
+	return ExecuteSlice(TimeSlicedContext, [&Inputs](FPCGMatchAndSetAttributesElement::ContextType* Context, const FPCGMatchAndSetAttributesExecutionState& ExecState, FPCGMatchAndSetAttributesIterationState& IterState, const uint32 IterationIndex) -> bool
 	{
-		FPCGTaggedData& Output = Outputs.Add_GetRef(Input);
+		const EPCGTimeSliceInitResult InitResult = Context->GetIterationStateResult(IterationIndex);
 
-		const UPCGPointData* InPointData = Cast<UPCGPointData>(Input.Data);
-		if (!InPointData)
+		// This iteration resulted in an early out for no sampling operation. Early out with a passthrough.
+		if (InitResult == EPCGTimeSliceInitResult::NoOperation)
 		{
-			PCGE_LOG(Error, GraphAndLog, LOCTEXT("InvalidInputDataType", "Input data must be of type Point"));
-			continue;
+			Context->OutputData.TaggedData.Add(Inputs[IterationIndex]);
+			return true;
 		}
 
-		UPCGPointData* OutPointData = NewObject<UPCGPointData>();
-		OutPointData->InitializeFromData(InPointData);
-		OutPointData->GetMutablePoints().Reserve(InPointData->GetPoints().Num());
-
-		int CurrentPointIndex = 0;
-		ensure(Partition.SelectPoints(*Context, InPointData, CurrentPointIndex, OutPointData));
-		Output.Data = OutPointData;
-	}
+		// It should be guaranteed to be a success at this point
+		check(InitResult == EPCGTimeSliceInitResult::Success);
 
-	return true;
+		// Run the execution until the time slice is finished
+		const bool bDone = ExecState.Partition->SelectPoints(*Context, IterState.InPointData, IterState.CurrentPointIndex, IterState.OutPointData);
+		if (bDone)
+		{
+			FPCGTaggedData& Output = Context->OutputData.TaggedData.Add_GetRef(Inputs[IterationIndex]);
+			Output.Data = IterState.OutPointData;
+		}
+		
+		return bDone;
+	});
 }
 
 #undef LOCTEXT_NAMESPACE
\ No newline at end of file
diff --git a/Engine/Plugins/Experimental/PCG/Source/PCG/Public/Elements/PCGMatchAndSetAttributes.h b/Engine/Plugins/Experimental/PCG/Source/PCG/Public/Elements/PCGMatchAndSetAttributes.h
--- a/Engine/Plugins/Experimental/PCG/Source/PCG/Public/Elements/PCGMatchAndSetAttributes.h
+++ b/Engine/Plugins/Experimental/PCG/Source/PCG/Public/Elements/PCGMatchAndSetAttributes.h
@@ -2,10 +2,31 @@
 
 #pragma once
 
+#include "PCGData.h"
 #include "PCGSettings.h"
+#include "Elements/PCGTimeSlicedElementBase.h"
 
 #include "PCGMatchAndSetAttributes.generated.h"
 
+// Defined in the cpp file
+class FPCGMatchAndSetPartition;
+class UPCGPointData;
+
+struct FPCGMatchAndSetAttributesExecutionState
+{
+	~FPCGMatchAndSetAttributesExecutionState();
+
+	FPCGMatchAndSetPartition* Partition = nullptr;
+};
+
+struct FPCGMatchAndSetAttributesIterationState
+{
+	int CurrentPointIndex = 0;
+	FPCGTaggedData InputData;
+	const UPCGPointData* InPointData = nullptr;
+	UPCGPointData* OutPointData = nullptr;
+};
+
 /** This class creates a PCG node that can match, select by weight or match & select by weight 
 * a 'matching' entry in a provided Attribute Set with multiple entries.
 * E.g. for a given point, if the point has the same specified attribute as the matching attribute in the attribute set,
@@ -66,8 +87,9 @@ public:
 	FName WeightAttribute = NAME_None;
 };
 
-class FPCGMatchAndSetAttributesElement : public IPCGElement
+class FPCGMatchAndSetAttributesElement : public TPCGTimeSlicedElementBase<FPCGMatchAndSetAttributesExecutionState, FPCGMatchAndSetAttributesIterationState>
 {
-public:
-	virtual bool ExecuteInternal(FPCGContext* Context) const override;
+protected:
+	virtual bool PrepareDataInternal(FPCGContext* InContext) const override;
+	virtual bool ExecuteInternal(FPCGContext* InContext) const override;
 };
\ No newline at end of file
diff --git a/Engine/Plugins/Experimental/PCG/Source/PCG/Public/Elements/PCGTimeSlicedElementBase.h b/Engine/Plugins/Experimental/PCG/Source/PCG/Public/Elements/PCGTimeSlicedElementBase.h
--- a/Engine/Plugins/Experimental/PCG/Source/PCG/Public/Elements/PCGTimeSlicedElementBase.h
+++ b/Engine/Plugins/Experimental/PCG/Source/PCG/Public/Elements/PCGTimeSlicedElementBase.h
@@ -154,7 +154,7 @@ public:
 
 	virtual FPCGContext* CreateContext() override { return new ContextType(); }
 
-	using ExecSignature = bool(ContextType* Context, const PerExecutionStateT& PerExecutionState, const PerIterationStateT& PerIterationState, const uint32 IterationIndex);
+	using ExecSignature = bool(ContextType* Context, const PerExecutionStateT& PerExecutionState, PerIterationStateT& PerIterationState, const uint32 IterationIndex);
 
 	/** Executes the delegate for every iteration. Will return false while still processing and true only when all tasks for all iterations are complete. */
 	bool ExecuteSlice(ContextType* Context, TFunctionRef<ExecSignature> ExecFunc) const;
