[PCG] Fix multiple tracking and partitioned issues

* Partition Actors are now unregistered on PostUnregisterAllComponents (because they could still be used while waiting to be destroyed)
* Don't keep a reference on non-created PAs when trying to generate a component (if the PA exists but is not loaded, we don't want to generate there)
* Don't refresh a component that is on a actor that gets deleted
* Don't refresh on actor unload (to mirror on actor load)
* Stop tracking Material Interfaces, as it was triggering way too many changes.
* All in all, this should fix errors happening because of refreshes comming when actors were in a bad state.

#rnx

diff --git a/Engine/Plugins/Experimental/PCG/Source/PCG/Private/Grid/PCGPartitionActor.cpp b/Engine/Plugins/Experimental/PCG/Source/PCG/Private/Grid/PCGPartitionActor.cpp
--- a/Engine/Plugins/Experimental/PCG/Source/PCG/Private/Grid/PCGPartitionActor.cpp
+++ b/Engine/Plugins/Experimental/PCG/Source/PCG/Private/Grid/PCGPartitionActor.cpp
@@ -169,6 +169,16 @@ void APCGPartitionActor::PostRegisterAllComponents()
 	}
 }
 
+void APCGPartitionActor::PostUnregisterAllComponents()
+{
+	if (!PCGHelpers::IsRuntimeOrPIE())
+	{
+		UnregisterPCG();
+	}
+
+	Super::PostUnregisterAllComponents();
+}
+
 void APCGPartitionActor::BeginPlay()
 {
 	// bIsRuntimeGenerated is not set yet, so we also need to check if the PA is transient.
diff --git a/Engine/Plugins/Experimental/PCG/Source/PCG/Private/Grid/PCGPartitionActor.h b/Engine/Plugins/Experimental/PCG/Source/PCG/Private/Grid/PCGPartitionActor.h
--- a/Engine/Plugins/Experimental/PCG/Source/PCG/Private/Grid/PCGPartitionActor.h
+++ b/Engine/Plugins/Experimental/PCG/Source/PCG/Private/Grid/PCGPartitionActor.h
@@ -38,6 +38,7 @@ public:
 	virtual void Destroyed() override;
 	virtual void GetActorBounds(bool bOnlyCollidingComponents, FVector& Origin, FVector& BoxExtent, bool bIncludeFromChildActors) const override;
 	virtual void PostRegisterAllComponents() override;
+	virtual void PostUnregisterAllComponents() override;
 #if WITH_EDITOR
 	virtual FBox GetStreamingBounds() const override;
 	virtual AActor* GetSceneOutlinerParent() const override;
diff --git a/Engine/Plugins/Experimental/PCG/Source/PCG/Private/PCGActorAndComponentMapping.cpp b/Engine/Plugins/Experimental/PCG/Source/PCG/Private/PCGActorAndComponentMapping.cpp
--- a/Engine/Plugins/Experimental/PCG/Source/PCG/Private/PCGActorAndComponentMapping.cpp
+++ b/Engine/Plugins/Experimental/PCG/Source/PCG/Private/PCGActorAndComponentMapping.cpp
@@ -22,6 +22,7 @@
 #include "LevelInstance/LevelInstanceInterface.h"
 #include "LevelInstance/LevelInstanceSubsystem.h"
 #include "LevelInstance/LevelInstanceEditorInstanceActor.h"
+#include "Materials/MaterialInterface.h"
 
 namespace PCGActorAndComponentMapping
 {
@@ -1359,7 +1360,8 @@ bool FPCGActorAndComponentMapping::UnregisterActor(AActor* InActor)
 
 void FPCGActorAndComponentMapping::OnActorUnloaded(AActor& InActor)
 {
-	OnActorDeleted(&InActor);
+	// Don't dirty on unload (to mirror the behavior in load)
+	OnActorDeleted_Internal(&InActor, /*bShouldDirty=*/false, /*LevelInstanceDepth=*/0);
 }
 
 void FPCGActorAndComponentMapping::OnActorDeleted(AActor* InActor)
@@ -1390,10 +1392,10 @@ void FPCGActorAndComponentMapping::OnActorDeleted(AActor* InActor)
 #endif
 
 	// Implementation note: since this is called only for actors directly in the current level, the depth here is 0.
-	OnActorDeleted_Internal(InActor, LevelInstanceDepth);
+	OnActorDeleted_Internal(InActor, /*bShouldDirty=*/true, LevelInstanceDepth);
 }
 
-void FPCGActorAndComponentMapping::OnActorDeleted_Internal(AActor* InActor, int32 LevelInstanceDepth)
+void FPCGActorAndComponentMapping::OnActorDeleted_Internal(AActor* InActor, bool bShouldDirty, int32 LevelInstanceDepth)
 {
 	TRACE_CPUPROFILER_EVENT_SCOPE(UPCGActorAndComponentMapping::OnActorDeleted);
 	check(InActor && PCGSubsystem && InActor->GetWorld() == PCGSubsystem->GetWorld());
@@ -1409,11 +1411,11 @@ void FPCGActorAndComponentMapping::OnActorDeleted_Internal(AActor* InActor, int3
 		{
 			if (ILevelInstanceInterface* LevelInstance = LevelInstanceSubsystem->GetLevelInstance(LevelInstanceId))
 			{
-				LevelInstanceSubsystem->ForEachActorInLevelInstance(LevelInstance, [this, LevelInstanceDepth, LevelInstanceEditorInstance](AActor* LevelActor)
+				LevelInstanceSubsystem->ForEachActorInLevelInstance(LevelInstance, [this, LevelInstanceDepth, LevelInstanceEditorInstance, bShouldDirty](AActor* LevelActor)
 				{
 					if (LevelActor != LevelInstanceEditorInstance)
 					{
-						OnActorDeleted_Internal(LevelActor, LevelInstanceDepth + 1);
+						OnActorDeleted_Internal(LevelActor, bShouldDirty, LevelInstanceDepth + 1);
 					}
 
 					return true;
@@ -1422,9 +1424,9 @@ void FPCGActorAndComponentMapping::OnActorDeleted_Internal(AActor* InActor, int3
 		}
 	}
 
-	PCGActorAndComponentMapping::PropagateToLevelInstanceActors(InActor, PCGSubsystem, [this, LevelInstanceDepth](AActor* LevelActor)
+	PCGActorAndComponentMapping::PropagateToLevelInstanceActors(InActor, PCGSubsystem, [this, LevelInstanceDepth, bShouldDirty](AActor* LevelActor)
 	{
-		OnActorDeleted_Internal(LevelActor, LevelInstanceDepth + 1);
+		OnActorDeleted_Internal(LevelActor, bShouldDirty, LevelInstanceDepth + 1);
 		return true;
 	});
 #endif // WITH_EDITOR
@@ -1434,8 +1436,11 @@ void FPCGActorAndComponentMapping::OnActorDeleted_Internal(AActor* InActor, int3
 		return;
 	}
 
-	// Notify all components that the actor has changed (was removed), but the Refresh will only happen AFTER the actor was actually removed from the world (because of delayed refresh).
-	OnActorChanged(InActor, /*bInHasMoved=*/ false, nullptr, LevelInstanceDepth);
+	if (bShouldDirty)
+	{
+		// Notify all components that the actor has changed (was removed), but the Refresh will only happen AFTER the actor was actually removed from the world (because of delayed refresh).
+		OnActorChanged(InActor, /*bInHasMoved=*/ false, nullptr, LevelInstanceDepth, /*bNoRefreshOnOwner=*/true);
+	}
 
 	// And then delete everything
 	UnregisterActor(InActor);
@@ -1588,7 +1593,7 @@ void FPCGActorAndComponentMapping::OnObjectPropertyChanged(UObject* InObject, FP
 	}
 }
 
-void FPCGActorAndComponentMapping::OnActorChanged(AActor* InActor, bool bInHasMoved, const UObject* InOriginatingChangeObject, int32 LevelInstanceDepth)
+void FPCGActorAndComponentMapping::OnActorChanged(AActor* InActor, bool bInHasMoved, const UObject* InOriginatingChangeObject, int32 LevelInstanceDepth, bool bNoRefreshOwner)
 {
 	TRACE_CPUPROFILER_EVENT_SCOPE(FPCGActorAndComponentMapping::OnActorChanged);
 
@@ -1632,14 +1637,15 @@ void FPCGActorAndComponentMapping::OnActorChanged(AActor* InActor, bool bInHasMo
 		// Then do an octree find to get all components that intersect with this actor.
 		// If the actor has moved, we also need to find components that intersected with it before
 		// We first do it for non-partitioned, then we do it for partitioned
-		auto UpdateNonPartitioned = [this, &DirtyComponents, InActor, CulledTrackedComponents, &RemovedTags, DirtyFlag, InOriginatingChangeObject](const FPCGComponentRef& ComponentRef) -> void
+		auto UpdateNonPartitioned = [this, &DirtyComponents, InActor, CulledTrackedComponents, &RemovedTags, DirtyFlag, InOriginatingChangeObject, bNoRefreshOwner](const FPCGComponentRef& ComponentRef) -> void
 		{
 			TRACE_CPUPROFILER_EVENT_SCOPE(FPCGActorAndComponentMapping::OnActorChanged::UpdateNonPartitioned);
 
-			// Don't dirty if the component was already dirtied, not tracked, or the origin of the change.
+			// Don't dirty if the component was already dirtied, not tracked, the origin of the change or its owner is the changed actor and we should not refresh.
 			if (DirtyComponents.Contains(ComponentRef.Component) || 
 				!CulledTrackedComponents->Contains(ComponentRef.Component) ||
-				InOriginatingChangeObject == ComponentRef.Component)
+				(InOriginatingChangeObject == ComponentRef.Component) ||
+				(bNoRefreshOwner && ComponentRef.Component->GetOwner() == InActor))
 			{
 				return;
 			}
@@ -1655,15 +1661,16 @@ void FPCGActorAndComponentMapping::OnActorChanged(AActor* InActor, bool bInHasMo
 
 		// For partitioned, we first need check if the original component intersect with the bounds, then forward the dirty call only to locals that intersect with the bounds.
 		// Note: CurrentActorBoundsPtr is passed by reference because it will be modified between lambda calls (cf comment above).
-		auto UpdatePartitioned = [this, &DirtyComponents, InActor, CulledTrackedComponents, &CurrentActorBoundsPtr, &RemovedTags, DirtyFlag, InOriginatingChangeObject](const FPCGComponentRef& ComponentRef)  -> void
+		auto UpdatePartitioned = [this, &DirtyComponents, InActor, CulledTrackedComponents, &CurrentActorBoundsPtr, &RemovedTags, DirtyFlag, InOriginatingChangeObject, bNoRefreshOwner](const FPCGComponentRef& ComponentRef)  -> void
 		{
 			TRACE_CPUPROFILER_EVENT_SCOPE(FPCGActorAndComponentMapping::OnActorChanged::UpdatePartitioned);
 
-			// Don't dirty if the component is not tracked, or the origin of the change.
+			// Don't dirty if the component is not tracked, the origin of the change or its owner is the changed actor and we should not refresh.
 			// We can "re-dirty" it because changes can impact different local components, from the same
 			// original component.
 			if (!CulledTrackedComponents->Contains(ComponentRef.Component) ||
-				InOriginatingChangeObject == ComponentRef.Component)
+				(InOriginatingChangeObject == ComponentRef.Component) ||
+				(bNoRefreshOwner && ComponentRef.Component->GetOwner() == InActor))
 			{
 				return;
 			}
@@ -1721,7 +1728,7 @@ void FPCGActorAndComponentMapping::OnActorChanged(AActor* InActor, bool bInHasMo
 
 		for (UPCGComponent* PCGComponent : *AlwaysTrackedComponents)
 		{
-			if (!PCGComponent || PCGComponent == InOriginatingChangeObject)
+			if (!PCGComponent || PCGComponent == InOriginatingChangeObject || (bNoRefreshOwner && PCGComponent->GetOwner() == InActor))
 			{
 				continue;
 			}
@@ -1766,7 +1773,7 @@ void FPCGActorAndComponentMapping::OnActorChanged(AActor* InActor, bool bInHasMo
 	// And refresh all dirtied components
 	for (UPCGComponent* Component : DirtyComponents)
 	{
-		if (Component)
+		if (Component && (!bNoRefreshOwner || Component->GetOwner() != InActor))
 		{
 			// When an object changes, we need to make sure that we don't trigger a refresh on PCG components that are "higher" in the
 			// level hierarchy, otherwise we will end up generating in the Level Instance level, which is wrong.
@@ -1841,7 +1848,8 @@ void FPCGActorAndComponentMapping::UpdateActorDependencies(AActor* InActor)
 	static const TArray<UClass*> ExcludedClasses =
 	{
 		UPCGComponent::StaticClass(),
-		UPCGGraphInterface::StaticClass()
+		UPCGGraphInterface::StaticClass(),
+		UMaterialInterface::StaticClass()
 	};
 
 	TSet<TObjectPtr<UObject>>& DependenciesSet = TrackedActorsToDependenciesMap.FindOrAdd(InActor);
diff --git a/Engine/Plugins/Experimental/PCG/Source/PCG/Private/PCGSubsystem.cpp b/Engine/Plugins/Experimental/PCG/Source/PCG/Private/PCGSubsystem.cpp
--- a/Engine/Plugins/Experimental/PCG/Source/PCG/Private/PCGSubsystem.cpp
+++ b/Engine/Plugins/Experimental/PCG/Source/PCG/Private/PCGSubsystem.cpp
@@ -797,8 +797,8 @@ namespace PCGSubsystem
 				if (IntersectedBounds.IsValid)
 				{
 					TSharedPtr<TSet<FWorldPartitionReference>> ActorReferences = MakeShared<TSet<FWorldPartitionReference>>();
-
-					auto PostCreation = [GridSize, &GridGuid](APartitionActor* Actor) { CastChecked<APCGPartitionActor>(Actor)->PostCreation(GridGuid); };
+					bool bWasCreated = false;
+					auto PostCreation = [GridSize, &GridGuid, &bWasCreated](APartitionActor* Actor) { CastChecked<APCGPartitionActor>(Actor)->PostCreation(GridGuid); bWasCreated = true;};
 
 					APCGPartitionActor* PCGActor = Cast<APCGPartitionActor>(PartitionSubsystem->GetActor(
 						APCGPartitionActor::StaticClass(),
@@ -841,9 +841,9 @@ namespace PCGSubsystem
 							PCGActor = Cast<APCGPartitionActor>(PCGActorDesc->GetActor());
 						}
 					}
-					else if(PCGActor)
+					// We still need to keep a reference on the PCG actor - note that newly created PCG actors will not have a reference here, but won't be unloaded
+					else if(PCGActor && bWasCreated)
 					{
-						// We still need to keep a reference on the PCG actor - note that newly created PCG actors will not have a reference here, but won't be unloaded
 						ActorReferences->Add(FWorldPartitionReference(World->GetWorldPartition(), PCGActor->GetActorGuid()));
 					}
 
diff --git a/Engine/Plugins/Experimental/PCG/Source/PCG/Public/PCGActorAndComponentMapping.h b/Engine/Plugins/Experimental/PCG/Source/PCG/Public/PCGActorAndComponentMapping.h
--- a/Engine/Plugins/Experimental/PCG/Source/PCG/Public/PCGActorAndComponentMapping.h
+++ b/Engine/Plugins/Experimental/PCG/Source/PCG/Public/PCGActorAndComponentMapping.h
@@ -138,7 +138,7 @@ private:
 	void OnActorAdded_Internal(AActor* InActor, bool bShouldDirty, int32 LevelInstanceDepth);
 	void OnActorDeleted(AActor* InActor);
 	void OnActorUnloaded(AActor& InActor);
-	void OnActorDeleted_Internal(AActor* InActor, int32 LevelInstanceDepth);
+	void OnActorDeleted_Internal(AActor* InActor, bool bShouldDirty, int32 LevelInstanceDepth);
 	void OnActorMoved(AActor* InActor);
 	void OnActorMoved_Internal(AActor* InActor, int32 LevelInstanceDepth);
 	void OnLandscapeChanged(ALandscapeProxy* InLandscape, const FLandscapeProxyComponentDataChangedParams& InChangeParams);
@@ -159,8 +159,9 @@ private:
 	/** Trigger an update when the actor changed. 
 	* Can specify if the actor has moved to also update components that were at its previous position.
 	* Can also specify an optional object, originating the change, to avoid re-dirtying a component if it was the origin.
+	* Another option when an actor is deleted/unload, don't refresh their components.
 	*/
-	void OnActorChanged(AActor* InActor, bool bInHasMoved, const UObject* InOriginatingChangeObject = nullptr, int32 LevelInstanceDepth = 0);
+	void OnActorChanged(AActor* InActor, bool bInHasMoved, const UObject* InOriginatingChangeObject = nullptr, int32 LevelInstanceDepth = 0, bool bNoRefreshOnOwner = false);
 
 	/** Update dependencies for a given tracked actor. */
 	void UpdateActorDependencies(AActor* InActor);
