[PCG] Fix graph cancellation adding nullptrs to the cancelled component set causing a crash.

This set is later used to cull active tasks so it might be a factor to our general cancellation issues. However the call site seems like it should normally crash if nullptrs are returned so its not clear if this fix will help the general problem.

#rnx

diff --git a/Engine/Plugins/Experimental/PCG/Source/PCG/Private/Graph/PCGGraphExecutor.cpp b/Engine/Plugins/Experimental/PCG/Source/PCG/Private/Graph/PCGGraphExecutor.cpp
--- a/Engine/Plugins/Experimental/PCG/Source/PCG/Private/Graph/PCGGraphExecutor.cpp
+++ b/Engine/Plugins/Experimental/PCG/Source/PCG/Private/Graph/PCGGraphExecutor.cpp
@@ -303,9 +303,10 @@ TSet<UPCGComponent*> FPCGGraphExecutor::Cancel(TFunctionRef<bool(TWeakObjectPtr<
 
 				if (bContainsDependency)
 				{
-					if (!CancelledComponents.Contains(ScheduledTask.SourceComponent.Get()))
+					UPCGComponent* TaskComponent = ScheduledTask.SourceComponent.Get();
+					if (TaskComponent && !CancelledComponents.Contains(TaskComponent))
 					{
-						CancelledComponents.Add(ScheduledTask.SourceComponent.Get());
+						CancelledComponents.Add(TaskComponent);
 						bStableCancellationSet = false;
 					}
 
diff --git a/Engine/Plugins/Experimental/PCG/Source/PCG/Private/PCGSubsystem.cpp b/Engine/Plugins/Experimental/PCG/Source/PCG/Private/PCGSubsystem.cpp
--- a/Engine/Plugins/Experimental/PCG/Source/PCG/Private/PCGSubsystem.cpp
+++ b/Engine/Plugins/Experimental/PCG/Source/PCG/Private/PCGSubsystem.cpp
@@ -630,7 +630,10 @@ void UPCGSubsystem::CancelGeneration(UPCGComponent* Component)
 	TArray<UPCGComponent*> CancelledComponents = GraphExecutor->Cancel(Component);
 	for (UPCGComponent* CancelledComponent : CancelledComponents)
 	{
-		CancelledComponent->OnProcessGraphAborted(/*bQuiet=*/true);
+		if (ensure(CancelledComponent))
+		{
+			CancelledComponent->OnProcessGraphAborted(/*bQuiet=*/true);
+		}
 	}	
 }
 
@@ -646,7 +649,10 @@ void UPCGSubsystem::CancelGeneration(UPCGGraph* Graph)
 	TArray<UPCGComponent*> CancelledComponents = GraphExecutor->Cancel(Graph);
 	for (UPCGComponent* CancelledComponent : CancelledComponents)
 	{
-		CancelledComponent->OnProcessGraphAborted(/*bQuiet=*/true);
+		if (ensure(CancelledComponent))
+		{
+			CancelledComponent->OnProcessGraphAborted(/*bQuiet=*/true);
+		}
 	}
 }
 
@@ -655,9 +661,12 @@ void UPCGSubsystem::CancelAllGeneration()
 	check(GraphExecutor);
 
 	TArray<UPCGComponent*> CancelledComponents = GraphExecutor->CancelAll();
-	for (UPCGComponent* Component : CancelledComponents)
+	for (UPCGComponent* CancelledComponent : CancelledComponents)
 	{
-		Component->OnProcessGraphAborted(/*bQuiet=*/true);
+		if (ensure(CancelledComponent))
+		{
+			CancelledComponent->OnProcessGraphAborted(/*bQuiet=*/true);
+		}
 	}
 }
 
