[PCG] Improve on PointFilter to become AttributeFilter

* Can now filter attribute sets
* Rename PointFilter(Range) to AttributeFilter(Range)
* Rename AttributeFilter to AttributeFilterNames (because of collision)
* Add GetAliasNodeName on the settings to expose the same node with different names (for deprecation)
* Add support for constructible for threshold (now allow to compare density (float) with a double constant)
* Add unit tests

diff --git a/Engine/Plugins/Experimental/PCG/Config/BasePCG.ini b/Engine/Plugins/Experimental/PCG/Config/BasePCG.ini
--- a/Engine/Plugins/Experimental/PCG/Config/BasePCG.ini
+++ b/Engine/Plugins/Experimental/PCG/Config/BasePCG.ini
@@ -63,5 +63,11 @@
 +ClassRedirects=(OldName="PCGInstancePackerByRegex",NewName="/Script/PCG.PCGInstanceDataPackerByRegex")
 +ClassRedirects=(OldName="PCGIntersectionSettings",NewName="/Script/PCG.PCGInnerIntersectionSettings")
 +ClassRedirects=(OldName="PCGCreateAttributeSettings",NewName="/Script/PCG.PCGAddAttributeSettings")
++ClassRedirects=(OldName="PCGPointFilterSettings",NewName="/Script/PCG.PCGAttributeFilteringSettings")
++ClassRedirects=(OldName="PCGPointFilterRangeSettings",NewName="/Script/PCG.PCGAttributeFilteringRangeSettings")
++ClassRedirects=(OldName="PCGAttributeFilterSettings",NewName="/Script/PCG.PCGAttributeFilterNamesSettings")
+
++StructRedirects=(OldName="PCGPointFilterThresholdSettings",NewName="/Script/PCG.PCGAttributeFilterThresholdSettings")
 
 +EnumRedirects=(OldName="EPCGDensityNoiseMode",NewName="EPCGAttributeNoiseMode")
++EnumRedirects=(OldName="EPCGPointFilterOperator",NewName="EPCGAttributeFilterOperator")
diff --git a/Engine/Plugins/Experimental/PCG/Source/PCG/Private/Elements/PCGAttributeFilter.cpp b/Engine/Plugins/Experimental/PCG/Source/PCG/Private/Elements/PCGAttributeFilter.cpp
new file mode 100644
--- /dev/null
+++ b/Engine/Plugins/Experimental/PCG/Source/PCG/Private/Elements/PCGAttributeFilter.cpp
@@ -0,0 +1,770 @@
+// Copyright Epic Games, Inc. All Rights Reserved.
+
+#include "Elements/PCGAttributeFilter.h"
+
+#include "PCGContext.h"
+#include "PCGCustomVersion.h"
+#include "PCGParamData.h"
+#include "PCGPin.h"
+#include "Data/PCGPointData.h"
+#include "Data/PCGSpatialData.h"
+#include "Metadata/Accessors/IPCGAttributeAccessorTpl.h"
+#include "Metadata/Accessors/PCGAttributeAccessorHelpers.h"
+#include "Metadata/Accessors/PCGCustomAccessor.h"
+
+#define LOCTEXT_NAMESPACE "PCGPointFilterElement"
+
+namespace PCGAttributeFilterConstants
+{
+	const FName DataToFilterLabel = TEXT("In");
+	const FName FilterLabel = TEXT("Filter");
+	const FName FilterMinLabel = TEXT("FilterMin");
+	const FName FilterMaxLabel = TEXT("FilterMax");
+	const FName InFilterLabel = TEXT("InsideFilter");
+	const FName OutFilterLabel = TEXT("OutsideFilter");
+
+	constexpr int32 ChunkSize = 256;
+
+#if WITH_EDITOR
+	const FText InputPinTooltip = LOCTEXT("InputPinTooltip", "This pin accepts Point data and Attribute Sets. Spatial data will be collapsed to point data.");
+
+	const FText FilterPinTooltip = LOCTEXT("FilterPinTooltip", "This pin accepts Statial data and Attribute Sets. If the data is Spatial, it will automatically sample input points in it. "
+		"If it is points, it will sample if \"Spatial Query\" is enabled, otherwise points number need to match with input.");
+#endif // WITH_EDITOR
+}
+
+namespace PCGAttributeFilterHelpers
+{
+	template <typename T>
+	bool ApplyCompare(const T& Input1, const T& Input2, EPCGAttributeFilterOperator Operation)
+	{
+		if (Operation == EPCGAttributeFilterOperator::Equal)
+		{
+			return PCG::Private::MetadataTraits<T>::Equal(Input1, Input2);
+		}
+		else if (Operation == EPCGAttributeFilterOperator::NotEqual)
+		{
+			return !PCG::Private::MetadataTraits<T>::Equal(Input1, Input2);
+		}
+
+		if constexpr (PCG::Private::MetadataTraits<T>::CanCompare)
+		{
+			switch (Operation)
+			{
+			case EPCGAttributeFilterOperator::Greater:
+				return PCG::Private::MetadataTraits<T>::Greater(Input1, Input2);
+			case EPCGAttributeFilterOperator::GreaterOrEqual:
+				return PCG::Private::MetadataTraits<T>::GreaterOrEqual(Input1, Input2);
+			case EPCGAttributeFilterOperator::Lesser:
+				return PCG::Private::MetadataTraits<T>::Less(Input1, Input2);
+			case EPCGAttributeFilterOperator::LesserOrEqual:
+				return PCG::Private::MetadataTraits<T>::LessOrEqual(Input1, Input2);
+			default:
+				break;
+			}
+		}
+
+		if constexpr (PCG::Private::MetadataTraits<T>::CanSearchString)
+		{
+			switch (Operation)
+			{
+			case EPCGAttributeFilterOperator::Substring:
+				return PCG::Private::MetadataTraits<T>::Substring(Input1, Input2);
+			case EPCGAttributeFilterOperator::Matches:
+				return PCG::Private::MetadataTraits<T>::Matches(Input1, Input2);
+			default:
+				break;
+			}
+		}
+
+		return false;
+	}
+
+	template <typename T>
+	bool ApplyRange(const T& Input, const T& InMin, const T& InMax, bool bMinIncluded, bool bMaxIncluded)
+	{
+		if constexpr (PCG::Private::MetadataTraits<T>::CanCompare)
+		{
+			return (bMinIncluded ? PCG::Private::MetadataTraits<T>::GreaterOrEqual(Input, InMin) : PCG::Private::MetadataTraits<T>::Greater(Input, InMin)) &&
+				(bMaxIncluded ? PCG::Private::MetadataTraits<T>::LessOrEqual(Input, InMax) : PCG::Private::MetadataTraits<T>::Less(Input, InMax));
+		}
+		else
+		{
+			return false;
+		}
+	}
+
+	struct ThresholdInfo
+	{
+		TUniquePtr<const IPCGAttributeAccessor> ThresholdAccessor;
+		TUniquePtr<const IPCGAttributeAccessorKeys> ThresholdKeys;
+		bool bUseInputDataForThreshold = false;
+		bool bUseSpatialQuery = false;
+
+		UPCGPointData* ThresholdPointData = nullptr;
+		const UPCGSpatialData* ThresholdSpatialData = nullptr;
+	};
+
+	bool InitialPrepareThresholdInfo(FPCGContext* InContext, TArray<FPCGTaggedData> FilterData, const FPCGAttributeFilterThresholdSettings& ThresholdSettings, ThresholdInfo& OutThresholdInfo)
+	{
+		if (ThresholdSettings.bUseConstantThreshold)
+		{
+			auto ConstantThreshold = [&OutThresholdInfo](auto&& Value)
+			{
+				using ConstantType = std::decay_t<decltype(Value)>;
+
+				OutThresholdInfo.ThresholdAccessor = MakeUnique<FPCGConstantValueAccessor<ConstantType>>(std::forward<decltype(Value)>(Value));
+				// Dummy keys
+				OutThresholdInfo.ThresholdKeys = MakeUnique<FPCGAttributeAccessorKeysSingleObjectPtr<void>>();
+			};
+
+			ThresholdSettings.AttributeTypes.Dispatcher(ConstantThreshold);
+		}
+		else
+		{
+			if (!FilterData.IsEmpty())
+			{
+				const UPCGData* ThresholdData = FilterData[0].Data;
+
+				if (const UPCGSpatialData* ThresholdSpatialData = Cast<const UPCGSpatialData>(ThresholdData))
+				{
+					// If the threshold is spatial or points (and spatial query is enabled), we'll use spatial query (meaning we'll have to sample points).
+					// Don't create an accessor yet (ThresholdData = nullptr), it will be created further down.
+					OutThresholdInfo.ThresholdSpatialData = ThresholdSpatialData;
+					if (!ThresholdSpatialData->IsA<UPCGPointData>() || ThresholdSettings.bUseSpatialQuery)
+					{
+						OutThresholdInfo.bUseSpatialQuery = true;
+						ThresholdData = nullptr;
+					}
+				}
+
+				if (ThresholdData)
+				{
+					FPCGAttributePropertyInputSelector ThresholdSelector = ThresholdSettings.ThresholdAttribute.CopyAndFixLast(ThresholdData);
+					OutThresholdInfo.ThresholdAccessor = PCGAttributeAccessorHelpers::CreateConstAccessor(ThresholdData, ThresholdSelector);
+					OutThresholdInfo.ThresholdKeys = PCGAttributeAccessorHelpers::CreateConstKeys(ThresholdData, ThresholdSelector);
+				}
+			}
+			else
+			{
+				OutThresholdInfo.bUseInputDataForThreshold = true;
+			}
+		}
+
+		return true;
+	}
+
+	bool PrepareThresholdInfoFromInput(FPCGContext* InContext, const UPCGData* InputData, const int32 NumInput, const FPCGAttributeFilterThresholdSettings& ThresholdSettings, ThresholdInfo& InOutThresholdInfo, int16 TargetType, bool bCheckCompare, bool bCheckStringSearch, const ThresholdInfo* OtherInfo = nullptr)
+	{
+		check(InContext && InputData);
+
+		if (InOutThresholdInfo.bUseInputDataForThreshold)
+		{
+			// If we have no threshold accessor, we use the same data as input
+			FPCGAttributePropertyInputSelector ThresholdSelector = ThresholdSettings.ThresholdAttribute.CopyAndFixLast(InputData);
+			InOutThresholdInfo.ThresholdAccessor = PCGAttributeAccessorHelpers::CreateConstAccessor(InputData, ThresholdSelector);
+			InOutThresholdInfo.ThresholdKeys = PCGAttributeAccessorHelpers::CreateConstKeys(InputData, ThresholdSelector);
+		}
+		else if (InOutThresholdInfo.ThresholdSpatialData != nullptr && InOutThresholdInfo.bUseSpatialQuery)
+		{
+			// Don't do 2 spatial query if we are iterating on the same input
+			if (OtherInfo != nullptr && OtherInfo->ThresholdSpatialData == InOutThresholdInfo.ThresholdSpatialData)
+			{
+				InOutThresholdInfo.ThresholdPointData = OtherInfo->ThresholdPointData;
+			}
+			else
+			{
+				// Reset the point data and reserving some points
+				// No need to reserve the full number of points, since we'll go by chunk
+				// Only allocate the chunk size
+				InOutThresholdInfo.ThresholdPointData = NewObject<UPCGPointData>();
+				InOutThresholdInfo.ThresholdPointData->InitializeFromData(InOutThresholdInfo.ThresholdSpatialData);
+				InOutThresholdInfo.ThresholdPointData->GetMutablePoints().SetNum(PCGAttributeFilterConstants::ChunkSize);
+			}
+
+			// Accessor will be valid, but keys will point to default points. But since it is a view, it will be updated when we sample the points.
+			FPCGAttributePropertyInputSelector ThresholdSelector = ThresholdSettings.ThresholdAttribute.CopyAndFixLast(InOutThresholdInfo.ThresholdPointData);
+			InOutThresholdInfo.ThresholdAccessor = PCGAttributeAccessorHelpers::CreateConstAccessor(InOutThresholdInfo.ThresholdPointData, ThresholdSelector);
+			InOutThresholdInfo.ThresholdKeys = PCGAttributeAccessorHelpers::CreateConstKeys(InOutThresholdInfo.ThresholdPointData, ThresholdSelector);
+		}
+
+		if (!InOutThresholdInfo.ThresholdAccessor.IsValid() || !InOutThresholdInfo.ThresholdKeys.IsValid())
+		{
+			PCGE_LOG_C(Warning, GraphAndLog, InContext, FText::Format(LOCTEXT("AttributeMissingForFilter", "DataToFilter does not have '{0}' threshold attribute/property"), FText::FromName(ThresholdSettings.ThresholdAttribute.GetName())));
+			return false;
+		}
+
+		// And also validate that types are comparable/constructible. Do it all at once for a single dispatch.
+		bool bCanCompare = true;
+		bool bCanSearchString = true;
+		bool bIsConstructible = true;
+		PCGMetadataAttribute::CallbackWithRightType(InOutThresholdInfo.ThresholdAccessor->GetUnderlyingType(), [&bCanCompare, &bCanSearchString, &bIsConstructible, TargetType](auto ThresholdValue)
+		{
+			bCanCompare = PCG::Private::MetadataTraits<decltype(ThresholdValue)>::CanCompare;
+			bCanSearchString = PCG::Private::MetadataTraits<decltype(ThresholdValue)>::CanSearchString;
+			bIsConstructible = PCGMetadataAttribute::CallbackWithRightType(TargetType, [](auto TargetValue) { return std::is_constructible_v<decltype(TargetValue),decltype(ThresholdValue)>;});
+		});
+
+		// Comparison between threshold and target data needs to be of the same type. So we have to make sure that we can
+		// request target type from threshold type. ie. We need to make sure we can broadcast threshold type to target type, or construct a target type from a threshold type.
+		// For example: if target is double but threshold is int32, we can broadcast int32 to double, to compare a double with a double.
+		if (!PCG::Private::IsBroadcastable(InOutThresholdInfo.ThresholdAccessor->GetUnderlyingType(), TargetType) && !bIsConstructible)
+		{
+			UEnum* PCGDataTypeEnum = StaticEnum<EPCGMetadataTypes>();
+			FText ThresholdTypeName = FText::FromString(PCGDataTypeEnum ? PCGDataTypeEnum->GetNameStringByValue(InOutThresholdInfo.ThresholdAccessor->GetUnderlyingType()) : FString(TEXT("Unknown")));
+			FText InputTypeName = FText::FromString(PCGDataTypeEnum ? PCGDataTypeEnum->GetNameStringByValue(TargetType) : FString(TEXT("Unknown")));
+			PCGE_LOG_C(Warning, GraphAndLog, InContext, FText::Format(LOCTEXT("TypeConversionFailed", "Cannot use threshold type '{0}' on input target type '{1}'"),
+				ThresholdTypeName,
+				InputTypeName));
+			return false;
+		}
+
+		if (bCheckCompare && !bCanCompare)
+		{
+			PCGE_LOG_C(Warning, GraphAndLog, InContext, LOCTEXT("TypeComparisonFailed", "Cannot compare target type"));
+			return false;
+		}
+
+		if (bCheckStringSearch && !bCanSearchString)
+		{
+			PCGE_LOG_C(Warning, GraphAndLog, InContext, LOCTEXT("TypeStringSearchFailed", "Cannot perform string operations on target type"));
+			return false;
+		}
+
+		// Check that if we have points as threshold, that the point data has the same number of point that the input data
+		if (InOutThresholdInfo.ThresholdSpatialData != nullptr && !InOutThresholdInfo.bUseSpatialQuery)
+		{
+			if (InOutThresholdInfo.ThresholdKeys->GetNum() != NumInput)
+			{
+				PCGE_LOG_C(Warning, GraphAndLog, InContext, FText::Format(LOCTEXT("InvalidNumberOfThresholdPoints", "Threshold point data doesn't have the same number of elements ({0}) than the input data ({1})."), InOutThresholdInfo.ThresholdKeys->GetNum(), NumInput));
+				return false;
+			}
+		}
+
+		return true;
+	}
+}
+
+#if WITH_EDITOR
+void FPCGAttributeFilterThresholdSettings::OnPostLoad()
+{
+	AttributeTypes.OnPostLoad();
+}
+#endif
+
+EPCGDataType UPCGAttributeFilteringSettings::GetCurrentPinTypes(const UPCGPin* InPin) const
+{
+	check(InPin);
+	if (!InPin->IsOutputPin())
+	{
+		return Super::GetCurrentPinTypes(InPin);
+	}
+
+	// Output pin narrows to union of inputs on first pin
+	const EPCGDataType InputTypeUnion = GetTypeUnionOfIncidentEdges(PCGPinConstants::DefaultInputLabel);
+
+	// Spatial is collapsed into points
+	if ((InputTypeUnion & EPCGDataType::Spatial) == InputTypeUnion)
+	{
+		return EPCGDataType::Point;
+	}
+	else
+	{
+		return (InputTypeUnion != EPCGDataType::None) ? InputTypeUnion : EPCGDataType::Any;
+	}
+}
+
+TArray<FPCGPinProperties> UPCGAttributeFilteringSettings::InputPinProperties() const
+{
+	TArray<FPCGPinProperties> PinProperties;
+	PinProperties.Emplace(PCGAttributeFilterConstants::DataToFilterLabel, EPCGDataType::Any);
+
+#if WITH_EDITOR
+	PinProperties.Last().Tooltip = PCGAttributeFilterConstants::InputPinTooltip;
+#endif //WITH_EDITOR
+
+	if (!bUseConstantThreshold)
+	{
+		PinProperties.Emplace(PCGAttributeFilterConstants::FilterLabel, EPCGDataType::Any, /*bInAllowMultipleConnections=*/ false);
+
+#if WITH_EDITOR
+		PinProperties.Last().Tooltip = PCGAttributeFilterConstants::FilterPinTooltip;
+#endif // WITH_EDITOR
+	}
+
+	return PinProperties;
+}
+
+TArray<FPCGPinProperties> UPCGAttributeFilteringSettings::OutputPinProperties() const
+{
+	TArray<FPCGPinProperties> PinProperties;
+	PinProperties.Emplace(PCGAttributeFilterConstants::InFilterLabel, EPCGDataType::Any);
+	PinProperties.Emplace(PCGAttributeFilterConstants::OutFilterLabel, EPCGDataType::Any);
+
+	return PinProperties;
+}
+
+FPCGElementPtr UPCGAttributeFilteringSettings::CreateElement() const
+{
+	return MakeShared<FPCGAttributeFilterElement>();
+}
+
+UPCGAttributeFilteringSettings::UPCGAttributeFilteringSettings()
+	: UPCGSettings()
+{
+	// Previous default object was: density for both selectors
+	// Recreate the same default
+	TargetAttribute.SetPointProperty(EPCGPointProperties::Density);
+	ThresholdAttribute.SetPointProperty(EPCGPointProperties::Density);
+}
+
+void UPCGAttributeFilteringSettings::PostLoad()
+{
+	Super::PostLoad();
+
+#if WITH_EDITOR
+	AttributeTypes.OnPostLoad();
+#endif // WITH_EDITOR
+}
+
+#if WITH_EDITOR
+void UPCGAttributeFilteringRangeSettings::PostLoad()
+{
+	Super::PostLoad();
+
+	MinThreshold.OnPostLoad();
+	MaxThreshold.OnPostLoad();
+}
+#endif
+
+FPCGElementPtr UPCGAttributeFilteringRangeSettings::CreateElement() const
+{
+	return MakeShared<FPCGAttributeFilterRangeElement>();
+}
+
+UPCGAttributeFilteringRangeSettings::UPCGAttributeFilteringRangeSettings()
+	: UPCGSettings()
+{
+	// Previous default object was: density for both selectors
+	// Recreate the same default
+	TargetAttribute.SetPointProperty(EPCGPointProperties::Density);
+	MinThreshold.ThresholdAttribute.SetPointProperty(EPCGPointProperties::Density);
+	MaxThreshold.ThresholdAttribute.SetPointProperty(EPCGPointProperties::Density);
+}
+
+
+EPCGDataType UPCGAttributeFilteringRangeSettings::GetCurrentPinTypes(const UPCGPin* InPin) const
+{
+	check(InPin);
+	if (!InPin->IsOutputPin())
+	{
+		return Super::GetCurrentPinTypes(InPin);
+	}
+
+	// Output pin narrows to union of inputs on first pin
+	const EPCGDataType InputTypeUnion = GetTypeUnionOfIncidentEdges(PCGPinConstants::DefaultInputLabel);
+
+	// Spatial is collapsed into points
+	if ((InputTypeUnion & EPCGDataType::Spatial) == InputTypeUnion)
+	{
+		return EPCGDataType::Point;
+	}
+	else
+	{
+		return (InputTypeUnion != EPCGDataType::None) ? InputTypeUnion : EPCGDataType::Any;
+	}
+}
+
+TArray<FPCGPinProperties> UPCGAttributeFilteringRangeSettings::InputPinProperties() const
+{
+	TArray<FPCGPinProperties> PinProperties;
+	PinProperties.Emplace(PCGAttributeFilterConstants::DataToFilterLabel, EPCGDataType::Any);
+
+#if WITH_EDITOR
+	PinProperties.Last().Tooltip = PCGAttributeFilterConstants::InputPinTooltip;
+#endif //WITH_EDITOR
+
+	if (!MinThreshold.bUseConstantThreshold)
+	{
+		PinProperties.Emplace(PCGAttributeFilterConstants::FilterMinLabel, EPCGDataType::Any, /*bInAllowMultipleConnections=*/ false);
+
+#if WITH_EDITOR
+		PinProperties.Last().Tooltip = PCGAttributeFilterConstants::FilterPinTooltip;
+#endif // WITH_EDITOR
+	}
+
+	if (!MaxThreshold.bUseConstantThreshold)
+	{
+		PinProperties.Emplace(PCGAttributeFilterConstants::FilterMaxLabel, EPCGDataType::Any, /*bInAllowMultipleConnections=*/ false);
+
+#if WITH_EDITOR
+		PinProperties.Last().Tooltip = PCGAttributeFilterConstants::FilterPinTooltip;
+#endif // WITH_EDITOR
+	}
+
+	return PinProperties;
+}
+
+TArray<FPCGPinProperties> UPCGAttributeFilteringRangeSettings::OutputPinProperties() const
+{
+	TArray<FPCGPinProperties> PinProperties;
+	PinProperties.Emplace(PCGAttributeFilterConstants::InFilterLabel, EPCGDataType::Any);
+	PinProperties.Emplace(PCGAttributeFilterConstants::OutFilterLabel, EPCGDataType::Any);
+
+	return PinProperties;
+}
+
+bool FPCGAttributeFilterElementBase::DoFiltering(FPCGContext* Context, EPCGAttributeFilterOperator InOperation, const FPCGAttributePropertyInputSelector& InTargetAttribute, const FPCGAttributeFilterThresholdSettings& FirstThreshold, const FPCGAttributeFilterThresholdSettings* SecondThreshold) const
+{
+	TRACE_CPUPROFILER_EVENT_SCOPE(FPCGAttributeFilterElementBase::DoFiltering);
+	check(Context);
+
+	struct FOperationData
+	{
+		const TArray<FPCGPoint>* OriginalPoints = nullptr;
+		TArray<FPCGPoint>* InFilterPoints = nullptr;
+		TArray<FPCGPoint>* OutFilterPoints = nullptr;
+
+		const UPCGMetadata* OriginalMetadata = nullptr;
+		UPCGMetadata* InFilterMetadata = nullptr;
+		UPCGMetadata* OutFilterMetadata = nullptr;
+
+		bool bIsInputPointData = false;
+	} OperationData;
+
+	TArray<FPCGTaggedData> DataToFilter = Context->InputData.GetInputsByPin(PCGAttributeFilterConstants::DataToFilterLabel);
+	TArray<FPCGTaggedData>& Outputs = Context->OutputData.TaggedData;
+	TArray<FPCGTaggedData> FirstFilterData = SecondThreshold ? Context->InputData.GetInputsByPin(PCGAttributeFilterConstants::FilterMinLabel) : Context->InputData.GetInputsByPin(PCGAttributeFilterConstants::FilterLabel);
+	TArray<FPCGTaggedData> SecondFilterData = SecondThreshold ? Context->InputData.GetInputsByPin(PCGAttributeFilterConstants::FilterMaxLabel) : TArray<FPCGTaggedData>{};
+
+	const EPCGAttributeFilterOperator Operator = InOperation;
+
+	// If there is no input, do nothing
+	if (DataToFilter.IsEmpty())
+	{
+		return true;
+	}
+
+	// Only support second threshold with the InRange Operation.
+	// We can't have a second threshold if it is not InRange
+	if (!ensure((SecondThreshold != nullptr) == (InOperation == EPCGAttributeFilterOperator::InRange)))
+	{
+		return true;
+	}
+
+	PCGAttributeFilterHelpers::ThresholdInfo FirstThresholdInfo;
+	PCGAttributeFilterHelpers::ThresholdInfo SecondThresholdInfo;
+
+	if (!PCGAttributeFilterHelpers::InitialPrepareThresholdInfo(Context, FirstFilterData, FirstThreshold, FirstThresholdInfo))
+	{
+		return true;
+	}
+
+	if (SecondThreshold && !PCGAttributeFilterHelpers::InitialPrepareThresholdInfo(Context, SecondFilterData, *SecondThreshold, SecondThresholdInfo))
+	{
+		return true;
+	}
+
+	for (const FPCGTaggedData& Input : DataToFilter)
+	{
+		const UPCGData* OriginalData = Input.Data;
+		if (!OriginalData)
+		{
+			continue;
+		}
+
+		if (const UPCGSpatialData* SpatialInput = Cast<const UPCGSpatialData>(OriginalData))
+		{
+			const UPCGPointData* OriginalPointData = SpatialInput->ToPointData(Context);
+			if (!OriginalPointData)
+			{
+				PCGE_LOG(Error, GraphAndLog, LOCTEXT("NoPointDataInInput", "Unable to get point data from input"));
+				continue;
+			}
+
+			OperationData.bIsInputPointData = true;
+			OriginalData = OriginalPointData;
+		}
+		else if (OriginalData->IsA<UPCGParamData>())
+		{
+			// Disable Spatial Queries
+			FirstThresholdInfo.bUseSpatialQuery = false;
+			SecondThresholdInfo.bUseSpatialQuery = false;
+		}
+		else
+		{
+			PCGE_LOG(Error, GraphAndLog, LOCTEXT("InvalidInput", "Input is not a point data nor an attribute set. Unsupported."));
+			continue;
+		}
+
+		// Helper lambdas to fail nicely and forward input to in/out filter pin
+		// If there is a problem with threshold -> forward to InFilter
+		auto ForwardInputToInFilterPin = [&Outputs, Input]()
+		{
+			FPCGTaggedData& InFilterOutput = Outputs.Add_GetRef(Input);
+			InFilterOutput.Pin = PCGAttributeFilterConstants::InFilterLabel;
+		};
+
+		// If there is a problem with target -> forward to OutFilter
+		auto ForwardInputToOutFilterPin = [&Outputs, Input]()
+		{
+			FPCGTaggedData& OutFilterOutput = Outputs.Add_GetRef(Input);
+			OutFilterOutput.Pin = PCGAttributeFilterConstants::OutFilterLabel;
+		};
+
+		FPCGAttributePropertyInputSelector TargetAttribute = InTargetAttribute.CopyAndFixLast(OriginalData);
+		TUniquePtr<const IPCGAttributeAccessor> TargetAccessor = PCGAttributeAccessorHelpers::CreateConstAccessor(OriginalData, TargetAttribute);
+		TUniquePtr<const IPCGAttributeAccessorKeys> TargetKeys = PCGAttributeAccessorHelpers::CreateConstKeys(OriginalData, TargetAttribute);
+
+		if (!TargetAccessor.IsValid() || !TargetKeys.IsValid())
+		{
+			PCGE_LOG(Warning, GraphAndLog, FText::Format(LOCTEXT("TargetMissingAttribute", "TargetData doesn't have target attribute/property '{0}'"), FText::FromName(TargetAttribute.GetName())));
+			ForwardInputToOutFilterPin();
+			continue;
+		}
+
+		const int16 TargetType = TargetAccessor->GetUnderlyingType();
+		const bool bCheckStringSearch = (Operator == EPCGAttributeFilterOperator::Substring || Operator == EPCGAttributeFilterOperator::Matches);
+		const bool bCheckCompare = (Operator != EPCGAttributeFilterOperator::Equal) && (Operator != EPCGAttributeFilterOperator::NotEqual) && !bCheckStringSearch;
+		const int32 NumInput = TargetKeys->GetNum();
+
+		if (!PCGAttributeFilterHelpers::PrepareThresholdInfoFromInput(Context, OriginalData, NumInput, FirstThreshold, FirstThresholdInfo, TargetType, bCheckCompare, bCheckStringSearch))
+		{
+			ForwardInputToInFilterPin();
+			continue;
+		}
+
+		if (SecondThreshold && !PCGAttributeFilterHelpers::PrepareThresholdInfoFromInput(Context, OriginalData, NumInput , *SecondThreshold, SecondThresholdInfo, TargetType, bCheckCompare, bCheckStringSearch, &FirstThresholdInfo))
+		{
+			ForwardInputToInFilterPin();
+			continue;
+		}
+
+		UPCGData* InFilterData = nullptr;
+		UPCGData* OutFilterData = nullptr;
+
+		if (OperationData.bIsInputPointData)
+		{
+			const UPCGPointData* OriginalPointData = CastChecked<UPCGPointData>(OriginalData);
+			UPCGPointData* InFilterPointData = NewObject<UPCGPointData>();
+			UPCGPointData* OutFilterPointData = NewObject<UPCGPointData>();
+
+			OperationData.OriginalPoints = &OriginalPointData->GetPoints();
+
+			InFilterPointData->InitializeFromData(OriginalPointData);
+			OutFilterPointData->InitializeFromData(OriginalPointData);
+			OperationData.InFilterPoints = &InFilterPointData->GetMutablePoints();
+			OperationData.OutFilterPoints = &OutFilterPointData->GetMutablePoints();
+
+			OperationData.InFilterPoints->Reserve(OriginalPointData->GetPoints().Num());
+			OperationData.OutFilterPoints->Reserve(OriginalPointData->GetPoints().Num());
+
+			InFilterData = InFilterPointData;
+			OutFilterData = OutFilterPointData;
+		}
+		else
+		{
+			// Param data
+			const UPCGParamData* OriginalParamData = CastChecked<UPCGParamData>(OriginalData);
+			UPCGParamData* InFilterParamData = NewObject<UPCGParamData>();
+			UPCGParamData* OutFilterParamData = NewObject<UPCGParamData>();
+
+			OperationData.OriginalMetadata = OriginalParamData->Metadata;
+
+			// Add all attributes from the original param data, but without any entry/value. It will be added later.
+			OperationData.InFilterMetadata = InFilterParamData->Metadata;
+			OperationData.OutFilterMetadata = OutFilterParamData->Metadata;
+			OperationData.InFilterMetadata->AddAttributesFiltered(OriginalParamData->Metadata, TSet<FName>(), EPCGMetadataFilterMode::ExcludeAttributes);
+			OperationData.OutFilterMetadata->AddAttributesFiltered(OriginalParamData->Metadata, TSet<FName>(), EPCGMetadataFilterMode::ExcludeAttributes);
+
+			InFilterData = InFilterParamData;
+			OutFilterData = OutFilterParamData;
+		}
+
+		auto Operation = [&Operator, &FirstThresholdInfo, &SecondThresholdInfo, &TargetAccessor, &TargetKeys, &FirstThreshold, &SecondThreshold, &OperationData](auto Dummy) -> bool
+		{
+			using Type = decltype(Dummy);
+
+			const int32 NumberOfEntries = TargetKeys->GetNum();
+
+			if (NumberOfEntries <= 0)
+			{
+				return false;
+			}
+
+			TArray<Type, TInlineAllocator<PCGAttributeFilterConstants::ChunkSize>> TargetValues;
+			TArray<Type, TInlineAllocator<PCGAttributeFilterConstants::ChunkSize>> FirstThresholdValues;
+			TArray<Type, TInlineAllocator<PCGAttributeFilterConstants::ChunkSize>> SecondThresholdValues;
+			TArray<bool, TInlineAllocator<PCGAttributeFilterConstants::ChunkSize>> SkipTests;
+			TargetValues.SetNum(PCGAttributeFilterConstants::ChunkSize);
+			FirstThresholdValues.SetNum(PCGAttributeFilterConstants::ChunkSize);
+			SecondThresholdValues.SetNum(PCGAttributeFilterConstants::ChunkSize);
+			SkipTests.SetNumZeroed(PCGAttributeFilterConstants::ChunkSize); // All initialized to false
+
+			const int32 NumberOfIterations = (NumberOfEntries + PCGAttributeFilterConstants::ChunkSize - 1) / PCGAttributeFilterConstants::ChunkSize;
+
+			for (int32 i = 0; i < NumberOfIterations; ++i)
+			{
+				const int32 StartIndex = i * PCGAttributeFilterConstants::ChunkSize;
+				const int32 Range = FMath::Min(NumberOfEntries - StartIndex, PCGAttributeFilterConstants::ChunkSize);
+				TArrayView<Type> TargetView(TargetValues.GetData(), Range);
+				TArrayView<Type> FirstThresholdView(FirstThresholdValues.GetData(), Range);
+				TArrayView<Type> SecondThresholdView(SecondThresholdValues.GetData(), Range);
+
+				// Sampling the points if needed
+				auto SamplePointData = [Range, StartIndex, &OperationData, &SkipTests](UPCGPointData* InPointData, const UPCGSpatialData* InSpatialData)
+				{
+					if (InPointData != nullptr)
+					{
+						// Threshold points only have "ChunkSize" points.
+						TArray<FPCGPoint>& ThresholdPoints = InPointData->GetMutablePoints();
+						check(OperationData.OriginalPoints);
+						for (int32 j = 0; j < Range; ++j)
+						{
+							FPCGPoint ThresholdPoint;
+							const FPCGPoint& SourcePoint = (*OperationData.OriginalPoints)[StartIndex + j];
+							// If we already mark this point to skip test, don't even try to sample.
+							// If the sample fails, mark the point to skip test.
+							if (!SkipTests[j] && InSpatialData->SamplePoint(SourcePoint.Transform, SourcePoint.GetLocalBounds(), ThresholdPoint, InPointData->Metadata))
+							{
+								ThresholdPoints[j] = ThresholdPoint;
+							}
+							else
+							{
+								SkipTests[j] = true;
+							}
+						}
+					}
+				};
+
+				SamplePointData(FirstThresholdInfo.ThresholdPointData, FirstThresholdInfo.ThresholdSpatialData);
+
+				if (FirstThresholdInfo.ThresholdPointData != SecondThresholdInfo.ThresholdPointData)
+				{
+					SamplePointData(SecondThresholdInfo.ThresholdPointData, SecondThresholdInfo.ThresholdSpatialData);
+				}
+
+				// If ThresholdView point on ThresholdPointData points, there are only "ChunkSize" points in it.
+				// But it wraps around, and since StartIndex is a multiple of "ChunkSize", we'll always start at point 0, as wanted. 
+				if (!TargetAccessor->GetRange(TargetView, StartIndex, *TargetKeys) ||
+					!FirstThresholdInfo.ThresholdAccessor->GetRange(FirstThresholdView, StartIndex, *FirstThresholdInfo.ThresholdKeys, EPCGAttributeAccessorFlags::AllowBroadcast | EPCGAttributeAccessorFlags::AllowConstructible) ||
+					(SecondThresholdInfo.ThresholdAccessor.IsValid() && !SecondThresholdInfo.ThresholdAccessor->GetRange(SecondThresholdView, StartIndex, *SecondThresholdInfo.ThresholdKeys, EPCGAttributeAccessorFlags::AllowBroadcast | EPCGAttributeAccessorFlags::AllowConstructible)))
+				{
+					return false;
+				}
+
+				auto AddElement = [&OperationData](bool bInFilter, int32 Index)
+				{
+					if (OperationData.bIsInputPointData)
+					{
+						TArray<FPCGPoint>* Points = bInFilter ? OperationData.InFilterPoints : OperationData.OutFilterPoints;
+						check(Points && OperationData.OriginalPoints);
+						Points->Add((*OperationData.OriginalPoints)[Index]);
+					}
+					else
+					{
+						UPCGMetadata* Metadata = bInFilter ? OperationData.InFilterMetadata : OperationData.OutFilterMetadata;
+						check(Metadata && OperationData.OriginalMetadata);
+						PCGMetadataEntryKey EntryKey = Metadata->AddEntry();
+						Metadata->SetAttributes(Index, OperationData.OriginalMetadata, EntryKey);
+					}
+				};
+
+				for (int32 j = 0; j < Range; ++j)
+				{
+					if (SkipTests[j])
+					{
+						AddElement(true, StartIndex + j);
+						continue;
+					}
+
+					const bool bShouldKeep = (Operator == EPCGAttributeFilterOperator::InRange ? 
+						PCGAttributeFilterHelpers::ApplyRange(TargetValues[j], FirstThresholdValues[j], SecondThresholdValues[j], FirstThreshold.bInclusive, SecondThreshold->bInclusive) : 
+						PCGAttributeFilterHelpers::ApplyCompare(TargetValues[j], FirstThresholdValues[j], Operator));
+
+					if (bShouldKeep)
+					{
+						AddElement(true, StartIndex + j);
+					}
+					else
+					{
+						AddElement(false, StartIndex + j);
+					}
+				}
+			}
+
+			return true;
+		};
+
+		if (PCGMetadataAttribute::CallbackWithRightType(TargetAccessor->GetUnderlyingType(), Operation))
+		{
+			FPCGTaggedData& InFilterOutput = Outputs.Add_GetRef(Input);
+			InFilterOutput.Pin = PCGAttributeFilterConstants::InFilterLabel;
+			InFilterOutput.Data = InFilterData;
+			InFilterOutput.Tags = Input.Tags;
+
+			FPCGTaggedData& OutFilterOutput = Outputs.Add_GetRef(Input);
+			OutFilterOutput.Pin = PCGAttributeFilterConstants::OutFilterLabel;
+			OutFilterOutput.Data = OutFilterData;
+			OutFilterOutput.Tags = Input.Tags;
+		}
+		else
+		{
+			// Should be caught before when computing threshold info.
+			PCGE_LOG(Warning, GraphAndLog, LOCTEXT("TypeCannotBeConverted", "Cannot convert threshold type to target type"));
+			ForwardInputToInFilterPin();
+		}
+	}
+
+	return true;
+}
+
+bool FPCGAttributeFilterElement::ExecuteInternal(FPCGContext* Context) const
+{
+	TRACE_CPUPROFILER_EVENT_SCOPE(FPCGAttributeFilterElement::Execute);
+
+#if !WITH_EDITOR
+	const bool bHasInFilterOutputPin = Context->Node && Context->Node->IsOutputPinConnected(PCGAttributeFilterConstants::InFilterLabel);
+	const bool bHasOutsideFilterOutputPin = Context->Node && Context->Node->IsOutputPinConnected(PCGAttributeFilterConstants::OutFilterLabel);
+
+	// Early out - only in non-editor builds, otherwise we will potentially poison the cache, since it is input-driven
+	if (!bHasInFilterOutputPin && !bHasOutsideFilterOutputPin)
+	{
+		return true;
+	}
+#endif
+
+	const UPCGAttributeFilteringSettings* Settings = Context->GetInputSettings<UPCGAttributeFilteringSettings>();
+	check(Settings);
+
+	FPCGAttributeFilterThresholdSettings ThresholdSettings{};
+	ThresholdSettings.bUseConstantThreshold = Settings->bUseConstantThreshold;
+	ThresholdSettings.bUseSpatialQuery = Settings->bUseSpatialQuery;
+	ThresholdSettings.ThresholdAttribute = Settings->ThresholdAttribute;
+	ThresholdSettings.AttributeTypes = Settings->AttributeTypes;
+
+	return DoFiltering(Context, Settings->Operator, Settings->TargetAttribute, ThresholdSettings);
+}
+
+bool FPCGAttributeFilterRangeElement::ExecuteInternal(FPCGContext* Context) const
+{
+	TRACE_CPUPROFILER_EVENT_SCOPE(FPCGAttributeFilterRangeElement::Execute);
+
+#if !WITH_EDITOR
+	const bool bHasInFilterOutputPin = Context->Node && Context->Node->IsOutputPinConnected(PCGAttributeFilterConstants::InFilterLabel);
+	const bool bHasOutsideFilterOutputPin = Context->Node && Context->Node->IsOutputPinConnected(PCGAttributeFilterConstants::OutFilterLabel);
+
+	// Early out - only in non-editor builds, otherwise we will potentially poison the cache, since it is input-driven
+	if (!bHasInFilterOutputPin && !bHasOutsideFilterOutputPin)
+	{
+		return true;
+	}
+#endif
+
+	const UPCGAttributeFilteringRangeSettings* Settings = Context->GetInputSettings<UPCGAttributeFilteringRangeSettings>();
+	check(Settings);
+
+	return DoFiltering(Context, EPCGAttributeFilterOperator::InRange, Settings->TargetAttribute, Settings->MinThreshold, &Settings->MaxThreshold);
+}
+
+#undef LOCTEXT_NAMESPACE
diff --git a/Engine/Plugins/Experimental/PCG/Source/PCG/Private/Elements/PCGAttributeFilterElement.cpp b/Engine/Plugins/Experimental/PCG/Source/PCG/Private/Elements/PCGAttributeFilterElement.cpp
deleted file mode 100644
diff --git a/Engine/Plugins/Experimental/PCG/Source/PCG/Private/Elements/PCGAttributeFilterNamesElement.cpp b/Engine/Plugins/Experimental/PCG/Source/PCG/Private/Elements/PCGAttributeFilterNamesElement.cpp
new file mode 100644
--- /dev/null
+++ b/Engine/Plugins/Experimental/PCG/Source/PCG/Private/Elements/PCGAttributeFilterNamesElement.cpp
@@ -0,0 +1,182 @@
+// Copyright Epic Games, Inc. All Rights Reserved.
+
+#include "Elements/PCGAttributeFilterNamesElement.h"
+
+#include "PCGContext.h"
+#include "PCGParamData.h"
+#include "PCGPin.h"
+#include "Data/PCGSpatialData.h"
+#include "Helpers/PCGHelpers.h"
+
+#define LOCTEXT_NAMESPACE "PCGAttributeFilterElement"
+
+namespace PCGAttributeFilterConstants
+{
+	const FName NodeName = TEXT("FilterAttributesByName");
+	const FText NodeTitle = LOCTEXT("NodeTitle", "Filter Attributes By Name");
+}
+
+void UPCGAttributeFilterNamesSettings::PostLoad()
+{
+	Super::PostLoad();
+
+#if WITH_EDITOR
+	if (!AttributesToKeep_DEPRECATED.IsEmpty())
+	{
+		SelectedAttributes.Empty();
+		Operation = EPCGAttributeFilterOperation::KeepSelectedAttributes;
+		// Can't use FString::Join since it is an array of FName
+		for (int i = 0; i < AttributesToKeep_DEPRECATED.Num(); ++i)
+		{
+			if (i != 0)
+			{
+				SelectedAttributes += TEXT(",");
+			}
+
+			SelectedAttributes += AttributesToKeep_DEPRECATED[i].ToString();
+		}
+
+		AttributesToKeep_DEPRECATED.Empty();
+	}
+#endif // WITH_EDITOR
+}
+
+#if WITH_EDITOR
+FName UPCGAttributeFilterNamesSettings::GetDefaultNodeName() const
+{
+	return PCGAttributeFilterConstants::NodeName;
+}
+
+FText UPCGAttributeFilterNamesSettings::GetDefaultNodeTitle() const
+{
+	return PCGAttributeFilterConstants::NodeTitle;
+}
+#endif
+
+EPCGDataType UPCGAttributeFilterNamesSettings::GetCurrentPinTypes(const UPCGPin* InPin) const
+{
+	check(InPin);
+	if (!InPin->IsOutputPin())
+	{
+		return Super::GetCurrentPinTypes(InPin);
+	}
+
+	// Output pin narrows to union of inputs on first pin
+	const EPCGDataType InputTypeUnion = GetTypeUnionOfIncidentEdges(PCGPinConstants::DefaultInputLabel);
+	return (InputTypeUnion != EPCGDataType::None) ? InputTypeUnion : EPCGDataType::Any;
+}
+
+FName UPCGAttributeFilterNamesSettings::AdditionalTaskName() const
+{
+	TArray<FString> AttributesToKeep = PCGHelpers::GetStringArrayFromCommaSeparatedString(SelectedAttributes);
+
+	FString NodeName = PCGAttributeFilterConstants::NodeTitle.ToString();
+
+	switch (Operation)
+	{
+	case EPCGAttributeFilterOperation::KeepSelectedAttributes:
+		NodeName += TEXT(" (Keep)");
+		break;
+	case EPCGAttributeFilterOperation::DeleteSelectedAttributes:
+		NodeName += TEXT(" (Delete)");
+		break;
+	}
+
+	// If we filter only one attribute, show its name
+	if (AttributesToKeep.Num() == 1)
+	{
+		return FName(FString::Printf(TEXT("%s: %s"), *NodeName, *AttributesToKeep[0]));
+	}
+	else
+	{
+		return FName(NodeName);
+	}
+}
+
+TArray<FPCGPinProperties> UPCGAttributeFilterNamesSettings::InputPinProperties() const
+{
+	TArray<FPCGPinProperties> PinProperties;
+	PinProperties.Emplace(PCGPinConstants::DefaultInputLabel, EPCGDataType::Any);
+
+	return PinProperties;
+}
+
+TArray<FPCGPinProperties> UPCGAttributeFilterNamesSettings::OutputPinProperties() const
+{
+	TArray<FPCGPinProperties> PinProperties;
+	PinProperties.Emplace(PCGPinConstants::DefaultOutputLabel, EPCGDataType::Any);
+
+	return PinProperties;
+}
+
+FPCGElementPtr UPCGAttributeFilterNamesSettings::CreateElement() const
+{
+	return MakeShared<FPCGAttributeFilterNamesElement>();
+}
+
+bool FPCGAttributeFilterNamesElement::ExecuteInternal(FPCGContext* Context) const
+{
+	TRACE_CPUPROFILER_EVENT_SCOPE(FPCGAttributeFilterNamesElement::Execute);
+
+	check(Context);
+
+	const UPCGAttributeFilterNamesSettings* Settings = Context->GetInputSettings<UPCGAttributeFilterNamesSettings>();
+
+	const bool bAddAttributesFromParent = (Settings->Operation == EPCGAttributeFilterOperation::DeleteSelectedAttributes);
+	const EPCGMetadataFilterMode FilterMode = bAddAttributesFromParent ? EPCGMetadataFilterMode::ExcludeAttributes : EPCGMetadataFilterMode::IncludeAttributes;
+
+	TSet<FName> AttributesToFilter; 
+	const TArray<FString> FilterAttributes = PCGHelpers::GetStringArrayFromCommaSeparatedString(Settings->SelectedAttributes);
+	for (const FString& FilterAttribute : FilterAttributes)
+	{
+		AttributesToFilter.Add(FName(*FilterAttribute));
+	}
+
+	TArray<FPCGTaggedData> Inputs = Context->InputData.GetInputsByPin(PCGPinConstants::DefaultInputLabel);
+
+	for (const FPCGTaggedData& InputTaggedData : Inputs)
+	{
+		const UPCGData* InputData = InputTaggedData.Data;
+		UPCGData* OutputData = nullptr;
+
+		const UPCGMetadata* ParentMetadata = nullptr;
+		UPCGMetadata* Metadata = nullptr;
+
+		if (const UPCGSpatialData* InputSpatialData = Cast<UPCGSpatialData>(InputData))
+		{
+			ParentMetadata = InputSpatialData->Metadata;
+
+			UPCGSpatialData* NewSpatialData = InputSpatialData->DuplicateData(/*bInitializeFromThisData=*/false);
+			Metadata = NewSpatialData->Metadata;
+			NewSpatialData->Metadata->InitializeWithAttributeFilter(ParentMetadata, AttributesToFilter, FilterMode);
+
+			// No need to inherit metadata since we already initialized it.
+			NewSpatialData->InitializeFromData(InputSpatialData, /*InMetadataParentOverride=*/ nullptr, /*bInheritMetadata=*/ false);
+
+			OutputData = NewSpatialData;
+		}
+		else if (const UPCGParamData* InputParamData = Cast<UPCGParamData>(InputData))
+		{
+			ParentMetadata = InputParamData->Metadata;
+
+			UPCGParamData* NewParamData = NewObject<UPCGParamData>();
+			Metadata = NewParamData->Metadata;
+			Metadata->InitializeAsCopyWithAttributeFilter(ParentMetadata, AttributesToFilter, FilterMode);
+
+			OutputData = NewParamData;
+		}
+		else
+		{
+			PCGE_LOG(Error, GraphAndLog, LOCTEXT("InvalidInputData", "Invalid data as input. Only Spatial and Attribute Set data are supported."));
+			continue;
+		}
+
+		TArray<FPCGTaggedData>& Outputs = Context->OutputData.TaggedData;
+		FPCGTaggedData& Output = Outputs.Emplace_GetRef();
+		Output.Data = OutputData;
+	}
+
+	return true;
+}
+
+#undef LOCTEXT_NAMESPACE
diff --git a/Engine/Plugins/Experimental/PCG/Source/PCG/Private/Elements/PCGPointFilter.cpp b/Engine/Plugins/Experimental/PCG/Source/PCG/Private/Elements/PCGPointFilter.cpp
deleted file mode 100644
diff --git a/Engine/Plugins/Experimental/PCG/Source/PCG/Private/Tests/Elements/PCGAttributeFilterTest.cpp b/Engine/Plugins/Experimental/PCG/Source/PCG/Private/Tests/Elements/PCGAttributeFilterTest.cpp
new file mode 100644
--- /dev/null
+++ b/Engine/Plugins/Experimental/PCG/Source/PCG/Private/Tests/Elements/PCGAttributeFilterTest.cpp
@@ -0,0 +1,299 @@
+// Copyright Epic Games, Inc. All Rights Reserved.
+
+#include "Tests/PCGTestsCommon.h"
+
+#include "PCGComponent.h"
+#include "PCGContext.h"
+#include "PCGParamData.h"
+#include "PCGPoint.h"
+#include "Data/PCGPointData.h"
+#include "Elements/PCGAttributeFilter.h"
+#include "Metadata/PCGMetadata.h"
+#include "Metadata/PCGMetadataAttributeTpl.h"
+
+#include "Math/RandomStream.h"
+
+
+IMPLEMENT_CUSTOM_SIMPLE_AUTOMATION_TEST(FPCGPointFilterDensity, FPCGTestBaseClass, "pcg.tests.AttributeFilter.Points.Density", PCGTestsCommon::TestFlags)
+IMPLEMENT_CUSTOM_SIMPLE_AUTOMATION_TEST(FPCGPointFilterDensityRange, FPCGTestBaseClass, "pcg.tests.AttributeFilter.Points.DensityRange", PCGTestsCommon::TestFlags)
+IMPLEMENT_CUSTOM_SIMPLE_AUTOMATION_TEST(FPCGAttributeFilterInt, FPCGTestBaseClass, "pcg.tests.AttributeFilter.Params.Int", PCGTestsCommon::TestFlags)
+IMPLEMENT_CUSTOM_SIMPLE_AUTOMATION_TEST(FPCGAttributeFilterIntRange, FPCGTestBaseClass, "pcg.tests.AttributeFilter.Params.IntRange", PCGTestsCommon::TestFlags)
+
+namespace PCGPointFilterTest
+{
+	const FName InsideFilterLabel = TEXT("InsideFilter");
+	const FName OutsideFilterLabel = TEXT("OutsideFilter");
+
+	UPCGPointData* GeneratePointDataWithRandomDensity(int32 InNumPoints, int32 InRandomSeed)
+	{
+		UPCGPointData* PointData = PCGTestsCommon::CreateEmptyPointData();
+		TArray<FPCGPoint>& Points = PointData->GetMutablePoints();
+
+		FRandomStream RandomSource(InRandomSeed);
+
+		for (int I = 0; I < InNumPoints; ++I)
+		{
+			FPCGPoint& Point = Points.Emplace_GetRef(FTransform(), RandomSource.FRand(), I);
+		}
+
+		return PointData;
+	}
+
+	static const FName IntAttributeName = TEXT("Int");
+
+	UPCGParamData* GenerateAttributeSetWithRandomInt(int32 InNumEntries, int32 InRandomSeed)
+	{
+		UPCGParamData* ParamData = NewObject<UPCGParamData>();
+		check(ParamData && ParamData->Metadata);
+		FPCGMetadataAttribute<int32>* Attribute = ParamData->Metadata->CreateAttribute<int32>(IntAttributeName, 0, true, false);
+		check(Attribute)
+
+		FRandomStream RandomSource(InRandomSeed);
+
+		for (int I = 0; I < InNumEntries; ++I)
+		{
+			PCGMetadataEntryKey Key = ParamData->Metadata->AddEntry();
+			Attribute->SetValue(Key, RandomSource.GetUnsignedInt());
+		}
+
+		return ParamData;
+	}
+}
+
+bool FPCGPointFilterDensity::RunTest(const FString& Parameters)
+{
+	PCGTestsCommon::FTestData TestData;
+	UPCGAttributeFilteringSettings* Settings = PCGTestsCommon::GenerateSettings<UPCGAttributeFilteringSettings>(TestData);
+	check(Settings);
+
+	FPCGElementPtr TestElement = TestData.Settings->GetElement();
+
+	static const int32 NumPoints = 50;
+	static const float DensityThreshold = 0.5f;
+
+	Settings->Operator = EPCGAttributeFilterOperator::Lesser;
+	Settings->TargetAttribute.SetPointProperty(EPCGPointProperties::Density);
+	Settings->bUseConstantThreshold = true;
+	Settings->AttributeTypes.FloatValue = DensityThreshold;
+	Settings->AttributeTypes.Type = EPCGMetadataTypes::Float;
+
+	FPCGTaggedData& TaggedData = TestData.InputData.TaggedData.Emplace_GetRef(FPCGTaggedData());
+	TaggedData.Pin = PCGPinConstants::DefaultInputLabel;
+	TaggedData.Data = PCGPointFilterTest::GeneratePointDataWithRandomDensity(NumPoints, TestData.Seed);
+
+	TUniquePtr<FPCGContext> Context = TestData.InitializeTestContext();
+
+	while (!TestElement->Execute(Context.Get())) {}
+
+	// We should have outputs on both
+	TArray<FPCGTaggedData> InFilterOutput = Context->OutputData.GetInputsByPin(PCGPointFilterTest::InsideFilterLabel);
+	TArray<FPCGTaggedData> OutFilterOutput = Context->OutputData.GetInputsByPin(PCGPointFilterTest::OutsideFilterLabel);
+
+	UTEST_EQUAL(TEXT("InFilter pin has 1 output"), InFilterOutput.Num(), 1);
+	UTEST_EQUAL(TEXT("OutFilter pin has 1 output"), OutFilterOutput.Num(), 1);
+
+	// Making sure the output is a PointData
+	const UPCGPointData* InFilterPointData = Cast<UPCGPointData>(InFilterOutput[0].Data);
+	const UPCGPointData* OutFilterPointData = Cast<UPCGPointData>(OutFilterOutput[0].Data);
+
+	UTEST_NOT_NULL(TEXT("InFilter data is a point data"), InFilterPointData);
+	UTEST_NOT_NULL(TEXT("OutFilter data is a point data"), OutFilterPointData);
+
+	// Verifying that all points have the right density
+	for (const FPCGPoint& Point : InFilterPointData->GetPoints())
+	{
+		UTEST_TRUE(*FString::Printf(TEXT("Point has a density (%f) lower than the threshold (%f)"), Point.Density, DensityThreshold), Point.Density < DensityThreshold);
+	}
+
+	for (const FPCGPoint& Point : OutFilterPointData->GetPoints())
+	{
+		UTEST_TRUE(*FString::Printf(TEXT("Point has a density (%f) higher or equal than the threshold (%f)"), Point.Density, DensityThreshold), Point.Density >= DensityThreshold);
+	}
+
+	return true;
+}
+
+bool FPCGPointFilterDensityRange::RunTest(const FString& Parameters)
+{
+	PCGTestsCommon::FTestData TestData;
+	UPCGAttributeFilteringRangeSettings* Settings = PCGTestsCommon::GenerateSettings<UPCGAttributeFilteringRangeSettings>(TestData);
+	check(Settings);
+
+	FPCGElementPtr TestElement = TestData.Settings->GetElement();
+
+	static const int32 NumPoints = 50;
+	static const float DensityMinThreshold = 0.3f;
+	static const float DensityMaxThreshold = 0.8f;
+
+	Settings->TargetAttribute.SetPointProperty(EPCGPointProperties::Density);
+
+	Settings->MinThreshold.bUseConstantThreshold = true;
+	Settings->MinThreshold.AttributeTypes.FloatValue = DensityMinThreshold;
+	Settings->MinThreshold.AttributeTypes.Type = EPCGMetadataTypes::Float;
+
+	Settings->MaxThreshold.bUseConstantThreshold = true;
+	Settings->MaxThreshold.AttributeTypes.FloatValue = DensityMaxThreshold;
+	Settings->MaxThreshold.AttributeTypes.Type = EPCGMetadataTypes::Float;
+
+	FPCGTaggedData& TaggedData = TestData.InputData.TaggedData.Emplace_GetRef(FPCGTaggedData());
+	TaggedData.Pin = PCGPinConstants::DefaultInputLabel;
+	TaggedData.Data = PCGPointFilterTest::GeneratePointDataWithRandomDensity(NumPoints, TestData.Seed);
+
+	TUniquePtr<FPCGContext> Context = TestData.InitializeTestContext();
+
+	while (!TestElement->Execute(Context.Get())) {}
+
+	// We should have outputs on both
+	TArray<FPCGTaggedData> InFilterOutput = Context->OutputData.GetInputsByPin(PCGPointFilterTest::InsideFilterLabel);
+	TArray<FPCGTaggedData> OutFilterOutput = Context->OutputData.GetInputsByPin(PCGPointFilterTest::OutsideFilterLabel);
+
+	UTEST_EQUAL(TEXT("InFilter pin has 1 output"), InFilterOutput.Num(), 1);
+	UTEST_EQUAL(TEXT("OutFilter pin has 1 output"), OutFilterOutput.Num(), 1);
+
+	// Making sure the output is a PointData
+	const UPCGPointData* InFilterPointData = Cast<UPCGPointData>(InFilterOutput[0].Data);
+	const UPCGPointData* OutFilterPointData = Cast<UPCGPointData>(OutFilterOutput[0].Data);
+
+	UTEST_NOT_NULL(TEXT("InFilter data is a point data"), InFilterPointData);
+	UTEST_NOT_NULL(TEXT("OutFilter data is a point data"), OutFilterPointData);
+
+	// Verifying that all points have the right density
+	for (const FPCGPoint& Point : InFilterPointData->GetPoints())
+	{
+		UTEST_TRUE(*FString::Printf(TEXT("Point has a density (%f) within the range ([%f, %f])"), Point.Density, DensityMinThreshold, DensityMaxThreshold), (Point.Density >= DensityMinThreshold) && (Point.Density <= DensityMaxThreshold));
+	}
+
+	for (const FPCGPoint& Point : OutFilterPointData->GetPoints())
+	{
+		UTEST_TRUE(*FString::Printf(TEXT("Point has a density (%f) outside the range ([%f, %f])"), Point.Density, DensityMinThreshold, DensityMaxThreshold), (Point.Density < DensityMinThreshold) || (Point.Density > DensityMaxThreshold));
+	}
+
+	return true;
+}
+
+bool FPCGAttributeFilterInt::RunTest(const FString& Parameters)
+{
+	PCGTestsCommon::FTestData TestData;
+	UPCGAttributeFilteringSettings* Settings = PCGTestsCommon::GenerateSettings<UPCGAttributeFilteringSettings>(TestData);
+	check(Settings);
+
+	FPCGElementPtr TestElement = TestData.Settings->GetElement();
+
+	static const int32 NumEntries = 50;
+	static const int64 IntThreshold = 999999;
+
+	Settings->Operator = EPCGAttributeFilterOperator::Lesser;
+	Settings->TargetAttribute.SetAttributeName(PCGPointFilterTest::IntAttributeName);
+	Settings->bUseConstantThreshold = true;
+	Settings->AttributeTypes.IntValue = IntThreshold;
+	Settings->AttributeTypes.Type = EPCGMetadataTypes::Integer64;
+
+	FPCGTaggedData& TaggedData = TestData.InputData.TaggedData.Emplace_GetRef(FPCGTaggedData());
+	TaggedData.Pin = PCGPinConstants::DefaultInputLabel;
+	TaggedData.Data = PCGPointFilterTest::GenerateAttributeSetWithRandomInt(NumEntries, TestData.Seed);
+
+	TUniquePtr<FPCGContext> Context = TestData.InitializeTestContext();
+
+	while (!TestElement->Execute(Context.Get())) {}
+
+	// We should have outputs on both
+	TArray<FPCGTaggedData> InFilterOutput = Context->OutputData.GetInputsByPin(PCGPointFilterTest::InsideFilterLabel);
+	TArray<FPCGTaggedData> OutFilterOutput = Context->OutputData.GetInputsByPin(PCGPointFilterTest::OutsideFilterLabel);
+
+	UTEST_EQUAL(TEXT("InFilter pin has 1 output"), InFilterOutput.Num(), 1);
+	UTEST_EQUAL(TEXT("OutFilter pin has 1 output"), OutFilterOutput.Num(), 1);
+
+	// Making sure the output is a ParamData, and has its attribute
+	const UPCGParamData* InFilterParamData = Cast<UPCGParamData>(InFilterOutput[0].Data);
+	const UPCGParamData* OutFilterParamData = Cast<UPCGParamData>(OutFilterOutput[0].Data);
+
+	UTEST_NOT_NULL(TEXT("InFilter data is a param data"), InFilterParamData);
+	UTEST_NOT_NULL(TEXT("OutFilter data is a param data"), OutFilterParamData);
+
+	const FPCGMetadataAttribute<int32>* InFilterAttribute = InFilterParamData->Metadata->GetConstTypedAttribute<int32>(PCGPointFilterTest::IntAttributeName);
+	const FPCGMetadataAttribute<int32>* OutFilterAttribute = OutFilterParamData->Metadata->GetConstTypedAttribute<int32>(PCGPointFilterTest::IntAttributeName);
+
+	UTEST_NOT_NULL(TEXT("InFilter metadata has the int attribute"), InFilterAttribute);
+	UTEST_NOT_NULL(TEXT("OutFilter metadata has the int attribute"), OutFilterAttribute);
+
+	// Verifying that all points have the right int value
+	for (int32 Key = 0; Key < InFilterParamData->Metadata->GetLocalItemCount(); ++Key)
+	{
+		int32 Value = InFilterAttribute->GetValueFromItemKey(Key);
+		UTEST_TRUE(*FString::Printf(TEXT("Attribute has a value (%d) lower than the threshold (%d)"), Value, (int32)IntThreshold), Value < (int32)IntThreshold);
+	}
+
+	for (int32 Key = 0; Key < OutFilterParamData->Metadata->GetLocalItemCount(); ++Key)
+	{
+		int32 Value = OutFilterAttribute->GetValueFromItemKey(Key);
+		UTEST_TRUE(*FString::Printf(TEXT("Attribute has a value (%d) higher than the threshold (%d)"), Value, (int32)IntThreshold), Value >= (int32)IntThreshold);
+	}
+
+	return true;
+}
+
+bool FPCGAttributeFilterIntRange::RunTest(const FString& Parameters)
+{
+	PCGTestsCommon::FTestData TestData;
+	UPCGAttributeFilteringRangeSettings* Settings = PCGTestsCommon::GenerateSettings<UPCGAttributeFilteringRangeSettings>(TestData);
+	check(Settings);
+
+	FPCGElementPtr TestElement = TestData.Settings->GetElement();
+
+	static const int32 NumEntries = 50;
+	static const int64 IntMinThreshold = 0;
+	static const int64 IntMaxThreshold = 1999999999;
+
+	Settings->TargetAttribute.SetAttributeName(PCGPointFilterTest::IntAttributeName);
+
+	Settings->MinThreshold.bUseConstantThreshold = true;
+	Settings->MinThreshold.AttributeTypes.IntValue = IntMinThreshold;
+	Settings->MinThreshold.AttributeTypes.Type = EPCGMetadataTypes::Integer64;
+
+	Settings->MaxThreshold.bUseConstantThreshold = true;
+	Settings->MaxThreshold.AttributeTypes.IntValue = IntMaxThreshold;
+	Settings->MaxThreshold.AttributeTypes.Type = EPCGMetadataTypes::Integer64;
+
+	FPCGTaggedData& TaggedData = TestData.InputData.TaggedData.Emplace_GetRef(FPCGTaggedData());
+	TaggedData.Pin = PCGPinConstants::DefaultInputLabel;
+	TaggedData.Data = PCGPointFilterTest::GenerateAttributeSetWithRandomInt(NumEntries, TestData.Seed);
+
+	TUniquePtr<FPCGContext> Context = TestData.InitializeTestContext();
+
+	while (!TestElement->Execute(Context.Get())) {}
+
+	// We should have outputs on both
+	TArray<FPCGTaggedData> InFilterOutput = Context->OutputData.GetInputsByPin(PCGPointFilterTest::InsideFilterLabel);
+	TArray<FPCGTaggedData> OutFilterOutput = Context->OutputData.GetInputsByPin(PCGPointFilterTest::OutsideFilterLabel);
+
+	UTEST_EQUAL(TEXT("InFilter pin has 1 output"), InFilterOutput.Num(), 1);
+	UTEST_EQUAL(TEXT("OutFilter pin has 1 output"), OutFilterOutput.Num(), 1);
+
+	// Making sure the output is a ParamData, and has its attribute
+	const UPCGParamData* InFilterParamData = Cast<UPCGParamData>(InFilterOutput[0].Data);
+	const UPCGParamData* OutFilterParamData = Cast<UPCGParamData>(OutFilterOutput[0].Data);
+
+	UTEST_NOT_NULL(TEXT("InFilter data is a param data"), InFilterParamData);
+	UTEST_NOT_NULL(TEXT("OutFilter data is a param data"), OutFilterParamData);
+
+	const FPCGMetadataAttribute<int32>* InFilterAttribute = InFilterParamData->Metadata->GetConstTypedAttribute<int32>(PCGPointFilterTest::IntAttributeName);
+	const FPCGMetadataAttribute<int32>* OutFilterAttribute = OutFilterParamData->Metadata->GetConstTypedAttribute<int32>(PCGPointFilterTest::IntAttributeName);
+
+	UTEST_NOT_NULL(TEXT("InFilter metadata has the int attribute"), InFilterAttribute);
+	UTEST_NOT_NULL(TEXT("OutFilter metadata has the int attribute"), OutFilterAttribute);
+
+	// Verifying that all points have the right int value
+	for (int32 Key = 0; Key < InFilterParamData->Metadata->GetLocalItemCount(); ++Key)
+	{
+		int32 Value = InFilterAttribute->GetValueFromItemKey(Key);
+		UTEST_TRUE(*FString::Printf(TEXT("Attribute has a value (%d) within the range [%d, %d]"), Value, (int32)IntMinThreshold, (int32)IntMaxThreshold), Value >= (int32)IntMinThreshold && Value <= (int32)IntMaxThreshold);
+	}
+
+	for (int32 Key = 0; Key < OutFilterParamData->Metadata->GetLocalItemCount(); ++Key)
+	{
+		int32 Value = OutFilterAttribute->GetValueFromItemKey(Key);
+		UTEST_TRUE(*FString::Printf(TEXT("Attribute has a value (%d) outside the range [%d, %d]"), Value, (int32)IntMinThreshold, (int32)IntMaxThreshold), Value < (int32)IntMinThreshold || Value > (int32)IntMaxThreshold);
+	}
+
+	return true;
+}
\ No newline at end of file
diff --git a/Engine/Plugins/Experimental/PCG/Source/PCG/Private/Tests/Elements/PCGPointFilterTest.cpp b/Engine/Plugins/Experimental/PCG/Source/PCG/Private/Tests/Elements/PCGPointFilterTest.cpp
deleted file mode 100644
diff --git a/Engine/Plugins/Experimental/PCG/Source/PCG/Public/Elements/PCGAttributeFilter.h b/Engine/Plugins/Experimental/PCG/Source/PCG/Public/Elements/PCGAttributeFilter.h
new file mode 100644
--- /dev/null
+++ b/Engine/Plugins/Experimental/PCG/Source/PCG/Public/Elements/PCGAttributeFilter.h
@@ -0,0 +1,188 @@
+// Copyright Epic Games, Inc. All Rights Reserved.
+
+#pragma once
+
+#include "PCGSettings.h"
+#include "Metadata/PCGAttributePropertySelector.h"
+#include "Metadata/PCGMetadataTypesConstantStruct.h"
+
+#include "PCGAttributeFilter.generated.h"
+
+UENUM()
+enum class EPCGAttributeFilterOperator : uint8
+{
+	Greater UMETA(DisplayName=">"),
+	GreaterOrEqual UMETA(DisplayName=">="),
+	Lesser UMETA(DisplayName="<"),
+	LesserOrEqual UMETA(DisplayName="<="),
+	Equal UMETA(DisplayName="="),
+	NotEqual UMETA(DisplayName="!="),
+	InRange UMETA(Hidden),
+	Substring,
+	Matches
+};
+
+USTRUCT(BlueprintType)
+struct PCG_API FPCGAttributeFilterThresholdSettings
+{
+	GENERATED_BODY()
+
+#if WITH_EDITOR
+	void OnPostLoad();
+#endif
+
+	/** If the threshold in included or excluded from the range. */
+	UPROPERTY(BlueprintReadWrite, EditAnywhere, Category = Settings, meta = (PCG_Overridable))
+	bool bInclusive = true;
+
+	UPROPERTY(BlueprintReadWrite, EditAnywhere, Category = Settings, meta = (PCG_Overridable))
+	bool bUseConstantThreshold = false;
+
+	UPROPERTY(BlueprintReadWrite, EditAnywhere, Category = Settings, meta = (EditCondition = "!bUseConstantThreshold", EditConditionHides, PCG_NotOverridable))
+	FPCGAttributePropertyInputSelector ThresholdAttribute;
+
+	/** For Point Data, enabling this option will use sampling rather than comparing points 1 to 1 directly. For other spatial data, this is always true, and for attribute sets, always false. */
+	UPROPERTY(BlueprintReadWrite, EditAnywhere, Category = Settings, meta = (EditCondition = "!bUseConstantThreshold", EditConditionHides, PCG_Overridable))
+	bool bUseSpatialQuery = true;
+
+	UPROPERTY(BlueprintReadWrite, EditAnywhere, Category = Settings, meta = (EditCondition = "bUseConstantThreshold", EditConditionHides, ShowOnlyInnerProperties, DisplayAfter = "bUseConstantThreshold", PCG_NotOverridable))
+	FPCGMetadataTypesConstantStruct AttributeTypes;
+};
+
+/**
+* Attribute filter that allows to do "A op B" type filtering, where A is the input spatial data or Attribute set,
+* and B is either a constant, another spatial data (if input is a spatial data), an Attribute set (in filter) or the input itself.
+* The filtering can be done either on properties or attributes.
+* Some examples:
+* - Threshold on property by constant (A.Density > 0.5)
+* - Threshold on attribute by constant (A.aaa != "bob")
+* - Threshold on property by metadata attribute(A.density >= B.bbb)
+* - Threshold on property by property(A.density <= B.steepness)
+* - Threshold on attribute by metadata attribute(A.aaa < B.bbb)
+* - Threshold on attribute by property(A.aaa == B.color)
+*/
+UCLASS(BlueprintType, ClassGroup = (Procedural))
+class PCG_API UPCGAttributeFilteringSettings : public UPCGSettings
+{
+	GENERATED_BODY()
+
+public:
+	UPCGAttributeFilteringSettings();
+
+	//~Begin UObject interface
+	virtual void PostLoad() override;
+	//~End UObject interface
+
+	//~Begin UPCGSettings interface
+#if WITH_EDITOR
+	virtual FName GetDefaultNodeName() const override { return FName(TEXT("AttributeFilter")); }
+	virtual FText GetDefaultNodeTitle() const override { return NSLOCTEXT("PCGAttributeFilteringElement", "NodeTitle", "Attribute Filter"); }
+	virtual TArray<FText> GetNodeTitleAliases() const { return { NSLOCTEXT("PCGAttributeFilteringElement", "AliasNodeTitle", "Point Filter") }; }
+	virtual EPCGSettingsType GetType() const override { return EPCGSettingsType::Filter; }
+	virtual bool HasDynamicPins() const override { return true; }
+#endif
+	virtual EPCGDataType GetCurrentPinTypes(const UPCGPin* InPin) const override;
+
+protected:
+	virtual TArray<FPCGPinProperties> InputPinProperties() const override;
+	virtual TArray<FPCGPinProperties> OutputPinProperties() const override;
+	virtual FPCGElementPtr CreateElement() const override;
+	//~End UPCGSettings interface
+
+public:
+	UPROPERTY(BlueprintReadWrite, EditAnywhere, Category = Settings, meta = (PCG_Overridable))
+	EPCGAttributeFilterOperator Operator = EPCGAttributeFilterOperator::Greater;
+
+	/** Target property/attribute related properties */
+	UPROPERTY(BlueprintReadWrite, EditAnywhere, Category = Settings, meta = (PCG_Overridable))
+	FPCGAttributePropertyInputSelector TargetAttribute;
+
+	/** Threshold property/attribute/constant related properties */
+	UPROPERTY(BlueprintReadWrite, EditAnywhere, Category = Settings, meta = (PCG_Overridable))
+	bool bUseConstantThreshold = false;
+
+	UPROPERTY(BlueprintReadWrite, EditAnywhere, Category = Settings, meta = (EditCondition = "!bUseConstantThreshold", EditConditionHides, PCG_NotOverridable))
+	FPCGAttributePropertyInputSelector ThresholdAttribute;
+
+	/** If the threshold data is Point data, it will sample input points in threshold data. Always true with Spatial data.*/
+	UPROPERTY(BlueprintReadWrite, EditAnywhere, Category = Settings, meta = (EditCondition = "!bUseConstantThreshold", EditConditionHides, PCG_Overridable))
+	bool bUseSpatialQuery = true;
+
+	UPROPERTY(BlueprintReadWrite, EditAnywhere, Category = Settings, meta = (EditCondition = "bUseConstantThreshold", EditConditionHides, ShowOnlyInnerProperties, DisplayAfter = "bUseConstantThreshold", PCG_NotOverridable))
+	FPCGMetadataTypesConstantStruct AttributeTypes;
+};
+
+
+/**
+* Attribute filter on range that allows to do "A op B" type filtering, where A is the input spatial data or Attribute set,
+* and B is either a constant, another spatial data (if input is a spatial data), an Attribute set (in filter) or the input itself.
+* The filtering can be done either on properties or attributes.
+* Some examples (that might not make sense, but are valid):
+* - Threshold on property by constant (A.Density in [0.2, 0.5])
+* - Threshold on attribute by constant (A.aaa in [0.4, 0.6])
+* - Threshold on property by metadata attribute(A.density in [B.bbmin, B.bbmax])
+* - Threshold on property by property(A.density in [B.position.x, B.steepness])
+* - Threshold on attribute by metadata attribute(A.aaa in [B.bbmin, B.bbmax])
+* - Threshold on attribute by property(A.aaa in [B.position, B.scale])
+*/
+UCLASS(BlueprintType, ClassGroup = (Procedural))
+class PCG_API UPCGAttributeFilteringRangeSettings : public UPCGSettings
+{
+	GENERATED_BODY()
+
+public:
+	UPCGAttributeFilteringRangeSettings();
+
+	//~Begin UObject interface
+#if WITH_EDITOR
+	virtual void PostLoad() override;
+#endif
+	//~End UObject interface
+
+	//~Begin UPCGSettings interface
+#if WITH_EDITOR
+	virtual FName GetDefaultNodeName() const override { return FName(TEXT("AttributeFilterRange")); }
+	virtual FText GetDefaultNodeTitle() const override { return NSLOCTEXT("PCGAttributeFilteringElement", "NodeTitleRange", "Attribute Filter Range"); }
+	virtual TArray<FText> GetNodeTitleAliases() const { return { NSLOCTEXT("PCGAttributeFilteringElement", "AliasNodeTitleRange", "Point Filter Range") }; }
+	virtual EPCGSettingsType GetType() const override { return EPCGSettingsType::Filter; }
+	virtual bool HasDynamicPins() const override { return true; }
+#endif
+	virtual EPCGDataType GetCurrentPinTypes(const UPCGPin* InPin) const override;
+
+protected:
+	virtual TArray<FPCGPinProperties> InputPinProperties() const override;
+	virtual TArray<FPCGPinProperties> OutputPinProperties() const override;
+	virtual FPCGElementPtr CreateElement() const override;
+	//~End UPCGSettings interface
+
+public:
+	/** Target property/attribute related properties */
+	UPROPERTY(BlueprintReadWrite, EditAnywhere, Category = Settings, meta = (PCG_Overridable))
+	FPCGAttributePropertyInputSelector TargetAttribute;
+
+	/** Threshold property/attribute/constant related properties */
+	UPROPERTY(BlueprintReadWrite, EditAnywhere, Category = Settings, meta = (PCG_Overridable))
+	FPCGAttributeFilterThresholdSettings MinThreshold;
+
+	UPROPERTY(BlueprintReadWrite, EditAnywhere, Category = Settings, meta = (PCG_Overridable))
+	FPCGAttributeFilterThresholdSettings MaxThreshold;
+};
+
+class FPCGAttributeFilterElementBase : public IPCGElement
+{
+protected:
+	bool DoFiltering(FPCGContext* Context, EPCGAttributeFilterOperator InOperation, const FPCGAttributePropertyInputSelector& TargetAttribute, const FPCGAttributeFilterThresholdSettings& FirstThreshold, const FPCGAttributeFilterThresholdSettings* SecondThreshold = nullptr) const;
+};
+
+class FPCGAttributeFilterElement : public FPCGAttributeFilterElementBase
+{
+protected:
+	virtual bool ExecuteInternal(FPCGContext* Context) const override;
+};
+
+class FPCGAttributeFilterRangeElement : public FPCGAttributeFilterElementBase
+{
+protected:
+	virtual bool ExecuteInternal(FPCGContext* Context) const override;
+};
+
diff --git a/Engine/Plugins/Experimental/PCG/Source/PCG/Public/Elements/PCGAttributeFilterElement.h b/Engine/Plugins/Experimental/PCG/Source/PCG/Public/Elements/PCGAttributeFilterElement.h
deleted file mode 100644
diff --git a/Engine/Plugins/Experimental/PCG/Source/PCG/Public/Elements/PCGAttributeFilterNamesElement.h b/Engine/Plugins/Experimental/PCG/Source/PCG/Public/Elements/PCGAttributeFilterNamesElement.h
new file mode 100644
--- /dev/null
+++ b/Engine/Plugins/Experimental/PCG/Source/PCG/Public/Elements/PCGAttributeFilterNamesElement.h
@@ -0,0 +1,67 @@
+// Copyright Epic Games, Inc. All Rights Reserved.
+
+#pragma once
+
+#include "PCGSettings.h"
+
+#include "PCGAttributeFilterNamesElement.generated.h"
+
+UENUM()
+enum class EPCGAttributeFilterOperation
+{
+	KeepSelectedAttributes,
+	DeleteSelectedAttributes
+};
+
+/**
+* Filter the attributes from a given input metadata.
+* Will remove all attributes that are not listed in AttributesToKeep.
+* If an attribute to keep is not in the original metadata, it won't be added.
+* 
+* The output will be the original data with the updated metadata.
+*/
+UCLASS(BlueprintType, ClassGroup = (Procedural))
+class PCG_API UPCGAttributeFilterNamesSettings : public UPCGSettings
+{
+	GENERATED_BODY()
+
+public:
+	//~Begin UObject interface
+	virtual void PostLoad() override;
+	//~End UObject interface
+
+	//~Begin UPCGSettings interface
+#if WITH_EDITOR
+	virtual FName GetDefaultNodeName() const override;
+	virtual FText GetDefaultNodeTitle() const override;
+	virtual EPCGSettingsType GetType() const override { return EPCGSettingsType::Metadata; }
+	virtual bool HasDynamicPins() const override { return true; }
+#endif
+	virtual EPCGDataType GetCurrentPinTypes(const UPCGPin* InPin) const override;
+	virtual FName AdditionalTaskName() const override;
+
+protected:
+	virtual TArray<FPCGPinProperties> InputPinProperties() const override;
+	virtual TArray<FPCGPinProperties> OutputPinProperties() const override;
+	virtual FPCGElementPtr CreateElement() const override;
+	//~End UPCGSettings interface
+
+public:
+	UPROPERTY(BlueprintReadWrite, EditAnywhere, Category = Settings)
+	EPCGAttributeFilterOperation Operation = EPCGAttributeFilterOperation::KeepSelectedAttributes;
+
+	UPROPERTY(BlueprintReadWrite, EditAnywhere, Category = Settings)
+	FString SelectedAttributes;
+
+#if WITH_EDITORONLY_DATA
+	UPROPERTY()
+	TArray<FName> AttributesToKeep_DEPRECATED;
+#endif // WITH_EDITORONLY_DATA
+};
+
+
+class FPCGAttributeFilterNamesElement : public IPCGElement
+{
+protected:
+	virtual bool ExecuteInternal(FPCGContext* Context) const override;
+};
diff --git a/Engine/Plugins/Experimental/PCG/Source/PCG/Public/Elements/PCGPointFilter.h b/Engine/Plugins/Experimental/PCG/Source/PCG/Public/Elements/PCGPointFilter.h
deleted file mode 100644
diff --git a/Engine/Plugins/Experimental/PCG/Source/PCG/Public/PCGSettings.h b/Engine/Plugins/Experimental/PCG/Source/PCG/Public/PCGSettings.h
--- a/Engine/Plugins/Experimental/PCG/Source/PCG/Public/PCGSettings.h
+++ b/Engine/Plugins/Experimental/PCG/Source/PCG/Public/PCGSettings.h
@@ -261,6 +261,10 @@ public:
 
 	virtual FName GetDefaultNodeName() const { return NAME_None; }
 	virtual FText GetDefaultNodeTitle() const { return FText::FromName(GetDefaultNodeName()); }
+
+	/** List of extra aliases that will be added to the node list in the Editor.Useful when we rename a node, but we still want the user to find the old one. */
+	virtual TArray<FText> GetNodeTitleAliases() const { return {}; }
+
 	virtual FText GetNodeTooltipText() const { return FText::GetEmpty(); }
 	virtual FLinearColor GetNodeTitleColor() const { return FLinearColor::White; }
 	virtual EPCGSettingsType GetType() const { return EPCGSettingsType::Generic; }
diff --git a/Engine/Plugins/Experimental/PCG/Source/PCGEditor/Private/PCGEditorGraphSchema.cpp b/Engine/Plugins/Experimental/PCG/Source/PCGEditor/Private/PCGEditorGraphSchema.cpp
--- a/Engine/Plugins/Experimental/PCG/Source/PCGEditor/Private/PCGEditorGraphSchema.cpp
+++ b/Engine/Plugins/Experimental/PCG/Source/PCGEditor/Private/PCGEditorGraphSchema.cpp
@@ -360,6 +360,14 @@ void UPCGEditorGraphSchema::GetNativeElementActions(FGraphActionMenuBuilder& Act
 					TSharedPtr<FPCGEditorGraphSchemaAction_NewNativeElement> NewAction(new FPCGEditorGraphSchemaAction_NewNativeElement(Category, MenuDesc, Description, 0));
 					NewAction->SettingsClass = SettingsClass;
 					ActionMenuBuilder.AddAction(NewAction);
+
+					// Also add all aliases
+					for (const FText& Alias : PCGSettings->GetNodeTitleAliases())
+					{
+						TSharedPtr<FPCGEditorGraphSchemaAction_NewNativeElement> NewAliasAction(new FPCGEditorGraphSchemaAction_NewNativeElement(Category, Alias, Description, 0));
+						NewAliasAction->SettingsClass = SettingsClass;
+						ActionMenuBuilder.AddAction(NewAliasAction);
+					}
 				}
 
 				// Also add preconfigured settings
