[PCG] Allow for deeper level of actor overrides in Spawn Actors

* Add a new "PropertyChain" accessor that can set values on a given chain of properties
* Cleanup AccessorHelpers and add new helper to create this chain accessor

diff --git a/Engine/Plugins/Experimental/PCG/Source/PCG/Private/Elements/PCGSpawnActor.cpp b/Engine/Plugins/Experimental/PCG/Source/PCG/Private/Elements/PCGSpawnActor.cpp
--- a/Engine/Plugins/Experimental/PCG/Source/PCG/Private/Elements/PCGSpawnActor.cpp
+++ b/Engine/Plugins/Experimental/PCG/Source/PCG/Private/Elements/PCGSpawnActor.cpp
@@ -42,17 +42,34 @@ namespace PCGSpawnActorHelpers
 		FActorSingleOverride(const FPCGAttributePropertySelector& InputSelector, const FString& OutputProperty, AActor* TemplateActor, const UPCGPointData* PointData)
 		{
 			ActorOverrideInputAccessor = PCGAttributeAccessorHelpers::CreateConstAccessor(PointData, InputSelector);
-			ActorOverrideOutputAccessor = PCGAttributeAccessorHelpers::CreatePropertyAccessor(FName(OutputProperty), TemplateActor->GetClass());
+			FPCGAttributePropertySelector OutputSelector = FPCGAttributePropertySelector::CreateSelectorFromString(OutputProperty);
+			const TArray<FString>& ExtraNames = OutputSelector.GetExtraNames();
+			if (ExtraNames.IsEmpty())
+			{
+				ActorOverrideOutputAccessor = PCGAttributeAccessorHelpers::CreatePropertyAccessor(FName(OutputProperty), TemplateActor->GetClass());
+			}
+			else
+			{
+				TArray<FName> PropertyNames;
+				PropertyNames.Reserve(ExtraNames.Num() + 1);
+				PropertyNames.Add(OutputSelector.GetAttributeName());
+				for (const FString& Name : ExtraNames)
+				{
+					PropertyNames.Add(FName(Name));
+				}
+
+				ActorOverrideOutputAccessor = PCGAttributeAccessorHelpers::CreatePropertyChainAccessor(PropertyNames, TemplateActor->GetClass());
+			}
 
 			if (!ActorOverrideInputAccessor.IsValid() || !ActorOverrideOutputAccessor.IsValid())
 			{
-				UE_LOG(LogPCG, Warning, TEXT("ActorOverride from input %s or output %s is invalid or unsupported. Will be skipped."), *InputSelector.GetName().ToString(), *OutputProperty);
+				UE_LOG(LogPCG, Warning, TEXT("ActorOverride from input %s or output %s is invalid or unsupported. Will be skipped."), *InputSelector.GetName().ToString(), *OutputSelector.GetDisplayText().ToString());
 				return;
 			}
 
 			if (!PCG::Private::IsBroadcastable(ActorOverrideInputAccessor->GetUnderlyingType(), ActorOverrideOutputAccessor->GetUnderlyingType()))
 			{
-				UE_LOG(LogPCG, Warning, TEXT("ActorOverride cannot set input %s to output %s. Types are incompatibles. Will be skipped."), *InputSelector.GetName().ToString(), *OutputProperty);
+				UE_LOG(LogPCG, Warning, TEXT("ActorOverride cannot set input %s to output %s. Types are incompatibles. Will be skipped."), *InputSelector.GetName().ToString(), *OutputSelector.GetDisplayText().ToString());
 				ActorOverrideInputAccessor.Reset();
 				ActorOverrideOutputAccessor.Reset();
 				return;
diff --git a/Engine/Plugins/Experimental/PCG/Source/PCG/Private/Metadata/Accessors/PCGAttributeAccessorHelpers.cpp b/Engine/Plugins/Experimental/PCG/Source/PCG/Private/Metadata/Accessors/PCGAttributeAccessorHelpers.cpp
--- a/Engine/Plugins/Experimental/PCG/Source/PCG/Private/Metadata/Accessors/PCGAttributeAccessorHelpers.cpp
+++ b/Engine/Plugins/Experimental/PCG/Source/PCG/Private/Metadata/Accessors/PCGAttributeAccessorHelpers.cpp
@@ -14,6 +14,7 @@
 #include "Metadata/Accessors/PCGPropertyAccessor.h"
 #include "Metadata/Accessors/PCGAttributeExtractor.h"
 
+#include "Engine/UserDefinedStruct.h"
 #include "UObject/EnumProperty.h"
 
 namespace PCGAttributeAccessorHelpers
@@ -348,12 +349,12 @@ namespace PCGAttributeAccessorHelpers
 		return ReturnType{};
 	}
 
-	template <typename ClassType, typename Func>
-	decltype(auto) DispatchPropertyTypes(const FName InPropertyName, const ClassType* InClass, Func&& Functor)
+	template <typename Func>
+	decltype(auto) DispatchPropertyTypes(const FName InPropertyName, const UStruct* InStruct, Func&& Functor)
 	{
-		if (InClass)
+		if (InStruct)
 		{
-			if (const FProperty* Property = InClass->FindPropertyByName(InPropertyName))
+			if (const FProperty* Property = InStruct->FindPropertyByName(InPropertyName))
 			{
 				return DispatchPropertyTypes(Property, std::forward<Func>(Functor));
 			}
@@ -362,6 +363,66 @@ namespace PCGAttributeAccessorHelpers
 		using ReturnType = decltype(Functor(static_cast<FPCGPropertyPathAccessor<FSoftObjectPath>*>(nullptr), static_cast<const FProperty*>(nullptr)));
 		return ReturnType{};
 	}
+
+	bool GetPropertyChain(const TArray<FName>& InPropertyNames, const UStruct* InStruct, TArray<const FProperty*>& OutProperties)
+	{
+		check(InStruct);
+
+		const UStruct* CurrentStruct = InStruct;
+		OutProperties.Reserve(InPropertyNames.Num());
+
+		for (int32 i = 0; i < InPropertyNames.Num(); ++i)
+		{
+			const FProperty* Property = nullptr;
+			const FName PropertyName = InPropertyNames[i];
+
+			// Try to get the property. If it is coming from a user struct, we need to iterate on all properties because the property name is mangled
+			if (const UUserDefinedStruct* UserDefinedStruct = Cast<UUserDefinedStruct>(CurrentStruct))
+			{
+				for (TFieldIterator<const FProperty> PropIt(UserDefinedStruct, EFieldIterationFlags::IncludeSuper); PropIt; ++PropIt)
+				{
+					const FName TempPropertyName = *UserDefinedStruct->GetAuthoredNameForField(*PropIt);
+					if (TempPropertyName == PropertyName)
+					{
+						Property = *PropIt;
+						break;
+					}
+				}
+			}
+			else
+			{
+				Property = FindFProperty<FProperty>(CurrentStruct, PropertyName);
+			}
+
+			if (!Property)
+			{
+				UE_LOG(LogPCG, Error, TEXT("Property '%s' does not exist in %s."), *PropertyName.ToString(), *CurrentStruct->GetName());
+				return false;
+			}
+
+			OutProperties.Add(Property);
+
+			// Check for a struct or object for all properties except the last one.
+			if (i < InPropertyNames.Num() - 1)
+			{
+				if (const FStructProperty* StructProperty = CastField<FStructProperty>(Property))
+				{
+					CurrentStruct = StructProperty->Struct;
+				}
+				else if (const FObjectProperty* ObjectProperty = CastField<FObjectProperty>(Property))
+				{
+					CurrentStruct = ObjectProperty->PropertyClass;
+				}
+				else
+				{
+					UE_LOG(LogPCG, Error, TEXT("Property '%s' does exist in % s, but is not extractable."), *PropertyName.ToString(), *CurrentStruct->GetName());
+					return false;
+				}
+			}
+		}
+
+		return true;
+	}
 }
 
 TUniquePtr<IPCGAttributeAccessor> PCGAttributeAccessorHelpers::CreatePropertyAccessor(const FProperty* InProperty)
@@ -373,35 +434,45 @@ TUniquePtr<IPCGAttributeAccessor> PCGAttributeAccessorHelpers::CreatePropertyAcc
 	});
 }
 
-TUniquePtr<IPCGAttributeAccessor> PCGAttributeAccessorHelpers::CreatePropertyAccessor(const FName InPropertyName, const UClass* InClass)
+TUniquePtr<IPCGAttributeAccessor> PCGAttributeAccessorHelpers::CreatePropertyAccessor(const FName InPropertyName, const UStruct* InStruct)
 {
-	return PCGAttributeAccessorHelpers::DispatchPropertyTypes(InPropertyName, InClass, [](auto SignatureDummy, const auto* TypedProperty) -> TUniquePtr<IPCGAttributeAccessor>
+	return PCGAttributeAccessorHelpers::DispatchPropertyTypes(InPropertyName, InStruct, [](auto SignatureDummy, const auto* TypedProperty) -> TUniquePtr<IPCGAttributeAccessor>
 	{
 		using TypedAccessor = typename decltype(SignatureDummy)::Type;
 		return MakeUnique<TypedAccessor>(TypedProperty);
 	});
 }
 
-TUniquePtr<IPCGAttributeAccessor> PCGAttributeAccessorHelpers::CreatePropertyAccessor(const FName InPropertyName, const UStruct* InStruct)
+TUniquePtr<IPCGAttributeAccessor> PCGAttributeAccessorHelpers::CreatePropertyChainAccessor(TArray<const FProperty*>&& InProperties)
 {
-	return PCGAttributeAccessorHelpers::DispatchPropertyTypes(InPropertyName, InStruct, [](auto SignatureDummy, const auto* TypedProperty) -> TUniquePtr<IPCGAttributeAccessor>
+	if (InProperties.IsEmpty())
+	{
+		return TUniquePtr<IPCGAttributeAccessor>{};
+	}
+
+	return PCGAttributeAccessorHelpers::DispatchPropertyTypes(InProperties.Last(), [&InProperties](auto SignatureDummy, const auto* TypedProperty) -> TUniquePtr<IPCGAttributeAccessor>
 	{
 		using TypedAccessor = typename decltype(SignatureDummy)::Type;
-		return MakeUnique<TypedAccessor>(TypedProperty);
+		return MakeUnique<TypedAccessor>(TypedProperty, std::forward<TArray<const FProperty*>>(InProperties));
 	});
 }
 
-bool PCGAttributeAccessorHelpers::IsPropertyAccessorSupported(const FProperty* InProperty)
+TUniquePtr<IPCGAttributeAccessor> PCGAttributeAccessorHelpers::CreatePropertyChainAccessor(const TArray<FName>& InPropertyNames, const UStruct* InStruct)
 {
-	return PCGAttributeAccessorHelpers::DispatchPropertyTypes(InProperty, [](auto, const auto*) -> bool
+	TArray<const FProperty*> PropertyChain;
+	if (!PCGAttributeAccessorHelpers::GetPropertyChain(InPropertyNames, InStruct, PropertyChain))
 	{
-		return true;
-	});
+		return TUniquePtr<IPCGAttributeAccessor>{};
+	}
+	else
+	{
+		return PCGAttributeAccessorHelpers::CreatePropertyChainAccessor(std::move(PropertyChain));
+	}
 }
 
-bool PCGAttributeAccessorHelpers::IsPropertyAccessorSupported(const FName InPropertyName, const UClass* InClass)
+bool PCGAttributeAccessorHelpers::IsPropertyAccessorSupported(const FProperty* InProperty)
 {
-	return PCGAttributeAccessorHelpers::DispatchPropertyTypes(InPropertyName, InClass, [](auto, const auto*) -> bool
+	return PCGAttributeAccessorHelpers::DispatchPropertyTypes(InProperty, [](auto, const auto*) -> bool
 	{
 		return true;
 	});
@@ -415,6 +486,14 @@ bool PCGAttributeAccessorHelpers::IsPropertyAccessorSupported(const FName InProp
 	});
 }
 
+bool PCGAttributeAccessorHelpers::IsPropertyAccessorChainSupported(const TArray<FName>& InPropertyNames, const UStruct* InStruct)
+{
+	TArray<const FProperty*> PropertyChain;
+	return PCGAttributeAccessorHelpers::GetPropertyChain(InPropertyNames, InStruct, PropertyChain) &&
+		!PropertyChain.IsEmpty() &&
+		PCGAttributeAccessorHelpers::IsPropertyAccessorSupported(PropertyChain.Last());
+}
+
 TUniquePtr<IPCGAttributeAccessor> PCGAttributeAccessorHelpers::CreateExtraAccessor(EPCGExtraProperties InExtraProperties)
 {
 	switch (InExtraProperties)
diff --git a/Engine/Plugins/Experimental/PCG/Source/PCG/Private/Metadata/Accessors/PCGPropertyAccessor.cpp b/Engine/Plugins/Experimental/PCG/Source/PCG/Private/Metadata/Accessors/PCGPropertyAccessor.cpp
--- a/Engine/Plugins/Experimental/PCG/Source/PCG/Private/Metadata/Accessors/PCGPropertyAccessor.cpp
+++ b/Engine/Plugins/Experimental/PCG/Source/PCG/Private/Metadata/Accessors/PCGPropertyAccessor.cpp
@@ -30,9 +30,19 @@ TArray<void*> PCGPropertyAccessor::GetContainerKeys(int32 Index, int32 Range, IP
 	return ContainerKeys;
 }
 
+IPCGPropertyChainAccessor::IPCGPropertyChainAccessor(const FProperty* Property, TArray<const FProperty*>&& ExtraProperties)
+	: PropertyChain(std::forward<TArray<const FProperty*>>(ExtraProperties))
+{
+	// Fix property chain
+	if (PropertyChain.IsEmpty() || PropertyChain.Last() != Property)
+	{
+		PropertyChain.Add(Property);
+	}
+}
+
 bool FPCGEnumPropertyAccessor::GetRangeImpl(TArrayView<int64> OutValues, int32 Index, const IPCGAttributeAccessorKeys& Keys) const
 {
-	return PCGPropertyAccessor::IterateGet(Property, OutValues, Index, Keys, [this](const void* PropertyAddressData) -> Type
+	return PCGPropertyAccessor::IterateGet(GetPropertyChain(), OutValues, Index, Keys, [this](const void* PropertyAddressData) -> Type
 		{
 			return Property->GetUnderlyingProperty()->GetSignedIntPropertyValue(PropertyAddressData);
 		});
@@ -40,7 +50,7 @@ bool FPCGEnumPropertyAccessor::GetRangeImpl(TArrayView<int64> OutValues, int32 I
 
 bool FPCGEnumPropertyAccessor::SetRangeImpl(TArrayView<const int64> InValues, int32 Index, IPCGAttributeAccessorKeys& Keys, EPCGAttributeAccessorFlags)
 {
-	return PCGPropertyAccessor::IterateSet(Property, InValues, Index, Keys, [this](void* PropertyAddressData, const int64& Value) -> void
+	return PCGPropertyAccessor::IterateSet(GetPropertyChain(), InValues, Index, Keys, [this](void* PropertyAddressData, const int64& Value) -> void
 		{
 			Property->GetUnderlyingProperty()->SetIntPropertyValue(PropertyAddressData, Value);
 		});
diff --git a/Engine/Plugins/Experimental/PCG/Source/PCG/Public/Metadata/Accessors/PCGAttributeAccessorHelpers.h b/Engine/Plugins/Experimental/PCG/Source/PCG/Public/Metadata/Accessors/PCGAttributeAccessorHelpers.h
--- a/Engine/Plugins/Experimental/PCG/Source/PCG/Public/Metadata/Accessors/PCGAttributeAccessorHelpers.h
+++ b/Engine/Plugins/Experimental/PCG/Source/PCG/Public/Metadata/Accessors/PCGAttributeAccessorHelpers.h
@@ -25,13 +25,15 @@ struct FPCGSettingsOverridableParam;
 namespace PCGAttributeAccessorHelpers
 {
 	PCG_API bool IsPropertyAccessorSupported(const FProperty* InProperty);
-	PCG_API bool IsPropertyAccessorSupported(const FName InPropertyName, const UClass* InClass);
 	PCG_API bool IsPropertyAccessorSupported(const FName InPropertyName, const UStruct* InStruct);
+	PCG_API bool IsPropertyAccessorChainSupported(const TArray<FName>& InPropertyNames, const UStruct* InStruct);
 
 	PCG_API TUniquePtr<IPCGAttributeAccessor> CreatePropertyAccessor(const FProperty* InProperty);
-	PCG_API TUniquePtr<IPCGAttributeAccessor> CreatePropertyAccessor(const FName InPropertyName, const UClass* InClass);
 	PCG_API TUniquePtr<IPCGAttributeAccessor> CreatePropertyAccessor(const FName InPropertyName, const UStruct* InStruct);
 
+	PCG_API TUniquePtr<IPCGAttributeAccessor> CreatePropertyChainAccessor(TArray<const FProperty*>&& InProperties);
+	PCG_API TUniquePtr<IPCGAttributeAccessor> CreatePropertyChainAccessor(const TArray<FName>& InPropertyNames, const UStruct* InStruct);
+
 	PCG_API TUniquePtr<IPCGAttributeAccessor> CreateExtraAccessor(EPCGExtraProperties InExtraProperties);
 
 	PCG_API TUniquePtr<IPCGAttributeAccessor> CreateChainAccessor(TUniquePtr<IPCGAttributeAccessor> InAccessor, FName Name, bool& bOutSuccess);
diff --git a/Engine/Plugins/Experimental/PCG/Source/PCG/Public/Metadata/Accessors/PCGPropertyAccessor.h b/Engine/Plugins/Experimental/PCG/Source/PCG/Public/Metadata/Accessors/PCGPropertyAccessor.h
--- a/Engine/Plugins/Experimental/PCG/Source/PCG/Public/Metadata/Accessors/PCGPropertyAccessor.h
+++ b/Engine/Plugins/Experimental/PCG/Source/PCG/Public/Metadata/Accessors/PCGPropertyAccessor.h
@@ -17,8 +17,34 @@ namespace PCGPropertyAccessor
 	TArray<const void*> GetContainerKeys(int32 Index, int32 Range, const IPCGAttributeAccessorKeys& Keys);
 	TArray<void*> GetContainerKeys(int32 Index, int32 Range, IPCGAttributeAccessorKeys& Keys);
 
+	template <typename T>
+	void AddressOffset(const TArray<const FProperty*>& InProperties, TArray<T>& InContainerKeys)
+	{
+		for (int32 j = 0; j < InProperties.Num(); ++j)
+		{
+			const FProperty* Property = InProperties[j];
+			check(Property);
+			// No indirection for last property
+			const FObjectProperty* ObjectProperty = (j < InProperties.Num() - 1) ? CastField<const FObjectProperty>(Property) : nullptr;
+			if (ObjectProperty)
+			{
+				for (int32 i = 0; i < InContainerKeys.Num(); ++i)
+				{
+					InContainerKeys[i] = ObjectProperty->GetObjectPropertyValue_InContainer(InContainerKeys[i]);
+				}
+			}
+			else
+			{
+				for (int32 i = 0; i < InContainerKeys.Num(); ++i)
+				{
+					InContainerKeys[i] = Property->ContainerPtrToValuePtr<void>(InContainerKeys[i]);
+				}
+			}
+		}
+	}
+
 	template <typename T, typename Func>
-	bool IterateGet(const FProperty* Property, TArrayView<T>& OutValues, int32 Index, const IPCGAttributeAccessorKeys& Keys, Func&& Getter)
+	bool IterateGet(const TArray<const FProperty*>& Properties, TArrayView<T>& OutValues, int32 Index, const IPCGAttributeAccessorKeys& Keys, Func&& Getter)
 	{
 		TArray<const void*> ContainerKeys = GetContainerKeys(Index, OutValues.Num(), Keys);
 		if (ContainerKeys.IsEmpty())
@@ -26,17 +52,19 @@ namespace PCGPropertyAccessor
 			return false;
 		}
 
+		// Update the addresses of all
+		AddressOffset<const void*>(Properties, ContainerKeys);
+
 		for (int32 i = 0; i < OutValues.Num(); ++i)
 		{
-			const void* PropertyAddressData = Property->ContainerPtrToValuePtr<void>(ContainerKeys[i]);
-			OutValues[i] = Getter(PropertyAddressData);
+			OutValues[i] = Getter(ContainerKeys[i]);
 		}
 
 		return true;
 	}
 
 	template <typename T, typename Func>
-	bool IterateSet(const FProperty* Property, TArrayView<const T>& InValues, int32 Index, IPCGAttributeAccessorKeys& Keys, Func&& Setter)
+	bool IterateSet(const TArray<const FProperty*>& Properties, TArrayView<const T>& InValues, int32 Index, IPCGAttributeAccessorKeys& Keys, Func&& Setter)
 	{
 		TArray<void*> ContainerKeys = GetContainerKeys(Index, InValues.Num(), Keys);
 		if (ContainerKeys.IsEmpty())
@@ -44,16 +72,35 @@ namespace PCGPropertyAccessor
 			return false;
 		}
 
+		// Update the addresses of all
+		AddressOffset<void*>(Properties, ContainerKeys);
+
 		for (int32 i = 0; i < InValues.Num(); ++i)
 		{
-			void* PropertyAddressData = Property->ContainerPtrToValuePtr<void>(ContainerKeys[i]);
-			Setter(PropertyAddressData, InValues[i]);
+			Setter(ContainerKeys[i], InValues[i]);
 		}
 
 		return true;
 	}
 }
 
+/**
+* Interface for Property chain to factorize ctor, fix the chain and storing the property chain
+*/
+class IPCGPropertyChainAccessor
+{
+public:
+	virtual ~IPCGPropertyChainAccessor() = default;
+
+protected:
+	IPCGPropertyChainAccessor(const FProperty* Property, TArray<const FProperty*>&& ExtraProperties);
+
+	const TArray<const FProperty*>& GetPropertyChain() const { return PropertyChain; }
+
+private:
+	TArray<const FProperty*> PropertyChain;
+};
+
 /**
 * Templated accessor class for numeric properties. Will wrap around a numeric property.
 * Do not instanciate it manually, use PCGAttributeAccessorHelpers::CreatePropertyAccessor.
@@ -61,14 +108,15 @@ namespace PCGPropertyAccessor
 * Key supported: Generic object
 */
 template <typename T>
-class FPCGNumericPropertyAccessor : public IPCGAttributeAccessorT<FPCGNumericPropertyAccessor<T>>
+class FPCGNumericPropertyAccessor : public IPCGAttributeAccessorT<FPCGNumericPropertyAccessor<T>>, IPCGPropertyChainAccessor
 {
 public:
 	using Type = T;
 	using Super = IPCGAttributeAccessorT<FPCGNumericPropertyAccessor<T>>;
 
-	FPCGNumericPropertyAccessor(const FNumericProperty* InProperty)
+	FPCGNumericPropertyAccessor(const FNumericProperty* InProperty, TArray<const FProperty*>&& ExtraProperties = {})
 		: Super(/*bInReadOnly=*/ false)
+		, IPCGPropertyChainAccessor(InProperty, std::forward<TArray<const FProperty*>>(ExtraProperties))
 		, Property(InProperty)
 	{
 		static_assert(std::is_integral_v<T> || std::is_floating_point_v<T>);
@@ -77,7 +125,7 @@ public:
 
 	bool GetRangeImpl(TArrayView<T> OutValues, int32 Index, const IPCGAttributeAccessorKeys& Keys) const
 	{
-		return PCGPropertyAccessor::IterateGet(Property, OutValues, Index, Keys, [this](const void* PropertyAddressData) -> T
+		return PCGPropertyAccessor::IterateGet(GetPropertyChain(), OutValues, Index, Keys, [this](const void* PropertyAddressData) -> T
 			{
 				if constexpr (std::is_integral_v<T>)
 				{
@@ -92,7 +140,7 @@ public:
 
 	bool SetRangeImpl(TArrayView<const T> InValues, int32 Index, IPCGAttributeAccessorKeys& Keys, EPCGAttributeAccessorFlags)
 	{
-		return PCGPropertyAccessor::IterateSet(Property, InValues, Index, Keys, [this](void* PropertyAddressData, const T& Value) -> void
+		return PCGPropertyAccessor::IterateSet(GetPropertyChain(), InValues, Index, Keys, [this](void* PropertyAddressData, const T& Value) -> void
 			{
 				if constexpr (std::is_integral_v<T>)
 				{
@@ -115,14 +163,15 @@ private:
 * Do not instanciate it manually, use PCGAttributeAccessorHelpers::CreatePropertyAccessor.
 * Key supported: Generic object
 */
-class FPCGEnumPropertyAccessor : public IPCGAttributeAccessorT<FPCGEnumPropertyAccessor>
+class FPCGEnumPropertyAccessor : public IPCGAttributeAccessorT<FPCGEnumPropertyAccessor>, IPCGPropertyChainAccessor
 {
 public:
 	using Type = int64;
 	using Super = IPCGAttributeAccessorT<FPCGEnumPropertyAccessor>;
 
-	FPCGEnumPropertyAccessor(const FEnumProperty* InProperty)
+	FPCGEnumPropertyAccessor(const FEnumProperty* InProperty, TArray<const FProperty*>&& ExtraProperties = {})
 		: Super(/*bInReadOnly=*/ false)
+		, IPCGPropertyChainAccessor(InProperty, std::forward<TArray<const FProperty*>>(ExtraProperties))
 		, Property(InProperty)
 	{
 		check(Property);
@@ -143,14 +192,15 @@ private:
 * Key supported: Generic object
 */
 template <typename T>
-class FPCGPropertyStructAccessor : public IPCGAttributeAccessorT<FPCGPropertyStructAccessor<T>>
+class FPCGPropertyStructAccessor : public IPCGAttributeAccessorT<FPCGPropertyStructAccessor<T>>, IPCGPropertyChainAccessor
 {
 public:
 	using Type = T;
 	using Super = IPCGAttributeAccessorT<FPCGPropertyStructAccessor<T>>;
 
-	FPCGPropertyStructAccessor(const FStructProperty* InProperty)
+	FPCGPropertyStructAccessor(const FStructProperty* InProperty, TArray<const FProperty*>&& ExtraProperties = {})
 		: Super(/*bInReadOnly=*/ false)
+		, IPCGPropertyChainAccessor(InProperty, std::forward<TArray<const FProperty*>>(ExtraProperties))
 		, Property(InProperty)
 	{
 		static_assert(PCG::Private::IsPCGType<T>());
@@ -159,7 +209,7 @@ public:
 
 	bool GetRangeImpl(TArrayView<T> OutValues, int32 Index, const IPCGAttributeAccessorKeys& Keys) const
 	{
-		return PCGPropertyAccessor::IterateGet(Property, OutValues, Index, Keys, [this](const void* PropertyAddressData) -> T
+		return PCGPropertyAccessor::IterateGet(GetPropertyChain(), OutValues, Index, Keys, [this](const void* PropertyAddressData) -> T
 			{
 				return *reinterpret_cast<const T*>(PropertyAddressData);
 			});
@@ -167,7 +217,7 @@ public:
 
 	bool SetRangeImpl(TArrayView<const T> InValues, int32 Index, IPCGAttributeAccessorKeys& Keys, EPCGAttributeAccessorFlags)
 	{
-		return PCGPropertyAccessor::IterateSet(Property, InValues, Index, Keys, [this](void* PropertyAddressData, const T& Value) -> void
+		return PCGPropertyAccessor::IterateSet(GetPropertyChain(), InValues, Index, Keys, [this](void* PropertyAddressData, const T& Value) -> void
 			{
 				*reinterpret_cast<T*>(PropertyAddressData) = Value;
 			});
@@ -184,14 +234,15 @@ private:
 * Key supported: Generic object
 */
 template <typename T, typename PropertyType>
-class FPCGPropertyAccessor : public IPCGAttributeAccessorT<FPCGPropertyAccessor<T, PropertyType>>
+class FPCGPropertyAccessor : public IPCGAttributeAccessorT<FPCGPropertyAccessor<T, PropertyType>>, IPCGPropertyChainAccessor
 {
 public:
 	using Type = T;
 	using Super = IPCGAttributeAccessorT<FPCGPropertyAccessor<T, PropertyType>>;
 
-	FPCGPropertyAccessor(const PropertyType* InProperty)
+	FPCGPropertyAccessor(const PropertyType* InProperty, TArray<const FProperty*>&& ExtraProperties = {})
 		: Super(/*bInReadOnly=*/ false)
+		, IPCGPropertyChainAccessor(InProperty, std::forward<TArray<const FProperty*>>(ExtraProperties))
 		, Property(InProperty)
 	{
 		static_assert(PCG::Private::IsPCGType<T>());
@@ -200,7 +251,7 @@ public:
 
 	bool GetRangeImpl(TArrayView<T> OutValues, int32 Index, const IPCGAttributeAccessorKeys& Keys) const
 	{
-		return PCGPropertyAccessor::IterateGet(Property, OutValues, Index, Keys, [this](const void* PropertyAddressData) -> T
+		return PCGPropertyAccessor::IterateGet(GetPropertyChain(), OutValues, Index, Keys, [this](const void* PropertyAddressData) -> T
 			{
 				return Property->GetPropertyValue(PropertyAddressData);
 			});
@@ -208,7 +259,7 @@ public:
 
 	bool SetRangeImpl(TArrayView<const T> InValues, int32 Index, IPCGAttributeAccessorKeys& Keys, EPCGAttributeAccessorFlags)
 	{
-		return PCGPropertyAccessor::IterateSet(Property, InValues, Index, Keys, [this](void* PropertyAddressData, const T& Value) -> void
+		return PCGPropertyAccessor::IterateSet(GetPropertyChain(), InValues, Index, Keys, [this](void* PropertyAddressData, const T& Value) -> void
 			{
 				Property->SetPropertyValue(PropertyAddressData, Value);
 			});
@@ -225,14 +276,15 @@ private:
 * Key supported: Generic object
 */
 template <typename T>
-class FPCGPropertyPathAccessor : public IPCGAttributeAccessorT<FPCGPropertyPathAccessor<T>>
+class FPCGPropertyPathAccessor : public IPCGAttributeAccessorT<FPCGPropertyPathAccessor<T>>, IPCGPropertyChainAccessor
 {
 public:
 	using Type = FString;
 	using Super = IPCGAttributeAccessorT<FPCGPropertyPathAccessor<T>>;
 
-	FPCGPropertyPathAccessor(const FProperty* InProperty)
+	FPCGPropertyPathAccessor(const FProperty* InProperty, TArray<const FProperty*>&& ExtraProperties = {})
 		: Super(/*bInReadOnly=*/ false)
+		, IPCGPropertyChainAccessor(InProperty, std::forward<TArray<const FProperty*>>(ExtraProperties))
 		, Property(InProperty)
 	{
 		static_assert(std::is_same_v<FSoftObjectPath, T> || std::is_same_v<FSoftClassPath, T>);
@@ -241,7 +293,7 @@ public:
 
 	bool GetRangeImpl(TArrayView<Type> OutValues, int32 Index, const IPCGAttributeAccessorKeys& Keys) const
 	{
-		return PCGPropertyAccessor::IterateGet(Property, OutValues, Index, Keys, [this](const void* PropertyAddressData) -> Type
+		return PCGPropertyAccessor::IterateGet(GetPropertyChain(), OutValues, Index, Keys, [this](const void* PropertyAddressData) -> Type
 			{
 				return reinterpret_cast<const T*>(PropertyAddressData)->ToString();
 			});
@@ -249,7 +301,7 @@ public:
 
 	bool SetRangeImpl(TArrayView<const Type> InValues, int32 Index, IPCGAttributeAccessorKeys& Keys, EPCGAttributeAccessorFlags)
 	{
-		return PCGPropertyAccessor::IterateSet(Property, InValues, Index, Keys, [this](void* PropertyAddressData, const Type& Value) -> void
+		return PCGPropertyAccessor::IterateSet(GetPropertyChain(), InValues, Index, Keys, [this](void* PropertyAddressData, const Type& Value) -> void
 			{
 				reinterpret_cast<T*>(PropertyAddressData)->SetPath(Value);
 			});
@@ -264,14 +316,15 @@ private:
 * Do not instantiate it manually, use PCGAttributeAccessorHelpers::CreatePropertyAccessor.
 * Key supported: Generic object
 */
-class FPCGPropertySoftObjectPathAccessor : public IPCGAttributeAccessorT<FPCGPropertySoftObjectPathAccessor>
+class FPCGPropertySoftObjectPathAccessor : public IPCGAttributeAccessorT<FPCGPropertySoftObjectPathAccessor>, IPCGPropertyChainAccessor
 {
 public:
 	using Type = FSoftObjectPath;
 	using Super = IPCGAttributeAccessorT<FPCGPropertySoftObjectPathAccessor>;
 
-	FPCGPropertySoftObjectPathAccessor(const FSoftObjectProperty* InProperty)
+	FPCGPropertySoftObjectPathAccessor(const FSoftObjectProperty* InProperty, TArray<const FProperty*>&& ExtraProperties = {})
 		: Super(/*bInReadOnly=*/ false)
+		, IPCGPropertyChainAccessor(InProperty, std::forward<TArray<const FProperty*>>(ExtraProperties))
 		, Property(InProperty)
 	{
 		check(Property);
@@ -279,7 +332,7 @@ public:
 
 	bool GetRangeImpl(TArrayView<Type> OutValues, int32 Index, const IPCGAttributeAccessorKeys& Keys) const
 	{
-		return PCGPropertyAccessor::IterateGet(Property, OutValues, Index, Keys, [this](const void* PropertyAddressData) -> Type
+		return PCGPropertyAccessor::IterateGet(GetPropertyChain(), OutValues, Index, Keys, [this](const void* PropertyAddressData) -> Type
 		{
 			return Property->GetPropertyValue(PropertyAddressData).ToSoftObjectPath();
 		});
@@ -287,7 +340,7 @@ public:
 
 	bool SetRangeImpl(TArrayView<const Type> InValues, int32 Index, IPCGAttributeAccessorKeys& Keys, EPCGAttributeAccessorFlags)
 	{
-		return PCGPropertyAccessor::IterateSet(Property, InValues, Index, Keys, [this](void* PropertyAddressData, const Type& Value) -> void
+		return PCGPropertyAccessor::IterateSet(GetPropertyChain(), InValues, Index, Keys, [this](void* PropertyAddressData, const Type& Value) -> void
 		{
 			Property->SetPropertyValue(PropertyAddressData, FSoftObjectPtr(Value));
 		});
@@ -302,14 +355,15 @@ private:
 * Do not instantiate it manually, use PCGAttributeAccessorHelpers::CreatePropertyAccessor.
 * Key supported: Generic object
 */
-class FPCGPropertySoftClassPathAccessor : public IPCGAttributeAccessorT<FPCGPropertySoftClassPathAccessor>
+class FPCGPropertySoftClassPathAccessor : public IPCGAttributeAccessorT<FPCGPropertySoftClassPathAccessor>, IPCGPropertyChainAccessor
 {
 public:
 	using Type = FSoftClassPath;
 	using Super = IPCGAttributeAccessorT<FPCGPropertySoftClassPathAccessor>;
 
-	FPCGPropertySoftClassPathAccessor(const FSoftClassProperty* InProperty)
+	FPCGPropertySoftClassPathAccessor(const FSoftClassProperty* InProperty, TArray<const FProperty*>&& ExtraProperties = {})
 		: Super(/*bInReadOnly=*/ false)
+		, IPCGPropertyChainAccessor(InProperty, std::forward<TArray<const FProperty*>>(ExtraProperties))
 		, Property(InProperty)
 	{
 		check(Property);
@@ -317,7 +371,7 @@ public:
 
 	bool GetRangeImpl(TArrayView<Type> OutValues, int32 Index, const IPCGAttributeAccessorKeys& Keys) const
 	{
-		return PCGPropertyAccessor::IterateGet(Property, OutValues, Index, Keys, [this](const void* PropertyAddressData) -> Type
+		return PCGPropertyAccessor::IterateGet(GetPropertyChain(), OutValues, Index, Keys, [this](const void* PropertyAddressData) -> Type
 		{
 			return FSoftClassPath(Property->GetPropertyValue(PropertyAddressData).ToString());
 		});
@@ -325,7 +379,7 @@ public:
 
 	bool SetRangeImpl(TArrayView<const Type> InValues, int32 Index, IPCGAttributeAccessorKeys& Keys, EPCGAttributeAccessorFlags)
 	{
-		return PCGPropertyAccessor::IterateSet(Property, InValues, Index, Keys, [this](void* PropertyAddressData, const Type& Value) -> void
+		return PCGPropertyAccessor::IterateSet(GetPropertyChain(), InValues, Index, Keys, [this](void* PropertyAddressData, const Type& Value) -> void
 		{
 			Property->SetPropertyValue(PropertyAddressData, FSoftObjectPtr(Value));
 		});
@@ -341,14 +395,15 @@ private:
 * Will always convert to FString for PCG
 * Key supported: Generic object
 */
-class FPCGPropertyObjectPtrAccessor : public IPCGAttributeAccessorT<FPCGPropertyObjectPtrAccessor>
+class FPCGPropertyObjectPtrAccessor : public IPCGAttributeAccessorT<FPCGPropertyObjectPtrAccessor>, IPCGPropertyChainAccessor
 {
 public:
 	using Type = FString;
 	using Super = IPCGAttributeAccessorT<FPCGPropertyObjectPtrAccessor>;
 
-	FPCGPropertyObjectPtrAccessor(const FObjectProperty* InProperty)
+	FPCGPropertyObjectPtrAccessor(const FObjectProperty* InProperty, TArray<const FProperty*>&& ExtraProperties = {})
 		: Super(/*bInReadOnly=*/ false)
+		, IPCGPropertyChainAccessor(InProperty, std::forward<TArray<const FProperty*>>(ExtraProperties))
 		, Property(InProperty)
 	{
 		check(Property);
@@ -356,7 +411,7 @@ public:
 
 	bool GetRangeImpl(TArrayView<Type> OutValues, int32 Index, const IPCGAttributeAccessorKeys& Keys) const
 	{
-		return PCGPropertyAccessor::IterateGet(Property, OutValues, Index, Keys, [this](const void* PropertyAddressData) -> Type
+		return PCGPropertyAccessor::IterateGet(GetPropertyChain(), OutValues, Index, Keys, [this](const void* PropertyAddressData) -> Type
 		{
 			return Property->GetPropertyValue(PropertyAddressData).GetPath();
 		});
@@ -364,7 +419,7 @@ public:
 
 	bool SetRangeImpl(TArrayView<const Type> InValues, int32 Index, IPCGAttributeAccessorKeys& Keys, EPCGAttributeAccessorFlags)
 	{
-		return PCGPropertyAccessor::IterateSet(Property, InValues, Index, Keys, [this](void* PropertyAddressData, const Type& Value) -> void
+		return PCGPropertyAccessor::IterateSet(GetPropertyChain(), InValues, Index, Keys, [this](void* PropertyAddressData, const Type& Value) -> void
 		{
 			Property->SetPropertyValue(PropertyAddressData, FSoftObjectPath(Value).TryLoad());
 		});
@@ -378,14 +433,15 @@ private:
 * Special accessor to support attribute selector overrides. Interface with a string.
 * Key supported: All
 */
-class FPCGAttributePropertySelectorAccessor : public IPCGAttributeAccessorT<FPCGAttributePropertySelectorAccessor>
+class FPCGAttributePropertySelectorAccessor : public IPCGAttributeAccessorT<FPCGAttributePropertySelectorAccessor>, IPCGPropertyChainAccessor
 {
 public:
 	using Type = FString;
 	using Super = IPCGAttributeAccessorT<FPCGAttributePropertySelectorAccessor>;
 
-	FPCGAttributePropertySelectorAccessor(const FStructProperty* InProperty)
+	FPCGAttributePropertySelectorAccessor(const FStructProperty* InProperty, TArray<const FProperty*>&& ExtraProperties = {})
 		: Super(/*bInReadOnly=*/ false)
+		, IPCGPropertyChainAccessor(InProperty, std::forward<TArray<const FProperty*>>(ExtraProperties))
 		, Property(InProperty)
 	{
 		ensure(InProperty && InProperty->Struct && InProperty->Struct->IsA(FPCGAttributePropertySelector::StaticStruct()->GetClass()));
@@ -393,7 +449,7 @@ public:
 
 	bool GetRangeImpl(TArrayView<FString> OutValues, int32 Index, const IPCGAttributeAccessorKeys& Keys) const
 	{
-		return PCGPropertyAccessor::IterateGet(Property, OutValues, Index, Keys, [this](const void* PropertyAddressData) -> FString
+		return PCGPropertyAccessor::IterateGet(GetPropertyChain(), OutValues, Index, Keys, [this](const void* PropertyAddressData) -> FString
 		{
 			return reinterpret_cast<const FPCGAttributePropertySelector*>(PropertyAddressData)->GetDisplayText().ToString();
 		});
@@ -401,7 +457,7 @@ public:
 
 	bool SetRangeImpl(TArrayView<const FString> InValues, int32 Index, IPCGAttributeAccessorKeys& Keys, EPCGAttributeAccessorFlags Flags)
 	{
-		return PCGPropertyAccessor::IterateSet(Property, InValues, Index, Keys, [this](void* PropertyAddressData, const FString& Value) -> void
+		return PCGPropertyAccessor::IterateSet(GetPropertyChain(), InValues, Index, Keys, [this](void* PropertyAddressData, const FString& Value) -> void
 		{
 			reinterpret_cast<FPCGAttributePropertySelector*>(PropertyAddressData)->Update(Value);
 		});
