[PCG] Fix localization

#rnx

diff --git a/Engine/Plugins/Experimental/PCG/Source/PCGEditor/Private/PCGEditorModule.cpp b/Engine/Plugins/Experimental/PCG/Source/PCGEditor/Private/PCGEditorModule.cpp
--- a/Engine/Plugins/Experimental/PCG/Source/PCGEditor/Private/PCGEditorModule.cpp
+++ b/Engine/Plugins/Experimental/PCG/Source/PCGEditor/Private/PCGEditorModule.cpp
@@ -175,8 +175,8 @@ void FPCGEditorModule::PopulateMenuActions(FMenuBuilder& MenuBuilder)
 		NAME_None);	
 	
 	MenuBuilder.AddMenuEntry(
-		LOCTEXT("DeletePCGPartitionActors", "Delete all PCG partition actors children"),
-		LOCTEXT("DeletePCGPartitionActors_Tooltip", "Deletes all PCG partition actors children in the current world, but not the Partition Actors themselves"),
+		LOCTEXT("DeletePCGPartitionActorsChildren", "Delete all PCG partition actors children"),
+		LOCTEXT("DeletePCGPartitionActorsChildren_Tooltip", "Deletes all PCG partition actors children in the current world, but not the Partition Actors themselves"),
 		FSlateIcon(),
 		FUIAction(
 			FExecuteAction::CreateLambda([]() {
