[PCG] Added native and blueprint element templated pins to the Proxy node

diff --git a/Engine/Plugins/Experimental/PCG/Source/PCG/Private/Elements/PCGIndirectionElement.cpp b/Engine/Plugins/Experimental/PCG/Source/PCG/Private/Elements/PCGIndirectionElement.cpp
--- a/Engine/Plugins/Experimental/PCG/Source/PCG/Private/Elements/PCGIndirectionElement.cpp
+++ b/Engine/Plugins/Experimental/PCG/Source/PCG/Private/Elements/PCGIndirectionElement.cpp
@@ -3,6 +3,9 @@
 #include "Elements/PCGIndirectionElement.h"
 
 #include "PCGPin.h"
+#include "PCGSettings.h"
+#include "Elements/PCGExecuteBlueprint.h" // Blueprint element class
+#include "Helpers/PCGSettingsHelpers.h"   // Graph and log errors
 
 #include UE_INLINE_GENERATED_CPP_BY_NAME(PCGIndirectionElement)
 
@@ -25,14 +28,114 @@ FText UPCGIndirectionSettings::GetNodeTooltipText() const
 	return LOCTEXT("NodeTooltip", "Executes another settings object, which can be overriden.");
 }
 
+FName UPCGIndirectionSettings::AdditionalTaskName() const
+{
+	switch (ProxyInterfaceMode)
+	{
+		case EPCGProxyInterfaceMode::ByNativeElement:
+			if (IsValid(SettingsClass))
+			{
+				const UPCGSettings* SettingsDefaultObject = CastChecked<UPCGSettings>(SettingsClass->GetDefaultObject());
+				return FName(FText::Format(LOCTEXT("NodeTitleNative", "Proxy ({0})"), SettingsDefaultObject->GetDefaultNodeTitle()).ToString());
+			}
+			break;
+		case EPCGProxyInterfaceMode::ByBlueprintElement:
+			if (IsValid(BlueprintElementClass))
+			{
+				return FName(FText::Format(LOCTEXT("NodeTitleBlueprint", "Proxy ({0})"), BlueprintElementClass->GetDisplayNameText()).ToString());
+			}
+			break;
+		case EPCGProxyInterfaceMode::BySettings:
+			if (const UPCGSettings* SettingsPtr = Settings.LoadSynchronous())
+			{
+				return FName(FText::Format(LOCTEXT("NodeTitleNative", "Proxy ({0})"), FText::FromString(SettingsPtr->GetName())).ToString());
+			}
+			break;
+
+		default:
+			checkNoEntry();
+	}
+
+	return GetDefaultNodeName();
+}
 #endif // WITH_EDITOR
 
+TArray<FPCGPinProperties> UPCGIndirectionSettings::InputPinProperties() const
+{
+	TArray<FPCGPinProperties> InputProperties;
+
+	switch (ProxyInterfaceMode)
+	{
+		case EPCGProxyInterfaceMode::ByNativeElement:
+			if (IsValid(SettingsClass))
+			{
+				const UPCGSettings* SettingsDefaultObject = CastChecked<UPCGSettings>(SettingsClass->GetDefaultObject());
+				InputProperties = SettingsDefaultObject->DefaultInputPinProperties();
+			}
+			break;
+		case EPCGProxyInterfaceMode::ByBlueprintElement:
+			if (IsValid(BlueprintElementClass))
+			{
+				const UPCGBlueprintElement* BlueprintSettingsDefaultObject = CastChecked<UPCGBlueprintElement>(BlueprintElementClass->GetDefaultObject());
+				InputProperties = BlueprintSettingsDefaultObject->GetInputPins();
+			}
+			break;
+		case EPCGProxyInterfaceMode::BySettings:
+			if (const UPCGSettings* SettingsPtr = Settings.LoadSynchronous())
+			{
+				InputProperties = SettingsPtr->DefaultInputPinProperties();
+			}
+			break;
+
+		default:
+			checkNoEntry();
+	}
+
+	if (InputProperties.IsEmpty())
+	{
+		return Super::InputPinProperties();
+	}
+
+	return InputProperties;
+}
+
 TArray<FPCGPinProperties> UPCGIndirectionSettings::OutputPinProperties() const
 {
-	TArray<FPCGPinProperties> PinProperties;
-	PinProperties.Emplace(PCGPinConstants::DefaultOutputLabel, EPCGDataType::Any);
+	TArray<FPCGPinProperties> OutputProperties;
+
+	switch (ProxyInterfaceMode)
+	{
+		case EPCGProxyInterfaceMode::ByNativeElement:
+			if (IsValid(SettingsClass))
+			{
+				const UPCGSettings* SettingsDefaultObject = CastChecked<UPCGSettings>(SettingsClass->GetDefaultObject());
+				OutputProperties = SettingsDefaultObject->DefaultOutputPinProperties();
+			}
+			break;
+		case EPCGProxyInterfaceMode::ByBlueprintElement:
+			if (IsValid(BlueprintElementClass))
+			{
+				const UPCGBlueprintElement* BlueprintSettingsDefaultObject = CastChecked<UPCGBlueprintElement>(BlueprintElementClass->GetDefaultObject());
+				OutputProperties = BlueprintSettingsDefaultObject->GetOutputPins();
+			}
+			break;
+		case EPCGProxyInterfaceMode::BySettings:
+			if (const UPCGSettings* SettingsPtr = Settings.LoadSynchronous())
+			{
+				OutputProperties = SettingsPtr->DefaultOutputPinProperties();
+			}
+			break;
+
+		default:
+			checkNoEntry();
+	}
+
+	if (OutputProperties.IsEmpty())
+	{
+		return Super::OutputPinProperties();
+	}
 
-	return PinProperties;
+	return OutputProperties;
 }
 
 FPCGElementPtr UPCGIndirectionSettings::CreateElement() const
@@ -47,7 +150,7 @@ bool FPCGIndirectionElement::CanExecuteOnlyOnMainThread(FPCGContext* InContext)
 		return true;
 	}
 
-	FPCGIndirectionContext* Context = static_cast<FPCGIndirectionContext*>(InContext);
+	const FPCGIndirectionContext* Context = static_cast<FPCGIndirectionContext*>(InContext);
 
 	if (Context->InnerElement)
 	{
@@ -91,6 +194,28 @@ bool FPCGIndirectionElement::PrepareDataInternal(FPCGContext* InContext) const
 
 	if (UPCGSettings* InnerSettings = Settings->Settings.LoadSynchronous())
 	{
+		if (Settings->ProxyInterfaceMode == EPCGProxyInterfaceMode::ByNativeElement)
+		{
+			if (!Settings->SettingsClass || !InnerSettings->GetClass()->IsChildOf(Settings->SettingsClass))
+			{
+				Context->bShouldActAsPassthrough = true;
+				PCGE_LOG_C(Error, GraphAndLog, Context, LOCTEXT("NativeProxySettingsInterfaceTypeMismatch", "The selected native settings template does not match the set or overridden settings!"));
+
+				return true;
+			}
+		}
+		else if (Settings->ProxyInterfaceMode == EPCGProxyInterfaceMode::ByBlueprintElement)
+		{
+			const UPCGBlueprintSettings* BlueprintSettings = Cast<UPCGBlueprintSettings>(InnerSettings);
+			if (!Settings->BlueprintElementClass || !BlueprintSettings || !BlueprintSettings->GetElementType()->IsChildOf(Settings->BlueprintElementClass))
+			{
+				Context->bShouldActAsPassthrough = true;
+				PCGE_LOG_C(Error, GraphAndLog, Context, LOCTEXT("BlueprintProxySettingsInterfaceTypeMismatch", "The selected blueprint settings template does not match the set or overridden settings!"));
+
+				return true;
+			}
+		} 
+
 		// TODO: while we can root it here, if this node or multiple indirection node were to execute in parallel,
 		// the lifetime of the inner settings would be a bit unclear - we need to improve this.
 		if (!InnerSettings->IsRooted())
@@ -133,8 +258,10 @@ bool FPCGIndirectionElement::ExecuteInternal(FPCGContext* InContext) const
 	const UPCGIndirectionSettings* Settings = Context->GetInputSettings<UPCGIndirectionSettings>();
 	check(Settings);
 
-	if (!Context->InnerElement || !Context->InnerContext)
+	// If Settings has not been set or overriden, act as passthrough
+	if (Context->bShouldActAsPassthrough || !Context->InnerElement || !Context->InnerContext)
 	{
+		Context->OutputData = Context->InputData;
 		return true;
 	}
 
@@ -142,7 +269,7 @@ bool FPCGIndirectionElement::ExecuteInternal(FPCGContext* InContext) const
 	// TODO: see what we can do for inspection data
 	// TODO: support pausing in inner element, might require some upstream changes in the graph executor
 	Context->InnerContext->AsyncState = Context->AsyncState;
-	bool bElementDone = Context->InnerElement->Execute(Context->InnerContext);
+	const bool bElementDone = Context->InnerElement->Execute(Context->InnerContext);
 
 	// Implementation note: to make sure everything is clean vs. the root set, we need to copy the output data
 	// regardless of whether the element is done or not
diff --git a/Engine/Plugins/Experimental/PCG/Source/PCG/Public/Elements/PCGIndirectionElement.h b/Engine/Plugins/Experimental/PCG/Source/PCG/Public/Elements/PCGIndirectionElement.h
--- a/Engine/Plugins/Experimental/PCG/Source/PCG/Public/Elements/PCGIndirectionElement.h
+++ b/Engine/Plugins/Experimental/PCG/Source/PCG/Public/Elements/PCGIndirectionElement.h
@@ -7,6 +7,16 @@
 
 #include "PCGIndirectionElement.generated.h"
 
+class UPCGBlueprintElement;
+
+UENUM()
+enum class EPCGProxyInterfaceMode : uint8
+{
+	ByNativeElement = 0u UMETA(Tooltip = "Select a native element to define the pin interface"),
+	ByBlueprintElement UMETA(Tooltip = "Select a custom blueprint element to define the pin interface"),
+	BySettings UMETA(Tooltip = "User defined settings will define the pin interface")
+};
+
 UCLASS(BlueprintType, ClassGroup=(Procedural))
 class PCG_API UPCGIndirectionSettings : public UPCGSettings
 {
@@ -19,16 +29,31 @@ public:
 	virtual FText GetDefaultNodeTitle() const override;
 	virtual FText GetNodeTooltipText() const override;
 	virtual EPCGSettingsType GetType() const override { return EPCGSettingsType::Generic; }
-#endif
+	virtual FName AdditionalTaskName() const override;
+#endif // WITH_EDITOR
 
 protected:
+	virtual TArray<FPCGPinProperties> InputPinProperties() const override;
 	virtual TArray<FPCGPinProperties> OutputPinProperties() const override;
 
 	virtual FPCGElementPtr CreateElement() const override;
 	//~End UPCGSettings interface
 
 public:
-	UPROPERTY(BlueprintReadWrite, EditAnywhere, Category = Settings, meta = (PCG_Overridable))
+	/** Defines which interface to use for populating pins */
+	UPROPERTY(BlueprintReadWrite, EditAnywhere, Category = Settings)
+	EPCGProxyInterfaceMode ProxyInterfaceMode = EPCGProxyInterfaceMode::BySettings;
+
+	/** The element settings class used to define the pin interface for this node instance */
+	UPROPERTY(EditAnywhere, Category = Settings, meta = (EditCondition = "ProxyInterfaceMode == EPCGProxyInterfaceMode::ByNativeElement", EditConditionHides))
+	TSubclassOf<UPCGSettings> SettingsClass;
+
+	/** The blueprint element class used to define the pin interface for this node instance */
+	UPROPERTY(EditAnywhere, Category = Settings, meta = (EditCondition = "ProxyInterfaceMode == EPCGProxyInterfaceMode::ByBlueprintElement", EditConditionHides))
+    TSubclassOf<UPCGBlueprintElement> BlueprintElementClass;
+
+	/** The element settings, which can be overriden, that will be used during the proxy execution */
+	UPROPERTY(EditAnywhere, BlueprintReadWrite, Category = Settings, meta = (PCG_Overridable))
 	TSoftObjectPtr<UPCGSettings> Settings;
 
 	UPROPERTY(BlueprintReadWrite, EditAnywhere, Category = Settings, meta = (PCG_Overridable))
@@ -42,6 +67,7 @@ struct FPCGIndirectionContext : public FPCGContext
 	FPCGElementPtr InnerElement;
 	FPCGContext* InnerContext = nullptr;
 	bool bNeedsToUnrootInnerSettings = false;
+	bool bShouldActAsPassthrough = false;
 };
 
 class FPCGIndirectionElement : public IPCGElement
