[PCG] TargetActor post process functions on Spawners

diff --git a/Engine/Plugins/Experimental/PCG/Source/PCG/Private/Elements/PCGCreateSpline.cpp b/Engine/Plugins/Experimental/PCG/Source/PCG/Private/Elements/PCGCreateSpline.cpp
--- a/Engine/Plugins/Experimental/PCG/Source/PCG/Private/Elements/PCGCreateSpline.cpp
+++ b/Engine/Plugins/Experimental/PCG/Source/PCG/Private/Elements/PCGCreateSpline.cpp
@@ -215,6 +215,15 @@ bool FPCGCreateSplineElement::ExecuteInternal(FPCGContext* Context) const
 			Context->SourceComponent->AddToManagedResources(ManagedComponent);
 		}
 
+		// Execute PostProcess Functions
+		if (SplineActor)
+		{
+			for (UFunction* Function : PCGHelpers::FindUserFunctions(SplineActor->GetClass(), Settings->PostProcessFunctionNames, Context))
+			{
+				SplineActor->ProcessEvent(Function, nullptr);
+			}
+		}
+
 		FPCGTaggedData& Output = Outputs.Add_GetRef(Input);
 		Output.Data = SplineData;
 	}
diff --git a/Engine/Plugins/Experimental/PCG/Source/PCG/Private/Elements/PCGCreateTargetActor.cpp b/Engine/Plugins/Experimental/PCG/Source/PCG/Private/Elements/PCGCreateTargetActor.cpp
--- a/Engine/Plugins/Experimental/PCG/Source/PCG/Private/Elements/PCGCreateTargetActor.cpp
+++ b/Engine/Plugins/Experimental/PCG/Source/PCG/Private/Elements/PCGCreateTargetActor.cpp
@@ -283,6 +283,11 @@ bool FPCGCreateTargetActorElement::ExecuteInternal(FPCGContext* Context) const
 
 	GeneratedActor->Tags.Add(PCGHelpers::DefaultPCGActorTag);
 
+	for (UFunction* Function : PCGHelpers::FindUserFunctions(GeneratedActor->GetClass(), Settings->PostProcessFunctionNames, Context))
+	{
+		GeneratedActor->ProcessEvent(Function, nullptr);
+	}
+
 	if (UPCGComponent* SourceComponent = Context->SourceComponent.Get())
 	{
 		UPCGManagedActors* ManagedActors = NewObject<UPCGManagedActors>(SourceComponent);
diff --git a/Engine/Plugins/Experimental/PCG/Source/PCG/Private/Elements/PCGSpawnActor.cpp b/Engine/Plugins/Experimental/PCG/Source/PCG/Private/Elements/PCGSpawnActor.cpp
--- a/Engine/Plugins/Experimental/PCG/Source/PCG/Private/Elements/PCGSpawnActor.cpp
+++ b/Engine/Plugins/Experimental/PCG/Source/PCG/Private/Elements/PCGSpawnActor.cpp
@@ -815,31 +815,6 @@ void FPCGSpawnActorElement::SpawnActors(FPCGSubgraphContext* Context, AActor* Ta
 	const UPCGSpawnActorSettings* Settings = Context->GetInputSettings<UPCGSpawnActorSettings>();
 	check(Settings && Settings->Option != EPCGSpawnActorOption::CollapseActors);
 
-	TArray<UFunction*> PostSpawnFunctions;
-	for (FName PostSpawnFunctionName : Settings->PostSpawnFunctionNames)
-	{
-		if (PostSpawnFunctionName == NAME_None)
-		{
-			continue;
-		}
-
-		if (UFunction* PostSpawnFunction = InTemplateActorClass->FindFunctionByName(PostSpawnFunctionName))
-		{
-			if (PostSpawnFunction->NumParms != 0)
-			{
-				PCGE_LOG(Warning, GraphAndLog, FText::Format(LOCTEXT("ParametersMissing", "PostSpawnFunction '{0}' requires parameters. We only support parameter-less functions. Call will be skipped."), FText::FromName(PostSpawnFunctionName)));
-			}
-			else
-			{
-				PostSpawnFunctions.Add(PostSpawnFunction);
-			}
-		}
-		else
-		{
-			PCGE_LOG(Warning, GraphAndLog, FText::Format(LOCTEXT("FunctionNotFound", "PostSpawnFunction '{0}' was not found in class '{1}'"), FText::FromName(PostSpawnFunctionName), FText::FromName(InTemplateActorClass->GetFName())));
-		}
-	}
-
 	const bool bForceDisableActorParsing = (Settings->bForceDisableActorParsing);
 
 	AActor* TemplateActor = nullptr;
@@ -954,6 +929,8 @@ void FPCGSpawnActorElement::SpawnActors(FPCGSubgraphContext* Context, AActor* Ta
 		PCGHelpers::GetGeneratedActorsFolderPath(TargetActor, GeneratedActorsFolderPath);
 #endif
 
+		const TArray<UFunction*> PostSpawnFunctions = PCGHelpers::FindUserFunctions(InTemplateActorClass, Settings->PostSpawnFunctionNames, Context);
+
 		for (int32 i = 0; i < Points.Num(); ++i)
 		{
 			const FPCGPoint& Point = Points[i];
diff --git a/Engine/Plugins/Experimental/PCG/Source/PCG/Private/Elements/PCGStaticMeshSpawner.cpp b/Engine/Plugins/Experimental/PCG/Source/PCG/Private/Elements/PCGStaticMeshSpawner.cpp
--- a/Engine/Plugins/Experimental/PCG/Source/PCG/Private/Elements/PCGStaticMeshSpawner.cpp
+++ b/Engine/Plugins/Experimental/PCG/Source/PCG/Private/Elements/PCGStaticMeshSpawner.cpp
@@ -9,6 +9,7 @@
 #include "Data/PCGSpatialData.h"
 #include "Elements/PCGStaticMeshSpawnerContext.h"
 #include "Helpers/PCGActorHelpers.h"
+#include "Helpers/PCGHelpers.h"
 #include "InstanceDataPackers/PCGInstanceDataPackerBase.h"
 #include "MeshSelectors/PCGMeshSelectorBase.h"
 #include "MeshSelectors/PCGMeshSelectorWeighted.h"
@@ -302,7 +303,19 @@ bool FPCGStaticMeshSpawnerElement::ExecuteInternal(FPCGContext* InContext) const
 		}
 	}
 
-	return Context->MeshInstancesData.IsEmpty();
+	const bool bFinishedExecution = Context->MeshInstancesData.IsEmpty();
+	if (bFinishedExecution)
+	{
+		if (AActor* TargetActor = Settings->TargetActor.Get() ? Settings->TargetActor.Get() : Context->GetTargetActor(nullptr))
+		{
+			for (UFunction* Function : PCGHelpers::FindUserFunctions(TargetActor->GetClass(), Settings->PostProcessFunctionNames, Context))
+			{
+				TargetActor->ProcessEvent(Function, nullptr);
+			}
+		}
+	}
+
+	return bFinishedExecution;
 }
 
 bool FPCGStaticMeshSpawnerElement::CanExecuteOnlyOnMainThread(FPCGContext* Context) const
diff --git a/Engine/Plugins/Experimental/PCG/Source/PCG/Private/Helpers/PCGHelpers.cpp b/Engine/Plugins/Experimental/PCG/Source/PCG/Private/Helpers/PCGHelpers.cpp
--- a/Engine/Plugins/Experimental/PCG/Source/PCG/Private/Helpers/PCGHelpers.cpp
+++ b/Engine/Plugins/Experimental/PCG/Source/PCG/Private/Helpers/PCGHelpers.cpp
@@ -3,6 +3,7 @@
 #include "Helpers/PCGHelpers.h"
 
 #include "PCGComponent.h"
+#include "PCGContext.h"
 #include "PCGGraph.h"
 #include "PCGModule.h"
 #include "PCGSubsystem.h"
@@ -19,6 +20,8 @@
 #include "Engine/World.h"
 #endif
 
+#define LOCTEXT_NAMESPACE "PCGHelpers"
+
 namespace PCGHelpers
 {
 	int ComputeSeed(int A)
@@ -492,4 +495,49 @@ namespace PCGHelpers
 		}
 #endif
 	}
+
+	TArray<UFunction*> FindUserFunctions(TSubclassOf<AActor> ActorClass, const TArray<FName>& FunctionNames, const FPCGContext* InContext)
+	{
+		TArray<UFunction*> Functions;
+
+		if (!ActorClass)
+		{
+			return Functions;
+		}
+
+		for (FName FunctionName : FunctionNames)
+		{
+			if (FunctionName == NAME_None)
+			{
+				continue;
+			}
+
+			if (UFunction* Function = ActorClass->FindFunctionByName(FunctionName))
+			{
+				// TODO: support Parameters
+				if (Function->NumParms != 0)
+				{
+					PCGLog::LogWarningOnGraph(FText::Format(LOCTEXT("ParametersMissing", "Function '{0}' in class '{1}' has parameters, but only parameter-less functions are supported. Call will be skipped."), FText::FromName(FunctionName), FText::FromName(ActorClass->GetFName())), InContext);
+				}
+#if WITH_EDITOR
+				else if (!Function->GetBoolMetaData(TEXT("CallInEditor")))
+				{
+					PCGLog::LogWarningOnGraph(FText::Format(LOCTEXT("CallInEditorFailed", "Function '{0}' in class '{1}' requires CallInEditor to be true while in-editor."), FText::FromName(FunctionName), FText::FromName(ActorClass->GetFName())), InContext);
+				}
+#endif
+				else
+				{
+					Functions.Add(Function);
+				}
+			}
+			else
+			{
+				PCGLog::LogWarningOnGraph(FText::Format(LOCTEXT("FunctionNotFound", "Function '{0}' was not found in class '{1}'."), FText::FromName(FunctionName), FText::FromName(ActorClass->GetFName())), InContext);
+			}
+		}
+
+		return Functions;
+	}
 }
+
+#undef LOCTEXT_NAMESPACE
diff --git a/Engine/Plugins/Experimental/PCG/Source/PCG/Public/Elements/PCGCreateSpline.h b/Engine/Plugins/Experimental/PCG/Source/PCG/Public/Elements/PCGCreateSpline.h
--- a/Engine/Plugins/Experimental/PCG/Source/PCG/Public/Elements/PCGCreateSpline.h
+++ b/Engine/Plugins/Experimental/PCG/Source/PCG/Public/Elements/PCGCreateSpline.h
@@ -66,6 +66,10 @@ public:
 
 	UPROPERTY(BlueprintReadWrite, EditAnywhere, Category = Settings, meta = (EditCondition = "Mode==EPCGCreateSplineMode::CreateNewActor", EditConditionHides))
 	EPCGAttachOptions AttachOptions = EPCGAttachOptions::Attached;
+
+	/** Specify a list of functions to be called on the target actor after spline creation. Functions need to be parameter-less and with "CallInEditor" flag enabled. */
+	UPROPERTY(BlueprintReadWrite, EditAnywhere, Category = Settings)
+	TArray<FName> PostProcessFunctionNames;
 };
 
 class FPCGCreateSplineElement : public IPCGElement
diff --git a/Engine/Plugins/Experimental/PCG/Source/PCG/Public/Elements/PCGCreateTargetActor.h b/Engine/Plugins/Experimental/PCG/Source/PCG/Public/Elements/PCGCreateTargetActor.h
--- a/Engine/Plugins/Experimental/PCG/Source/PCG/Public/Elements/PCGCreateTargetActor.h
+++ b/Engine/Plugins/Experimental/PCG/Source/PCG/Public/Elements/PCGCreateTargetActor.h
@@ -71,6 +71,10 @@ public:
 
 	UPROPERTY(meta = (PCG_Overridable))
 	FTransform ActorPivot;
+
+	/** Specify a list of functions to be called on the target actor after creation. Functions need to be parameter-less and with "CallInEditor" flag enabled. */
+	UPROPERTY(BlueprintReadWrite, EditAnywhere, Category = Settings)
+	TArray<FName> PostProcessFunctionNames;
 };
 
 class FPCGCreateTargetActorElement : public IPCGElement
diff --git a/Engine/Plugins/Experimental/PCG/Source/PCG/Public/Elements/PCGStaticMeshSpawner.h b/Engine/Plugins/Experimental/PCG/Source/PCG/Public/Elements/PCGStaticMeshSpawner.h
--- a/Engine/Plugins/Experimental/PCG/Source/PCG/Public/Elements/PCGStaticMeshSpawner.h
+++ b/Engine/Plugins/Experimental/PCG/Source/PCG/Public/Elements/PCGStaticMeshSpawner.h
@@ -97,6 +97,10 @@ public:
 	UPROPERTY(BlueprintReadWrite, Category = Settings, meta = (PCG_Overridable))
 	TSoftObjectPtr<AActor> TargetActor;
 
+	/** Specify a list of functions to be called on the target actor after instances are spawned. Functions need to be parameter-less and with "CallInEditor" flag enabled. */
+	UPROPERTY(BlueprintReadWrite, EditAnywhere, Category = Settings)
+	TArray<FName> PostProcessFunctionNames;
+
 	UPROPERTY(meta=(DeprecatedProperty, DeprecationMessage="Use MeshSelectorType and MeshSelectorParameters instead."))
 	TArray<FPCGStaticMeshSpawnerEntry> Meshes_DEPRECATED;
 
diff --git a/Engine/Plugins/Experimental/PCG/Source/PCG/Public/Helpers/PCGHelpers.h b/Engine/Plugins/Experimental/PCG/Source/PCG/Public/Helpers/PCGHelpers.h
--- a/Engine/Plugins/Experimental/PCG/Source/PCG/Public/Helpers/PCGHelpers.h
+++ b/Engine/Plugins/Experimental/PCG/Source/PCG/Public/Helpers/PCGHelpers.h
@@ -13,6 +13,7 @@ class ALandscapeProxy;
 class UPCGComponent;
 class UPCGGraph;
 class UWorld;
+struct FPCGContext;
 
 namespace PCGHelpers
 {
@@ -70,4 +71,7 @@ namespace PCGHelpers
 #endif
 
 	PCG_API void AttachToParent(AActor* InActorToAttach, AActor* InParent, EPCGAttachOptions AttachOptions, const FString& GeneratedPath = FString());
+
+	/** Finds functions on the actor matching the provided function names. Functions must be marked as CallInEditor and take no parameters. */
+	PCG_API TArray<UFunction*> FindUserFunctions(TSubclassOf<AActor> ActorClass, const TArray<FName>& FunctionNames, const FPCGContext* InContext = nullptr);
 };
