[PCG] Engine Settings coming from the Defaultengine.ini are now updating the CVar correctly

* Also now use the helper class to avoid duplicating boilerplate code.

#rnx

diff --git a/Engine/Plugins/Experimental/PCG/Source/PCG/Private/PCGEngineSettings.cpp b/Engine/Plugins/Experimental/PCG/Source/PCG/Private/PCGEngineSettings.cpp
--- a/Engine/Plugins/Experimental/PCG/Source/PCG/Private/PCGEngineSettings.cpp
+++ b/Engine/Plugins/Experimental/PCG/Source/PCG/Private/PCGEngineSettings.cpp
@@ -2,6 +2,9 @@
 
 #include "PCGEngineSettings.h"
 
+#include "HAL/IConsoleManager.h"
+#include "Misc/ConfigUtilities.h"
+
 #define LOCTEXT_NAMESPACE "PCGEngineSettings"
 
 FName UPCGEngineSettings::GetCategoryName() const
@@ -18,26 +21,15 @@ FText UPCGEngineSettings::GetSectionText() const
 
 void UPCGEngineSettings::PostInitProperties()
 {
-	Super::PostInitProperties();
-
 #if WITH_EDITOR
 	if (IsTemplate())
 	{
-		ImportConsoleVariableValues();
+		// We want the .ini file to have precedence over the CVar constructor, so we apply the ini to the CVar before following the regular UDeveloperSettingsBackedByCVars flow
+		UE::ConfigUtilities::ApplyCVarSettingsFromIni(TEXT("/Script/PCG.PCGEngineSettings"), *GEngineIni, ECVF_SetByProjectSetting);
 	}
 #endif // WITH_EDITOR
-}
 
-#if WITH_EDITOR
-void UPCGEngineSettings::PostEditChangeProperty(FPropertyChangedEvent& PropertyChangedEvent)
-{
-	Super::PostEditChangeProperty(PropertyChangedEvent);
-
-	if (PropertyChangedEvent.Property)
-	{
-		ExportValuesToConsoleVariables(PropertyChangedEvent.Property);
-	}
+	Super::PostInitProperties();
 }
-#endif // WITH_EDITOR
 
 #undef LOCTEXT_NAMESPACE
\ No newline at end of file
diff --git a/Engine/Plugins/Experimental/PCG/Source/PCG/Public/PCGEngineSettings.h b/Engine/Plugins/Experimental/PCG/Source/PCG/Public/PCGEngineSettings.h
--- a/Engine/Plugins/Experimental/PCG/Source/PCG/Public/PCGEngineSettings.h
+++ b/Engine/Plugins/Experimental/PCG/Source/PCG/Public/PCGEngineSettings.h
@@ -3,11 +3,12 @@
 #pragma once
 
 #include "Engine/DeveloperSettings.h"
+#include "Engine/DeveloperSettingsBackedByCVars.h"
 
 #include "PCGEngineSettings.generated.h"
 
 UCLASS(config = Engine, defaultconfig)
-class PCG_API UPCGEngineSettings : public UDeveloperSettings
+class PCG_API UPCGEngineSettings : public UDeveloperSettingsBackedByCVars
 {
 	GENERATED_BODY()
 
@@ -49,8 +50,5 @@ public:
 
 	//~ Begin UObject Interface
 	virtual void PostInitProperties() override;
-#if WITH_EDITOR
-	virtual void PostEditChangeProperty(FPropertyChangedEvent& PropertyChangedEvent) override;
-#endif
 	//~ End UObject Interface
 };
