[PCG] GatherDependencies was not checking Excluded classes correctly

#rnx

diff --git a/Engine/Plugins/Experimental/PCG/Source/PCG/Private/Helpers/PCGHelpers.cpp b/Engine/Plugins/Experimental/PCG/Source/PCG/Private/Helpers/PCGHelpers.cpp
--- a/Engine/Plugins/Experimental/PCG/Source/PCG/Private/Helpers/PCGHelpers.cpp
+++ b/Engine/Plugins/Experimental/PCG/Source/PCG/Private/Helpers/PCGHelpers.cpp
@@ -306,35 +306,34 @@ namespace PCGHelpers
 
 		for (FProperty* Property = ObjectClass->PropertyLink; Property != nullptr; Property = Property->PropertyLinkNext)
 		{
-			// Skip any kind of internal property and the ones that are susceptible to be instable
-			if (!Property || Property->HasAnyPropertyFlags(CPF_Transient | CPF_DuplicateTransient | CPF_Deprecated))
-			{
-				continue;
-			}
-
-			// For Object properties, if the class is in the excluded classes, continue
-			if (FObjectProperty* ObjectProperty = CastField<FObjectProperty>(Property))
-			{
-				if (ObjectProperty->PropertyClass && Algo::AnyOf(InExcludedClasses, [InClass = ObjectProperty->PropertyClass](const UClass* ExcludedClass){ return InClass->IsChildOf(ExcludedClass); }))
-				{
-					continue;
-				}
-			}
-
-			GatherDependencies(Property, Object, OutDependencies, MaxDepth);
+			GatherDependencies(Property, Object, OutDependencies, MaxDepth, InExcludedClasses);
 		}
 	}
 
 	// Inspired by IteratePropertiesRecursive in ObjectPropertyTrace.cpp
-	void GatherDependencies(FProperty* Property, const void* InContainer, TSet<TObjectPtr<UObject>>& OutDependencies, int32 MaxDepth)
+	void GatherDependencies(FProperty* Property, const void* InContainer, TSet<TObjectPtr<UObject>>& OutDependencies, int32 MaxDepth, const TArray<UClass*>& InExcludedClasses)
 	{
-		auto AddToDependenciesAndGatherRecursively = [&OutDependencies, MaxDepth](UObject* Object) {
+		// Skip any kind of internal property and the ones that are susceptible to be unstable
+		if (!Property || Property->HasAnyPropertyFlags(CPF_Transient | CPF_DuplicateTransient | CPF_Deprecated))
+		{
+			return;
+		}
+
+		auto AddToDependenciesAndGatherRecursively = [&OutDependencies, MaxDepth, &InExcludedClasses](UObject* Object) {
 			if (Object && !OutDependencies.Contains(Object))
 			{
+				// If we explicitly don't want to track this object, early out.
+				if (!Object->GetClass() || 
+					!CanBeExpanded(Object->GetClass()) ||
+					Algo::AnyOf(InExcludedClasses, [InClass = Object->GetClass()](const UClass* ExcludedClass) { return InClass->IsChildOf(ExcludedClass); }))
+				{
+					return;
+				}
+
 				OutDependencies.Add(Object);
 				if (MaxDepth != 0)
 				{
-					GatherDependencies(Object, OutDependencies, MaxDepth - 1);
+					GatherDependencies(Object, OutDependencies, MaxDepth - 1, InExcludedClasses);
 				}
 			}
 		};
@@ -364,7 +363,7 @@ namespace PCGHelpers
 			const void* StructContainer = StructProperty->ContainerPtrToValuePtr<const void>(InContainer);
 			for (TFieldIterator<FProperty> It(StructProperty->Struct); It; ++It)
 			{
-				GatherDependencies(*It, StructContainer, OutDependencies, MaxDepth);
+				GatherDependencies(*It, StructContainer, OutDependencies, MaxDepth, InExcludedClasses);
 			}
 		}
 		else if (FArrayProperty* ArrayProperty = CastField<FArrayProperty>(Property))
@@ -373,7 +372,7 @@ namespace PCGHelpers
 			for (int32 DynamicIndex = 0; DynamicIndex < Helper.Num(); ++DynamicIndex)
 			{
 				const void* ValuePtr = Helper.GetRawPtr(DynamicIndex);
-				GatherDependencies(ArrayProperty->Inner, ValuePtr, OutDependencies, MaxDepth);
+				GatherDependencies(ArrayProperty->Inner, ValuePtr, OutDependencies, MaxDepth, InExcludedClasses);
 			}
 		}
 		else if (FMapProperty* MapProperty = CastField<FMapProperty>(Property))
@@ -387,8 +386,8 @@ namespace PCGHelpers
 					// Key and Value are stored next to each other in memory.
 					// ValueProp has an offset, so we should use the same starting address for both.
 					const void* PairKeyValuePtr = Helper.GetKeyPtr(DynamicIndex);
-					GatherDependencies(MapProperty->KeyProp, PairKeyValuePtr, OutDependencies, MaxDepth);
-					GatherDependencies(MapProperty->ValueProp, PairKeyValuePtr, OutDependencies, MaxDepth);
+					GatherDependencies(MapProperty->KeyProp, PairKeyValuePtr, OutDependencies, MaxDepth, InExcludedClasses);
+					GatherDependencies(MapProperty->ValueProp, PairKeyValuePtr, OutDependencies, MaxDepth, InExcludedClasses);
 
 					--Num;
 				}
@@ -403,7 +402,7 @@ namespace PCGHelpers
 				if (Helper.IsValidIndex(DynamicIndex))
 				{
 					const void* ValuePtr = Helper.GetElementPtr(DynamicIndex);
-					GatherDependencies(SetProperty->ElementProp, ValuePtr, OutDependencies, MaxDepth);
+					GatherDependencies(SetProperty->ElementProp, ValuePtr, OutDependencies, MaxDepth, InExcludedClasses);
 
 					--Num;
 				}
diff --git a/Engine/Plugins/Experimental/PCG/Source/PCG/Public/Helpers/PCGHelpers.h b/Engine/Plugins/Experimental/PCG/Source/PCG/Public/Helpers/PCGHelpers.h
--- a/Engine/Plugins/Experimental/PCG/Source/PCG/Public/Helpers/PCGHelpers.h
+++ b/Engine/Plugins/Experimental/PCG/Source/PCG/Public/Helpers/PCGHelpers.h
@@ -51,8 +51,8 @@ namespace PCGHelpers
 	PCG_API TArray<FString> GetStringArrayFromCommaSeparatedString(const FString& InCommaSeparatedString);
 
 #if WITH_EDITOR
-	PCG_API void GatherDependencies(UObject* Object, TSet<TObjectPtr<UObject>>& OutDependencies, int32 MaxDepth = -1, const TArray<UClass*>& InExcludedClasses = TArray<UClass*>());
-	PCG_API void GatherDependencies(FProperty* Property, const void* InContainer, TSet<TObjectPtr<UObject>>& OutDependencies, int32 MaxDepth);
+	PCG_API void GatherDependencies(UObject* Object, TSet<TObjectPtr<UObject>>& OutDependencies, int32 MaxDepth = -1, const TArray<UClass*>& InExcludedClasses = {});
+	PCG_API void GatherDependencies(FProperty* Property, const void* InContainer, TSet<TObjectPtr<UObject>>& OutDependencies, int32 MaxDepth, const TArray<UClass*>& InExcludedClasses = {});
 #endif
 
 	/** 
