First pass at cleaning ensures/errors around cancellation with multi-graph setups - this removes a spurious ensure until we fix this more properly (done in a second pass) and a hole wrt scheduled-but-not-processed task dependencies.

diff --git a/Engine/Plugins/Experimental/PCG/Source/PCG/Private/Graph/PCGGraphExecutor.cpp b/Engine/Plugins/Experimental/PCG/Source/PCG/Private/Graph/PCGGraphExecutor.cpp
--- a/Engine/Plugins/Experimental/PCG/Source/PCG/Private/Graph/PCGGraphExecutor.cpp
+++ b/Engine/Plugins/Experimental/PCG/Source/PCG/Private/Graph/PCGGraphExecutor.cpp
@@ -1014,8 +1014,9 @@ void FPCGGraphExecutor::QueueNextTasks(FPCGTaskId FinishedTask)
 			bool bAllPrerequisitesMet = true;
 			FPCGGraphTask* SuccessorTaskPtr = Tasks.Find(Successor);
 
-			// This should never be null, but later recovery should be able to cleanup this properly
-			if (ensure(SuccessorTaskPtr))
+			// This should rarely be null, but can happen when we have a task waiting on the execution of multiple other components, and one/many get cancelled.
+			// Example: R gets data from C0 and C1. C0 is cancelled due to a change in params, which will cancel R, but not C1. C1 finishes executing and has a null successor here, but that's fine.
+			if (SuccessorTaskPtr)
 			{
 				FPCGGraphTask& SuccessorTask = *SuccessorTaskPtr;
 
@@ -1061,6 +1062,19 @@ bool FPCGGraphExecutor::CancelNextTasks(FPCGTaskId CancelledTask, TSet<UPCGCompo
 		TaskSuccessors.Remove(CancelledTask);
 	}
 
+	// Tasks cancelled might have an impact on scheduled-but-not-processed tasks
+	ScheduleLock.Lock();
+	for (FPCGGraphScheduleTask& ScheduledTask : ScheduledTasks)
+	{
+		if (!OutCancelledComponents.Contains(ScheduledTask.SourceComponent.Get()) &&
+			Algo::AnyOf(ScheduledTask.Tasks[ScheduledTask.FirstTaskIndex].Inputs, [CancelledTask](const FPCGGraphTaskInput& Input) { return Input.TaskId == CancelledTask; }))
+		{
+			OutCancelledComponents.Add(ScheduledTask.SourceComponent.Get());
+			bAddedComponents = true;
+		}
+	}
+	ScheduleLock.Unlock();
+
 	return bAddedComponents;
 }
 
