[PCG] Add option to disable refresh (global) as a nuke function and another to disable tracking other than itself (on the component)

#rnx

diff --git a/Engine/Plugins/Experimental/PCG/Source/PCG/Private/PCGActorAndComponentMapping.cpp b/Engine/Plugins/Experimental/PCG/Source/PCG/Private/PCGActorAndComponentMapping.cpp
--- a/Engine/Plugins/Experimental/PCG/Source/PCG/Private/PCGActorAndComponentMapping.cpp
+++ b/Engine/Plugins/Experimental/PCG/Source/PCG/Private/PCGActorAndComponentMapping.cpp
@@ -1773,7 +1773,14 @@ void FPCGActorAndComponentMapping::OnActorChanged(AActor* InActor, bool bInHasMo
 	// And refresh all dirtied components
 	for (UPCGComponent* Component : DirtyComponents)
 	{
-		if (Component && (!bNoRefreshOwner || Component->GetOwner() != InActor))
+		if (!ensure(Component))
+		{
+			continue;
+		}
+
+		const bool bOwnerHasChanged = Component->GetOwner() == InActor;
+
+		if ((!bNoRefreshOwner || !bOwnerHasChanged) && (!Component->bOnlyTrackItself || bOwnerHasChanged))
 		{
 			// When an object changes, we need to make sure that we don't trigger a refresh on PCG components that are "higher" in the
 			// level hierarchy, otherwise we will end up generating in the Level Instance level, which is wrong.
diff --git a/Engine/Plugins/Experimental/PCG/Source/PCG/Private/PCGComponent.cpp b/Engine/Plugins/Experimental/PCG/Source/PCG/Private/PCGComponent.cpp
--- a/Engine/Plugins/Experimental/PCG/Source/PCG/Private/PCGComponent.cpp
+++ b/Engine/Plugins/Experimental/PCG/Source/PCG/Private/PCGComponent.cpp
@@ -61,6 +61,11 @@
 namespace PCGComponent
 {
 	const bool bSaveOnCleanupAndGenerate = false;
+
+	static TAutoConsoleVariable<bool> CVarGlobalDisableRefresh(
+		TEXT("pcg.GlobalDisableRefresh"),
+		false,
+		TEXT("Disable refresh for all PCG Components."));
 }
 
 UPCGComponent::UPCGComponent(const FObjectInitializer& InObjectInitializer)
@@ -1891,6 +1896,12 @@ void UPCGComponent::Refresh(bool bStructural)
 		return;
 	}
 
+	// If the refresh is globally disabled, just exit
+	if (PCGComponent::CVarGlobalDisableRefresh.GetValueOnAnyThread())
+	{
+		return;
+	}
+
 	// Discard any refresh if have already one scheduled.
 	if (UPCGSubsystem* Subsystem = GetSubsystem())
 	{
diff --git a/Engine/Plugins/Experimental/PCG/Source/PCG/Public/PCGComponent.h b/Engine/Plugins/Experimental/PCG/Source/PCG/Public/PCGComponent.h
--- a/Engine/Plugins/Experimental/PCG/Source/PCG/Public/PCGComponent.h
+++ b/Engine/Plugins/Experimental/PCG/Source/PCG/Public/PCGComponent.h
@@ -232,6 +232,10 @@ public:
 	UPROPERTY(BlueprintReadWrite, EditAnywhere, Category = "Editing Settings", meta = (DisplayName = "Regenerate PCG Volume In Editor", DisplayPriority = 400))
 	bool bRegenerateInEditor = true;
 
+	/** Even if the graph has external dependencies, the component won't react to them. */
+	UPROPERTY(BlueprintReadWrite, EditAnywhere, Category = "Editing Settings", meta = (DisplayPriority = 450))
+	bool bOnlyTrackItself = false;
+
 	UPROPERTY(BlueprintReadOnly, VisibleAnywhere, Transient, Category = Debug, meta = (NoResetToDefault))
 	bool bDirtyGenerated = false;
 
diff --git a/Engine/Plugins/Experimental/PCG/Source/PCG/Public/PCGEngineSettings.h b/Engine/Plugins/Experimental/PCG/Source/PCG/Public/PCGEngineSettings.h
--- a/Engine/Plugins/Experimental/PCG/Source/PCG/Public/PCGEngineSettings.h
+++ b/Engine/Plugins/Experimental/PCG/Source/PCG/Public/PCGEngineSettings.h
@@ -25,16 +25,20 @@ public:
 	// Console variables defined in PCGActorAndComponentMapping.cpp
 	// Use the variables to get the values, not GetDefault<UPCGEngineSettings>
 
+	/** Globally disable refresh. */
+	UPROPERTY(EditAnywhere, Config, Category = Tracking, meta = (ConsoleVariable = "pcg.GlobalDisableRefresh"))
+	bool bGlobalDisableRefresh = false;
+
 	/** Completely disable landscape refresh when it changes. */
-	UPROPERTY(EditAnywhere, Config, Category = LandscapeTracking, meta = (ConsoleVariable = "pcg.LandscapeDisableRefreshTracking"))
+	UPROPERTY(EditAnywhere, Config, Category = Tracking, meta = (ConsoleVariable = "pcg.LandscapeDisableRefreshTracking"))
 	bool bLandscapeDisableRefreshTracking = false;
 
 	/** Completely disable landscape refresh when it changes in edit mode. Will force a refresh when landscape edit mode is exited. */
-	UPROPERTY(EditAnywhere, Config, Category = LandscapeTracking, meta = (ConsoleVariable = "pcg.LandscapeDisableRefreshTrackingInLandscapeEditingMode"))
+	UPROPERTY(EditAnywhere, Config, Category = Tracking, meta = (ConsoleVariable = "pcg.LandscapeDisableRefreshTrackingInLandscapeEditingMode"))
 	bool bLandscapeDisableRefreshTrackingInLandscapeEditingMode = false;
 
 	/** Time in MS between a landscape change and PCG refresh. Set it to 0 or negative value to disable the delay. */
-	UPROPERTY(EditAnywhere, Config, Category = "LandscapeTracking|Advanced", meta = (ConsoleVariable = "pcg.LandscapeRefreshTimeDelayMS"))
+	UPROPERTY(EditAnywhere, Config, Category = "Tracking|Advanced", meta = (ConsoleVariable = "pcg.LandscapeRefreshTimeDelayMS"))
 	int32 LandscapeRefreshTimeDelayMS = 1000;
 
 	/** Strips empty point data from being passed along as output. */
