[PCG] Ensure we don't schedule any tasks for a null source component.

#rnx

diff --git a/Engine/Plugins/Experimental/PCG/Source/PCG/Private/Graph/PCGGraphExecutor.cpp b/Engine/Plugins/Experimental/PCG/Source/PCG/Private/Graph/PCGGraphExecutor.cpp
--- a/Engine/Plugins/Experimental/PCG/Source/PCG/Private/Graph/PCGGraphExecutor.cpp
+++ b/Engine/Plugins/Experimental/PCG/Source/PCG/Private/Graph/PCGGraphExecutor.cpp
@@ -443,6 +443,11 @@ void FPCGGraphExecutor::Execute()
 	for (int32 ScheduledTaskIndex = ScheduledTasks.Num() - 1; ScheduledTaskIndex >= 0; --ScheduledTaskIndex)
 	{
 		FPCGGraphScheduleTask& ScheduledTask = ScheduledTasks[ScheduledTaskIndex];
+		if (!ScheduledTask.SourceComponent.IsValid())
+		{
+			// If the source component has gone invalid, don't schedule any tasks.
+			continue;
+		}
 
 		check(ScheduledTask.Tasks.Num() > 0);
 		// Push tasks to the primary task list & build successors map
