[PCG] Add condition to early out if no item to filter in attribute filter

#rnx

diff --git a/Engine/Plugins/Experimental/PCG/Source/PCG/Private/Elements/PCGAttributeFilter.cpp b/Engine/Plugins/Experimental/PCG/Source/PCG/Private/Elements/PCGAttributeFilter.cpp
--- a/Engine/Plugins/Experimental/PCG/Source/PCG/Private/Elements/PCGAttributeFilter.cpp
+++ b/Engine/Plugins/Experimental/PCG/Source/PCG/Private/Elements/PCGAttributeFilter.cpp
@@ -528,6 +528,12 @@ bool FPCGAttributeFilterElementBase::DoFiltering(FPCGContext* Context, EPCGAttri
 		const bool bCheckCompare = (Operator != EPCGAttributeFilterOperator::Equal) && (Operator != EPCGAttributeFilterOperator::NotEqual) && !bCheckStringSearch;
 		const int32 NumInput = TargetKeys->GetNum();
 
+		if (NumInput == 0)
+		{
+			ForwardInputToInFilterPin();
+			continue;
+		}
+
 		if (!PCGAttributeFilterHelpers::PrepareThresholdInfoFromInput(Context, OriginalData, NumInput, FirstThreshold, FirstThresholdInfo, TargetType, bCheckCompare, bCheckStringSearch))
 		{
 			ForwardInputToInFilterPin();
