[PCG] New node to find the convex hull of a point cloud

diff --git a/Engine/Plugins/Experimental/PCG/Source/PCG/Private/Elements/PCGConvexHull2D.cpp b/Engine/Plugins/Experimental/PCG/Source/PCG/Private/Elements/PCGConvexHull2D.cpp
new file mode 100644
--- /dev/null
+++ b/Engine/Plugins/Experimental/PCG/Source/PCG/Private/Elements/PCGConvexHull2D.cpp
@@ -0,0 +1,64 @@
+// Copyright Epic Games, Inc. All Rights Reserved.
+
+#include "Elements/PCGConvexHull2D.h"
+
+#include "PCGContext.h"
+#include "PCGModule.h"
+#include "Data/PCGPointData.h"
+
+#include "Math/ConvexHull2d.h"
+
+#define LOCTEXT_NAMESPACE "PCGConvexHull2DElement"
+
+FPCGElementPtr UPCGConvexHull2DSettings::CreateElement() const
+{
+	return MakeShared<FPCGConvexHull2DElement>();
+}
+
+bool FPCGConvexHull2DElement::ExecuteInternal(FPCGContext* Context) const
+{
+	TRACE_CPUPROFILER_EVENT_SCOPE(FPCGConvexHullElement::Execute);
+
+	check(Context);
+
+	const UPCGConvexHull2DSettings* Settings = Context->GetInputSettings<UPCGConvexHull2DSettings>();
+	check(Settings);
+
+	TArray<FPCGTaggedData>& Outputs = Context->OutputData.TaggedData;
+
+	for (const FPCGTaggedData& Input : Context->InputData.GetInputsByPin(PCGPinConstants::DefaultInputLabel))
+	{
+		const UPCGPointData* PointData = Cast<UPCGPointData>(Input.Data);
+
+		if (!PointData)
+		{
+			PCGE_LOG(Error, GraphAndLog, LOCTEXT("InputNotPointData", "Input is not a point data"));
+			continue;
+		}
+
+		const TArray<FPCGPoint>& Points = PointData->GetPoints();
+
+		TArray<FVector> PointsPositions;
+		TArray<int32> ConvexHullIndices;
+		Algo::Transform(Points, PointsPositions, [](const FPCGPoint& Point) { return Point.Transform.GetLocation(); });
+
+		ConvexHull2D::ComputeConvexHull(PointsPositions, ConvexHullIndices);
+
+		FPCGTaggedData& Output = Outputs.Add_GetRef(Input);
+		UPCGPointData* OutputPointData = NewObject<UPCGPointData>();
+		OutputPointData->InitializeFromData(PointData);
+		TArray<FPCGPoint>& OutPoints = OutputPointData->GetMutablePoints();
+		OutPoints.Reserve(ConvexHullIndices.Num());
+		for (int32 Index : ConvexHullIndices)
+		{
+			OutPoints.Add(Points[Index]);
+		}
+
+		Output.Data = OutputPointData;
+		Output.Pin = PCGPinConstants::DefaultOutputLabel;
+	}
+
+	return true;
+}
+
+#undef LOCTEXT_NAMESPACE
\ No newline at end of file
diff --git a/Engine/Plugins/Experimental/PCG/Source/PCG/Public/Elements/PCGConvexHull2D.h b/Engine/Plugins/Experimental/PCG/Source/PCG/Public/Elements/PCGConvexHull2D.h
new file mode 100644
--- /dev/null
+++ b/Engine/Plugins/Experimental/PCG/Source/PCG/Public/Elements/PCGConvexHull2D.h
@@ -0,0 +1,37 @@
+// Copyright Epic Games, Inc. All Rights Reserved.
+
+#pragma once
+
+#include "PCGSettings.h"
+
+#include "PCGConvexHull2D.generated.h"
+
+/**
+* Return the convex hull of a set of points on the XY plane
+*/
+UCLASS(BlueprintType)
+class UPCGConvexHull2DSettings : public UPCGSettings
+{
+	GENERATED_BODY()
+
+public:
+	//~Begin UPCGSettings interface
+#if WITH_EDITOR
+	virtual FName GetDefaultNodeName() const override { return FName(TEXT("FindConvexHull2D")); }
+	virtual FText GetDefaultNodeTitle() const override { return NSLOCTEXT("PCGConvexHullElement", "NodeTitle", "Find Convex Hull 2D"); }
+	virtual FText GetNodeTooltipText() const override { return NSLOCTEXT("PCGConvexHullElement", "NodeTooltip", "Return the 2D convex hull of a set of points on the XY plane."); }
+	virtual EPCGSettingsType GetType() const override { return EPCGSettingsType::Spatial; }
+#endif
+
+protected:
+	virtual TArray<FPCGPinProperties> InputPinProperties() const override { return Super::DefaultPointInputPinProperties(); }
+	virtual TArray<FPCGPinProperties> OutputPinProperties() const override { return Super::DefaultPointOutputPinProperties(); }
+	virtual FPCGElementPtr CreateElement() const override;
+	//~End UPCGSettings interface
+};
+
+class FPCGConvexHull2DElement : public IPCGElement
+{
+protected:
+	virtual bool ExecuteInternal(FPCGContext* Context) const override;
+};
\ No newline at end of file
