[PCG] Fixed some remaining cases of using strings instead of soft object paths for actor reference and physical material

#rnx

diff --git a/Engine/Plugins/Experimental/PCG/Source/PCG/Private/Elements/PCGSpawnActor.cpp b/Engine/Plugins/Experimental/PCG/Source/PCG/Private/Elements/PCGSpawnActor.cpp
--- a/Engine/Plugins/Experimental/PCG/Source/PCG/Private/Elements/PCGSpawnActor.cpp
+++ b/Engine/Plugins/Experimental/PCG/Source/PCG/Private/Elements/PCGSpawnActor.cpp
@@ -803,13 +803,13 @@ void FPCGSpawnActorElement::SpawnActors(FPCGSubgraphContext* Context, AActor* Ta
 	}
 
 	int32 OutPointOffset = 0;
-	FPCGMetadataAttribute<FString>* ActorReferenceAttribute = nullptr;
+	FPCGMetadataAttribute<FSoftObjectPath>* ActorReferenceAttribute = nullptr;
 
 	if (OutPointData)
 	{
 		OutPointOffset = OutPointData->GetMutablePoints().Num();
 		OutPointData->GetMutablePoints().Append(Points);
-		ActorReferenceAttribute = OutPointData->MutableMetadata()->FindOrCreateAttribute<FString>(PCGPointDataConstants::ActorReferenceAttribute, FString(), /*bAllowsInterpolation=*/false, /*bOverrideParent=*/false, /*bOverwriteIfTypeMismatch=*/false);
+		ActorReferenceAttribute = OutPointData->MutableMetadata()->FindOrCreateAttribute<FSoftObjectPath>(PCGPointDataConstants::ActorReferenceAttribute, FSoftObjectPath(), /*bAllowsInterpolation=*/false, /*bOverrideParent=*/false, /*bOverwriteIfTypeMismatch=*/false);
 	}
 
 	const UPCGSpawnActorSettings* Settings = Context->GetInputSettings<UPCGSpawnActorSettings>();
@@ -967,7 +967,7 @@ void FPCGSpawnActorElement::SpawnActors(FPCGSubgraphContext* Context, AActor* Ta
 			{
 				FPCGPoint& OutPoint = OutPointData->GetMutablePoints()[i + OutPointOffset];
 				OutPointData->Metadata->InitializeOnSet(OutPoint.MetadataEntry);
-				ActorReferenceAttribute->SetValue(OutPoint.MetadataEntry, FSoftObjectPath(GeneratedActor).ToString());
+				ActorReferenceAttribute->SetValue(OutPoint.MetadataEntry, FSoftObjectPath(GeneratedActor));
 			}
 		}
 
diff --git a/Engine/Plugins/Experimental/PCG/Source/PCG/Private/Elements/PCGWorldQuery.cpp b/Engine/Plugins/Experimental/PCG/Source/PCG/Private/Elements/PCGWorldQuery.cpp
--- a/Engine/Plugins/Experimental/PCG/Source/PCG/Private/Elements/PCGWorldQuery.cpp
+++ b/Engine/Plugins/Experimental/PCG/Source/PCG/Private/Elements/PCGWorldQuery.cpp
@@ -135,7 +135,7 @@ bool FPCGWorldRayHitQueryElement::ExecuteInternal(FPCGContext* Context) const
 
 	if (Data->QueryParams.bGetReferenceToActorHit && Data->Metadata)
 	{
-		Data->Metadata->FindOrCreateAttribute(PCGPointDataConstants::ActorReferenceAttribute, FString(), /*bAllowInterpolation=*/false, /*bOverrideParent=*/false);
+		Data->Metadata->FindOrCreateAttribute(PCGPointDataConstants::ActorReferenceAttribute, FSoftObjectPath(), /*bAllowInterpolation=*/false, /*bOverrideParent=*/false);
 	}
 	else
 	{
@@ -144,7 +144,7 @@ bool FPCGWorldRayHitQueryElement::ExecuteInternal(FPCGContext* Context) const
 
 	if (Data->QueryParams.bGetReferenceToPhysicalMaterial && Data->Metadata)
 	{
-		Data->Metadata->FindOrCreateAttribute(PCGWorldRayHitConstants::PhysicalMaterialReferenceAttribute, FString(), /*bAllowInterpolation=*/false, /*bOverrideParent=*/false);
+		Data->Metadata->FindOrCreateAttribute(PCGWorldRayHitConstants::PhysicalMaterialReferenceAttribute, FSoftObjectPath(), /*bAllowInterpolation=*/false, /*bOverrideParent=*/false);
 	}
 	else
 	{
