[PCG] Make sure AddAttribute from source always take the first entry value as default value for the newly created attribute if there is just one entry

#rnx

diff --git a/Engine/Plugins/Experimental/PCG/Source/PCG/Private/Elements/PCGCreateAttribute.cpp b/Engine/Plugins/Experimental/PCG/Source/PCG/Private/Elements/PCGCreateAttribute.cpp
--- a/Engine/Plugins/Experimental/PCG/Source/PCG/Private/Elements/PCGCreateAttribute.cpp
+++ b/Engine/Plugins/Experimental/PCG/Source/PCG/Private/Elements/PCGCreateAttribute.cpp
@@ -271,6 +271,8 @@ bool FPCGCreateAttributeElement::ExecuteInternal(FPCGContext* Context) const
 	FName SourceParamAttributeName = NAME_None;
 	FName OutputAttributeName = NAME_None;
 
+	FPCGAttributePropertyInputSelector InputSource{};
+
 	if (!SourceParams.IsEmpty())
 	{
 		SourceParamData = CastChecked<UPCGParamData>(SourceParams[0].Data);
@@ -281,7 +283,7 @@ bool FPCGCreateAttributeElement::ExecuteInternal(FPCGContext* Context) const
 			return true;
 		}
 
-		FPCGAttributePropertyInputSelector InputSource = Settings->InputSource.CopyAndFixLast(SourceParamData);
+		InputSource = Settings->InputSource.CopyAndFixLast(SourceParamData);
 
 		SourceParamAttributeName = InputSource.GetName();
 		OutputAttributeName = Settings->GetOutputAttributeName(&InputSource, SourceParamData);
@@ -352,8 +354,6 @@ bool FPCGCreateAttributeElement::ExecuteInternal(FPCGContext* Context) const
 
 		if (SourceParamData)
 		{
-			const FPCGAttributePropertyInputSelector InputSource = Settings->InputSource.CopyAndFixLast(SourceParamData);
-
 			// We need accessors if we have a multi entry source attribute or we have extractors
 			const bool bIsMultiEntries = SourceParamData->Metadata->GetLocalItemCount() > 1;
 			const bool bNeedAccessors = bIsMultiEntries || !InputSource.GetExtraNames().IsEmpty();
@@ -362,13 +362,18 @@ bool FPCGCreateAttributeElement::ExecuteInternal(FPCGContext* Context) const
 			if (!bNeedAccessors)
 			{
 				const FPCGMetadataAttributeBase* SourceAttribute = SourceParamData->Metadata->GetConstAttribute(SourceParamAttributeName);
-				Attribute = Metadata->CopyAttribute(SourceAttribute, OutputAttributeName, /*bKeepParent=*/false, /*bCopyEntries=*/false, /*bCopyValues=*/false);
+				// Presence of attribute was already checked before, this should not return null
+				check(SourceAttribute);
 
-				// We perhaps need to fix the default value. If the first entry is different from the default value, we override the default value.
-				if (Attribute)
+				// Copy the attribute using the first entry of the source attribute as the default value (there is just a single entry or none). If there is no first entry, will be the default value anyway.
+				auto CreateAttribute = [Metadata, SourceAttribute, OutputAttributeName](auto Dummy) -> FPCGMetadataAttributeBase*
 				{
-					Attribute->SetDefaultValueToFirstEntry();
-				}
+					using AttributeType = decltype(Dummy);
+					AttributeType DefaultValue = static_cast<const FPCGMetadataAttribute<AttributeType>*>(SourceAttribute)->GetValue(PCGMetadataEntryKey(0));
+					return PCGMetadataElementCommon::ClearOrCreateAttribute<AttributeType>(Metadata, OutputAttributeName, DefaultValue);
+				};
+
+				Attribute = PCGMetadataAttribute::CallbackWithRightType(SourceAttribute->GetTypeId(), std::move(CreateAttribute));
 			}
 			else // Create a new attribute of the accessed field's type manually
 			{
