Revisited cancellation to cleanup successors a bit more completely to remove spurious cases, which will the executor validation a bit more robust.

diff --git a/Engine/Plugins/Experimental/PCG/Source/PCG/Private/Elements/PCGLoopElement.cpp b/Engine/Plugins/Experimental/PCG/Source/PCG/Private/Elements/PCGLoopElement.cpp
--- a/Engine/Plugins/Experimental/PCG/Source/PCG/Private/Elements/PCGLoopElement.cpp
+++ b/Engine/Plugins/Experimental/PCG/Source/PCG/Private/Elements/PCGLoopElement.cpp
@@ -161,9 +161,11 @@ bool FPCGLoopElement::ExecuteInternal(FPCGContext* InContext) const
 			for (FPCGTaskId SubgraphTaskId : Context->SubgraphTaskIds)
 			{
 				FPCGDataCollection SubgraphOutput;
-				ensure(Subsystem->GetOutputData(SubgraphTaskId, SubgraphOutput));
-
-				Context->OutputData.TaggedData.Append(SubgraphOutput.TaggedData);
+				// While this should be always return true, if a scheduled task was cancelled, this can still happen because the dependency has been removed but this code isn't aware of this
+				if (Subsystem->GetOutputData(SubgraphTaskId, SubgraphOutput))
+				{
+					Context->OutputData.TaggedData.Append(SubgraphOutput.TaggedData);
+				}
 			}
 		}
 		else
diff --git a/Engine/Plugins/Experimental/PCG/Source/PCG/Private/Graph/PCGGraphExecutor.cpp b/Engine/Plugins/Experimental/PCG/Source/PCG/Private/Graph/PCGGraphExecutor.cpp
--- a/Engine/Plugins/Experimental/PCG/Source/PCG/Private/Graph/PCGGraphExecutor.cpp
+++ b/Engine/Plugins/Experimental/PCG/Source/PCG/Private/Graph/PCGGraphExecutor.cpp
@@ -327,6 +327,8 @@ TSet<UPCGComponent*> FPCGGraphExecutor::Cancel(TFunctionRef<bool(TWeakObjectPtr<
 			if (CancelledComponents.Contains(Task.SourceComponent.Get()))
 			{
 				FPCGTaskId CancelledTaskId = Task.NodeId;
+				RemoveTaskFromInputSuccessors(CancelledTaskId, Task.Inputs);
+
 				delete Task.Context;
 				ReadyTasks.RemoveAtSwap(ReadyTaskIndex);
 				bStableCancellationSet &= !CancelNextTasks(CancelledTaskId, CancelledComponents);
@@ -341,6 +343,8 @@ TSet<UPCGComponent*> FPCGGraphExecutor::Cancel(TFunctionRef<bool(TWeakObjectPtr<
 			{
 				FPCGTaskId CancelledTaskId = Task.NodeId;
 				Task.bWasCancelled = true;
+				RemoveTaskFromInputSuccessors(CancelledTaskId, Task.Inputs);
+
 				bStableCancellationSet &= !CancelNextTasks(CancelledTaskId, CancelledComponents);
 			}
 		}
@@ -352,6 +356,8 @@ TSet<UPCGComponent*> FPCGGraphExecutor::Cancel(TFunctionRef<bool(TWeakObjectPtr<
 			if (Task.Context && CancelledComponents.Contains(Task.Context->SourceComponent.Get()))
 			{
 				FPCGTaskId CancelledTaskId = Task.NodeId;
+				RemoveTaskFromInputSuccessors(CancelledTaskId, Task.Inputs);
+
 				SleepingTasks.RemoveAtSwap(SleepingTaskIndex);
 				bStableCancellationSet &= !CancelNextTasks(CancelledTaskId, CancelledComponents);
 			}
@@ -693,6 +699,7 @@ void FPCGGraphExecutor::Execute()
 				if (!bIsMainThreadTask || bMainThreadAvailable)
 				{
 					FPCGGraphActiveTask& ActiveTask = ActiveTasks.Emplace_GetRef();
+					ActiveTask.Inputs = MoveTemp(Task.Inputs);
 					ActiveTask.Element = Task.Element;
 					ActiveTask.NodeId = Task.NodeId;
 					ActiveTask.Context = TUniquePtr<FPCGContext>(Task.Context);
@@ -1019,9 +1026,8 @@ void FPCGGraphExecutor::QueueNextTasks(FPCGTaskId FinishedTask)
 			bool bAllPrerequisitesMet = true;
 			FPCGGraphTask* SuccessorTaskPtr = Tasks.Find(Successor);
 
-			// This should rarely be null, but can happen when we have a task waiting on the execution of multiple other components, and one/many get cancelled.
-			// Example: R gets data from C0 and C1. C0 is cancelled due to a change in params, which will cancel R, but not C1. C1 finishes executing and has a null successor here, but that's fine.
-			if (SuccessorTaskPtr)
+			// This should never be null, but later recovery should be able to cleanup this properly
+			if (ensure(SuccessorTaskPtr))
 			{
 				FPCGGraphTask& SuccessorTask = *SuccessorTaskPtr;
 
@@ -1048,7 +1054,10 @@ bool FPCGGraphExecutor::CancelNextTasks(FPCGTaskId CancelledTask, TSet<UPCGCompo
 
 	if (TSet<FPCGTaskId>* Successors = TaskSuccessors.Find(CancelledTask))
 	{
-		for (FPCGTaskId Successor : *Successors)
+		TSet<FPCGTaskId> LocalSuccessors = MoveTemp(*Successors);
+		TaskSuccessors.Remove(CancelledTask);
+
+		for (FPCGTaskId Successor : LocalSuccessors)
 		{
 			if (FPCGGraphTask* Task = Tasks.Find(Successor))
 			{
@@ -1058,13 +1067,12 @@ bool FPCGGraphExecutor::CancelNextTasks(FPCGTaskId CancelledTask, TSet<UPCGCompo
 					bAddedComponents = true;
 				}
 
+				RemoveTaskFromInputSuccessors(Task->NodeId, Task->Inputs);
 				Tasks.Remove(Successor);
 			}
 
 			bAddedComponents |= CancelNextTasks(Successor, OutCancelledComponents);
 		}
-
-		TaskSuccessors.Remove(CancelledTask);
 	}
 
 	// Tasks cancelled might have an impact on scheduled-but-not-processed tasks
@@ -1083,6 +1091,17 @@ bool FPCGGraphExecutor::CancelNextTasks(FPCGTaskId CancelledTask, TSet<UPCGCompo
 	return bAddedComponents;
 }
 
+void FPCGGraphExecutor::RemoveTaskFromInputSuccessors(FPCGTaskId CancelledTask, const TArray<FPCGGraphTaskInput>& CancelledTaskInputs)
+{
+	for (const FPCGGraphTaskInput& Input : CancelledTaskInputs)
+	{
+		if (TSet<FPCGTaskId>* Successors = TaskSuccessors.Find(Input.TaskId))
+		{
+			Successors->Remove(CancelledTask);
+		}
+	}
+}
+
 void FPCGGraphExecutor::BuildTaskInput(const FPCGGraphTask& Task, FPCGDataCollection& TaskInput)
 {
 	TRACE_CPUPROFILER_EVENT_SCOPE(FPCGGraphExecutor::BuildTaskInput);
diff --git a/Engine/Plugins/Experimental/PCG/Source/PCG/Private/Graph/PCGGraphExecutor.h b/Engine/Plugins/Experimental/PCG/Source/PCG/Private/Graph/PCGGraphExecutor.h
--- a/Engine/Plugins/Experimental/PCG/Source/PCG/Private/Graph/PCGGraphExecutor.h
+++ b/Engine/Plugins/Experimental/PCG/Source/PCG/Private/Graph/PCGGraphExecutor.h
@@ -70,6 +70,7 @@ struct FPCGGraphScheduleTask
 
 struct FPCGGraphActiveTask
 {
+	TArray<FPCGGraphTaskInput> Inputs;
 	FPCGElementPtr Element;
 	TUniquePtr<FPCGContext> Context;
 	FPCGTaskId NodeId = InvalidPCGTaskId;
@@ -150,6 +151,7 @@ private:
 	void ClearAllTasks();
 	void QueueNextTasks(FPCGTaskId FinishedTask);
 	bool CancelNextTasks(FPCGTaskId CancelledTask, TSet<UPCGComponent*>& OutCancelledComponents);
+	void RemoveTaskFromInputSuccessors(FPCGTaskId CancelledTask, const TArray<FPCGGraphTaskInput>& CancelledTaskInputs);
 	void BuildTaskInput(const FPCGGraphTask& Task, FPCGDataCollection& TaskInput);
 	/** Combine all param data into one on the Params pin, if any.*/
 	void CombineParams(FPCGTaskId InTaskId, FPCGDataCollection& InTaskInput);
diff --git a/Engine/Plugins/Experimental/PCG/Source/PCG/Private/PCGSubgraph.cpp b/Engine/Plugins/Experimental/PCG/Source/PCG/Private/PCGSubgraph.cpp
--- a/Engine/Plugins/Experimental/PCG/Source/PCG/Private/PCGSubgraph.cpp
+++ b/Engine/Plugins/Experimental/PCG/Source/PCG/Private/PCGSubgraph.cpp
@@ -579,7 +579,7 @@ bool FPCGSubgraphElement::ExecuteInternal(FPCGContext* InContext) const
 					// as we just can call the GetOutputData on a fresh data collection and merge it afterwards in the output data,
 					// but this is not needed here so we will keep the full assignment & benefit from the crc as well.
 					ensure(Context->SubgraphTaskIds.Num() == 1);
-					ensure(Subsystem->GetOutputData(Context->SubgraphTaskIds[0], Context->OutputData));
+					Subsystem->GetOutputData(Context->SubgraphTaskIds[0], Context->OutputData);
 				}
 			}
 			else
