[PCG] Prevent RuntimeGen system from collecting GenSources when there are no runtime gen components in the level

#rnx

diff --git a/Engine/Plugins/Experimental/PCG/Source/PCG/Private/PCGActorAndComponentMapping.cpp b/Engine/Plugins/Experimental/PCG/Source/PCG/Private/PCGActorAndComponentMapping.cpp
--- a/Engine/Plugins/Experimental/PCG/Source/PCG/Private/PCGActorAndComponentMapping.cpp
+++ b/Engine/Plugins/Experimental/PCG/Source/PCG/Private/PCGActorAndComponentMapping.cpp
@@ -237,6 +237,8 @@ bool FPCGActorAndComponentMapping::RegisterOrUpdatePCGComponent(UPCGComponent* I
 		UnregisterPartitionedPCGComponent(InComponent);
 	}
 
+	PCGSubsystem->OnOriginalComponentRegistered(InComponent);
+
 	// Then register/update accordingly
 	bool bHasChanged = false;
 	if (bIsPartitioned)
@@ -416,6 +418,11 @@ void FPCGActorAndComponentMapping::UnregisterPCGComponent(UPCGComponent* InCompo
 #endif // WITH_EDITOR
 	}
 
+	if (InComponent->GetOwner() && !InComponent->GetOwner()->IsA<APCGPartitionActor>())
+	{
+		PCGSubsystem->OnOriginalComponentUnregistered(InComponent);
+	}
+
 	UnregisterPartitionedPCGComponent(InComponent);
 	UnregisterNonPartitionedPCGComponent(InComponent);
 
@@ -428,13 +435,6 @@ void FPCGActorAndComponentMapping::UnregisterPCGComponent(UPCGComponent* InCompo
 
 void FPCGActorAndComponentMapping::UnregisterPartitionedPCGComponent(UPCGComponent* InComponent)
 {
-	check(InComponent);
-
-	if (InComponent->GetOwner() && !InComponent->GetOwner()->IsA<APCGPartitionActor>())
-	{
-		PCGSubsystem->OnOriginalComponentUnregistered(InComponent);
-	}
-
 	if (!PartitionedOctree.RemoveComponent(InComponent) || InComponent->GenerationTrigger == EPCGComponentGenerationTrigger::GenerateAtRuntime)
 	{
 		return;
@@ -462,13 +462,6 @@ void FPCGActorAndComponentMapping::UnregisterPartitionedPCGComponent(UPCGCompone
 
 void FPCGActorAndComponentMapping::UnregisterNonPartitionedPCGComponent(UPCGComponent* InComponent)
 {
-	check(InComponent);
-
-	if (InComponent->GetOwner() && !InComponent->GetOwner()->IsA<APCGPartitionActor>())
-	{
-		PCGSubsystem->OnOriginalComponentUnregistered(InComponent);
-	}
-
 	NonPartitionedOctree.RemoveComponent(InComponent);
 }
 
@@ -716,6 +709,27 @@ bool FPCGActorAndComponentMapping::IsComponentRegistered(const UPCGComponent* In
 	return PartitionedOctree.Contains(InComponent) || NonPartitionedOctree.Contains(InComponent);
 }
 
+bool FPCGActorAndComponentMapping::AnyRuntimeGenComponentsExist() const
+{
+	for (UPCGComponent* Component : PartitionedOctree.GetAllComponents())
+	{
+		if (Component && Component->GenerationTrigger == EPCGComponentGenerationTrigger::GenerateAtRuntime)
+		{
+			return true;
+		}
+	}
+
+	for (UPCGComponent* Component : NonPartitionedOctree.GetAllComponents())
+	{
+		if (Component && Component->GenerationTrigger == EPCGComponentGenerationTrigger::GenerateAtRuntime)
+		{
+			return true;
+		}
+	}
+
+	return false;
+}
+
 TSet<UPCGComponent*> FPCGActorAndComponentMapping::GetAllRegisteredPartitionedComponents() const
 {
 	return PartitionedOctree.GetAllComponents();
diff --git a/Engine/Plugins/Experimental/PCG/Source/PCG/Private/PCGSubsystem.cpp b/Engine/Plugins/Experimental/PCG/Source/PCG/Private/PCGSubsystem.cpp
--- a/Engine/Plugins/Experimental/PCG/Source/PCG/Private/PCGSubsystem.cpp
+++ b/Engine/Plugins/Experimental/PCG/Source/PCG/Private/PCGSubsystem.cpp
@@ -328,6 +328,14 @@ void UPCGSubsystem::UnregisterPCGWorldActor(APCGWorldActor* InActor)
 	}
 }
 
+void UPCGSubsystem::OnOriginalComponentRegistered(UPCGComponent* InComponent)
+{
+	if (RuntimeGenScheduler)
+	{
+		RuntimeGenScheduler->OnOriginalComponentRegistered(InComponent);
+	}
+}
+
 void UPCGSubsystem::OnOriginalComponentUnregistered(UPCGComponent* InComponent)
 {
 	if (RuntimeGenScheduler)
diff --git a/Engine/Plugins/Experimental/PCG/Source/PCG/Private/RuntimeGen/PCGRuntimeGenScheduler.cpp b/Engine/Plugins/Experimental/PCG/Source/PCG/Private/RuntimeGen/PCGRuntimeGenScheduler.cpp
--- a/Engine/Plugins/Experimental/PCG/Source/PCG/Private/RuntimeGen/PCGRuntimeGenScheduler.cpp
+++ b/Engine/Plugins/Experimental/PCG/Source/PCG/Private/RuntimeGen/PCGRuntimeGenScheduler.cpp
@@ -87,7 +87,7 @@ void FPCGRuntimeGenScheduler::Tick(const APCGWorldActor* InPCGWorldActor)
 
 	TickCVars(InPCGWorldActor);
 
-	const TSet<IPCGGenSourceBase*> GenSources = GenSourceManager->GetGenSources(InPCGWorldActor);
+	const TSet<IPCGGenSourceBase*> GenSources = bAnyRuntimeGenComponentsExist ? GenSourceManager->GetGenSources(InPCGWorldActor) : TSet<IPCGGenSourceBase*>();
 
 	// Initialize RuntimeGen PA pool if necessary. If PoolSize is 0, then we have not initialized the pool yet.
 	if (!GenSources.IsEmpty() || !GeneratedComponents.IsEmpty())
@@ -102,11 +102,17 @@ void FPCGRuntimeGenScheduler::Tick(const APCGWorldActor* InPCGWorldActor)
 	
 	// Mapping of component + coordinates to priorities - needed to compute max priority over all gen sources.
 	TMap<FGridGenerationKey, double> ComponentsToGenerate;
-	TickQueueComponentsForGeneration(GenSources, InPCGWorldActor, ComponentsToGenerate);
+	if (!GenSources.IsEmpty())
+	{
+		TickQueueComponentsForGeneration(GenSources, InPCGWorldActor, ComponentsToGenerate);
+	}
 
 	// 2. Schedule cleanup on components that become out of range.
 
-	TickCleanup(GenSources, InPCGWorldActor);
+	if (!GeneratedComponents.IsEmpty())
+	{
+		TickCleanup(GenSources, InPCGWorldActor);
+	}
 
 	// 3. Schedule generation on components in priority order.
 
@@ -128,7 +134,7 @@ void FPCGRuntimeGenScheduler::Tick(const APCGWorldActor* InPCGWorldActor)
 
 bool FPCGRuntimeGenScheduler::ShouldTick()
 {
-	check(World);
+	check(World && ActorAndComponentMapping);
 
 	if (!PCGRuntimeGenSchedulerHelpers::CVarRuntimeGenerationEnable.GetValueOnAnyThread())
 	{
@@ -159,6 +165,34 @@ bool FPCGRuntimeGenScheduler::ShouldTick()
 	}
 #endif
 
+	if (bAnyRuntimeGenComponentsExistDirty)
+	{
+		const bool bDidAnyRuntimeGenComponentsExist = bAnyRuntimeGenComponentsExist;
+		bAnyRuntimeGenComponentsExist = ActorAndComponentMapping->AnyRuntimeGenComponentsExist();
+		bAnyRuntimeGenComponentsExistDirty = false;
+
+		if (PCGRuntimeGenSchedulerHelpers::CVarRuntimeGenerationEnableDebugging.GetValueOnAnyThread())
+		{
+			if (bDidAnyRuntimeGenComponentsExist != bAnyRuntimeGenComponentsExist)
+			{
+				if (bAnyRuntimeGenComponentsExist)
+				{
+					UE_LOG(LogPCG, Warning, TEXT("[RUNTIMEGEN] THERE ARE NOW RUNTIME COMPONENTS IN THE LEVEL. SCHEDULER WILL BEGIN TICKING."));
+				}
+				else
+				{
+					UE_LOG(LogPCG, Warning, TEXT("[RUNTIMEGEN] THERE ARE NO MORE RUNTIME COMPONENTS. SCHEDULER WILL ONLY TICK TO CLEANUP."));
+				}
+			}
+		}
+	}
+
+	// We can stop ticking if there are no runtime gen components alive and there are no generated components that need cleaning up.
+	if (!bAnyRuntimeGenComponentsExist && GeneratedComponents.IsEmpty())
+	{
+		return false;
+	}
+
 	return true;
 }
 
@@ -650,6 +684,19 @@ void FPCGRuntimeGenScheduler::TickCVars(const APCGWorldActor* InPCGWorldActor)
 	}
 }
 
+void FPCGRuntimeGenScheduler::OnOriginalComponentRegistered(UPCGComponent* InOriginalComponent)
+{
+	// Ensure we are not a local component.
+	if (!InOriginalComponent || Cast<APCGPartitionActor>(InOriginalComponent->GetOwner()))
+	{
+		ensure(false);
+		return;
+	}
+
+	// When an original/non-partitioned component is registered, we need to dirty the state.
+	bAnyRuntimeGenComponentsExistDirty = true;
+}
+
 void FPCGRuntimeGenScheduler::OnOriginalComponentUnregistered(UPCGComponent* InOriginalComponent)
 {
 	check(ActorAndComponentMapping);
@@ -661,6 +708,9 @@ void FPCGRuntimeGenScheduler::OnOriginalComponentUnregistered(UPCGComponent* InO
 		return;
 	}
 
+	// When an original/non-partitioned component is unregistered, we need to dirty the state.
+	bAnyRuntimeGenComponentsExistDirty = true;
+
 	// Gather all generated components which originated from this original component.
 	TSet<FGridGenerationKey> KeysToCleanup;
 	for (FGridGenerationKey GenerationKey : GeneratedComponents)
@@ -929,6 +979,9 @@ void FPCGRuntimeGenScheduler::RefreshComponent(UPCGComponent* InComponent, bool
 				RefreshLocalComponent(LocalComponent);
 			}
 		}
+
+		// When an original/non-partitioned component is refreshed, we need to dirty the state.
+		bAnyRuntimeGenComponentsExistDirty = true;
 	}
 }
 
diff --git a/Engine/Plugins/Experimental/PCG/Source/PCG/Public/PCGActorAndComponentMapping.h b/Engine/Plugins/Experimental/PCG/Source/PCG/Public/PCGActorAndComponentMapping.h
--- a/Engine/Plugins/Experimental/PCG/Source/PCG/Public/PCGActorAndComponentMapping.h
+++ b/Engine/Plugins/Experimental/PCG/Source/PCG/Public/PCGActorAndComponentMapping.h
@@ -124,6 +124,9 @@ private:
 	/** Returns true if a component is registered */
 	bool IsComponentRegistered(const UPCGComponent* InComponent) const;
 
+	/** Return true if there are any Original or Non-Partitioned components set to GenerateAtRuntime. */
+	bool AnyRuntimeGenComponentsExist() const;
+
 #if WITH_EDITOR
 	/* Return true if something is still tracked or was just untracked. */
 	bool AddOrUpdateTrackedActor(AActor* InActor);
diff --git a/Engine/Plugins/Experimental/PCG/Source/PCG/Public/PCGSubsystem.h b/Engine/Plugins/Experimental/PCG/Source/PCG/Public/PCGSubsystem.h
--- a/Engine/Plugins/Experimental/PCG/Source/PCG/Public/PCGSubsystem.h
+++ b/Engine/Plugins/Experimental/PCG/Source/PCG/Public/PCGSubsystem.h
@@ -84,6 +84,7 @@ public:
 	void RegisterPCGWorldActor(APCGWorldActor* InActor);
 	void UnregisterPCGWorldActor(APCGWorldActor* InActor);
 
+	void OnOriginalComponentRegistered(UPCGComponent* InComponent);
 	void OnOriginalComponentUnregistered(UPCGComponent* InComponent);
 
 	UPCGLandscapeCache* GetLandscapeCache();
diff --git a/Engine/Plugins/Experimental/PCG/Source/PCG/Public/RuntimeGen/PCGRuntimeGenScheduler.h b/Engine/Plugins/Experimental/PCG/Source/PCG/Public/RuntimeGen/PCGRuntimeGenScheduler.h
--- a/Engine/Plugins/Experimental/PCG/Source/PCG/Public/RuntimeGen/PCGRuntimeGenScheduler.h
+++ b/Engine/Plugins/Experimental/PCG/Source/PCG/Public/RuntimeGen/PCGRuntimeGenScheduler.h
@@ -38,7 +38,8 @@ public:
 
 	void Tick(const APCGWorldActor* InPCGWorldActor);
 
-	void OnOriginalComponentUnregistered(UPCGComponent* InComponent);
+	void OnOriginalComponentRegistered(UPCGComponent* InOriginalComponent);
+	void OnOriginalComponentUnregistered(UPCGComponent* InOriginalComponent);
 
 protected:
 	struct FGridGenerationKey : TTuple<uint32, FIntVector, UPCGComponent*>
@@ -113,4 +114,8 @@ private:
 
 	bool bPoolingWasEnabledLastFrame = true;
 	uint32 BasePoolSizeLastFrame = 0;
+
+	/** Track the existence of runtime gen components to avoid unnecessary computation when there is no work to do. */
+	bool bAnyRuntimeGenComponentsExist = false;
+	bool bAnyRuntimeGenComponentsExistDirty = false;
 };
