[PCG] Make a pass on CreateAttribute to use ZeroValue instead of default value

* Default value can yield to uninitialized fields (like for Vectors)

#rnx

diff --git a/Engine/Plugins/Experimental/PCG/Source/PCG/Private/Elements/Metadata/PCGMetadataElement.cpp b/Engine/Plugins/Experimental/PCG/Source/PCG/Private/Elements/Metadata/PCGMetadataElement.cpp
--- a/Engine/Plugins/Experimental/PCG/Source/PCG/Private/Elements/Metadata/PCGMetadataElement.cpp
+++ b/Engine/Plugins/Experimental/PCG/Source/PCG/Private/Elements/Metadata/PCGMetadataElement.cpp
@@ -223,7 +223,7 @@ bool FPCGMetadataOperationElement::ExecuteInternal(FPCGContext* Context) const
 			auto CreateAttribute = [SampledData, LocalDestinationAttribute](auto Dummy)
 			{
 				using AttributeType = decltype(Dummy);
-				return PCGMetadataElementCommon::ClearOrCreateAttribute(SampledData->Metadata, LocalDestinationAttribute, AttributeType{}) != nullptr;
+				return PCGMetadataElementCommon::ClearOrCreateAttribute<AttributeType>(SampledData->Metadata, LocalDestinationAttribute) != nullptr;
 			};
 			
 			if (!PCGMetadataAttribute::CallbackWithRightType(InputAccessor->GetUnderlyingType(), CreateAttribute))
diff --git a/Engine/Plugins/Experimental/PCG/Source/PCG/Private/Elements/Metadata/PCGMetadataOpElementBase.cpp b/Engine/Plugins/Experimental/PCG/Source/PCG/Private/Elements/Metadata/PCGMetadataOpElementBase.cpp
--- a/Engine/Plugins/Experimental/PCG/Source/PCG/Private/Elements/Metadata/PCGMetadataOpElementBase.cpp
+++ b/Engine/Plugins/Experimental/PCG/Source/PCG/Private/Elements/Metadata/PCGMetadataOpElementBase.cpp
@@ -524,7 +524,7 @@ bool FPCGMetadataElementBase::ExecuteInternal(FPCGContext* Context) const
 			// But if we have extra selectors, we need to handle it the same way as properties.
 			// There is no point of failure before duplicating. So duplicate, create the attribute and then the accessor.
 			PCGMetadataElementCommon::DuplicateTaggedData(InputTaggedData[InputPinToForward], OutputData, OutMetadata);
-			FPCGMetadataAttributeBase* OutputAttribute = PCGMetadataElementCommon::ClearOrCreateAttribute(OutMetadata, OutputName, AttributeType{});
+			FPCGMetadataAttributeBase* OutputAttribute = PCGMetadataElementCommon::ClearOrCreateAttribute<AttributeType>(OutMetadata, OutputName);
 			if (!OutputAttribute)
 			{
 				return false;
diff --git a/Engine/Plugins/Experimental/PCG/Source/PCG/Private/Elements/PCGAttributeTransferElement.cpp b/Engine/Plugins/Experimental/PCG/Source/PCG/Private/Elements/PCGAttributeTransferElement.cpp
--- a/Engine/Plugins/Experimental/PCG/Source/PCG/Private/Elements/PCGAttributeTransferElement.cpp
+++ b/Engine/Plugins/Experimental/PCG/Source/PCG/Private/Elements/PCGAttributeTransferElement.cpp
@@ -258,7 +258,7 @@ bool FPCGAttributeTransferElement::ExecuteInternal(FPCGContext* Context) const
 		auto CreateAttribute = [OutputData, TargetAttributeName](auto Dummy)
 		{
 			using AttributeType = decltype(Dummy);
-			return PCGMetadataElementCommon::ClearOrCreateAttribute(OutputData->Metadata, TargetAttributeName, AttributeType{}) != nullptr;
+			return PCGMetadataElementCommon::ClearOrCreateAttribute<AttributeType>(OutputData->Metadata, TargetAttributeName) != nullptr;
 		};
 		
 		if (!PCGMetadataAttribute::CallbackWithRightType(InputAccessor->GetUnderlyingType(), CreateAttribute))
diff --git a/Engine/Plugins/Experimental/PCG/Source/PCG/Private/MatchAndSet/PCGMatchAndSetBase.cpp b/Engine/Plugins/Experimental/PCG/Source/PCG/Private/MatchAndSet/PCGMatchAndSetBase.cpp
--- a/Engine/Plugins/Experimental/PCG/Source/PCG/Private/MatchAndSet/PCGMatchAndSetBase.cpp
+++ b/Engine/Plugins/Experimental/PCG/Source/PCG/Private/MatchAndSet/PCGMatchAndSetBase.cpp
@@ -59,7 +59,7 @@ bool UPCGMatchAndSetBase::CreateAttributeIfNeeded(FPCGContext& Context, const FP
 			auto CreateAttribute = [OutPointData, &DestinationAttribute](auto&& Value)
 			{
 				using ConstantType = std::decay_t<decltype(Value)>;
-				return PCGMetadataElementCommon::ClearOrCreateAttribute(OutPointData->Metadata, DestinationAttribute, ConstantType{}) != nullptr;
+				return PCGMetadataElementCommon::ClearOrCreateAttribute<ConstantType>(OutPointData->Metadata, DestinationAttribute) != nullptr;
 			};
 
 			if (!ConstantValue.Dispatcher(CreateAttribute))
diff --git a/Engine/Plugins/Experimental/PCG/Source/PCG/Private/Tests/Elements/PCGAttributeReduceTest.cpp b/Engine/Plugins/Experimental/PCG/Source/PCG/Private/Tests/Elements/PCGAttributeReduceTest.cpp
--- a/Engine/Plugins/Experimental/PCG/Source/PCG/Private/Tests/Elements/PCGAttributeReduceTest.cpp
+++ b/Engine/Plugins/Experimental/PCG/Source/PCG/Private/Tests/Elements/PCGAttributeReduceTest.cpp
@@ -126,7 +126,7 @@ protected:
 
 		UPCGParamData* ParamData = NewObject<UPCGParamData>();
 		FPCGMetadataAttribute<int32>* IntAttribute = ParamData->Metadata->CreateAttribute<int32>(IntName, 0, /*bAllowsInterpolation=*/ true, /*bOverrideParent=*/false);
-		FPCGMetadataAttribute<FVector2D>* Vec2Attribute = ParamData->Metadata->CreateAttribute<FVector2D>(Vec2Name, FVector2D{}, /*bAllowsInterpolation=*/ true, /*bOverrideParent=*/false);
+		FPCGMetadataAttribute<FVector2D>* Vec2Attribute = ParamData->Metadata->CreateAttribute<FVector2D>(Vec2Name, FVector2D::ZeroVector, /*bAllowsInterpolation=*/ true, /*bOverrideParent=*/false);
 
 		constexpr int32 EntryCount = 5;
 		for (int32 i = 0; i < EntryCount; ++i)
@@ -335,7 +335,7 @@ bool FPCGAttributeReduceIncompatibleType::RunTest(const FString& Parameters)
 	// String are always invalid
 	FPCGMetadataAttribute<FString>* StringAttribute = ParamData->Metadata->CreateAttribute<FString>(StringName, FString{}, /*bAllowsInterpolation=*/ false, /*bOverrideParent=*/false);
 	// Quat are invalid for min/max
-	FPCGMetadataAttribute<FQuat>* QuatAttribute = ParamData->Metadata->CreateAttribute<FQuat>(QuatName, FQuat{}, /*bAllowsInterpolation=*/ true, /*bOverrideParent=*/false);
+	FPCGMetadataAttribute<FQuat>* QuatAttribute = ParamData->Metadata->CreateAttribute<FQuat>(QuatName, FQuat{ForceInitToZero}, /*bAllowsInterpolation=*/ true, /*bOverrideParent=*/false);
 
 	constexpr int32 EntryCount = 5;
 	for (int32 i = 0; i < EntryCount; ++i)
diff --git a/Engine/Plugins/Experimental/PCG/Source/PCG/Public/Elements/Metadata/PCGMetadataElementCommon.h b/Engine/Plugins/Experimental/PCG/Source/PCG/Public/Elements/Metadata/PCGMetadataElementCommon.h
--- a/Engine/Plugins/Experimental/PCG/Source/PCG/Public/Elements/Metadata/PCGMetadataElementCommon.h
+++ b/Engine/Plugins/Experimental/PCG/Source/PCG/Public/Elements/Metadata/PCGMetadataElementCommon.h
@@ -4,6 +4,7 @@
 
 #include "PCGSettings.h"
 #include "Metadata/PCGMetadata.h"
+#include "Metadata/PCGMetadataAttributeTraits.h"
 #include "Metadata/Accessors/IPCGAttributeAccessor.h" // IWYU pragma: keep
 #include "Metadata/Accessors/PCGAttributeAccessorKeys.h"
 
@@ -19,9 +20,9 @@ namespace PCGMetadataElementCommon
 	/** Copies the entry to value key relationship stored in the given Metadata, including its parents */
 	PCG_API void CopyEntryToValueKeyMap(const UPCGMetadata* MetadataToCopy, const FPCGMetadataAttributeBase* AttributeToCopy, FPCGMetadataAttributeBase* OutAttribute);
 
-	/** Creates a new attribute, or clears the attribute if it already exists and is a 'T' type */
+	/** Creates a new attribute, or clears the attribute if it already exists and is a 'T' type. If default value not provided, will take the zero value for that type. */
 	template<typename T>
-	FPCGMetadataAttribute<T>* ClearOrCreateAttribute(UPCGMetadata* Metadata, const FName& DestinationAttribute, T DefaultValue)
+	FPCGMetadataAttribute<T>* ClearOrCreateAttribute(UPCGMetadata* Metadata, const FName& DestinationAttribute, T DefaultValue = PCG::Private::MetadataTraits<T>::ZeroValue())
 	{
 		TRACE_CPUPROFILER_EVENT_SCOPE(PCGMetadataElementCommon::ClearOrCreateAttribute);
 
diff --git a/Engine/Plugins/Experimental/PCGExternalDataInterop/Source/PCGExternalDataInterop/Private/Alembic/PCGAlembicInterop.cpp b/Engine/Plugins/Experimental/PCGExternalDataInterop/Source/PCGExternalDataInterop/Private/Alembic/PCGAlembicInterop.cpp
--- a/Engine/Plugins/Experimental/PCGExternalDataInterop/Source/PCGExternalDataInterop/Private/Alembic/PCGAlembicInterop.cpp
+++ b/Engine/Plugins/Experimental/PCGExternalDataInterop/Source/PCGExternalDataInterop/Private/Alembic/PCGAlembicInterop.cpp
@@ -194,7 +194,7 @@ bool CreatePointAccessorAndValidate(FPCGContext* Context, UPCGPointData* PointDa
 		auto CreateAttribute = [PointMetadata, &PropName](auto Dummy)
 		{
 			using AttributeType = decltype(Dummy);
-			return PCGMetadataElementCommon::ClearOrCreateAttribute(PointMetadata, FName(PropName), AttributeType{}) != nullptr;
+			return PCGMetadataElementCommon::ClearOrCreateAttribute<AttributeType>(PointMetadata, FName(PropName)) != nullptr;
 		};
 
 		if (!PCGMetadataAttribute::CallbackWithRightType(AlembicPropAccessor->GetUnderlyingType(), CreateAttribute))
