Tentative fix for landscape cache holding on other embedded landscape caches longer than they should

#rnx

diff --git a/Engine/Plugins/Experimental/PCG/Source/PCG/Private/Grid/PCGLandscapeCache.cpp b/Engine/Plugins/Experimental/PCG/Source/PCG/Private/Grid/PCGLandscapeCache.cpp
--- a/Engine/Plugins/Experimental/PCG/Source/PCG/Private/Grid/PCGLandscapeCache.cpp
+++ b/Engine/Plugins/Experimental/PCG/Source/PCG/Private/Grid/PCGLandscapeCache.cpp
@@ -400,6 +400,13 @@ void FPCGLandscapeCacheEntry::SerializeToBulkData(EPCGLandscapeCacheSerializatio
 bool FPCGLandscapeCacheEntry::SerializeFromBulkData() const
 {
 	check(!bDataLoaded);
+
+	// If owner object has been unloaded, we can't serialize the bulk data.
+	if (!OwningCache.IsValid())
+	{
+		return false;
+	}
+
 	LLM_SCOPE_BYNAME(TEXT("PCGLandscape"));
 
 	// Note: this call is not threadsafe by itself, it is meant to be called from a locked region
@@ -502,6 +509,7 @@ void UPCGLandscapeCache::Serialize(FArchive& Archive)
 			Archive << Key;
 
 			FPCGLandscapeCacheEntry* Entry = new FPCGLandscapeCacheEntry();
+			Entry->OwningCache = this;
 			Entry->Serialize(Archive, this, EntryIndex, SerializedContents);
 
 			CacheMapKey MapKey(Key.Key, Key.Value, nullptr);
@@ -670,10 +678,11 @@ void UPCGLandscapeCache::PrimeCache()
 	TArray<FPCGLandscapeCacheEntry*> NewEntries;
 	NewEntries.SetNum(CacheEntriesToBuild.Num());
 
-	ParallelFor(CacheEntriesToBuild.Num(), [&CacheEntriesToBuild, &NewEntries](int32 EntryIndex)
+	ParallelFor(CacheEntriesToBuild.Num(), [this, &CacheEntriesToBuild, &NewEntries](int32 EntryIndex)
 	{
 		const TTuple<ULandscapeInfo*, ULandscapeComponent*, CacheMapKey>& CacheEntryInfo = CacheEntriesToBuild[EntryIndex];
 		NewEntries[EntryIndex] = FPCGLandscapeCacheEntry::CreateCacheEntry(CacheEntryInfo.Get<0>(), CacheEntryInfo.Get<1>());
+		NewEntries[EntryIndex]->OwningCache = this;
 	});
 
 	if (SerializationMode != EPCGLandscapeCacheSerializationMode::NeverSerialize)
@@ -748,11 +757,6 @@ void UPCGLandscapeCache::TakeOwnership(UPCGLandscapeCache* InLandscapeCache)
 	CacheEntryCount = CachedData.Num();
 #endif
 
-	if (bShouldDirty)
-	{
-		EmbeddedCaches.AddUnique(InLandscapeCache);
-	}
-
 	if (bShouldDirty && SerializationMode != EPCGLandscapeCacheSerializationMode::NeverSerialize)
 	{
 		MarkPackageDirty();
@@ -783,6 +787,7 @@ const FPCGLandscapeCacheEntry* UPCGLandscapeCache::GetCacheEntry(ULandscapeCompo
 			check(LandscapeComponent->SectionBaseX / LandscapeComponent->ComponentSizeQuads == ComponentKey.Coordinate.X && LandscapeComponent->SectionBaseY / LandscapeComponent->ComponentSizeQuads == ComponentKey.Coordinate.Y);
 			if (FPCGLandscapeCacheEntry* NewEntry = FPCGLandscapeCacheEntry::CreateCacheEntry(LandscapeComponent->GetLandscapeInfo(), LandscapeComponent))
 			{
+				NewEntry->OwningCache = this;
 				CacheEntry = NewEntry;
 				CachedData.Add(ComponentKey, NewEntry);
 				++CacheEntryCount;
diff --git a/Engine/Plugins/Experimental/PCG/Source/PCG/Public/Grid/PCGLandscapeCache.h b/Engine/Plugins/Experimental/PCG/Source/PCG/Public/Grid/PCGLandscapeCache.h
--- a/Engine/Plugins/Experimental/PCG/Source/PCG/Public/Grid/PCGLandscapeCache.h
+++ b/Engine/Plugins/Experimental/PCG/Source/PCG/Public/Grid/PCGLandscapeCache.h
@@ -14,6 +14,7 @@ class ALandscapeProxy;
 class FLandscapeProxyComponentDataChangedParams;
 class ULandscapeComponent;
 class ULandscapeInfo;
+class UPCGLandscapeCache;
 class UPCGPointData;
 class UPCGMetadata;
 struct FPCGPoint;
@@ -112,6 +113,7 @@ private:
 	mutable FCriticalSection DataLock;
 	mutable int32 Touch = 0;
 	mutable bool bDataLoaded = false;
+	TWeakObjectPtr<UPCGLandscapeCache> OwningCache = nullptr;
 };
 
 UCLASS()
@@ -226,10 +228,6 @@ private:
 	TSet<TWeakObjectPtr<ALandscapeProxy>> Landscapes;
 #endif
 
-	// Keep "merged" landscape caches alive so that the bulk data is reachable
-	UPROPERTY(Transient)
-	TArray<TObjectPtr<UPCGLandscapeCache>> EmbeddedCaches;
-
 #if WITH_EDITOR
 	FRWLock CacheLock;
 #endif
