[PCG] Fix CopyPoints relative rotation for negative scale

#rnx

diff --git a/Engine/Plugins/Experimental/PCG/Source/PCG/Private/Elements/PCGCopyPoints.cpp b/Engine/Plugins/Experimental/PCG/Source/PCG/Private/Elements/PCGCopyPoints.cpp
--- a/Engine/Plugins/Experimental/PCG/Source/PCG/Private/Elements/PCGCopyPoints.cpp
+++ b/Engine/Plugins/Experimental/PCG/Source/PCG/Private/Elements/PCGCopyPoints.cpp
@@ -206,10 +206,15 @@ bool FPCGCopyPointsElement::ExecuteInternal(FPCGContext* Context) const
 
 				OutPoint = SourcePoint;
 
+				// Compute the RelativeTransform only once.
+				const bool bUseRelativeTransform = (RotationInheritance == EPCGCopyPointsInheritanceMode::Relative) || (ScaleInheritance == EPCGCopyPointsInheritanceMode::Relative);
+				const FTransform RelativeTransform = bUseRelativeTransform ? SourcePoint.Transform.GetRelativeTransform(TargetPoint.Transform) : FTransform{};
+
 				// Set Rotation, Scale, and Color based on inheritance mode
 				if (RotationInheritance == EPCGCopyPointsInheritanceMode::Relative)
 				{
-					OutPoint.Transform.SetRotation(TargetPoint.Transform.GetRotation() * SourcePoint.Transform.GetRotation());
+					// GetRelativeTransform will take into account if the scale is negative
+					OutPoint.Transform.SetRotation(RelativeTransform.GetRotation());
 				}
 				else if (RotationInheritance == EPCGCopyPointsInheritanceMode::Source)
 				{
@@ -222,7 +227,8 @@ bool FPCGCopyPointsElement::ExecuteInternal(FPCGContext* Context) const
 
 				if (ScaleInheritance == EPCGCopyPointsInheritanceMode::Relative)
 				{
-					OutPoint.Transform.SetScale3D(SourcePoint.Transform.GetScale3D() * TargetPoint.Transform.GetScale3D());
+					// GetRelativeTransform will take into account if the scale is negative
+					OutPoint.Transform.SetScale3D(RelativeTransform.GetScale3D());
 				}
 				else if (ScaleInheritance == EPCGCopyPointsInheritanceMode::Source)
 				{ 
