[PCG] Add actor property overrides to spawner nodes

diff --git a/Engine/Plugins/Experimental/PCG/Source/PCG/Private/Elements/PCGCreateTargetActor.cpp b/Engine/Plugins/Experimental/PCG/Source/PCG/Private/Elements/PCGCreateTargetActor.cpp
--- a/Engine/Plugins/Experimental/PCG/Source/PCG/Private/Elements/PCGCreateTargetActor.cpp
+++ b/Engine/Plugins/Experimental/PCG/Source/PCG/Private/Elements/PCGCreateTargetActor.cpp
@@ -30,6 +30,12 @@ FText UPCGCreateTargetActor::GetDefaultNodeTitle() const
 
 #endif // WITH_EDITOR
 
+namespace PCGCreateTargetActorConstants
+{
+	const FName ActorPropertyOverridesLabel = TEXT("Property Overrides");
+	const FText ActorPropertyOverridesTooltip = LOCTEXT("ActorOverrideToolTip", "Provide property overrides for the created target actor. The attribute name must match the InputSource name in the actor property override description.");
+}
+
 FPCGElementPtr UPCGCreateTargetActor::CreateElement() const
 {
 	return MakeShared<FPCGCreateTargetActorElement>();
@@ -37,7 +43,9 @@ FPCGElementPtr UPCGCreateTargetActor::CreateElement() const
 
 TArray<FPCGPinProperties> UPCGCreateTargetActor::InputPinProperties() const
 {
-	return TArray<FPCGPinProperties>();
+	TArray<FPCGPinProperties> PinProperties;
+	PinProperties.Add(PCGActorPropertyOverrideHelpers::CreateActorPropertiesOverridePin(PCGCreateTargetActorConstants::ActorPropertyOverridesLabel, PCGCreateTargetActorConstants::ActorPropertyOverridesTooltip));
+	return PinProperties;
 }
 
 TArray<FPCGPinProperties> UPCGCreateTargetActor::OutputPinProperties() const
@@ -283,6 +291,9 @@ bool FPCGCreateTargetActorElement::ExecuteInternal(FPCGContext* Context) const
 
 	GeneratedActor->Tags.Add(PCGHelpers::DefaultPCGActorTag);
 
+	// Apply property overrides to the GeneratedActor
+	PCGActorPropertyOverrideHelpers::ApplyOverridesFromParams(Settings->PropertyOverrideDescriptions, GeneratedActor, PCGCreateTargetActorConstants::ActorPropertyOverridesLabel, Context);
+
 	for (UFunction* Function : PCGHelpers::FindUserFunctions(GeneratedActor->GetClass(), Settings->PostProcessFunctionNames, Context))
 	{
 		GeneratedActor->ProcessEvent(Function, nullptr);
diff --git a/Engine/Plugins/Experimental/PCG/Source/PCG/Private/Elements/PCGSpawnActor.cpp b/Engine/Plugins/Experimental/PCG/Source/PCG/Private/Elements/PCGSpawnActor.cpp
--- a/Engine/Plugins/Experimental/PCG/Source/PCG/Private/Elements/PCGSpawnActor.cpp
+++ b/Engine/Plugins/Experimental/PCG/Source/PCG/Private/Elements/PCGSpawnActor.cpp
@@ -33,121 +33,6 @@ static TAutoConsoleVariable<bool> CVarAllowActorReuse(
 	true,
 	TEXT("Controls whether PCG spawned actors can be reused and skipped when re-executing"));
 
-namespace PCGSpawnActorHelpers
-{
-	struct FActorSingleOverride
-	{
-		using ApplyOverrideFunction = TFunction<void(int32, const IPCGAttributeAccessorKeys&, IPCGAttributeAccessorKeys&)>;
-
-		FActorSingleOverride(const FPCGAttributePropertySelector& InputSelector, const FString& OutputProperty, AActor* TemplateActor, const UPCGPointData* PointData)
-		{
-			ActorOverrideInputAccessor = PCGAttributeAccessorHelpers::CreateConstAccessor(PointData, InputSelector);
-			FPCGAttributePropertySelector OutputSelector = FPCGAttributePropertySelector::CreateSelectorFromString(OutputProperty);
-			const TArray<FString>& ExtraNames = OutputSelector.GetExtraNames();
-			if (ExtraNames.IsEmpty())
-			{
-				ActorOverrideOutputAccessor = PCGAttributeAccessorHelpers::CreatePropertyAccessor(FName(OutputProperty), TemplateActor->GetClass());
-			}
-			else
-			{
-				TArray<FName> PropertyNames;
-				PropertyNames.Reserve(ExtraNames.Num() + 1);
-				PropertyNames.Add(OutputSelector.GetAttributeName());
-				for (const FString& Name : ExtraNames)
-				{
-					PropertyNames.Add(FName(Name));
-				}
-
-				ActorOverrideOutputAccessor = PCGAttributeAccessorHelpers::CreatePropertyChainAccessor(PropertyNames, TemplateActor->GetClass());
-			}
-
-			if (!ActorOverrideInputAccessor.IsValid() || !ActorOverrideOutputAccessor.IsValid())
-			{
-				UE_LOG(LogPCG, Warning, TEXT("ActorOverride from input %s or output %s is invalid or unsupported. Will be skipped."), *InputSelector.GetName().ToString(), *OutputSelector.GetDisplayText().ToString());
-				return;
-			}
-
-			if (!PCG::Private::IsBroadcastable(ActorOverrideInputAccessor->GetUnderlyingType(), ActorOverrideOutputAccessor->GetUnderlyingType()))
-			{
-				UE_LOG(LogPCG, Warning, TEXT("ActorOverride cannot set input %s to output %s. Types are incompatibles. Will be skipped."), *InputSelector.GetName().ToString(), *OutputSelector.GetDisplayText().ToString());
-				ActorOverrideInputAccessor.Reset();
-				ActorOverrideOutputAccessor.Reset();
-				return;
-			}
-
-			auto CreateGetterSetter = [this](auto Dummy)
-			{
-				using Type = decltype(Dummy);
-
-				ActorOverrideFunction = [this](int32 Index, const IPCGAttributeAccessorKeys& InputKeys, IPCGAttributeAccessorKeys& OutputKey)
-				{
-					if (!IsValid())
-					{
-						return;
-					}
-
-					Type Value{};
-					if (ActorOverrideInputAccessor->Get<Type>(Value, Index, InputKeys, EPCGAttributeAccessorFlags::AllowBroadcast))
-					{
-						ActorOverrideOutputAccessor->Set<Type>(Value, OutputKey);
-					}
-				};
-			};
-
-			PCGMetadataAttribute::CallbackWithRightType(ActorOverrideOutputAccessor->GetUnderlyingType(), CreateGetterSetter);
-		}
-
-		bool IsValid() const
-		{
-			return ActorOverrideInputAccessor.IsValid() && ActorOverrideOutputAccessor.IsValid() && ActorOverrideFunction;
-		}
-
-		void Apply(int32 Index, const IPCGAttributeAccessorKeys& InputKeys, IPCGAttributeAccessorKeys& OutputKey)
-		{
-			ActorOverrideFunction(Index, InputKeys, OutputKey);
-		}
-
-	private:
-		TUniquePtr<const IPCGAttributeAccessor> ActorOverrideInputAccessor;
-		TUniquePtr<IPCGAttributeAccessor> ActorOverrideOutputAccessor;
-		ApplyOverrideFunction ActorOverrideFunction;
-	};
-
-	struct FActorOverrides
-	{
-		FActorOverrides(const TArray<FPCGActorPropertyOverride>& Overrides, AActor* TemplateActor, const UPCGPointData* PointData)
-			: InputKeys(PointData->GetPoints())
-			, OutputKey(TemplateActor)
-		{
-			ActorSingleOverrides.Reserve(Overrides.Num());
-
-			for (int32 i = 0; i < Overrides.Num(); ++i)
-			{
-				FPCGAttributePropertyInputSelector InputSelector = Overrides[i].InputSource.CopyAndFixLast(PointData);
-				const FString& OutputProperty = Overrides[i].PropertyTarget;
-
-				ActorSingleOverrides.Emplace(InputSelector, OutputProperty, TemplateActor, PointData);
-			}
-		}
-
-		void Apply(int32 Index)
-		{
-			for (FActorSingleOverride& ActorSingleOverride : ActorSingleOverrides)
-			{
-				if (ActorSingleOverride.IsValid())
-				{
-					ActorSingleOverride.Apply(Index, InputKeys, OutputKey);
-				}
-			}
-		}
-
-	private:
-		FPCGAttributeAccessorKeysPoints InputKeys;
-		FPCGAttributeAccessorKeysSingleObjectPtr<AActor> OutputKey;
-		TArray<FActorSingleOverride> ActorSingleOverrides;
-	};
-}
-
 class FPCGSpawnActorPartitionByAttribute : public FPCGPointDataPartitionBase<FPCGSpawnActorPartitionByAttribute, TSubclassOf<AActor>>
 {
 public:
@@ -322,6 +207,16 @@ void UPCGSpawnActorSettings::PostLoad()
 		bGenerationTrigger_DEPRECATED = EPCGSpawnActorGenerationTrigger::Default;
 	}
 
+	if (!ActorOverrides_DEPRECATED.IsEmpty())
+	{
+		for (const FPCGActorPropertyOverride& Override : ActorOverrides_DEPRECATED)
+		{
+			SpawnedActorPropertyOverrideDescriptions.Emplace(Override.InputSource, Override.PropertyTarget);
+		}
+
+		ActorOverrides_DEPRECATED.Empty();
+	}
+
 	// Since the template actor editing is set to false by default, this needs to be corrected on post-load for proper deprecation
 	if (TemplateActor)
 	{
@@ -820,7 +715,7 @@ void FPCGSpawnActorElement::SpawnActors(FPCGSubgraphContext* Context, AActor* Ta
 	AActor* TemplateActor = nullptr;
 	if (InTemplateActor)
 	{
-		if (Settings->ActorOverrides.IsEmpty())
+		if (Settings->SpawnedActorPropertyOverrideDescriptions.IsEmpty())
 		{
 			TemplateActor = InTemplateActor;
 		}
@@ -831,7 +726,7 @@ void FPCGSpawnActorElement::SpawnActors(FPCGSubgraphContext* Context, AActor* Ta
 	}
 	else
 	{
-		if (Settings->ActorOverrides.IsEmpty())
+		if (Settings->SpawnedActorPropertyOverrideDescriptions.IsEmpty())
 		{
 			TemplateActor = Cast<AActor>(InTemplateActorClass->GetDefaultObject());
 		}
@@ -843,7 +738,8 @@ void FPCGSpawnActorElement::SpawnActors(FPCGSubgraphContext* Context, AActor* Ta
 
 	check(TemplateActor);
 
-	PCGSpawnActorHelpers::FActorOverrides ActorOverrides(Settings->ActorOverrides, TemplateActor, PointData);
+	FPCGActorOverrides ActorOverrides(TemplateActor);
+	ActorOverrides.Initialize(Settings->SpawnedActorPropertyOverrideDescriptions, TemplateActor, PointData, Context);
 
 	FActorSpawnParameters SpawnParams;
 	SpawnParams.Template = TemplateActor;
@@ -931,11 +827,13 @@ void FPCGSpawnActorElement::SpawnActors(FPCGSubgraphContext* Context, AActor* Ta
 
 		const TArray<UFunction*> PostSpawnFunctions = PCGHelpers::FindUserFunctions(InTemplateActorClass, Settings->PostSpawnFunctionNames, Context);
 
+		bool bAllActorOverridesSucceeded = true;
+
 		for (int32 i = 0; i < Points.Num(); ++i)
 		{
 			const FPCGPoint& Point = Points[i];
 
-			ActorOverrides.Apply(i);
+			bAllActorOverridesSucceeded &= ActorOverrides.Apply(i);
 
 			AActor* GeneratedActor = TargetActor->GetWorld()->SpawnActor(InTemplateActorClass, &Point.Transform, SpawnParams);
 
@@ -971,6 +869,11 @@ void FPCGSpawnActorElement::SpawnActors(FPCGSubgraphContext* Context, AActor* Ta
 			}
 		}
 
+		if (!bAllActorOverridesSucceeded)
+		{
+			PCGE_LOG(Error, GraphAndLog, LOCTEXT("ActorOverridesFailed", "At least one actor property override failed."));
+		}
+
 		Context->SourceComponent->AddToManagedResources(ManagedActors);
 
 		PCGE_LOG(Verbose, LogOnly, FText::Format(LOCTEXT("GenerationInfo", "Generated {0} actors"), Points.Num()));
diff --git a/Engine/Plugins/Experimental/PCG/Source/PCG/Private/Metadata/PCGActorPropertyOverride.cpp b/Engine/Plugins/Experimental/PCG/Source/PCG/Private/Metadata/PCGActorPropertyOverride.cpp
new file mode 100644
--- /dev/null
+++ b/Engine/Plugins/Experimental/PCG/Source/PCG/Private/Metadata/PCGActorPropertyOverride.cpp
@@ -0,0 +1,165 @@
+// Copyright Epic Games, Inc. All Rights Reserved.
+
+#include "Metadata/PCGActorPropertyOverride.h"
+
+#include "PCGContext.h"
+#include "PCGData.h"
+#include "PCGModule.h"
+#include "PCGParamData.h"
+#include "Metadata/Accessors/PCGAttributeAccessorHelpers.h"
+
+#include "GameFramework/Actor.h"
+
+#define LOCTEXT_NAMESPACE "PCGActorPropertyOverride"
+
+namespace PCGActorPropertyOverrideHelpers
+{
+	FPCGPinProperties CreateActorPropertiesOverridePin(FName Label, const FText& Tooltip)
+	{
+		FPCGPinProperties ActorOverridePinProperties(Label, EPCGDataType::Param, /*bAllowMultipleConnections=*/false, /*bAllowMultipleData=*/false, Tooltip);
+		ActorOverridePinProperties.bAdvancedPin = true;
+		return ActorOverridePinProperties;
+	}
+
+	void ApplyOverridesFromParams(const TArray<FPCGActorPropertyOverrideDescription>& InActorPropertyOverrideDescriptions, AActor* TargetActor, FName OverridesPinLabel, FPCGContext* Context)
+	{
+		if (!Context)
+		{
+			return;
+		}
+
+		const TArray<FPCGTaggedData> OverrideInputs = Context->InputData.GetInputsByPin(OverridesPinLabel);
+
+		if (OverrideInputs.Num() == 0)
+		{
+			return;
+		}
+		else if (OverrideInputs.Num() > 1)
+		{
+			PCGLog::LogWarningOnGraph(FText::Format(LOCTEXT("MoreThanOneData", "More than one data was found on pin '{0}'. Only using the first one."), FText::FromName(OverridesPinLabel)), Context);
+		}
+
+		const UPCGParamData* ParamData = Cast<const UPCGParamData>(OverrideInputs[0].Data);
+
+		if (!ParamData)
+		{
+			PCGLog::LogErrorOnGraph(LOCTEXT("InvalidActorOverrideData", "Invalid input data type for Actor Property Overrides pin, must be of type Param."), Context);
+			return;
+		}
+
+		FPCGActorOverrides ActorOverrides(TargetActor);
+		ActorOverrides.Initialize(InActorPropertyOverrideDescriptions, TargetActor, ParamData, Context);
+		if (!ActorOverrides.Apply(/*InputKeyIndex=*/0)) // Use the First Entry of the param data for override (similar to what is done in Parameter Overrides in FPCGContext)
+		{
+			PCGLog::LogErrorOnGraph(FText::Format(LOCTEXT("ApplyOverrideFailed", "Failed to apply property overrides to actor '%s' from an attribute set."), FText::FromName(TargetActor->GetClass()->GetFName())), Context);
+		}
+	}
+}
+
+void FPCGActorSingleOverride::Initialize(const FPCGAttributePropertySelector& InputSelector, const FString& OutputProperty, AActor* TemplateActor, const UPCGData* SourceData, FPCGContext* Context)
+{
+	InputKeys = PCGAttributeAccessorHelpers::CreateConstKeys(SourceData, InputSelector);
+	ActorOverrideInputAccessor = PCGAttributeAccessorHelpers::CreateConstAccessor(SourceData, InputSelector);
+	FPCGAttributePropertySelector OutputSelector = FPCGAttributePropertySelector::CreateSelectorFromString(OutputProperty);
+	const TArray<FString>& ExtraNames = OutputSelector.GetExtraNames();
+	if (ExtraNames.IsEmpty())
+	{
+		ActorOverrideOutputAccessor = PCGAttributeAccessorHelpers::CreatePropertyAccessor(FName(OutputProperty), TemplateActor->GetClass());
+	}
+	else
+	{
+		TArray<FName> PropertyNames;
+		PropertyNames.Reserve(ExtraNames.Num() + 1);
+		PropertyNames.Add(OutputSelector.GetAttributeName());
+		for (const FString& Name : ExtraNames)
+		{
+			PropertyNames.Add(FName(Name));
+		}
+
+		ActorOverrideOutputAccessor = PCGAttributeAccessorHelpers::CreatePropertyChainAccessor(PropertyNames, TemplateActor->GetClass());
+	}
+
+	if (!ActorOverrideInputAccessor.IsValid() || !ActorOverrideOutputAccessor.IsValid())
+	{
+		PCGLog::LogWarningOnGraph(FText::Format(LOCTEXT("OverrideInvalid", "ActorOverride from input '{0}' or output '{1}' is invalid or unsupported. Will be skipped."), InputSelector.GetDisplayText(), OutputSelector.GetDisplayText()), Context);
+		return;
+	}
+
+	// TODO: Replace with something less restrictive (e.g. allow Double -> Float conversion).
+	if (!PCG::Private::IsBroadcastable(ActorOverrideInputAccessor->GetUnderlyingType(), ActorOverrideOutputAccessor->GetUnderlyingType()))
+	{
+		PCGLog::LogWarningOnGraph(
+			FText::Format(LOCTEXT("TypesIncompatible", "ActorOverride cannot set input '{0}' to output '{1}'. Cannot convert type '{2}' to type '{3}'. Will be skipped."),
+				InputSelector.GetDisplayText(),
+				OutputSelector.GetDisplayText(),
+				PCG::Private::GetTypeNameText(ActorOverrideInputAccessor->GetUnderlyingType()),
+				PCG::Private::GetTypeNameText(ActorOverrideOutputAccessor->GetUnderlyingType())),
+			Context);
+
+		ActorOverrideInputAccessor.Reset();
+		ActorOverrideOutputAccessor.Reset();
+		return;
+	}
+
+	auto CreateGetterSetter = [this](auto Dummy)
+	{
+		using Type = decltype(Dummy);
+
+		ActorOverrideFunction = &FPCGActorSingleOverride::ApplyImpl<Type>;
+	};
+
+	PCGMetadataAttribute::CallbackWithRightType(ActorOverrideOutputAccessor->GetUnderlyingType(), CreateGetterSetter);
+}
+
+bool FPCGActorSingleOverride::IsValid() const
+{
+	return InputKeys.IsValid() && ActorOverrideInputAccessor.IsValid() && ActorOverrideOutputAccessor.IsValid() && ActorOverrideFunction;
+}
+
+bool FPCGActorSingleOverride::Apply(int32 InputKeyIndex, IPCGAttributeAccessorKeys& OutputKey)
+{
+	return Invoke(ActorOverrideFunction, this, InputKeyIndex, OutputKey);
+}
+
+void FPCGActorOverrides::Initialize(const TArray<FPCGActorPropertyOverrideDescription>& OverrideDescriptions, AActor* TemplateActor, const UPCGData* SourceData, FPCGContext* Context)
+{
+	if (!TemplateActor)
+	{
+		PCGLog::LogErrorOnGraph(LOCTEXT("InitializeOverrideFailedNoActor", "Failed to initialize property overrides. No template actor was provided."), Context);
+		return;
+	}
+
+	ActorSingleOverrides.Reserve(OverrideDescriptions.Num());
+
+	for (int32 i = 0; i < OverrideDescriptions.Num(); ++i)
+	{
+		FPCGAttributePropertyInputSelector InputSelector = OverrideDescriptions[i].InputSource.CopyAndFixLast(SourceData);
+		const FString& OutputProperty = OverrideDescriptions[i].PropertyTarget;
+
+		FPCGActorSingleOverride Override;
+		Override.Initialize(InputSelector, OutputProperty, TemplateActor, SourceData, Context);
+
+		if (Override.IsValid())
+		{
+			ActorSingleOverrides.Add(std::move(Override));
+		}
+		else
+		{
+			PCGLog::LogErrorOnGraph(FText::Format(LOCTEXT("InitializeOverrideFailed", "Failed to initialize override '{0}' for property {1} on actor '{2}'."), InputSelector.GetDisplayText(), FText::FromString(OutputProperty), FText::FromName(TemplateActor->GetClass()->GetFName())), Context);
+		}
+	}
+}
+
+bool FPCGActorOverrides::Apply(int32 InputKeyIndex)
+{
+	bool bAllSucceeded = true;
+
+	for (FPCGActorSingleOverride& ActorSingleOverride : ActorSingleOverrides)
+	{
+		bAllSucceeded &= ActorSingleOverride.Apply(InputKeyIndex, OutputKey);
+	}
+
+	return bAllSucceeded;
+}
+
+#undef LOCTEXT_NAMESPACE
diff --git a/Engine/Plugins/Experimental/PCG/Source/PCG/Public/Elements/PCGCreateTargetActor.h b/Engine/Plugins/Experimental/PCG/Source/PCG/Public/Elements/PCGCreateTargetActor.h
--- a/Engine/Plugins/Experimental/PCG/Source/PCG/Public/Elements/PCGCreateTargetActor.h
+++ b/Engine/Plugins/Experimental/PCG/Source/PCG/Public/Elements/PCGCreateTargetActor.h
@@ -3,7 +3,7 @@
 #pragma once
 
 #include "PCGSettings.h"
-#include "Metadata/PCGAttributePropertySelector.h"
+#include "Metadata/PCGActorPropertyOverride.h"
 
 #include "PCGCreateTargetActor.generated.h"
 
@@ -72,6 +72,10 @@ public:
 	UPROPERTY(meta = (PCG_Overridable))
 	FTransform ActorPivot;
 
+	/** Override the default property values on the created target actor. Applied before post-process functions. */
+	UPROPERTY(EditAnywhere, BlueprintReadWrite, Category = Settings)
+	TArray<FPCGActorPropertyOverrideDescription> PropertyOverrideDescriptions;
+
 	/** Specify a list of functions to be called on the target actor after creation. Functions need to be parameter-less and with "CallInEditor" flag enabled. */
 	UPROPERTY(BlueprintReadWrite, EditAnywhere, Category = Settings)
 	TArray<FName> PostProcessFunctionNames;
diff --git a/Engine/Plugins/Experimental/PCG/Source/PCG/Public/Elements/PCGSpawnActor.h b/Engine/Plugins/Experimental/PCG/Source/PCG/Public/Elements/PCGSpawnActor.h
--- a/Engine/Plugins/Experimental/PCG/Source/PCG/Public/Elements/PCGSpawnActor.h
+++ b/Engine/Plugins/Experimental/PCG/Source/PCG/Public/Elements/PCGSpawnActor.h
@@ -3,7 +3,7 @@
 #pragma once
 
 #include "PCGSubgraph.h"
-#include "Metadata/PCGAttributePropertySelector.h"
+#include "Metadata/PCGActorPropertyOverride.h"
 
 #include "PCGSpawnActor.generated.h"
 
@@ -28,18 +28,6 @@ enum class EPCGSpawnActorGenerationTrigger : uint8
 	DoNotGenerate // Does not call generate
 };
 
-USTRUCT(BlueprintType)
-struct FPCGActorPropertyOverride
-{
-	GENERATED_BODY()
-
-	UPROPERTY(EditAnywhere, Category = Settings, meta = (PCG_Overridable))
-	FPCGAttributePropertyInputSelector InputSource;
-
-	UPROPERTY(EditAnywhere, Category = Settings)
-	FString PropertyTarget;
-};
-
 /*
 * PCG settings class that allows spawning actors with some options to perform the work more efficiently.
 * Note that depending on the options, any PCG components on the spawned actors can be also generated,
@@ -83,7 +71,7 @@ public:
 	TObjectPtr<AActor> TemplateActor;
 
 	UPROPERTY(EditAnywhere, BlueprintReadWrite, Category = Settings, meta = (EditCondition = "Option != EPCGSpawnActorOption::CollapseActors", EditConditionHides))
-	TArray<FPCGActorPropertyOverride> ActorOverrides;
+	TArray<FPCGActorPropertyOverrideDescription> SpawnedActorPropertyOverrideDescriptions;
 
 	UPROPERTY(meta = (PCG_Overridable))
 	TSoftObjectPtr<AActor> RootActor;
@@ -101,6 +89,9 @@ protected:
 #if WITH_EDITORONLY_DATA
 	UPROPERTY()
 	EPCGSpawnActorGenerationTrigger bGenerationTrigger_DEPRECATED = EPCGSpawnActorGenerationTrigger::Default;
+
+	UPROPERTY()
+	TArray<FPCGActorPropertyOverride> ActorOverrides_DEPRECATED;
 #endif
 
 public:
diff --git a/Engine/Plugins/Experimental/PCG/Source/PCG/Public/Metadata/PCGActorPropertyOverride.h b/Engine/Plugins/Experimental/PCG/Source/PCG/Public/Metadata/PCGActorPropertyOverride.h
new file mode 100644
--- /dev/null
+++ b/Engine/Plugins/Experimental/PCG/Source/PCG/Public/Metadata/PCGActorPropertyOverride.h
@@ -0,0 +1,126 @@
+// Copyright Epic Games, Inc. All Rights Reserved.
+
+#pragma once
+
+#include "PCGPin.h"
+#include "Metadata/PCGAttributePropertySelector.h"
+#include "Metadata/Accessors/IPCGAttributeAccessor.h"
+#include "Metadata/Accessors/PCGAttributeAccessorKeys.h"
+
+#include "PCGActorPropertyOverride.generated.h"
+
+class IPCGAttributeAccessor;
+class UPCGData;
+struct FPCGContext;
+
+class AActor;
+
+/**
+* Represents the override source (to be read) and the actor property (to be written).
+*/
+USTRUCT(BlueprintType)
+struct FPCGActorPropertyOverrideDescription
+{
+	GENERATED_BODY()
+
+	/** Provide an attribute or property to read the override value from. */
+	UPROPERTY(EditAnywhere, Category = Settings, meta = (PCG_Overridable))
+	FPCGAttributePropertyInputSelector InputSource;
+
+	/**
+	* Provide an actor property name to be overridden. If you have a property "A" on your actor, use "A" as the property target.
+	*
+	* For example, if you want to override the "Is Editor Only" flag, find it in the details panel, right-click, select 'Copy Internal Name', and paste that as the property target.
+	*
+	* If you have a component property, such as the static mesh of a static mesh component, use "StaticMeshComponent.StaticMesh".
+	*/
+	UPROPERTY(EditAnywhere, Category = Settings)
+	FString PropertyTarget;
+};
+
+namespace PCGActorPropertyOverrideHelpers
+{
+	/** Create an advanced ParamData pin for capturing property overrides. */
+	FPCGPinProperties CreateActorPropertiesOverridePin(FName Label, const FText& Tooltip);
+	
+	/** Apply property overrides to the TargetActor directly from the ActorPropertiesOverride pin. Use CreateActorPropertiesOverridePin(). */
+	void ApplyOverridesFromParams(const TArray<FPCGActorPropertyOverrideDescription>& InActorPropertyOverrideDescriptions, AActor* TargetActor, FName OverridesPinLabel, FPCGContext* Context);
+}
+
+/**
+* Represents a single property override on the provided actor. Applies an override function to read the InputAccessor
+* and write its value to the OutputAccessor.
+* 
+* The InputAccessor's InputKeys are created from the given SourceData and InputSelector.
+*/
+struct FPCGActorSingleOverride
+{
+	/** Initialize the single actor override. Call before using Apply(InputKeyIndex, OutputKey). */
+	void Initialize(const FPCGAttributePropertySelector& InputSelector, const FString& OutputProperty, AActor* TemplateActor, const UPCGData* SourceData, FPCGContext* Context);
+	
+	/** Returns true if initialization succeeded in creating the accessors and accessor keys. */
+	bool IsValid() const;
+
+	/** Applies a single property override to the actor by reading from the InputAccessor at the given KeyIndex, and writing to the OutputKey which represents the actor property. */
+	bool Apply(int32 InputKeyIndex, IPCGAttributeAccessorKeys& OutputKey);
+
+private:
+	TUniquePtr<const IPCGAttributeAccessorKeys> InputKeys;
+	TUniquePtr<const IPCGAttributeAccessor> ActorOverrideInputAccessor;
+	TUniquePtr<IPCGAttributeAccessor> ActorOverrideOutputAccessor;
+
+	// InputKeyIndex, OutputKeys
+	using ApplyOverrideFunction = bool(FPCGActorSingleOverride::*)(int32, IPCGAttributeAccessorKeys&);
+	ApplyOverrideFunction ActorOverrideFunction;
+
+	template <typename Type>
+	bool ApplyImpl(int32 InputKeyIndex, IPCGAttributeAccessorKeys& OutputKey)
+	{
+		if (!IsValid())
+		{
+			return false;
+		}
+
+		Type Value{};
+		if (ActorOverrideInputAccessor->Get<Type>(Value, InputKeyIndex, *InputKeys.Get(), EPCGAttributeAccessorFlags::AllowBroadcast))
+		{
+			if (ActorOverrideOutputAccessor->Set<Type>(Value, OutputKey))
+			{
+				return true;
+			}
+		}
+
+		return false;
+	}
+};
+
+/**
+* Represents a set of property overrides for the provided actor. Provide a SourceData to read from, and a collection of ActorPropertyOverrides matching the TemplateActor's class properties.
+*/
+struct FPCGActorOverrides
+{
+	FPCGActorOverrides(AActor* TemplateActor) : OutputKey(TemplateActor)
+	{}
+
+	/** Initialize the actor overrides. Call before using Apply(InputKeyIndex). */
+	void Initialize(const TArray<FPCGActorPropertyOverrideDescription>& OverrideDescriptions, AActor* TemplateActor, const UPCGData* SourceData, FPCGContext* Context);
+
+	/** Applies each property override to the actor by reading from the InputAccessor at the given KeyIndex, and writing to the OutputKey which represents the actor property. */
+	bool Apply(int32 InputKeyIndex);
+
+private:
+	FPCGAttributeAccessorKeysSingleObjectPtr<AActor> OutputKey;
+	TArray<FPCGActorSingleOverride> ActorSingleOverrides;
+};
+
+USTRUCT(BlueprintType, meta=(Deprecated = "5.4", DeprecationMessage="Use FPCGActorPropertyOverrideDescription instead."))
+struct FPCGActorPropertyOverride
+{
+	GENERATED_BODY()
+
+	UPROPERTY(EditAnywhere, Category = Settings, meta = (PCG_Overridable))
+	FPCGAttributePropertyInputSelector InputSource;
+
+	UPROPERTY(EditAnywhere, Category = Settings)
+	FString PropertyTarget;
+};
