Fixed get landscape data Get Layer Weights option deprecation

#rnx

diff --git a/Engine/Plugins/Experimental/PCG/Source/PCG/Private/Elements/PCGTypedGetter.cpp b/Engine/Plugins/Experimental/PCG/Source/PCG/Private/Elements/PCGTypedGetter.cpp
--- a/Engine/Plugins/Experimental/PCG/Source/PCG/Private/Elements/PCGTypedGetter.cpp
+++ b/Engine/Plugins/Experimental/PCG/Source/PCG/Private/Elements/PCGTypedGetter.cpp
@@ -45,10 +45,10 @@ void UPCGGetLandscapeSettings::PostLoad()
 		bGetHeightOnly_DEPRECATED = false;
 	}
 
-	if (bGetLayerWeights_DEPRECATED)
+	if (!bGetLayerWeights_DEPRECATED)
 	{
 		SamplingProperties.bGetLayerWeights = bGetLayerWeights_DEPRECATED;
-		bGetLayerWeights_DEPRECATED = false;
+		bGetLayerWeights_DEPRECATED = true;
 	}
 #endif
 }
