[Backout] - CL29376412
Original CL Desc
-----------------------------------------------------------------
[PCG] Add actor property overrides to spawner nodes

diff --git a/Engine/Plugins/Experimental/PCG/Source/PCG/Private/Elements/PCGCreateTargetActor.cpp b/Engine/Plugins/Experimental/PCG/Source/PCG/Private/Elements/PCGCreateTargetActor.cpp
--- a/Engine/Plugins/Experimental/PCG/Source/PCG/Private/Elements/PCGCreateTargetActor.cpp
+++ b/Engine/Plugins/Experimental/PCG/Source/PCG/Private/Elements/PCGCreateTargetActor.cpp
@@ -30,12 +30,6 @@ FText UPCGCreateTargetActor::GetDefaultNodeTitle() const
 
 #endif // WITH_EDITOR
 
-namespace PCGCreateTargetActorConstants
-{
-	const FName ActorPropertyOverridesLabel = TEXT("Property Overrides");
-	const FText ActorPropertyOverridesTooltip = LOCTEXT("ActorOverrideToolTip", "Provide property overrides for the created target actor. The attribute name must match the InputSource name in the actor property override description.");
-}
-
 FPCGElementPtr UPCGCreateTargetActor::CreateElement() const
 {
 	return MakeShared<FPCGCreateTargetActorElement>();
@@ -43,9 +37,7 @@ FPCGElementPtr UPCGCreateTargetActor::CreateElement() const
 
 TArray<FPCGPinProperties> UPCGCreateTargetActor::InputPinProperties() const
 {
-	TArray<FPCGPinProperties> PinProperties;
-	PinProperties.Add(PCGActorPropertyOverrideHelpers::CreateActorPropertiesOverridePin(PCGCreateTargetActorConstants::ActorPropertyOverridesLabel, PCGCreateTargetActorConstants::ActorPropertyOverridesTooltip));
-	return PinProperties;
+	return TArray<FPCGPinProperties>();
 }
 
 TArray<FPCGPinProperties> UPCGCreateTargetActor::OutputPinProperties() const
@@ -291,9 +283,6 @@ bool FPCGCreateTargetActorElement::ExecuteInternal(FPCGContext* Context) const
 
 	GeneratedActor->Tags.Add(PCGHelpers::DefaultPCGActorTag);
 
-	// Apply property overrides to the GeneratedActor
-	PCGActorPropertyOverrideHelpers::ApplyOverridesFromParams(Settings->PropertyOverrideDescriptions, GeneratedActor, PCGCreateTargetActorConstants::ActorPropertyOverridesLabel, Context);
-
 	for (UFunction* Function : PCGHelpers::FindUserFunctions(GeneratedActor->GetClass(), Settings->PostProcessFunctionNames, Context))
 	{
 		GeneratedActor->ProcessEvent(Function, nullptr);
diff --git a/Engine/Plugins/Experimental/PCG/Source/PCG/Private/Elements/PCGSpawnActor.cpp b/Engine/Plugins/Experimental/PCG/Source/PCG/Private/Elements/PCGSpawnActor.cpp
--- a/Engine/Plugins/Experimental/PCG/Source/PCG/Private/Elements/PCGSpawnActor.cpp
+++ b/Engine/Plugins/Experimental/PCG/Source/PCG/Private/Elements/PCGSpawnActor.cpp
@@ -33,6 +33,121 @@ static TAutoConsoleVariable<bool> CVarAllowActorReuse(
 	true,
 	TEXT("Controls whether PCG spawned actors can be reused and skipped when re-executing"));
 
+namespace PCGSpawnActorHelpers
+{
+	struct FActorSingleOverride
+	{
+		using ApplyOverrideFunction = TFunction<void(int32, const IPCGAttributeAccessorKeys&, IPCGAttributeAccessorKeys&)>;
+
+		FActorSingleOverride(const FPCGAttributePropertySelector& InputSelector, const FString& OutputProperty, AActor* TemplateActor, const UPCGPointData* PointData)
+		{
+			ActorOverrideInputAccessor = PCGAttributeAccessorHelpers::CreateConstAccessor(PointData, InputSelector);
+			FPCGAttributePropertySelector OutputSelector = FPCGAttributePropertySelector::CreateSelectorFromString(OutputProperty);
+			const TArray<FString>& ExtraNames = OutputSelector.GetExtraNames();
+			if (ExtraNames.IsEmpty())
+			{
+				ActorOverrideOutputAccessor = PCGAttributeAccessorHelpers::CreatePropertyAccessor(FName(OutputProperty), TemplateActor->GetClass());
+			}
+			else
+			{
+				TArray<FName> PropertyNames;
+				PropertyNames.Reserve(ExtraNames.Num() + 1);
+				PropertyNames.Add(OutputSelector.GetAttributeName());
+				for (const FString& Name : ExtraNames)
+				{
+					PropertyNames.Add(FName(Name));
+				}
+
+				ActorOverrideOutputAccessor = PCGAttributeAccessorHelpers::CreatePropertyChainAccessor(PropertyNames, TemplateActor->GetClass());
+			}
+
+			if (!ActorOverrideInputAccessor.IsValid() || !ActorOverrideOutputAccessor.IsValid())
+			{
+				UE_LOG(LogPCG, Warning, TEXT("ActorOverride from input %s or output %s is invalid or unsupported. Will be skipped."), *InputSelector.GetName().ToString(), *OutputSelector.GetDisplayText().ToString());
+				return;
+			}
+
+			if (!PCG::Private::IsBroadcastable(ActorOverrideInputAccessor->GetUnderlyingType(), ActorOverrideOutputAccessor->GetUnderlyingType()))
+			{
+				UE_LOG(LogPCG, Warning, TEXT("ActorOverride cannot set input %s to output %s. Types are incompatibles. Will be skipped."), *InputSelector.GetName().ToString(), *OutputSelector.GetDisplayText().ToString());
+				ActorOverrideInputAccessor.Reset();
+				ActorOverrideOutputAccessor.Reset();
+				return;
+			}
+
+			auto CreateGetterSetter = [this](auto Dummy)
+			{
+				using Type = decltype(Dummy);
+
+				ActorOverrideFunction = [this](int32 Index, const IPCGAttributeAccessorKeys& InputKeys, IPCGAttributeAccessorKeys& OutputKey)
+				{
+					if (!IsValid())
+					{
+						return;
+					}
+
+					Type Value{};
+					if (ActorOverrideInputAccessor->Get<Type>(Value, Index, InputKeys, EPCGAttributeAccessorFlags::AllowBroadcast))
+					{
+						ActorOverrideOutputAccessor->Set<Type>(Value, OutputKey);
+					}
+				};
+			};
+
+			PCGMetadataAttribute::CallbackWithRightType(ActorOverrideOutputAccessor->GetUnderlyingType(), CreateGetterSetter);
+		}
+
+		bool IsValid() const
+		{
+			return ActorOverrideInputAccessor.IsValid() && ActorOverrideOutputAccessor.IsValid() && ActorOverrideFunction;
+		}
+
+		void Apply(int32 Index, const IPCGAttributeAccessorKeys& InputKeys, IPCGAttributeAccessorKeys& OutputKey)
+		{
+			ActorOverrideFunction(Index, InputKeys, OutputKey);
+		}
+
+	private:
+		TUniquePtr<const IPCGAttributeAccessor> ActorOverrideInputAccessor;
+		TUniquePtr<IPCGAttributeAccessor> ActorOverrideOutputAccessor;
+		ApplyOverrideFunction ActorOverrideFunction;
+	};
+
+	struct FActorOverrides
+	{
+		FActorOverrides(const TArray<FPCGActorPropertyOverride>& Overrides, AActor* TemplateActor, const UPCGPointData* PointData)
+			: InputKeys(PointData->GetPoints())
+			, OutputKey(TemplateActor)
+		{
+			ActorSingleOverrides.Reserve(Overrides.Num());
+
+			for (int32 i = 0; i < Overrides.Num(); ++i)
+			{
+				FPCGAttributePropertyInputSelector InputSelector = Overrides[i].InputSource.CopyAndFixLast(PointData);
+				const FString& OutputProperty = Overrides[i].PropertyTarget;
+
+				ActorSingleOverrides.Emplace(InputSelector, OutputProperty, TemplateActor, PointData);
+			}
+		}
+
+		void Apply(int32 Index)
+		{
+			for (FActorSingleOverride& ActorSingleOverride : ActorSingleOverrides)
+			{
+				if (ActorSingleOverride.IsValid())
+				{
+					ActorSingleOverride.Apply(Index, InputKeys, OutputKey);
+				}
+			}
+		}
+
+	private:
+		FPCGAttributeAccessorKeysPoints InputKeys;
+		FPCGAttributeAccessorKeysSingleObjectPtr<AActor> OutputKey;
+		TArray<FActorSingleOverride> ActorSingleOverrides;
+	};
+}
+
 class FPCGSpawnActorPartitionByAttribute : public FPCGPointDataPartitionBase<FPCGSpawnActorPartitionByAttribute, TSubclassOf<AActor>>
 {
 public:
@@ -207,16 +322,6 @@ void UPCGSpawnActorSettings::PostLoad()
 		bGenerationTrigger_DEPRECATED = EPCGSpawnActorGenerationTrigger::Default;
 	}
 
-	if (!ActorOverrides_DEPRECATED.IsEmpty())
-	{
-		for (const FPCGActorPropertyOverride& Override : ActorOverrides_DEPRECATED)
-		{
-			SpawnedActorPropertyOverrideDescriptions.Emplace(Override.InputSource, Override.PropertyTarget);
-		}
-
-		ActorOverrides_DEPRECATED.Empty();
-	}
-
 	// Since the template actor editing is set to false by default, this needs to be corrected on post-load for proper deprecation
 	if (TemplateActor)
 	{
@@ -715,7 +820,7 @@ void FPCGSpawnActorElement::SpawnActors(FPCGSubgraphContext* Context, AActor* Ta
 	AActor* TemplateActor = nullptr;
 	if (InTemplateActor)
 	{
-		if (Settings->SpawnedActorPropertyOverrideDescriptions.IsEmpty())
+		if (Settings->ActorOverrides.IsEmpty())
 		{
 			TemplateActor = InTemplateActor;
 		}
@@ -726,7 +831,7 @@ void FPCGSpawnActorElement::SpawnActors(FPCGSubgraphContext* Context, AActor* Ta
 	}
 	else
 	{
-		if (Settings->SpawnedActorPropertyOverrideDescriptions.IsEmpty())
+		if (Settings->ActorOverrides.IsEmpty())
 		{
 			TemplateActor = Cast<AActor>(InTemplateActorClass->GetDefaultObject());
 		}
@@ -738,8 +843,7 @@ void FPCGSpawnActorElement::SpawnActors(FPCGSubgraphContext* Context, AActor* Ta
 
 	check(TemplateActor);
 
-	FPCGActorOverrides ActorOverrides(TemplateActor);
-	ActorOverrides.Initialize(Settings->SpawnedActorPropertyOverrideDescriptions, TemplateActor, PointData, Context);
+	PCGSpawnActorHelpers::FActorOverrides ActorOverrides(Settings->ActorOverrides, TemplateActor, PointData);
 
 	FActorSpawnParameters SpawnParams;
 	SpawnParams.Template = TemplateActor;
@@ -827,13 +931,11 @@ void FPCGSpawnActorElement::SpawnActors(FPCGSubgraphContext* Context, AActor* Ta
 
 		const TArray<UFunction*> PostSpawnFunctions = PCGHelpers::FindUserFunctions(InTemplateActorClass, Settings->PostSpawnFunctionNames, Context);
 
-		bool bAllActorOverridesSucceeded = true;
-
 		for (int32 i = 0; i < Points.Num(); ++i)
 		{
 			const FPCGPoint& Point = Points[i];
 
-			bAllActorOverridesSucceeded &= ActorOverrides.Apply(i);
+			ActorOverrides.Apply(i);
 
 			AActor* GeneratedActor = TargetActor->GetWorld()->SpawnActor(InTemplateActorClass, &Point.Transform, SpawnParams);
 
@@ -869,11 +971,6 @@ void FPCGSpawnActorElement::SpawnActors(FPCGSubgraphContext* Context, AActor* Ta
 			}
 		}
 
-		if (!bAllActorOverridesSucceeded)
-		{
-			PCGE_LOG(Error, GraphAndLog, LOCTEXT("ActorOverridesFailed", "At least one actor property override failed."));
-		}
-
 		Context->SourceComponent->AddToManagedResources(ManagedActors);
 
 		PCGE_LOG(Verbose, LogOnly, FText::Format(LOCTEXT("GenerationInfo", "Generated {0} actors"), Points.Num()));
diff --git a/Engine/Plugins/Experimental/PCG/Source/PCG/Private/Metadata/PCGActorPropertyOverride.cpp b/Engine/Plugins/Experimental/PCG/Source/PCG/Private/Metadata/PCGActorPropertyOverride.cpp
deleted file mode 100644
diff --git a/Engine/Plugins/Experimental/PCG/Source/PCG/Public/Elements/PCGCreateTargetActor.h b/Engine/Plugins/Experimental/PCG/Source/PCG/Public/Elements/PCGCreateTargetActor.h
--- a/Engine/Plugins/Experimental/PCG/Source/PCG/Public/Elements/PCGCreateTargetActor.h
+++ b/Engine/Plugins/Experimental/PCG/Source/PCG/Public/Elements/PCGCreateTargetActor.h
@@ -3,7 +3,7 @@
 #pragma once
 
 #include "PCGSettings.h"
-#include "Metadata/PCGActorPropertyOverride.h"
+#include "Metadata/PCGAttributePropertySelector.h"
 
 #include "PCGCreateTargetActor.generated.h"
 
@@ -72,10 +72,6 @@ public:
 	UPROPERTY(meta = (PCG_Overridable))
 	FTransform ActorPivot;
 
-	/** Override the default property values on the created target actor. Applied before post-process functions. */
-	UPROPERTY(EditAnywhere, BlueprintReadWrite, Category = Settings)
-	TArray<FPCGActorPropertyOverrideDescription> PropertyOverrideDescriptions;
-
 	/** Specify a list of functions to be called on the target actor after creation. Functions need to be parameter-less and with "CallInEditor" flag enabled. */
 	UPROPERTY(BlueprintReadWrite, EditAnywhere, Category = Settings)
 	TArray<FName> PostProcessFunctionNames;
diff --git a/Engine/Plugins/Experimental/PCG/Source/PCG/Public/Elements/PCGSpawnActor.h b/Engine/Plugins/Experimental/PCG/Source/PCG/Public/Elements/PCGSpawnActor.h
--- a/Engine/Plugins/Experimental/PCG/Source/PCG/Public/Elements/PCGSpawnActor.h
+++ b/Engine/Plugins/Experimental/PCG/Source/PCG/Public/Elements/PCGSpawnActor.h
@@ -3,7 +3,7 @@
 #pragma once
 
 #include "PCGSubgraph.h"
-#include "Metadata/PCGActorPropertyOverride.h"
+#include "Metadata/PCGAttributePropertySelector.h"
 
 #include "PCGSpawnActor.generated.h"
 
@@ -28,6 +28,18 @@ enum class EPCGSpawnActorGenerationTrigger : uint8
 	DoNotGenerate // Does not call generate
 };
 
+USTRUCT(BlueprintType)
+struct FPCGActorPropertyOverride
+{
+	GENERATED_BODY()
+
+	UPROPERTY(EditAnywhere, Category = Settings, meta = (PCG_Overridable))
+	FPCGAttributePropertyInputSelector InputSource;
+
+	UPROPERTY(EditAnywhere, Category = Settings)
+	FString PropertyTarget;
+};
+
 /*
 * PCG settings class that allows spawning actors with some options to perform the work more efficiently.
 * Note that depending on the options, any PCG components on the spawned actors can be also generated,
@@ -71,7 +83,7 @@ public:
 	TObjectPtr<AActor> TemplateActor;
 
 	UPROPERTY(EditAnywhere, BlueprintReadWrite, Category = Settings, meta = (EditCondition = "Option != EPCGSpawnActorOption::CollapseActors", EditConditionHides))
-	TArray<FPCGActorPropertyOverrideDescription> SpawnedActorPropertyOverrideDescriptions;
+	TArray<FPCGActorPropertyOverride> ActorOverrides;
 
 	UPROPERTY(meta = (PCG_Overridable))
 	TSoftObjectPtr<AActor> RootActor;
@@ -89,9 +101,6 @@ protected:
 #if WITH_EDITORONLY_DATA
 	UPROPERTY()
 	EPCGSpawnActorGenerationTrigger bGenerationTrigger_DEPRECATED = EPCGSpawnActorGenerationTrigger::Default;
-
-	UPROPERTY()
-	TArray<FPCGActorPropertyOverride> ActorOverrides_DEPRECATED;
 #endif
 
 public:
diff --git a/Engine/Plugins/Experimental/PCG/Source/PCG/Public/Metadata/PCGActorPropertyOverride.h b/Engine/Plugins/Experimental/PCG/Source/PCG/Public/Metadata/PCGActorPropertyOverride.h
deleted file mode 100644
