[PCG] Re-assess the issue with CopyPoints

* Last changes broke some usages

#rnx

diff --git a/Engine/Plugins/Experimental/PCG/Source/PCG/Private/Elements/PCGCopyPoints.cpp b/Engine/Plugins/Experimental/PCG/Source/PCG/Private/Elements/PCGCopyPoints.cpp
--- a/Engine/Plugins/Experimental/PCG/Source/PCG/Private/Elements/PCGCopyPoints.cpp
+++ b/Engine/Plugins/Experimental/PCG/Source/PCG/Private/Elements/PCGCopyPoints.cpp
@@ -206,35 +206,24 @@ bool FPCGCopyPointsElement::ExecuteInternal(FPCGContext* Context) const
 
 				OutPoint = SourcePoint;
 
-				// Compute the RelativeTransform only once.
-				const bool bUseRelativeTransform = (RotationInheritance == EPCGCopyPointsInheritanceMode::Relative) || (ScaleInheritance == EPCGCopyPointsInheritanceMode::Relative);
-				const FTransform RelativeTransform = bUseRelativeTransform ? SourcePoint.Transform.GetRelativeTransform(TargetPoint.Transform) : FTransform{};
+				// Set the position, rotation and scale as relative by default.
+				OutPoint.Transform = SourcePoint.Transform * TargetPoint.Transform;
 
 				// Set Rotation, Scale, and Color based on inheritance mode
-				if (RotationInheritance == EPCGCopyPointsInheritanceMode::Relative)
-				{
-					// GetRelativeTransform will take into account if the scale is negative
-					OutPoint.Transform.SetRotation(RelativeTransform.GetRotation());
-				}
-				else if (RotationInheritance == EPCGCopyPointsInheritanceMode::Source)
+				if (RotationInheritance == EPCGCopyPointsInheritanceMode::Source)
 				{
 					OutPoint.Transform.SetRotation(SourcePoint.Transform.GetRotation());
 				}
-				else // if (RotationInheritance == EPCGCopyPointsInheritanceMode::Target)
+				else if (RotationInheritance == EPCGCopyPointsInheritanceMode::Target)
 				{
 					OutPoint.Transform.SetRotation(TargetPoint.Transform.GetRotation());
 				}
 
-				if (ScaleInheritance == EPCGCopyPointsInheritanceMode::Relative)
-				{
-					// GetRelativeTransform will take into account if the scale is negative
-					OutPoint.Transform.SetScale3D(RelativeTransform.GetScale3D());
-				}
-				else if (ScaleInheritance == EPCGCopyPointsInheritanceMode::Source)
+				if (ScaleInheritance == EPCGCopyPointsInheritanceMode::Source)
 				{ 
 					OutPoint.Transform.SetScale3D(SourcePoint.Transform.GetScale3D());
 				}
-				else // if (ScaleInheritance == EPCGCopyPointsInheritanceMode::Target)
+				else if (ScaleInheritance == EPCGCopyPointsInheritanceMode::Target)
 				{
 					OutPoint.Transform.SetScale3D(TargetPoint.Transform.GetScale3D());
 				}
@@ -252,9 +241,6 @@ bool FPCGCopyPointsElement::ExecuteInternal(FPCGContext* Context) const
 					OutPoint.Color = TargetPoint.Color;
 				}
 
-				const FVector Location = TargetPoint.Transform.TransformPosition(SourcePoint.Transform.GetLocation());
-				OutPoint.Transform.SetLocation(Location);
-
 				// Set seed based on inheritance mode
 				if (SeedInheritance == EPCGCopyPointsInheritanceMode::Relative)
 				{
