[PCG] Add first pass of world builder

diff --git a/Engine/Plugins/Experimental/PCG/Source/PCG/Private/PCGComponent.cpp b/Engine/Plugins/Experimental/PCG/Source/PCG/Private/PCGComponent.cpp
--- a/Engine/Plugins/Experimental/PCG/Source/PCG/Private/PCGComponent.cpp
+++ b/Engine/Plugins/Experimental/PCG/Source/PCG/Private/PCGComponent.cpp
@@ -1896,8 +1896,8 @@ void UPCGComponent::Refresh(bool bStructural)
 		return;
 	}
 
-	// If the refresh is globally disabled, just exit
-	if (PCGComponent::CVarGlobalDisableRefresh.GetValueOnAnyThread())
+	// If refresh is disabled, just exit
+	if (PCGComponent::CVarGlobalDisableRefresh.GetValueOnAnyThread() || IsRunningCommandlet())
 	{
 		return;
 	}
diff --git a/Engine/Plugins/Experimental/PCG/Source/PCGEditor/Private/WorldPartitionBuilder/PCGWorldPartitionBuilder.cpp b/Engine/Plugins/Experimental/PCG/Source/PCGEditor/Private/WorldPartitionBuilder/PCGWorldPartitionBuilder.cpp
new file mode 100644
--- /dev/null
+++ b/Engine/Plugins/Experimental/PCG/Source/PCGEditor/Private/WorldPartitionBuilder/PCGWorldPartitionBuilder.cpp
@@ -0,0 +1,407 @@
+// Copyright Epic Games, Inc. All Rights Reserved.
+
+#include "WorldPartitionBuilder/PCGWorldPartitionBuilder.h"
+
+#include "PCGComponent.h"
+#include "PCGGraph.h"
+
+#include "AssetCompilingManager.h"
+#include "Editor.h"
+#include "FileHelpers.h"
+#include "Modules/ModuleManager.h"
+#include "UObject/Linker.h"
+#include "WorldPartition/IWorldPartitionEditorModule.h"
+#include "WorldPartition/WorldPartitionHelpers.h"
+
+DEFINE_LOG_CATEGORY_STATIC(LogPCGWorldPartitionBuilder, All, All);
+
+namespace PCGWorldPartitionBuilder
+{
+	/** Grab all original PCG components in world, applying given filter. */
+	void CollectComponentsToGenerate(UWorld* InWorld, TFunctionRef<bool(const UPCGComponent*)> ComponentFilter, TArray<TWeakObjectPtr<UPCGComponent>>& OutComponents);
+	
+	/** Generate the given components. Optionally generate only one component at a time with a wait on async processes after each. Optionally apply given filter to select components. */
+	bool GenerateComponents(TArray<TWeakObjectPtr<UPCGComponent>>& Components, UWorld* InWorld, bool bOneComponentAtATime);
+	bool GenerateComponents(TArray<TWeakObjectPtr<UPCGComponent>>& Components, UWorld* InWorld, bool bOneComponentAtATime, TFunctionRef<bool(const UPCGComponent*)> ComponentFilter);
+
+	/** Generate a component. Applies correct editing mode if necessary. */
+	void GenerateComponent(UPCGComponent* InComponent, UWorld* InWorld);
+
+	/** Waits for background activity to be quiet. */
+	void WaitForAllAsyncEditorProcesses(UWorld* InWorld);
+
+	/** Runs builder on current editor world. */
+	void Build(const TArray<FString>& Args);
+
+	static FAutoConsoleCommand CommandBuildComponents(
+		TEXT("pcg.BuildComponents"),
+		TEXT("Runs PCG world builder on current world. Args: [-GenerateComponentEditingModeNormal] [-IncludeGraphNames=PCG_GraphA;PCG_GraphB] [-IncludeActorIDs=MyActor1_UID1234678;MyActor2_UID1234678]"),
+		FConsoleCommandWithArgsDelegate::CreateLambda([](const TArray<FString>& Args) { PCGWorldPartitionBuilder::Build(Args); }));
+};
+
+UPCGWorldPartitionBuilder::UPCGWorldPartitionBuilder(const FObjectInitializer& ObjectInitializer)
+	: Super(ObjectInitializer)
+{
+	if (!HasAnyFlags(RF_ClassDefaultObject))
+	{
+		bGenerateEditingModeNormalComponents = HasParam("GenerateComponentEditingModeNormal");
+
+		bOneComponentAtATime = HasParam("OneComponentAtATime");
+
+		FString IncludeGraphNamesValue;
+		if (GetParamValue("IncludeGraphNames=", IncludeGraphNamesValue) && !IncludeGraphNamesValue.IsEmpty())
+		{
+			TArray<FString> IncludeGraphNameStrings;
+			IncludeGraphNamesValue.ParseIntoArray(IncludeGraphNameStrings, TEXT(";"), true);
+
+			IncludeGraphNames.Append(IncludeGraphNameStrings);
+		}
+
+		FString IncludeActorIDsValue;
+		if (GetParamValue("IncludeActorIDs=", IncludeActorIDsValue) && !IncludeActorIDsValue.IsEmpty())
+		{
+			TArray<FString> IncludeActorIDStrings;
+			IncludeActorIDsValue.ParseIntoArray(IncludeActorIDStrings, TEXT(";"), true);
+
+			IncludeActorIDs.Append(IncludeActorIDStrings);
+		}
+	}
+}
+
+bool UPCGWorldPartitionBuilder::RunInternal(UWorld* World, const FCellInfo& InCellInfo, FPackageSourceControlHelper& PackageHelper)
+{
+	if (!ensure(GEngine) || !ensure(World))
+	{
+		return false;
+	}
+
+	// Try to eliminate as much background noise as possible before getting started.
+	PCGWorldPartitionBuilder::WaitForAllAsyncEditorProcesses(World);
+
+	// 1. Collect components of interest once upfront.
+
+	TArray<TWeakObjectPtr<UPCGComponent>> ComponentsToGenerate;
+	{
+		auto ComponentFilter = [this](const UPCGComponent* InComponent)
+		{
+			if (!InComponent || !InComponent->GetOwner())
+			{
+				return false;
+			}
+
+			// Check actor in inclusion list if provided.
+			if (!IncludeActorIDs.IsEmpty() && !IncludeActorIDs.Contains(InComponent->GetOwner()->GetName()))
+			{
+				return false;
+			}
+
+			// Accept based on editing mode.
+			return (InComponent->GetSerializedEditingMode() == EPCGEditorDirtyMode::LoadAsPreview)
+				|| (bGenerateEditingModeNormalComponents && InComponent->GetSerializedEditingMode() == EPCGEditorDirtyMode::Normal);
+		};
+
+		PCGWorldPartitionBuilder::CollectComponentsToGenerate(World, ComponentFilter, ComponentsToGenerate);
+	}
+
+	// 2. Clear the dirty flag on any packages that dirtied themselves while loading. We're not interested in saving these.
+
+	{
+		TArray<UPackage*> DirtyPackages;
+		FEditorFileUtils::GetDirtyPackages(DirtyPackages);
+		DirtyPackages.RemoveSwap(nullptr);
+		for (UPackage* DirtyPackage : DirtyPackages)
+		{
+			// Make clean as it's not a change we have made.
+			if (!PendingDirtyPackages.Contains(DirtyPackage->GetPersistentGuid()))
+			{
+				DirtyPackage->SetDirtyFlag(false);
+			}
+		}
+		DirtyPackages.Reset();
+	}
+
+	const int32 NumPendingDirtyPackagesBefore = PendingDirtyPackages.Num();
+
+	// 3. Generate components.
+
+	bool bGeneratedAnyComponent = false;
+
+	if (!IncludeGraphNames.IsEmpty())
+	{
+		// Generate all components with graph of each name, in specified order.
+		for (const FName& IncludeGraphName : IncludeGraphNames)
+		{
+			auto FilterOnGraphName = [IncludeGraphName](const UPCGComponent* InComponent)
+			{
+				const UPCGGraph* Graph = InComponent ? InComponent->GetGraph() : nullptr;
+				return Graph && Graph->GetName() == IncludeGraphName;
+			};
+
+			bGeneratedAnyComponent |= PCGWorldPartitionBuilder::GenerateComponents(ComponentsToGenerate, World, bOneComponentAtATime, FilterOnGraphName);
+		}
+	}
+	else
+	{
+		bGeneratedAnyComponent |= PCGWorldPartitionBuilder::GenerateComponents(ComponentsToGenerate, World, bOneComponentAtATime);
+	}
+
+	// 4. Get packages that were dirtied during generation and record them for saving later.
+
+	if (bGeneratedAnyComponent)
+	{
+		TArray<UPackage*> DirtyPackages;
+		FEditorFileUtils::GetDirtyPackages(DirtyPackages);
+		
+		DirtyPackages.RemoveSwap(nullptr);
+		for (UPackage* DirtyPackage : DirtyPackages)
+		{
+			PendingDirtyPackages.Add(DirtyPackage->GetPersistentGuid(), DirtyPackage);
+		}
+	}
+	
+	const int NumPackagesDirtied = PendingDirtyPackages.Num() - NumPendingDirtyPackagesBefore;
+	if (!bGeneratedAnyComponent)
+	{
+		UE_LOG(LogPCGWorldPartitionBuilder, Log, TEXT("No components found to generate. No packages will be saved."));
+	}
+	else if (NumPackagesDirtied <= 0)
+	{
+		UE_LOG(LogPCGWorldPartitionBuilder, Log, TEXT("At least one component was generated but no additional packages were dirtied."));
+	}
+	else
+	{
+		UE_LOG(LogPCGWorldPartitionBuilder, Log, TEXT("Generation complete, %d packages dirtied."), NumPackagesDirtied);
+	}
+
+	return bGeneratedAnyComponent;
+}
+
+bool UPCGWorldPartitionBuilder::PostRun(UWorld* World, FPackageSourceControlHelper& PackageHelper, const bool bInRunSuccess)
+{
+	if (!bInRunSuccess)
+	{
+		return true;
+	}
+
+	// Save/delete pending packages.
+
+	TArray<UPackage*> DirtyPackages;
+	PendingDirtyPackages.GenerateValueArray(DirtyPackages);
+
+	// Empty packages should be deleted - mirrors logic in InternalPromptForCheckoutAndSave()
+	TArray<UPackage*> PackagesToDelete;
+	for (int i = DirtyPackages.Num() - 1; i >= 0; --i)
+	{
+		if (UPackage::IsEmptyPackage(DirtyPackages[i]))
+		{
+			PackagesToDelete.Add(DirtyPackages[i]);
+			DirtyPackages.RemoveAtSwap(i);
+		}
+	}
+
+	if (!SavePackages(DirtyPackages, PackageHelper))
+	{
+		return false;
+	}
+	
+	if (!PackagesToDelete.IsEmpty())
+	{
+		// Release any file handles and prepare for delete.
+		for (UPackage* PackageToDelete : PackagesToDelete)
+		{
+			ResetLoaders(PackageToDelete);
+		}
+
+		if (!DeletePackages(PackagesToDelete, PackageHelper))
+		{
+			return false;
+		}
+	}
+
+	TArray<FString> FilesToSubmit;
+	FilesToSubmit.Append(SourceControlHelpers::PackageFilenames(DirtyPackages));
+
+	const FString ChangeDescription = FString::Printf(TEXT("Generated PCG components for world '%s'"), *World->GetName());
+	return OnFilesModified(FilesToSubmit, ChangeDescription);
+}
+
+void PCGWorldPartitionBuilder::CollectComponentsToGenerate(
+	UWorld* InWorld,
+	TFunctionRef<bool(const UPCGComponent*)> ComponentFilter,
+	TArray<TWeakObjectPtr<UPCGComponent>>& OutComponents)
+{
+	check(InWorld);
+
+	TArray<UObject*> AllComponents;
+	GetObjectsOfClass(UPCGComponent::StaticClass(), AllComponents, /*bIncludeDerivedClasses=*/true);
+	for (UObject* ComponentObject : AllComponents)
+	{
+		if (!IsValid(ComponentObject))
+		{
+			continue;
+		}
+
+		UPCGComponent* Component = Cast<UPCGComponent>(ComponentObject);
+		if (!Component || !Component->GetOwner() || Component->GetWorld() != InWorld)
+		{
+			continue;
+		}
+
+		if (Component->IsLocalComponent())
+		{
+			// Skip LCs, rely on original component to generate everything.
+			continue;
+		}
+
+		if (!ComponentFilter(Component))
+		{
+			continue;
+		}
+
+		OutComponents.Add(Component);
+	}
+}
+
+bool PCGWorldPartitionBuilder::GenerateComponents(TArray<TWeakObjectPtr<UPCGComponent>>& Components, UWorld* InWorld, bool bOneComponentAtATime)
+{
+	return PCGWorldPartitionBuilder::GenerateComponents(Components, InWorld, bOneComponentAtATime, [](const UPCGComponent*) { return true; });
+}
+
+bool PCGWorldPartitionBuilder::GenerateComponents(
+	TArray<TWeakObjectPtr<UPCGComponent>>& Components,
+	UWorld* InWorld,
+	bool bOneComponentAtATime,
+	TFunctionRef<bool(const UPCGComponent*)> ComponentFilter)
+{
+	if (!bOneComponentAtATime)
+	{
+		PCGWorldPartitionBuilder::WaitForAllAsyncEditorProcesses(InWorld);
+	}
+
+	auto WaitForComponentGeneration = [InWorld](const UPCGComponent* InComponent)
+	{
+		UE_LOG(LogPCGWorldPartitionBuilder, Log, TEXT("Completing generation on PCG component on actor '%s' label '%s', graph '%s'"),
+			*InComponent->GetOwner()->GetName(),
+			*InComponent->GetOwner()->GetActorNameOrLabel(),
+			*InComponent->GetGraph()->GetName());
+
+		while (InComponent->IsGenerating())
+		{
+			FWorldPartitionHelpers::FakeEngineTick(InWorld);
+		}
+
+		// Can be useful to let some things flush/update after generation.
+		FWorldPartitionHelpers::FakeEngineTick(InWorld);
+	};
+
+	TSet<TObjectKey<UPCGComponent>> GeneratedComponents;
+
+	for (TWeakObjectPtr<UPCGComponent> ComponentWeakPtr : Components)
+	{
+		UPCGComponent* Component = ComponentWeakPtr.Get();
+		const UPCGGraph* Graph = Component ? Component->GetGraph() : nullptr;
+		if (!Component || !Graph)
+		{
+			UE_LOG(LogPCGWorldPartitionBuilder, Warning, TEXT("Lost a PCG component weak pointer, component will not be generated."));
+			continue;
+		}
+
+		if (!ComponentFilter(Component))
+		{
+			continue;
+		}
+
+		if (bOneComponentAtATime)
+		{
+			PCGWorldPartitionBuilder::WaitForAllAsyncEditorProcesses(InWorld);
+		}
+
+		UE_LOG(LogPCGWorldPartitionBuilder, Log, TEXT("Generate PCG component on actor '%s' label '%s', graph '%s'"),
+			*Component->GetOwner()->GetName(),
+			*Component->GetOwner()->GetActorNameOrLabel(),
+			*Component->GetGraph()->GetName());
+
+		PCGWorldPartitionBuilder::GenerateComponent(Component, InWorld);
+		GeneratedComponents.Add(Component);
+
+		if (bOneComponentAtATime)
+		{
+			WaitForComponentGeneration(Component);
+		}
+	}
+
+	if (!bOneComponentAtATime)
+	{
+		for (TObjectKey<UPCGComponent>& GeneratedComponent : GeneratedComponents)
+		{
+			if (const UPCGComponent* Component = GeneratedComponent.ResolveObjectPtr())
+			{
+				WaitForComponentGeneration(Component);
+			}
+		}
+	}
+
+	return !GeneratedComponents.IsEmpty();
+}
+
+void PCGWorldPartitionBuilder::GenerateComponent(UPCGComponent* InComponent, UWorld* InWorld)
+{
+	if (InComponent->GetSerializedEditingMode() == EPCGEditorDirtyMode::LoadAsPreview)
+	{
+		InComponent->SetEditingMode(EPCGEditorDirtyMode::LoadAsPreview, InComponent->GetSerializedEditingMode());
+		InComponent->ChangeTransientState(EPCGEditorDirtyMode::LoadAsPreview);
+	}
+
+	// Force generate as components that are already generated may decline the request.
+	InComponent->GenerateLocal(/*bForce=*/true);
+}
+
+void PCGWorldPartitionBuilder::WaitForAllAsyncEditorProcesses(UWorld* InWorld)
+{
+	InWorld->BlockTillLevelStreamingCompleted();
+
+	// Quite a lot of activity can happen here..
+	FWorldPartitionHelpers::FakeEngineTick(InWorld);
+
+	// Make sure all actor changes are out of the way, as we are sensitive to these.
+	bool bActorsStable = false;
+	FDelegateHandle ActorAddedHandle = GEngine->OnLevelActorAdded().AddLambda([&bActorsStable](AActor* InActor) { bActorsStable = false; });
+	FDelegateHandle ActorDeletedHandle = GEngine->OnLevelActorDeleted().AddLambda([&bActorsStable](AActor* InActor) { bActorsStable = false; });
+
+	while (!bActorsStable)
+	{
+		bActorsStable = true;
+		FWorldPartitionHelpers::FakeEngineTick(InWorld);
+	}
+
+	GEngine->OnLevelActorAdded().Remove(ActorAddedHandle);
+	GEngine->OnLevelActorDeleted().Remove(ActorDeletedHandle);
+
+	// Finalize asset compilation before we potentially use them during generation. Example: static mesh collision depends on built SMs.
+	// Done before every generation to ensure everything up until now is compiled/built.
+	FAssetCompilingManager::Get().FinishAllCompilation();
+
+	// This may execute pending construction scripts.
+	FAssetCompilingManager::Get().ProcessAsyncTasks();
+}
+
+void PCGWorldPartitionBuilder::Build(const TArray<FString>& Args)
+{
+	if (UWorld* World = (GEditor ? GEditor->GetEditorWorldContext().World() : nullptr))
+	{
+		IWorldPartitionEditorModule& WorldPartitionEditorModule = FModuleManager::LoadModuleChecked<IWorldPartitionEditorModule>("WorldPartitionEditor");
+		IWorldPartitionEditorModule::FRunBuilderParams Params;
+		Params.BuilderClass = UPCGWorldPartitionBuilder::StaticClass();
+		Params.World = World;
+		Params.OperationDescription = FText::FromString("Generating PCG Components...");
+		
+		Params.ExtraArgs = TEXT("-AllowCommandletRendering");
+		for (const FString& Arg : Args)
+		{
+			Params.ExtraArgs += " ";
+			Params.ExtraArgs += Arg;
+		}
+
+		WorldPartitionEditorModule.RunBuilder(Params);
+	}
+}
diff --git a/Engine/Plugins/Experimental/PCG/Source/PCGEditor/Private/WorldPartitionBuilder/PCGWorldPartitionBuilder.h b/Engine/Plugins/Experimental/PCG/Source/PCGEditor/Private/WorldPartitionBuilder/PCGWorldPartitionBuilder.h
new file mode 100644
--- /dev/null
+++ b/Engine/Plugins/Experimental/PCG/Source/PCGEditor/Private/WorldPartitionBuilder/PCGWorldPartitionBuilder.h
@@ -0,0 +1,53 @@
+// Copyright Epic Games, Inc. All Rights Reserved.
+
+#pragma once
+
+#include "WorldPartition/WorldPartitionBuilder.h"
+
+#include "PCGWorldPartitionBuilder.generated.h"
+
+class UWorld;
+
+/**
+* Builder that triggers generation on PCG components.
+* 
+* Example command line:
+*   ProjectName MapName -Unattended -AllowCommandletRendering -run=WorldPartitionBuilderCommandlet -Builder=PCGWorldPartitionBuilder  -IncludeGraphNames=PCG_GraphA;PCG_GraphB
+*/
+UCLASS()
+class UPCGWorldPartitionBuilder : public UWorldPartitionBuilder
+{
+	GENERATED_UCLASS_BODY()
+
+public:
+	// Whether to require initialization of rendering, for now default to true. Requires a GPU present.
+	// Get Texture Data benefits from rendering - means the GPU will be available as a fallback sampling method so
+	// that the "CPU sampling" option does not have to be enabled on the texture. This may expand to other nodes if
+	// we use GPU resources or compute more in the future.
+	virtual bool RequiresCommandletRendering() const override { return true; }
+
+	// In the future we may load a world in stages, but for now load everything to make sure we have the complete level.
+	virtual ELoadingMode GetLoadingMode() const override { return ELoadingMode::EntireWorld; }
+
+protected:
+	virtual bool RunInternal(UWorld* World, const FCellInfo& InCellInfo, FPackageSourceControlHelper& PackageHelper) override;
+	virtual bool PostRun(UWorld* World, FPackageSourceControlHelper& PackageHelper, const bool bInRunSuccess) override;
+	virtual bool CanProcessNonPartitionedWorlds() const { return true; }
+
+private:
+	/** The packages dirtied while generating all components. */
+	UPROPERTY(Transient)
+	TMap<FGuid, UPackage*> PendingDirtyPackages;
+
+	/** Include components which have editing mode set to Normal. */
+	bool bGenerateEditingModeNormalComponents = false;
+
+	/** If non empty, components with each graph name will be generated (if editing mode is included), in order. */
+	TArray<FName> IncludeGraphNames;
+
+	/** Call generate on each component and wait until completion and any async processes before generating the next. */
+	bool bOneComponentAtATime = false;
+
+	/** If non empty, only components on actors with given ID(s) will be generated. */
+	TArray<FString> IncludeActorIDs;
+};
