[PCG] Fixed a bug with the Projection node's Color Blend mode

#rnx

diff --git a/Engine/Plugins/Experimental/PCG/Source/PCG/Private/Data/PCGProjectionData.cpp b/Engine/Plugins/Experimental/PCG/Source/PCG/Private/Data/PCGProjectionData.cpp
--- a/Engine/Plugins/Experimental/PCG/Source/PCG/Private/Data/PCGProjectionData.cpp
+++ b/Engine/Plugins/Experimental/PCG/Source/PCG/Private/Data/PCGProjectionData.cpp
@@ -32,9 +32,9 @@ namespace PCGProjectionPrivate
 			case EPCGProjectionColorBlendMode::TargetValue:
 				return TargetColor;
 			case EPCGProjectionColorBlendMode::Add:
-				return ClampVector(TargetColor + SourceColor);
+				return ClampVector(SourceColor + TargetColor);
 			case EPCGProjectionColorBlendMode::Subtract:
-				return ClampVector(TargetColor - SourceColor);
+				return ClampVector(SourceColor - TargetColor);
 			case EPCGProjectionColorBlendMode::Multiply:
 				return SourceColor * TargetColor;
 			default:
@@ -268,11 +268,12 @@ const UPCGPointData* UPCGProjectionData::CreatePointData(FPCGContext* Context) c
 
 	TArray<FPCGPoint>& Points = PointData->GetMutablePoints();
 
-	FPCGAsync::AsyncPointProcessing(Context, SourcePoints.Num(), Points, [this, SourcePointData, SourceMetadata, PointData, OutMetadata, TempTargetMetadata, &SourcePoints](int32 Index, FPCGPoint& OutPoint)
+	FPCGAsync::AsyncPointProcessing(Context, SourcePoints.Num(), Points, [this, SourceMetadata, OutMetadata, TempTargetMetadata, &SourcePoints](int32 Index, FPCGPoint& OutPoint)
 	{
 		const FPCGPoint& SourcePoint = SourcePoints[Index];
 
 		FPCGPoint PointFromTarget;
+		bool bValidProjection = true;
 		if (!Target->ProjectPoint(SourcePoint.Transform, SourcePoint.GetLocalBounds(), ProjectionParams, PointFromTarget, TempTargetMetadata))
 		{
 			if (!bKeepZeroDensityPoints)
@@ -282,8 +283,9 @@ const UPCGPointData* UPCGProjectionData::CreatePointData(FPCGContext* Context) c
 			else
 			{
 				// Point is rejected, mark its density to zero, put it in a state where we won't affect the output point
+				bValidProjection = false;
 				PointFromTarget.Transform = SourcePoint.Transform;
-				PointFromTarget.Color = FVector4::One();
+				PointFromTarget.Color = SourcePoint.Color;
 				PointFromTarget.Density = 0;
 				PointFromTarget.MetadataEntry = PCGInvalidEntryKey;
 			}
@@ -296,7 +298,11 @@ const UPCGPointData* UPCGProjectionData::CreatePointData(FPCGContext* Context) c
 		// TODO this would be cleaner if there was a ProjectPoint that took an FPCGPoint
 		OutPoint.Transform = PointFromTarget.Transform;
 
-		OutPoint.Color = PCGProjectionPrivate::ApplyProjectionColorBlend(SourcePoint.Color, PointFromTarget.Color, ProjectionParams.ColorBlendMode);
+		// There is no reason to project a color blend if the projection failed
+		if (bValidProjection)
+		{
+			OutPoint.Color = PCGProjectionPrivate::ApplyProjectionColorBlend(SourcePoint.Color, PointFromTarget.Color, ProjectionParams.ColorBlendMode);
+		}
 
 		OutPoint.Density *= PointFromTarget.Density;
 
