[PCG] Minor bug fix to Graph Parameter Helpers and soft object paths

#rnx

diff --git a/Engine/Plugins/Experimental/PCG/Source/PCG/Private/Helpers/PCGGraphParameterExtension.cpp b/Engine/Plugins/Experimental/PCG/Source/PCG/Private/Helpers/PCGGraphParameterExtension.cpp
--- a/Engine/Plugins/Experimental/PCG/Source/PCG/Private/Helpers/PCGGraphParameterExtension.cpp
+++ b/Engine/Plugins/Experimental/PCG/Source/PCG/Private/Helpers/PCGGraphParameterExtension.cpp
@@ -39,9 +39,9 @@ template<> TValueOrError<FString, EPropertyBagResult> FPCGGraphParameterExtensio
 {
 	return PropertyBag.GetValueString(PropertyName);
 }
-template<> TValueOrError<FSoftObjectPath, EPropertyBagResult> FPCGGraphParameterExtension::GetGraphParameter(const FInstancedPropertyBag& PropertyBag, const FName PropertyName)
+template<> TValueOrError<FSoftObjectPath*, EPropertyBagResult> FPCGGraphParameterExtension::GetGraphParameter(const FInstancedPropertyBag& PropertyBag, const FName PropertyName)
 {
-	return PropertyBag.GetValueSoftPath(PropertyName);
+	return PropertyBag.GetValueStruct<FSoftObjectPath>(PropertyName);
 }
 template<> TValueOrError<TSoftObjectPtr<UObject>, EPropertyBagResult> FPCGGraphParameterExtension::GetGraphParameter(const FInstancedPropertyBag& PropertyBag, const FName PropertyName)
 {
@@ -137,7 +137,7 @@ template<> EPropertyBagResult FPCGGraphParameterExtension::SetGraphParameter(FIn
 }
 template<> EPropertyBagResult FPCGGraphParameterExtension::SetGraphParameter(FInstancedPropertyBag& PropertyBag, const FName PropertyName, const FSoftObjectPath& Value)
 {
-	return PropertyBag.SetValueSoftPath(PropertyName, Value);
+	return PropertyBag.SetValueStruct(PropertyName, Value);
 }
 template<> EPropertyBagResult FPCGGraphParameterExtension::SetGraphParameter(FInstancedPropertyBag& PropertyBag, const FName PropertyName, const TSoftObjectPtr<UObject>& Value)
 {
diff --git a/Engine/Plugins/Experimental/PCG/Source/PCG/Private/Helpers/PCGGraphParametersHelpers.cpp b/Engine/Plugins/Experimental/PCG/Source/PCG/Private/Helpers/PCGGraphParametersHelpers.cpp
--- a/Engine/Plugins/Experimental/PCG/Source/PCG/Private/Helpers/PCGGraphParametersHelpers.cpp
+++ b/Engine/Plugins/Experimental/PCG/Source/PCG/Private/Helpers/PCGGraphParametersHelpers.cpp
@@ -174,7 +174,7 @@ uint8 UPCGGraphParametersHelpers::GetEnumParameter(const UPCGGraphInterface* Gra
 
 FSoftObjectPath UPCGGraphParametersHelpers::GetSoftObjectPathParameter(const UPCGGraphInterface* GraphInterface, const FName Name)
 {
-	return PCGGraphParametersHelpersPrivate::ValidateAndReturnResult<FSoftObjectPath>(GraphInterface, Name);
+	return PCGGraphParametersHelpersPrivate::ValidateAndReturnResult<FSoftObjectPath, FSoftObjectPath*>(GraphInterface, Name);
 }
 
 TSoftObjectPtr<UObject> UPCGGraphParametersHelpers::GetSoftObjectParameter(const UPCGGraphInterface* GraphInterface, const FName Name)
diff --git a/Engine/Plugins/Experimental/PCG/Source/PCG/Public/Helpers/PCGGraphParameterExtension.h b/Engine/Plugins/Experimental/PCG/Source/PCG/Public/Helpers/PCGGraphParameterExtension.h
--- a/Engine/Plugins/Experimental/PCG/Source/PCG/Public/Helpers/PCGGraphParameterExtension.h
+++ b/Engine/Plugins/Experimental/PCG/Source/PCG/Public/Helpers/PCGGraphParameterExtension.h
@@ -30,7 +30,7 @@ template<> PCG_API TValueOrError<int32, EPropertyBagResult> FPCGGraphParameterEx
 template<> PCG_API TValueOrError<int64, EPropertyBagResult> FPCGGraphParameterExtension::GetGraphParameter(const FInstancedPropertyBag& PropertyBag, const FName PropertyName);
 template<> PCG_API TValueOrError<FName, EPropertyBagResult> FPCGGraphParameterExtension::GetGraphParameter(const FInstancedPropertyBag& PropertyBag, const FName PropertyName);
 template<> PCG_API TValueOrError<FString, EPropertyBagResult> FPCGGraphParameterExtension::GetGraphParameter(const FInstancedPropertyBag& PropertyBag, const FName PropertyName);
-template<> PCG_API TValueOrError<FSoftObjectPath, EPropertyBagResult> FPCGGraphParameterExtension::GetGraphParameter(const FInstancedPropertyBag& PropertyBag, const FName PropertyName);
+template<> PCG_API TValueOrError<FSoftObjectPath*, EPropertyBagResult> FPCGGraphParameterExtension::GetGraphParameter(const FInstancedPropertyBag& PropertyBag, const FName PropertyName);
 template<> PCG_API TValueOrError<TSoftObjectPtr<UObject>, EPropertyBagResult> FPCGGraphParameterExtension::GetGraphParameter(const FInstancedPropertyBag& PropertyBag, const FName PropertyName);
 template<> PCG_API TValueOrError<TSoftClassPtr<UObject>, EPropertyBagResult> FPCGGraphParameterExtension::GetGraphParameter(const FInstancedPropertyBag& PropertyBag, const FName PropertyName);
 template<> PCG_API TValueOrError<FVector*, EPropertyBagResult> FPCGGraphParameterExtension::GetGraphParameter(const FInstancedPropertyBag& PropertyBag, const FName PropertyName);
