[PCG] Demote a check on generated resources to an ensure as it's important but not fatal.

#rnx

diff --git a/Engine/Plugins/Experimental/PCG/Source/PCG/Private/PCGComponent.cpp b/Engine/Plugins/Experimental/PCG/Source/PCG/Private/PCGComponent.cpp
--- a/Engine/Plugins/Experimental/PCG/Source/PCG/Private/PCGComponent.cpp
+++ b/Engine/Plugins/Experimental/PCG/Source/PCG/Private/PCGComponent.cpp
@@ -2662,7 +2662,12 @@ bool UPCGComponent::ShouldTrackLandscape() const
 void UPCGComponent::SetManagedResources(const TArray<TObjectPtr<UPCGManagedResource>>& Resources)
 {
 	FScopeLock ResourcesLock(&GeneratedResourcesLock);
-	check(GeneratedResources.IsEmpty());
+
+	// We expect the GeneratedResources to be empty here, as otherwise they might not be taken care of properly - they
+	// will be lost down below, but this should not happen. However, if the GeneratedResources are marked as Visible,
+	// then they will be copied over during BP duplication, hence why this will happen, hence the ensure here.
+	ensure(GeneratedResources.IsEmpty());
+
 	GeneratedResources = Resources;
 
 	// Remove any null entries
diff --git a/Engine/Plugins/Experimental/PCG/Source/PCG/Public/PCGComponent.h b/Engine/Plugins/Experimental/PCG/Source/PCG/Public/PCGComponent.h
--- a/Engine/Plugins/Experimental/PCG/Source/PCG/Public/PCGComponent.h
+++ b/Engine/Plugins/Experimental/PCG/Source/PCG/Public/PCGComponent.h
@@ -458,6 +458,8 @@ private:
 	TSet<TSoftObjectPtr<AActor>> GeneratedActors_DEPRECATED;
 #endif
 
+	// NOTE: This should not be made visible or editable because it will change the way the BP actors are
+	// duplicated/setup and might trigger an ensure in the resources.
 	UPROPERTY()
 	TArray<TObjectPtr<UPCGManagedResource>> GeneratedResources;
 
