[PCG] Fix unit test fail

#rnx

diff --git a/Engine/Plugins/Experimental/PCG/Source/PCG/Private/Tests/Elements/PCGPropertyToParamDataElementTest.cpp b/Engine/Plugins/Experimental/PCG/Source/PCG/Private/Tests/Elements/PCGPropertyToParamDataElementTest.cpp
--- a/Engine/Plugins/Experimental/PCG/Source/PCG/Private/Tests/Elements/PCGPropertyToParamDataElementTest.cpp
+++ b/Engine/Plugins/Experimental/PCG/Source/PCG/Private/Tests/Elements/PCGPropertyToParamDataElementTest.cpp
@@ -307,6 +307,7 @@ bool FPCGPropertyToParamDataPropertyTypeTest::RunTest(const FString& Parameters)
 
 	// Unknown property
 	AddExpectedError(TEXT("Property 'DummyMissingProperty' does not exist"), EAutomationExpectedErrorFlags::Contains, 1);
+	AddExpectedError(TEXT("Fail to extract the property 'DummyMissingProperty' on actor"), EAutomationExpectedErrorFlags::Contains, 1);
 	bSuccess &= VerifyAttributeValueInvalid(this, TestData, TEXT("DummyMissingProperty"), 42, ExtraTestWhat);
 
 	// Missing property
