[PCG] Fix UProp category specifiers for Instanced UObjects in PCG
* MeshSelectors
* InstanceDataPackers
* MatchAndSets

diff --git a/Engine/Plugins/Experimental/PCG/Source/PCG/Public/Elements/PCGPointMatchAndSet.h b/Engine/Plugins/Experimental/PCG/Source/PCG/Public/Elements/PCGPointMatchAndSet.h
--- a/Engine/Plugins/Experimental/PCG/Source/PCG/Public/Elements/PCGPointMatchAndSet.h
+++ b/Engine/Plugins/Experimental/PCG/Source/PCG/Public/Elements/PCGPointMatchAndSet.h
@@ -14,7 +14,7 @@ class UPCGMatchAndSetBase;
 * on the point data it consumes as input.
 * E.g. for a given point, if it matches with something in the Match & Set object, it will set a value on the point.
 */
-UCLASS(BlueprintType, ClassGroup = (Procedural))
+UCLASS(BlueprintType, ClassGroup = (Procedural), meta = (PrioritizeCategories = "Settings"))
 class PCG_API UPCGPointMatchAndSetSettings : public UPCGSettings
 {
 	GENERATED_BODY()
@@ -55,11 +55,11 @@ protected:
 
 public:
 	/** Defines the type of Match & Set object to use. */
-	UPROPERTY(BlueprintReadOnly, EditAnywhere, NoClear, Category = Settings)
+	UPROPERTY(BlueprintReadOnly, EditAnywhere, NoClear, Category = MatchAndSet)
 	TSubclassOf<UPCGMatchAndSetBase> MatchAndSetType;
 
 	/** Instance of MatchAndSetType, stores the data that will be used in these settings. */
-	UPROPERTY(BlueprintReadOnly, VisibleAnywhere, Instanced, Category = Settings)
+	UPROPERTY(BlueprintReadOnly, VisibleAnywhere, Instanced, Category = MatchAndSet)
 	TObjectPtr<UPCGMatchAndSetBase> MatchAndSetInstance;
 
 	/** "Set" part of the Match & Set - defines what will be changed in the operation */
diff --git a/Engine/Plugins/Experimental/PCG/Source/PCG/Public/Elements/PCGStaticMeshSpawner.h b/Engine/Plugins/Experimental/PCG/Source/PCG/Public/Elements/PCGStaticMeshSpawner.h
--- a/Engine/Plugins/Experimental/PCG/Source/PCG/Public/Elements/PCGStaticMeshSpawner.h
+++ b/Engine/Plugins/Experimental/PCG/Source/PCG/Public/Elements/PCGStaticMeshSpawner.h
@@ -35,7 +35,7 @@ struct PCG_API FPCGStaticMeshSpawnerEntry
 	FCollisionProfileName CollisionProfile;
 };
 
-UCLASS(BlueprintType, ClassGroup = (Procedural))
+UCLASS(BlueprintType, ClassGroup = (Procedural), meta = (PrioritizeCategories = "Settings"))
 class PCG_API UPCGStaticMeshSpawnerSettings : public UPCGSettings
 {
 	GENERATED_BODY()
@@ -73,17 +73,17 @@ public:
 
 public:
 	/** Defines the method of mesh selection per input data */
-	UPROPERTY(BlueprintReadOnly, EditAnywhere, NoClear, Category = Settings)
+	UPROPERTY(BlueprintReadOnly, EditAnywhere, NoClear, Category = MeshSelector)
 	TSubclassOf<UPCGMeshSelectorBase> MeshSelectorType;
 
-	UPROPERTY(BlueprintReadOnly, VisibleAnywhere, Instanced, Category = Settings)
+	UPROPERTY(BlueprintReadOnly, VisibleAnywhere, Instanced, Category = MeshSelector)
 	TObjectPtr<UPCGMeshSelectorBase> MeshSelectorParameters;
 
 	/** Defines the method of custom data packing for spawned (H)ISMCs */
-	UPROPERTY(BlueprintReadOnly, EditAnywhere, Category = Settings)
+	UPROPERTY(BlueprintReadOnly, EditAnywhere, Category = InstanceDataPacker)
 	TSubclassOf<UPCGInstanceDataPackerBase> InstanceDataPackerType;
 
-	UPROPERTY(BlueprintReadOnly, VisibleAnywhere, Instanced, Category = Settings)
+	UPROPERTY(BlueprintReadOnly, VisibleAnywhere, Instanced, Category = InstanceDataPacker)
 	TObjectPtr<UPCGInstanceDataPackerBase> InstanceDataPackerParameters = nullptr;
 
 	/** Attribute name to store mesh SoftObjectPaths inside if the output pin is connected. Note: Will overwrite existing data if the attribute name already exists. */
diff --git a/Engine/Plugins/Experimental/PCG/Source/PCG/Public/InstanceDataPackers/PCGInstanceDataPackerByAttribute.h b/Engine/Plugins/Experimental/PCG/Source/PCG/Public/InstanceDataPackers/PCGInstanceDataPackerByAttribute.h
--- a/Engine/Plugins/Experimental/PCG/Source/PCG/Public/InstanceDataPackers/PCGInstanceDataPackerByAttribute.h
+++ b/Engine/Plugins/Experimental/PCG/Source/PCG/Public/InstanceDataPackers/PCGInstanceDataPackerByAttribute.h
@@ -15,6 +15,6 @@ public:
 	virtual void PackInstances_Implementation(UPARAM(ref) FPCGContext& Context, const UPCGSpatialData* InSpatialData, UPARAM(ref) const FPCGMeshInstanceList& InstanceList, FPCGPackedCustomData& OutPackedCustomData) const override;
 
 public:
-	UPROPERTY(BlueprintReadWrite, EditAnywhere, Category = Settings)
+	UPROPERTY(BlueprintReadWrite, EditAnywhere, Category = InstanceDataPacker)
 	TArray<FName> AttributeNames;
 };
diff --git a/Engine/Plugins/Experimental/PCG/Source/PCG/Public/InstanceDataPackers/PCGInstanceDataPackerByRegex.h b/Engine/Plugins/Experimental/PCG/Source/PCG/Public/InstanceDataPackers/PCGInstanceDataPackerByRegex.h
--- a/Engine/Plugins/Experimental/PCG/Source/PCG/Public/InstanceDataPackers/PCGInstanceDataPackerByRegex.h
+++ b/Engine/Plugins/Experimental/PCG/Source/PCG/Public/InstanceDataPackers/PCGInstanceDataPackerByRegex.h
@@ -15,6 +15,6 @@ public:
 	virtual void PackInstances_Implementation(UPARAM(ref) FPCGContext& Context, const UPCGSpatialData* InSpatialData, UPARAM(ref) const FPCGMeshInstanceList& InstanceList, FPCGPackedCustomData& OutPackedCustomData) const override;
 
 public:
-	UPROPERTY(BlueprintReadWrite, EditAnywhere, Category = Settings)
+	UPROPERTY(BlueprintReadWrite, EditAnywhere, Category = InstanceDataPacker)
 	TArray<FString> RegexPatterns; 
 };
diff --git a/Engine/Plugins/Experimental/PCG/Source/PCG/Public/MatchAndSet/PCGMatchAndSetByAttribute.h b/Engine/Plugins/Experimental/PCG/Source/PCG/Public/MatchAndSet/PCGMatchAndSetByAttribute.h
--- a/Engine/Plugins/Experimental/PCG/Source/PCG/Public/MatchAndSet/PCGMatchAndSetByAttribute.h
+++ b/Engine/Plugins/Experimental/PCG/Source/PCG/Public/MatchAndSet/PCGMatchAndSetByAttribute.h
@@ -58,11 +58,11 @@ public:
 
 public:
 	/** Attribute to match on the data */
-	UPROPERTY(BlueprintReadWrite, EditAnywhere, Category = Settings)
+	UPROPERTY(BlueprintReadWrite, EditAnywhere, Category = MatchAndSet)
 	FName MatchSourceAttribute;
 
 	/** Type of the attribute to match against. */
-	UPROPERTY(BlueprintReadWrite, EditAnywhere, Category = Settings)
+	UPROPERTY(BlueprintReadWrite, EditAnywhere, Category = MatchAndSet)
 	EPCGMetadataTypes MatchSourceType = EPCGMetadataTypes::Double;
 
 	/** String type of the attribute to match against (if required). */
@@ -70,6 +70,6 @@ public:
 	EPCGMetadataTypesConstantStructStringMode MatchSourceStringMode_DEPRECATED;
 
 	/** Lookup entries (key-value pairs) */
-	UPROPERTY(BlueprintReadWrite, EditAnywhere, Category = Settings)
+	UPROPERTY(BlueprintReadWrite, EditAnywhere, Category = MatchAndSet)
 	TArray<FPCGMatchAndSetByAttributeEntry> Entries;
 };
diff --git a/Engine/Plugins/Experimental/PCG/Source/PCG/Public/MatchAndSet/PCGMatchAndSetWeighted.h b/Engine/Plugins/Experimental/PCG/Source/PCG/Public/MatchAndSet/PCGMatchAndSetWeighted.h
--- a/Engine/Plugins/Experimental/PCG/Source/PCG/Public/MatchAndSet/PCGMatchAndSetWeighted.h
+++ b/Engine/Plugins/Experimental/PCG/Source/PCG/Public/MatchAndSet/PCGMatchAndSetWeighted.h
@@ -57,10 +57,10 @@ public:
 
 public:
 	/** Values and their respective weights */
-	UPROPERTY(BlueprintReadWrite, EditAnywhere, Category = Settings)
+	UPROPERTY(BlueprintReadWrite, EditAnywhere, Category = MatchAndSet)
 	TArray<FPCGMatchAndSetWeightedEntry> Entries;
 
 	/** Controls whether the output data should mutate its seed - prevents issues when doing multiple random processes in a row */
-	UPROPERTY(BlueprintReadWrite, EditAnywhere, Category = Settings)
+	UPROPERTY(BlueprintReadWrite, EditAnywhere, Category = MatchAndSet)
 	bool bShouldMutateSeed = true;
 };
diff --git a/Engine/Plugins/Experimental/PCG/Source/PCG/Public/MatchAndSet/PCGMatchAndSetWeightedByCategory.h b/Engine/Plugins/Experimental/PCG/Source/PCG/Public/MatchAndSet/PCGMatchAndSetWeightedByCategory.h
--- a/Engine/Plugins/Experimental/PCG/Source/PCG/Public/MatchAndSet/PCGMatchAndSetWeightedByCategory.h
+++ b/Engine/Plugins/Experimental/PCG/Source/PCG/Public/MatchAndSet/PCGMatchAndSetWeightedByCategory.h
@@ -65,22 +65,22 @@ public:
 
 public:
 	/** Attribute to match against */
-	UPROPERTY(BlueprintReadWrite, EditAnywhere, Category = Settings)
+	UPROPERTY(BlueprintReadWrite, EditAnywhere, Category = MatchAndSet)
 	FName CategoryAttribute;
 
 	/** Type of the attribute to match against. */
-	UPROPERTY(BlueprintReadWrite, EditAnywhere, Category = Settings)
+	UPROPERTY(BlueprintReadWrite, EditAnywhere, Category = MatchAndSet)
 	EPCGMetadataTypes CategoryType = EPCGMetadataTypes::Double;
 
 	UPROPERTY()
 	EPCGMetadataTypesConstantStructStringMode CategoryStringMode_DEPRECATED;
 
 	/** Lookup entries (key -> weighted list) */
-	UPROPERTY(BlueprintReadWrite, EditAnywhere, Category = Settings)
+	UPROPERTY(BlueprintReadWrite, EditAnywhere, Category = MatchAndSet)
 	TArray<FPCGMatchAndSetWeightedByCategoryEntryList> Categories;
 
 	/** Controls whether the output data should mutate its seed - prevents issues when doing multiple random processes in a row */
-	UPROPERTY(BlueprintReadWrite, EditAnywhere, Category = Settings)
+	UPROPERTY(BlueprintReadWrite, EditAnywhere, Category = MatchAndSet)
 	bool bShouldMutateSeed = true;
 };
 
diff --git a/Engine/Plugins/Experimental/PCG/Source/PCG/Public/MeshSelectors/PCGMeshSelectorByAttribute.h b/Engine/Plugins/Experimental/PCG/Source/PCG/Public/MeshSelectors/PCGMeshSelectorByAttribute.h
--- a/Engine/Plugins/Experimental/PCG/Source/PCG/Public/MeshSelectors/PCGMeshSelectorByAttribute.h
+++ b/Engine/Plugins/Experimental/PCG/Source/PCG/Public/MeshSelectors/PCGMeshSelectorByAttribute.h
@@ -24,16 +24,16 @@ public:
 		UPCGPointData* OutPointData) const override;
 
 public:
-	UPROPERTY(BlueprintReadWrite, EditAnywhere, Category = Settings)
+	UPROPERTY(BlueprintReadWrite, EditAnywhere, Category = MeshSelector)
 	FName AttributeName;
 
-	UPROPERTY(EditAnywhere, Category = Settings)
+	UPROPERTY(EditAnywhere, Category = MeshSelector)
 	FSoftISMComponentDescriptor TemplateDescriptor;
 
-	UPROPERTY(BlueprintReadWrite, EditAnywhere, Category = Settings, meta = (InlineEditConditionToggle))
+	UPROPERTY(BlueprintReadWrite, EditAnywhere, Category = MeshSelector, meta = (InlineEditConditionToggle))
 	bool bUseAttributeMaterialOverrides = false;
 
-	UPROPERTY(BlueprintReadWrite, EditAnywhere, DisplayName = "By Attribute Material Overrides", Category = Settings, meta = (EditCondition = "bUseAttributeMaterialOverrides"))
+	UPROPERTY(BlueprintReadWrite, EditAnywhere, DisplayName = "By Attribute Material Overrides", Category = MeshSelector, meta = (EditCondition = "bUseAttributeMaterialOverrides"))
 	TArray<FName> MaterialOverrideAttributes;
 
 #if WITH_EDITORONLY_DATA
diff --git a/Engine/Plugins/Experimental/PCG/Source/PCG/Public/MeshSelectors/PCGMeshSelectorWeighted.h b/Engine/Plugins/Experimental/PCG/Source/PCG/Public/MeshSelectors/PCGMeshSelectorWeighted.h
--- a/Engine/Plugins/Experimental/PCG/Source/PCG/Public/MeshSelectors/PCGMeshSelectorWeighted.h
+++ b/Engine/Plugins/Experimental/PCG/Source/PCG/Public/MeshSelectors/PCGMeshSelectorWeighted.h
@@ -95,12 +95,12 @@ protected:
 #endif // WITH_EDITOR
 
 public:
-	UPROPERTY(BlueprintReadWrite, EditAnywhere, Category = Settings, meta = (TitleProperty = "DisplayName"))
+	UPROPERTY(BlueprintReadWrite, EditAnywhere, Category = MeshSelector, meta = (TitleProperty = "DisplayName"))
 	TArray<FPCGMeshSelectorWeightedEntry> MeshEntries;
 
-	UPROPERTY(BlueprintReadWrite, EditAnywhere, Category = Settings, meta = (InlineEditConditionToggle))
+	UPROPERTY(BlueprintReadWrite, EditAnywhere, Category = MeshSelector, meta = (InlineEditConditionToggle))
 	bool bUseAttributeMaterialOverrides = false;
 
-	UPROPERTY(BlueprintReadWrite, EditAnywhere, DisplayName = "By Attribute Material Overrides", Category = Settings, meta = (EditCondition = "bUseAttributeMaterialOverrides"))
+	UPROPERTY(BlueprintReadWrite, EditAnywhere, DisplayName = "By Attribute Material Overrides", Category = MeshSelector, meta = (EditCondition = "bUseAttributeMaterialOverrides"))
 	TArray<FName> MaterialOverrideAttributes;
 };
diff --git a/Engine/Plugins/Experimental/PCG/Source/PCG/Public/MeshSelectors/PCGMeshSelectorWeightedByCategory.h b/Engine/Plugins/Experimental/PCG/Source/PCG/Public/MeshSelectors/PCGMeshSelectorWeightedByCategory.h
--- a/Engine/Plugins/Experimental/PCG/Source/PCG/Public/MeshSelectors/PCGMeshSelectorWeightedByCategory.h
+++ b/Engine/Plugins/Experimental/PCG/Source/PCG/Public/MeshSelectors/PCGMeshSelectorWeightedByCategory.h
@@ -59,15 +59,15 @@ protected:
 #endif // WITH_EDITOR
 
 public:
-	UPROPERTY(BlueprintReadWrite, EditAnywhere, Category = Settings)
+	UPROPERTY(BlueprintReadWrite, EditAnywhere, Category = MeshSelector)
 	FName CategoryAttribute;
 
-	UPROPERTY(BlueprintReadWrite, EditAnywhere, Category = Settings)
+	UPROPERTY(BlueprintReadWrite, EditAnywhere, Category = MeshSelector)
 	TArray<FPCGWeightedByCategoryEntryList> Entries;
 
-	UPROPERTY(BlueprintReadWrite, EditAnywhere, Category = Settings, meta = (InlineEditConditionToggle))
+	UPROPERTY(BlueprintReadWrite, EditAnywhere, Category = MeshSelector, meta = (InlineEditConditionToggle))
 	bool bUseAttributeMaterialOverrides = false;
 
-	UPROPERTY(BlueprintReadWrite, EditAnywhere, DisplayName = "By Attribute Material Overrides", Category = Settings, meta = (EditCondition = "bUseAttributeMaterialOverrides"))
+	UPROPERTY(BlueprintReadWrite, EditAnywhere, DisplayName = "By Attribute Material Overrides", Category = MeshSelector, meta = (EditCondition = "bUseAttributeMaterialOverrides"))
 	TArray<FName> MaterialOverrideAttributes;
 };
