[PCG] Automatically set 'GenerateAtRuntime' components to the 'Preview' editing mode
* This fixes Original/NonPartitioned components incorrectly serializing their generated state and managed resources

#rnx

diff --git a/Engine/Plugins/Experimental/PCG/Source/PCG/Private/Graph/PCGGraphExecutor.cpp b/Engine/Plugins/Experimental/PCG/Source/PCG/Private/Graph/PCGGraphExecutor.cpp
--- a/Engine/Plugins/Experimental/PCG/Source/PCG/Private/Graph/PCGGraphExecutor.cpp
+++ b/Engine/Plugins/Experimental/PCG/Source/PCG/Private/Graph/PCGGraphExecutor.cpp
@@ -1611,7 +1611,7 @@ namespace PCGGraphExecutor
 					const FIntVector CellCoords = UPCGActorHelpers::GetCellCoord(ComponentActor->GetActorLocation(), FromGridSize, PCGWorldActor->bUse2DGrid);
 
 					// Search for a transient local component if the source component is runtime managed.
-					const bool bTransientComponent = InContext->SourceComponent->GenerationTrigger == EPCGComponentGenerationTrigger::GenerateAtRuntime;
+					const bool bTransientComponent = InContext->SourceComponent->IsManagedByRuntimeGenSystem();
 					ComponentWithData = Subsystem->GetLocalComponent(FromGridSize, CellCoords, InContext->SourceComponent->GetOriginalComponent(), bTransientComponent);
 				}
 			}
@@ -1671,7 +1671,7 @@ namespace PCGGraphExecutor
 			{
 				PCGGraphExecutionLogging::LogGridLinkageTaskExecuteRetrieveScheduleGraph(InContext, ComponentWithData, InResourceKey);
 
-				EPCGComponentGenerationTrigger GenTrigger = (InContext->SourceComponent->GenerationTrigger == EPCGComponentGenerationTrigger::GenerateAtRuntime) ?
+				EPCGComponentGenerationTrigger GenTrigger = (InContext->SourceComponent->IsManagedByRuntimeGenSystem()) ?
 					EPCGComponentGenerationTrigger::GenerateAtRuntime : EPCGComponentGenerationTrigger::GenerateOnDemand;
 
 				// Wake up this task after graph has generated.
diff --git a/Engine/Plugins/Experimental/PCG/Source/PCG/Private/Grid/PCGPartitionActor.cpp b/Engine/Plugins/Experimental/PCG/Source/PCG/Private/Grid/PCGPartitionActor.cpp
--- a/Engine/Plugins/Experimental/PCG/Source/PCG/Private/Grid/PCGPartitionActor.cpp
+++ b/Engine/Plugins/Experimental/PCG/Source/PCG/Private/Grid/PCGPartitionActor.cpp
@@ -191,7 +191,7 @@ void APCGPartitionActor::BeginPlay()
 			{
 				// If we have an original component that is generated (or generating), this one is automatically generated => GenerateOnLoad.
 				// But if its runtime generated then it's handled by the runtime generation scheduler.
-				if ((It.Key->bGenerated || It.Key->IsGenerating()) && ensure(It.Key->GenerationTrigger != EPCGComponentGenerationTrigger::GenerateAtRuntime))
+				if ((It.Key->bGenerated || It.Key->IsGenerating()) && ensure(!It.Key->IsManagedByRuntimeGenSystem()))
 				{
 					It.Value->GenerationTrigger = EPCGComponentGenerationTrigger::GenerateOnLoad;
 				}
diff --git a/Engine/Plugins/Experimental/PCG/Source/PCG/Private/PCGActorAndComponentMapping.cpp b/Engine/Plugins/Experimental/PCG/Source/PCG/Private/PCGActorAndComponentMapping.cpp
--- a/Engine/Plugins/Experimental/PCG/Source/PCG/Private/PCGActorAndComponentMapping.cpp
+++ b/Engine/Plugins/Experimental/PCG/Source/PCG/Private/PCGActorAndComponentMapping.cpp
@@ -165,7 +165,7 @@ TArray<FPCGTaskId> FPCGActorAndComponentMapping::DispatchToRegisteredLocalCompon
 		return {};
 	}
 
-	const bool bIsRuntimeGenerated = OriginalComponent->GenerationTrigger == EPCGComponentGenerationTrigger::GenerateAtRuntime;
+	const bool bIsRuntimeGenerated = OriginalComponent->IsManagedByRuntimeGenSystem();
 
 	const TMap<const UPCGComponent*, TSet<TObjectPtr<APCGPartitionActor>>>& Map = bIsRuntimeGenerated ? ComponentToRuntimeGenPartitionActorsMap : ComponentToPartitionActorsMap;
 
@@ -279,7 +279,7 @@ bool FPCGActorAndComponentMapping::RegisterOrUpdatePartitionedPCGComponent(UPCGC
 	// In Editor only, we will create new partition actors depending on the new bounds and generation trigger. Runtime managed components should not create PAs here
 	// TODO: For now it will always create the PA. But if we want to create them only when we generate, we need to make
 	// sure to update the runtime flow, for them to also create PA if they need to.
-	if ((bComponentHasChanged || bComponentWasAdded) && InComponent->GenerationTrigger != EPCGComponentGenerationTrigger::GenerateAtRuntime)
+	if ((bComponentHasChanged || bComponentWasAdded) && !InComponent->IsManagedByRuntimeGenSystem())
 	{
 		bool bHasUnbounded = false;
 		PCGHiGenGrid::FSizeArray GridSizes;
@@ -435,7 +435,7 @@ void FPCGActorAndComponentMapping::UnregisterPCGComponent(UPCGComponent* InCompo
 
 void FPCGActorAndComponentMapping::UnregisterPartitionedPCGComponent(UPCGComponent* InComponent)
 {
-	if (!PartitionedOctree.RemoveComponent(InComponent) || InComponent->GenerationTrigger == EPCGComponentGenerationTrigger::GenerateAtRuntime)
+	if (!PartitionedOctree.RemoveComponent(InComponent) || InComponent->IsManagedByRuntimeGenSystem())
 	{
 		return;
 	}
@@ -574,22 +574,18 @@ void FPCGActorAndComponentMapping::ForAllIntersectingPartitionActors(const FBox&
 		return;
 	}
 
-	PCGHiGenGrid::FSizeToGuidMap GridSizeToGuid;
-	PCGWorldActor->GetGridGuids(GridSizeToGuid);
-	for (const TPair<uint32, FGuid>& SizeAndGuid : GridSizeToGuid)
+	auto ForAllIntersectingPartitionActorsOfGridSize = [InFunc, PCGWorldActor, &InBounds](const TMap<uint32, TMap<FIntVector, TObjectPtr<APCGPartitionActor>>>& Map, FRWLock& Lock, uint32 GridSize)
 	{
-		const uint32 GridSize = SizeAndGuid.Key;
-
 		const bool bUse2DGrid = PCGWorldActor->bUse2DGrid;
 		FIntVector MinCellCoords = UPCGActorHelpers::GetCellCoord(InBounds.Min, GridSize, bUse2DGrid);
 		FIntVector MaxCellCoords = UPCGActorHelpers::GetCellCoord(InBounds.Max, GridSize, bUse2DGrid);
 
-		FReadScopeLock ReadLock(PartitionActorsMapLock);
+		FReadScopeLock ReadLock(Lock);
 
-		const TMap<FIntVector, TObjectPtr<APCGPartitionActor>>* PartitionActorsMapGrid = PartitionActorsMap.Find(GridSize);
+		const TMap<FIntVector, TObjectPtr<APCGPartitionActor>>* PartitionActorsMapGrid = Map.Find(GridSize);
 		if (!PartitionActorsMapGrid || PartitionActorsMapGrid->IsEmpty())
 		{
-			continue;
+			return;
 		}
 
 		for (int32 z = MinCellCoords.Z; z <= MaxCellCoords.Z; z++)
@@ -609,6 +605,16 @@ void FPCGActorAndComponentMapping::ForAllIntersectingPartitionActors(const FBox&
 				}
 			}
 		}
+	};
+
+	PCGHiGenGrid::FSizeToGuidMap GridSizeToGuid;
+	PCGWorldActor->GetGridGuids(GridSizeToGuid);
+	for (const TPair<uint32, FGuid>& SizeAndGuid : GridSizeToGuid)
+	{
+		const uint32 GridSize = SizeAndGuid.Key;
+
+		ForAllIntersectingPartitionActorsOfGridSize(PartitionActorsMap, PartitionActorsMapLock, GridSize);
+		ForAllIntersectingPartitionActorsOfGridSize(RuntimeGenPartitionActorsMap, RuntimeGenPartitionActorsMapLock, GridSize);
 	}
 }
 
@@ -713,7 +719,7 @@ bool FPCGActorAndComponentMapping::AnyRuntimeGenComponentsExist() const
 {
 	for (UPCGComponent* Component : PartitionedOctree.GetAllComponents())
 	{
-		if (Component && Component->GenerationTrigger == EPCGComponentGenerationTrigger::GenerateAtRuntime)
+		if (Component && Component->IsManagedByRuntimeGenSystem())
 		{
 			return true;
 		}
@@ -721,7 +727,7 @@ bool FPCGActorAndComponentMapping::AnyRuntimeGenComponentsExist() const
 
 	for (UPCGComponent* Component : NonPartitionedOctree.GetAllComponents())
 	{
-		if (Component && Component->GenerationTrigger == EPCGComponentGenerationTrigger::GenerateAtRuntime)
+		if (Component && Component->IsManagedByRuntimeGenSystem())
 		{
 			return true;
 		}
diff --git a/Engine/Plugins/Experimental/PCG/Source/PCG/Private/PCGComponent.cpp b/Engine/Plugins/Experimental/PCG/Source/PCG/Private/PCGComponent.cpp
--- a/Engine/Plugins/Experimental/PCG/Source/PCG/Private/PCGComponent.cpp
+++ b/Engine/Plugins/Experimental/PCG/Source/PCG/Private/PCGComponent.cpp
@@ -86,8 +86,7 @@ UPCGComponent::UPCGComponent(const FObjectInitializer& InObjectInitializer)
 bool UPCGComponent::CanPartition() const
 {
 	// Support/Force partitioning on non-PCG partition actors in WP worlds. GenerateAtRuntime components can be partitioned even if WorldPartition is not enabled.
-	return ((GetOwner() && GetOwner()->GetWorld() && GetOwner()->GetWorld()->GetWorldPartition() != nullptr)
-		|| GenerationTrigger == EPCGComponentGenerationTrigger::GenerateAtRuntime)
+	return ((GetOwner() && GetOwner()->GetWorld() && GetOwner()->GetWorld()->GetWorldPartition() != nullptr) || IsManagedByRuntimeGenSystem())
 		&& Cast<APCGPartitionActor>(GetOwner()) == nullptr;
 }
 
@@ -178,7 +177,7 @@ bool UPCGComponent::ShouldGenerate(bool bForce, EPCGComponentGenerationTrigger R
 		return false;
 	}
 
-	if (GenerationTrigger == EPCGComponentGenerationTrigger::GenerateAtRuntime)
+	if (IsManagedByRuntimeGenSystem())
 	{
 		// If we're runtime generated, turn down other requests.
 		const bool bShouldGenerate = RequestedGenerationTrigger == EPCGComponentGenerationTrigger::GenerateAtRuntime;
@@ -560,7 +559,7 @@ void UPCGComponent::Cleanup()
 		return;
 	}
 
-	if (GenerationTrigger == EPCGComponentGenerationTrigger::GenerateAtRuntime)
+	if (IsManagedByRuntimeGenSystem())
 	{
 		UE_LOG(LogPCG, Warning, TEXT("Cleanup request denied as this component is managed by the runtime generation scheduler."));
 		return;
@@ -1283,10 +1282,23 @@ void UPCGComponent::PostLoad()
 
 	SetupCallbacksOnCreation();
 
-	CurrentEditingMode = SerializedEditingMode;
+	// Always set the editing mode to Preview when we're in GenerateAtRuntime mode
+	CurrentEditingMode = IsManagedByRuntimeGenSystem() ? EPCGEditorDirtyMode::Preview : SerializedEditingMode;
 
 	if (CurrentEditingMode == EPCGEditorDirtyMode::Preview)
 	{
+		// Make sure we update the transient state if we have been forced into Preview mode by runtime generation.
+		if (CurrentEditingMode != SerializedEditingMode)
+		{
+			if (bGenerated)
+			{
+				CleanupLocalImmediate(/*bRemoveComponents=*/true);
+			}
+
+			PreviousEditingMode = SerializedEditingMode;
+			ChangeTransientState(CurrentEditingMode);
+		}
+
 		bGenerated = false;
 	}
 	else if (CurrentEditingMode == EPCGEditorDirtyMode::LoadAsPreview && !PCGHelpers::IsRuntimeOrPIE())
@@ -1297,7 +1309,7 @@ void UPCGComponent::PostLoad()
 	}
 #endif
 
-	if (!SchedulingPolicy || GenerationTrigger != EPCGComponentGenerationTrigger::GenerateAtRuntime)
+	if (!SchedulingPolicy || IsManagedByRuntimeGenSystem())
 	{
 		RefreshSchedulingPolicy();
 	}
@@ -1306,7 +1318,7 @@ void UPCGComponent::PostLoad()
 		const EObjectFlags Flags = GetMaskedFlags(RF_PropagateToSubObjects) | RF_Transactional;
 		SchedulingPolicy->SetFlags(Flags);
 #if WITH_EDITOR
-		SchedulingPolicy->SetShouldDisplayProperties(GenerationTrigger == EPCGComponentGenerationTrigger::GenerateAtRuntime);
+		SchedulingPolicy->SetShouldDisplayProperties(IsManagedByRuntimeGenSystem());
 #endif
 	}
 }
@@ -1410,19 +1422,10 @@ TStructOnScope<FActorComponentInstanceData> UPCGComponent::GetComponentInstanceD
 
 void UPCGComponent::OnGraphChanged(UPCGGraphInterface* InGraph, EPCGChangeType ChangeType)
 {
-	if (!!(ChangeType & EPCGChangeType::Debug) && GenerationTrigger == EPCGComponentGenerationTrigger::GenerateAtRuntime)
-	{
-		// TODO: Handle RuntimeGen business in the RuntimeGenScheduler, subscribe to OnGraphChanged there instead.
-		CleanupLocalImmediate(/*bRemoveComponents=*/true);
-		GenerateLocal(EPCGComponentGenerationTrigger::GenerateAtRuntime, /*bForce=*/true);
-	}
-	else
-	{
-		const bool bIsStructural = ((ChangeType & (EPCGChangeType::Edge | EPCGChangeType::Structural)) != EPCGChangeType::None);
-		const bool bDirtyInputs = bIsStructural || ((ChangeType & EPCGChangeType::Input) != EPCGChangeType::None);
+	const bool bIsStructural = ((ChangeType & (EPCGChangeType::Edge | EPCGChangeType::Structural)) != EPCGChangeType::None);
+	const bool bDirtyInputs = bIsStructural || ((ChangeType & EPCGChangeType::Input) != EPCGChangeType::None);
 
-		RefreshAfterGraphChanged(InGraph, bIsStructural, bDirtyInputs);
-	}
+	RefreshAfterGraphChanged(InGraph, bIsStructural, bDirtyInputs);
 }
 
 void UPCGComponent::RefreshAfterGraphChanged(UPCGGraphInterface* InGraph, bool bIsStructural, bool bDirtyInputs)
@@ -1464,14 +1467,24 @@ void UPCGComponent::RefreshAfterGraphChanged(UPCGGraphInterface* InGraph, bool b
 	}
 #endif
 
-	// Otherwise, if we are in PIE or runtime, force generate if we have a graph (and were generated). Or cleanup if we have no graph
-	if (bHasGraph && bGenerated)
+	if (IsManagedByRuntimeGenSystem())
 	{
-		GenerateLocal(/*bForce=*/true);
+		if (UPCGSubsystem* Subsystem = GetSubsystem())
+		{
+			Subsystem->RefreshRuntimeGenComponent(this, /*bRemovePartitionActors=*/true);
+		}
 	}
-	else if (!bHasGraph)
+	else
 	{
-		CleanupLocal(/*bRemoveComponents=*/true, /*bSave=*/ false);
+		// Otherwise, if we are in PIE or runtime, force generate if we have a graph (and were generated). Or cleanup if we have no graph
+		if (bHasGraph && bGenerated)
+		{
+			GenerateLocal(/*bForce=*/true);
+		}
+		else if (!bHasGraph)
+		{
+			CleanupLocal(/*bRemoveComponents=*/true, /*bSave=*/ false);
+		}
 	}
 }
 
@@ -1482,13 +1495,16 @@ void UPCGComponent::PreEditChange(FProperty* PropertyAboutToChange)
 
 	if (PropertyAboutToChange && PropertyAboutToChange->GetFName() == GET_MEMBER_NAME_CHECKED(UPCGComponent, GenerationTrigger))
 	{
-		if (GenerationTrigger == EPCGComponentGenerationTrigger::GenerateAtRuntime)
+		if (IsManagedByRuntimeGenSystem())
 		{
 			if (UPCGSubsystem* Subsystem = GetSubsystem())
 			{
 				// When toggling off of GenerateAtRuntime, we should flush the RuntimeGenScheduler state for this component.
 				Subsystem->RefreshRuntimeGenComponent(this, /*bRemovePartitionActors=*/true);
 			}
+
+			// Reset to the the editing mode we were in before entering GenerateAtRuntime mode.
+			SetEditingMode(PreviousEditingMode, SerializedEditingMode);
 		}
 	}
 }
@@ -1588,14 +1604,19 @@ void UPCGComponent::PostEditChangeProperty(FPropertyChangedEvent& PropertyChange
 		if (SchedulingPolicy)
 		{
 			// We should only display scheduling policy parameters when in runtime generation mode.
-			SchedulingPolicy->SetShouldDisplayProperties(GenerationTrigger == EPCGComponentGenerationTrigger::GenerateAtRuntime);
+			SchedulingPolicy->SetShouldDisplayProperties(IsManagedByRuntimeGenSystem());
 		}
 
-		if (GenerationTrigger == EPCGComponentGenerationTrigger::GenerateAtRuntime)
+		if (IsManagedByRuntimeGenSystem())
 		{
 			// If we have been set to GenerateAtRuntime, we should cleanup any artifacts.
 			// TODO: Should this include PartitionActors?
 			CleanupLocalImmediate(/*bRemoveComponents=*/true);
+
+			PreviousEditingMode = CurrentEditingMode;
+
+			SetEditingMode(/*InEditingMode=*/EPCGEditorDirtyMode::Preview, SerializedEditingMode);
+			ChangeTransientState(CurrentEditingMode);
 		}
 		else
 		{
@@ -1869,7 +1890,7 @@ void UPCGComponent::Refresh(bool bStructural)
 	}
 
 	// Runtime component refreshes should go through the runtime scheduler.
-	if (GenerationTrigger == EPCGComponentGenerationTrigger::GenerateAtRuntime)
+	if (IsManagedByRuntimeGenSystem())
 	{
 		if (UPCGSubsystem* Subsystem = GetSubsystem())
 		{
@@ -2009,7 +2030,7 @@ bool UPCGComponent::IsActorTracked(AActor* InActor, bool& bOutIsCulled) const
 
 void UPCGComponent::OnRefresh(bool bForceRefresh)
 {
-	check(GenerationTrigger != EPCGComponentGenerationTrigger::GenerateAtRuntime);
+	check(!IsManagedByRuntimeGenSystem());
 
 	// Mark the refresh task invalid to allow re-triggering refreshes
 	CurrentRefreshTask = InvalidPCGTaskId;
@@ -2430,13 +2451,13 @@ void UPCGComponent::RefreshSchedulingPolicy()
 	SchedulingPolicy = nullptr;
 
 	// We should never create the scheduling policy when not in runtime generation mode.
-	if (SchedulingPolicyClass && GenerationTrigger == EPCGComponentGenerationTrigger::GenerateAtRuntime)
+	if (SchedulingPolicyClass && IsManagedByRuntimeGenSystem())
 	{
 		const EObjectFlags Flags = GetMaskedFlags(RF_PropagateToSubObjects);
 		SchedulingPolicy = NewObject<UPCGSchedulingPolicyBase>(this, SchedulingPolicyClass, NAME_None, Flags);
 
 #if WITH_EDITOR
-		SchedulingPolicy->SetShouldDisplayProperties(GenerationTrigger == EPCGComponentGenerationTrigger::GenerateAtRuntime);
+		SchedulingPolicy->SetShouldDisplayProperties(IsManagedByRuntimeGenSystem());
 #endif
 	}
 }
@@ -2860,6 +2881,7 @@ void FPCGComponentInstanceData::ApplyToComponent(UActorComponent* Component, con
 			PCGComponent->bForceGenerateOnBPAddedToWorld = SourceComponent->bForceGenerateOnBPAddedToWorld;
 
 			PCGComponent->CurrentEditingMode = SourceComponent->CurrentEditingMode;
+			PCGComponent->PreviousEditingMode = SourceComponent->PreviousEditingMode;
 #endif // WITH_EDITOR
 
 			// Non-critical but should be done: transient data, tracked actors cache, landscape tracking
diff --git a/Engine/Plugins/Experimental/PCG/Source/PCG/Private/PCGSubsystem.cpp b/Engine/Plugins/Experimental/PCG/Source/PCG/Private/PCGSubsystem.cpp
--- a/Engine/Plugins/Experimental/PCG/Source/PCG/Private/PCGSubsystem.cpp
+++ b/Engine/Plugins/Experimental/PCG/Source/PCG/Private/PCGSubsystem.cpp
@@ -365,7 +365,7 @@ FPCGTaskId UPCGSubsystem::ScheduleComponent(UPCGComponent* PCGComponent, EPCGHiG
 
 #if WITH_EDITOR
 	// Create the PartitionActors if necessary. Skip if this is a runtime managed component, PAs are handled manually by the RuntimeGenScheduler.
-	if (PCGComponent->IsPartitioned() && !PCGHelpers::IsRuntimeOrPIE() && PCGComponent->GenerationTrigger != EPCGComponentGenerationTrigger::GenerateAtRuntime)
+	if (PCGComponent->IsPartitioned() && !PCGHelpers::IsRuntimeOrPIE() && !PCGComponent->IsManagedByRuntimeGenSystem())
 	{
 		if (!GridSizes.IsEmpty())
 		{
@@ -680,7 +680,7 @@ void UPCGSubsystem::CancelAllGeneration()
 
 void UPCGSubsystem::RefreshRuntimeGenComponent(UPCGComponent* RuntimeComponent, bool bRemovePartitionActors)
 {
-	if (!ensure(RuntimeComponent && RuntimeComponent->GenerationTrigger == EPCGComponentGenerationTrigger::GenerateAtRuntime))
+	if (!ensure(RuntimeComponent && RuntimeComponent->IsManagedByRuntimeGenSystem()))
 	{
 		return;
 	}
@@ -971,7 +971,7 @@ void UPCGSubsystem::CreatePartitionActorsWithinBounds(const FBox& InBounds, cons
 
 FPCGTaskId UPCGSubsystem::ScheduleRefresh(UPCGComponent* Component, bool bForceRegen)
 {
-	check(Component && Component->GenerationTrigger != EPCGComponentGenerationTrigger::GenerateAtRuntime);
+	check(Component && !Component->IsManagedByRuntimeGenSystem());
 
 	TWeakObjectPtr<UPCGComponent> ComponentPtr(Component);
 
diff --git a/Engine/Plugins/Experimental/PCG/Source/PCG/Private/RuntimeGen/PCGRuntimeGenScheduler.cpp b/Engine/Plugins/Experimental/PCG/Source/PCG/Private/RuntimeGen/PCGRuntimeGenScheduler.cpp
--- a/Engine/Plugins/Experimental/PCG/Source/PCG/Private/RuntimeGen/PCGRuntimeGenScheduler.cpp
+++ b/Engine/Plugins/Experimental/PCG/Source/PCG/Private/RuntimeGen/PCGRuntimeGenScheduler.cpp
@@ -73,7 +73,7 @@ FPCGRuntimeGenScheduler::~FPCGRuntimeGenScheduler()
 	GenSourceManager = nullptr;
 }
 
-void FPCGRuntimeGenScheduler::Tick(const APCGWorldActor* InPCGWorldActor)
+void FPCGRuntimeGenScheduler::Tick(APCGWorldActor* InPCGWorldActor)
 {
 	TRACE_CPUPROFILER_EVENT_SCOPE(FPCGRuntimeGenScheduler::Tick);
 	check(InPCGWorldActor && GenSourceManager);
@@ -198,7 +198,7 @@ bool FPCGRuntimeGenScheduler::ShouldTick()
 
 void FPCGRuntimeGenScheduler::TickQueueComponentsForGeneration(
 	const TSet<IPCGGenSourceBase*>& InGenSources,
-	const APCGWorldActor* InPCGWorldActor,
+	APCGWorldActor* InPCGWorldActor,
 	TMap<FGridGenerationKey, double>& OutComponentsToGenerate)
 {
 	TRACE_CPUPROFILER_EVENT_SCOPE(FPCGRuntimeGenScheduler::TickQueueComponentsForGeneration);
@@ -246,7 +246,7 @@ void FPCGRuntimeGenScheduler::TickQueueComponentsForGeneration(
 		const UPCGSchedulingPolicyBase* Policy = OriginalComponent->GetRuntimeGenSchedulingPolicy();
 
 		// TODO: For each execution domain (for now only GenAtRuntime/dynamic), assuming we run Preview through this scheduler, which it seems like we will.
-		if (OriginalComponent->GenerationTrigger == EPCGComponentGenerationTrigger::GenerateAtRuntime)
+		if (OriginalComponent->IsManagedByRuntimeGenSystem())
 		{
 			bool bHasUnbounded = false;
 			PCGHiGenGrid::FSizeArray GridSizes;
@@ -388,7 +388,7 @@ void FPCGRuntimeGenScheduler::TickQueueComponentsForGeneration(
 		const UPCGSchedulingPolicyBase* Policy = OriginalComponent->GetRuntimeGenSchedulingPolicy();
 
 		// TODO: For each execution domain (for now only GenAtRuntime/dynamic), assuming we run Preview through this scheduler, which it seems like we will.
-		if (OriginalComponent->GenerationTrigger == EPCGComponentGenerationTrigger::GenerateAtRuntime)
+		if (OriginalComponent->IsManagedByRuntimeGenSystem())
 		{
 			// Unbounded will grab the base GenerationRadius used for non-partitioned and unbounded.
 			const double MaxGenerationRadius = OriginalComponent->GetGenerationRadiusFromGrid(EPCGHiGenGrid::Unbounded);
@@ -686,13 +686,14 @@ void FPCGRuntimeGenScheduler::TickCVars(const APCGWorldActor* InPCGWorldActor)
 
 void FPCGRuntimeGenScheduler::OnOriginalComponentRegistered(UPCGComponent* InOriginalComponent)
 {
-	// Ensure we are not a local component.
-	if (!InOriginalComponent || Cast<APCGPartitionActor>(InOriginalComponent->GetOwner()))
+	// Ensure we are non-local runtime managed component.
+	if (!InOriginalComponent || !InOriginalComponent->IsManagedByRuntimeGenSystem() || Cast<APCGPartitionActor>(InOriginalComponent->GetOwner()))
 	{
-		ensure(false);
 		return;
 	}
 
+	CreateGridGuidsForComponent(InOriginalComponent);
+
 	// When an original/non-partitioned component is registered, we need to dirty the state.
 	bAnyRuntimeGenComponentsExistDirty = true;
 }
@@ -920,6 +921,9 @@ void FPCGRuntimeGenScheduler::RefreshComponent(UPCGComponent* InComponent, bool
 	}
 	else
 	{
+		// Grid sizes might have changed, so we should defensively create grid guids.
+		CreateGridGuidsForComponent(OriginalComponent);
+
 		TArray<FGridGenerationKey> GenerationKeys;
 
 		for (FGridGenerationKey GenerationKey : GeneratedComponents)
@@ -1129,3 +1133,19 @@ void FPCGRuntimeGenScheduler::ResetPartitionActorPoolToSize(uint32 NewPoolSize)
 	PartitionActorPoolSize = 0;
 	AddPartitionActorPoolCount(NewPoolSize);
 }
+
+void FPCGRuntimeGenScheduler::CreateGridGuidsForComponent(UPCGComponent* InComponent)
+{
+	if (InComponent && InComponent->IsPartitioned() && InComponent->IsManagedByRuntimeGenSystem())
+	{
+		if (APCGWorldActor* PCGWorldActor = PCGHelpers::GetPCGWorldActor(World))
+		{
+			bool bHasUnbounded;
+			PCGHiGenGrid::FSizeArray GridSizes;
+			ensure(PCGHelpers::GetGenerationGridSizes(InComponent->GetGraph(), PCGWorldActor, GridSizes, bHasUnbounded));
+
+			// TODO: This will mark the package as dirty, which should not happen from runtime gen.
+			PCGWorldActor->CreateGridGuidsIfNecessary(GridSizes);
+		}
+	}
+}
diff --git a/Engine/Plugins/Experimental/PCG/Source/PCG/Public/PCGCommon.h b/Engine/Plugins/Experimental/PCG/Source/PCG/Public/PCGCommon.h
--- a/Engine/Plugins/Experimental/PCG/Source/PCG/Public/PCGCommon.h
+++ b/Engine/Plugins/Experimental/PCG/Source/PCG/Public/PCGCommon.h
@@ -19,7 +19,6 @@ enum class EPCGChangeType : uint8
 	Edge = 1 << 3,
 	Node = 1 << 4,
 	Structural = 1 << 5,
-	Debug = 1 << 6, // The Debug change type propagates when debug is toggled on a node.
 };
 ENUM_CLASS_FLAGS(EPCGChangeType);
 
diff --git a/Engine/Plugins/Experimental/PCG/Source/PCG/Public/PCGComponent.h b/Engine/Plugins/Experimental/PCG/Source/PCG/Public/PCGComponent.h
--- a/Engine/Plugins/Experimental/PCG/Source/PCG/Public/PCGComponent.h
+++ b/Engine/Plugins/Experimental/PCG/Source/PCG/Public/PCGComponent.h
@@ -310,6 +310,9 @@ public:
 	bool IsPartitioned() const;
 	bool IsLocalComponent() const { return bIsComponentLocal; }
 
+	/** Returns true if the component is managed by the runtime generation system. Nothing else should generate or cleanup this component. */
+	bool IsManagedByRuntimeGenSystem() const { return GenerationTrigger == EPCGComponentGenerationTrigger::GenerateAtRuntime; }
+
 	/* Responsibility of the PCG Partition Actor to mark is local */
 	void MarkAsLocalComponent() { bIsComponentLocal = true; }
 
@@ -350,13 +353,16 @@ protected:
 	UPROPERTY(Transient, VisibleAnywhere, Category = Debug)
 	uint32 GenerationGridSize = PCGHiGenGrid::UnboundedGridSize();
 
-	// Current editing mode that depends on the serialized editing mode and loading
-	UPROPERTY(Transient, EditAnywhere, Category = "Editing Settings", meta = (DisplayName = "Editing Mode", EditCondition = "!bIsComponentLocal", DisplayPriority = 300))
+	/** Current editing mode that depends on the serialized editing mode and loading. If the component is set to GenerateAtRuntime, this will behave as Preview. */
+	UPROPERTY(Transient, EditAnywhere, Category = "Editing Settings", meta = (DisplayName = "Editing Mode", EditCondition = "!bIsComponentLocal && GenerationTrigger != EPCGComponentGenerationTrigger::GenerateAtRuntime", DisplayPriority = 300))
 	EPCGEditorDirtyMode CurrentEditingMode = EPCGEditorDirtyMode::Normal;
 
 	UPROPERTY(VisibleAnywhere, Category = Debug, meta = (NoResetToDefault))
 	EPCGEditorDirtyMode SerializedEditingMode = EPCGEditorDirtyMode::Normal;
 
+	/** Used to store the CurrentEditingMode when it is forcefully changed by another system, such as runtime generation. */
+	EPCGEditorDirtyMode PreviousEditingMode = EPCGEditorDirtyMode::Normal;
+
 public:
 	UPROPERTY(BlueprintReadWrite, EditAnywhere, Category = "Input Node Settings (Deprecated)", meta = (DisplayPriority = 800))
 	EPCGComponentInput InputType = EPCGComponentInput::Actor;
diff --git a/Engine/Plugins/Experimental/PCG/Source/PCG/Public/RuntimeGen/PCGRuntimeGenScheduler.h b/Engine/Plugins/Experimental/PCG/Source/PCG/Public/RuntimeGen/PCGRuntimeGenScheduler.h
--- a/Engine/Plugins/Experimental/PCG/Source/PCG/Public/RuntimeGen/PCGRuntimeGenScheduler.h
+++ b/Engine/Plugins/Experimental/PCG/Source/PCG/Public/RuntimeGen/PCGRuntimeGenScheduler.h
@@ -36,7 +36,7 @@ public:
 	FPCGRuntimeGenScheduler& operator=(const FPCGRuntimeGenScheduler& other) = delete;
 	FPCGRuntimeGenScheduler& operator=(FPCGRuntimeGenScheduler&& other) = delete; 
 
-	void Tick(const APCGWorldActor* InPCGWorldActor);
+	void Tick(APCGWorldActor* InPCGWorldActor);
 
 	void OnOriginalComponentRegistered(UPCGComponent* InOriginalComponent);
 	void OnOriginalComponentUnregistered(UPCGComponent* InOriginalComponent);
@@ -58,7 +58,7 @@ protected:
 	/** Queue nearby components for generation. */
 	void TickQueueComponentsForGeneration(
 		const TSet<IPCGGenSourceBase*>& GenSources,
-		const APCGWorldActor* InPCGWorldActor,
+		APCGWorldActor* InPCGWorldActor,
 		TMap<FGridGenerationKey, double>& OutComponentsToGenerate);
 
 	/** Perform immediate cleanup on components that become out of range. */
@@ -93,6 +93,9 @@ protected:
 	/** Destroy all pooled partition actors and rebuild with the NewPoolSize. */
 	void ResetPartitionActorPoolToSize(uint32 NewPoolSize);
 
+	/** Create grid guids for the given component (if necessary). Only succeeds on partitioned original components. */
+	void CreateGridGuidsForComponent(UPCGComponent* InComponent);
+
 private:
 	/** Tracks the generated components managed by the RuntimeGenScheduler. For local components, this generation key will hold the original component.
 	* For non-partitioned components, the generation key should have unbounded grid size and (0, 0, 0) grid coordinates.
diff --git a/Engine/Plugins/Experimental/PCG/Source/PCGEditor/Private/Details/PCGComponentDetails.cpp b/Engine/Plugins/Experimental/PCG/Source/PCGEditor/Private/Details/PCGComponentDetails.cpp
--- a/Engine/Plugins/Experimental/PCG/Source/PCGEditor/Private/Details/PCGComponentDetails.cpp
+++ b/Engine/Plugins/Experimental/PCG/Source/PCGEditor/Private/Details/PCGComponentDetails.cpp
@@ -159,7 +159,7 @@ EVisibility FPCGComponentDetails::GenerateButtonVisible() const
 	for (const TWeakObjectPtr<UPCGComponent>& Component : SelectedComponents)
 	{
 		// If component is runtime generated then generate/cleanup is managed by the scheduler.
-		if (Component.IsValid() && !Component->IsGenerating() && Component->GenerationTrigger != EPCGComponentGenerationTrigger::GenerateAtRuntime)
+		if (Component.IsValid() && !Component->IsGenerating() && !Component->IsManagedByRuntimeGenSystem())
 		{
 			return EVisibility::Visible;
 		}
@@ -173,7 +173,7 @@ EVisibility FPCGComponentDetails::CancelButtonVisible() const
 	for (const TWeakObjectPtr<UPCGComponent>& Component : SelectedComponents)
 	{
 		// If component is runtime generated then generate/cleanup is managed by the scheduler.
-		if (Component.IsValid() && Component->IsGenerating() && Component->GenerationTrigger != EPCGComponentGenerationTrigger::GenerateAtRuntime)
+		if (Component.IsValid() && Component->IsGenerating() && !Component->IsManagedByRuntimeGenSystem())
 		{
 			return EVisibility::Visible;
 		}
@@ -187,7 +187,7 @@ EVisibility FPCGComponentDetails::CleanupButtonVisible() const
 	for (const TWeakObjectPtr<UPCGComponent>& Component : SelectedComponents)
 	{
 		// If component is runtime generated then generate/cleanup is managed by the scheduler.
-		if (Component.IsValid() && Component->GenerationTrigger != EPCGComponentGenerationTrigger::GenerateAtRuntime)
+		if (Component.IsValid() && !Component->IsManagedByRuntimeGenSystem())
 		{
 			return EVisibility::Visible;
 		}
@@ -200,7 +200,7 @@ EVisibility FPCGComponentDetails::RefreshButtonVisible() const
 {
 	for (const TWeakObjectPtr<UPCGComponent>& Component : SelectedComponents)
 	{
-		if (Component.IsValid() && Component->GenerationTrigger == EPCGComponentGenerationTrigger::GenerateAtRuntime)
+		if (Component.IsValid() && Component->IsManagedByRuntimeGenSystem())
 		{
 			return EVisibility::Visible;
 		}
diff --git a/Engine/Plugins/Experimental/PCG/Source/PCGEditor/Private/PCGEditor.cpp b/Engine/Plugins/Experimental/PCG/Source/PCGEditor/Private/PCGEditor.cpp
--- a/Engine/Plugins/Experimental/PCG/Source/PCGEditor/Private/PCGEditor.cpp
+++ b/Engine/Plugins/Experimental/PCG/Source/PCGEditor/Private/PCGEditor.cpp
@@ -278,7 +278,7 @@ void FPCGEditor::UpdateDebugAfterComponentSelection(UPCGComponent* InOldComponen
 		check(Component);
 
 		// GenerateAtRuntime components should be refreshed through the runtime gen scheduler.
-		if (Component->GenerationTrigger == EPCGComponentGenerationTrigger::GenerateAtRuntime)
+		if (Component->IsManagedByRuntimeGenSystem())
 		{
 			if (UPCGSubsystem* Subsystem = GetSubsystem())
 			{
@@ -1376,7 +1376,7 @@ void FPCGEditor::OnToggleDebug()
 			if (PCGSettingsInterface->bDebug != bNewCheckState)
 			{
 				PCGSettingsInterface->bDebug = bNewCheckState;
-				PCGNode->OnNodeChangedDelegate.Broadcast(PCGNode, EPCGChangeType::Debug);
+				PCGNode->OnNodeChangedDelegate.Broadcast(PCGNode, EPCGChangeType::Settings);
 			}
 		}
 	}
diff --git a/Engine/Plugins/Experimental/PCG/Source/PCGEditor/Private/SPCGEditorGraphDebugObjectTree.cpp b/Engine/Plugins/Experimental/PCG/Source/PCGEditor/Private/SPCGEditorGraphDebugObjectTree.cpp
--- a/Engine/Plugins/Experimental/PCG/Source/PCGEditor/Private/SPCGEditorGraphDebugObjectTree.cpp
+++ b/Engine/Plugins/Experimental/PCG/Source/PCGEditor/Private/SPCGEditorGraphDebugObjectTree.cpp
@@ -398,7 +398,7 @@ void SPCGEditorGraphDebugObjectTree::RefreshTree()
 
 	for (UObject* PCGComponentObject : PCGComponents)
 	{
-		if (!IsValid(PCGComponentObject) || PCGComponentObject->HasAnyFlags(RF_Transient))
+		if (!IsValid(PCGComponentObject))
 		{
 			continue;
 		}
