[PCG] Refactor and improve Flatten for Metadata

* Duplicate Data is now a base function on PCG Data to avoid copying metadata if not necessary
* Flatten is now a base function on PCG Data
* PointData has a special Flatten that also compress it (remove unused entry/values)
* Refactor duplicate in a few places

#rnx

diff --git a/Engine/Plugins/Experimental/PCG/Source/PCG/Private/Data/PCGPointData.cpp b/Engine/Plugins/Experimental/PCG/Source/PCG/Private/Data/PCGPointData.cpp
--- a/Engine/Plugins/Experimental/PCG/Source/PCG/Private/Data/PCGPointData.cpp
+++ b/Engine/Plugins/Experimental/PCG/Source/PCG/Private/Data/PCGPointData.cpp
@@ -668,3 +668,59 @@ UPCGSpatialData* UPCGPointData::CopyInternal() const
 
 	return NewPointData;
 }
+
+void UPCGPointData::Flatten()
+{
+	if (!Metadata)
+	{
+		return;
+	}
+
+	// If there is no more attributes, reset all keys from points to invalid
+	if (Metadata->GetAttributeCount() == 0)
+	{
+		bool bWasModified = false;
+		for (FPCGPoint& Point : Points)
+		{
+			if (Point.MetadataEntry != PCGInvalidEntryKey)
+			{
+				if (!bWasModified)
+				{
+					Modify();
+					bWasModified = true;
+				}
+
+				Point.MetadataEntry = PCGInvalidEntryKey;
+			}
+		}
+
+		return;
+	}
+
+	// Gather all the keys that are not invalid
+	TArray<PCGMetadataEntryKey> EntryKeys;
+	EntryKeys.Reserve(Points.Num());
+	for (const FPCGPoint& Point : Points)
+	{
+		if (Point.MetadataEntry != PCGInvalidEntryKey)
+		{
+			EntryKeys.Add(Point.MetadataEntry);
+		}
+	}
+
+	// Then flatten and compress the Metadata for all invalid entry keys. Return true if something changed.
+	if (Metadata->FlattenAndCompress(EntryKeys))
+	{
+		Modify();
+
+		// Go over all the points and assign all a new entry key for all points that has a valid entry key in the first place.
+		PCGMetadataEntryKey CurrentEntryKey = 0;
+		for (FPCGPoint& Point : Points)
+		{
+			if (Point.MetadataEntry != PCGInvalidEntryKey)
+			{
+				Point.MetadataEntry = CurrentEntryKey++;
+			}
+		}
+	}
+}
diff --git a/Engine/Plugins/Experimental/PCG/Source/PCG/Private/Elements/Metadata/PCGMetadataElementCommon.cpp b/Engine/Plugins/Experimental/PCG/Source/PCG/Private/Elements/Metadata/PCGMetadataElementCommon.cpp
--- a/Engine/Plugins/Experimental/PCG/Source/PCG/Private/Elements/Metadata/PCGMetadataElementCommon.cpp
+++ b/Engine/Plugins/Experimental/PCG/Source/PCG/Private/Elements/Metadata/PCGMetadataElementCommon.cpp
@@ -10,22 +10,12 @@ namespace PCGMetadataElementCommon
 	void DuplicateTaggedData(const FPCGTaggedData& InTaggedData, FPCGTaggedData& OutTaggedData, UPCGMetadata*& OutMetadata)
 	{
 		TRACE_CPUPROFILER_EVENT_SCOPE(PCGMetadataElementCommon::DuplicateTaggedData);
-
-		if (const UPCGSpatialData* SpatialInput = Cast<const UPCGSpatialData>(InTaggedData.Data))
-		{
-			UPCGSpatialData* NewSpatialData = SpatialInput->DuplicateData();
-			OutTaggedData.Data = NewSpatialData;
-
-			OutMetadata = NewSpatialData->Metadata;
-		}
-		else if (const UPCGParamData* ParamsInput = Cast<const UPCGParamData>(InTaggedData.Data))
+		if (InTaggedData.Data)
 		{
-			UPCGParamData* NewParamData = NewObject<UPCGParamData>();
-			NewParamData->Metadata->InitializeAsCopy(ParamsInput->Metadata);
-
-			OutTaggedData.Data = NewParamData;
-
-			OutMetadata = NewParamData->Metadata;
+			UPCGData* NewData = InTaggedData.Data->DuplicateData();
+			check(NewData);
+			OutTaggedData.Data = NewData;
+			OutMetadata = NewData->MutableMetadata();
 		}
 	}
 
diff --git a/Engine/Plugins/Experimental/PCG/Source/PCG/Private/Metadata/PCGMetadata.cpp b/Engine/Plugins/Experimental/PCG/Source/PCG/Private/Metadata/PCGMetadata.cpp
--- a/Engine/Plugins/Experimental/PCG/Source/PCG/Private/Metadata/PCGMetadata.cpp
+++ b/Engine/Plugins/Experimental/PCG/Source/PCG/Private/Metadata/PCGMetadata.cpp
@@ -275,6 +275,20 @@ bool UPCGMetadata::HasParent(const UPCGMetadata* InTentativeParent) const
 
 void UPCGMetadata::Flatten()
 {
+	// Check if we have a UPCGData owner, if so call it, otherwise just call FlattenImpl
+	if (UPCGData* Owner = Cast<UPCGData>(GetOuter()))
+	{
+		Owner->Flatten();
+	}
+	else
+	{
+		FlattenImpl();
+	}
+}
+
+void UPCGMetadata::FlattenImpl()
+{
+	TRACE_CPUPROFILER_EVENT_SCOPE(UPCGMetadata::FlattenImpl);
 	Modify();
 
 	const int32 NumEntries = GetItemCountForChild();
@@ -311,6 +325,36 @@ void UPCGMetadata::Flatten()
 	ItemKeyOffset = 0;
 }
 
+bool UPCGMetadata::FlattenAndCompress(const TArray<PCGMetadataEntryKey>& InEntryKeysToKeep)
+{
+	TRACE_CPUPROFILER_EVENT_SCOPE(UPCGMetadata::FlattenAndCompress);
+
+	// No keys or no parents, nothing to do
+	if (Attributes.IsEmpty())
+	{
+		return false;
+	}
+
+	Modify();
+
+	AttributeLock.WriteLock();
+	for (auto& AttributePair : Attributes)
+	{
+		FPCGMetadataAttributeBase* Attribute = AttributePair.Value;
+		check(Attribute);
+
+		Attribute->FlattenAndCompress(InEntryKeysToKeep);
+	}
+	AttributeLock.WriteUnlock();
+
+	Parent = nullptr;
+	ParentKeys.Reset();
+	ParentKeys.Init(PCGInvalidEntryKey, InEntryKeysToKeep.Num());
+	ItemKeyOffset = 0;
+
+	return true;
+}
+
 void UPCGMetadata::AddAttributeInternal(FName AttributeName, FPCGMetadataAttributeBase* Attribute)
 {
 	// This call assumes we have a write lock on the attribute map.
diff --git a/Engine/Plugins/Experimental/PCG/Source/PCG/Private/PCGComponent.cpp b/Engine/Plugins/Experimental/PCG/Source/PCG/Private/PCGComponent.cpp
--- a/Engine/Plugins/Experimental/PCG/Source/PCG/Private/PCGComponent.cpp
+++ b/Engine/Plugins/Experimental/PCG/Source/PCG/Private/PCGComponent.cpp
@@ -408,25 +408,14 @@ void UPCGComponent::PostProcessGraph(const FBox& InNewBounds, bool bInGenerated,
 			//GeneratedGraphOutput = Context->InputData;
 			for (const FPCGTaggedData& TaggedData : Context->InputData.TaggedData)
 			{
-				FPCGTaggedData& DuplicatedTaggedData = GeneratedGraphOutput.TaggedData.Add_GetRef(TaggedData);
 				// TODO: outering the first layer might not be sufficient here - might need to expose
 				// some methods in the data to traverse all the data to outer everything for serialization
-				DuplicatedTaggedData.Data = Cast<UPCGData>(StaticDuplicateObject(TaggedData.Data, this));
-
-				UPCGMetadata* DuplicatedMetadata = nullptr;
-				if (const UPCGSpatialData* DuplicatedSpatialData = Cast<UPCGSpatialData>(DuplicatedTaggedData.Data))
-				{
-					DuplicatedMetadata = DuplicatedSpatialData->Metadata;
-				}
-				else if (const UPCGParamData* DuplicatedParamData = Cast<UPCGParamData>(DuplicatedTaggedData.Data))
-				{
-					DuplicatedMetadata = DuplicatedParamData->Metadata;
-				}
-
-				// Make sure the metadata can be serialized independently
-				if (DuplicatedMetadata)
+				if (UPCGData* DuplicatedData = TaggedData.Data->DuplicateData())
 				{
-					DuplicatedMetadata->Flatten();
+					FPCGTaggedData& DuplicatedTaggedData = GeneratedGraphOutput.TaggedData.Add_GetRef(TaggedData);
+					DuplicatedTaggedData.Data = DuplicatedData;
+					DuplicatedData->Rename(nullptr, this);
+					DuplicatedData->Flatten();
 				}
 			}
 
diff --git a/Engine/Plugins/Experimental/PCG/Source/PCG/Private/PCGData.cpp b/Engine/Plugins/Experimental/PCG/Source/PCG/Private/PCGData.cpp
--- a/Engine/Plugins/Experimental/PCG/Source/PCG/Private/PCGData.cpp
+++ b/Engine/Plugins/Experimental/PCG/Source/PCG/Private/PCGData.cpp
@@ -138,6 +138,19 @@ void UPCGData::InitUID()
 	UID = ++UIDCounter;
 }
 
+void UPCGData::Flatten()
+{
+	if (UPCGMetadata* Metadata = MutableMetadata())
+	{
+		Metadata->FlattenImpl();
+	}
+}
+
+UPCGData* UPCGData::DuplicateData() const
+{
+	return Cast<UPCGData>(StaticDuplicateObject(this, GetTransientPackage()));
+}
+
 bool FPCGTaggedData::operator==(const FPCGTaggedData& Other) const
 {
 	return Data == Other.Data &&
diff --git a/Engine/Plugins/Experimental/PCG/Source/PCG/Private/PCGParamData.cpp b/Engine/Plugins/Experimental/PCG/Source/PCG/Private/PCGParamData.cpp
--- a/Engine/Plugins/Experimental/PCG/Source/PCG/Private/PCGParamData.cpp
+++ b/Engine/Plugins/Experimental/PCG/Source/PCG/Private/PCGParamData.cpp
@@ -161,4 +161,12 @@ void UPCGParamData::SetLastSelector(const FPCGAttributePropertySelector& InSelec
 
 	bHasCachedLastSelector = true;
 	CachedLastSelector.ImportFromOtherSelector(InSelector);
+}
+
+UPCGParamData* UPCGParamData::DuplicateData() const
+{
+	UPCGParamData* NewParamData = NewObject<UPCGParamData>();
+	NewParamData->Metadata->InitializeAsCopy(Metadata);
+
+	return NewParamData;
 }
\ No newline at end of file
diff --git a/Engine/Plugins/Experimental/PCG/Source/PCG/Private/Tests/Metadata/PCGMetadataFlattenTest.cpp b/Engine/Plugins/Experimental/PCG/Source/PCG/Private/Tests/Metadata/PCGMetadataFlattenTest.cpp
new file mode 100644
--- /dev/null
+++ b/Engine/Plugins/Experimental/PCG/Source/PCG/Private/Tests/Metadata/PCGMetadataFlattenTest.cpp
@@ -0,0 +1,152 @@
+// Copyright Epic Games, Inc. All Rights Reserved.
+
+#if WITH_EDITOR
+
+#include "Tests/PCGTestsCommon.h"
+
+#include "Data/PCGPointData.h"
+#include "Metadata/PCGMetadata.h"
+#include "Metadata/PCGMetadataAttributeTpl.h"
+
+IMPLEMENT_CUSTOM_SIMPLE_AUTOMATION_TEST(FPCGMetadataFlatten, FPCGTestBaseClass, "pcg.tests.Metadata.Flatten", PCGTestsCommon::TestFlags)
+
+/**
+* Series of operations to manipulate points and attributes, to validate the flatten operation.
+*/
+bool FPCGMetadataFlatten::RunTest(const FString& Parameters)
+{
+	static const FName Attribute1Name = TEXT("FloatAttr");
+	static const FName Attribute2Name = TEXT("StringAttr");
+	static const FName Attribute3Name = TEXT("IntAttr");
+
+	UPCGPointData* RootPointData = NewObject<UPCGPointData>();
+	TArray<FPCGPoint>& Points = RootPointData->GetMutablePoints();
+	FPCGMetadataAttribute<float>* Attribute1 = RootPointData->Metadata->CreateAttribute<float>(Attribute1Name, -0.1f, true, true);
+	FPCGMetadataAttribute<FString>* Attribute2 = RootPointData->Metadata->CreateAttribute<FString>(Attribute2Name, TEXT("Default"), true, true);
+
+	check(Attribute1 && Attribute2);
+	
+	for (int32 i = 0; i < 10; ++i)
+	{
+		FPCGPoint& Point = Points.Emplace_GetRef();
+		RootPointData->Metadata->InitializeOnSet(Point.MetadataEntry);
+		Attribute1->SetValue(Point.MetadataEntry, static_cast<float>(i) * 0.1f);
+
+		if (i % 2 == 0)
+		{
+			// Will be either "0" or "1"
+			Attribute2->SetValue(Point.MetadataEntry, FString::Printf(TEXT("%d"), i % 4));
+		}
+	}
+
+	// At the end of the first set, metadata has 10 entries, and values for each points are
+	// for Attribute 1: [0.0, 0.1, 0.2, 0.3, 0.4, 0.5, 0.6, 0.7, 0.8, 0.9]
+	// for Attribute 2: ["0", "Default", "2", "Default", "0", ...]
+	UTEST_EQUAL("RootMetadata has 10 entries", RootPointData->Metadata->GetItemCountForChild(), 10);
+
+	UPCGPointData* FirstChildPointData = Cast<UPCGPointData>(RootPointData->DuplicateData());
+	Attribute1 = FirstChildPointData->Metadata->GetMutableTypedAttribute<float>(Attribute1Name);
+	Attribute2 = FirstChildPointData->Metadata->GetMutableTypedAttribute<FString>(Attribute2Name);
+
+	UTEST_TRUE("Attributes exists in first child", Attribute1 && Attribute2);
+	check(Attribute1 && Attribute2);
+
+	// Override all the values for Attribute 1, and replace all "1" by "0" in attribute 2
+	// Also add a third attribute
+	FPCGMetadataAttribute<int>* Attribute3 = FirstChildPointData->Metadata->CreateAttribute<int>(Attribute3Name, -1, true, true);
+	check(Attribute3);
+
+	TArray<FPCGPoint>& FirstChildPoints = FirstChildPointData->GetMutablePoints();
+	
+	UTEST_EQUAL("First child has 10 points", FirstChildPoints.Num(), 10);
+
+	for (int32 i = 0; i < 10; ++i)
+	{
+		FirstChildPointData->Metadata->InitializeOnSet(FirstChildPoints[i].MetadataEntry);
+		Attribute1->SetValue(FirstChildPoints[i].MetadataEntry, static_cast<float>(i) * 1.1f);
+
+		if (i % 2 == 0)
+		{
+			Attribute2->SetValue(FirstChildPoints[i].MetadataEntry, TEXT("0"));
+		}
+
+		Attribute3->SetValue(FirstChildPoints[i].MetadataEntry, i);
+	}
+
+	// At the end of the second set, metadata has 20 entries, and values for each points are
+	// for Attribute 1: [0.0, 1.1, 2.2, 3.3, 4.4, 5.5, 6.6, 7.7, 8.8, 9.9]
+	// for Attribute 2: ["0", "Default", "0", "Default", "0", ...]
+	// for Attribute 3: [0, 1, 2, 3, 4, 5, 6, 7, 8, 9]
+	UTEST_EQUAL("FirstChildMetadata has 20 entries", FirstChildPointData->Metadata->GetItemCountForChild(), 20);
+
+	// Second child is keeping the metadata entry for even numbers and reset metadata entry for the rest (point back to default). Also override Attribute 3 values for even numbers
+	UPCGPointData* SecondChildPointData = NewObject<UPCGPointData>();
+	SecondChildPointData->InitializeFromData(FirstChildPointData);
+	TArray<FPCGPoint>& SecondChildPoints = SecondChildPointData->GetMutablePoints();
+
+	Attribute3 = SecondChildPointData->Metadata->GetMutableTypedAttribute<int>(Attribute3Name);
+	UTEST_NOT_NULL("Attributes exists in first child", Attribute3);
+	check(Attribute3);
+
+	for (int32 i = 0; i < 10; ++i)
+	{
+		FPCGPoint& Point = SecondChildPoints.Add_GetRef(FirstChildPoints[i]);
+		if (i % 2 == 0)
+		{
+			SecondChildPointData->Metadata->InitializeOnSet(Point.MetadataEntry);
+			Attribute3->SetValue(Point.MetadataEntry, 10 * i);
+		}
+		else
+		{
+			Point.MetadataEntry = PCGInvalidEntryKey;
+		}
+	}
+
+	// At the end of the third set, metadata has 25 entries, and values for each points are:
+	// for Attribute 1: [0.0, -0.1, 2.2, -0.1, 4.4, -0.1, 6.6, -0.1, 8.8, -0.1]
+	// for Attribute 2: ["0", "Default", "0", "Default", "0", ...]
+	// for Attribute 3: [0, -1, 20, -1, 40, -1, 60, -1, 80, -1]
+	UTEST_EQUAL("SecondChildMetadata has 25 entries", SecondChildPointData->Metadata->GetItemCountForChild(), 25);
+
+	// For final set, duplicate the data and flatten it
+	UPCGPointData* FinalPointData = Cast<UPCGPointData>(SecondChildPointData->DuplicateData());
+	FinalPointData->Flatten();
+
+	Attribute1 = FinalPointData->Metadata->GetMutableTypedAttribute<float>(Attribute1Name);
+	Attribute2 = FinalPointData->Metadata->GetMutableTypedAttribute<FString>(Attribute2Name);
+	Attribute3 = FinalPointData->Metadata->GetMutableTypedAttribute<int>(Attribute3Name);
+
+	UTEST_TRUE("Attributes exists in final child", Attribute1 && Attribute2 && Attribute3);
+	check(Attribute1 && Attribute2 && Attribute3);
+
+	UTEST_EQUAL("Final metadata has 5 entries", FinalPointData->Metadata->GetItemCountForChild(), 5);
+	UTEST_EQUAL("Attribute 1 has 5 values", Attribute1->GetValueKeyOffsetForChild(), 5);
+	UTEST_EQUAL("Attribute 2 has 1 value", Attribute2->GetValueKeyOffsetForChild(), 1);
+	UTEST_EQUAL("Attribute 3 has 5 values", Attribute3->GetValueKeyOffsetForChild(), 5);
+	UTEST_EQUAL("Final point data has 10 points", FinalPointData->GetPoints().Num(), 10);
+
+	// Validate the values
+	for (int32 i = 0; i < 10; ++i)
+	{
+		auto FormatMessage = [i](const FString& Msg) { return *FString::Printf(TEXT("Point %d: %s"), i, *Msg); };
+
+		const FPCGPoint& Point = FinalPointData->GetPoints()[i];
+		if (i % 2 == 0)
+		{
+			UTEST_EQUAL(FormatMessage(TEXT("Valid metadata entry")), Point.MetadataEntry, static_cast<PCGMetadataEntryKey>(i / 2));
+			UTEST_EQUAL(FormatMessage(TEXT("Attribute 1 value")), Attribute1->GetValueFromItemKey(Point.MetadataEntry), i * 1.1f);
+			UTEST_EQUAL(FormatMessage(TEXT("Attribute 2 value")), Attribute2->GetValueFromItemKey(Point.MetadataEntry), TEXT("0"));
+			UTEST_EQUAL(FormatMessage(TEXT("Attribute 3 value")), Attribute3->GetValueFromItemKey(Point.MetadataEntry), i * 10);
+		}
+		else
+		{
+			UTEST_EQUAL(FormatMessage(TEXT("Invalid metadata entry")), Point.MetadataEntry, PCGInvalidEntryKey);
+			UTEST_EQUAL(FormatMessage(TEXT("Attribute 1 value")), Attribute1->GetValueFromItemKey(Point.MetadataEntry), -0.1f);
+			UTEST_EQUAL(FormatMessage(TEXT("Attribute 2 value")), Attribute2->GetValueFromItemKey(Point.MetadataEntry), TEXT("Default"));
+			UTEST_EQUAL(FormatMessage(TEXT("Attribute 3 value")), Attribute3->GetValueFromItemKey(Point.MetadataEntry), -1);
+		}
+	}
+
+	return true;
+}
+#endif // WITH_EDITOR
\ No newline at end of file
diff --git a/Engine/Plugins/Experimental/PCG/Source/PCG/Public/Data/PCGPointData.h b/Engine/Plugins/Experimental/PCG/Source/PCG/Public/Data/PCGPointData.h
--- a/Engine/Plugins/Experimental/PCG/Source/PCG/Public/Data/PCGPointData.h
+++ b/Engine/Plugins/Experimental/PCG/Source/PCG/Public/Data/PCGPointData.h
@@ -113,6 +113,9 @@ public:
 	UFUNCTION(BlueprintCallable, Category = SpatialData)
 	void CopyPointsFrom(const UPCGPointData* InData, const TArray<int>& InDataIndices);
 
+	/** Make a pass on Metadata to flatten parenting and only keep entries used by points. */
+	virtual void Flatten() override;
+
 	TArray<FPCGPoint>& GetMutablePoints();
 
 	const PointOctree& GetOctree() const;
diff --git a/Engine/Plugins/Experimental/PCG/Source/PCG/Public/Data/PCGSpatialData.h b/Engine/Plugins/Experimental/PCG/Source/PCG/Public/Data/PCGSpatialData.h
--- a/Engine/Plugins/Experimental/PCG/Source/PCG/Public/Data/PCGSpatialData.h
+++ b/Engine/Plugins/Experimental/PCG/Source/PCG/Public/Data/PCGSpatialData.h
@@ -46,6 +46,8 @@ public:
 	virtual bool HasCachedLastSelector() const override;
 	virtual FPCGAttributePropertyInputSelector GetCachedLastSelector() const override;
 	virtual void SetLastSelector(const FPCGAttributePropertySelector& InSelector) override;
+
+	virtual UPCGSpatialData* DuplicateData() const override { return DuplicateData(/*bInitializeMetadata=*/ true); }
 	// ~End UPCGData interface
 
 	/** Virtual call to allocate a new spacial data object, duplicate this spatial data into
@@ -56,7 +58,7 @@ public:
 	*   They are mainly cached values (and octree for points).
 	*   TODO: If we want to also copy those values (can be an optimization), we need to guard the copy.
 	*/
-	UPCGSpatialData* DuplicateData(const bool bInitializeMetadata = true) const;
+	UPCGSpatialData* DuplicateData(const bool bInitializeMetadata) const;
 
 	/** Returns the dimension of the data type, which has nothing to do with the dimension of its points */
 	UFUNCTION(BlueprintCallable, Category = SpatialData)
diff --git a/Engine/Plugins/Experimental/PCG/Source/PCG/Public/Metadata/PCGMetadata.h b/Engine/Plugins/Experimental/PCG/Source/PCG/Public/Metadata/PCGMetadata.h
--- a/Engine/Plugins/Experimental/PCG/Source/PCG/Public/Metadata/PCGMetadata.h
+++ b/Engine/Plugins/Experimental/PCG/Source/PCG/Public/Metadata/PCGMetadata.h
@@ -81,10 +81,16 @@ public:
 	const UPCGMetadata* GetRoot() const;
 	bool HasParent(const UPCGMetadata* InTentativeParent) const;
 
-	/** Unparents current metadata by flattening the attributes (values, entries, etc.) */
+	/** Unparents current metadata by flattening the attributes (values, entries, etc.) and potentially compress the data to remove unused values. */
 	UFUNCTION(BlueprintCallable, Category = "PCG|Metadata|Advanced")
 	void Flatten();
 
+	/** Unparents current metadata by flattening the attributes (values, entries, etc.) */
+	void FlattenImpl();
+
+	/** Unparents current metadata, flatten attribute and only keep the entries specified. Return true if something has changed and keys needs be updated. */
+	bool FlattenAndCompress(const TArray<PCGMetadataEntryKey>& InEntryKeysToKeep);
+
 	UFUNCTION(BlueprintCallable, Category = "PCG|Metadata")
 	UPARAM(DisplayName = "Metadata") UPCGMetadata* CreateInteger32Attribute(FName AttributeName, int32 DefaultValue, bool bAllowsInterpolation, bool bOverrideParent = true);
 
diff --git a/Engine/Plugins/Experimental/PCG/Source/PCG/Public/Metadata/PCGMetadataAttribute.h b/Engine/Plugins/Experimental/PCG/Source/PCG/Public/Metadata/PCGMetadataAttribute.h
--- a/Engine/Plugins/Experimental/PCG/Source/PCG/Public/Metadata/PCGMetadataAttribute.h
+++ b/Engine/Plugins/Experimental/PCG/Source/PCG/Public/Metadata/PCGMetadataAttribute.h
@@ -25,7 +25,13 @@ public:
 	virtual ~FPCGMetadataAttributeBase() = default;
 	virtual void Serialize(UPCGMetadata* InMetadata, FArchive& InArchive);
 
+	/** Unparents current attribute by flattening the values, entries, etc. */
 	virtual void Flatten() = 0;
+	/** Unparents current attribute by flattening the values, entries, etc while only keeping the entries referenced in InEntryKeysToKeep. There must be NO invalid entry keys. */
+	virtual void FlattenAndCompress(const TArray<PCGMetadataEntryKey>& InEntryKeysToKeep) = 0;
+
+	/** Remove all entries, values and parenting. */
+	virtual void Reset() = 0;
 
 	const UPCGMetadata* GetMetadata() const { return Metadata; }
 	int16 GetTypeId() const { return TypeId; }
diff --git a/Engine/Plugins/Experimental/PCG/Source/PCG/Public/Metadata/PCGMetadataAttributeTpl.h b/Engine/Plugins/Experimental/PCG/Source/PCG/Public/Metadata/PCGMetadataAttributeTpl.h
--- a/Engine/Plugins/Experimental/PCG/Source/PCG/Public/Metadata/PCGMetadataAttributeTpl.h
+++ b/Engine/Plugins/Experimental/PCG/Source/PCG/Public/Metadata/PCGMetadataAttributeTpl.h
@@ -58,7 +58,7 @@ public:
 		// We don't need to flatten the EntryToValueKeyMap - this will have been taken care of in the metadata flatten
 
 		// Flatten values, from root to current attribute
-		if(Parent)
+		if (Parent)
 		{
 			FWriteScopeLock ScopeLock(ValueLock);
 
@@ -89,6 +89,99 @@ public:
 		Parent = nullptr;
 	}
 
+	virtual void FlattenAndCompress(const TArray<PCGMetadataEntryKey>& InEntryKeysToKeep) override
+	{
+		// No entries, we can just delete everything in the attribute.
+		if (InEntryKeysToKeep.IsEmpty())
+		{
+			Reset();
+			return;
+		}
+
+		TArray<PCGMetadataValueKey> AllValueKeys;
+		TArray<PCGMetadataValueKey> AllUniqueValueKeys;
+		AllValueKeys.Reserve(InEntryKeysToKeep.Num());
+		constexpr bool bUseValueKeys = PCG::Private::MetadataTraits<T>::CompressData;
+		if constexpr (bUseValueKeys)
+		{
+			AllUniqueValueKeys.Reserve(InEntryKeysToKeep.Num());
+		}
+
+		// First gather all value keys associated with the entry keys to keep.
+		// If we compress data, we also store the unique value keys used (that is not default).
+		for (const PCGMetadataEntryKey& EntryKey : InEntryKeysToKeep)
+		{
+			AllValueKeys.Add(GetValueKey(EntryKey));
+			if constexpr (bUseValueKeys)
+			{
+				if (AllValueKeys.Last() != PCGDefaultValueKey)
+				{
+					AllUniqueValueKeys.AddUnique(AllValueKeys.Last());
+				}
+			}
+		}
+
+		// Then for each value key (or unique values keys), gather the value in a NewValues array
+		// and also keep a mapping between old value key and new value key.
+		// Only done if the old value key is not the default one.
+		TMap<PCGMetadataValueKey, PCGMetadataValueKey> ValueKeyMapping;
+		TArray<T> NewValues;
+		const TArray<PCGMetadataValueKey>& AllValueKeysRef = !bUseValueKeys ? AllValueKeys : AllUniqueValueKeys;
+		NewValues.Reserve(AllValueKeysRef.Num());
+		ValueKeyMapping.Reserve(AllValueKeysRef.Num());
+		for (PCGMetadataValueKey ValueKey : AllValueKeysRef)
+		{
+			if (ValueKey != PCGDefaultValueKey)
+			{
+				ValueKeyMapping.Add(ValueKey, NewValues.Num());
+				NewValues.Add(GetValue(ValueKey));
+			}
+		}
+
+		// Move the new values in place of the old values.
+		ValueLock.WriteLock();
+		Values = std::move(NewValues);
+		ValueLock.WriteUnlock();
+
+		// And finally, create a new entry to value mapping.
+		// Logic is that each entry to keep will have their "index" as new entry key
+		// (like if the entries to keep are [25, 47, 54], the new entries would be [0, 1, 2]).
+		// So the operation is:
+		// All pairs Old EK -> Old VK transform to New EK -> New VK.
+		TMap<PCGMetadataEntryKey, PCGMetadataValueKey> NewMap;
+		for (int32 i = 0; i < InEntryKeysToKeep.Num(); ++i)
+		{
+			PCGMetadataValueKey ValueKey = AllValueKeys[i];
+			if (ValueKey != PCGDefaultValueKey && ensure(InEntryKeysToKeep[i] != PCGInvalidEntryKey))
+			{
+				NewMap.Add(i, ValueKeyMapping[ValueKey]);
+			}
+		}
+
+		// And move the map.
+		EntryMapLock.WriteLock();
+		EntryToValueKeyMap = std::move(NewMap);
+		EntryMapLock.WriteUnlock();
+
+		// At the end, reset value offset, and lose parent.
+		ValueKeyOffset = 0;
+		Parent = nullptr;
+	}
+
+	virtual void Reset() override
+	{
+		ValueKeyOffset = 0;
+		Parent = nullptr;
+
+		EntryMapLock.WriteLock();
+		EntryToValueKeyMap.Empty();
+		EntryMapLock.WriteUnlock();
+
+		ValueLock.WriteLock();
+		Values.Empty();
+		ValueLock.WriteUnlock();
+	}
+
 	const FPCGMetadataAttribute* GetParent() const { return static_cast<const FPCGMetadataAttribute*>(Parent); }
 
 	FPCGMetadataAttribute* TypedCopy(FName NewName, UPCGMetadata* InMetadata, bool bKeepParent, bool bCopyEntries = true, bool bCopyValues = true)
diff --git a/Engine/Plugins/Experimental/PCG/Source/PCG/Public/PCGData.h b/Engine/Plugins/Experimental/PCG/Source/PCG/Public/PCGData.h
--- a/Engine/Plugins/Experimental/PCG/Source/PCG/Public/PCGData.h
+++ b/Engine/Plugins/Experimental/PCG/Source/PCG/Public/PCGData.h
@@ -50,13 +50,17 @@ public:
 	virtual FPCGAttributePropertyInputSelector GetCachedLastSelector() const { return FPCGAttributePropertyInputSelector{}; }
 	virtual void SetLastSelector(const FPCGAttributePropertySelector& InSelector) {};
 
+	/** Return a copy of the data, with Metadata inheritence for spatial data. */
+	virtual UPCGData* DuplicateData() const;
+
 	// ~Begin UObject interface
 	virtual void PostDuplicate(bool bDuplicateForPIE) override { InitUID(); }
 	// ~End UObject interface
 
-	// Metadata getters, to be implemented if data supports Metadata
+	// Metadata ops, to be implemented if data supports Metadata
 	virtual UPCGMetadata* MutableMetadata() { return nullptr; }
 	virtual const UPCGMetadata* ConstMetadata() const { return nullptr; }
+	virtual void Flatten();
 
 protected:
 	/** Computes Crc for this and any connected data. */
diff --git a/Engine/Plugins/Experimental/PCG/Source/PCG/Public/PCGParamData.h b/Engine/Plugins/Experimental/PCG/Source/PCG/Public/PCGParamData.h
--- a/Engine/Plugins/Experimental/PCG/Source/PCG/Public/PCGParamData.h
+++ b/Engine/Plugins/Experimental/PCG/Source/PCG/Public/PCGParamData.h
@@ -25,6 +25,8 @@ public:
 	virtual bool HasCachedLastSelector() const override;
 	virtual FPCGAttributePropertyInputSelector GetCachedLastSelector() const override;
 	virtual void SetLastSelector(const FPCGAttributePropertySelector& InSelector) override;
+
+	virtual UPCGParamData* DuplicateData() const override;
 	// ~End UPCGData interface
 
 	UFUNCTION(BlueprintCallable, Category = Metadata)
