[PCG] Fix compilation issue for test on MacArm64

#rnx

diff --git a/Engine/Plugins/Experimental/PCG/Source/PCG/Private/Tests/Metadata/PCGMetadataFlattenTest.cpp b/Engine/Plugins/Experimental/PCG/Source/PCG/Private/Tests/Metadata/PCGMetadataFlattenTest.cpp
--- a/Engine/Plugins/Experimental/PCG/Source/PCG/Private/Tests/Metadata/PCGMetadataFlattenTest.cpp
+++ b/Engine/Plugins/Experimental/PCG/Source/PCG/Private/Tests/Metadata/PCGMetadataFlattenTest.cpp
@@ -128,22 +128,22 @@ bool FPCGMetadataFlatten::RunTest(const FString& Parameters)
 	// Validate the values
 	for (int32 i = 0; i < 10; ++i)
 	{
-		auto FormatMessage = [i](const FString& Msg) { return *FString::Printf(TEXT("Point %d: %s"), i, *Msg); };
+		auto FormatMessage = [i](const FString& Msg) { return FString::Printf(TEXT("Point %d: %s"), i, *Msg); };
 
 		const FPCGPoint& Point = FinalPointData->GetPoints()[i];
 		if (i % 2 == 0)
 		{
-			UTEST_EQUAL(FormatMessage(TEXT("Valid metadata entry")), Point.MetadataEntry, static_cast<PCGMetadataEntryKey>(i / 2));
-			UTEST_EQUAL(FormatMessage(TEXT("Attribute 1 value")), Attribute1->GetValueFromItemKey(Point.MetadataEntry), i * 1.1f);
-			UTEST_EQUAL(FormatMessage(TEXT("Attribute 2 value")), Attribute2->GetValueFromItemKey(Point.MetadataEntry), TEXT("0"));
-			UTEST_EQUAL(FormatMessage(TEXT("Attribute 3 value")), Attribute3->GetValueFromItemKey(Point.MetadataEntry), i * 10);
+			UTEST_EQUAL(*FormatMessage(TEXT("Valid metadata entry")), Point.MetadataEntry, static_cast<PCGMetadataEntryKey>(i / 2));
+			UTEST_EQUAL(*FormatMessage(TEXT("Attribute 1 value")), Attribute1->GetValueFromItemKey(Point.MetadataEntry), i * 1.1f);
+			UTEST_EQUAL(*FormatMessage(TEXT("Attribute 2 value")), Attribute2->GetValueFromItemKey(Point.MetadataEntry), TEXT("0"));
+			UTEST_EQUAL(*FormatMessage(TEXT("Attribute 3 value")), Attribute3->GetValueFromItemKey(Point.MetadataEntry), i * 10);
 		}
 		else
 		{
-			UTEST_EQUAL(FormatMessage(TEXT("Invalid metadata entry")), Point.MetadataEntry, PCGInvalidEntryKey);
-			UTEST_EQUAL(FormatMessage(TEXT("Attribute 1 value")), Attribute1->GetValueFromItemKey(Point.MetadataEntry), -0.1f);
-			UTEST_EQUAL(FormatMessage(TEXT("Attribute 2 value")), Attribute2->GetValueFromItemKey(Point.MetadataEntry), TEXT("Default"));
-			UTEST_EQUAL(FormatMessage(TEXT("Attribute 3 value")), Attribute3->GetValueFromItemKey(Point.MetadataEntry), -1);
+			UTEST_EQUAL(*FormatMessage(TEXT("Invalid metadata entry")), Point.MetadataEntry, PCGInvalidEntryKey);
+			UTEST_EQUAL(*FormatMessage(TEXT("Attribute 1 value")), Attribute1->GetValueFromItemKey(Point.MetadataEntry), -0.1f);
+			UTEST_EQUAL(*FormatMessage(TEXT("Attribute 2 value")), Attribute2->GetValueFromItemKey(Point.MetadataEntry), TEXT("Default"));
+			UTEST_EQUAL(*FormatMessage(TEXT("Attribute 3 value")), Attribute3->GetValueFromItemKey(Point.MetadataEntry), -1);
 		}
 	}
 
