Cleaned up some object renames to make it clear it's only in editor and potentially prevent rename issues in multithread

#rnx

diff --git a/Engine/Plugins/Experimental/PCG/Source/PCG/Private/Elements/PCGPointMatchAndSet.cpp b/Engine/Plugins/Experimental/PCG/Source/PCG/Private/Elements/PCGPointMatchAndSet.cpp
--- a/Engine/Plugins/Experimental/PCG/Source/PCG/Private/Elements/PCGPointMatchAndSet.cpp
+++ b/Engine/Plugins/Experimental/PCG/Source/PCG/Private/Elements/PCGPointMatchAndSet.cpp
@@ -164,7 +164,9 @@ void UPCGPointMatchAndSetSettings::RefreshMatchAndSet()
 		// Forget previous instance
 		if (MatchAndSetInstance)
 		{
+#if WITH_EDITOR
 			MatchAndSetInstance->Rename(nullptr, GetTransientPackage(), REN_DontCreateRedirectors | REN_ForceNoResetLoaders);
+#endif
 			MatchAndSetInstance->MarkAsGarbage();
 			MatchAndSetInstance = nullptr;
 		}
diff --git a/Engine/Plugins/Experimental/PCG/Source/PCG/Private/Elements/PCGStaticMeshSpawner.cpp b/Engine/Plugins/Experimental/PCG/Source/PCG/Private/Elements/PCGStaticMeshSpawner.cpp
--- a/Engine/Plugins/Experimental/PCG/Source/PCG/Private/Elements/PCGStaticMeshSpawner.cpp
+++ b/Engine/Plugins/Experimental/PCG/Source/PCG/Private/Elements/PCGStaticMeshSpawner.cpp
@@ -509,7 +509,9 @@ void UPCGStaticMeshSpawnerSettings::RefreshMeshSelector()
 	{
 		if (MeshSelectorParameters)
 		{
+#if WITH_EDITOR
 			MeshSelectorParameters->Rename(nullptr, GetTransientPackage(), REN_DontCreateRedirectors | REN_ForceNoResetLoaders);
+#endif
 			MeshSelectorParameters->MarkAsGarbage();
 			MeshSelectorParameters = nullptr;
 		}
@@ -529,7 +531,9 @@ void UPCGStaticMeshSpawnerSettings::RefreshInstancePacker()
 	{
 		if (InstanceDataPackerParameters)
 		{
+#if WITH_EDITOR
 			InstanceDataPackerParameters->Rename(nullptr, GetTransientPackage(), REN_DontCreateRedirectors | REN_ForceNoResetLoaders);
+#endif
 			InstanceDataPackerParameters->MarkAsGarbage();
 			InstanceDataPackerParameters = nullptr;
 		}
diff --git a/Engine/Plugins/Experimental/PCG/Source/PCG/Private/PCGComponent.cpp b/Engine/Plugins/Experimental/PCG/Source/PCG/Private/PCGComponent.cpp
--- a/Engine/Plugins/Experimental/PCG/Source/PCG/Private/PCGComponent.cpp
+++ b/Engine/Plugins/Experimental/PCG/Source/PCG/Private/PCGComponent.cpp
@@ -856,10 +856,10 @@ void UPCGComponent::CleanupLocalImmediate(bool bRemoveComponents)
 						LoadedPreviewResources.Add(Resource);
 					}
 					else
-#endif
 					{
 						Resource->Rename(nullptr, nullptr, REN_ForceNoResetLoaders | REN_DoNotDirty | REN_DontCreateRedirectors | REN_NonTransactional);
 					}
+#endif
 				}
 
 				GeneratedResources.RemoveAtSwap(ResourceIndex);
@@ -958,10 +958,10 @@ FPCGTaskId UPCGComponent::CreateCleanupTask(bool bRemoveComponents, const TArray
 							ThisComponent->LoadedPreviewResources.Add(Resource);
 						}
 						else
-#endif
 						{
 							Resource->Rename(nullptr, nullptr, REN_ForceNoResetLoaders | REN_DoNotDirty | REN_DontCreateRedirectors | REN_NonTransactional);
 						}
+#endif
 					}
 
 					ThisComponent->GeneratedResources.RemoveAtSwap(Context->ResourceIndex);
@@ -1081,10 +1081,10 @@ void UPCGComponent::CleanupUnusedManagedResources()
 						LoadedPreviewResources.Add(Resource);
 					}
 					else
-#endif
 					{
 						Resource->Rename(nullptr, nullptr, REN_ForceNoResetLoaders | REN_DoNotDirty | REN_DontCreateRedirectors | REN_NonTransactional);
 					}
+#endif
 				}
 
 				GeneratedResources.RemoveAtSwap(ResourceIndex);
@@ -2433,7 +2433,9 @@ void UPCGComponent::RefreshSchedulingPolicy()
 	// Only delete it if we are the owner, it's for deprecation where local components had hard ref on original policy.
 	if (IsValid(SchedulingPolicy) && SchedulingPolicy->GetOuter() == this)
 	{
+#if WITH_EDITOR
 		SchedulingPolicy->Rename(nullptr, GetTransientPackage(), REN_DontCreateRedirectors | REN_ForceNoResetLoaders);
+#endif
 		SchedulingPolicy->MarkAsGarbage();
 	}
 
diff --git a/Engine/Plugins/Experimental/PCG/Source/PCG/Private/PCGNode.cpp b/Engine/Plugins/Experimental/PCG/Source/PCG/Private/PCGNode.cpp
--- a/Engine/Plugins/Experimental/PCG/Source/PCG/Private/PCGNode.cpp
+++ b/Engine/Plugins/Experimental/PCG/Source/PCG/Private/PCGNode.cpp
@@ -585,7 +585,9 @@ void UPCGNode::SetSettingsInterface(UPCGSettingsInterface* InSettingsInterface,
 		// a node and get both settings objects in the clipboard text, and the wrong settings can be used upon paste.
 		if (ensure(SettingsInterface->GetOuter() == this))
 		{
+#if WITH_EDITOR
 			SettingsInterface->Rename(nullptr, GetTransientPackage(), REN_ForceNoResetLoaders | REN_DoNotDirty | REN_DontCreateRedirectors | REN_NonTransactional);
+#endif
 			SettingsInterface->MarkAsGarbage();
 		}
 	}
