Added some tools to log abnormal component status for investigation

diff --git a/Engine/Plugins/Experimental/PCG/Source/PCG/Private/PCGSubsystem.cpp b/Engine/Plugins/Experimental/PCG/Source/PCG/Private/PCGSubsystem.cpp
--- a/Engine/Plugins/Experimental/PCG/Source/PCG/Private/PCGSubsystem.cpp
+++ b/Engine/Plugins/Experimental/PCG/Source/PCG/Private/PCGSubsystem.cpp
@@ -305,7 +305,91 @@ void UPCGSubsystem::DestroyAllPCGWorldActors()
 		ActorToDestroy->Destroy();
 	}
 }
-#endif
+
+void UPCGSubsystem::LogAbnormalComponentStates(bool bGroupByState) const
+{
+	TArray<UPCGComponent*> DeactivatedComponents;
+	TArray<UPCGComponent*> NotGeneratedComponents;
+	TArray<UPCGComponent*> DirtyGeneratedComponents;
+
+	UE_LOG(LogPCG, Log, TEXT("--- Logging Abnormal PCG Component States ---"));
+
+	UPCGActorHelpers::ForEachActorInWorld(GetWorld(), AActor::StaticClass(), [bGroupByState, &DeactivatedComponents, &NotGeneratedComponents, &DirtyGeneratedComponents](AActor* InActor)
+	{
+		if (!InActor || !IsValid(InActor))
+		{
+			return true;
+		}
+
+		TInlineComponentArray<UPCGComponent*, 1> PCGComponents;
+		InActor->GetComponents(PCGComponents);
+
+		for (UPCGComponent* PCGComponent : PCGComponents)
+		{
+			if (!PCGComponent->bActivated)
+			{
+				if (bGroupByState)
+				{
+					DeactivatedComponents.Add(PCGComponent);
+				}
+				else
+				{
+					UE_LOG(LogPCG, Log, TEXT("%s - %s - Deactivated Component"), *InActor->GetName(), *PCGComponent->GetName());
+				}
+			}
+			else if (!PCGComponent->bGenerated && PCGComponent->GenerationTrigger != EPCGComponentGenerationTrigger::GenerateAtRuntime)
+			{
+				if (bGroupByState)
+				{
+					NotGeneratedComponents.Add(PCGComponent);
+				}
+				else
+				{
+					UE_LOG(LogPCG, Log, TEXT("%s - %s - Not Generated Component"), *InActor->GetName(), *PCGComponent->GetName());
+				}
+			}
+			else if (PCGComponent->bDirtyGenerated)
+			{
+				if (bGroupByState)
+				{
+					DirtyGeneratedComponents.Add(PCGComponent);
+				}
+				else
+				{
+					UE_LOG(LogPCG, Log, TEXT("%s - %s - Dirty Generated Component "), *InActor->GetName(), *PCGComponent->GetName());
+				}
+			}
+		}
+
+		return true;
+	});
+
+	if (bGroupByState)
+	{
+		UE_LOG(LogPCG, Log, TEXT("--- Deactivated PCG Components ---"));
+		for (UPCGComponent* Component : DeactivatedComponents)
+		{
+			check(Component && Component->GetOwner());
+			UE_LOG(LogPCG, Log, TEXT("%s - %s"), *Component->GetOwner()->GetName(), *Component->GetName());
+		}
+
+		UE_LOG(LogPCG, Log, TEXT("--- Not Generated Components ---"));
+		for (UPCGComponent* Component : NotGeneratedComponents)
+		{
+			check(Component && Component->GetOwner());
+			UE_LOG(LogPCG, Log, TEXT("%s - %s"), *Component->GetOwner()->GetName(), *Component->GetName());
+		}
+
+		UE_LOG(LogPCG, Log, TEXT("--- Dirty Components ---"));
+		for (UPCGComponent* Component : DirtyGeneratedComponents)
+		{
+			check(Component && Component->GetOwner());
+			UE_LOG(LogPCG, Log, TEXT("%s - %s"), *Component->GetOwner()->GetName(), *Component->GetName());
+		}
+	}
+}
+
+#endif // WITH_EDITOR
 
 void UPCGSubsystem::RegisterPCGWorldActor(APCGWorldActor* InActor)
 {
diff --git a/Engine/Plugins/Experimental/PCG/Source/PCG/Public/PCGSubsystem.h b/Engine/Plugins/Experimental/PCG/Source/PCG/Public/PCGSubsystem.h
--- a/Engine/Plugins/Experimental/PCG/Source/PCG/Public/PCGSubsystem.h
+++ b/Engine/Plugins/Experimental/PCG/Source/PCG/Public/PCGSubsystem.h
@@ -80,6 +80,7 @@ public:
 #if WITH_EDITOR
 	void DestroyAllPCGWorldActors();
 	void DestroyCurrentPCGWorldActor();
+	void LogAbnormalComponentStates(bool bGroupByState) const;
 #endif
 	void RegisterPCGWorldActor(APCGWorldActor* InActor);
 	void UnregisterPCGWorldActor(APCGWorldActor* InActor);
diff --git a/Engine/Plugins/Experimental/PCG/Source/PCGEditor/Private/PCGEditorModule.cpp b/Engine/Plugins/Experimental/PCG/Source/PCGEditor/Private/PCGEditorModule.cpp
--- a/Engine/Plugins/Experimental/PCG/Source/PCGEditor/Private/PCGEditorModule.cpp
+++ b/Engine/Plugins/Experimental/PCG/Source/PCGEditor/Private/PCGEditorModule.cpp
@@ -279,6 +279,41 @@ void FPCGEditorModule::PopulateMenuActions(FMenuBuilder& MenuBuilder)
 				PCGEditorUtils::ForcePCGBlueprintVariableVisibility();
 				})),
 		NAME_None);
+
+	MenuBuilder.AddSubMenu(
+		LOCTEXT("PCGToolsLoggingSubMenu", "Logging / Reporting"),
+		LOCTEXT("PCGToolsLoggingSubMenu_Tooltip", "Logging and reporting related editor commands"),
+		FNewMenuDelegate::CreateLambda([this](FMenuBuilder& LoggingMenuBuilder)
+		{
+			LoggingMenuBuilder.AddMenuEntry(
+			LOCTEXT("LogAbnormalComponentState", "Log abnormal component state (actor order)"),
+			LOCTEXT("LogAbnormalComponentState_Tooltip", "Logs unusual PCG components state, for every loaded actor"),
+			FSlateIcon(),
+			FUIAction(
+				FExecuteAction::CreateLambda([]() {
+					if (UPCGSubsystem* PCGSubsystem = UPCGSubsystem::GetInstance(GEditor->GetEditorWorldContext().World()))
+					{
+						PCGSubsystem->LogAbnormalComponentStates(/*bGroupByState=*/false);
+					}
+				})),
+			NAME_None);
+
+			LoggingMenuBuilder.AddMenuEntry(
+			LOCTEXT("LogAbnormalComponentState", "Log abnormal component state (grouped by state)"),
+			LOCTEXT("LogAbnormalComponentState_Tooltip", "Logs unusual PCG components, for every loaded actor, grouped by state"),
+			FSlateIcon(),
+			FUIAction(
+				FExecuteAction::CreateLambda([]() {
+					if (UPCGSubsystem* PCGSubsystem = UPCGSubsystem::GetInstance(GEditor->GetEditorWorldContext().World()))
+					{
+						PCGSubsystem->LogAbnormalComponentStates(/*bGroupByState=*/true);
+					}
+				})),
+			NAME_None);
+		}),
+		false,
+		FSlateIcon());
+
 }
 
 void FPCGEditorModule::RegisterSettings()
