[PCG] Provide the option to sample GPU textures from CPU-accessible memory in PCGTextureSampler
* This new path will be taken by default when the texture settings permit it, but can be forced with bForceEditorOnlyCPUSampling to create a duplicate texture with the correct settings

#localization none
#tests tried locally on a few textures

diff --git a/Engine/Plugins/Experimental/PCG/Source/PCG/Private/Data/PCGTextureData.cpp b/Engine/Plugins/Experimental/PCG/Source/PCG/Private/Data/PCGTextureData.cpp
--- a/Engine/Plugins/Experimental/PCG/Source/PCG/Private/Data/PCGTextureData.cpp
+++ b/Engine/Plugins/Experimental/PCG/Source/PCG/Private/Data/PCGTextureData.cpp
@@ -14,8 +14,33 @@
 
 #include UE_INLINE_GENERATED_CPP_BY_NAME(PCGTextureData)
 
-namespace PCGTextureSampling
+namespace PCGTextureSamplingHelpers
 {
+	bool IsTextureCPUAccessible(UTexture2D* Texture)
+	{
+		FTexturePlatformData* PlatformData = Texture ? Texture->GetPlatformData() : nullptr;
+		return PlatformData && PlatformData->GetHasCpuCopy();
+	}
+
+	bool CanGPUTextureBeCPUAccessed(UTexture2D* Texture)
+	{
+		// SRGB textures need to be GPU sampled.
+		if (!Texture || Texture->SRGB)
+		{
+			return false;
+		}
+
+		FTexturePlatformData* PlatformData = Texture->GetPlatformData();
+
+		// If a CPU copy is available, this is a CPU texture and not a GPU texture, so we return false.
+		if (!PlatformData || PlatformData->GetHasCpuCopy())
+		{
+			return false;
+		}
+
+		return PlatformData->Mips.Num() == 1 && PlatformData->PixelFormat == PF_B8G8R8A8;
+	}
+
 	template<typename ValueType>
 	bool Sample(const FVector2D& InPosition,
 		const FBox2D& InSurface,
@@ -162,10 +187,10 @@ bool UPCGBaseTextureData::SamplePoint(const FTransform& InTransform, const FBox&
 	FBox2D Surface(FVector2D(-1.0f, -1.0f), FVector2D(1.0f, 1.0f));
 
 	FLinearColor Color = FLinearColor(EForceInit::ForceInit);
-	if (PCGTextureSampling::Sample<FLinearColor>(Position2D, Surface, this, Width, Height, Color, [this](int32 Index) { return ColorData[Index]; }))
+	if (PCGTextureSamplingHelpers::Sample<FLinearColor>(Position2D, Surface, this, Width, Height, Color, [this](int32 Index) { return ColorData[Index]; }))
 	{
 		OutPoint.Color = Color;
-		OutPoint.Density = ((DensityFunction == EPCGTextureDensityFunction::Ignore) ? 1.0f : PCGTextureSampling::SampleFloatChannel(Color, ColorChannel));
+		OutPoint.Density = ((DensityFunction == EPCGTextureDensityFunction::Ignore) ? 1.0f : PCGTextureSamplingHelpers::SampleFloatChannel(Color, ColorChannel));
 		return OutPoint.Density > 0;
 	}
 	else
@@ -219,9 +244,9 @@ const UPCGPointData* UPCGBaseTextureData::CreatePointData(FPCGContext* Context)
 		FVector2D LocalCoordinate((2.0 * X + 0.5) / XCount - 1.0, (2.0 * Y + 0.5) / YCount - 1.0);
 		FLinearColor Color = FLinearColor(EForceInit::ForceInit);
 
-		if (PCGTextureSampling::Sample<FLinearColor>(LocalCoordinate, Surface, this, Width, Height, Color, [this](int32 Index) { return ColorData[Index]; }))
+		if (PCGTextureSamplingHelpers::Sample<FLinearColor>(LocalCoordinate, Surface, this, Width, Height, Color, [this](int32 Index) { return ColorData[Index]; }))
 		{
-			const float Density = ((DensityFunction == EPCGTextureDensityFunction::Ignore) ? 1.0f : PCGTextureSampling::SampleFloatChannel(Color, ColorChannel));
+			const float Density = ((DensityFunction == EPCGTextureDensityFunction::Ignore) ? 1.0f : PCGTextureSamplingHelpers::SampleFloatChannel(Color, ColorChannel));
 			if (Density > 0 || bKeepZeroDensityPoints)
 			{
 				FVector LocalPosition(LocalCoordinate, 0);
@@ -266,7 +291,7 @@ void UPCGBaseTextureData::CopyBaseTextureData(UPCGBaseTextureData* NewTextureDat
 	NewTextureData->Width = Width;
 }
 
-void UPCGTextureData::Initialize(UTexture* InTexture, uint32 InTextureIndex, const FTransform& InTransform, const TFunction<void()>& PostInitializeCallback)
+void UPCGTextureData::Initialize(UTexture* InTexture, uint32 InTextureIndex, const FTransform& InTransform, const TFunction<void()>& PostInitializeCallback, bool bCreateCPUDuplicateEditorOnly)
 {
 	TRACE_CPUPROFILER_EVENT_SCOPE(UPCGTextureData::Initialize);
 
@@ -287,6 +312,29 @@ void UPCGTextureData::Initialize(UTexture* InTexture, uint32 InTextureIndex, con
 		return;
 	}
 
+#if WITH_EDITOR
+	// Create a duplicate texture if necessary.
+	if (bCreateCPUDuplicateEditorOnly)
+	{
+		UTexture2D* Texture2D = Cast<UTexture2D>(Texture);
+		if (!Texture2D)
+		{
+			UTexture2DArray* Texture2DArray = CastChecked<UTexture2DArray>(Texture);
+			Texture2D = Texture2DArray->SourceTextures.IsValidIndex(TextureIndex) ? Texture2DArray->SourceTextures[TextureIndex] : nullptr;
+		}
+
+		if (Texture2D && !PCGTextureSamplingHelpers::IsTextureCPUAccessible(Texture2D) && !PCGTextureSamplingHelpers::CanGPUTextureBeCPUAccessed(Texture2D))
+		{
+			FObjectDuplicationParameters DuplicationParams(Texture2D, /*Outer=*/this);
+			DuplicateTexture = CastChecked<UTexture2D>(StaticDuplicateObjectEx(DuplicationParams));
+			DuplicateTexture->MipGenSettings = TextureMipGenSettings::TMGS_NoMipmaps;
+			DuplicateTexture->CompressionSettings = TextureCompressionSettings::TC_VectorDisplacementmap; // Allows the texture to be in a non-compressed format (B8G8R8A8), which is necessary to convince the data to remain CPU-side.
+			DuplicateTexture->SRGB = false;
+			DuplicateTexture->UpdateResource();
+		}
+	}
+#endif
+
 	// Prioritize initializing from a CPU texture when the provided texture is marked as CPU accessible
 	if (InitializeFromCPUTexture())
 	{
@@ -294,6 +342,15 @@ void UPCGTextureData::Initialize(UTexture* InTexture, uint32 InTextureIndex, con
 	}
 	else
 	{
+#if WITH_EDITOR
+		// Try reading the texture back from CPU-accessible memory if possible.
+		if (InitializeGPUTextureFromCPU())
+		{
+			PostInitializeCallback();
+			return;
+		}
+#endif
+
 		if (!InitializeFromGPUTexture(PostInitializeCallback))
 		{
 			UE_LOG(LogPCG, Error, TEXT("PCGTextureData failed to initialize texture '%s'"), *Texture->GetFName().ToString());
@@ -536,3 +593,45 @@ bool UPCGTextureData::InitializeFromGPUTexture(const TFunction<void()>& PostInit
 
 	return true;
 }
+
+#if WITH_EDITOR
+bool UPCGTextureData::InitializeGPUTextureFromCPU()
+{
+	if (!DuplicateTexture || !PCGTextureSamplingHelpers::CanGPUTextureBeCPUAccessed(DuplicateTexture))
+	{
+		return false;
+	}
+
+	FTexturePlatformData* PlatformData = DuplicateTexture->GetPlatformData();
+	if (!PlatformData)
+	{
+		return false;
+	}
+
+	bool bBulkDataAccessed = false;
+
+	if (const uint8_t* BulkData = reinterpret_cast<const uint8_t*>(PlatformData->Mips[0].BulkData.LockReadOnly()))
+	{
+		bBulkDataAccessed = true;
+
+		Width = PlatformData->SizeX;
+		Height = PlatformData->SizeY;
+		const int32 PixelCount = Width * Height;
+		ColorData.SetNum(PixelCount);
+
+		const FColor* FormattedImageData = reinterpret_cast<const FColor*>(BulkData);
+		for (int32 D = 0; D < PixelCount; ++D)
+		{
+			ColorData[D] = FormattedImageData[D].ReinterpretAsLinear();
+		}
+	}
+	else
+	{
+		UE_LOG(LogPCG, Error, TEXT("PCGTextureData unable to get bulk data from '%s'."), *Texture->GetFName().ToString());
+	}
+
+	PlatformData->Mips[0].BulkData.Unlock();
+
+	return bBulkDataAccessed;
+}
+#endif
diff --git a/Engine/Plugins/Experimental/PCG/Source/PCG/Private/Elements/PCGTextureSampler.cpp b/Engine/Plugins/Experimental/PCG/Source/PCG/Private/Elements/PCGTextureSampler.cpp
--- a/Engine/Plugins/Experimental/PCG/Source/PCG/Private/Elements/PCGTextureSampler.cpp
+++ b/Engine/Plugins/Experimental/PCG/Source/PCG/Private/Elements/PCGTextureSampler.cpp
@@ -142,6 +142,11 @@ bool FPCGTextureSamplerElement::ExecuteInternal(FPCGContext* InContext) const
 	const bool bUseTileBounds = Settings->bUseTileBounds;
 	const FVector2D& TileBoundsMin = Settings->TileBoundsMin;
 	const FVector2D& TileBoundsMax = Settings->TileBoundsMax;
+#if WITH_EDITOR
+	const bool bForceEditorOnlyCPUSampling = Settings->bForceEditorOnlyCPUSampling;
+#else
+	const bool bForceEditorOnlyCPUSampling = false;
+#endif
 
 	TArray<FPCGTaggedData>& Outputs = Context->OutputData.TaggedData;
 	FPCGTaggedData& Output = Outputs.Emplace_GetRef();
@@ -174,7 +179,7 @@ bool FPCGTextureSamplerElement::ExecuteInternal(FPCGContext* InContext) const
 		}
 	};
 
-	TextureData->Initialize(Texture, TextureArrayIndex, FinalTransform, PostInitializeCallback);
+	TextureData->Initialize(Texture, TextureArrayIndex, FinalTransform, PostInitializeCallback, bForceEditorOnlyCPUSampling);
 
 	TextureData->DensityFunction = DensityFunction;
 	TextureData->ColorChannel = ColorChannel;
diff --git a/Engine/Plugins/Experimental/PCG/Source/PCG/Public/Data/PCGTextureData.h b/Engine/Plugins/Experimental/PCG/Source/PCG/Public/Data/PCGTextureData.h
--- a/Engine/Plugins/Experimental/PCG/Source/PCG/Public/Data/PCGTextureData.h
+++ b/Engine/Plugins/Experimental/PCG/Source/PCG/Public/Data/PCGTextureData.h
@@ -8,6 +8,7 @@
 
 class UPCGSpatialData;
 class UTexture;
+class UTexture2D;
 
 UENUM(BlueprintType)
 enum class EPCGTextureColorChannel : uint8
@@ -25,6 +26,15 @@ enum class EPCGTextureDensityFunction : uint8
 	Multiply
 };
 
+namespace PCGTextureSamplingHelpers
+{
+	/** Returns true if a texture is CPU-accessible. */
+	bool IsTextureCPUAccessible(UTexture2D* Texture);
+
+	/** Returns true if a texture is both GPU-accessible and reachable from CPU memory. */
+	bool CanGPUTextureBeCPUAccessed(UTexture2D* Texture);
+}
+
 UCLASS(Abstract)
 class PCG_API UPCGBaseTextureData : public UPCGSurfaceData
 {
@@ -100,7 +110,7 @@ class PCG_API UPCGTextureData : public UPCGBaseTextureData
 	GENERATED_BODY()
 
 public:
-	void Initialize(UTexture* InTexture, uint32 InTextureIndex, const FTransform& InTransform, const TFunction<void()>& PostInitializeCallback);
+	void Initialize(UTexture* InTexture, uint32 InTextureIndex, const FTransform& InTransform, const TFunction<void()>& PostInitializeCallback, bool bCreateCPUDuplicateEditorOnly = false);
 
 	// ~Begin UPCGData interface
 	virtual EPCGDataType GetDataType() const override { return EPCGDataType::Texture; }
@@ -122,10 +132,21 @@ private:
 	*/
 	bool InitializeFromGPUTexture(const TFunction<void()>& PostInitializeCallback);
 
+#if WITH_EDITOR
+	/** Attempts to initialize the UPCGTextureData from a GPU-accessible texture, but with CPU-accessible memory. Returns true if initialization succeeds. */
+	bool InitializeGPUTextureFromCPU();
+#endif
+
 public:
 	UPROPERTY(BlueprintReadOnly, VisibleAnywhere, Category = Properties)
 	TWeakObjectPtr<UTexture> Texture = nullptr;
 
+#if WITH_EDITORONLY_DATA
+	/** Transient CPU visible duplicate of Texture created and used only when initialized with bCreateCPUDuplicateEditorOnly. */
+	UPROPERTY(Transient)
+	TObjectPtr<UTexture2D> DuplicateTexture = nullptr;
+#endif
+
 	UPROPERTY(BlueprintReadOnly, VisibleAnywhere, Category = Properties)
 	int TextureIndex = 0;
 };
diff --git a/Engine/Plugins/Experimental/PCG/Source/PCG/Public/Elements/PCGTextureSampler.h b/Engine/Plugins/Experimental/PCG/Source/PCG/Public/Elements/PCGTextureSampler.h
--- a/Engine/Plugins/Experimental/PCG/Source/PCG/Public/Elements/PCGTextureSampler.h
+++ b/Engine/Plugins/Experimental/PCG/Source/PCG/Public/Elements/PCGTextureSampler.h
@@ -91,6 +91,15 @@ public:
 	UPROPERTY(BlueprintReadWrite, EditAnywhere, Category = "Settings|Tiling", meta = (EditCondition = "bUseAdvancedTiling && bUseTileBounds", PCG_Overridable))
 	FVector2D TileBoundsMax = FVector2D(0.5, 0.5);
 
+#if WITH_EDITORONLY_DATA
+	/** Even if the texture is not set to CPU-available, it can still be accessed from CPU memory under certain conditions (sRGB disabled, no mipmaps, and non-compressed format).
+	 * Reading from CPU memory will be faster and more accurate than reading from GPU memory, since the texture will not be subject to compression or resolution clamping. Enable
+	 * this flag to force a duplicate of the texture with the correct settings for CPU memory access. This is editor-only.
+	 */
+	UPROPERTY(BlueprintReadWrite, EditAnywhere, Category = "Settings")
+	bool bForceEditorOnlyCPUSampling = false;
+#endif
+
 protected:
 #if WITH_EDITORONLY_DATA
 	// Used to hide the 'TextureArrayIndex' property.
