Fixed duplicate localization tag

#rnx

diff --git a/Engine/Plugins/Experimental/PCG/Source/PCGEditor/Private/PCGEditorModule.cpp b/Engine/Plugins/Experimental/PCG/Source/PCGEditor/Private/PCGEditorModule.cpp
--- a/Engine/Plugins/Experimental/PCG/Source/PCGEditor/Private/PCGEditorModule.cpp
+++ b/Engine/Plugins/Experimental/PCG/Source/PCGEditor/Private/PCGEditorModule.cpp
@@ -299,8 +299,8 @@ void FPCGEditorModule::PopulateMenuActions(FMenuBuilder& MenuBuilder)
 			NAME_None);
 
 			LoggingMenuBuilder.AddMenuEntry(
-			LOCTEXT("LogAbnormalComponentState", "Log abnormal component state (grouped by state)"),
-			LOCTEXT("LogAbnormalComponentState_Tooltip", "Logs unusual PCG components, for every loaded actor, grouped by state"),
+			LOCTEXT("LogAbnormalComponentState_GroupedByState", "Log abnormal component state (grouped by state)"),
+			LOCTEXT("LogAbnormalComponentState_GroupedByState_Tooltip", "Logs unusual PCG components, for every loaded actor, grouped by state"),
 			FSlateIcon(),
 			FUIAction(
 				FExecuteAction::CreateLambda([]() {
