[PCG] Add DisplayName "PCG Volume" for APCGVolume

#rnx

diff --git a/Engine/Plugins/Experimental/PCG/Source/PCG/Public/PCGVolume.h b/Engine/Plugins/Experimental/PCG/Source/PCG/Public/PCGVolume.h
--- a/Engine/Plugins/Experimental/PCG/Source/PCG/Public/PCGVolume.h
+++ b/Engine/Plugins/Experimental/PCG/Source/PCG/Public/PCGVolume.h
@@ -8,7 +8,7 @@
 
 class UPCGComponent;
 
-UCLASS(BlueprintType)
+UCLASS(BlueprintType, DisplayName = "PCG Volume")
 class PCG_API APCGVolume : public AVolume
 {
 	GENERATED_BODY()
