[PCG] CPU-Accessible textures have swapped R and B channels after sampling

#rnx

diff --git a/Engine/Plugins/Experimental/PCG/Source/PCG/Private/Data/PCGTextureData.cpp b/Engine/Plugins/Experimental/PCG/Source/PCG/Private/Data/PCGTextureData.cpp
--- a/Engine/Plugins/Experimental/PCG/Source/PCG/Private/Data/PCGTextureData.cpp
+++ b/Engine/Plugins/Experimental/PCG/Source/PCG/Private/Data/PCGTextureData.cpp
@@ -418,7 +418,7 @@ bool UPCGTextureData::InitializeFromCPUTexture()
 	{
 		const TArrayView64<const FColor> DataView = CPUTextureRef->AsBGRA8();
 
-		// Memory representation of FColor is BGRA
+		// Memory representation of FColor is BGRA, so we reinterpret as FLinearColor to get RGBA.
 		for (int32 D = 0; D < PixelCount; ++D)
 		{
 			ColorData[D] = DataView[D].ReinterpretAsLinear();
@@ -428,7 +428,7 @@ bool UPCGTextureData::InitializeFromCPUTexture()
 	{
 		const TArrayView64<const FColor> DataView = CPUTextureRef->AsBGRE8();
 
-		// Memory representation of FColor is BGRA
+		// Memory representation of FColor is BGRA, so we reinterpret as FLinearColor to get RGBA.
 		for (int32 D = 0; D < PixelCount; ++D)
 		{
 			ColorData[D] = DataView[D].ReinterpretAsLinear();
@@ -442,9 +442,7 @@ bool UPCGTextureData::InitializeFromCPUTexture()
 		for (int32 D = 0; D < PixelCount; ++D)
 		{
 			const uint32 Index = D * 4;
-
-			// To avoid swapping R and B to BGRA format, we can access in the correct order here
-			ColorData[D] = FColor(DataView[Index + 2], DataView[Index + 1], DataView[Index + 0], DataView[Index + 3]).ReinterpretAsLinear();
+			ColorData[D] = FLinearColor(DataView[Index + 0], DataView[Index + 1], DataView[Index + 2], DataView[Index + 3]);
 		}
 	}
 	else if (CPUTextureRef->Format == ERawImageFormat::RGBA16F)
@@ -453,10 +451,7 @@ bool UPCGTextureData::InitializeFromCPUTexture()
 
 		for (int32 D = 0; D < PixelCount; ++D)
 		{
-			// Swap R and B to achieve BGRA format
-			FLinearColor Temp = FLinearColor(DataView[D]);
-			Swap(Temp.R, Temp.B);
-			ColorData[D] = Temp;
+			ColorData[D] = FLinearColor(DataView[D]);
 		}
 	}
 	else if (CPUTextureRef->Format == ERawImageFormat::RGBA32F)
@@ -465,10 +460,7 @@ bool UPCGTextureData::InitializeFromCPUTexture()
 
 		for (int32 D = 0; D < PixelCount; ++D)
 		{
-			// Swap R and B to achieve BGRA format
-			FLinearColor Temp = DataView[D];
-			Swap(Temp.R, Temp.B);
-			ColorData[D] = Temp;
+			ColorData[D] = DataView[D];
 		}
 	}
 	else if (CPUTextureRef->Format == ERawImageFormat::G16)
