[PCG] Add HideEditConditionToggle to SchedulingPolicies
* Prevents exposing bShouldDisplayProperties in the editor, which should be an invariant managed by the component

#rnx

diff --git a/Engine/Plugins/Experimental/PCG/Source/PCG/Public/RuntimeGen/SchedulingPolicies/PCGSchedulingPolicyDistanceAndDirection.h b/Engine/Plugins/Experimental/PCG/Source/PCG/Public/RuntimeGen/SchedulingPolicies/PCGSchedulingPolicyDistanceAndDirection.h
--- a/Engine/Plugins/Experimental/PCG/Source/PCG/Public/RuntimeGen/SchedulingPolicies/PCGSchedulingPolicyDistanceAndDirection.h
+++ b/Engine/Plugins/Experimental/PCG/Source/PCG/Public/RuntimeGen/SchedulingPolicies/PCGSchedulingPolicyDistanceAndDirection.h
@@ -28,7 +28,7 @@ public:
 
 public:
 	/** Toggle whether or not distance is used to calculate the scheduling priority. */
-	UPROPERTY(BlueprintReadWrite, EditAnywhere, Category = "RuntimeGeneration|Scheduling Policy Parameters", meta = (EditCondition = "bShouldDisplayProperties", EditConditionHides))
+	UPROPERTY(BlueprintReadWrite, EditAnywhere, Category = "RuntimeGeneration|Scheduling Policy Parameters", meta = (EditCondition = "bShouldDisplayProperties", EditConditionHides, HideEditConditionToggle))
 	bool bUseDistance = true;
 
 	/** Scalar value used to increase/decrease the impact of distance in the scheduling priority. */
@@ -36,7 +36,7 @@ public:
 	float DistanceWeight = 1.0f;
 
 	/** Toggle whether or not direction is used to calculate the scheduling priority. */
-	UPROPERTY(BlueprintReadWrite, EditAnywhere, Category = "RuntimeGeneration|Scheduling Policy Parameters", meta = (EditCondition = "bShouldDisplayProperties", EditConditionHides))
+	UPROPERTY(BlueprintReadWrite, EditAnywhere, Category = "RuntimeGeneration|Scheduling Policy Parameters", meta = (EditCondition = "bShouldDisplayProperties", EditConditionHides, HideEditConditionToggle))
 	bool bUseDirection = true;
 
 	/** Scalar value used to increase/decrease the impact of direction in the scheduling priority. */
