[PCG] Improving Self-Pruning node

* We can now choose what attribute we want to compare.
* It still defaults to the preivous code path in the case of extents, since it is a bit faster
* Make the element time sliced
* Extend the sort by attribute to support randomized/similar value sorting
* Add SquaredLength as extractor.

diff --git a/Engine/Plugins/Experimental/PCG/Source/PCG/Private/Elements/PCGSelfPruning.cpp b/Engine/Plugins/Experimental/PCG/Source/PCG/Private/Elements/PCGSelfPruning.cpp
--- a/Engine/Plugins/Experimental/PCG/Source/PCG/Private/Elements/PCGSelfPruning.cpp
+++ b/Engine/Plugins/Experimental/PCG/Source/PCG/Private/Elements/PCGSelfPruning.cpp
@@ -6,6 +6,10 @@
 #include "PCGCustomVersion.h"
 #include "Data/PCGPointData.h"
 #include "Data/PCGSpatialData.h"
+#include "Metadata/Accessors/IPCGAttributeAccessor.h"
+#include "Metadata/Accessors/PCGAttributeAccessorHelpers.h"
+#include "Metadata/Accessors/PCGAttributeAccessorKeys.h"
+#include "Metadata/Accessors/PCGAttributeExtractor.h"
 
 #include UE_INLINE_GENERATED_CPP_BY_NAME(PCGSelfPruning)
 
@@ -40,106 +44,182 @@ namespace PCGSelfPruningAlgorithms
 			return RandomSort(A, B);
 		}
 	}
+
+	/**
+	* Keep Extents sort alive since it is a bit faster than the SortByAttribute.
+	*/
+	void SortExtents(TArray<FPCGPointRef>& SortedPoints, const FPCGSelfPruningParameters& Parameters, FVector::FReal SquaredRadiusEquality)
+	{
+		TRACE_CPUPROFILER_EVENT_SCOPE(FPCGSelfPruningElement::Execute::SortExtents);
+
+		if (Parameters.PruningType == EPCGSelfPruningType::LargeToSmall)
+		{
+			if (Parameters.bRandomizedPruning)
+			{
+				Algo::Sort(SortedPoints, [SquaredRadiusEquality](const FPCGPointRef& A, const FPCGPointRef& B) { return !PCGSelfPruningAlgorithms::SortSmallToLargeWithRandom(A, B, SquaredRadiusEquality); });
+			}
+			else
+			{
+				Algo::Sort(SortedPoints, [SquaredRadiusEquality](const FPCGPointRef& A, const FPCGPointRef& B) { return !PCGSelfPruningAlgorithms::SortSmallToLargeNoRandom(A, B, SquaredRadiusEquality); });
+			}
+		}
+		else //if (Parameters.PruningType == EPCGSelfPruningType::SmallToLarge)
+		{
+			if (Parameters.bRandomizedPruning)
+			{
+				Algo::Sort(SortedPoints, [SquaredRadiusEquality](const FPCGPointRef& A, const FPCGPointRef& B) { return PCGSelfPruningAlgorithms::SortSmallToLargeWithRandom(A, B, SquaredRadiusEquality); });
+			}
+			else
+			{
+				Algo::Sort(SortedPoints, [SquaredRadiusEquality](const FPCGPointRef& A, const FPCGPointRef& B) { return PCGSelfPruningAlgorithms::SortSmallToLargeNoRandom(A, B, SquaredRadiusEquality); });
+			}
+		}
+	}
 }
 
 namespace PCGSelfPruningElement
 {
-	// all points are in contiguous memory so this will create a bitset to just look up the index to set/test a bit
-	// this should reduce the 'set' cost to O[1]
-	struct FPointBitSet
+	void FPointBitSet::Initialize(const TArray<FPCGPoint>& Points)
 	{
-		TArray<uint32> Bits;
-		const FPCGPoint* FirstPoint = nullptr;
+		Bits.SetNumZeroed(Points.Num() / 32 + 1);
+		FirstPoint = Points.Num() > 0 ? &Points[0] : nullptr;
+	}
 
-		FPointBitSet(const TArray<FPCGPoint>& Points)
-		{
-			Bits.SetNumZeroed(Points.Num()/32 + 1);
-			FirstPoint = Points.Num() > 0 ? &Points[0] : nullptr;
-		}
+	void FPointBitSet::Add(const FPCGPoint* Point)
+	{
+		const uint32 Index = Point - FirstPoint;
+		const uint32 BitsIndex = Index / 32;
+		const uint32 Bit = Index % 32;
 
-		void Add(const FPCGPoint* Point)
-		{
-			const uint32 Index = Point-FirstPoint;
-			const uint32 BitsIndex = Index / 32;
-			const uint32 Bit = Index % 32;
+		Bits[BitsIndex] |= (1 << Bit);
+	}
 
-			Bits[BitsIndex] |= (1<<Bit);
-		}
+	bool FPointBitSet::Contains(const FPCGPoint* Point)
+	{
+		const uint32 Index = Point - FirstPoint;
+		const uint32 BitsIndex = Index / 32;
+		const uint32 Bit = Index % 32;
 
-		bool Contains(const FPCGPoint* Point)
-		{
-			const uint32 Index = Point-FirstPoint;
-			const uint32 BitsIndex = Index / 32;
-			const uint32 Bit = Index % 32;
+		return (Bits[BitsIndex] & (1 << Bit)) != 0;
+	}
 
-			return (Bits[BitsIndex] & (1<<Bit)) != 0;
-		}
-	};
+	static constexpr int32 TimeSliceFrequencyCheck = 255;
+
+	bool ShouldStop(FPCGContext* InOptionalContext)
+	{
+		return InOptionalContext && InOptionalContext->TimeSliceIsEnabled() && InOptionalContext->ShouldStop();
+	}
 
-	void DensityBoundsExclusion(const TArray<FPCGPoint>& Points, const TArray<FPCGPointRef>& SortedPoints, const UPCGPointData::PointOctree& Octree, FPointBitSet& ExclusionPoints)
+	bool DensityBoundsExclusion(FIterationState& IterationState, FPCGContext* InOptionalContext)
 	{
-		FPointBitSet ExcludedPoints(Points);
-		
-		for (const FPCGPointRef& PointRef : SortedPoints)
+		TRACE_CPUPROFILER_EVENT_SCOPE(FPCGSelfPruningElement::Execute::DensityBoundsExclusion);
+
+		check(IterationState.InputPointData);
+		const UPCGPointData::PointOctree& Octree = IterationState.InputPointData->GetOctree();
+
+		int32 CheckTimeSlicingCount = 0;
+
+		while (IterationState.CurrentPointIndex < IterationState.SortedPoints.Num())
 		{
-			if (ExcludedPoints.Contains(PointRef.Point))
+			// Don't check too many times.
+			if (++CheckTimeSlicingCount >= TimeSliceFrequencyCheck)
+			{
+				if (ShouldStop(InOptionalContext))
+				{
+					return false;
+				}
+
+				CheckTimeSlicingCount = 0;
+			}
+
+			const FPCGPointRef& PointRef = IterationState.SortedPoints[IterationState.CurrentPointIndex++];
+			if (IterationState.ExcludedPoints.Contains(PointRef.Point))
 			{
 				continue;
 			}
 
-			ExclusionPoints.Add(PointRef.Point);
+			IterationState.ExclusionPoints.Add(PointRef.Point);
 
-			Octree.FindElementsWithBoundsTest(FBoxCenterAndExtent(PointRef.Bounds.Origin, PointRef.Bounds.BoxExtent), [&ExclusionPoints, &ExcludedPoints](const FPCGPointRef& InPointRef)
+			Octree.FindElementsWithBoundsTest(FBoxCenterAndExtent(PointRef.Bounds.Origin, PointRef.Bounds.BoxExtent), [&IterationState](const FPCGPointRef& InPointRef)
 			{
 				// TODO: check on an oriented-box basis?
-				if (!ExclusionPoints.Contains(InPointRef.Point))
+				if (!IterationState.ExclusionPoints.Contains(InPointRef.Point))
 				{
-					ExcludedPoints.Add(InPointRef.Point);
+					IterationState.ExcludedPoints.Add(InPointRef.Point);
 				}
 			});
 		}
+
+		return true;
 	}
 
-	void DuplicatePointsExclusion(const TArray<FPCGPoint>& Points, const TArray<FPCGPointRef>& SortedPoints, const UPCGPointData::PointOctree& Octree, FPointBitSet& ExclusionPoints)
+	bool DuplicatePointsExclusion(FIterationState& IterationState, FPCGContext* InOptionalContext)
 	{
-		FPointBitSet ExcludedPoints(Points);
+		TRACE_CPUPROFILER_EVENT_SCOPE(FPCGSelfPruningElement::Execute::DuplicatePointsExclusion);
 
-		for (const FPCGPointRef& PointRef : SortedPoints)
+		check(IterationState.InputPointData);
+		const UPCGPointData::PointOctree& Octree = IterationState.InputPointData->GetOctree();
+
+		int32 CheckTimeSlicingCount = 0;
+
+		while (IterationState.CurrentPointIndex < IterationState.SortedPoints.Num())
 		{
-			if (ExcludedPoints.Contains(PointRef.Point))
+			// Don't check too many times.
+			if (++CheckTimeSlicingCount >= TimeSliceFrequencyCheck)
+			{
+				if (ShouldStop(InOptionalContext))
+				{
+					return false;
+				}
+
+				CheckTimeSlicingCount = 0;
+			}
+
+			const FPCGPointRef& PointRef = IterationState.SortedPoints[IterationState.CurrentPointIndex++];
+			if (IterationState.ExcludedPoints.Contains(PointRef.Point))
 			{
 				continue;
 			}
 
-			ExclusionPoints.Add(PointRef.Point);
+			IterationState.ExclusionPoints.Add(PointRef.Point);
 
-			Octree.FindElementsWithBoundsTest(FBoxCenterAndExtent(PointRef.Point->Transform.TransformPosition(PointRef.Point->GetLocalCenter()), FVector::Zero()), [&ExclusionPoints, &ExcludedPoints, PointRef](const FPCGPointRef& InPointRef)
+			Octree.FindElementsWithBoundsTest(FBoxCenterAndExtent(PointRef.Point->Transform.TransformPosition(PointRef.Point->GetLocalCenter()), FVector::Zero()), [&IterationState, PointRef](const FPCGPointRef& InPointRef)
 			{
 				if ((PointRef.Point->Transform.GetLocation() - InPointRef.Point->Transform.GetLocation()).SquaredLength() <= SMALL_NUMBER &&
-					!ExclusionPoints.Contains(InPointRef.Point))
+				!IterationState.ExclusionPoints.Contains(InPointRef.Point))
 				{
-					ExcludedPoints.Add(InPointRef.Point);
+					IterationState.ExcludedPoints.Add(InPointRef.Point);
 				}
 			});
 		}
+
+		return true;
 	}
 
 	void Execute(FPCGContext* Context, EPCGSelfPruningType PruningType, float RadiusSimilarityFactor, bool bRandomizedPruning)
 	{
+		FPCGSelfPruningParameters Parameters;
+		Parameters.PruningType = PruningType;
+		Parameters.RadiusSimilarityFactor = RadiusSimilarityFactor;
+		Parameters.bRandomizedPruning = bRandomizedPruning;
+		Parameters.ComparisonSource.SetPointProperty(EPCGPointProperties::Extents);
+		return Execute(Context, Parameters);
+	}
+
+	void Execute(FPCGContext* Context, const FPCGSelfPruningParameters& InParameters)
+	{
+		check(Context);
+
 		// Early out: if pruning is disabled
-		if (PruningType == EPCGSelfPruningType::None)
+		if (InParameters.PruningType == EPCGSelfPruningType::None)
 		{
 			Context->OutputData = Context->InputData;
 			PCGE_LOG_C(Verbose, LogOnly, Context, LOCTEXT("TypeNotSpecified", "Skipped - Type is None"));
 			return;
 		}
 
-		const FVector::FReal RadiusEquality = 1.0f + RadiusSimilarityFactor;
-		const FVector::FReal SquaredRadiusEquality = FMath::Square(RadiusEquality);
-
 		TArray<FPCGTaggedData>& Outputs = Context->OutputData.TaggedData;
 
-		// TODO: embarassingly parallel loop
 		TArray<FPCGTaggedData> Inputs = Context->InputData.GetInputs();
 		for (const FPCGTaggedData& Input : Inputs)
 		{
@@ -152,93 +232,252 @@ namespace PCGSelfPruningElement
 			}
 
 			const UPCGPointData* InputPointData = SpatialInput->ToPointData(Context);
-			const TArray<FPCGPoint>& Points = InputPointData->GetPoints();
-			const UPCGPointData::PointOctree& Octree = InputPointData->GetOctree();
-
-			// Self-pruning will be done as follows:
-			// For each point:
-			//  if in its vicinity, there is >=1 non-rejected point with a radius significantly larger
-			//  or in its range + has a randomly assigned index -> we'll look at its seed
-			//  then remove this point
-			TArray<FPCGPointRef> SortedPoints;
-			SortedPoints.Reserve(Points.Num());
+			FIterationState IterationState;
+			IterationState.InputPointData = InputPointData;
+
+			while (!ExecuteSlice(IterationState, InParameters, Context)) {}
+
+			FPCGTaggedData& Output = Context->OutputData.TaggedData.Add_GetRef(Input);
+			if (IterationState.OutputPointData != nullptr)
+			{
+				Output.Data = IterationState.OutputPointData;
+			}
+		}
+	}
+
+	bool ExecuteSlice(FIterationState& InState, const FPCGSelfPruningParameters& InParameters, FPCGContext* InOptionalContext)
+	{
+		// Early out: if pruning is disabled
+		if (InParameters.PruningType == EPCGSelfPruningType::None)
+		{
+			if (InOptionalContext)
+			{
+				PCGE_LOG_C(Verbose, LogOnly, InOptionalContext, LOCTEXT("TypeNotSpecified", "Skipped - Type is None"));
+			}
+			return true;
+		}
+
+		const FVector::FReal RadiusEquality = 1.0f + InParameters.RadiusSimilarityFactor;
+		const FVector::FReal SquaredRadiusEquality = FMath::Square(RadiusEquality);
+		const TArray<FPCGPoint>& Points = InState.InputPointData->GetPoints();
+
+		// Force octree computation, and check if we need to stop after that
+		const UPCGPointData::PointOctree& Octree = InState.InputPointData->GetOctree();
+		if (ShouldStop(InOptionalContext))
+		{
+			return false;
+		}
+
+		// Self-pruning will be done as follows:
+		// For each point:
+		//  if in its vicinity, there is >=1 non-rejected point with a radius significantly larger
+		//  or in its range + has a randomly assigned index -> we'll look at its seed
+		//  then remove this point
+		if (!InState.bSortDone)
+		{
+			InState.SortedPoints.Empty();
+			InState.SortedPoints.Reserve(Points.Num());
 			for (const FPCGPoint& Point : Points)
 			{
-				SortedPoints.Add(FPCGPointRef(Point));
+				InState.SortedPoints.Add(FPCGPointRef(Point));
 			}
 
-			// Apply proper sort algorithm
-			if (PruningType == EPCGSelfPruningType::LargeToSmall)
+			FPCGAttributePropertySelector ComparisonSource = InParameters.ComparisonSource.CopyAndFixLast(InState.InputPointData);
+
+			// First sort the points if we need sorting. 
+			// Randomize for other cases than LargeToSmall and SmallToLarge are just a seed sort, so special case for that.
+			// Special case for extents too, since the SortByAttribute is slightly slower (~10%) due to the overhead of working with accessors.
+			// To not impact existing graphs (and since it is already one of the slowest operations in the graph), we will keep the old sorting directly on the extents.
+			// Note that sorting is not time sliced.
+			// TODO: Might be a good thing to have a future there.
+			bool bSortAttribute = InParameters.PruningType == EPCGSelfPruningType::LargeToSmall || InParameters.PruningType == EPCGSelfPruningType::SmallToLarge;
+			if (InParameters.bRandomizedPruning && !bSortAttribute)
+			{
+				TRACE_CPUPROFILER_EVENT_SCOPE(FPCGSelfPruningElement::Execute::SortingSeeds);
+				Algo::Sort(InState.SortedPoints, PCGSelfPruningAlgorithms::RandomSort);
+			}
+			else if (ComparisonSource.GetSelection() == EPCGAttributePropertySelection::PointProperty && ComparisonSource.GetPointProperty() == EPCGPointProperties::Extents && ComparisonSource.GetExtraNames().IsEmpty())
+			{
+				TRACE_CPUPROFILER_EVENT_SCOPE(FPCGSelfPruningElement::Execute::SortingExtents);
+				PCGSelfPruningAlgorithms::SortExtents(InState.SortedPoints, InParameters, SquaredRadiusEquality);
+			}
+			else
 			{
-				if (bRandomizedPruning)
+				TRACE_CPUPROFILER_EVENT_SCOPE(FPCGSelfPruningElement::Execute::SortingGeneric);
+
+				const bool bAscending = InParameters.PruningType != EPCGSelfPruningType::LargeToSmall;
+
+				TUniquePtr<const IPCGAttributeAccessor> InputAccessor = PCGAttributeAccessorHelpers::CreateConstAccessor(InState.InputPointData, ComparisonSource);
+
+				if (!InputAccessor.IsValid())
 				{
-					Algo::Sort(SortedPoints, [SquaredRadiusEquality](const FPCGPointRef& A, const FPCGPointRef& B) { return !PCGSelfPruningAlgorithms::SortSmallToLargeWithRandom(A, B, SquaredRadiusEquality); });
+					PCGLog::LogErrorOnGraph(FText::Format(LOCTEXT("InvalidAttribute", "Attribute/Property '{0}' was not found."), ComparisonSource.GetDisplayText()), InOptionalContext);
+					return true;
 				}
-				else
+
+				// For vector attributes, collapse them to their square length for comparison (was the previous behavior with extents).
+				if (PCG::Private::IsOfTypes<FVector2D, FVector, FVector4>(InputAccessor->GetUnderlyingType()))
 				{
-					Algo::Sort(SortedPoints, [SquaredRadiusEquality](const FPCGPointRef& A, const FPCGPointRef& B) { return !PCGSelfPruningAlgorithms::SortSmallToLargeNoRandom(A, B, SquaredRadiusEquality); });
+					// Force squared length extraction on vectors
+					ComparisonSource.GetExtraNamesMutable().Add(PCGAttributeExtractorConstants::VectorSquaredLength.ToString());
+					InputAccessor = PCGAttributeAccessorHelpers::CreateConstAccessor(InState.InputPointData, ComparisonSource);
 				}
-			}
-			else if (PruningType == EPCGSelfPruningType::SmallToLarge)
-			{
-				if (bRandomizedPruning)
+
+				if (!PCG::Private::IsOfTypes<int32, int64, float, double, uint8>(InputAccessor->GetUnderlyingType()))
 				{
-					Algo::Sort(SortedPoints, [SquaredRadiusEquality](const FPCGPointRef& A, const FPCGPointRef& B) { return PCGSelfPruningAlgorithms::SortSmallToLargeWithRandom(A, B, SquaredRadiusEquality); });
+					PCGLog::LogErrorOnGraph(FText::Format(LOCTEXT("InvalidAttribute", "Attribute/Property '{0}' is not numeric ({1})"), ComparisonSource.GetDisplayText(), PCG::Private::GetTypeNameText(InputAccessor->GetUnderlyingType())), InOptionalContext);
+					return true;
 				}
-				else
+
+				if (InParameters.bRandomizedPruning)
 				{
-					Algo::Sort(SortedPoints, [SquaredRadiusEquality](const FPCGPointRef& A, const FPCGPointRef& B) { return PCGSelfPruningAlgorithms::SortSmallToLargeNoRandom(A, B, SquaredRadiusEquality); });
+					auto CompareLessWithRandom = [SquaredRadiusEquality, &Points](const auto& A, const auto& B, int32 IndexA, int32 IndexB, bool)
+					{
+						static_assert(std::is_same_v<decltype(A), decltype(B)>);
+
+						using ValueType = std::remove_const_t<std::remove_reference_t<decltype(A)>>;
+
+						if constexpr (std::is_arithmetic_v<ValueType>)
+						{
+							ValueType ACopy = static_cast<ValueType>(static_cast<double>(A) * SquaredRadiusEquality);
+							if (PCG::Private::MetadataTraits<ValueType>::Less(ACopy, B))
+							{
+								return true;
+							}
+
+							ValueType BCopy = static_cast<ValueType>(static_cast<double>(B) * SquaredRadiusEquality);
+							if (PCG::Private::MetadataTraits<ValueType>::Less(BCopy, A))
+							{
+								return false;
+							}
+						}
+						else
+						{
+							if (PCG::Private::MetadataTraits<ValueType>::Less(A, B))
+							{
+								return true;
+							}
+							else if (PCG::Private::MetadataTraits<ValueType>::Less(B, A))
+							{
+								return false;
+							}
+						}
+
+						return Points[IndexA].Seed < Points[IndexB].Seed;
+					};
+
+					PCGAttributeAccessorHelpers::SortByAttribute(*InputAccessor, FPCGAttributeAccessorKeysPoints(Points), InState.SortedPoints, bAscending, PCGAttributeAccessorHelpers::Private::DefaultIndexGetter, std::move(CompareLessWithRandom));
 				}
-			}
-			else
-			{
-				if (bRandomizedPruning)
+				else
 				{
-					Algo::Sort(SortedPoints, PCGSelfPruningAlgorithms::RandomSort);
+					auto CompareLess = [SquaredRadiusEquality](const auto& A, const auto& B, int32, int32, bool)
+					{
+						static_assert(std::is_same_v<decltype(A), decltype(B)>);
+
+						using ValueType = std::remove_const_t<std::remove_reference_t<decltype(A)>>;
+						using CompareType = std::conditional_t<std::is_arithmetic_v<ValueType>, ValueType, const ValueType&>;
+
+						CompareType ACopy = A;
+
+						if constexpr (std::is_arithmetic_v<ValueType>)
+						{
+							ACopy = static_cast<ValueType>(static_cast<double>(A) * SquaredRadiusEquality);
+						}
+
+						return PCG::Private::MetadataTraits<ValueType>::Less(ACopy, B);
+					};
+
+					PCGAttributeAccessorHelpers::SortByAttribute(*InputAccessor, FPCGAttributeAccessorKeysPoints(Points), InState.SortedPoints, bAscending, PCGAttributeAccessorHelpers::Private::DefaultIndexGetter, std::move(CompareLess));
 				}
 			}
 
-			FPointBitSet ExclusionPoints(Points);
+			InState.ExclusionPoints.Initialize(Points);
+			InState.ExcludedPoints.Initialize(Points);
 
-			const bool bIsDuplicateTest = (PruningType == EPCGSelfPruningType::RemoveDuplicates);
+			InState.bSortDone = true;
 
-			if (bIsDuplicateTest)
-			{
-				PCGSelfPruningElement::DuplicatePointsExclusion(Points, SortedPoints, Octree, ExclusionPoints);
-			}
-			else
+			// After sorting, allow ourselves to stop if needed.
+			if (ShouldStop(InOptionalContext))
 			{
-				PCGSelfPruningElement::DensityBoundsExclusion(Points, SortedPoints, Octree, ExclusionPoints);
+				return false;
 			}
+		}
 
-			// Finally, output all points that are present in the ExclusionPoints.
-			FPCGTaggedData& Output = Outputs.Emplace_GetRef();
-			Output = Input;
+		// Find excluded/duplicate points. Time sliced.
+		const bool bIsDuplicateTest = (InParameters.PruningType == EPCGSelfPruningType::RemoveDuplicates);
+		bool bIsDone = false;
+		if (bIsDuplicateTest)
+		{
+			bIsDone = PCGSelfPruningElement::DuplicatePointsExclusion(InState, InOptionalContext);
+		}
+		else
+		{
+			bIsDone = PCGSelfPruningElement::DensityBoundsExclusion(InState, InOptionalContext);
+		}
 
+		// Finally, output all points that are present in the ExclusionPoints. This part is not time sliced, should it be too?
+		if (bIsDone)
+		{
 			UPCGPointData* PrunedData = NewObject<UPCGPointData>();
-			PrunedData->InitializeFromData(InputPointData);
-			Output.Data = PrunedData;
+			PrunedData->InitializeFromData(InState.InputPointData);
+			InState.OutputPointData = PrunedData;
 
 			TArray<FPCGPoint>& OutputPoints = PrunedData->GetMutablePoints();
 
 			for (const FPCGPoint& Point : Points)
 			{
-				if (ExclusionPoints.Contains(&Point))
+				if (InState.ExclusionPoints.Contains(&Point))
 				{
 					OutputPoints.Add(Point);
 				}
-			}				
-
-			if (bIsDuplicateTest)
-			{
-				PCGE_LOG_C(Verbose, LogOnly, Context, FText::Format(LOCTEXT("GenerationInfoDuplicate", "Removed {0} duplicate points from {1} source points"), Points.Num() - OutputPoints.Num(), Points.Num()));
 			}
-			else
+
+			if (InOptionalContext)
 			{
-				PCGE_LOG_C(Verbose, LogOnly, Context, FText::Format(LOCTEXT("GenerationInfo", "Generated {0} points from {1} source points"), OutputPoints.Num(), Points.Num()));
+				if (bIsDuplicateTest)
+				{
+					PCGE_LOG_C(Verbose, LogOnly, InOptionalContext, FText::Format(LOCTEXT("GenerationInfoDuplicate", "Removed {0} duplicate points from {1} source points"), Points.Num() - OutputPoints.Num(), Points.Num()));
+				}
+				else
+				{
+					PCGE_LOG_C(Verbose, LogOnly, InOptionalContext, FText::Format(LOCTEXT("GenerationInfo", "Generated {0} points from {1} source points"), OutputPoints.Num(), Points.Num()));
+				}
 			}
 		}
+
+		return bIsDone;
+	}
+}
+
+UPCGSelfPruningSettings::UPCGSelfPruningSettings()
+{
+	// Previous Default behavior was Extents
+	Parameters.ComparisonSource.SetPointProperty(EPCGPointProperties::Extents);
+}
+
+void UPCGSelfPruningSettings::PostLoad()
+{
+	Super::PostLoad();
+
+#if WITH_EDITOR
+	if (PruningType_DEPRECATED != EPCGSelfPruningType::LargeToSmall)
+	{
+		Parameters.PruningType = PruningType_DEPRECATED;
+		PruningType_DEPRECATED = EPCGSelfPruningType::LargeToSmall;
 	}
+
+	if (RadiusSimilarityFactor_DEPRECATED != 0.25f)
+	{
+		Parameters.RadiusSimilarityFactor = RadiusSimilarityFactor_DEPRECATED;
+		RadiusSimilarityFactor_DEPRECATED = 0.25;
+	}
+
+	if (!bRandomizedPruning_DEPRECATED)
+	{
+		Parameters.bRandomizedPruning = bRandomizedPruning_DEPRECATED;
+		bRandomizedPruning_DEPRECATED = true;
+	}
+#endif // WITH_EDITOR
 }
 
 FPCGElementPtr UPCGSelfPruningSettings::CreateElement() const
@@ -246,18 +485,76 @@ FPCGElementPtr UPCGSelfPruningSettings::CreateElement() const
 	return MakeShared<FPCGSelfPruningElement>();
 }
 
+bool FPCGSelfPruningElement::PrepareDataInternal(FPCGContext* InContext) const
+{
+	TRACE_CPUPROFILER_EVENT_SCOPE(FPCGSelfPruningElement::PrepareDataInternal);
+
+	FPCGSelfPruningElement::ContextType* TimeSlicedContext = static_cast<FPCGSelfPruningElement::ContextType*>(InContext);
+	check(TimeSlicedContext);
+
+	TimeSlicedContext->SetTimeSliceIsEnabled(true);
+
+	TArray<FPCGTaggedData> Inputs = InContext->InputData.GetInputsByPin(PCGPinConstants::DefaultInputLabel);
+
+	// No global execution state.
+	EPCGTimeSliceInitResult InitResult = TimeSlicedContext->InitializePerExecutionState();
+
+	TimeSlicedContext->InitializePerIterationStates(Inputs.Num(), [&Inputs, InContext](PCGSelfPruningElement::FIterationState& OutState, const FPCGSelfPruningElement::ExecStateType&, int32 Index) -> EPCGTimeSliceInitResult
+	{
+		const UPCGSpatialData* SpatialData = Cast<UPCGSpatialData>(Inputs[Index].Data);
+		if (!SpatialData)
+		{
+			PCGE_LOG_C(Error, GraphAndLog, InContext, FText::Format(LOCTEXT("InvalidInputDataType", "Input {0}: Input data must be of type Spatial"), FText::AsNumber(Index)));
+			return EPCGTimeSliceInitResult::NoOperation;
+		}
+
+		// Retro-compatibility
+		OutState.InputPointData = SpatialData->ToPointData();
+		if (!OutState.InputPointData)
+		{
+			return EPCGTimeSliceInitResult::NoOperation;
+		}
+
+		return EPCGTimeSliceInitResult::Success;
+	});
+
+	return true;
+}
+
 bool FPCGSelfPruningElement::ExecuteInternal(FPCGContext* Context) const
 {
 	TRACE_CPUPROFILER_EVENT_SCOPE(FPCGSelfPruningElement::Execute);
-	// TODO: time-sliced implementation
+
 	const UPCGSelfPruningSettings* Settings = Context->GetInputSettings<UPCGSelfPruningSettings>();
 	check(Settings);
 
-	const EPCGSelfPruningType PruningType = Settings->PruningType;
-	const float RadiusSimilarityFactor = Settings->RadiusSimilarityFactor;
-	const bool bRandomizedPruning = Settings->bRandomizedPruning;
+	FPCGSelfPruningElement::ContextType* TimeSlicedContext = static_cast<FPCGSelfPruningElement::ContextType*>(Context);
+	check(TimeSlicedContext);
+
+	// Prepare data failed, no need to execute. Return an empty output
+	if (!TimeSlicedContext->DataIsPreparedForExecution())
+	{
+		TimeSlicedContext->OutputData.TaggedData.Empty();
+		return true;
+	}
+
+	TArray<FPCGTaggedData> Inputs = Context->InputData.GetInputsByPin(PCGPinConstants::DefaultInputLabel);
+
+	return ExecuteSlice(TimeSlicedContext, [Settings, &Inputs](FPCGSelfPruningElement::ContextType* Context, const FPCGSelfPruningElement::ExecStateType&, PCGSelfPruningElement::FIterationState& IterState, const uint32 IterationIndex) -> bool
+	{
+		const bool bIsDone = PCGSelfPruningElement::ExecuteSlice(IterState, Settings->Parameters, Context);
+
+		if (bIsDone)
+		{
+			FPCGTaggedData& Output = Context->OutputData.TaggedData.Add_GetRef(Inputs[IterationIndex]);
+			if (IterState.OutputPointData != nullptr)
+			{
+				Output.Data = IterState.OutputPointData;
+			}
+		}
 
-	PCGSelfPruningElement::Execute(Context, PruningType, RadiusSimilarityFactor, bRandomizedPruning);
+		return bIsDone;
+	});
 
 	return true;
 }
diff --git a/Engine/Plugins/Experimental/PCG/Source/PCG/Private/Metadata/Accessors/PCGAttributeExtractor.cpp b/Engine/Plugins/Experimental/PCG/Source/PCG/Private/Metadata/Accessors/PCGAttributeExtractor.cpp
--- a/Engine/Plugins/Experimental/PCG/Source/PCG/Private/Metadata/Accessors/PCGAttributeExtractor.cpp
+++ b/Engine/Plugins/Experimental/PCG/Source/PCG/Private/Metadata/Accessors/PCGAttributeExtractor.cpp
@@ -70,6 +70,12 @@ namespace PCGAttributeExtractor
 			return MakeUnique<FPCGChainAccessor<double, VectorType>>(std::move(InAccessor),
 				[](const VectorType& Value) -> double { return Value.Size(); });
 		}
+		else if (Name == PCGAttributeExtractorConstants::VectorSquaredLength)
+		{
+			bOutSuccess = true;
+			return MakeUnique<FPCGChainAccessor<double, VectorType>>(std::move(InAccessor),
+				[](const VectorType& Value) -> double { return Value.SizeSquared(); });
+		}
 		else if (Name == PCGAttributeExtractorConstants::VectorNormalized)
 		{
 			if constexpr (std::is_same_v<FQuat, VectorType>)
diff --git a/Engine/Plugins/Experimental/PCG/Source/PCG/Public/Elements/PCGSelfPruning.h b/Engine/Plugins/Experimental/PCG/Source/PCG/Public/Elements/PCGSelfPruning.h
--- a/Engine/Plugins/Experimental/PCG/Source/PCG/Public/Elements/PCGSelfPruning.h
+++ b/Engine/Plugins/Experimental/PCG/Source/PCG/Public/Elements/PCGSelfPruning.h
@@ -2,7 +2,12 @@
 
 #pragma once
 
+#include "PCGContext.h"
 #include "PCGSettings.h"
+#include "Data/PCGPointData.h"
+#include "Elements/PCGTimeSlicedElementBase.h"
+
+#include "Containers/Array.h"
 
 #include "PCGSelfPruning.generated.h"
 
@@ -16,9 +21,60 @@ enum class EPCGSelfPruningType : uint8
 	RemoveDuplicates
 };
 
+USTRUCT(BlueprintType)
+struct FPCGSelfPruningParameters
+{
+	GENERATED_BODY()
+
+public:
+	UPROPERTY(BlueprintReadWrite, EditAnywhere, Category = Settings, meta = (PCG_Overridable))
+	EPCGSelfPruningType PruningType = EPCGSelfPruningType::LargeToSmall;
+
+	/** By default, it will prune according to the extents of each point, but you can provide another comparison like Density, or a dynamic attribute. Only support numeric values or vector (vector will be reduced to its squared length). */
+	UPROPERTY(BlueprintReadWrite, EditAnywhere, Category = Settings, meta = (EditCondition = "PruningType == EPCGSelfPruningType::LargeToSmall || PruningType == EPCGSelfPruningType::SmallToLarge", PCG_Overridable))
+	FPCGAttributePropertyInputSelector ComparisonSource;
+
+	/** Similarity factor to consider 2 points "equal". (For example, if a point extents squared length is 10 and factor is 0.25, all points between 7.5 and 12.5 will be considered "the same"). */
+	UPROPERTY(BlueprintReadWrite, EditAnywhere, Category = Settings, meta = (ClampMin = 0.0f, EditCondition = "PruningType == EPCGSelfPruningType::LargeToSmall || PruningType == EPCGSelfPruningType::SmallToLarge", PCG_Overridable))
+	float RadiusSimilarityFactor = 0.25f;
+
+	UPROPERTY(BlueprintReadWrite, EditAnywhere, Category = Settings, meta = (PCG_Overridable))
+	bool bRandomizedPruning = true;
+};
+
 namespace PCGSelfPruningElement
 {
+	// all points are in contiguous memory so this will create a bitset to just look up the index to set/test a bit
+	// this should reduce the 'set' cost to O[1]
+	struct PCG_API FPointBitSet
+	{
+		TArray<uint32> Bits;
+		const FPCGPoint* FirstPoint = nullptr;
+
+		void Initialize(const TArray<FPCGPoint>& Points);
+		void Add(const FPCGPoint* Point);
+		bool Contains(const FPCGPoint* Point);
+	};
+
+	struct FIterationState
+	{
+		const UPCGPointData* InputPointData = nullptr;
+		UPCGPointData* OutputPointData = nullptr;
+		TArray<FPCGPointRef> SortedPoints;
+		FPointBitSet ExcludedPoints;
+		FPointBitSet ExclusionPoints;
+		int32 CurrentPointIndex = 0;
+		bool bSortDone = false;
+	};
+
+	/** Will do the self pruning on all inputs in the context. Blocking call (no time slicing). Be careful, it can be very costly. */
 	void Execute(FPCGContext* Context, EPCGSelfPruningType PruningType, float RadiusSimilarityFactor, bool bRandomizedPruning);
+
+	/** Will do the self pruning on all inputs in the context. Blocking call (no time slicing). Be careful, it can be very costly. */
+	PCG_API void Execute(FPCGContext* Context, const FPCGSelfPruningParameters& InParameters);
+
+	/** Will execute self pruning according to the info from FIterationState. Time sliced call, will return true if it is done. */
+	PCG_API bool ExecuteSlice(FIterationState& InState, const FPCGSelfPruningParameters& InParameters, FPCGContext* InOptionalContext = nullptr);
 }
 
 UCLASS(BlueprintType, ClassGroup=(Procedural))
@@ -26,15 +82,15 @@ class PCG_API UPCGSelfPruningSettings : public UPCGSettings
 {
 	GENERATED_BODY()
 
-public:
-	UPROPERTY(BlueprintReadWrite, EditAnywhere, Category = Settings, meta = (PCG_Overridable))
-	EPCGSelfPruningType PruningType = EPCGSelfPruningType::LargeToSmall;
+	UPCGSelfPruningSettings();
 
-	UPROPERTY(BlueprintReadWrite, EditAnywhere, Category = Settings, meta=(ClampMin=0.0f, EditCondition="PruningType == EPCGSelfPruningType::LargeToSmall || PruningType == EPCGSelfPruningType::SmallToLarge", PCG_Overridable))
-	float RadiusSimilarityFactor = 0.25f;
+	//~Begin UObject interface
+	virtual void PostLoad() override;
+	//~End UObject interface
 
-	UPROPERTY(BlueprintReadWrite, EditAnywhere, Category = Settings, meta = (PCG_Overridable))
-	bool bRandomizedPruning = true;
+public:
+	UPROPERTY(BlueprintReadWrite, EditAnywhere, Category = Settings, meta = (ShowOnlyInnerProperties, PCG_Overridable))
+	FPCGSelfPruningParameters Parameters;
 
 #if WITH_EDITOR
 	//~Begin UPCGSettings interface
@@ -48,15 +104,23 @@ protected:
 	virtual TArray<FPCGPinProperties> OutputPinProperties() const override { return Super::DefaultPointOutputPinProperties(); }
 	virtual FPCGElementPtr CreateElement() const override;
 	// ~End UPCGSettings interface
+
+public:
+#if WITH_EDITORONLY_DATA
+	UPROPERTY()
+	EPCGSelfPruningType PruningType_DEPRECATED = EPCGSelfPruningType::LargeToSmall;
+
+	UPROPERTY()
+	float RadiusSimilarityFactor_DEPRECATED = 0.25f;
+
+	UPROPERTY()
+	bool bRandomizedPruning_DEPRECATED = true;
+#endif // WITH_EDITORONLY_DATA
 };
 
-class FPCGSelfPruningElement : public IPCGElement
+class FPCGSelfPruningElement : public TPCGTimeSlicedElementBase<PCGTimeSlice::FEmptyStruct, PCGSelfPruningElement::FIterationState>
 {
 protected:
+	virtual bool PrepareDataInternal(FPCGContext* Context) const override;
 	virtual bool ExecuteInternal(FPCGContext* Context) const override;
 };
-
-#if UE_ENABLE_INCLUDE_ORDER_DEPRECATED_IN_5_2
-#include "CoreMinimal.h"
-#include "PCGNode.h"
-#endif
diff --git a/Engine/Plugins/Experimental/PCG/Source/PCG/Public/Metadata/Accessors/PCGAttributeAccessorHelpers.h b/Engine/Plugins/Experimental/PCG/Source/PCG/Public/Metadata/Accessors/PCGAttributeAccessorHelpers.h
--- a/Engine/Plugins/Experimental/PCG/Source/PCG/Public/Metadata/Accessors/PCGAttributeAccessorHelpers.h
+++ b/Engine/Plugins/Experimental/PCG/Source/PCG/Public/Metadata/Accessors/PCGAttributeAccessorHelpers.h
@@ -6,6 +6,8 @@
 #include "Metadata/PCGMetadataAttributeTraits.h"
 #include "Metadata/Accessors/IPCGAttributeAccessor.h"
 #include "Metadata/Accessors/PCGAttributeAccessorKeys.h"
+
+#include "Containers/Array.h"
 #include "Templates/UniquePtr.h"
 #include "UObject/NameTypes.h"
 
@@ -73,23 +75,39 @@ namespace PCGAttributeAccessorHelpers
 	PCG_API TUniquePtr<const IPCGAttributeAccessorKeys> CreateConstKeys(const UPCGData* InData, const FPCGAttributePropertySelector& InSelector);
 	PCG_API TUniquePtr<IPCGAttributeAccessorKeys> CreateKeys(UPCGData* InData, const FPCGAttributePropertySelector& InSelector);
 
-	/**
-	* Sorts array given the accessors and keys of the array
-	*/
-	template <typename T>
-	void SortByAttribute(const IPCGAttributeAccessor& InAccessor, const IPCGAttributeAccessorKeys& InKeys, TArray<T>& InArray, bool bAscending)
+	namespace Private
 	{
-		SortByAttribute(InAccessor, InKeys, InArray, bAscending, [](int Index) { return Index; });
+		// Use a lambda to have this code more likely to be inlined in SortByAttribute.
+		static inline auto DefaultIndexGetter = [](int32 Index) -> int32 { return Index; };
+
+		// Use bAscending bool to know if you need to negate the condition for equal, since CompareDescending is !CompareAscending
+		template<typename T>
+		bool DefaultStableCompareLess(const T& A, const T& B, int32 IndexA, int32 IndexB, bool bAscending)
+		{
+			if (PCG::Private::MetadataTraits<T>::Equal(A, B))
+			{
+				return (bAscending == (IndexA < IndexB));
+			}
+
+			return PCG::Private::MetadataTraits<T>::Less(A, B);
+		}
+
+		// We need the lambda, because we can't use a templated function as a default parameter without specifying the template, but we can with a generic lambda.
+		static inline auto DefaultStableCompareLessLambda = [](const auto& A, const auto& B, int32 IndexA, int32 IndexB, bool bAscending) -> bool { return DefaultStableCompareLess(A, B, IndexA, IndexB, bAscending); };
 	}
 
 	/**
-	* Sorts array given the accessors and keys of the array.
-	* A custom function is used to extract the index of the elements of the array,
-	* and sort by the values associated with that index.
+	* Sorts array given the accessors and keys of the array. Sort is stable by default.
+	* A custom function can be used to extract the index of the elements of the array,
+	* and sort ascending or descending by the values associated with that index.
+	* CompareLess method can also be provided, method signature needs to follow DefaultStableCompareLess.
+	* We need the decltypes on both function types to allow to provide default values for both callbacks.
+	* Check DefaultIndexGetter and DefaultStableCompareLess for their signatures.
 	*/
-	template <typename T, typename Func>
-	void SortByAttribute(const IPCGAttributeAccessor& InAccessor, const IPCGAttributeAccessorKeys& InKeys, TArray<T>& InArray, bool bAscending, Func&& CustomGetIndex)
+	template <typename T, typename GetIndexFunc = decltype(Private::DefaultIndexGetter), typename CompareLessFunc = decltype(Private::DefaultStableCompareLessLambda)>
+	void SortByAttribute(const IPCGAttributeAccessor& InAccessor, const IPCGAttributeAccessorKeys& InKeys, TArray<T>& InArray, bool bAscending, GetIndexFunc CustomGetIndex = Private::DefaultIndexGetter, CompareLessFunc CompareLess = Private::DefaultStableCompareLessLambda)
 	{
+		TRACE_CPUPROFILER_EVENT_SCOPE(PCGAttributeAccessorHelpers::SortByAttribute);
 		check(InArray.Num() <= InKeys.GetNum())
 
 		if (InArray.IsEmpty())
@@ -97,7 +115,18 @@ namespace PCGAttributeAccessorHelpers
 			return;
 		}
 
-		auto Callback = [&InAccessor, &InKeys, &InArray, bAscending, &CustomGetIndex](auto Dummy)
+		// Prepare integer sequence
+		TArray<int32> ElementsIndexes;
+		{
+			TRACE_CPUPROFILER_EVENT_SCOPE(PCGAttributeAccessorHelpers::SortByAttribute::PrepareIndexes);
+			ElementsIndexes.Reserve(InArray.Num());
+			for (int i = 0; i < InArray.Num(); ++i)
+			{
+				ElementsIndexes.Add(i);
+			}
+		}
+
+		auto SortIndexes = [&InAccessor, &InKeys, &InArray, bAscending, &CustomGetIndex, &CompareLess, &ElementsIndexes](auto Dummy)
 		{
 			using ValueType = decltype(Dummy);
 
@@ -105,78 +134,63 @@ namespace PCGAttributeAccessorHelpers
 			{
 				TArray<ValueType> CachedValues;
 
-				if constexpr (std::is_trivially_copyable_v<ValueType>)
 				{
-					CachedValues.SetNumUninitialized(InKeys.GetNum());
+					TRACE_CPUPROFILER_EVENT_SCOPE(PCGAttributeAccessorHelpers::SortByAttribute::AllocateCachedValues);
+					if constexpr (std::is_trivially_copyable_v<ValueType>)
+					{
+						CachedValues.SetNumUninitialized(InKeys.GetNum());
+					}
+					else
+					{
+						CachedValues.SetNum(InKeys.GetNum());
+					}
 				}
-				else
+
 				{
-					CachedValues.SetNum(InKeys.GetNum());
+					TRACE_CPUPROFILER_EVENT_SCOPE(PCGAttributeAccessorHelpers::SortByAttribute::GatherCachedValues);
+					InAccessor.GetRange(TArrayView<ValueType>(CachedValues), 0, InKeys);
 				}
 
-				InAccessor.GetRange(TArrayView<ValueType>(CachedValues), 0, InKeys);
 
-				auto CompareAscending = [&CachedValues, &CustomGetIndex](int LHS, int RHS)
+				// Pass bAscending bool to the compare function to be able to negate the condition on equal, for it to be also stable in descending mode.
+				auto CompareAscending = [&CachedValues, &CustomGetIndex, &CompareLess, bAscending](int LHS, int RHS)
 				{
 					const int32 LHSIndex = CustomGetIndex(LHS);
 					const int32 RHSIndex = CustomGetIndex(RHS);
+					return CompareLess(CachedValues[LHSIndex], CachedValues[RHSIndex], LHSIndex, RHSIndex, bAscending);
+				};
 
-					const ValueType& LHSValue = CachedValues[LHSIndex];
-					const ValueType& RHSValue = CachedValues[RHSIndex];
+				auto CompareDescending = [&CompareAscending](int LHS, int RHS) { return !CompareAscending(LHS, RHS); };
 
-					if (PCG::Private::MetadataTraits<ValueType>::Equal(LHSValue, RHSValue))
+				{
+					TRACE_CPUPROFILER_EVENT_SCOPE(PCGAttributeAccessorHelpers::SortByAttribute::SortIndexes);
+					if (bAscending)
 					{
-						return LHSIndex < RHSIndex;
+						ElementsIndexes.Sort(CompareAscending);
 					}
-
-					return PCG::Private::MetadataTraits<ValueType>::Less(LHSValue, RHSValue);
-				};
-
-				auto CompareDescending = [&CachedValues, &CustomGetIndex](int LHS, int RHS)
-				{
-					const int32 LHSIndex = CustomGetIndex(LHS);
-					const int32 RHSIndex = CustomGetIndex(RHS);
-
-					const ValueType& LHSValue = CachedValues[LHSIndex];
-					const ValueType& RHSValue = CachedValues[RHSIndex];
-
-					if (PCG::Private::MetadataTraits<ValueType>::Equal(LHSValue, RHSValue))
+					else
 					{
-						return LHSIndex > RHSIndex;
+						ElementsIndexes.Sort(CompareDescending);
 					}
-
-					return PCG::Private::MetadataTraits<ValueType>::Greater(LHSValue, RHSValue);
-				};
-
-				// Fill integer sequence
-				TArray<int32> ElementsIndexes;
-				ElementsIndexes.Reserve(InArray.Num());
-				for (int i = 0; i < InArray.Num(); ++i)
-				{
-					ElementsIndexes.Add(i);
-				}
-
-				if (bAscending)
-				{
-					ElementsIndexes.Sort(CompareAscending);
-				}
-				else
-				{
-					ElementsIndexes.Sort(CompareDescending);
 				}
+			}
+		};
 
-				TArray<T> SortedArray;
-				SortedArray.Reserve(InArray.Num());
+		// Perform the sorting on the indexes using the attribute values provided by the accessor.
+		PCGMetadataAttribute::CallbackWithRightType(InAccessor.GetUnderlyingType(), SortIndexes);
 
-				for (int i = 0; i < InArray.Num(); ++i)
-				{
-					SortedArray.Add(MoveTemp(InArray[ElementsIndexes[i]]));
-				}
+		// Write back the values according to the sorted indexes.
+		TArray<T> SortedArray;
+		{
+			TRACE_CPUPROFILER_EVENT_SCOPE(PCGAttributeAccessorHelpers::SortByAttribute::WriteBackInSortedArray);
+			SortedArray.Reserve(InArray.Num());
 
-				InArray = MoveTemp(SortedArray);
+			for (int i = 0; i < InArray.Num(); ++i)
+			{
+				SortedArray.Add(MoveTemp(InArray[ElementsIndexes[i]]));
 			}
-		};
+		}
 
-		PCGMetadataAttribute::CallbackWithRightType(InAccessor.GetUnderlyingType(), Callback);
+		InArray = MoveTemp(SortedArray);
 	}
 }
diff --git a/Engine/Plugins/Experimental/PCG/Source/PCG/Public/Metadata/Accessors/PCGAttributeExtractor.h b/Engine/Plugins/Experimental/PCG/Source/PCG/Public/Metadata/Accessors/PCGAttributeExtractor.h
--- a/Engine/Plugins/Experimental/PCG/Source/PCG/Public/Metadata/Accessors/PCGAttributeExtractor.h
+++ b/Engine/Plugins/Experimental/PCG/Source/PCG/Public/Metadata/Accessors/PCGAttributeExtractor.h
@@ -19,6 +19,7 @@ namespace PCGAttributeExtractorConstants
 	const FName VectorB = TEXT("B");
 	const FName VectorA = TEXT("A");
 	const FName VectorLength = TEXT("Length");
+	const FName VectorSquaredLength = TEXT("SquaredLength");
 	const FName VectorSize = TEXT("Size");
 	const FName VectorNormalized = TEXT("Normalized");
 
@@ -34,6 +35,8 @@ namespace PCGAttributeExtractorConstants
 		VectorA,
 		VectorLength,
 		VectorSize,
+		VectorSquaredLength,
+		VectorNormalized
 	};
 
 	// Transform
diff --git a/Engine/Plugins/Experimental/PCG/Source/PCG/Public/Metadata/PCGAttributePropertySelector.h b/Engine/Plugins/Experimental/PCG/Source/PCG/Public/Metadata/PCGAttributePropertySelector.h
--- a/Engine/Plugins/Experimental/PCG/Source/PCG/Public/Metadata/PCGAttributePropertySelector.h
+++ b/Engine/Plugins/Experimental/PCG/Source/PCG/Public/Metadata/PCGAttributePropertySelector.h
@@ -54,6 +54,7 @@ public:
 	// Getters
 	EPCGAttributePropertySelection GetSelection() const { return Selection; }
 	const TArray<FString>& GetExtraNames() const { return ExtraNames; }
+	TArray<FString>& GetExtraNamesMutable() { return ExtraNames; }
 	FName GetAttributeName() const { return AttributeName; }
 	EPCGPointProperties GetPointProperty() const { return PointProperty; }
 	EPCGExtraProperties GetExtraProperty() const { return ExtraProperty; }
