[PCG] Fix localization

#rnx

diff --git a/Engine/Plugins/Experimental/PCG/Source/PCG/Private/Elements/PCGSelfPruning.cpp b/Engine/Plugins/Experimental/PCG/Source/PCG/Private/Elements/PCGSelfPruning.cpp
--- a/Engine/Plugins/Experimental/PCG/Source/PCG/Private/Elements/PCGSelfPruning.cpp
+++ b/Engine/Plugins/Experimental/PCG/Source/PCG/Private/Elements/PCGSelfPruning.cpp
@@ -325,7 +325,7 @@ namespace PCGSelfPruningElement
 
 				if (!PCG::Private::IsOfTypes<int32, int64, float, double, uint8>(InputAccessor->GetUnderlyingType()))
 				{
-					PCGLog::LogErrorOnGraph(FText::Format(LOCTEXT("InvalidAttribute", "Attribute/Property '{0}' is not numeric ({1})"), ComparisonSource.GetDisplayText(), PCG::Private::GetTypeNameText(InputAccessor->GetUnderlyingType())), InOptionalContext);
+					PCGLog::LogErrorOnGraph(FText::Format(LOCTEXT("NotNumericAttribute", "Attribute/Property '{0}' is not numeric ({1})"), ComparisonSource.GetDisplayText(), PCG::Private::GetTypeNameText(InputAccessor->GetUnderlyingType())), InOptionalContext);
 					return true;
 				}
 
