Changed the way we sample on the landscape to have better intersection results. This makes the landscape data consistent with the world query results as well.

#rnx

diff --git a/Engine/Plugins/Experimental/PCG/Source/PCG/Private/Data/PCGLandscapeData.cpp b/Engine/Plugins/Experimental/PCG/Source/PCG/Private/Data/PCGLandscapeData.cpp
--- a/Engine/Plugins/Experimental/PCG/Source/PCG/Private/Data/PCGLandscapeData.cpp
+++ b/Engine/Plugins/Experimental/PCG/Source/PCG/Private/Data/PCGLandscapeData.cpp
@@ -173,17 +173,59 @@ FBox UPCGLandscapeData::GetStrictBounds() const
 
 bool UPCGLandscapeData::SamplePoint(const FTransform& InTransform, const FBox& InBounds, FPCGPoint& OutPoint, UPCGMetadata* OutMetadata) const
 {
-	// The point is in/on the shape if it coincides with its projection. I.e. projecting on landscape does not move the point. Implementing
-	// this way shares the sampling code.
-	if (ProjectPoint(InTransform, InBounds, {}, OutPoint, OutMetadata))
+	const ULandscapeInfo* LandscapeInfo = GetLandscapeInfo(InTransform.GetLocation());
+	if (!LandscapeInfo || !LandscapeInfo->GetLandscapeProxy())
 	{
-		if (InBounds.IsValid)
+		return false;
+	}
+
+#if WITH_EDITOR
+	const FTransform& LandscapeTransform = LandscapeInfo->GetLandscapeProxy()->GetTransform();
+#else
+	const FTransform LandscapeTransform = LandscapeInfo->GetLandscapeProxy()->LandscapeActorToWorld();
+#endif
+
+	// Box in local space -> box in world space -> box in landscape space
+	const FTransform BoundsTransformInLanscapeSpace = InTransform.GetRelativeTransform(LandscapeTransform);
+	FBox BoundsInLanscapeSpace = InBounds.TransformBy(BoundsTransformInLanscapeSpace);
+
+	// Gather all landscape heightfield components we need to test
+	const int ComponentMapKeyMinX = FMath::FloorToInt(BoundsInLanscapeSpace.Min.X / LandscapeInfo->ComponentSizeQuads);
+	const int ComponentMapKeyMaxX = FMath::FloorToInt(BoundsInLanscapeSpace.Max.X / LandscapeInfo->ComponentSizeQuads);
+	const int ComponentMapKeyMinY = FMath::FloorToInt(BoundsInLanscapeSpace.Min.Y / LandscapeInfo->ComponentSizeQuads);
+	const int ComponentMapKeyMaxY = FMath::FloorToInt(BoundsInLanscapeSpace.Max.Y / LandscapeInfo->ComponentSizeQuads);
+
+	TArray<ULandscapeHeightfieldCollisionComponent*, TInlineAllocator<1>> LandscapeCollisionComponents;
+
+	for (int X = ComponentMapKeyMinX; X <= ComponentMapKeyMaxX; ++X)
+	{
+		for (int Y = ComponentMapKeyMinY; Y <= ComponentMapKeyMaxY; ++Y)
 		{
-			return FMath::PointBoxIntersection(OutPoint.Transform.GetLocation(), InBounds.TransformBy(InTransform));
+			ULandscapeHeightfieldCollisionComponent* CollisionComponent = LandscapeInfo->XYtoCollisionComponentMap.FindRef(FIntPoint(X, Y));
+
+			if (CollisionComponent)
+			{
+				LandscapeCollisionComponents.AddUnique(CollisionComponent);
+			}
 		}
-		else
+	}
+
+	FCollisionShape CollisionShape;
+	if (!LandscapeCollisionComponents.IsEmpty())
+	{
+		CollisionShape.SetBox(FVector3f(InBounds.GetExtent() * InTransform.GetScale3D())); 
+	}
+
+	// Test collision against all gathered collision components
+	for (ULandscapeHeightfieldCollisionComponent* Component : LandscapeCollisionComponents)
+	{
+		check(Component);
+		TArray<FOverlapResult> OutOverlap;
+		if (Component->OverlapComponentWithResult(InTransform.GetLocation(), InTransform.GetRotation(), CollisionShape, OutOverlap))
 		{
-			return (InTransform.GetLocation() - OutPoint.Transform.GetLocation()).SquaredLength() < UE_SMALL_NUMBER;
+			new(&OutPoint) FPCGPoint(InTransform, /*Density=*/1.0f, /*Seed=*/0);
+			OutPoint.SetLocalBounds(InBounds);
+			return true;
 		}
 	}
 
