[PCG] Prevent RuntimeGen from dirtying the level
* Runtime PAs are now renamed with REN_DoNotDirty
* Created a second map of grid-size to grid-guid to track transient guids

#rnx

diff --git a/Engine/Plugins/Experimental/PCG/Source/PCG/Private/PCGActorAndComponentMapping.cpp b/Engine/Plugins/Experimental/PCG/Source/PCG/Private/PCGActorAndComponentMapping.cpp
--- a/Engine/Plugins/Experimental/PCG/Source/PCG/Private/PCGActorAndComponentMapping.cpp
+++ b/Engine/Plugins/Experimental/PCG/Source/PCG/Private/PCGActorAndComponentMapping.cpp
@@ -608,12 +608,18 @@ void FPCGActorAndComponentMapping::ForAllIntersectingPartitionActors(const FBox&
 	};
 
 	PCGHiGenGrid::FSizeToGuidMap GridSizeToGuid;
-	PCGWorldActor->GetGridGuids(GridSizeToGuid);
+	PCGWorldActor->GetSerializedGridGuids(GridSizeToGuid);
 	for (const TPair<uint32, FGuid>& SizeAndGuid : GridSizeToGuid)
 	{
 		const uint32 GridSize = SizeAndGuid.Key;
-
 		ForAllIntersectingPartitionActorsOfGridSize(PartitionActorsMap, PartitionActorsMapLock, GridSize);
+	}
+
+	GridSizeToGuid.Empty();
+	PCGWorldActor->GetTransientGridGuids(GridSizeToGuid);
+	for (const TPair<uint32, FGuid>& SizeAndGuid : GridSizeToGuid)
+	{
+		const uint32 GridSize = SizeAndGuid.Key;
 		ForAllIntersectingPartitionActorsOfGridSize(RuntimeGenPartitionActorsMap, RuntimeGenPartitionActorsMapLock, GridSize);
 	}
 }
diff --git a/Engine/Plugins/Experimental/PCG/Source/PCG/Private/PCGSubsystem.cpp b/Engine/Plugins/Experimental/PCG/Source/PCG/Private/PCGSubsystem.cpp
--- a/Engine/Plugins/Experimental/PCG/Source/PCG/Private/PCGSubsystem.cpp
+++ b/Engine/Plugins/Experimental/PCG/Source/PCG/Private/PCGSubsystem.cpp
@@ -861,10 +861,10 @@ namespace PCGSubsystem
 		APCGWorldActor* PCGWorldActor = PCGHelpers::GetPCGWorldActor(World);
 		check(PCGWorldActor);
 
-		PCGWorldActor->CreateGridGuidsIfNecessary(GridSizes);
+		PCGWorldActor->CreateGridGuidsIfNecessary(GridSizes, /*bAreGridsSerialized=*/true);
 
 		PCGHiGenGrid::FSizeToGuidMap GridSizeToGuid;
-		PCGWorldActor->GetGridGuids(GridSizeToGuid);
+		PCGWorldActor->GetSerializedGridGuids(GridSizeToGuid);
 		
 		TArray<FPCGTaskId> CellTasks;
 		for (uint32 GridSize : GridSizes)
diff --git a/Engine/Plugins/Experimental/PCG/Source/PCG/Private/PCGWorldActor.cpp b/Engine/Plugins/Experimental/PCG/Source/PCG/Private/PCGWorldActor.cpp
--- a/Engine/Plugins/Experimental/PCG/Source/PCG/Private/PCGWorldActor.cpp
+++ b/Engine/Plugins/Experimental/PCG/Source/PCG/Private/PCGWorldActor.cpp
@@ -130,16 +130,19 @@ void APCGWorldActor::EndPlay(const EEndPlayReason::Type EndPlayReason)
 	Super::EndPlay(EndPlayReason);
 }
 
-void APCGWorldActor::CreateGridGuidsIfNecessary(const PCGHiGenGrid::FSizeArray& InGridSizes)
+void APCGWorldActor::CreateGridGuidsIfNecessary(const PCGHiGenGrid::FSizeArray& InGridSizes, bool bAreGridsSerialized)
 {
+	TMap<uint32, FGuid>& GuidsMap = bAreGridsSerialized ? GridGuids : TransientGridGuids;
+	FRWLock& GuidsLock = bAreGridsSerialized ? GridGuidsLock : TransientGridGuidsLock;
+
 	// Check if any need adding
 	ensure(!InGridSizes.IsEmpty());
 	PCGHiGenGrid::FSizeArray GridSizesToAdd;
 	{
-		FReadScopeLock ReadLock(GridGuidsLock);
+		FReadScopeLock ReadLock(GuidsLock);
 		for (uint32 GridSize : InGridSizes)
 		{
-			if (!GridGuids.Contains(GridSize))
+			if (!GuidsMap.Contains(GridSize))
 			{
 				GridSizesToAdd.Push(GridSize);
 			}
@@ -148,19 +151,19 @@ void APCGWorldActor::CreateGridGuidsIfNecessary(const PCGHiGenGrid::FSizeArray&
 
 	if (GridSizesToAdd.Num() > 0)
 	{
-		FWriteScopeLock WriteLock(GridGuidsLock);
+		FWriteScopeLock WriteLock(GuidsLock);
 
 		bool bModified = false;
 		for (uint32 GridSize : GridSizesToAdd)
 		{
-			if (!GridGuids.Contains(GridSize))
+			if (!GuidsMap.Contains(GridSize))
 			{
-				GridGuids.Add(GridSize, FGuid::NewGuid());
+				GuidsMap.Add(GridSize, FGuid::NewGuid());
 				bModified = true;
 			}
 		}
 
-		if (bModified)
+		if (bModified && bAreGridsSerialized)
 		{
 			// Set dirty flag if we added guids. Unfortunately if the guids are not up to date, this will produce save prompts
 			// to users. However, this was needed to ensure the guids are saved - without this guids were lost and PAs were leaked.
@@ -173,14 +176,14 @@ void APCGWorldActor::CreateGridGuidsIfNecessary(const PCGHiGenGrid::FSizeArray&
 				PCGSubsystem->ScheduleGeneric([this]()
 				{
 					this->MarkPackageDirty();
-				return true;
+					return true;
 				}, nullptr, {});
 			}
 		}
 	}
 }
 
-void APCGWorldActor::GetGridGuids(PCGHiGenGrid::FSizeToGuidMap& OutSizeToGuidMap) const
+void APCGWorldActor::GetSerializedGridGuids(PCGHiGenGrid::FSizeToGuidMap& OutSizeToGuidMap) const
 {
 	FReadScopeLock ReadLock(GridGuidsLock);
 	for (const TPair<uint32, FGuid>& SizeGuid : GridGuids)
@@ -189,6 +192,15 @@ void APCGWorldActor::GetGridGuids(PCGHiGenGrid::FSizeToGuidMap& OutSizeToGuidMap
 	}
 }
 
+void APCGWorldActor::GetTransientGridGuids(PCGHiGenGrid::FSizeToGuidMap& OutSizeToGuidMap) const
+{
+	FReadScopeLock ReadLock(TransientGridGuidsLock);
+	for (const TPair<uint32, FGuid>& SizeGuid : TransientGridGuids)
+	{
+		OutSizeToGuidMap.Add(SizeGuid.Key, SizeGuid.Value);
+	}
+}
+
 void APCGWorldActor::MergeFrom(APCGWorldActor* OtherWorldActor)
 {
 	check(OtherWorldActor && this != OtherWorldActor);
diff --git a/Engine/Plugins/Experimental/PCG/Source/PCG/Private/RuntimeGen/PCGRuntimeGenScheduler.cpp b/Engine/Plugins/Experimental/PCG/Source/PCG/Private/RuntimeGen/PCGRuntimeGenScheduler.cpp
--- a/Engine/Plugins/Experimental/PCG/Source/PCG/Private/RuntimeGen/PCGRuntimeGenScheduler.cpp
+++ b/Engine/Plugins/Experimental/PCG/Source/PCG/Private/RuntimeGen/PCGRuntimeGenScheduler.cpp
@@ -854,7 +854,7 @@ void FPCGRuntimeGenScheduler::CleanupComponent(const FGridGenerationKey& Generat
 			}
 
 #if WITH_EDITOR
-			PartitionActor->Rename(nullptr, PartitionActor->GetOuter(), REN_NonTransactional);
+			PartitionActor->Rename(nullptr, PartitionActor->GetOuter(), REN_NonTransactional | REN_DoNotDirty);
 			PartitionActor->SetActorLabel(*PCGRuntimeGenSchedulerConstants::PooledPartitionActorName);
 #endif
 			PartitionActorPool.Push(PartitionActor);
@@ -867,7 +867,7 @@ void FPCGRuntimeGenScheduler::CleanupComponent(const FGridGenerationKey& Generat
 			}
 
 #if WITH_EDITOR
-			PartitionActor->Rename(nullptr, PartitionActor->GetOuter(), REN_NonTransactional);
+			PartitionActor->Rename(nullptr, PartitionActor->GetOuter(), REN_NonTransactional | REN_DoNotDirty);
 #endif
 			World->DestroyActor(PartitionActor);
 		}
@@ -1071,7 +1071,7 @@ APCGPartitionActor* FPCGRuntimeGenScheduler::GetPartitionActorFromPool(uint32 Gr
 #if WITH_EDITOR
 	const FName ActorName = *APCGPartitionActor::GetRuntimeGenActorName(GridSize, GridCoords);
 
-	PartitionActor->Rename(*ActorName.ToString(), PartitionActor->GetOuter(), REN_NonTransactional);
+	PartitionActor->Rename(*ActorName.ToString(), PartitionActor->GetOuter(), REN_NonTransactional | REN_DoNotDirty);
 	PartitionActor->SetActorLabel(ActorName.ToString());
 #endif
 
@@ -1124,7 +1124,7 @@ void FPCGRuntimeGenScheduler::ResetPartitionActorPoolToSize(uint32 NewPoolSize)
 	for (APCGPartitionActor* PartitionActor : PartitionActorPool)
 	{
 #if WITH_EDITOR
-		PartitionActor->Rename(nullptr, PartitionActor->GetOuter(), REN_NonTransactional);
+		PartitionActor->Rename(nullptr, PartitionActor->GetOuter(), REN_NonTransactional | REN_DoNotDirty);
 #endif
 		World->DestroyActor(PartitionActor);
 	}
@@ -1144,8 +1144,7 @@ void FPCGRuntimeGenScheduler::CreateGridGuidsForComponent(UPCGComponent* InCompo
 			PCGHiGenGrid::FSizeArray GridSizes;
 			ensure(PCGHelpers::GetGenerationGridSizes(InComponent->GetGraph(), PCGWorldActor, GridSizes, bHasUnbounded));
 
-			// TODO: This will mark the package as dirty, which should not happen from runtime gen.
-			PCGWorldActor->CreateGridGuidsIfNecessary(GridSizes);
+			PCGWorldActor->CreateGridGuidsIfNecessary(GridSizes, /*bAreGridsSerialized=*/false);
 		}
 	}
 }
diff --git a/Engine/Plugins/Experimental/PCG/Source/PCG/Public/PCGWorldActor.h b/Engine/Plugins/Experimental/PCG/Source/PCG/Public/PCGWorldActor.h
--- a/Engine/Plugins/Experimental/PCG/Source/PCG/Public/PCGWorldActor.h
+++ b/Engine/Plugins/Experimental/PCG/Source/PCG/Public/PCGWorldActor.h
@@ -25,9 +25,14 @@ public:
 	virtual void BeginPlay() override;
 	virtual void EndPlay(const EEndPlayReason::Type EndPlayReason) override;
 
-	/** Returns the grid GUIDs used for the partitioned actors, one per grid size. */
-	void CreateGridGuidsIfNecessary(const PCGHiGenGrid::FSizeArray& InGridSizes);
-	void GetGridGuids(PCGHiGenGrid::FSizeToGuidMap& OutSizeToGuidMap) const;
+	/** Creates guids for unused grid sizes. */
+	void CreateGridGuidsIfNecessary(const PCGHiGenGrid::FSizeArray& InGridSizes, bool bAreGridsSerialized);
+
+	/** Returns the serialized grid GUIDs used for the partitioned actors, one per grid size. */
+	void GetSerializedGridGuids(PCGHiGenGrid::FSizeToGuidMap& OutSizeToGuidMap) const;
+
+	/** Returns the transient grid GUIDs used for the partitioned actors, one per grid size. */
+	void GetTransientGridGuids(PCGHiGenGrid::FSizeToGuidMap& OutSizeToGuidMap) const;
 
 	void MergeFrom(APCGWorldActor* OtherWorldActor);
 
@@ -78,10 +83,15 @@ private:
 	void OnPartitionGridSizeChanged();
 #endif
 
-	/** GUIDs of the partitioned actor grids, one per grid size. */
+	/** GUIDs of the serialized partitioned actor grids, one per grid size. */
 	UPROPERTY()
 	TMap<uint32, FGuid> GridGuids;
 	mutable FRWLock GridGuidsLock;
+
+	/** GUIDs of the transient partitioned actor grids, one per grid size. */
+	UPROPERTY(Transient)
+	TMap<uint32, FGuid> TransientGridGuids;
+	mutable FRWLock TransientGridGuidsLock;
 };
 
 #if UE_ENABLE_INCLUDE_ORDER_DEPRECATED_IN_5_2
