[PCG] Add ApplyOnActor node
* Applies property overrides and executes provided functions on a target actor

diff --git a/Engine/Plugins/Experimental/PCG/Source/PCG/Private/Elements/PCGApplyOnActor.cpp b/Engine/Plugins/Experimental/PCG/Source/PCG/Private/Elements/PCGApplyOnActor.cpp
new file mode 100644
--- /dev/null
+++ b/Engine/Plugins/Experimental/PCG/Source/PCG/Private/Elements/PCGApplyOnActor.cpp
@@ -0,0 +1,59 @@
+// Copyright Epic Games, Inc. All Rights Reserved.
+
+#include "Elements/PCGApplyOnActor.h"
+
+#include "PCGContext.h"
+#include "PCGModule.h"
+#include "Helpers/PCGHelpers.h"
+
+#include "GameFramework/Actor.h"
+
+#define LOCTEXT_NAMESPACE "PCGApplyOnActorElement"
+
+namespace PCGApplyOnActorConstants
+{
+	const FText DependencyTooltip = LOCTEXT("DependencyPinTooltip", "Data passed to this pin will be used to order execution but will otherwise not contribute to the results of this node.");
+	const FName ActorPropertyOverridesLabel = TEXT("Property Overrides");
+	const FText ActorPropertyOverridesTooltip = LOCTEXT("ActorOverrideToolTip", "Provide property overrides for the target actor. The attribute name must match the InputSource name in the actor property override description.");
+}
+
+TArray<FPCGPinProperties> UPCGApplyOnActorSettings::InputPinProperties() const
+{
+	TArray<FPCGPinProperties> PinProperties;
+	PinProperties.Emplace(PCGPinConstants::DefaultDependencyOnlyLabel, EPCGDataType::Any, /*bAllowMultipleConnections=*/true, /*bAllowMultipleData=*/true, PCGApplyOnActorConstants::DependencyTooltip);
+	PinProperties.Add(PCGActorPropertyOverrideHelpers::CreateActorPropertiesOverridePin(PCGApplyOnActorConstants::ActorPropertyOverridesLabel, PCGApplyOnActorConstants::ActorPropertyOverridesTooltip));
+	return PinProperties;
+}
+
+FPCGElementPtr UPCGApplyOnActorSettings::CreateElement() const
+{
+	return MakeShared<FPCGApplyOnActorElement>();
+}
+
+bool FPCGApplyOnActorElement::ExecuteInternal(FPCGContext* Context) const
+{
+	TRACE_CPUPROFILER_EVENT_SCOPE(FPCGApplyOnActorElement::Execute);
+
+	check(Context);
+
+	const UPCGApplyOnActorSettings* Settings = Context->GetInputSettings<UPCGApplyOnActorSettings>();
+	check(Settings);
+
+	AActor* TargetActor = Settings->TargetActor.Get() ? Settings->TargetActor.Get() : Context->GetTargetActor(nullptr);
+	if (!TargetActor)
+	{
+		PCGE_LOG(Error, GraphAndLog, LOCTEXT("InvalidTargetActor", "Invalid target actor."));
+		return true;
+	}
+
+	PCGActorPropertyOverrideHelpers::ApplyOverridesFromParams(Settings->PropertyOverrideDescriptions, TargetActor, PCGApplyOnActorConstants::ActorPropertyOverridesLabel, Context);
+
+	for (UFunction* Function : PCGHelpers::FindUserFunctions(TargetActor->GetClass(), Settings->PostProcessFunctionNames, Context))
+	{
+		TargetActor->ProcessEvent(Function, nullptr);
+	}
+
+	return true;
+}
+
+#undef LOCTEXT_NAMESPACE
diff --git a/Engine/Plugins/Experimental/PCG/Source/PCG/Public/Elements/PCGApplyOnActor.h b/Engine/Plugins/Experimental/PCG/Source/PCG/Public/Elements/PCGApplyOnActor.h
new file mode 100644
--- /dev/null
+++ b/Engine/Plugins/Experimental/PCG/Source/PCG/Public/Elements/PCGApplyOnActor.h
@@ -0,0 +1,55 @@
+// Copyright Epic Games, Inc. All Rights Reserved.
+
+#pragma once
+
+#include "PCGSettings.h"
+#include "Metadata/PCGActorPropertyOverride.h"
+
+#include "PCGApplyOnActor.generated.h"
+
+class AActor;
+
+/**
+* Apply property overrides and executes functions on a target actor.
+*/
+UCLASS(BlueprintType)
+class UPCGApplyOnActorSettings : public UPCGSettings
+{
+	GENERATED_BODY()
+
+public:
+	//~Begin UPCGSettings interface
+#if WITH_EDITOR
+	virtual FName GetDefaultNodeName() const override { return FName(TEXT("ApplyOnActor")); }
+	virtual FText GetDefaultNodeTitle() const override { return NSLOCTEXT("PCGApplyOnActorElement", "NodeTitle", "Apply On Actor"); }
+	virtual FText GetNodeTooltipText() const override { return NSLOCTEXT("PCGApplyOnActorElement", "NodeTooltip", "Applies property overrides and executes functions on a target actor."); }
+	virtual EPCGSettingsType GetType() const override { return EPCGSettingsType::Generic; }
+#endif
+
+protected:
+	virtual TArray<FPCGPinProperties> InputPinProperties() const override;
+	virtual TArray<FPCGPinProperties> OutputPinProperties() const override { return TArray<FPCGPinProperties>(); }
+	virtual FPCGElementPtr CreateElement() const override;
+	//~End UPCGSettings interface
+
+public:
+	UPROPERTY(BlueprintReadWrite, Category = Settings, meta = (PCG_Overridable))
+	TSoftObjectPtr<AActor> TargetActor;
+
+	/** Override the default property values on the target actor. Applied before post-process functions. */
+	UPROPERTY(EditAnywhere, BlueprintReadWrite, Category = Settings)
+	TArray<FPCGActorPropertyOverrideDescription> PropertyOverrideDescriptions;
+
+	/** Specify a list of functions to be called on the target actor. Functions need to be parameter-less and with "CallInEditor" flag enabled. */
+	UPROPERTY(BlueprintReadWrite, EditAnywhere, Category = Settings)
+	TArray<FName> PostProcessFunctionNames;
+};
+
+class FPCGApplyOnActorElement : public IPCGElement
+{
+public:
+	virtual bool IsCacheable(const UPCGSettings* InSettings) const override { return false; }
+
+protected:
+	virtual bool ExecuteInternal(FPCGContext* Context) const override;
+};
