Removed deprecated pins by default on input node.

diff --git a/Engine/Plugins/Experimental/PCG/Source/PCG/Private/Elements/PCGIndirectionElement.cpp b/Engine/Plugins/Experimental/PCG/Source/PCG/Private/Elements/PCGIndirectionElement.cpp
--- a/Engine/Plugins/Experimental/PCG/Source/PCG/Private/Elements/PCGIndirectionElement.cpp
+++ b/Engine/Plugins/Experimental/PCG/Source/PCG/Private/Elements/PCGIndirectionElement.cpp
@@ -63,6 +63,7 @@ FName UPCGIndirectionSettings::AdditionalTaskName() const
 TArray<FPCGPinProperties> UPCGIndirectionSettings::InputPinProperties() const
 {
 	TArray<FPCGPinProperties> InputProperties;
+	bool bSetProperties = false;
 
 	switch (ProxyInterfaceMode)
 	{
@@ -71,6 +72,7 @@ TArray<FPCGPinProperties> UPCGIndirectionSettings::InputPinProperties() const
 			{
 				const UPCGSettings* SettingsDefaultObject = CastChecked<UPCGSettings>(SettingsClass->GetDefaultObject());
 				InputProperties = SettingsDefaultObject->DefaultInputPinProperties();
+				bSetProperties = true;
 			}
 			break;
 		case EPCGProxyInterfaceMode::ByBlueprintElement:
@@ -78,12 +80,14 @@ TArray<FPCGPinProperties> UPCGIndirectionSettings::InputPinProperties() const
 			{
 				const UPCGBlueprintElement* BlueprintSettingsDefaultObject = CastChecked<UPCGBlueprintElement>(BlueprintElementClass->GetDefaultObject());
 				InputProperties = BlueprintSettingsDefaultObject->GetInputPins();
+				bSetProperties = true;
 			}
 			break;
 		case EPCGProxyInterfaceMode::BySettings:
 			if (const UPCGSettings* SettingsPtr = Settings.LoadSynchronous())
 			{
 				InputProperties = SettingsPtr->DefaultInputPinProperties();
+				bSetProperties = true;
 			}
 			break;
 
@@ -91,7 +95,7 @@ TArray<FPCGPinProperties> UPCGIndirectionSettings::InputPinProperties() const
 			checkNoEntry();
 	}
 
-	if (InputProperties.IsEmpty())
+	if (!bSetProperties)
 	{
 		return Super::InputPinProperties();
 	}
diff --git a/Engine/Plugins/Experimental/PCG/Source/PCG/Private/Helpers/PCGSubgraphHelpers.cpp b/Engine/Plugins/Experimental/PCG/Source/PCG/Private/Helpers/PCGSubgraphHelpers.cpp
--- a/Engine/Plugins/Experimental/PCG/Source/PCG/Private/Helpers/PCGSubgraphHelpers.cpp
+++ b/Engine/Plugins/Experimental/PCG/Source/PCG/Private/Helpers/PCGSubgraphHelpers.cpp
@@ -55,7 +55,7 @@ namespace PCGSubgraphHelpersExtra
 		}
 
 		NewProperties.Label = FName(NewName);
-		NewProperties = NewInputOutputSettings->AddCustomPin(NewProperties);
+		NewProperties = NewInputOutputSettings->AddPin(NewProperties);
 		NewInputOutputNode->UpdateAfterSettingsChangeDuringCreation();
 		return NewProperties.Label;
 	}
@@ -329,18 +329,7 @@ UPCGGraph* FPCGSubgraphHelpers::CollapseIntoSubgraph(UPCGGraph* InOriginalGraph,
 			OutsideSubgraphEdge.InputPin = Edge->InputPin;
 			InsideSubgraphEdge.OutputPin = *OutPin;
 
-			if (Edge->InputPin->Node == InOriginalGraph->GetInputNode())
-			{
-				const UPCGGraphInputOutputSettings* Settings = Cast<const UPCGGraphInputOutputSettings>(Edge->InputPin->Node->GetSettings());
-
-				if (Settings && !Settings->IsCustomPin(Edge->InputPin))
-				{
-					OutsideSubgraphEdge.OutputPinLabel = Edge->InputPin->Properties.Label;
-					InsideSubgraphEdge.InputPin = NewPCGGraph->GetInputNode()->GetOutputPin(Edge->InputPin->Properties.Label);
-					bProcessed = true;
-				}
-			}
-			else if (Edge->InputPin->Node && Edge->InputPin->Node->GetSettings() && Edge->InputPin->Node->GetSettings()->IsA<UPCGUserParameterGetSettings>())
+			if (Edge->InputPin->Node && Edge->InputPin->Node->GetSettings() && Edge->InputPin->Node->GetSettings()->IsA<UPCGUserParameterGetSettings>())
 			{
 				const UPCGUserParameterGetSettings* OldSettings = CastChecked<const UPCGUserParameterGetSettings>(Edge->InputPin->Node->GetSettings());
 				if (CollapseInfo.GetUserParametersOutputPins.Contains(OldSettings->PropertyGuid))
@@ -408,18 +397,6 @@ UPCGGraph* FPCGSubgraphHelpers::CollapseIntoSubgraph(UPCGGraph* InOriginalGraph,
 			OutsideSubgraphEdge.OutputPin = Edge->OutputPin;
 			InsideSubgraphEdge.InputPin = *InPin;
 
-			if (Edge->OutputPin->Node == InOriginalGraph->GetOutputNode())
-			{
-				const UPCGGraphInputOutputSettings* Settings = Cast<const UPCGGraphInputOutputSettings>(Edge->OutputPin->Node->GetSettings());
-
-				if (Settings && !Settings->IsCustomPin(Edge->OutputPin))
-				{
-					OutsideSubgraphEdge.InputPinLabel = Edge->OutputPin->Properties.Label;
-					InsideSubgraphEdge.OutputPin = NewPCGGraph->GetOutputNode()->GetInputPin(Edge->OutputPin->Properties.Label);
-					bProcessed = true;
-				}
-			}
-
 			if (!bProcessed)
 			{
 				FName NewPinName = PCGSubgraphHelpersExtra::CreateNewCustomPin(NewPCGGraph, Edge->InputPin, false, NameCollisionMapping);
diff --git a/Engine/Plugins/Experimental/PCG/Source/PCG/Private/PCGInputOutputSettings.cpp b/Engine/Plugins/Experimental/PCG/Source/PCG/Private/PCGInputOutputSettings.cpp
--- a/Engine/Plugins/Experimental/PCG/Source/PCG/Private/PCGInputOutputSettings.cpp
+++ b/Engine/Plugins/Experimental/PCG/Source/PCG/Private/PCGInputOutputSettings.cpp
@@ -3,12 +3,77 @@
 #include "PCGInputOutputSettings.h"
 
 #include "PCGContext.h"
+#include "PCGCustomVersion.h"
 #include "PCGPin.h"
 
 #include UE_INLINE_GENERATED_CPP_BY_NAME(PCGInputOutputSettings)
 
 #define LOCTEXT_NAMESPACE "PCGInputOutputElement"
 
+namespace PCGInputOutputPrivate
+{
+	struct FLabelAndTooltip
+	{
+		FLabelAndTooltip(FName InLabel, FText InTooltip = FText::GetEmpty())
+			: Label(InLabel), Tooltip(InTooltip)
+		{
+		}
+
+		FName Label;
+		FText Tooltip;
+	};
+
+	static TArray<FLabelAndTooltip> StaticInLabels =
+	{
+		{
+			PCGPinConstants::DefaultInputLabel, 
+			LOCTEXT("InputOutputInPinTooltip", "Default input pin.")
+		}
+	};
+
+	static TArray<FLabelAndTooltip> StaticAdvancedInLabels =
+	{
+		{
+			PCGInputOutputConstants::DefaultInputLabel,
+			LOCTEXT("InputOutputInputPinTooltip", "DEPRECATED: Please use custom pins to pass inputs into the graph, or use the 'getter' nodes (e.g. Get Volume Data, Get Actor Data, ...) to import data from the level.\n\n"
+			"Takes the output of the Actor pin and if the 'Input Type' setting on the PCG Component is set to Landscape, combines it with the result of the Landscape pin. "
+			"If the Actor data is two dimensional it will be projected onto the landscape, otherwise it will be intersected.")
+		},
+		{
+			PCGInputOutputConstants::DefaultActorLabel,
+			LOCTEXT("InputOutputActorPinTooltip", "DEPRECATED: Please use custom pins to pass inputs into the graph, or use the 'getter' nodes (e.g. Get Volume Data, Get Actor Data, ...) to import data from the level.\n\n"
+			"If this is a partitioned component, then this will be the intersection of the current partition actor bounds with the following. "
+			"If the actor is a Landscape Proxy, then this provide a landscape data. "
+			"Otherwise if the actor is a volume, this will provide a volume shape matching the actor bounds. "
+			"Otherwise if the 'Parse Actor Components' setting is enabled on the PCG Component, this will be all compatible components on the actor (Landscape Splines, Splines, Shapes, Primitives) unioned together. "
+			"Otherwise a single point will be provided at the actor position.")
+		},
+		{
+			PCGInputOutputConstants::DefaultOriginalActorLabel,
+			LOCTEXT("InputOutputOriginalActorPinTooltip", "DEPRECATED: Please use custom pins to pass inputs into the graph, or use the 'getter' nodes (e.g. Get Volume Data, Get Actor Data, ...) to import data from the level.\n\n"
+			"If the actor is a partition actor, this will pull data from the generating PCG actor. Otherwise it will provide the same data as the Actor pin.")
+		},
+		{
+			PCGInputOutputConstants::DefaultLandscapeLabel,
+			LOCTEXT("InputOutputLandscapePinTooltip", "DEPRECATED: Please use custom pins to pass inputs into the graph, or use the 'getter' nodes (e.g. Get Volume Data, Get Actor Data, ...) to import data from the level.\n\n"
+			"Provides the landscape represented by this actor if it is a Landscape Proxy, otherwise it returns any landscapes overlapping this actor in the level.")
+		},
+		{
+			PCGInputOutputConstants::DefaultLandscapeHeightLabel,
+			LOCTEXT("InputOutputLandscapeHeightPinTooltip", "DEPRECATED: Please use custom pins to pass inputs into the graph, or use the 'getter' nodes (e.g. Get Volume Data, Get Actor Data, ...) to import data from the level.\n\n"
+			"Similar to Landscape pin, but only provides height data and not other layers.")
+		}
+	};
+
+	static TArray<FLabelAndTooltip> StaticOutLabels =
+	{
+		{
+			PCGPinConstants::DefaultOutputLabel
+		}
+	};
+	static TArray<FLabelAndTooltip> StaticAdvancedOutLabels;
+}
+
 bool FPCGInputOutputElement::ExecuteInternal(FPCGContext* Context) const
 {
 	// Essentially a pass-through element
@@ -16,43 +81,6 @@ bool FPCGInputOutputElement::ExecuteInternal(FPCGContext* Context) const
 	return true;
 }
 
-UPCGGraphInputOutputSettings::UPCGGraphInputOutputSettings(const FObjectInitializer& ObjectInitializer)
-	: Super(ObjectInitializer)
-{
-	// This label is kept into "StaticInLabels" to avoid changing the default object. But it will be marked as advanced too, since we want to move away from it.
-	StaticInLabels.Emplace(PCGPinConstants::DefaultInputLabel, LOCTEXT("InputOutputInPinTooltip",
-		"DEPRECATED: Please use custom pins to pass inputs into the graph, or use the 'getter' nodes (e.g. Get Volume Data, Get Actor Data, ...) to import data from the level.\n\n"
-		"Provides the same result as the Input pin."
-	));
-	StaticAdvancedInLabels.Emplace(PCGInputOutputConstants::DefaultInputLabel, LOCTEXT("InputOutputInputPinTooltip",
-		"DEPRECATED: Please use custom pins to pass inputs into the graph, or use the 'getter' nodes (e.g. Get Volume Data, Get Actor Data, ...) to import data from the level.\n\n"
-		"Takes the output of the Actor pin and if the 'Input Type' setting on the PCG Component is set to Landscape, combines it with the result of the Landscape pin. "
-		"If the Actor data is two dimensional it will be projected onto the landscape, otherwise it will be intersected."
-	));
-	StaticAdvancedInLabels.Emplace(PCGInputOutputConstants::DefaultActorLabel, LOCTEXT("InputOutputActorPinTooltip",
-		"DEPRECATED: Please use custom pins to pass inputs into the graph, or use the 'getter' nodes (e.g. Get Volume Data, Get Actor Data, ...) to import data from the level.\n\n"
-		"If this is a partitioned component, then this will be the intersection of the current partition actor bounds with the following. "
-		"If the actor is a Landscape Proxy, then this provide a landscape data. "
-		"Otherwise if the actor is a volume, this will provide a volume shape matching the actor bounds. "
-		"Otherwise if the 'Parse Actor Components' setting is enabled on the PCG Component, this will be all compatible components on the actor (Landscape Splines, Splines, Shapes, Primitives) unioned together. "
-		"Otherwise a single point will be provided at the actor position."
-	));
-	StaticAdvancedInLabels.Emplace(PCGInputOutputConstants::DefaultOriginalActorLabel, LOCTEXT("InputOutputOriginalActorPinTooltip",
-		"DEPRECATED: Please use custom pins to pass inputs into the graph, or use the 'getter' nodes (e.g. Get Volume Data, Get Actor Data, ...) to import data from the level.\n\n"
-		"If the actor is a partition actor, this will pull data from the generating PCG actor. Otherwise it will provide the same data as the Actor pin."
-	));
-	StaticAdvancedInLabels.Emplace(PCGInputOutputConstants::DefaultLandscapeLabel, LOCTEXT("InputOutputLandscapePinTooltip",
-		"DEPRECATED: Please use custom pins to pass inputs into the graph, or use the 'getter' nodes (e.g. Get Volume Data, Get Actor Data, ...) to import data from the level.\n\n"
-		"Provides the landscape represented by this actor if it is a Landscape Proxy, otherwise it returns any landscapes overlapping this actor in the level."
-	));
-	StaticAdvancedInLabels.Emplace(PCGInputOutputConstants::DefaultLandscapeHeightLabel, LOCTEXT("InputOutputLandscapeHeightPinTooltip",
-		"DEPRECATED: Please use custom pins to pass inputs into the graph, or use the 'getter' nodes (e.g. Get Volume Data, Get Actor Data, ...) to import data from the level.\n\n"
-		"Similar to Landscape pin, but only provides height data and not other layers."
-	));
-	
-	StaticOutLabels.Emplace(PCGPinConstants::DefaultOutputLabel);
-}
-
 void UPCGGraphInputOutputSettings::PostLoad()
 {
 	Super::PostLoad();
@@ -61,53 +89,83 @@ void UPCGGraphInputOutputSettings::PostLoad()
 	{
 		for (const FName& PinLabel : PinLabels_DEPRECATED)
 		{
-			CustomPins.Emplace(PinLabel);
+			Pins.Emplace(PinLabel);
 		}
 
 		PinLabels_DEPRECATED.Reset();
 	}
+
+	if (!CustomPins_DEPRECATED.IsEmpty())
+	{
+		Pins.Append(CustomPins_DEPRECATED);
+		CustomPins_DEPRECATED.Empty();
+	}
 }
 
-TArray<FPCGPinProperties> UPCGGraphInputOutputSettings::GetPinProperties() const
+void UPCGGraphInputOutputSettings::SetInput(bool bInIsInput)
 {
-	TArray<FPCGPinProperties> PinProperties;
-	const EPCGDataType DefaultPinDataType = bIsInput ? EPCGDataType::Spatial : EPCGDataType::Any;
-	Algo::Transform(StaticLabels(), PinProperties, [DefaultPinDataType](const FLabelAndTooltip& InLabelAndTooltip) {
-		FPCGPinProperties Res = FPCGPinProperties(InLabelAndTooltip.Label, DefaultPinDataType, /*bMultiConnections=*/true, /*bMultiData=*/true, InLabelAndTooltip.Tooltip);
-		Res.bAdvancedPin = true;
-		return Res;
-	});
-	
-
-	Algo::Transform(StaticAdvancedLabels(), PinProperties, [this, DefaultPinDataType](const FLabelAndTooltip& InLabelAndTooltip) {
-		const bool bIsLandscapePin = (InLabelAndTooltip.Label == PCGInputOutputConstants::DefaultLandscapeLabel || InLabelAndTooltip.Label == PCGInputOutputConstants::DefaultLandscapeHeightLabel);
-		const EPCGDataType PinType = bIsLandscapePin ? EPCGDataType::Surface : DefaultPinDataType;
-		FPCGPinProperties Res(InLabelAndTooltip.Label, PinType, /*bMultiConnection=*/true, /*bMultiData=*/false, InLabelAndTooltip.Tooltip);
-		Res.bAdvancedPin = true;
-		return Res;
-	});
+	bIsInput = bInIsInput;
 
+	if (!bHasAddedDefaultPin)
+	{
+		bHasAddedDefaultPin = true;
 
-	PinProperties.Append(CustomPins);
-	return PinProperties;
+		if (Pins.IsEmpty())
+		{
+			const EPCGDataType DefaultPinDataType = bIsInput ? EPCGDataType::Spatial : EPCGDataType::Any;
+			const TArray<PCGInputOutputPrivate::FLabelAndTooltip>& StaticLabels = bIsInput ? PCGInputOutputPrivate::StaticInLabels : PCGInputOutputPrivate::StaticOutLabels;
+
+			Algo::Transform(StaticLabels, Pins, [DefaultPinDataType](const PCGInputOutputPrivate::FLabelAndTooltip& InLabelAndTooltip) {
+				FPCGPinProperties Res = FPCGPinProperties(InLabelAndTooltip.Label, DefaultPinDataType, /*bMultiConnections=*/true, /*bMultiData=*/true, InLabelAndTooltip.Tooltip);
+				Res.bAdvancedPin = true;
+				return Res;
+			});
+		}
+	}
 }
 
-TArray<FPCGPinProperties> UPCGGraphInputOutputSettings::InputPinProperties() const
+#if WITH_EDITOR
+void UPCGGraphInputOutputSettings::ApplyDeprecationBeforeUpdatePins(UPCGNode* InOutNode, TArray<TObjectPtr<UPCGPin>>& InputPins, TArray<TObjectPtr<UPCGPin>>& OutputPins)
 {
-	return GetPinProperties();
-}
+	if (DataVersion < FPCGCustomVersion::UpdateInputOutputNodesDefaults)
+	{
+		// At this point we should already have the default in/out, done in the SetInput, so we don't need to consider the default in/out here.
+		check(Pins.Num() > 0);
+
+		// Add default pins to the node if and only if any of its default pins were connected.
+		TArray<FPCGPinProperties> DeprecatedPropertiesToAdd;
+		TArray<TObjectPtr<UPCGPin>>& PreviousPins = (bIsInput ? OutputPins : InputPins);
+		// Basically, any previous pins that's connected and also not part of the current Pins array is an old pin
+		for (UPCGPin* PreviousPin : PreviousPins)
+		{
+			if (!PreviousPin || !PreviousPin->IsConnected())
+			{
+				continue;
+			}
 
-TArray<FPCGPinProperties> UPCGGraphInputOutputSettings::OutputPinProperties() const
-{
-	return GetPinProperties();
+			// If the previous pin isn't found, just copy directly
+			if (!Pins.FindByPredicate([PreviousPin](const FPCGPinProperties& PinProperty) { return PinProperty.Label == PreviousPin->Properties.Label; }))
+			{
+				DeprecatedPropertiesToAdd.Add(PreviousPin->Properties);
+			}
+		}
+
+		if (!DeprecatedPropertiesToAdd.IsEmpty())
+		{
+			Pins.Insert(MoveTemp(DeprecatedPropertiesToAdd), 1);
+		}
+	}
 }
+#endif // WITH_EDITOR
 
 TArray<FPCGPinProperties> UPCGGraphInputOutputSettings::DefaultInputPinProperties() const
 {
 	// It is important for serialization that this is not modified, or it could break existing graphs.
 	TArray<FPCGPinProperties> PinProperties;
 	const EPCGDataType DefaultPinDataType = bIsInput ? EPCGDataType::Spatial : EPCGDataType::Any;
-	Algo::Transform(StaticLabels(), PinProperties, [DefaultPinDataType](const FLabelAndTooltip& InLabelAndTooltip) {
+	const TArray<PCGInputOutputPrivate::FLabelAndTooltip>& StaticLabels = bIsInput ? PCGInputOutputPrivate::StaticInLabels : PCGInputOutputPrivate::StaticOutLabels;
+
+	Algo::Transform(StaticLabels, PinProperties, [DefaultPinDataType](const PCGInputOutputPrivate::FLabelAndTooltip& InLabelAndTooltip) {
 		FPCGPinProperties Res = FPCGPinProperties(InLabelAndTooltip.Label, DefaultPinDataType, /*bMultiConnections=*/true, /*bMultiData=*/true, InLabelAndTooltip.Tooltip);
 		Res.bAdvancedPin = true;
 		return Res;
@@ -121,18 +179,12 @@ TArray<FPCGPinProperties> UPCGGraphInputOutputSettings::DefaultOutputPinProperti
 	return DefaultInputPinProperties();
 }
 
-const FPCGPinProperties& UPCGGraphInputOutputSettings::AddCustomPin(const FPCGPinProperties& NewCustomPinProperties)
+const FPCGPinProperties& UPCGGraphInputOutputSettings::AddPin(const FPCGPinProperties& NewPinProperties)
 {
 	Modify();
-	int32 Index = CustomPins.Add(NewCustomPinProperties);
-	FixCustomPinProperties();
-	return CustomPins[Index];
-}
-
-bool UPCGGraphInputOutputSettings::IsCustomPin(const UPCGPin* InPin) const
-{
-	check(InPin);
-	return CustomPins.Contains(InPin->Properties);
+	int32 Index = Pins.Add(NewPinProperties);
+	FixPinProperties();
+	return Pins[Index];
 }
 
 #if WITH_EDITOR
@@ -141,19 +193,19 @@ void UPCGGraphInputOutputSettings::PostEditChangeProperty(struct FPropertyChange
 	if (PropertyChangedEvent.Property)
 	{
 		// Label has changed if we have modified "CustomPins" array
-		bool bLabelChanged = PropertyChangedEvent.Property->GetFName() == GET_MEMBER_NAME_CHECKED(UPCGGraphInputOutputSettings, CustomPins);
+		bool bLabelChanged = PropertyChangedEvent.Property->GetFName() == GET_MEMBER_NAME_CHECKED(UPCGGraphInputOutputSettings, Pins);
 
 		if (!bLabelChanged)
 		{
 			// Or might has changed if an element of the array "CustomPin" was modified.
 			bLabelChanged = PropertyChangedEvent.MemberProperty &&
-				PropertyChangedEvent.MemberProperty->GetFName() == GET_MEMBER_NAME_CHECKED(UPCGGraphInputOutputSettings, CustomPins) &&
+				PropertyChangedEvent.MemberProperty->GetFName() == GET_MEMBER_NAME_CHECKED(UPCGGraphInputOutputSettings, Pins) &&
 				PropertyChangedEvent.Property->GetFName() == GET_MEMBER_NAME_CHECKED(FPCGPinProperties, Label);
 		}
 
 		if (bLabelChanged)
 		{
-			FixCustomPinProperties();
+			FixPinProperties();
 		}
 	}
 
@@ -161,33 +213,21 @@ void UPCGGraphInputOutputSettings::PostEditChangeProperty(struct FPropertyChange
 }
 #endif // WITH_EDITOR
 
-void UPCGGraphInputOutputSettings::FixCustomPinProperties()
+void UPCGGraphInputOutputSettings::FixPinProperties()
 {
-	// No need to fix if we have no custom pins
-	if (CustomPins.IsEmpty())
+	// No need to fix if we have no pins
+	if (Pins.IsEmpty())
 	{
 		return;
 	}
 
 	TSet<FName> AllLabels;
-
-	// First gather all our static labels
-	for (FLabelAndTooltip LabelAndTooltip : StaticLabels())
-	{
-		AllLabels.Emplace(LabelAndTooltip.Label);
-	}
-
-	for (FLabelAndTooltip LabelAndTooltip : StaticAdvancedLabels())
-	{
-		AllLabels.Emplace(LabelAndTooltip.Label);
-	}
-
 	bool bWasModified = false;
 
-	for (FPCGPinProperties& CustomPinProperties : CustomPins)
+	for (FPCGPinProperties& PinProperties : Pins)
 	{
 		// Avoid "None" pin label
-		if (CustomPinProperties.Label == NAME_None)
+		if (PinProperties.Label == NAME_None)
 		{
 			if (!bWasModified)
 			{
@@ -195,13 +235,13 @@ void UPCGGraphInputOutputSettings::FixCustomPinProperties()
 				Modify();
 			}
 
-			CustomPinProperties.Label = PCGInputOutputConstants::DefaultNewCustomPinName;
+			PinProperties.Label = PCGInputOutputConstants::DefaultNewCustomPinName;
 		}
 
 		uint32 Count = 1;
-		FString OriginalLabel = CustomPinProperties.Label.ToString();
+		FString OriginalLabel = PinProperties.Label.ToString();
 
-		while (AllLabels.Contains(CustomPinProperties.Label))
+		while (AllLabels.Contains(PinProperties.Label))
 		{
 			if (!bWasModified)
 			{
@@ -209,10 +249,10 @@ void UPCGGraphInputOutputSettings::FixCustomPinProperties()
 				Modify();
 			}
 
-			CustomPinProperties.Label = FName(FString::Printf(TEXT("%s%d"), *OriginalLabel, Count++));
+			PinProperties.Label = FName(FString::Printf(TEXT("%s%d"), *OriginalLabel, Count++));
 		}
 
-		AllLabels.Emplace(CustomPinProperties.Label);
+		AllLabels.Emplace(PinProperties.Label);
 	}
 }
 
diff --git a/Engine/Plugins/Experimental/PCG/Source/PCG/Public/PCGCustomVersion.h b/Engine/Plugins/Experimental/PCG/Source/PCG/Public/PCGCustomVersion.h
--- a/Engine/Plugins/Experimental/PCG/Source/PCG/Public/PCGCustomVersion.h
+++ b/Engine/Plugins/Experimental/PCG/Source/PCG/Public/PCGCustomVersion.h
@@ -56,6 +56,9 @@ struct PCG_API FPCGCustomVersion
 		// Update TransferAttribute with selectors
 		UpdateTransferAttributeWithSelectors = 15,
 
+		// Removed by-default pins on input node. Note, this breaks cooked binary compatibility
+		UpdateInputOutputNodesDefaults = 16,
+
 		// -----<new versions can be added above this line>-------------------------------------------------
 		VersionPlusOne,
 		LatestVersion = VersionPlusOne - 1
diff --git a/Engine/Plugins/Experimental/PCG/Source/PCG/Public/PCGGraph.h b/Engine/Plugins/Experimental/PCG/Source/PCG/Public/PCGGraph.h
--- a/Engine/Plugins/Experimental/PCG/Source/PCG/Public/PCGGraph.h
+++ b/Engine/Plugins/Experimental/PCG/Source/PCG/Public/PCGGraph.h
@@ -320,14 +320,14 @@ protected:
 	/** Calculates node grid size. Not thread safe, called within write lock. */
 	uint32 CalculateNodeGridSizeRecursive_Unsafe(const UPCGNode* InNode, uint32 InDefaultGridSize) const;
 
-	UPROPERTY(BlueprintReadOnly, VisibleAnywhere, Category = Graph)
+	UPROPERTY(BlueprintReadOnly, VisibleAnywhere, Category = Graph, meta = (NoResetToDefault))
 	TArray<TObjectPtr<UPCGNode>> Nodes;
 
 	// Add input/output nodes
-	UPROPERTY(BlueprintReadOnly, VisibleAnywhere, Category = Graph)
+	UPROPERTY(BlueprintReadOnly, VisibleAnywhere, Category = Graph, meta = (NoResetToDefault))
 	TObjectPtr<UPCGNode> InputNode;
 
-	UPROPERTY(BlueprintReadOnly, VisibleAnywhere, Category = Graph)
+	UPROPERTY(BlueprintReadOnly, VisibleAnywhere, Category = Graph, meta = (NoResetToDefault))
 	TObjectPtr<UPCGNode> OutputNode;
 
 #if WITH_EDITORONLY_DATA
diff --git a/Engine/Plugins/Experimental/PCG/Source/PCG/Public/PCGInputOutputSettings.h b/Engine/Plugins/Experimental/PCG/Source/PCG/Public/PCGInputOutputSettings.h
--- a/Engine/Plugins/Experimental/PCG/Source/PCG/Public/PCGInputOutputSettings.h
+++ b/Engine/Plugins/Experimental/PCG/Source/PCG/Public/PCGInputOutputSettings.h
@@ -29,8 +29,6 @@ class PCG_API UPCGGraphInputOutputSettings : public UPCGSettings
 	GENERATED_BODY()
 
 public:
-	UPCGGraphInputOutputSettings(const FObjectInitializer& ObjectInitializer);
-
 	// ~Begin UObject interface
 	virtual void PostLoad() override;
 	// ~End UObject interface
@@ -40,57 +38,41 @@ public:
 	virtual FName GetDefaultNodeName() const override { return bIsInput ? FName(TEXT("InputNode")) : FName(TEXT("OutputNode")); }
 	virtual FText GetDefaultNodeTitle() const override { return bIsInput ? NSLOCTEXT("PCGGraphInputOutputSettings", "InputNodeTitle", "Input Node") : NSLOCTEXT("PCGGraphInputOutputSettings", "OutputNodeTitle", "Output Node"); }
 	virtual EPCGSettingsType GetType() const override { return EPCGSettingsType::InputOutput; }
-	virtual void PostEditChangeProperty(struct FPropertyChangedEvent& PropertyChangedEvent);
+	virtual void PostEditChangeProperty(struct FPropertyChangedEvent& PropertyChangedEvent) override;
+	virtual void ApplyDeprecationBeforeUpdatePins(UPCGNode* InOutNode, TArray<TObjectPtr<UPCGPin>>& InputPins, TArray<TObjectPtr<UPCGPin>>& OutputPins) override;
 #endif
 
-	TArray<FPCGPinProperties> InputPinProperties() const override;
-	TArray<FPCGPinProperties> OutputPinProperties() const override;
+	TArray<FPCGPinProperties> InputPinProperties() const override { return Pins; }
+	TArray<FPCGPinProperties> OutputPinProperties() const override { return Pins; }
 
 	virtual TArray<FPCGPinProperties> DefaultInputPinProperties() const override;
 	virtual TArray<FPCGPinProperties> DefaultOutputPinProperties() const override;
 
-	void SetInput(bool bInIsInput) { bIsInput = bInIsInput; }
-
-	bool IsCustomPin(const UPCGPin* InPin) const;
+	void SetInput(bool bInIsInput);
 
 	// Add a new custom pin
 	// Note that you should use the return value of this function, since it can be different from
 	// the one passed as argument. It will change if its label collides with existing pins.
-	UE_NODISCARD const FPCGPinProperties& AddCustomPin(const FPCGPinProperties& NewCustomPinProperties);
+	UE_NODISCARD const FPCGPinProperties& AddPin(const FPCGPinProperties& NewCustomPinProperties);
 
 protected:
 	virtual FPCGElementPtr CreateElement() const override { return MakeShared<FPCGInputOutputElement>(); }
 	// ~End UPCGSettings interface
 
-	struct FLabelAndTooltip
-	{
-		explicit FLabelAndTooltip(FName InLabel, FText InTooltip = FText::GetEmpty())
-			: Label(InLabel), Tooltip(InTooltip)
-		{
-		}
-
-		FName Label;
-		FText Tooltip;
-	};
-
-	const TArray<FLabelAndTooltip>& StaticLabels() const { return bIsInput ? StaticInLabels : StaticOutLabels; }
-	const TArray<FLabelAndTooltip>& StaticAdvancedLabels() const { return bIsInput ? StaticAdvancedInLabels : StaticAdvancedOutLabels; }
-
-	TArray<FPCGPinProperties> GetPinProperties() const;
-
-	void FixCustomPinProperties();
+	void FixPinProperties();
 
 protected:
 	UPROPERTY()
 	TSet<FName> PinLabels_DEPRECATED;
 
-	UPROPERTY(EditAnywhere, BlueprintReadWrite, Category = "Settings|Input")
-	TArray<FPCGPinProperties> CustomPins;
+	UPROPERTY()
+	TArray<FPCGPinProperties> CustomPins_DEPRECATED;
+
+	UPROPERTY(EditAnywhere, BlueprintReadWrite, Category = Settings)
+	TArray<FPCGPinProperties> Pins;
 
-	TArray<FLabelAndTooltip> StaticInLabels;
-	TArray<FLabelAndTooltip> StaticAdvancedInLabels;
-	TArray<FLabelAndTooltip> StaticOutLabels;
-	TArray<FLabelAndTooltip> StaticAdvancedOutLabels;
+	UPROPERTY()
+	bool bHasAddedDefaultPin = false;
 
 	bool bIsInput = false;
 };
