Fixed typo around the metadata copy point where clearly we weren't doing the right thing
Added possibility to override filter by name node

#rnx

diff --git a/Engine/Plugins/Experimental/PCG/Source/PCG/Private/Metadata/PCGMetadataAccessor.cpp b/Engine/Plugins/Experimental/PCG/Source/PCG/Private/Metadata/PCGMetadataAccessor.cpp
--- a/Engine/Plugins/Experimental/PCG/Source/PCG/Private/Metadata/PCGMetadataAccessor.cpp
+++ b/Engine/Plugins/Experimental/PCG/Source/PCG/Private/Metadata/PCGMetadataAccessor.cpp
@@ -311,7 +311,7 @@ void UPCGMetadataAccessorHelpers::InitializeMetadata(FPCGPoint& Point, UPCGMetad
 	if (Metadata)
 	{
 		// If we're not given the parent metadata, we'll assume it is the current metadata's parent
-		if (!ParentMetadata || ParentMetadata->HasParent(ParentMetadata))
+		if (!ParentMetadata || Metadata->HasParent(ParentMetadata))
 		{
 			Point.MetadataEntry = Metadata->AddEntry(ParentPoint.MetadataEntry);
 		}
diff --git a/Engine/Plugins/Experimental/PCG/Source/PCG/Public/Elements/PCGAttributeFilterNamesElement.h b/Engine/Plugins/Experimental/PCG/Source/PCG/Public/Elements/PCGAttributeFilterNamesElement.h
--- a/Engine/Plugins/Experimental/PCG/Source/PCG/Public/Elements/PCGAttributeFilterNamesElement.h
+++ b/Engine/Plugins/Experimental/PCG/Source/PCG/Public/Elements/PCGAttributeFilterNamesElement.h
@@ -47,10 +47,10 @@ protected:
 	//~End UPCGSettings interface
 
 public:
-	UPROPERTY(BlueprintReadWrite, EditAnywhere, Category = Settings)
+	UPROPERTY(BlueprintReadWrite, EditAnywhere, Category = Settings, meta = (PCG_Overridable))
 	EPCGAttributeFilterOperation Operation = EPCGAttributeFilterOperation::KeepSelectedAttributes;
 
-	UPROPERTY(BlueprintReadWrite, EditAnywhere, Category = Settings)
+	UPROPERTY(BlueprintReadWrite, EditAnywhere, Category = Settings, meta = (PCG_Overridable))
 	FString SelectedAttributes;
 
 #if WITH_EDITORONLY_DATA
