[PCG] Nativize ApplyScaleToBounds and  EditPointCenter nodes
* ApplyScaleToBounds - Converts scale of a point into the bounds
* EditPointPivot - (renamed ResetPointCenter) Sets a new center point for the point, adjusting the bounds
-- Includes UnitTests

diff --git a/Engine/Plugins/Experimental/PCG/Source/PCG/Private/Elements/PCGApplyScaleToBounds.cpp b/Engine/Plugins/Experimental/PCG/Source/PCG/Private/Elements/PCGApplyScaleToBounds.cpp
new file mode 100644
--- /dev/null
+++ b/Engine/Plugins/Experimental/PCG/Source/PCG/Private/Elements/PCGApplyScaleToBounds.cpp
@@ -0,0 +1,33 @@
+// Copyright Epic Games, Inc. All Rights Reserved.
+
+#include "Elements/PCGApplyScaleToBounds.h"
+
+#include "PCGContext.h"
+#include "PCGPin.h"
+#include "Data/PCGPointData.h"
+
+#define LOCTEXT_NAMESPACE "PCGApplyScaleToBoundsElement"
+
+FPCGElementPtr UPCGApplyScaleToBoundsSettings::CreateElement() const
+{
+	return MakeShared<FPCGApplyScaleToBoundsElement>();
+}
+
+bool FPCGApplyScaleToBoundsElement::ExecuteInternal(FPCGContext* Context) const
+{
+	TRACE_CPUPROFILER_EVENT_SCOPE(FPCGApplyScaleToBoundsElement::Execute);
+	check(Context);
+
+	FPCGApplyScaleToBoundsElement::ContextType* ApplyScaleContext = static_cast<FPCGApplyScaleToBoundsElement::ContextType*>(Context);
+
+	return ExecutePointOperation(ApplyScaleContext, [](const FPCGPoint& InPoint, FPCGPoint& OutPoint)
+	{
+		OutPoint = InPoint;
+
+		OutPoint.ApplyScaleToBounds();
+		
+		return true;
+	});
+}
+
+#undef LOCTEXT_NAMESPACE
\ No newline at end of file
diff --git a/Engine/Plugins/Experimental/PCG/Source/PCG/Private/Elements/PCGResetPointCenter.cpp b/Engine/Plugins/Experimental/PCG/Source/PCG/Private/Elements/PCGResetPointCenter.cpp
new file mode 100644
--- /dev/null
+++ b/Engine/Plugins/Experimental/PCG/Source/PCG/Private/Elements/PCGResetPointCenter.cpp
@@ -0,0 +1,36 @@
+// Copyright Epic Games, Inc. All Rights Reserved.
+
+#include "Elements/PCGResetPointCenter.h"
+
+#include "PCGContext.h"
+#include "PCGPin.h"
+#include "Data/PCGPointData.h"
+
+#define LOCTEXT_NAMESPACE "PCGResetPointCenterElement"
+
+FPCGElementPtr UPCGResetPointCenterSettings::CreateElement() const
+{
+	return MakeShared<FPCGResetPointCenterElement>();
+}
+
+bool FPCGResetPointCenterElement::ExecuteInternal(FPCGContext* Context) const
+{
+	TRACE_CPUPROFILER_EVENT_SCOPE(FPCGResetPointCenterElement::Execute);
+	check(Context);
+
+	FPCGResetPointCenterElement::ContextType* PointCenterContext = static_cast<FPCGResetPointCenterElement::ContextType*>(Context);
+
+	const UPCGResetPointCenterSettings* Settings = Context->GetInputSettings<UPCGResetPointCenterSettings>();
+	check(Settings);
+
+	return ExecutePointOperation(PointCenterContext, [&PointCenterLocation = Settings->PointCenterLocation](const FPCGPoint& InPoint, FPCGPoint& OutPoint)
+	{
+		OutPoint = InPoint;
+
+		OutPoint.ResetPointCenter(PointCenterLocation);
+
+		return true;
+	});
+}
+
+#undef LOCTEXT_NAMESPACE
\ No newline at end of file
diff --git a/Engine/Plugins/Experimental/PCG/Source/PCG/Private/PCGPoint.cpp b/Engine/Plugins/Experimental/PCG/Source/PCG/Private/PCGPoint.cpp
--- a/Engine/Plugins/Experimental/PCG/Source/PCG/Private/PCGPoint.cpp
+++ b/Engine/Plugins/Experimental/PCG/Source/PCG/Private/PCGPoint.cpp
@@ -45,6 +45,24 @@ FBoxSphereBounds FPCGPoint::GetDensityBounds() const
 	return FBoxSphereBounds(GetLocalDensityBounds().TransformBy(Transform));
 }
 
+void FPCGPoint::ApplyScaleToBounds()
+{
+	const FVector PointScale = Transform.GetScale3D();
+	Transform.SetScale3D(PointScale.GetSignVector());
+	BoundsMin *= PointScale.GetAbs();
+	BoundsMax *= PointScale.GetAbs();
+}
+
+void FPCGPoint::ResetPointCenter(const FVector& BoundsRatio)
+{
+	const FVector NewCenterLocal = FMath::Lerp(BoundsMin, BoundsMax, BoundsRatio);
+
+	BoundsMin -= NewCenterLocal;
+	BoundsMax -= NewCenterLocal;
+
+	Transform.SetLocation(Transform.GetLocation() + Transform.TransformVector(NewCenterLocal));
+}
+
 bool FPCGPoint::HasCustomPropertyGetterSetter(FName Name)
 {
 	return Name == PCGPointCustomPropertyNames::ExtentsName ||
diff --git a/Engine/Plugins/Experimental/PCG/Source/PCG/Private/Tests/Elements/PCGApplyScaleToBoundsTest.cpp b/Engine/Plugins/Experimental/PCG/Source/PCG/Private/Tests/Elements/PCGApplyScaleToBoundsTest.cpp
new file mode 100644
--- /dev/null
+++ b/Engine/Plugins/Experimental/PCG/Source/PCG/Private/Tests/Elements/PCGApplyScaleToBoundsTest.cpp
@@ -0,0 +1,105 @@
+// Copyright Epic Games, Inc. All Rights Reserved.
+
+#include "Tests/PCGTestsCommon.h"
+
+#include "PCGContext.h"
+#include "Data/PCGPointData.h"
+
+#include "Elements/PCGApplyScaleToBounds.h"
+
+IMPLEMENT_CUSTOM_SIMPLE_AUTOMATION_TEST(FPCGApplyScaleToBoundsTest_Basic, FPCGTestBaseClass, "pcg.tests.ApplyScaleToBounds.Basic", PCGTestsCommon::TestFlags)
+
+bool FPCGApplyScaleToBoundsTest_Basic::RunTest(const FString& Parameters)
+{
+	PCGTestsCommon::FTestData TestData;
+	
+	FPCGTaggedData& Inputs = TestData.InputData.TaggedData.Emplace_GetRef();
+	Inputs.Pin = PCGPinConstants::DefaultInputLabel;
+	UPCGApplyScaleToBoundsSettings* Settings = PCGTestsCommon::GenerateSettings<UPCGApplyScaleToBoundsSettings>(TestData);
+	UPCGPointData* InData = PCGTestsCommon::CreateRandomPointData(10, 42, false);
+	Inputs.Data = InData;
+
+	// Setting the scale explicitly
+	UPCGPointData* InputPointData = Cast<UPCGPointData>(InData);
+	TArray<FPCGPoint>& InPoints = InputPointData->GetMutablePoints();
+	for (FPCGPoint& InPoint : InPoints)
+	{
+		InPoint.Transform.SetScale3D(FVector(3.0));
+	}
+
+	FPCGElementPtr TestElement = TestData.Settings->GetElement();
+
+	TUniquePtr<FPCGContext> Context = TestData.InitializeTestContext();
+
+	while (!TestElement->Execute(Context.Get())) {}
+
+	// test our point data
+	const TArray<FPCGTaggedData>& Outputs = Context->OutputData.TaggedData;
+
+	UTEST_EQUAL("Output count", Outputs.Num(), 1);
+
+	const UPCGPointData* OutPointData = Cast<UPCGPointData>(Outputs[0].Data);
+
+	UTEST_NOT_NULL("Output point data", OutPointData);
+
+	const TArray<FPCGPoint>& OutPoints = OutPointData->GetPoints();
+
+	UTEST_EQUAL("Output point count", OutPoints.Num(), 10);
+
+	for (int i = 0; i < OutPoints.Num(); ++i)
+	{
+		UTEST_EQUAL("Output points should have a Bounds Min of -3", OutPoints[i].BoundsMin, FVector(-3));
+		UTEST_EQUAL("Output points should have a Bounds Max of 3", OutPoints[i].BoundsMax, FVector(3));
+	}
+
+	return true;
+}
+
+IMPLEMENT_CUSTOM_SIMPLE_AUTOMATION_TEST(FPCGApplyScaleToBoundsTest_NegativeScale, FPCGTestBaseClass, "pcg.tests.ApplyScaleToBounds.NegativeScale", PCGTestsCommon::TestFlags)
+
+bool FPCGApplyScaleToBoundsTest_NegativeScale::RunTest(const FString& Parameters)
+{
+	PCGTestsCommon::FTestData TestData;
+
+	FPCGTaggedData& Inputs = TestData.InputData.TaggedData.Emplace_GetRef();
+	Inputs.Pin = PCGPinConstants::DefaultInputLabel;
+	UPCGApplyScaleToBoundsSettings* Settings = PCGTestsCommon::GenerateSettings<UPCGApplyScaleToBoundsSettings>(TestData);
+	UPCGPointData* InData = PCGTestsCommon::CreateRandomPointData(10, 42, false);
+	Inputs.Data = InData;
+
+	// Setting the scale explicitly
+	UPCGPointData* InputPointData = Cast<UPCGPointData>(InData);
+	TArray<FPCGPoint>& InPoints = InputPointData->GetMutablePoints();
+	for (FPCGPoint& InPoint : InPoints)
+	{
+		InPoint.Transform.SetScale3D(FVector(-2.0));
+	}
+
+	FPCGElementPtr TestElement = TestData.Settings->GetElement(); // returns nullptr
+
+	TUniquePtr<FPCGContext> Context = TestData.InitializeTestContext();
+
+	while (!TestElement->Execute(Context.Get())) {}
+
+	// test our point data
+	const TArray<FPCGTaggedData>& Outputs = Context->OutputData.TaggedData;
+
+	UTEST_EQUAL("Output count", Outputs.Num(), 1);
+
+	const UPCGPointData* OutPointData = Cast<UPCGPointData>(Outputs[0].Data);
+
+	UTEST_NOT_NULL("Output point data", OutPointData);
+
+	const TArray<FPCGPoint>& OutPoints = OutPointData->GetPoints();
+
+	UTEST_EQUAL("Output point count", OutPoints.Num(), 10);
+
+	for (int i = 0; i < OutPoints.Num(); ++i)
+	{
+		UTEST_EQUAL("Output points should have a Bounds Min of -2", OutPoints[i].BoundsMin, FVector(-2.0));
+		UTEST_EQUAL("Output points should have a Bounds Min of 2", OutPoints[i].BoundsMax, FVector(2.0));
+		UTEST_EQUAL("Output points Bounds Max should be negative -1", OutPoints[i].Transform.GetScale3D(), FVector(-1.0));
+	}
+
+	return true;
+}
\ No newline at end of file
diff --git a/Engine/Plugins/Experimental/PCG/Source/PCG/Private/Tests/Elements/PCGResetPointCenterTest.cpp b/Engine/Plugins/Experimental/PCG/Source/PCG/Private/Tests/Elements/PCGResetPointCenterTest.cpp
new file mode 100644
--- /dev/null
+++ b/Engine/Plugins/Experimental/PCG/Source/PCG/Private/Tests/Elements/PCGResetPointCenterTest.cpp
@@ -0,0 +1,248 @@
+// Copyright Epic Games, Inc. All Rights Reserved.
+
+#include "Tests/PCGTestsCommon.h"
+
+#include "PCGContext.h"
+#include "Data/PCGPointData.h"
+
+#include "Elements/PCGResetPointCenter.h"
+
+namespace ResetPointCenterCommonTests
+{
+	TUniquePtr<FPCGContext> GenerateTestDataAndRun(FVector PointPosition, FRotator PointRotation, FVector PointScale)
+	{
+		PCGTestsCommon::FTestData TestData;
+		PCGTestsCommon::GenerateSettings<UPCGResetPointCenterSettings>(TestData);
+		UPCGResetPointCenterSettings* Settings = CastChecked<UPCGResetPointCenterSettings>(TestData.Settings);
+		Settings->PointCenterLocation = PointPosition;
+
+		FPCGTaggedData& Inputs = TestData.InputData.TaggedData.Emplace_GetRef();
+		Inputs.Pin = PCGPinConstants::DefaultInputLabel;
+		UPCGPointData* InData = PCGTestsCommon::CreatePointData();
+		Inputs.Data = InData;
+
+		// Setting the scale and rotation explicitly
+		UPCGPointData* InputPointData = Cast<UPCGPointData>(InData);
+		TArray<FPCGPoint>& InPoints = InputPointData->GetMutablePoints();
+		for (FPCGPoint& InPoint : InPoints)
+		{
+			InPoint.Transform.SetRotation(PointRotation.Quaternion());
+			InPoint.Transform.SetScale3D(PointScale);
+		}
+		
+		FPCGElementPtr TestElement = TestData.Settings->GetElement();
+
+		TUniquePtr<FPCGContext> Context = TestData.InitializeTestContext();
+
+		while (!TestElement->Execute(Context.Get())) {}
+
+		return Context;
+	}
+}
+
+IMPLEMENT_CUSTOM_SIMPLE_AUTOMATION_TEST(FPCGResetPointCenterTest_Center, FPCGTestBaseClass, "pcg.tests.ResetPointCenter.Center", PCGTestsCommon::TestFlags)
+
+bool FPCGResetPointCenterTest_Center::RunTest(const FString& Parameters)
+{
+	TUniquePtr<FPCGContext> Context = ResetPointCenterCommonTests::GenerateTestDataAndRun(FVector(0.5), FRotator::ZeroRotator, FVector::One());
+
+	// Point center should be in the center of the bounding box
+	const TArray<FPCGTaggedData>& Outputs = Context->OutputData.TaggedData;
+
+	UTEST_EQUAL("Output count", Outputs.Num(), 1);
+
+	const UPCGPointData* OutPointData = Cast<UPCGPointData>(Outputs[0].Data);
+
+	UTEST_NOT_NULL("Output point data", OutPointData);
+
+	const TArray<FPCGPoint>& OutPoints = OutPointData->GetPoints();
+
+	UTEST_EQUAL("Output point count", OutPoints.Num(), 1);
+
+	UTEST_EQUAL("Output has the correct transform,", OutPoints[0].Transform, FTransform(FRotator::ZeroRotator, FVector(0.0), FVector::One()));
+	UTEST_EQUAL("Output has the correct minimum bounds,", OutPoints[0].BoundsMin, FVector(-1.0));
+	UTEST_EQUAL("Output has the correct maximum bounds,", OutPoints[0].BoundsMax, FVector(1.0));
+
+	return true;
+}
+
+IMPLEMENT_CUSTOM_SIMPLE_AUTOMATION_TEST(FPCGResetPointCenterTest_Zero, FPCGTestBaseClass, "pcg.tests.ResetPointCenter.Zero", PCGTestsCommon::TestFlags)
+
+bool FPCGResetPointCenterTest_Zero::RunTest(const FString& Parameters)
+{
+	TUniquePtr<FPCGContext> Context = ResetPointCenterCommonTests::GenerateTestDataAndRun(FVector::ZeroVector, FRotator::ZeroRotator, FVector::One());
+
+	// Point center should be in the bottom left corner
+	const TArray<FPCGTaggedData>& Outputs = Context->OutputData.TaggedData;
+
+	UTEST_EQUAL("Output count", Outputs.Num(), 1);
+
+	const UPCGPointData* OutPointData = Cast<UPCGPointData>(Outputs[0].Data);
+
+	UTEST_NOT_NULL("Output point data", OutPointData);
+
+	const TArray<FPCGPoint>& OutPoints = OutPointData->GetPoints();
+
+	UTEST_EQUAL("Output point count", OutPoints.Num(), 1);
+
+	UTEST_EQUAL("Output has the correct transform,", OutPoints[0].Transform, FTransform(FRotator::ZeroRotator, FVector(-1.0), FVector::One()));
+	UTEST_EQUAL("Output has the correct minimum bounds,", OutPoints[0].BoundsMin, FVector(0.0));
+	UTEST_EQUAL("Output has the correct maximum bounds,", OutPoints[0].BoundsMax, FVector(2.0));
+
+	return true;
+}
+
+IMPLEMENT_CUSTOM_SIMPLE_AUTOMATION_TEST(FPCGResetPointCenterTest_One, FPCGTestBaseClass, "pcg.tests.ResetPointCenter.One", PCGTestsCommon::TestFlags)
+
+bool FPCGResetPointCenterTest_One::RunTest(const FString& Parameters)
+{
+	TUniquePtr<FPCGContext> Context = ResetPointCenterCommonTests::GenerateTestDataAndRun(FVector::One(), FRotator::ZeroRotator, FVector::One());
+
+	// Point center should be in the top right corner
+	const TArray<FPCGTaggedData>& Outputs = Context->OutputData.TaggedData;
+
+	UTEST_EQUAL("Output count", Outputs.Num(), 1);
+
+	const UPCGPointData* OutPointData = Cast<UPCGPointData>(Outputs[0].Data);
+
+	UTEST_NOT_NULL("Output point data", OutPointData);
+
+	const TArray<FPCGPoint>& OutPoints = OutPointData->GetPoints();
+
+	UTEST_EQUAL("Output point count", OutPoints.Num(), 1);
+
+	UTEST_EQUAL("Output has the correct transform,", OutPoints[0].Transform, FTransform(FRotator::ZeroRotator, FVector(1.0), FVector::One()));
+	UTEST_EQUAL("Output has the correct minimum bounds,", OutPoints[0].BoundsMin, FVector(-2.0));
+	UTEST_EQUAL("Output has the correct maximum bounds,", OutPoints[0].BoundsMax, FVector(0.0));
+
+	return true;
+}
+
+IMPLEMENT_CUSTOM_SIMPLE_AUTOMATION_TEST(FPCGResetPointCenterTest_NegativeOutside, FPCGTestBaseClass, "pcg.tests.ResetPointCenter.NegativeOutside", PCGTestsCommon::TestFlags)
+
+bool FPCGResetPointCenterTest_NegativeOutside::RunTest(const FString& Parameters)
+{
+	TUniquePtr<FPCGContext> Context = ResetPointCenterCommonTests::GenerateTestDataAndRun(FVector(-3.0), FRotator::ZeroRotator, FVector::One());
+
+	// Point center should be outside its bounding range, further bottom left
+	const TArray<FPCGTaggedData>& Outputs = Context->OutputData.TaggedData;
+
+	UTEST_EQUAL("Output count", Outputs.Num(), 1);
+
+	const UPCGPointData* OutPointData = Cast<UPCGPointData>(Outputs[0].Data);
+
+	UTEST_NOT_NULL("Output point data", OutPointData);
+
+	const TArray<FPCGPoint>& OutPoints = OutPointData->GetPoints();
+
+	UTEST_EQUAL("Output point count", OutPoints.Num(), 1);
+
+	UTEST_EQUAL("Output has the correct transform,", OutPoints[0].Transform, FTransform(FRotator::ZeroRotator, FVector(-7.0), FVector::One()));
+	UTEST_EQUAL("Output has the correct minimum bounds,", OutPoints[0].BoundsMin, FVector(6.0));
+	UTEST_EQUAL("Output has the correct maximum bounds,", OutPoints[0].BoundsMax, FVector(8.0));
+
+	return true;
+}
+
+IMPLEMENT_CUSTOM_SIMPLE_AUTOMATION_TEST(FPCGResetPointCenterTest_PositiveOutside, FPCGTestBaseClass, "pcg.tests.ResetPointCenter.PositiveOutside", PCGTestsCommon::TestFlags)
+
+bool FPCGResetPointCenterTest_PositiveOutside::RunTest(const FString& Parameters)
+{
+	TUniquePtr<FPCGContext> Context = ResetPointCenterCommonTests::GenerateTestDataAndRun(FVector(4.0), FRotator::ZeroRotator, FVector::One());
+
+	// Point center should be outside its bounding range, further top right
+	const TArray<FPCGTaggedData>& Outputs = Context->OutputData.TaggedData;
+
+	UTEST_EQUAL("Output count", Outputs.Num(), 1);
+
+	const UPCGPointData* OutPointData = Cast<UPCGPointData>(Outputs[0].Data);
+
+	UTEST_NOT_NULL("Output point data", OutPointData);
+
+	const TArray<FPCGPoint>& OutPoints = OutPointData->GetPoints();
+
+	UTEST_EQUAL("Output point count", OutPoints.Num(), 1);
+
+	UTEST_EQUAL("Output has the correct transform,", OutPoints[0].Transform, FTransform(FRotator::ZeroRotator, FVector(7.0), FVector::One()));
+	UTEST_EQUAL("Output has the correct minimum bounds,", OutPoints[0].BoundsMin, FVector(-8.0));
+	UTEST_EQUAL("Output has the correct maximum bounds,", OutPoints[0].BoundsMax, FVector(-6.0));
+
+	return true;
+}
+
+IMPLEMENT_CUSTOM_SIMPLE_AUTOMATION_TEST(FPCGResetPointCenterTest_NonUniform, FPCGTestBaseClass, "pcg.tests.ResetPointCenter.NonUniform", PCGTestsCommon::TestFlags)
+
+bool FPCGResetPointCenterTest_NonUniform::RunTest(const FString& Parameters)
+{
+	TUniquePtr<FPCGContext> Context = ResetPointCenterCommonTests::GenerateTestDataAndRun(FVector(-1.0, 0.5, 2.0), FRotator::ZeroRotator, FVector::One());
+
+	const TArray<FPCGTaggedData>& Outputs = Context->OutputData.TaggedData;
+
+	UTEST_EQUAL("Output count", Outputs.Num(), 1);
+
+	const UPCGPointData* OutPointData = Cast<UPCGPointData>(Outputs[0].Data);
+
+	UTEST_NOT_NULL("Output point data", OutPointData);
+
+	const TArray<FPCGPoint>& OutPoints = OutPointData->GetPoints();
+
+	UTEST_EQUAL("Output point count", OutPoints.Num(), 1);
+
+	UTEST_EQUAL("Output has the correct transform,", OutPoints[0].Transform, FTransform(FRotator::ZeroRotator, FVector(-3.0, 0.0, 3.0), FVector::One()));
+	UTEST_EQUAL("Output has the correct minimum bounds,", OutPoints[0].BoundsMin, FVector(2.0, -1.0, -4.0));
+	UTEST_EQUAL("Output has the correct maximum bounds,", OutPoints[0].BoundsMax, FVector(4.0, 1.0, -2.0));
+
+	return true;
+}
+
+IMPLEMENT_CUSTOM_SIMPLE_AUTOMATION_TEST(FPCGResetPointCenterTest_CenterRotation, FPCGTestBaseClass, "pcg.tests.ResetPointCenter.CenterRotation", PCGTestsCommon::TestFlags)
+
+bool FPCGResetPointCenterTest_CenterRotation::RunTest(const FString& Parameters)
+{
+	TUniquePtr<FPCGContext> Context = ResetPointCenterCommonTests::GenerateTestDataAndRun(FVector(0.5), FRotator(20.0), FVector::One());
+
+	// Point center should be in the center of the bounding box
+	const TArray<FPCGTaggedData>& Outputs = Context->OutputData.TaggedData;
+
+	UTEST_EQUAL("Output count", Outputs.Num(), 1);
+
+	const UPCGPointData* OutPointData = Cast<UPCGPointData>(Outputs[0].Data);
+
+	UTEST_NOT_NULL("Output point data", OutPointData);
+
+	const TArray<FPCGPoint>& OutPoints = OutPointData->GetPoints();
+
+	UTEST_EQUAL("Output point count", OutPoints.Num(), 1);
+
+	UTEST_EQUAL("Output has the correct transform,", OutPoints[0].Transform, FTransform(FRotator(20.0), FVector(0.0), FVector::One()));
+	UTEST_EQUAL("Output has the correct minimum bounds,", OutPoints[0].BoundsMin, FVector(-1.0));
+	UTEST_EQUAL("Output has the correct maximum bounds,", OutPoints[0].BoundsMax, FVector(1.0));
+
+	return true;
+}
+
+IMPLEMENT_CUSTOM_SIMPLE_AUTOMATION_TEST(FPCGResetPointCenterTest_CenterScale, FPCGTestBaseClass, "pcg.tests.ResetPointCenter.CenterScale", PCGTestsCommon::TestFlags)
+
+bool FPCGResetPointCenterTest_CenterScale::RunTest(const FString& Parameters)
+{
+	TUniquePtr<FPCGContext> Context = ResetPointCenterCommonTests::GenerateTestDataAndRun(FVector(0.5), FRotator::ZeroRotator, FVector(3.0));
+
+	// Point center should be in the center of the bounding box
+	const TArray<FPCGTaggedData>& Outputs = Context->OutputData.TaggedData;
+
+	UTEST_EQUAL("Output count", Outputs.Num(), 1);
+
+	const UPCGPointData* OutPointData = Cast<UPCGPointData>(Outputs[0].Data);
+
+	UTEST_NOT_NULL("Output point data", OutPointData);
+
+	const TArray<FPCGPoint>& OutPoints = OutPointData->GetPoints();
+
+	UTEST_EQUAL("Output point count", OutPoints.Num(), 1);
+
+	UTEST_EQUAL("Output has the correct transform,", OutPoints[0].Transform, FTransform(FRotator::ZeroRotator, FVector(0.0), FVector(3.0)));
+	UTEST_EQUAL("Output has the correct minimum bounds,", OutPoints[0].BoundsMin, FVector(-1.0));
+	UTEST_EQUAL("Output has the correct maximum bounds,", OutPoints[0].BoundsMax, FVector(1.0));
+
+	return true;
+}
\ No newline at end of file
diff --git a/Engine/Plugins/Experimental/PCG/Source/PCG/Public/Elements/PCGApplyScaleToBounds.h b/Engine/Plugins/Experimental/PCG/Source/PCG/Public/Elements/PCGApplyScaleToBounds.h
new file mode 100644
--- /dev/null
+++ b/Engine/Plugins/Experimental/PCG/Source/PCG/Public/Elements/PCGApplyScaleToBounds.h
@@ -0,0 +1,37 @@
+// Copyright Epic Games, Inc. All Rights Reserved.
+#pragma once
+
+#include "PCGSettings.h"
+#include "Elements/PCGPointOperationElementBase.h"
+
+#include "PCGApplyScaleToBounds.generated.h"
+
+/**
+ * Applies the scale of each point to its bounds and resets the scale.
+ */
+UCLASS(BlueprintType, ClassGroup = (Procedural))
+class UPCGApplyScaleToBoundsSettings : public UPCGSettings
+{
+	GENERATED_BODY()
+
+public:
+	//~Begin UPCGSettings interface
+#if WITH_EDITOR
+	virtual FName GetDefaultNodeName() const override { return FName(TEXT("ApplyScaleToBounds")); }
+	virtual FText GetDefaultNodeTitle() const override { return NSLOCTEXT("PCGApplyScaleToBoundsElement", "NodeTitle", "Apply Scale To Bounds"); }
+	virtual FText GetNodeTooltipText() const override { return NSLOCTEXT("PCGApplyScaleToBoundsElement", "NodeTooltip", "Applies the scale of each point to its bounds and resets the scale."); }
+	virtual EPCGSettingsType GetType() const override { return EPCGSettingsType::Spatial; }
+#endif
+
+protected:
+	virtual TArray<FPCGPinProperties> InputPinProperties() const override { return Super::DefaultPointInputPinProperties(); }
+	virtual TArray<FPCGPinProperties> OutputPinProperties() const override { return Super::DefaultPointOutputPinProperties(); }
+	virtual FPCGElementPtr CreateElement() const override;
+	//~End UPCGSettings interface
+};
+
+class FPCGApplyScaleToBoundsElement : public FPCGPointOperationElementBase
+{
+protected:
+	virtual bool ExecuteInternal(FPCGContext* Context) const override;
+};
\ No newline at end of file
diff --git a/Engine/Plugins/Experimental/PCG/Source/PCG/Public/Elements/PCGResetPointCenter.h b/Engine/Plugins/Experimental/PCG/Source/PCG/Public/Elements/PCGResetPointCenter.h
new file mode 100644
--- /dev/null
+++ b/Engine/Plugins/Experimental/PCG/Source/PCG/Public/Elements/PCGResetPointCenter.h
@@ -0,0 +1,43 @@
+// Copyright Epic Games, Inc. All Rights Reserved.
+
+#pragma once
+
+#include "PCGSettings.h"
+#include "Elements/PCGPointOperationElementBase.h"
+
+#include "PCGResetPointCenter.generated.h"
+
+/**
+ * Modify the position of a point within its bounds, while keeping its bounds the same.
+ */
+UCLASS(BlueprintType, ClassGroup = (Procedural))
+class UPCGResetPointCenterSettings : public UPCGSettings
+{
+	GENERATED_BODY()
+
+public:
+	//~Begin UPCGSettings interface
+#if WITH_EDITOR
+	virtual FName GetDefaultNodeName() const override { return FName(TEXT("ResetPointCenter")); }
+	virtual FText GetDefaultNodeTitle() const override { return NSLOCTEXT("PCGResetPointCenterElement", "NodeTitle", "Reset Point Center"); }
+	virtual FText GetNodeTooltipText() const override { return NSLOCTEXT("PCGResetPointCenterElement", "NodeTooltip", "Modify the position of a point within its bounds, while keeping its bounds the same."); }
+	virtual EPCGSettingsType GetType() const override { return EPCGSettingsType::Spatial; }
+#endif
+
+protected:
+	virtual TArray<FPCGPinProperties> InputPinProperties() const override { return Super::DefaultPointInputPinProperties(); }
+	virtual TArray<FPCGPinProperties> OutputPinProperties() const override { return Super::DefaultPointOutputPinProperties(); }
+	virtual FPCGElementPtr CreateElement() const override;
+	//~End UPCGSettings interface
+
+public:
+	/** Set the normalized center of the point */
+	UPROPERTY(BlueprintReadWrite, EditAnywhere, Category = Settings, meta = (PCG_Overridable))
+	FVector PointCenterLocation = FVector(0.5, 0.5, 0.5);
+};
+
+class FPCGResetPointCenterElement : public FPCGPointOperationElementBase
+{
+protected:
+	virtual bool ExecuteInternal(FPCGContext* Context) const override;
+};
\ No newline at end of file
diff --git a/Engine/Plugins/Experimental/PCG/Source/PCG/Public/PCGPoint.h b/Engine/Plugins/Experimental/PCG/Source/PCG/Public/PCGPoint.h
--- a/Engine/Plugins/Experimental/PCG/Source/PCG/Public/PCGPoint.h
+++ b/Engine/Plugins/Experimental/PCG/Source/PCG/Public/PCGPoint.h
@@ -80,6 +80,10 @@ public:
 		BoundsMax += Delta;
 	}
 
+	void ApplyScaleToBounds();
+
+	void ResetPointCenter(const FVector& BoundsRatio);
+
 	using PointCustomPropertyGetter = TFunction<bool(const FPCGPoint&, void*)>;
 	using PointCustomPropertySetter = TFunction<bool(FPCGPoint&, const void*)>;
 
