Prevent pinless data from being corralled into a single pin, was causing spurious warnings following removal of deprecated input pins.

#rnx

diff --git a/Engine/Plugins/Experimental/PCG/Source/PCG/Private/PCGElement.cpp b/Engine/Plugins/Experimental/PCG/Source/PCG/Private/PCGElement.cpp
--- a/Engine/Plugins/Experimental/PCG/Source/PCG/Private/PCGElement.cpp
+++ b/Engine/Plugins/Experimental/PCG/Source/PCG/Private/PCGElement.cpp
@@ -411,7 +411,10 @@ void IPCGElement::CleanupAndValidateOutput(FPCGContext* Context) const
 		{
 			for (FPCGTaggedData& TaggedData : Context->OutputData.TaggedData)
 			{
-				TaggedData.Pin = OutputPinProperties[0].Label;
+				if (!TaggedData.bPinlessData)
+				{
+					TaggedData.Pin = OutputPinProperties[0].Label;
+				}				
 			}
 		}
 
