[PCG] Fix failing test: Editor.Plugins.Tools.PCG.PCGVolumeExposedInCreationMenu

#rnx

diff --git a/Engine/Plugins/Experimental/PCG/Source/PCGEditor/Private/Tests/Display/PCGVolumeExposedInCreationMenu.cpp b/Engine/Plugins/Experimental/PCG/Source/PCGEditor/Private/Tests/Display/PCGVolumeExposedInCreationMenu.cpp
--- a/Engine/Plugins/Experimental/PCG/Source/PCGEditor/Private/Tests/Display/PCGVolumeExposedInCreationMenu.cpp
+++ b/Engine/Plugins/Experimental/PCG/Source/PCGEditor/Private/Tests/Display/PCGVolumeExposedInCreationMenu.cpp
@@ -9,12 +9,12 @@ IMPLEMENT_CUSTOM_SIMPLE_AUTOMATION_TEST(FPCGVolumeExposedInCreationMenu, FPCGTes
 
 bool FPCGVolumeExposedInCreationMenu::RunTest(const FString& Parameters)
 {
-	// create menu queries place mode categories to dislplay 
+	// create menu queries place mode categories to display
 	const IPlacementModeModule& PlacementModeModule = IPlacementModeModule::Get();
 	// PCG volume is located in Volumes category
 	const FName PlacementModeCategoryHandle = TEXT("Volumes");
-	FString PCGDisplayVolumeName = TEXT("PCGVolume");
-	TArray<TSharedPtr<FPlaceableItem>>OutItems;
+	FString PCGDisplayVolumeName = TEXT("PCG Volume");
+	TArray<TSharedPtr<FPlaceableItem>> OutItems;
 
 
 	PlacementModeModule.GetItemsForCategory(PlacementModeCategoryHandle, OutItems);
