[PCG] Update the Attribute (metadata) operations to support multiple inputs in cases where it would make sense

* Now accepts N:1, N:N operations

diff --git a/Engine/Plugins/Experimental/PCG/Source/PCG/Private/Elements/Metadata/PCGMetadataBitwiseOpElement.cpp b/Engine/Plugins/Experimental/PCG/Source/PCG/Private/Elements/Metadata/PCGMetadataBitwiseOpElement.cpp
--- a/Engine/Plugins/Experimental/PCG/Source/PCG/Private/Elements/Metadata/PCGMetadataBitwiseOpElement.cpp
+++ b/Engine/Plugins/Experimental/PCG/Source/PCG/Private/Elements/Metadata/PCGMetadataBitwiseOpElement.cpp
@@ -62,7 +62,7 @@ FName UPCGMetadataBitwiseSettings::GetInputPinLabel(uint32 Index) const
 	}
 }
 
-uint32 UPCGMetadataBitwiseSettings::GetInputPinNum() const
+uint32 UPCGMetadataBitwiseSettings::GetOperandNum() const
 {
 	return (Operation == EPCGMedadataBitwiseOperation::Not) ? 1 : 2;
 }
@@ -136,7 +136,7 @@ uint16 UPCGMetadataBitwiseSettings::GetOutputType(uint16 InputTypeId) const
 	return (uint16)EPCGMetadataTypes::Integer64;
 }
 
-bool FPCGMetadataBitwiseElement::DoOperation(FOperationData& OperationData) const
+bool FPCGMetadataBitwiseElement::DoOperation(PCGMetadataOps::FOperationData& OperationData) const
 {
 	TRACE_CPUPROFILER_EVENT_SCOPE(FPCGMetadataBitwiseElement::Execute);
 
diff --git a/Engine/Plugins/Experimental/PCG/Source/PCG/Private/Elements/Metadata/PCGMetadataBooleanOpElement.cpp b/Engine/Plugins/Experimental/PCG/Source/PCG/Private/Elements/Metadata/PCGMetadataBooleanOpElement.cpp
--- a/Engine/Plugins/Experimental/PCG/Source/PCG/Private/Elements/Metadata/PCGMetadataBooleanOpElement.cpp
+++ b/Engine/Plugins/Experimental/PCG/Source/PCG/Private/Elements/Metadata/PCGMetadataBooleanOpElement.cpp
@@ -62,7 +62,7 @@ FName UPCGMetadataBooleanSettings::GetInputPinLabel(uint32 Index) const
 	}
 }
 
-uint32 UPCGMetadataBooleanSettings::GetInputPinNum() const
+uint32 UPCGMetadataBooleanSettings::GetOperandNum() const
 {
 	return (Operation == EPCGMedadataBooleanOperation::Not) ? 1 : 2;
 }
@@ -136,7 +136,7 @@ FPCGElementPtr UPCGMetadataBooleanSettings::CreateElement() const
 	return MakeShared<FPCGMetadataBooleanElement>();
 }
 
-bool FPCGMetadataBooleanElement::DoOperation(FOperationData& OperationData) const
+bool FPCGMetadataBooleanElement::DoOperation(PCGMetadataOps::FOperationData& OperationData) const
 {
 	TRACE_CPUPROFILER_EVENT_SCOPE(FPCGMetadataBooleanElement::Execute);
 
diff --git a/Engine/Plugins/Experimental/PCG/Source/PCG/Private/Elements/Metadata/PCGMetadataBreakTransform.cpp b/Engine/Plugins/Experimental/PCG/Source/PCG/Private/Elements/Metadata/PCGMetadataBreakTransform.cpp
--- a/Engine/Plugins/Experimental/PCG/Source/PCG/Private/Elements/Metadata/PCGMetadataBreakTransform.cpp
+++ b/Engine/Plugins/Experimental/PCG/Source/PCG/Private/Elements/Metadata/PCGMetadataBreakTransform.cpp
@@ -57,7 +57,7 @@ FName UPCGMetadataBreakTransformSettings::GetOutputPinLabel(uint32 Index) const
 	}
 }
 
-uint32 UPCGMetadataBreakTransformSettings::GetOutputPinNum() const
+uint32 UPCGMetadataBreakTransformSettings::GetResultNum() const
 {
 	return 3;
 }
@@ -107,7 +107,7 @@ FPCGElementPtr UPCGMetadataBreakTransformSettings::CreateElement() const
 	return MakeShared<FPCGMetadataBreakTransformElement>();
 }
 
-bool FPCGMetadataBreakTransformElement::DoOperation(FOperationData& OperationData) const
+bool FPCGMetadataBreakTransformElement::DoOperation(PCGMetadataOps::FOperationData& OperationData) const
 {
 	TRACE_CPUPROFILER_EVENT_SCOPE(FPCGMetadataBreakVectorElement::Execute);
 
diff --git a/Engine/Plugins/Experimental/PCG/Source/PCG/Private/Elements/Metadata/PCGMetadataBreakVector.cpp b/Engine/Plugins/Experimental/PCG/Source/PCG/Private/Elements/Metadata/PCGMetadataBreakVector.cpp
--- a/Engine/Plugins/Experimental/PCG/Source/PCG/Private/Elements/Metadata/PCGMetadataBreakVector.cpp
+++ b/Engine/Plugins/Experimental/PCG/Source/PCG/Private/Elements/Metadata/PCGMetadataBreakVector.cpp
@@ -107,7 +107,7 @@ FName UPCGMetadataBreakVectorSettings::GetOutputPinLabel(uint32 Index) const
 	}
 }
 
-uint32 UPCGMetadataBreakVectorSettings::GetOutputPinNum() const
+uint32 UPCGMetadataBreakVectorSettings::GetResultNum() const
 {
 	return 4;
 }
@@ -160,7 +160,7 @@ FPCGElementPtr UPCGMetadataBreakVectorSettings::CreateElement() const
 	return MakeShared<FPCGMetadataBreakVectorElement>();
 }
 
-bool FPCGMetadataBreakVectorElement::DoOperation(FOperationData& OperationData) const
+bool FPCGMetadataBreakVectorElement::DoOperation(PCGMetadataOps::FOperationData& OperationData) const
 {
 	TRACE_CPUPROFILER_EVENT_SCOPE(FPCGMetadataBreakVectorElement::Execute);
 
diff --git a/Engine/Plugins/Experimental/PCG/Source/PCG/Private/Elements/Metadata/PCGMetadataCompareOpElement.cpp b/Engine/Plugins/Experimental/PCG/Source/PCG/Private/Elements/Metadata/PCGMetadataCompareOpElement.cpp
--- a/Engine/Plugins/Experimental/PCG/Source/PCG/Private/Elements/Metadata/PCGMetadataCompareOpElement.cpp
+++ b/Engine/Plugins/Experimental/PCG/Source/PCG/Private/Elements/Metadata/PCGMetadataCompareOpElement.cpp
@@ -76,7 +76,7 @@ FName UPCGMetadataCompareSettings::GetInputPinLabel(uint32 Index) const
 	}
 }
 
-uint32 UPCGMetadataCompareSettings::GetInputPinNum() const
+uint32 UPCGMetadataCompareSettings::GetOperandNum() const
 {
 	return 2;
 }
@@ -162,7 +162,7 @@ FPCGElementPtr UPCGMetadataCompareSettings::CreateElement() const
 	return MakeShared<FPCGMetadataCompareElement>();
 }
 
-bool FPCGMetadataCompareElement::DoOperation(FOperationData& OperationData) const
+bool FPCGMetadataCompareElement::DoOperation(PCGMetadataOps::FOperationData& OperationData) const
 {
 	TRACE_CPUPROFILER_EVENT_SCOPE(FPCGMetadataTrigElement::Execute);
 
diff --git a/Engine/Plugins/Experimental/PCG/Source/PCG/Private/Elements/Metadata/PCGMetadataMakeRotator.cpp b/Engine/Plugins/Experimental/PCG/Source/PCG/Private/Elements/Metadata/PCGMetadataMakeRotator.cpp
--- a/Engine/Plugins/Experimental/PCG/Source/PCG/Private/Elements/Metadata/PCGMetadataMakeRotator.cpp
+++ b/Engine/Plugins/Experimental/PCG/Source/PCG/Private/Elements/Metadata/PCGMetadataMakeRotator.cpp
@@ -39,7 +39,7 @@ FName UPCGMetadataMakeRotatorSettings::GetInputPinLabel(uint32 Index) const
 	}
 }
 
-uint32 UPCGMetadataMakeRotatorSettings::GetInputPinNum() const
+uint32 UPCGMetadataMakeRotatorSettings::GetOperandNum() const
 {
 	switch (Operation)
 	{
@@ -151,7 +151,7 @@ FString UPCGMetadataMakeRotatorSettings::GetDefaultValueString(uint32 Index) con
 }
 #endif // WITH_EDITOR
 
-bool FPCGMetadataMakeRotatorElement::DoOperation(FOperationData& OperationData) const
+bool FPCGMetadataMakeRotatorElement::DoOperation(PCGMetadataOps::FOperationData& OperationData) const
 {
 	TRACE_CPUPROFILER_EVENT_SCOPE(FPCGMetadataMakeRotatorElement::Execute);
 
diff --git a/Engine/Plugins/Experimental/PCG/Source/PCG/Private/Elements/Metadata/PCGMetadataMakeTransform.cpp b/Engine/Plugins/Experimental/PCG/Source/PCG/Private/Elements/Metadata/PCGMetadataMakeTransform.cpp
--- a/Engine/Plugins/Experimental/PCG/Source/PCG/Private/Elements/Metadata/PCGMetadataMakeTransform.cpp
+++ b/Engine/Plugins/Experimental/PCG/Source/PCG/Private/Elements/Metadata/PCGMetadataMakeTransform.cpp
@@ -70,7 +70,7 @@ FName UPCGMetadataMakeTransformSettings::GetInputPinLabel(uint32 Index) const
 	}
 }
 
-uint32 UPCGMetadataMakeTransformSettings::GetInputPinNum() const
+uint32 UPCGMetadataMakeTransformSettings::GetOperandNum() const
 {
 	return 3;
 }
@@ -187,7 +187,7 @@ FPCGElementPtr UPCGMetadataMakeTransformSettings::CreateElement() const
 	return MakeShared<FPCGMetadataMakeTransformElement>();
 }
 
-bool FPCGMetadataMakeTransformElement::DoOperation(FOperationData& OperationData) const
+bool FPCGMetadataMakeTransformElement::DoOperation(PCGMetadataOps::FOperationData& OperationData) const
 {
 	TRACE_CPUPROFILER_EVENT_SCOPE(FPCGMetadataMakeTransformElement::Execute);
 
diff --git a/Engine/Plugins/Experimental/PCG/Source/PCG/Private/Elements/Metadata/PCGMetadataMakeVector.cpp b/Engine/Plugins/Experimental/PCG/Source/PCG/Private/Elements/Metadata/PCGMetadataMakeVector.cpp
--- a/Engine/Plugins/Experimental/PCG/Source/PCG/Private/Elements/Metadata/PCGMetadataMakeVector.cpp
+++ b/Engine/Plugins/Experimental/PCG/Source/PCG/Private/Elements/Metadata/PCGMetadataMakeVector.cpp
@@ -122,7 +122,7 @@ FName UPCGMetadataMakeVectorSettings::GetInputPinLabel(uint32 Index) const
 	}
 }
 
-uint32 UPCGMetadataMakeVectorSettings::GetInputPinNum() const
+uint32 UPCGMetadataMakeVectorSettings::GetOperandNum() const
 {
 	if (OutputType == EPCGMetadataTypes::Vector2 ||
 		(OutputType == EPCGMetadataTypes::Vector && MakeVector3Op == EPCGMetadataMakeVector3::Vector2AndValue) ||
@@ -249,7 +249,7 @@ FPCGElementPtr UPCGMetadataMakeVectorSettings::CreateElement() const
 	return MakeShared<FPCGMetadataMakeVectorElement>();
 }
 
-bool FPCGMetadataMakeVectorElement::DoOperation(FOperationData& OperationData) const
+bool FPCGMetadataMakeVectorElement::DoOperation(PCGMetadataOps::FOperationData& OperationData) const
 {
 	TRACE_CPUPROFILER_EVENT_SCOPE(FPCGMetadataMakeVectorElement::Execute);
 
diff --git a/Engine/Plugins/Experimental/PCG/Source/PCG/Private/Elements/Metadata/PCGMetadataMathsOpElement.cpp b/Engine/Plugins/Experimental/PCG/Source/PCG/Private/Elements/Metadata/PCGMetadataMathsOpElement.cpp
--- a/Engine/Plugins/Experimental/PCG/Source/PCG/Private/Elements/Metadata/PCGMetadataMathsOpElement.cpp
+++ b/Engine/Plugins/Experimental/PCG/Source/PCG/Private/Elements/Metadata/PCGMetadataMathsOpElement.cpp
@@ -120,7 +120,7 @@ namespace PCGMetadataMathsSettings
 		case EPCGMedadataMathsOperation::Multiply:
 			return Value1 * Value2;
 		case EPCGMedadataMathsOperation::Divide:
-			return Value1 / Value2;
+			return (Value2 != T{0}) ? (Value1 / Value2) : T{0}; // To mirror FMath
 		case EPCGMedadataMathsOperation::Max:
 			return PCGMetadataMaths::Max(Value1, Value2);
 		case EPCGMedadataMathsOperation::Min:
@@ -195,7 +195,7 @@ FName UPCGMetadataMathsSettings::GetInputPinLabel(uint32 Index) const
 	}
 }
 
-uint32 UPCGMetadataMathsSettings::GetInputPinNum() const
+uint32 UPCGMetadataMathsSettings::GetOperandNum() const
 {
 	if (PCGMetadataMathsSettings::IsUnaryOp(Operation))
 	{
@@ -318,7 +318,7 @@ FPCGElementPtr UPCGMetadataMathsSettings::CreateElement() const
 	return MakeShared<FPCGMetadataMathsElement>();
 }
 
-bool FPCGMetadataMathsElement::DoOperation(FOperationData& OperationData) const
+bool FPCGMetadataMathsElement::DoOperation(PCGMetadataOps::FOperationData& OperationData) const
 {
 	TRACE_CPUPROFILER_EVENT_SCOPE(FPCGMetadataMathsElement::Execute);
 
diff --git a/Engine/Plugins/Experimental/PCG/Source/PCG/Private/Elements/Metadata/PCGMetadataOpElementBase.cpp b/Engine/Plugins/Experimental/PCG/Source/PCG/Private/Elements/Metadata/PCGMetadataOpElementBase.cpp
--- a/Engine/Plugins/Experimental/PCG/Source/PCG/Private/Elements/Metadata/PCGMetadataOpElementBase.cpp
+++ b/Engine/Plugins/Experimental/PCG/Source/PCG/Private/Elements/Metadata/PCGMetadataOpElementBase.cpp
@@ -7,7 +7,6 @@
 #include "PCGParamData.h"
 #include "PCGPin.h"
 #include "Data/PCGPointData.h"
-#include "Data/PCGSpatialData.h"
 #include "Elements/Metadata/PCGMetadataElementCommon.h"
 #include "Metadata/Accessors/PCGAttributeAccessorHelpers.h"
 
@@ -43,13 +42,13 @@ EPCGDataType UPCGMetadataSettingsBase::GetInputPinType(uint32 Index) const
 
 TArray<FName> UPCGMetadataSettingsBase::GetOutputDataFromPinOptions() const
 {
-	const uint32 NumInputs = GetInputPinNum();
+	const uint32 OperandNum = GetOperandNum();
 
 	TArray<FName> AllOptions;
-	AllOptions.SetNum(NumInputs + 1);
+	AllOptions.SetNum(OperandNum + 1);
 	AllOptions[0] = PCGMetadataSettingsBaseConstants::DefaultOutputDataFromPinName;
 
-	for (uint32 Index = 0; Index < NumInputs; ++Index)
+	for (uint32 Index = 0; Index < OperandNum; ++Index)
 	{
 		AllOptions[Index + 1] = GetInputPinLabel(Index);
 	}
@@ -61,7 +60,7 @@ uint32 UPCGMetadataSettingsBase::GetInputPinIndex(FName InPinLabel) const
 {
 	if (InPinLabel != PCGMetadataSettingsBaseConstants::DefaultOutputDataFromPinName)
 	{
-		for (uint32 Index = 0; Index < GetInputPinNum(); ++Index)
+		for (uint32 Index = 0; Index < GetOperandNum(); ++Index)
 		{
 			if (InPinLabel == GetInputPinLabel(Index))
 			{
@@ -70,29 +69,29 @@ uint32 UPCGMetadataSettingsBase::GetInputPinIndex(FName InPinLabel) const
 		}
 	}
 
-	return (uint32)INDEX_NONE;
+	return static_cast<uint32>(INDEX_NONE);
 }
 
 uint32 UPCGMetadataSettingsBase::GetInputPinToForward() const
 {
-	const uint32 NumberOfInputs = GetInputPinNum();
+	const uint32 OperandNum = GetOperandNum();
 	uint32 InputPinToForward = 0;
 
 	// If there is only one input, it is trivial
-	if (NumberOfInputs != 1)
+	if (OperandNum != 1)
 	{
 		// Heuristic:
 		//	* If OutputDataFromPin is set, use this value
 		//	* If there are connected pins, use the first spatial input (not Any)
 		//	* Otherwise, take the first pin
 		const uint32 OutputDataFromPinIndex = GetInputPinIndex(OutputDataFromPin);
-		if (OutputDataFromPinIndex != (uint32)INDEX_NONE)
+		if (OutputDataFromPinIndex != static_cast<uint32>(INDEX_NONE))
 		{
 			InputPinToForward = OutputDataFromPinIndex;
 		}
 		else
 		{
-			for (uint32 InputPinIndex = 0; InputPinIndex < NumberOfInputs; ++InputPinIndex)
+			for (uint32 InputPinIndex = 0; InputPinIndex < OperandNum; ++InputPinIndex)
 			{
 				const EPCGDataType PinType = GetInputPinType(InputPinIndex);
 
@@ -111,8 +110,8 @@ uint32 UPCGMetadataSettingsBase::GetInputPinToForward() const
 EPCGDataType UPCGMetadataSettingsBase::GetCurrentPinTypes(const UPCGPin* InPin) const
 {
 	check(InPin);
-	const uint32 NumberOfInputs = GetInputPinNum();
-	if (!InPin->IsOutputPin() || NumberOfInputs == 0)
+	const uint32 OperandNum = GetOperandNum();
+	if (!InPin->IsOutputPin() || OperandNum == 0)
 	{
 		// Fall back to default for input pins, or if no input pins present from which to obtain type
 		return Super::GetCurrentPinTypes(InPin);
@@ -128,7 +127,7 @@ void UPCGMetadataSettingsBase::PostEditChangeProperty(FPropertyChangedEvent& Pro
 	Super::PostEditChangeProperty(PropertyChangedEvent);
 
 	// Make sure the output data from pin value is always valid. Reset it otherwise.
-	if (GetInputPinIndex(OutputDataFromPin) == (uint32)INDEX_NONE)
+	if (GetInputPinIndex(OutputDataFromPin) == static_cast<uint32>(INDEX_NONE))
 	{
 		OutputDataFromPin = PCGMetadataSettingsBaseConstants::DefaultOutputDataFromPinName;
 	}
@@ -141,13 +140,13 @@ bool UPCGMetadataSettingsBase::CanEditChange(const FProperty* InProperty) const
 		return false;
 	}
 
-	return InProperty && ((InProperty->GetFName() != GET_MEMBER_NAME_CHECKED(UPCGMetadataSettingsBase, OutputDataFromPin)) || (GetInputPinNum() != 1));
+	return InProperty && ((InProperty->GetFName() != GET_MEMBER_NAME_CHECKED(UPCGMetadataSettingsBase, OutputDataFromPin)) || (GetOperandNum() != 1));
 }
 
 bool UPCGMetadataSettingsBase::GetPinExtraIcon(const UPCGPin* InPin, FName& OutExtraIcon, FText& OutTooltip) const
 {
 	// Only set the arrow if the output data from pin is forced.
-	if (GetInputPinNum() > 1 && InPin && OutputDataFromPin == InPin->Properties.Label)
+	if (GetOperandNum() > 1 && InPin && OutputDataFromPin == InPin->Properties.Label)
 	{
 		OutExtraIcon = TEXT("Icons.ArrowRight");
 		return true;
@@ -193,18 +192,18 @@ TArray<FPCGPinProperties> UPCGMetadataSettingsBase::InputPinProperties() const
 {
 	TArray<FPCGPinProperties> PinProperties;
 
-	for (uint32 InputPinIndex = 0; InputPinIndex < GetInputPinNum(); ++InputPinIndex)
+	for (uint32 InputPinIndex = 0; InputPinIndex < GetOperandNum(); ++InputPinIndex)
 	{
 		const FName PinLabel = GetInputPinLabel(InputPinIndex);
 		if (PinLabel != NAME_None)
 		{
-			PinProperties.Emplace(PinLabel, EPCGDataType::Any, /*bAllowMultipleConnections=*/false);
+			PinProperties.Emplace(PinLabel, EPCGDataType::Any);
 
 #if WITH_EDITOR
 			TArray<FText> AllTooltips;
 			TArray<FString> SupportedTypes;
 
-			for (uint8 TypeId = 0; TypeId < (uint8)EPCGMetadataTypes::Count; ++TypeId)
+			for (uint8 TypeId = 0; TypeId < static_cast<uint8>(EPCGMetadataTypes::Count); ++TypeId)
 			{
 				bool DummyBool;
 				if (IsSupportedInputType(TypeId, InputPinIndex, DummyBool))
@@ -212,13 +211,13 @@ TArray<FPCGPinProperties> UPCGMetadataSettingsBase::InputPinProperties() const
 					SupportedTypes.Add(PCG::Private::GetTypeName(TypeId));
 				}
 			}
-			
+
 			if (!SupportedTypes.IsEmpty())
 			{
 				AllTooltips.Add(FText::Format(LOCTEXT("PinTooltipSupportedTypes", "Supported types: {0}"), FText::FromString(FString::Join(SupportedTypes, TEXT(", ")))));
 			}
 
-			if (GetInputPinNum() > 1 && OutputDataFromPin == PinLabel)
+			if (GetOperandNum() > 1 && OutputDataFromPin == PinLabel)
 			{
 				AllTooltips.Add(LOCTEXT("PinTooltipForwardInput", "This input will be forwarded to the output."));
 			}
@@ -242,8 +241,8 @@ TArray<FPCGPinProperties> UPCGMetadataSettingsBase::InputPinProperties() const
 TArray<FPCGPinProperties> UPCGMetadataSettingsBase::OutputPinProperties() const
 {
 	TArray<FPCGPinProperties> PinProperties;
-	
-	for (uint32 OutputPinIndex = 0; OutputPinIndex < GetOutputPinNum(); ++OutputPinIndex)
+
+	for (uint32 OutputPinIndex = 0; OutputPinIndex < GetResultNum(); ++OutputPinIndex)
 	{
 		const FName PinLabel = GetOutputPinLabel(OutputPinIndex);
 		if (PinLabel != NAME_None)
@@ -255,176 +254,91 @@ TArray<FPCGPinProperties> UPCGMetadataSettingsBase::OutputPinProperties() const
 	return PinProperties;
 }
 
-void FPCGMetadataElementBase::PassthroughInputs(FPCGContext* InContext, const UPCGMetadataSettingsBase* InSettings) const
+void FPCGMetadataElementBase::PassthroughInput(FPCGContext* Context, TArray<FPCGTaggedData>& Outputs, const int32 Index) const
 {
-	check(InContext && InSettings);
-
-	const uint32 NumberOfOutputs = InSettings->GetOutputPinNum();
-	uint32 InputPinToForward = InSettings->GetInputPinToForward();
+	check(Context);
 
-	TArray<FPCGTaggedData> InputsToForward = InContext->InputData.GetInputsByPin(InSettings->GetInputPinLabel(InputPinToForward));
-	TArray<FPCGTaggedData>& Outputs = InContext->OutputData.TaggedData;
+	const UPCGMetadataSettingsBase* Settings = Context->GetInputSettings<UPCGMetadataSettingsBase>();
+	check(Settings);
 
-	// Empty and reserve memory at the same time.
-	Outputs.Empty(InputsToForward.Num() * NumberOfOutputs);
+	const uint32 NumberOfOutputs = Settings->GetResultNum();
+	const uint32 PrimaryPinIndex = Settings->GetInputPinToForward();
+	TArray<FPCGTaggedData> InputsToForward = Context->InputData.GetSpatialInputsByPin(Settings->GetInputPinLabel(PrimaryPinIndex));
 
-	for (uint32 i = 0; i < NumberOfOutputs; ++i)
+	// Passthrough this single input to all of the outputs
+	for (uint32 I = 0; I < NumberOfOutputs; ++I)
 	{
-		const FName OutputPin = InSettings->GetOutputPinLabel(i);
-		for (const FPCGTaggedData& Input : InputsToForward)
-		{
-			Outputs.Emplace_GetRef(Input).Pin = OutputPin;
-		}
+		Outputs.Emplace_GetRef(InputsToForward[Index]).Pin = Settings->GetOutputPinLabel(I);
 	}
 }
 
-bool FPCGMetadataElementBase::ExecuteInternal(FPCGContext* Context) const
+void FPCGMetadataElementBase::PassthroughAllInputs(FPCGContext* Context, TArray<FPCGTaggedData>& Outputs) const
 {
-	TRACE_CPUPROFILER_EVENT_SCOPE(FPCGMetadataElementBase::Execute);
+	check(Context);
 
 	const UPCGMetadataSettingsBase* Settings = Context->GetInputSettings<UPCGMetadataSettingsBase>();
 	check(Settings);
 
-	const uint32 NumberOfInputs = Settings->GetInputPinNum();
-	const uint32 NumberOfOutputs = Settings->GetOutputPinNum();
+	const uint32 NumberOfOutputs = Settings->GetResultNum();
+	const uint32 PrimaryPinIndex = Settings->GetInputPinToForward();
+	TArray<FPCGTaggedData> InputsToForward = Context->InputData.GetSpatialInputsByPin(Settings->GetInputPinLabel(PrimaryPinIndex));
 
-	check(NumberOfInputs > 0);
-	check(NumberOfOutputs <= UPCGMetadataSettingsBase::MaxNumberOfOutputs);
-
-	const TArray<FPCGTaggedData>& Inputs = Context->InputData.TaggedData;
-	TArray<FPCGTaggedData>& Outputs = Context->OutputData.TaggedData;
-
-	// Gathering all the inputs metadata
-	TArray<const UPCGMetadata*> SourceMetadata;
-	TArray<const FPCGMetadataAttributeBase*> SourceAttribute;
-	TArray<FPCGTaggedData> InputTaggedData;
-	SourceMetadata.SetNum(NumberOfInputs);
-	SourceAttribute.SetNum(NumberOfInputs);
-	InputTaggedData.SetNum(NumberOfInputs);
-
-	uint32 InputPinToForward = Settings->GetInputPinToForward();
+	Outputs.Empty(InputsToForward.Num() * NumberOfOutputs);
 
-	for (uint32 InputPinIndex = 0; InputPinIndex < NumberOfInputs; ++InputPinIndex)
+	// Passthrough all inputs to all of the outputs
+	for (uint32 I = 0; I < NumberOfOutputs; ++I)
 	{
-		FName PinLabel = Settings->GetInputPinLabel(InputPinIndex);
-		const bool bIsInputConnected = Context->Node ? Context->Node->IsInputPinConnected(PinLabel) : false;
-		TArray<FPCGTaggedData> InputData = Context->InputData.GetInputsByPin(PinLabel);
-
-		// If input data is empty and not connected but we can have default values, add a "fake" param data input.
-		if (InputData.IsEmpty() && !bIsInputConnected && Settings->DoesInputSupportDefaultValue(InputPinIndex))
-		{
-			FPCGTaggedData& DummyData = InputData.Emplace_GetRef();
-			DummyData.Pin = PinLabel;
-			DummyData.Data = Settings->CreateDefaultValueParam(InputPinIndex);
-			if (!DummyData.Data)
-			{
-				PCGE_LOG(Error, GraphAndLog, FText::Format(LOCTEXT("CantCreateDefaultValue", "Pin {0} supports default value but we could not create it."), FText::FromName(PinLabel)));
-				return true;
-			}
-		}
-
-		if (InputData.IsEmpty())
-		{
-			// If we have no data, just passthrough the input and return
-			PCGE_LOG(Verbose, LogOnly, FText::Format(LOCTEXT("MissingInputDataForPin", "No data provided on pin {0}"), FText::FromName(PinLabel)));
-			PassthroughInputs(Context, Settings);
-			return true;
-		}
-		else if (InputData.Num() > 1)
-		{
-			PCGE_LOG(Error, GraphAndLog, FText::Format(LOCTEXT("TooMuchDataForPin", "Too many data items ({0}) provided on pin {1}"), InputData.Num(), FText::FromName(PinLabel)));
-			return true;
-		}
-
-		// By construction, there can only be one of then(hence the 0 index)
-		InputTaggedData[InputPinIndex] = MoveTemp(InputData[0]);
-
-		// Check if we have any points
-		if (const UPCGPointData* PointInput = Cast<const UPCGPointData>(InputTaggedData[InputPinIndex].Data))
-		{
-			if (PointInput->GetPoints().IsEmpty())
-			{
-				// If we have no points, just passthrough the input and return
-				PCGE_LOG(Verbose, LogOnly, FText::Format(LOCTEXT("NoPointsForPin", "No points in point data provided on pin {0}"), FText::FromName(PinLabel)));
-				PassthroughInputs(Context, Settings);
-				return true;
-			}
-		}
-
-		// Only gather Spacial and Params input. 
-		if (const UPCGSpatialData* SpatialInput = Cast<const UPCGSpatialData>(InputTaggedData[InputPinIndex].Data))
-		{
-			SourceMetadata[InputPinIndex] = SpatialInput->Metadata;
-		}
-		else if (const UPCGParamData* ParamsInput = Cast<const UPCGParamData>(InputTaggedData[InputPinIndex].Data))
-		{
-			SourceMetadata[InputPinIndex] = ParamsInput->Metadata;
-		}
-		else
+		const FName OutputPin = Settings->GetOutputPinLabel(I);
+		for (const FPCGTaggedData& Input : InputsToForward)
 		{
-			// Since this aborts execution, and the user can fix it, it should be a node error
-			PCGE_LOG(Error, GraphAndLog, FText::Format(LOCTEXT("InvalidInputDataTypeForPin", "Invalid data provided on pin {0}, must be of type Spatial or Attribute Set"), FText::FromName(PinLabel)));
-			return true;
+			Outputs.Emplace_GetRef(Input).Pin = OutputPin;
 		}
 	}
+}
 
-	FOperationData OperationData;
-	OperationData.InputAccessors.SetNum(NumberOfInputs);
-	OperationData.InputKeys.SetNum(NumberOfInputs);
-	OperationData.InputSources.SetNum(NumberOfInputs);
-	TArray<int32> NumberOfElements;
-	NumberOfElements.SetNum(NumberOfInputs);
-
-	OperationData.MostComplexInputType = (uint16)EPCGMetadataTypes::Unknown;
-	OperationData.NumberOfElementsToProcess = -1;
-
-	bool bNoOperationNeeded = false;
+namespace PCGMetadataOpPrivate
+{
+	using ContextType = FPCGMetadataElementBase::ContextType;
+	using ExecStateType = FPCGMetadataElementBase::ExecStateType;
 
-	auto CreateInputAccessorAndValidate = [this, Settings, InputPinToForward, &InputTaggedData, Context, &SourceMetadata, &OperationData, &NumberOfElements, &bNoOperationNeeded, &SourceAttribute](uint32 Index) -> bool
+	void CreateAccessor(const UPCGMetadataSettingsBase* Settings, const FPCGTaggedData& InputData, PCGMetadataOps::FOperationData& OperationData, const int32 Index)
 	{
-		// First we verify if the input data match the first one. If the pin to forward is not connected, behave like a param data
-		const UClass* InputPinToForwardClass = (InputTaggedData[InputPinToForward].Data ? InputTaggedData[InputPinToForward].Data->GetClass() : UPCGParamData::StaticClass());
+		OperationData.InputSources[Index] = Settings->GetInputSource(Index).CopyAndFixLast(InputData.Data);
+		const FPCGAttributePropertyInputSelector& InputSource = OperationData.InputSources[Index];
 
-		if (Index != InputPinToForward && 
-			InputPinToForwardClass != InputTaggedData[Index].Data->GetClass() &&
-			!InputTaggedData[Index].Data->IsA<UPCGParamData>())
-		{
-			PCGE_LOG(Error, GraphAndLog, FText::Format(LOCTEXT("InputTypeMismatch", "Data on pin {0} is not of the same type than on pin {1} and is not an Attribute Set. This is not supported."), FText::FromName(InputTaggedData[Index].Pin), FText::FromName(InputTaggedData[InputPinToForward].Pin)));
-			return false;
-		}
+		OperationData.InputAccessors[Index] = PCGAttributeAccessorHelpers::CreateConstAccessor(InputData.Data, InputSource);
+		OperationData.InputKeys[Index] = PCGAttributeAccessorHelpers::CreateConstKeys(InputData.Data, InputSource);
+	}
 
-		OperationData.InputSources[Index] = Settings->GetInputSource(Index).CopyAndFixLast(InputTaggedData[Index].Data);
+	bool ValidateAccessor(const FPCGContext* Context, const UPCGMetadataSettingsBase* Settings, const FPCGTaggedData& InputData, PCGMetadataOps::FOperationData& OperationData, int32 Index)
+	{
 		const FPCGAttributePropertyInputSelector& InputSource = OperationData.InputSources[Index];
 		const FText InputSourceText = InputSource.GetDisplayText();
 
-		OperationData.InputAccessors[Index] = PCGAttributeAccessorHelpers::CreateConstAccessor(InputTaggedData[Index].Data, InputSource);
-		OperationData.InputKeys[Index] = PCGAttributeAccessorHelpers::CreateConstKeys(InputTaggedData[Index].Data, InputSource);
-
 		if (!OperationData.InputAccessors[Index].IsValid() || !OperationData.InputKeys[Index].IsValid())
 		{
-			PCGE_LOG(Error, GraphAndLog, FText::Format(LOCTEXT("AttributeDoesNotExist", "Attribute/Property '{0}' from pin {1} does not exist"), InputSourceText, FText::FromName(InputTaggedData[Index].Pin)));
+			PCGE_LOG_C(Error, GraphAndLog, Context, FText::Format(LOCTEXT("AttributeDoesNotExist", "Attribute/Property '{0}' from pin {1} does not exist"), InputSourceText, FText::FromName(InputData.Pin)));
 			return false;
 		}
 
-		uint16 AttributeTypeId = OperationData.InputAccessors[Index]->GetUnderlyingType();
+		const uint16 AttributeTypeId = OperationData.InputAccessors[Index]->GetUnderlyingType();
 
 		// Then verify that the type is OK
 		bool bHasSpecialRequirement = false;
 		if (!Settings->IsSupportedInputType(AttributeTypeId, Index, bHasSpecialRequirement))
 		{
 			const FText AttributeTypeName = PCG::Private::GetTypeNameText(AttributeTypeId);
-			PCGE_LOG(Error, GraphAndLog, FText::Format(LOCTEXT("UnsupportedAttributeType", "Attribute/Property '{0}' from pin {1} is not a supported type ('{2}')"), 
+			PCGE_LOG_C(Error, GraphAndLog, Context, FText::Format(LOCTEXT("UnsupportedAttributeType", "Attribute/Property '{0}' from pin {1} is not a supported type ('{2}')"),
 				InputSourceText,
-				FText::FromName(InputTaggedData[Index].Pin),
-				AttributeTypeName
-			));
+				FText::FromName(InputData.Pin),
+				AttributeTypeName));
 			return false;
 		}
 
 		if (!bHasSpecialRequirement)
 		{
 			// In this case, check if we have a more complex type, or if we can broadcast to the most complex type.
-			if (OperationData.MostComplexInputType == (uint16)EPCGMetadataTypes::Unknown || PCG::Private::IsMoreComplexType(AttributeTypeId, OperationData.MostComplexInputType))
+			if (OperationData.MostComplexInputType == static_cast<uint16>(EPCGMetadataTypes::Unknown) || PCG::Private::IsMoreComplexType(AttributeTypeId, OperationData.MostComplexInputType))
 			{
 				OperationData.MostComplexInputType = AttributeTypeId;
 			}
@@ -432,206 +346,373 @@ bool FPCGMetadataElementBase::ExecuteInternal(FPCGContext* Context) const
 			{
 				const FText AttributeTypeName = PCG::Private::GetTypeNameText(AttributeTypeId);
 				const FText MostComplexTypeName = PCG::Private::GetTypeNameText(OperationData.MostComplexInputType);
-				PCGE_LOG(Error, GraphAndLog, FText::Format(LOCTEXT("AttributeCannotBeBroadcasted", "Attribute '{0}' (from pin {1}) of type '{2}' cannot be used for operation with type '{3}'"), 
+				PCGE_LOG_C(Error, GraphAndLog, Context, FText::Format(LOCTEXT("AttributeCannotBeBroadcasted", "Attribute '{0}' (from pin {1}) of type '{2}' cannot be used for operation with type '{3}'"),
 					InputSourceText,
-					FText::FromName(InputTaggedData[Index].Pin),
+					FText::FromName(InputData.Pin),
 					AttributeTypeName,
 					MostComplexTypeName));
 				return false;
 			}
 		}
 
-		NumberOfElements[Index] = OperationData.InputKeys[Index]->GetNum();
-
-		if (Index == InputPinToForward)
-		{
-			OperationData.NumberOfElementsToProcess = NumberOfElements[Index];
-		}
+		return true;
+	}
 
-		// There is nothing to do if one input doesn't have any element to process.
-		// Therefore mark that we have nothing to do and early out.
-		if (NumberOfElements[Index] == 0)
-		{
-			PCGE_LOG(Verbose, LogOnly, FText::Format(LOCTEXT("NoElementsInInput", "No elements in data from pin {0}"), FText::FromName(InputTaggedData[Index].Pin)));
-			bNoOperationNeeded = true;
-			return true;
-		}
+	bool ValidateSecondaryInputClassMatches(const FPCGTaggedData& PrimaryInputData, const FPCGTaggedData& SecondaryInputData)
+	{
+		// First, verify the input data matches the primary. If the pin to forward is not connected, behave like a param data
+		const UClass* InputPinToForwardClass = (PrimaryInputData.Data ? PrimaryInputData.Data->GetClass() : UPCGParamData::StaticClass());
 
-		// Verify that the number of elements makes sense
-		if (OperationData.NumberOfElementsToProcess % NumberOfElements[Index] != 0)
+		// TODO: Consider updating this to check if its a child class instead to be more future proof. For now this is good.
+		// Check for data mismatch between primary pin and current pin
+		if (InputPinToForwardClass != SecondaryInputData.Data->GetClass() && !SecondaryInputData.Data->IsA<UPCGParamData>())
 		{
-			PCGE_LOG(Error, GraphAndLog, FText::Format(LOCTEXT("MismatchInNumberOfElements", "Mismatch between the number of elements from pin {0} ({1}) and from pin {2} ({3})"), FText::FromName(InputTaggedData[InputPinToForward].Pin), OperationData.NumberOfElementsToProcess, FText::FromName(InputTaggedData[Index].Pin), NumberOfElements[Index]));
 			return false;
 		}
 
-		if (InputSource.GetSelection() == EPCGAttributePropertySelection::Attribute)
-		{
-			SourceAttribute[Index] = SourceMetadata[Index]->GetConstAttribute(InputSource.GetName());
-		}
-		else
-		{
-			SourceAttribute[Index] = nullptr;
-		}
-
-		return true;
-	};
-
-	// First do it for the input to forward (it's our control data)
-	if (!CreateInputAccessorAndValidate(InputPinToForward))
-	{
 		return true;
 	}
+}
 
-	for (uint32 i = 0; i < NumberOfInputs && !bNoOperationNeeded; ++i)
-	{
-		if (i != InputPinToForward && !CreateInputAccessorAndValidate(i))
-		{
-			return true;
-		}
-	}
+bool FPCGMetadataElementBase::PrepareDataInternal(FPCGContext* Context) const
+{
+	TRACE_CPUPROFILER_EVENT_SCOPE(FPCGMetadataElementBase::PrepareDataInternal);
+
+	FPCGMetadataElementBase::ContextType* TimeSlicedContext = static_cast<FPCGMetadataElementBase::ContextType*>(Context);
+	check(TimeSlicedContext);
 
-	// If no operation is needed, just forward input 
-	if (bNoOperationNeeded)
+	const UPCGMetadataSettingsBase* Settings = Context->GetInputSettings<UPCGMetadataSettingsBase>();
+	check(Settings);
+
+	const uint32 OperandNum = Settings->GetOperandNum();
+	const uint32 ResultNum = Settings->GetResultNum();
+
+	check(OperandNum > 0);
+	check(ResultNum <= UPCGMetadataSettingsBase::MaxNumberOfOutputs);
+
+	const FName PrimaryPinLabel = Settings->GetInputPinLabel(Settings->GetInputPinToForward());
+	const TArray<FPCGTaggedData> PrimaryInputs = Context->InputData.GetInputsByPin(PrimaryPinLabel);
+
+	// There are no inputs on the primary pin, so passthrough inputs
+	if (PrimaryInputs.IsEmpty())
 	{
-		PassthroughInputs(Context, Settings);
 		return true;
 	}
 
-	// At this point, we verified everything, so we can go forward with the computation, depending on the most complex type
-	// So first forward outputs and create the attribute
-	OperationData.OutputAccessors.SetNum(Settings->GetOutputPinNum());
-	OperationData.OutputKeys.SetNum(Settings->GetOutputPinNum());
-
-	FPCGAttributePropertyOutputSelector OutputTarget = Settings->OutputTarget.CopyAndFixSource(&OperationData.InputSources[InputPinToForward]);
+	// There's no execution state, so just flag that it is ready to continue
+	TimeSlicedContext->InitializePerExecutionState();
 
-	// Use implicit capture, since we capture a lot
-	auto CreateAttribute = [&](uint32 OutputIndex, auto DummyOutValue) -> bool
+	// Set up the iterations on the multiple inputs of the primary pin
+	TimeSlicedContext->InitializePerIterationStates(PrimaryInputs.Num(), [this, Context, Settings, OperandNum, &PrimaryInputs, &PrimaryPinLabel](IterStateType& OutState, const ExecStateType& ExecState, const uint32 IterationIndex)
 	{
-		using AttributeType = decltype(DummyOutValue);
+		TArray<FPCGTaggedData>& Outputs = Context->OutputData.TaggedData;
+		const uint32 NumberOfResults = Settings->GetResultNum();
 
-		FPCGTaggedData& OutputData = Outputs.Add_GetRef(InputTaggedData[InputPinToForward]);
-		OutputData.Pin = Settings->GetOutputPinLabel(OutputIndex);
+		// Gathering all the inputs metadata
+		TArray<const UPCGMetadata*> SourceMetadata;
+		TArray<const FPCGMetadataAttributeBase*> SourceAttribute;
+		TArray<FPCGTaggedData> InputTaggedData;
+		SourceMetadata.SetNum(OperandNum);
+		SourceAttribute.SetNum(OperandNum);
+		InputTaggedData.SetNum(OperandNum);
 
-		UPCGMetadata* OutMetadata = nullptr;
+		const uint32 PrimaryPinIndex = Settings->GetInputPinToForward();
 
-		const FName OutputName = OutputTarget.GetName();
-		const FText OutputTargetText = OutputTarget.GetDisplayText();
-
-		if (OutputTarget.GetSelection() == EPCGAttributePropertySelection::Attribute && OutputTarget.GetExtraNames().IsEmpty())
+		// Iterate over the inputs and validate
+		for (uint32 OperandPinIndex = 0; OperandPinIndex < OperandNum; ++OperandPinIndex)
 		{
-			// In case of an attribute, we check if we have extra selectors. If not, we can just delete the attribute
-			// and create a new one of the right type.
-			// But if we have extra selectors, we need to handle it the same way as properties.
-			// There is no point of failure before duplicating. So duplicate, create the attribute and then the accessor.
-			PCGMetadataElementCommon::DuplicateTaggedData(InputTaggedData[InputPinToForward], OutputData, OutMetadata);
-			FPCGMetadataAttributeBase* OutputAttribute = PCGMetadataElementCommon::ClearOrCreateAttribute<AttributeType>(OutMetadata, OutputName);
-			if (!OutputAttribute)
+			const FName CurrentPinLabel = Settings->GetInputPinLabel(OperandPinIndex);
+			const bool bIsInputConnected = Context->Node ? Context->Node->IsInputPinConnected(CurrentPinLabel) : false;
+			TArray<FPCGTaggedData> CurrentPinInputData = Context->InputData.GetInputsByPin(CurrentPinLabel);
+
+			// This only needs to be checked once
+			if (!bIsInputConnected && CurrentPinInputData.IsEmpty())
 			{
-				return false;
+				FPCGTaggedData& DefaultData = CurrentPinInputData.Emplace_GetRef();
+				DefaultData.Pin = CurrentPinLabel;
+				DefaultData.Data = Settings->CreateDefaultValueParam(OperandPinIndex);
+
+				if (!DefaultData.Data)
+				{
+					PCGE_LOG(Error, GraphAndLog, FText::Format(LOCTEXT("CantCreateDefaultValue", "Pin '{0}' supports default value but we could not create it."), FText::FromName(CurrentPinLabel)));
+					return EPCGTimeSliceInitResult::AbortExecution;
+				}
 			}
 
-			// And copy the mapping from the original attribute, if it is not points
-			if (!InputTaggedData[InputPinToForward].Data->IsA<UPCGPointData>() && SourceMetadata[InputPinToForward] && SourceAttribute[InputPinToForward])
+			if (CurrentPinInputData.IsEmpty())
 			{
-				PCGMetadataElementCommon::CopyEntryToValueKeyMap(SourceMetadata[InputPinToForward], SourceAttribute[InputPinToForward], OutputAttribute);
+				// If we have no data, there is no operation
+				PCGE_LOG(Verbose, LogOnly, FText::Format(LOCTEXT("MissingInputDataForPin", "No data provided on pin '{0}'."), FText::FromName(CurrentPinLabel)));
+				return EPCGTimeSliceInitResult::NoOperation;
+			}
+			else if (CurrentPinInputData.Num() != 1 && CurrentPinInputData.Num() != PrimaryInputs.Num())
+			{
+				PCGE_LOG(Error, GraphAndLog,
+					FText::Format(LOCTEXT("MismatchedDataCountForPin", "Number of data elements ({0}) provided on pin '{1}' doesn't match number of elements ({2}) on primary pin ({3}). Only 1 input or {2} are supported."),
+						CurrentPinInputData.Num(),
+						FText::FromName(CurrentPinLabel),
+						PrimaryInputs.Num(),
+						FText::FromName(PrimaryPinLabel)));
+				return EPCGTimeSliceInitResult::AbortExecution;
 			}
 
-			OperationData.OutputAccessors[OutputIndex] = PCGAttributeAccessorHelpers::CreateAccessor(const_cast<UPCGData*>(OutputData.Data.Get()), OutputTarget);
-		}
-		else
-		{
-			// In case of property or attribute with extra accessor, we need to validate that the property/attribute can accept the output type.
-			// Verify this before duplicating, because an extra allocation is certainly less costly than duplicating the data.
-			// Do it with a const accessor, since OutputData.Data is still pointing on the const input data.
-			TUniquePtr<const IPCGAttributeAccessor> TempConstAccessor = PCGAttributeAccessorHelpers::CreateConstAccessor(OutputData.Data.Get(), OutputTarget);
+			// The operand inputs must either be N:1 or N:N
+			InputTaggedData[OperandPinIndex] = CurrentPinInputData.Num() == 1 ? CurrentPinInputData[0] : MoveTemp(CurrentPinInputData[IterationIndex]);
 
-			if (TempConstAccessor.IsValid())
+			// Check if we have any points
+			if (const UPCGPointData* PointInput = Cast<const UPCGPointData>(InputTaggedData[OperandPinIndex].Data))
 			{
-				// We matched an attribute/property, check if the output type is valid.
-				if (!PCG::Private::IsBroadcastable(PCG::Private::MetadataTypes<AttributeType>::Id, TempConstAccessor->GetUnderlyingType()))
+				if (PointInput->GetPoints().IsEmpty())
 				{
-					PCGE_LOG(Error, GraphAndLog, FText::Format(LOCTEXT("AttributeTypeBroadcastFailed_Updated", "Attribute/Property '{0}' ({1}) is not compatible with operation output type ({2})."), 
-						OutputTargetText,
-						PCG::Private::GetTypeNameText(TempConstAccessor->GetUnderlyingType()),
-						PCG::Private::GetTypeNameText<AttributeType>()));
-					return false;
+					// If we have no points, there is no operation
+					PCGE_LOG(Verbose, LogOnly, FText::Format(LOCTEXT("NoPointsForPin", "No points in point data provided on pin {0}"), FText::FromName(CurrentPinLabel)));
+					return EPCGTimeSliceInitResult::NoOperation;
 				}
+			}
 
-				PCGMetadataElementCommon::DuplicateTaggedData(InputTaggedData[InputPinToForward], OutputData, OutMetadata);
-
-				// Re-create the accessor to point to the right data (since we just duplicated the data)
-				OperationData.OutputAccessors[OutputIndex] = PCGAttributeAccessorHelpers::CreateAccessor(const_cast<UPCGData*>(OutputData.Data.Get()), OutputTarget);
+			SourceMetadata[OperandPinIndex] = InputTaggedData[OperandPinIndex].Data->ConstMetadata();
+			if (!SourceMetadata[OperandPinIndex])
+			{
+				// Since this aborts execution, and the user can fix it, it should be a node error
+				PCGE_LOG(Error, GraphAndLog, FText::Format(LOCTEXT("InvalidInputDataTypeForPin", "Invalid data provided on pin '{0}', must be of type Spatial or Attribute Set."), FText::FromName(CurrentPinLabel)));
+				return EPCGTimeSliceInitResult::AbortExecution;
 			}
 		}
 
-		if (!OperationData.OutputAccessors[OutputIndex].IsValid())
+		PCGMetadataOps::FOperationData& OperationData = OutState;
+		OperationData.Settings = Settings;
+		OperationData.InputAccessors.SetNum(OperandNum);
+		OperationData.InputKeys.SetNum(OperandNum);
+		OperationData.InputSources.SetNum(OperandNum);
+		OperationData.MostComplexInputType = static_cast<uint16>(EPCGMetadataTypes::Unknown);
+
+		const FPCGTaggedData& PrimaryPinData = InputTaggedData[PrimaryPinIndex];
+
+		// First create an accessor for the input to forward (it's our control data)
+		PCGMetadataOpPrivate::CreateAccessor(Settings, PrimaryPinData, OperationData, PrimaryPinIndex);
+		if (!PCGMetadataOpPrivate::ValidateAccessor(Context, Settings, PrimaryPinData, OperationData, PrimaryPinIndex))
 		{
-			return false;
+			return EPCGTimeSliceInitResult::AbortExecution;
 		}
 
-		if (OperationData.OutputAccessors[OutputIndex]->IsReadOnly())
+		// Update the number of elements to process
+		OperationData.NumberOfElementsToProcess = OperationData.InputKeys[PrimaryPinIndex]->GetNum();
+		if (OperationData.NumberOfElementsToProcess == 0)
 		{
-			PCGE_LOG(Error, GraphAndLog, FText::Format(LOCTEXT("OutputAccessorIsReadOnly", "Attribute/Property '{0}' is read only."), OutputTargetText));
-			return false;
+			PCGE_LOG(Verbose, LogOnly, FText::Format(LOCTEXT("NoElementsInForwardedInput", "No elements in data from forwarded pin '{0}'."), FText::FromName(PrimaryPinData.Pin)));
+			return EPCGTimeSliceInitResult::NoOperation;
 		}
 
-		OperationData.OutputKeys[OutputIndex] = PCGAttributeAccessorHelpers::CreateKeys(const_cast<UPCGData*>(OutputData.Data.Get()), OutputTarget);
+		// Create the accessors and validate them for each of the other operands
+		for (uint32 Index = 0; Index < OperandNum; ++Index)
+		{
+			if (Index != PrimaryPinIndex)
+			{
+				// Secondary input class should match the forwarded one
+				if (!PCGMetadataOpPrivate::ValidateSecondaryInputClassMatches(PrimaryPinData, InputTaggedData[Index]))
+				{
+					PCGE_LOG(Error, GraphAndLog, FText::Format(LOCTEXT("InputTypeMismatch", "Data on pin '{0}' is not of the same type than on pin '{1}' and is not an Attribute Set. This is not supported."), FText::FromName(InputTaggedData[Index].Pin), FText::FromName(PrimaryPinData.Pin)));
+					return EPCGTimeSliceInitResult::AbortExecution;
+				}
 
-		return OperationData.OutputKeys[OutputIndex].IsValid();
-	};
+				PCGMetadataOpPrivate::CreateAccessor(Settings, InputTaggedData[Index], OperationData, Index);
 
-	auto CreateAllSameAttributes = [NumberOfOutputs, &CreateAttribute](auto DummyOutValue) -> bool
-	{
-		for (uint32 i = 0; i < NumberOfOutputs; ++i)
+				if (!PCGMetadataOpPrivate::ValidateAccessor(Context, Settings, InputTaggedData[Index], OperationData, Index))
+				{
+					return EPCGTimeSliceInitResult::AbortExecution;
+				}
+
+				const int32 ElementNum = OperationData.InputKeys[Index]->GetNum();
+
+				// No elements on secondary pin, early out for no operation
+				if (ElementNum == 0)
+				{
+					PCGE_LOG(Verbose, LogOnly, FText::Format(LOCTEXT("NoElementsInInput", "No elements in data from secondary pin '{0}'."), FText::FromName(PrimaryPinData.Pin)));
+					return EPCGTimeSliceInitResult::NoOperation;
+				}
+
+				// Verify that the number of elements makes sense
+				if (OperationData.NumberOfElementsToProcess % ElementNum != 0)
+				{
+					PCGE_LOG(Error, GraphAndLog, FText::Format(LOCTEXT("MismatchInNumberOfElements", "Mismatch between the number of elements from pin '{0}' ({1}) and from pin '{2}' ({3})."), FText::FromName(PrimaryPinData.Pin), OperationData.NumberOfElementsToProcess, FText::FromName(InputTaggedData[Index].Pin), ElementNum));
+					return EPCGTimeSliceInitResult::AbortExecution;
+				}
+
+				// If selection is an attribute, get it from the metadata
+				FPCGAttributePropertyInputSelector& InputSource = OperationData.InputSources[Index];
+				if (InputSource.GetSelection() == EPCGAttributePropertySelection::Attribute)
+				{
+					SourceAttribute[Index] = SourceMetadata[Index]->GetConstAttribute(InputSource.GetName());
+				}
+				else
+				{
+					SourceAttribute[Index] = nullptr;
+				}
+			}
+		}
+
+		// At this point, we verified everything, so we can go forward with the computation, depending on the most complex type
+		// So first forward outputs and create the attribute
+		OperationData.OutputAccessors.SetNum(Settings->GetResultNum());
+		OperationData.OutputKeys.SetNum(Settings->GetResultNum());
+
+		const FPCGAttributePropertyOutputSelector OutputTarget = Settings->OutputTarget.CopyAndFixSource(&OperationData.InputSources[PrimaryPinIndex]);
+
+		// Use implicit capture, since we capture a lot
+		auto CreateAttribute = [&]<typename AttributeType>(uint32 OutputIndex, AttributeType DummyOutValue) -> bool
 		{
-			if (!CreateAttribute(i, DummyOutValue))
+			FPCGTaggedData& OutputData = Outputs.Add_GetRef(InputTaggedData[PrimaryPinIndex]);
+			OutputData.Pin = Settings->GetOutputPinLabel(OutputIndex);
+
+			UPCGMetadata* OutMetadata = nullptr;
+
+			const FName OutputName = OutputTarget.GetName();
+			const FText OutputTargetText = OutputTarget.GetDisplayText();
+
+			if (OutputTarget.GetSelection() == EPCGAttributePropertySelection::Attribute && OutputTarget.GetExtraNames().IsEmpty())
+			{
+				// In case of an attribute, we check if we have extra selectors. If not, we can just delete the attribute
+				// and create a new one of the right type.
+				// But if we have extra selectors, we need to handle it the same way as properties.
+				// There is no point of failure before duplicating. So duplicate, create the attribute and then the accessor.
+				PCGMetadataElementCommon::DuplicateTaggedData(InputTaggedData[PrimaryPinIndex], OutputData, OutMetadata);
+				FPCGMetadataAttributeBase* OutputAttribute = PCGMetadataElementCommon::ClearOrCreateAttribute<AttributeType>(OutMetadata, OutputName);
+				if (!OutputAttribute)
+				{
+					return false;
+				}
+
+				// And copy the mapping from the original attribute, if it is not points
+				if (!InputTaggedData[PrimaryPinIndex].Data->IsA<UPCGPointData>() && SourceMetadata[PrimaryPinIndex] && SourceAttribute[PrimaryPinIndex])
+				{
+					PCGMetadataElementCommon::CopyEntryToValueKeyMap(SourceMetadata[PrimaryPinIndex], SourceAttribute[PrimaryPinIndex], OutputAttribute);
+				}
+
+				OperationData.OutputAccessors[OutputIndex] = PCGAttributeAccessorHelpers::CreateAccessor(const_cast<UPCGData*>(OutputData.Data.Get()), OutputTarget);
+			}
+			else
+			{
+				// In case of property or attribute with extra accessor, we need to validate that the property/attribute can accept the output type.
+				// Verify this before duplicating, because an extra allocation is certainly less costly than duplicating the data.
+				// Do it with a const accessor, since OutputData.Data is still pointing on the const input data.
+				const TUniquePtr<const IPCGAttributeAccessor> TempConstAccessor = PCGAttributeAccessorHelpers::CreateConstAccessor(OutputData.Data.Get(), OutputTarget);
+
+				if (TempConstAccessor.IsValid())
+				{
+					// We matched an attribute/property, check if the output type is valid.
+					if (!PCG::Private::IsBroadcastable(PCG::Private::MetadataTypes<AttributeType>::Id, TempConstAccessor->GetUnderlyingType()))
+					{
+						PCGE_LOG(Error, GraphAndLog, FText::Format(LOCTEXT("AttributeTypeBroadcastFailed_Updated", "Attribute/Property '{0}' ({1}) is not compatible with operation output type ({2})."),
+							OutputTargetText,
+							PCG::Private::GetTypeNameText(TempConstAccessor->GetUnderlyingType()),
+							PCG::Private::GetTypeNameText<AttributeType>()));
+						return false;
+					}
+
+					PCGMetadataElementCommon::DuplicateTaggedData(InputTaggedData[PrimaryPinIndex], OutputData, OutMetadata);
+
+					// Re-create the accessor to point to the right data (since we just duplicated the data)
+					OperationData.OutputAccessors[OutputIndex] = PCGAttributeAccessorHelpers::CreateAccessor(const_cast<UPCGData*>(OutputData.Data.Get()), OutputTarget);
+				}
+			}
+
+			if (!OperationData.OutputAccessors[OutputIndex].IsValid())
 			{
 				return false;
 			}
-		}
 
-		return true;
-	};
+			if (OperationData.OutputAccessors[OutputIndex]->IsReadOnly())
+			{
+				PCGE_LOG(Error, GraphAndLog, FText::Format(LOCTEXT("OutputAccessorIsReadOnly", "Attribute/Property '{0}' is read only."), OutputTargetText));
+				return false;
+			}
 
-	OperationData.OutputType = Settings->GetOutputType(OperationData.MostComplexInputType);
+			OperationData.OutputKeys[OutputIndex] = PCGAttributeAccessorHelpers::CreateKeys(const_cast<UPCGData*>(OutputData.Data.Get()), OutputTarget);
 
-	bool bCreateAttributeSucceeded = true;
+			return OperationData.OutputKeys[OutputIndex].IsValid();
+		};
 
-	if (!Settings->HasDifferentOutputTypes())
-	{
-		bCreateAttributeSucceeded = PCGMetadataAttribute::CallbackWithRightType(OperationData.OutputType, CreateAllSameAttributes);
-	}
-	else
-	{
-		TArray<uint16> OutputTypes = Settings->GetAllOutputTypes();
-		check(OutputTypes.Num() == NumberOfOutputs);
+		auto CreateAllSameAttributes = [NumberOfResults, &CreateAttribute](auto DummyOutValue) -> bool
+		{
+			for (uint32 i = 0; i < NumberOfResults; ++i)
+			{
+				if (!CreateAttribute(i, DummyOutValue))
+				{
+					return false;
+				}
+			}
+
+			return true;
+		};
+
+		OperationData.OutputType = Settings->GetOutputType(OperationData.MostComplexInputType);
+
+		bool bCreateAttributeSucceeded = true;
 
-		for (uint32 i = 0; i < NumberOfOutputs && bCreateAttributeSucceeded; ++i)
+		if (!Settings->HasDifferentOutputTypes())
 		{
-			bCreateAttributeSucceeded &= PCGMetadataAttribute::CallbackWithRightType(OutputTypes[i],
-				[&CreateAttribute, i](auto DummyOutValue) -> bool 
+			bCreateAttributeSucceeded = PCGMetadataAttribute::CallbackWithRightType(OperationData.OutputType, CreateAllSameAttributes);
+		}
+		else
+		{
+			TArray<uint16> OutputTypes = Settings->GetAllOutputTypes();
+			check(OutputTypes.Num() == NumberOfResults);
+
+			for (uint32 i = 0; i < NumberOfResults && bCreateAttributeSucceeded; ++i)
+			{
+				bCreateAttributeSucceeded &= PCGMetadataAttribute::CallbackWithRightType(OutputTypes[i], [&CreateAttribute, i](auto DummyOutValue) -> bool
 				{
 					return CreateAttribute(i, DummyOutValue);
 				});
+			}
 		}
-	}
 
-	if (!bCreateAttributeSucceeded)
+		if (!bCreateAttributeSucceeded)
+		{
+			PCGE_LOG(Error, GraphAndLog, LOCTEXT("ErrorCreatingOutputAttributes", "Error while creating output attributes"));
+			return EPCGTimeSliceInitResult::AbortExecution;
+		}
+
+		OperationData.Settings = Settings;
+
+		return EPCGTimeSliceInitResult::Success;
+	});
+
+	return true;
+}
+
+bool FPCGMetadataElementBase::ExecuteInternal(FPCGContext* Context) const
+{
+	TRACE_CPUPROFILER_EVENT_SCOPE(FPCGMetadataElementBase::Execute);
+	ContextType* TimeSlicedContext = static_cast<ContextType*>(Context);
+	check(TimeSlicedContext);
+
+	TArray<FPCGTaggedData>& Outputs = Context->OutputData.TaggedData;
+
+	// Prepare data failed, no need to execute. Return an empty output
+	if (!TimeSlicedContext->DataIsPreparedForExecution())
 	{
-		PCGE_LOG(Error, GraphAndLog, LOCTEXT("ErrorCreatingOutputAttributes", "Error while creating output attributes"));
-		Outputs.Empty();
 		return true;
 	}
 
-	OperationData.Settings = Settings;
-
-	if (!DoOperation(OperationData))
+	return ExecuteSlice(TimeSlicedContext, [this, &Outputs](ContextType* Context, const ExecStateType& ExecState, IterStateType& IterState, const uint32 IterationIndex) -> bool
 	{
-		PCGE_LOG(Error, GraphAndLog, LOCTEXT("ErrorOccurred", "Error while performing the metadata operation, check logs for more information"));
-		Outputs.Empty();
-	}
+		// No operation, so skip the iteration.
+		if (Context->GetIterationStateResult(IterationIndex) == EPCGTimeSliceInitResult::NoOperation)
+		{
+			PassthroughAllInputs(Context, Outputs);
+			return true;
+		}
 
-	return true;
+		// TODO: Add range-based async evaluation to the DoOperation function in the future
+		if (!DoOperation(IterState))
+		{
+			PCGE_LOG(Error, GraphAndLog, LOCTEXT("ErrorOccurred", "Error while performing the metadata operation, check logs for more information"));
+			PassthroughInput(Context, Outputs, IterationIndex);
+		}
+
+		return true;
+	});
 }
 
 #undef LOCTEXT_NAMESPACE
diff --git a/Engine/Plugins/Experimental/PCG/Source/PCG/Private/Elements/Metadata/PCGMetadataRotatorOpElement.cpp b/Engine/Plugins/Experimental/PCG/Source/PCG/Private/Elements/Metadata/PCGMetadataRotatorOpElement.cpp
--- a/Engine/Plugins/Experimental/PCG/Source/PCG/Private/Elements/Metadata/PCGMetadataRotatorOpElement.cpp
+++ b/Engine/Plugins/Experimental/PCG/Source/PCG/Private/Elements/Metadata/PCGMetadataRotatorOpElement.cpp
@@ -144,7 +144,7 @@ FName UPCGMetadataRotatorSettings::GetInputPinLabel(uint32 Index) const
 	}
 }
 
-uint32 UPCGMetadataRotatorSettings::GetInputPinNum() const
+uint32 UPCGMetadataRotatorSettings::GetOperandNum() const
 {
 	if (PCGMetadataRotatorSettings::IsUnaryOp(Operation))
 	{
@@ -239,7 +239,7 @@ FPCGElementPtr UPCGMetadataRotatorSettings::CreateElement() const
 	return MakeShared<FPCGMetadataRotatorElement>();
 }
 
-bool FPCGMetadataRotatorElement::DoOperation(FOperationData& OperationData) const
+bool FPCGMetadataRotatorElement::DoOperation(PCGMetadataOps::FOperationData& OperationData) const
 {
 	TRACE_CPUPROFILER_EVENT_SCOPE(FPCGMetadataRotatorElement::Execute);
 
diff --git a/Engine/Plugins/Experimental/PCG/Source/PCG/Private/Elements/Metadata/PCGMetadataStringOpElement.cpp b/Engine/Plugins/Experimental/PCG/Source/PCG/Private/Elements/Metadata/PCGMetadataStringOpElement.cpp
--- a/Engine/Plugins/Experimental/PCG/Source/PCG/Private/Elements/Metadata/PCGMetadataStringOpElement.cpp
+++ b/Engine/Plugins/Experimental/PCG/Source/PCG/Private/Elements/Metadata/PCGMetadataStringOpElement.cpp
@@ -16,7 +16,7 @@ FName UPCGMetadataStringOpSettings::GetInputPinLabel(uint32 Index) const
 	}
 }
 
-uint32 UPCGMetadataStringOpSettings::GetInputPinNum() const
+uint32 UPCGMetadataStringOpSettings::GetOperandNum() const
 {
 	return 2;
 }
@@ -71,7 +71,7 @@ uint16 UPCGMetadataStringOpSettings::GetOutputType(uint16 InputTypeId) const
 	return (uint16)EPCGMetadataTypes::String;
 }
 
-bool FPCGMetadataStringOpElement::DoOperation(FOperationData& OperationData) const
+bool FPCGMetadataStringOpElement::DoOperation(PCGMetadataOps::FOperationData& OperationData) const
 {
 	TRACE_CPUPROFILER_EVENT_SCOPE(FPCGMetadataStringOpElement::Execute);
 
diff --git a/Engine/Plugins/Experimental/PCG/Source/PCG/Private/Elements/Metadata/PCGMetadataTransformOpElement.cpp b/Engine/Plugins/Experimental/PCG/Source/PCG/Private/Elements/Metadata/PCGMetadataTransformOpElement.cpp
--- a/Engine/Plugins/Experimental/PCG/Source/PCG/Private/Elements/Metadata/PCGMetadataTransformOpElement.cpp
+++ b/Engine/Plugins/Experimental/PCG/Source/PCG/Private/Elements/Metadata/PCGMetadataTransformOpElement.cpp
@@ -112,7 +112,7 @@ FName UPCGMetadataTransformSettings::GetInputPinLabel(uint32 Index) const
 	}
 }
 
-uint32 UPCGMetadataTransformSettings::GetInputPinNum() const
+uint32 UPCGMetadataTransformSettings::GetOperandNum() const
 {
 	if (Operation == EPCGMedadataTransformOperation::Invert)
 	{
@@ -202,7 +202,7 @@ FPCGElementPtr UPCGMetadataTransformSettings::CreateElement() const
 	return MakeShared<FPCGMetadataTransformElement>();
 }
 
-bool FPCGMetadataTransformElement::DoOperation(FOperationData& OperationData) const
+bool FPCGMetadataTransformElement::DoOperation(PCGMetadataOps::FOperationData& OperationData) const
 {
 	TRACE_CPUPROFILER_EVENT_SCOPE(FPCGMetadataTransformElement::Execute);
 
diff --git a/Engine/Plugins/Experimental/PCG/Source/PCG/Private/Elements/Metadata/PCGMetadataTrigOpElement.cpp b/Engine/Plugins/Experimental/PCG/Source/PCG/Private/Elements/Metadata/PCGMetadataTrigOpElement.cpp
--- a/Engine/Plugins/Experimental/PCG/Source/PCG/Private/Elements/Metadata/PCGMetadataTrigOpElement.cpp
+++ b/Engine/Plugins/Experimental/PCG/Source/PCG/Private/Elements/Metadata/PCGMetadataTrigOpElement.cpp
@@ -74,7 +74,7 @@ FName UPCGMetadataTrigSettings::GetInputPinLabel(uint32 Index) const
 	}
 }
 
-uint32 UPCGMetadataTrigSettings::GetInputPinNum() const
+uint32 UPCGMetadataTrigSettings::GetOperandNum() const
 {
 	return (Operation != EPCGMedadataTrigOperation::Atan2) ? 1 : 2;
 }
@@ -148,7 +148,7 @@ FPCGElementPtr UPCGMetadataTrigSettings::CreateElement() const
 	return MakeShared<FPCGMetadataTrigElement>();
 }
 
-bool FPCGMetadataTrigElement::DoOperation(FOperationData& OperationData) const
+bool FPCGMetadataTrigElement::DoOperation(PCGMetadataOps::FOperationData& OperationData) const
 {
 	TRACE_CPUPROFILER_EVENT_SCOPE(FPCGMetadataTrigElement::Execute);
 
diff --git a/Engine/Plugins/Experimental/PCG/Source/PCG/Private/Elements/Metadata/PCGMetadataVectorOpElement.cpp b/Engine/Plugins/Experimental/PCG/Source/PCG/Private/Elements/Metadata/PCGMetadataVectorOpElement.cpp
--- a/Engine/Plugins/Experimental/PCG/Source/PCG/Private/Elements/Metadata/PCGMetadataVectorOpElement.cpp
+++ b/Engine/Plugins/Experimental/PCG/Source/PCG/Private/Elements/Metadata/PCGMetadataVectorOpElement.cpp
@@ -268,7 +268,7 @@ FName UPCGMetadataVectorSettings::GetInputPinLabel(uint32 Index) const
 	}
 }
 
-uint32 UPCGMetadataVectorSettings::GetInputPinNum() const
+uint32 UPCGMetadataVectorSettings::GetOperandNum() const
 {
 	if (PCGMetadataVectorSettings::IsUnaryOp(Operation))
 	{
@@ -381,7 +381,7 @@ uint16 UPCGMetadataVectorSettings::GetOutputType(uint16 InputTypeId) const
 	}
 }
 
-bool FPCGMetadataVectorElement::DoOperation(FOperationData& OperationData) const
+bool FPCGMetadataVectorElement::DoOperation(PCGMetadataOps::FOperationData& OperationData) const
 {
 	TRACE_CPUPROFILER_EVENT_SCOPE(FPCGMetadataVectorElement::Execute);
 
diff --git a/Engine/Plugins/Experimental/PCG/Source/PCG/Public/Elements/Metadata/PCGMetadataBitwiseOpElement.h b/Engine/Plugins/Experimental/PCG/Source/PCG/Public/Elements/Metadata/PCGMetadataBitwiseOpElement.h
--- a/Engine/Plugins/Experimental/PCG/Source/PCG/Public/Elements/Metadata/PCGMetadataBitwiseOpElement.h
+++ b/Engine/Plugins/Experimental/PCG/Source/PCG/Public/Elements/Metadata/PCGMetadataBitwiseOpElement.h
@@ -37,10 +37,10 @@ public:
 	//~End UPCGSettings interface
 
 	//~Begin UPCGMetadataSettingsBase interface
-	FPCGAttributePropertyInputSelector GetInputSource(uint32 Index) const override;
+	virtual FPCGAttributePropertyInputSelector GetInputSource(uint32 Index) const override;
 
 	virtual FName GetInputPinLabel(uint32 Index) const override;
-	virtual uint32 GetInputPinNum() const override;
+	virtual uint32 GetOperandNum() const override;
 
 	virtual bool IsSupportedInputType(uint16 TypeId, uint32 InputIndex, bool& bHasSpecialRequirement) const override;
 	virtual uint16 GetOutputType(uint16 InputTypeId) const override;
@@ -73,5 +73,5 @@ public:
 class FPCGMetadataBitwiseElement : public FPCGMetadataElementBase
 {
 protected:
-	virtual bool DoOperation(FOperationData& OperationData) const override;
+	virtual bool DoOperation(PCGMetadataOps::FOperationData& OperationData) const override;
 };
diff --git a/Engine/Plugins/Experimental/PCG/Source/PCG/Public/Elements/Metadata/PCGMetadataBooleanOpElement.h b/Engine/Plugins/Experimental/PCG/Source/PCG/Public/Elements/Metadata/PCGMetadataBooleanOpElement.h
--- a/Engine/Plugins/Experimental/PCG/Source/PCG/Public/Elements/Metadata/PCGMetadataBooleanOpElement.h
+++ b/Engine/Plugins/Experimental/PCG/Source/PCG/Public/Elements/Metadata/PCGMetadataBooleanOpElement.h
@@ -37,10 +37,10 @@ public:
 	//~End UPCGSettings interface
 
 	//~Begin UPCGMetadataSettingsBase interface
-	FPCGAttributePropertyInputSelector GetInputSource(uint32 Index) const override;
+	virtual FPCGAttributePropertyInputSelector GetInputSource(uint32 Index) const override;
 
 	virtual FName GetInputPinLabel(uint32 Index) const override;
-	virtual uint32 GetInputPinNum() const override;
+	virtual uint32 GetOperandNum() const override;
 
 	virtual bool IsSupportedInputType(uint16 TypeId, uint32 InputIndex, bool& bHasSpecialRequirement) const override;
 	virtual uint16 GetOutputType(uint16 InputTypeId) const override;
@@ -73,5 +73,5 @@ public:
 class FPCGMetadataBooleanElement : public FPCGMetadataElementBase
 {
 protected:
-	virtual bool DoOperation(FOperationData& OperationData) const override;
+	virtual bool DoOperation(PCGMetadataOps::FOperationData& OperationData) const override;
 };
diff --git a/Engine/Plugins/Experimental/PCG/Source/PCG/Public/Elements/Metadata/PCGMetadataBreakTransform.h b/Engine/Plugins/Experimental/PCG/Source/PCG/Public/Elements/Metadata/PCGMetadataBreakTransform.h
--- a/Engine/Plugins/Experimental/PCG/Source/PCG/Public/Elements/Metadata/PCGMetadataBreakTransform.h
+++ b/Engine/Plugins/Experimental/PCG/Source/PCG/Public/Elements/Metadata/PCGMetadataBreakTransform.h
@@ -23,10 +23,10 @@ public:
 	virtual FText GetDefaultNodeTitle() const override;
 #endif
 
-	FPCGAttributePropertyInputSelector GetInputSource(uint32 Index) const override;
+	virtual FPCGAttributePropertyInputSelector GetInputSource(uint32 Index) const override;
 
 	virtual FName GetOutputPinLabel(uint32 Index) const override;
-	virtual uint32 GetOutputPinNum() const override;
+	virtual uint32 GetResultNum() const override;
 
 	virtual bool IsSupportedInputType(uint16 TypeId, uint32 InputIndex, bool& bHasSpecialRequirement) const override;
 	virtual FName GetOutputAttributeName(FName BaseName, uint32 Index) const override;
@@ -51,7 +51,7 @@ public:
 class FPCGMetadataBreakTransformElement : public FPCGMetadataElementBase
 {
 protected:
-	virtual bool DoOperation(FOperationData& OperationData) const override;
+	virtual bool DoOperation(PCGMetadataOps::FOperationData& OperationData) const override;
 };
 
 #if UE_ENABLE_INCLUDE_ORDER_DEPRECATED_IN_5_2
diff --git a/Engine/Plugins/Experimental/PCG/Source/PCG/Public/Elements/Metadata/PCGMetadataBreakVector.h b/Engine/Plugins/Experimental/PCG/Source/PCG/Public/Elements/Metadata/PCGMetadataBreakVector.h
--- a/Engine/Plugins/Experimental/PCG/Source/PCG/Public/Elements/Metadata/PCGMetadataBreakVector.h
+++ b/Engine/Plugins/Experimental/PCG/Source/PCG/Public/Elements/Metadata/PCGMetadataBreakVector.h
@@ -31,10 +31,10 @@ public:
 	virtual FText GetDefaultNodeTitle() const override;
 #endif
 
-	FPCGAttributePropertyInputSelector GetInputSource(uint32 Index) const override;
+	virtual FPCGAttributePropertyInputSelector GetInputSource(uint32 Index) const override;
 
 	virtual FName GetOutputPinLabel(uint32 Index) const override;
-	virtual uint32 GetOutputPinNum() const override;
+	virtual uint32 GetResultNum() const override;
 
 	virtual bool IsSupportedInputType(uint16 TypeId, uint32 InputIndex, bool& bHasSpecialRequirement) const override;
 	virtual uint16 GetOutputType(uint16 InputTypeId) const override;
@@ -57,5 +57,5 @@ public:
 class FPCGMetadataBreakVectorElement : public FPCGMetadataElementBase
 {
 protected:
-	virtual bool DoOperation(FOperationData& OperationData) const override;
+	virtual bool DoOperation(PCGMetadataOps::FOperationData& OperationData) const override;
 };
diff --git a/Engine/Plugins/Experimental/PCG/Source/PCG/Public/Elements/Metadata/PCGMetadataCompareOpElement.h b/Engine/Plugins/Experimental/PCG/Source/PCG/Public/Elements/Metadata/PCGMetadataCompareOpElement.h
--- a/Engine/Plugins/Experimental/PCG/Source/PCG/Public/Elements/Metadata/PCGMetadataCompareOpElement.h
+++ b/Engine/Plugins/Experimental/PCG/Source/PCG/Public/Elements/Metadata/PCGMetadataCompareOpElement.h
@@ -39,10 +39,10 @@ public:
 	//~End UPCGSettings interface
 
 	//~Begin UPCGMetadataSettingsBase interface
-	FPCGAttributePropertyInputSelector GetInputSource(uint32 Index) const override;
+	virtual FPCGAttributePropertyInputSelector GetInputSource(uint32 Index) const override;
 
 	virtual FName GetInputPinLabel(uint32 Index) const override;
-	virtual uint32 GetInputPinNum() const override;
+	virtual uint32 GetOperandNum() const override;
 
 	virtual bool IsSupportedInputType(uint16 TypeId, uint32 InputIndex, bool& bHasSpecialRequirement) const override;
 	virtual uint16 GetOutputType(uint16 InputTypeId) const override;
@@ -78,5 +78,5 @@ public:
 class FPCGMetadataCompareElement : public FPCGMetadataElementBase
 {
 protected:
-	virtual bool DoOperation(FOperationData& OperationData) const override;
+	virtual bool DoOperation(PCGMetadataOps::FOperationData& OperationData) const override;
 };
diff --git a/Engine/Plugins/Experimental/PCG/Source/PCG/Public/Elements/Metadata/PCGMetadataMakeRotator.h b/Engine/Plugins/Experimental/PCG/Source/PCG/Public/Elements/Metadata/PCGMetadataMakeRotator.h
--- a/Engine/Plugins/Experimental/PCG/Source/PCG/Public/Elements/Metadata/PCGMetadataMakeRotator.h
+++ b/Engine/Plugins/Experimental/PCG/Source/PCG/Public/Elements/Metadata/PCGMetadataMakeRotator.h
@@ -57,10 +57,10 @@ public:
 #endif // WITH_EDITOR
 	//~End UPCGSettings interface
 
-	FPCGAttributePropertyInputSelector GetInputSource(uint32 Index) const override;
+	virtual FPCGAttributePropertyInputSelector GetInputSource(uint32 Index) const override;
 
 	virtual FName GetInputPinLabel(uint32 Index) const override;
-	virtual uint32 GetInputPinNum() const override;
+	virtual uint32 GetOperandNum() const override;
 
 	virtual bool IsSupportedInputType(uint16 TypeId, uint32 InputIndex, bool& bHasSpecialRequirement) const override;
 	virtual uint16 GetOutputType(uint16 InputTypeId) const override;
@@ -87,5 +87,5 @@ public:
 class FPCGMetadataMakeRotatorElement : public FPCGMetadataElementBase
 {
 protected:
-	virtual bool DoOperation(FOperationData& OperationData) const override;
+	virtual bool DoOperation(PCGMetadataOps::FOperationData& OperationData) const override;
 };
diff --git a/Engine/Plugins/Experimental/PCG/Source/PCG/Public/Elements/Metadata/PCGMetadataMakeTransform.h b/Engine/Plugins/Experimental/PCG/Source/PCG/Public/Elements/Metadata/PCGMetadataMakeTransform.h
--- a/Engine/Plugins/Experimental/PCG/Source/PCG/Public/Elements/Metadata/PCGMetadataMakeTransform.h
+++ b/Engine/Plugins/Experimental/PCG/Source/PCG/Public/Elements/Metadata/PCGMetadataMakeTransform.h
@@ -36,10 +36,10 @@ public:
 #endif // WITH_EDITORs
 	//~End UPCGSettings interface
 
-	FPCGAttributePropertyInputSelector GetInputSource(uint32 Index) const override;
+	virtual FPCGAttributePropertyInputSelector GetInputSource(uint32 Index) const override;
 
 	virtual FName GetInputPinLabel(uint32 Index) const override;
-	virtual uint32 GetInputPinNum() const override;
+	virtual uint32 GetOperandNum() const override;
 
 	virtual bool IsSupportedInputType(uint16 TypeId, uint32 InputIndex, bool& bHasSpecialRequirement) const override;
 	virtual uint16 GetOutputType(uint16 InputTypeId) const override;
@@ -74,5 +74,5 @@ public:
 class FPCGMetadataMakeTransformElement : public FPCGMetadataElementBase
 {
 protected:
-	virtual bool DoOperation(FOperationData& OperationData) const override;
+	virtual bool DoOperation(PCGMetadataOps::FOperationData& OperationData) const override;
 };
diff --git a/Engine/Plugins/Experimental/PCG/Source/PCG/Public/Elements/Metadata/PCGMetadataMakeVector.h b/Engine/Plugins/Experimental/PCG/Source/PCG/Public/Elements/Metadata/PCGMetadataMakeVector.h
--- a/Engine/Plugins/Experimental/PCG/Source/PCG/Public/Elements/Metadata/PCGMetadataMakeVector.h
+++ b/Engine/Plugins/Experimental/PCG/Source/PCG/Public/Elements/Metadata/PCGMetadataMakeVector.h
@@ -58,10 +58,10 @@ public:
 #endif // WITH_EDITOR
 	//~End UPCGSettings interface
 
-	FPCGAttributePropertyInputSelector GetInputSource(uint32 Index) const override;
+	virtual FPCGAttributePropertyInputSelector GetInputSource(uint32 Index) const override;
 
 	virtual FName GetInputPinLabel(uint32 Index) const override;
-	virtual uint32 GetInputPinNum() const override;
+	virtual uint32 GetOperandNum() const override;
 
 	virtual bool IsSupportedInputType(uint16 TypeId, uint32 InputIndex, bool& bHasSpecialRequirement) const override;
 	virtual uint16 GetOutputType(uint16 InputTypeId) const override;
@@ -111,5 +111,5 @@ public:
 class FPCGMetadataMakeVectorElement : public FPCGMetadataElementBase
 {
 protected:
-	virtual bool DoOperation(FOperationData& OperationData) const override;
+	virtual bool DoOperation(PCGMetadataOps::FOperationData& OperationData) const override;
 };
diff --git a/Engine/Plugins/Experimental/PCG/Source/PCG/Public/Elements/Metadata/PCGMetadataMathsOpElement.h b/Engine/Plugins/Experimental/PCG/Source/PCG/Public/Elements/Metadata/PCGMetadataMathsOpElement.h
--- a/Engine/Plugins/Experimental/PCG/Source/PCG/Public/Elements/Metadata/PCGMetadataMathsOpElement.h
+++ b/Engine/Plugins/Experimental/PCG/Source/PCG/Public/Elements/Metadata/PCGMetadataMathsOpElement.h
@@ -64,11 +64,11 @@ public:
 	//~End UPCGSettings interface
 
 	//~Begin UPCGMetadataSettingsBase interface
-	FPCGAttributePropertyInputSelector GetInputSource(uint32 Index) const override;
+	virtual FPCGAttributePropertyInputSelector GetInputSource(uint32 Index) const override;
 
 	virtual FName GetInputPinLabel(uint32 Index) const override;
-	virtual uint32 GetInputPinNum() const override;
-	virtual uint16 GetOutputType(uint16 InputTypeId) const;
+	virtual uint32 GetOperandNum() const override;
+	virtual uint16 GetOutputType(uint16 InputTypeId) const override;
 
 	virtual bool IsSupportedInputType(uint16 TypeId, uint32 InputIndex, bool& bHasSpecialRequirement) const override;
 	//~End UPCGMetadataSettingsBase interface
@@ -118,5 +118,5 @@ private:
 class FPCGMetadataMathsElement : public FPCGMetadataElementBase
 {
 protected:
-	virtual bool DoOperation(FOperationData& OperationData) const override;
+	virtual bool DoOperation(PCGMetadataOps::FOperationData& OperationData) const override;
 };
diff --git a/Engine/Plugins/Experimental/PCG/Source/PCG/Public/Elements/Metadata/PCGMetadataOpElementBase.h b/Engine/Plugins/Experimental/PCG/Source/PCG/Public/Elements/Metadata/PCGMetadataOpElementBase.h
--- a/Engine/Plugins/Experimental/PCG/Source/PCG/Public/Elements/Metadata/PCGMetadataOpElementBase.h
+++ b/Engine/Plugins/Experimental/PCG/Source/PCG/Public/Elements/Metadata/PCGMetadataOpElementBase.h
@@ -9,6 +9,7 @@
 #include "Metadata/Accessors/PCGAttributeAccessorKeys.h"
 
 #include "Containers/StaticArray.h"
+#include "Elements/PCGTimeSlicedElementBase.h"
 
 #include "PCGMetadataOpElementBase.generated.h"
 
@@ -99,20 +100,20 @@ protected:
 	//~End UPCGSettings interface
 
 public:
-	virtual FPCGAttributePropertyInputSelector GetInputSource(uint32 Index) const { return FPCGAttributePropertyInputSelector(); };
+	virtual FPCGAttributePropertyInputSelector GetInputSource(uint32 Index) const { return FPCGAttributePropertyInputSelector(); }
 
 	virtual FName GetInputPinLabel(uint32 Index) const { return PCGPinConstants::DefaultInputLabel; }
-	virtual uint32 GetInputPinNum() const { return 1; };
+	virtual uint32 GetOperandNum() const { return 1; }
 
 	virtual FName GetOutputPinLabel(uint32 Index) const { return PCGPinConstants::DefaultOutputLabel; }
-	virtual uint32 GetOutputPinNum() const { return 1; }
+	virtual uint32 GetResultNum() const { return 1; }
 
-	virtual bool IsSupportedInputType(uint16 TypeId, uint32 InputIndex, bool& bHasSpecialRequirement) const { return false; };
-	virtual uint16 GetOutputType(uint16 InputTypeId) const { return InputTypeId; };
+	virtual bool IsSupportedInputType(uint16 TypeId, uint32 InputIndex, bool& bHasSpecialRequirement) const { return false; }
+	virtual uint16 GetOutputType(uint16 InputTypeId) const { return InputTypeId; }
 	virtual FName GetOutputAttributeName(FName BaseName, uint32 Index) const { return BaseName; }
 
 	virtual bool HasDifferentOutputTypes() const { return false; }
-	virtual TArray<uint16> GetAllOutputTypes() const { return TArray<uint16>(); };
+	virtual TArray<uint16> GetAllOutputTypes() const { return TArray<uint16>(); }
 
 	/* Can be overriden by child class to support default values on unplugged pins. */
 	virtual bool DoesInputSupportDefaultValue(uint32 Index) const { return false; }
@@ -159,14 +160,12 @@ protected:
 	TArray<FName> GetOutputDataFromPinOptions() const;
 };
 
-
-class FPCGMetadataElementBase : public IPCGElement
+namespace PCGMetadataOps
 {
-public:
 	struct FOperationData
 	{
 		int32 NumberOfElementsToProcess = -1;
-		uint16 MostComplexInputType;
+		uint16 MostComplexInputType = static_cast<uint16>(EPCGMetadataTypes::Unknown);
 		uint16 OutputType;
 		const UPCGMetadataSettingsBase* Settings = nullptr;
 
@@ -181,55 +180,43 @@ public:
 		template <int32 NbInputs, int32 NbOutputs>
 		void Validate();
 	};
+}
 
+class FPCGMetadataElementBase : public TPCGTimeSlicedElementBase<PCGTimeSlice::FEmptyStruct, PCGMetadataOps::FOperationData>
+{
 protected:
+	virtual bool PrepareDataInternal(FPCGContext* Context) const override;
 	virtual bool ExecuteInternal(FPCGContext* Context) const override;
-	virtual bool DoOperation(FOperationData& InOperationData) const = 0;
+
+	virtual bool DoOperation(PCGMetadataOps::FOperationData& InOperationData) const = 0;
 
 	/**
 	* Generic method to factorise all the boilerplate code for a variable number of inputs/outputs
 	*/
 	template <typename... InputTypes, typename... Callbacks>
-	inline bool DoNAryOp(FOperationData& InOperationData, TTuple<Callbacks...>&& InCallbacks) const;
+	inline bool DoNAryOp(PCGMetadataOps::FOperationData& InOperationData, TTuple<Callbacks...>&& InCallbacks) const;
 
 	/* All operations can have a fixed number of inputs and a variable number of outputs.
 	* Each output need to have its own callback, all taking the exact number of "const InType&" as input
 	* and each can return a different output type.
 	*/
 	template <typename InType, typename... Callbacks>
-	bool DoUnaryOp(FOperationData& InOperationData, Callbacks&& ...InCallbacks) const;
+	bool DoUnaryOp(PCGMetadataOps::FOperationData& InOperationData, Callbacks&& ...InCallbacks) const;
 
 	template <typename InType1, typename InType2, typename... Callbacks>
-	bool DoBinaryOp(FOperationData& InOperationData, Callbacks&& ...InCallbacks) const;
+	bool DoBinaryOp(PCGMetadataOps::FOperationData& InOperationData, Callbacks&& ...InCallbacks) const;
 
 	template <typename InType1, typename InType2, typename InType3, typename... Callbacks>
-	bool DoTernaryOp(FOperationData& InOperationData, Callbacks&& ...InCallbacks) const;
+	bool DoTernaryOp(PCGMetadataOps::FOperationData& InOperationData, Callbacks&& ...InCallbacks) const;
 
 	template <typename InType1, typename InType2, typename InType3, typename InType4, typename... Callbacks>
-	bool DoQuaternaryOp(FOperationData& InOperationData, Callbacks&& ...InCallbacks) const;
+	bool DoQuaternaryOp(PCGMetadataOps::FOperationData& InOperationData, Callbacks&& ...InCallbacks) const;
 
 	/** To be called if we have no data to perform any operation, it will passthrough the input. */
-	void PassthroughInputs(FPCGContext* InContext, const UPCGMetadataSettingsBase* InSettings) const;
+	void PassthroughInput(FPCGContext* Context, TArray<FPCGTaggedData>& Outputs, const int32 Index) const;
+	void PassthroughAllInputs(FPCGContext* Context, TArray<FPCGTaggedData>& Outputs) const;
 };
 
-template <int32 NbInputs, int32 NbOutputs>
-inline void FPCGMetadataElementBase::FOperationData::Validate()
-{
-	check(OutputKeys.Num() == NbOutputs);
-
-	for (int32 i = 0; i < NbInputs; ++i)
-	{
-		check(InputAccessors[i].IsValid());
-		check(InputKeys[i].IsValid());
-	}
-
-	for (int32 j = 0; j < NbOutputs; ++j)
-	{
-		check(OutputAccessors[j].IsValid());
-		check(OutputKeys[j].IsValid());
-	}
-}
-
 /**
 * Heavy templated code to factorize the gathering of inputs, the application of callbacks and the set of outputs.
 * Note that the order of calls is from the bottom to the top: Operation -> Gather -> Apply.
@@ -265,7 +252,7 @@ namespace NAryOperation
 	* Note that we go in reverse order, and stop when OutputIndex is negative.
 	*/
 	template<int OutputIndex, typename ...Callbacks, typename ...Args>
-	bool Apply(FPCGMetadataElementBase::FOperationData& InOperationData, int32 StartIndex, int32 Range, const Options& InOptions, const TTuple<Callbacks...>& InCallbacks, Args&& ...InArgs)
+	bool Apply(PCGMetadataOps::FOperationData& InOperationData, int32 StartIndex, int32 Range, const Options& InOptions, const TTuple<Callbacks...>& InCallbacks, Args&& ...InArgs)
 	{
 		if constexpr (OutputIndex < 0)
 		{
@@ -321,7 +308,7 @@ namespace NAryOperation
 	* To do so, we use Apply templated with the LastOutputIndex, and go backwards (last output to first output).
 	*/
 	template <typename... Callbacks, typename... Args>
-	bool Gather(FPCGMetadataElementBase::FOperationData& InOperationData, int32 StartIndex, int32 Range, const Options& InOptions, const TTuple<Callbacks...>& InCallbacks, int InputIndex, Signature<> S, Args&& ...InArgs)
+	bool Gather(PCGMetadataOps::FOperationData& InOperationData, int32 StartIndex, int32 Range, const Options& InOptions, const TTuple<Callbacks...>& InCallbacks, int InputIndex, Signature<> S, Args&& ...InArgs)
 	{
 		TRACE_CPUPROFILER_EVENT_SCOPE(PCG::Private::NAryOperation::Apply);
 
@@ -341,7 +328,7 @@ namespace NAryOperation
 	* Gather(InOperationData, StartIndex, Range, InOptions, InCallbacks, 2, Signature<>, IntValues, FloatValues);
 	*/
 	template <typename... Callbacks, typename InputType, typename... InputTypes, typename... Args>
-	bool Gather(FPCGMetadataElementBase::FOperationData& InOperationData, int32 StartIndex, int32 Range, const Options& InOptions, const TTuple<Callbacks...>& InCallbacks, int InputIndex, Signature<InputType, InputTypes...> S, Args&& ...InArgs)
+	bool Gather(PCGMetadataOps::FOperationData& InOperationData, int32 StartIndex, int32 Range, const Options& InOptions, const TTuple<Callbacks...>& InCallbacks, int InputIndex, Signature<InputType, InputTypes...> S, Args&& ...InArgs)
 	{
 		bool bSuccess = true;
 
@@ -370,7 +357,7 @@ namespace NAryOperation
 	* We use an index, starting at 0, and also use our empty struct to pack all our input types.
 	*/
 	template <typename... InputTypes, typename... Callbacks>
-	bool Operation(FPCGMetadataElementBase::FOperationData& InOperationData, int32 StartIndex, int32 Range, const Options& InOptions, const TTuple<Callbacks...>& InCallbacks)
+	bool Operation(PCGMetadataOps::FOperationData& InOperationData, int32 StartIndex, int32 Range, const Options& InOptions, const TTuple<Callbacks...>& InCallbacks)
 	{
 		TRACE_CPUPROFILER_EVENT_SCOPE(PCG::Private::NAryOperation::Operation);
 		return Gather(InOperationData, StartIndex, Range, InOptions, InCallbacks, 0, Signature<InputTypes...>());
@@ -379,10 +366,28 @@ namespace NAryOperation
 }
 }
 
+template <int32 NbInputs, int32 NbOutputs>
+void PCGMetadataOps::FOperationData::Validate()
+{
+	check(OutputKeys.Num() == NbOutputs);
+
+	for (int32 i = 0; i < NbInputs; ++i)
+	{
+		check(InputAccessors[i].IsValid());
+		check(InputKeys[i].IsValid());
+	}
+
+	for (int32 j = 0; j < NbOutputs; ++j)
+	{
+		check(OutputAccessors[j].IsValid());
+		check(OutputKeys[j].IsValid());
+	}
+}
+
 template <typename... InputTypes, typename... Callbacks>
-inline bool FPCGMetadataElementBase::DoNAryOp(FOperationData& InOperationData, TTuple<Callbacks...>&& InCallbacks) const
+inline bool FPCGMetadataElementBase::DoNAryOp(PCGMetadataOps::FOperationData& InOperationData, TTuple<Callbacks...>&& InCallbacks) const
 {
-	// If nothing to do, exit immediatly
+	// If nothing to do, exit immediately
 	if (InOperationData.NumberOfElementsToProcess == 0)
 	{
 		return true;
@@ -420,25 +425,25 @@ inline bool FPCGMetadataElementBase::DoNAryOp(FOperationData& InOperationData, T
 }
 
 template <typename InType, typename... Callbacks>
-inline bool FPCGMetadataElementBase::DoUnaryOp(FOperationData& InOperationData, Callbacks&& ...InCallbacks) const
+inline bool FPCGMetadataElementBase::DoUnaryOp(PCGMetadataOps::FOperationData& InOperationData, Callbacks&& ...InCallbacks) const
 {
 	return DoNAryOp<InType>(InOperationData, ForwardAsTuple(std::forward<Callbacks>(InCallbacks)...));
 }
 
 template <typename InType1, typename InType2, typename... Callbacks>
-inline bool FPCGMetadataElementBase::DoBinaryOp(FOperationData& InOperationData, Callbacks&& ...InCallbacks) const
+inline bool FPCGMetadataElementBase::DoBinaryOp(PCGMetadataOps::FOperationData& InOperationData, Callbacks&& ...InCallbacks) const
 {
 	return DoNAryOp<InType1, InType2>(InOperationData, ForwardAsTuple(std::forward<Callbacks>(InCallbacks)...));
 }
 
 template <typename InType1, typename InType2, typename InType3, typename... Callbacks>
-inline bool FPCGMetadataElementBase::DoTernaryOp(FOperationData& InOperationData, Callbacks&& ...InCallbacks) const
+inline bool FPCGMetadataElementBase::DoTernaryOp(PCGMetadataOps::FOperationData& InOperationData, Callbacks&& ...InCallbacks) const
 {
 	return DoNAryOp<InType1, InType2, InType3>(InOperationData, ForwardAsTuple(std::forward<Callbacks>(InCallbacks)...));
 }
 
 template <typename InType1, typename InType2, typename InType3, typename InType4, typename... Callbacks>
-inline bool FPCGMetadataElementBase::DoQuaternaryOp(FOperationData& InOperationData, Callbacks&& ...InCallbacks) const
+inline bool FPCGMetadataElementBase::DoQuaternaryOp(PCGMetadataOps::FOperationData& InOperationData, Callbacks&& ...InCallbacks) const
 {
 	return DoNAryOp<InType1, InType2, InType3, InType4>(InOperationData, ForwardAsTuple(std::forward<Callbacks>(InCallbacks)...));
 }
diff --git a/Engine/Plugins/Experimental/PCG/Source/PCG/Public/Elements/Metadata/PCGMetadataRotatorOpElement.h b/Engine/Plugins/Experimental/PCG/Source/PCG/Public/Elements/Metadata/PCGMetadataRotatorOpElement.h
--- a/Engine/Plugins/Experimental/PCG/Source/PCG/Public/Elements/Metadata/PCGMetadataRotatorOpElement.h
+++ b/Engine/Plugins/Experimental/PCG/Source/PCG/Public/Elements/Metadata/PCGMetadataRotatorOpElement.h
@@ -48,10 +48,10 @@ public:
 	//~End UPCGSettings interface
 
 	//~Begin UPCGMetadataSettingsBase interface
-	FPCGAttributePropertyInputSelector GetInputSource(uint32 Index) const override;
+	virtual FPCGAttributePropertyInputSelector GetInputSource(uint32 Index) const override;
 
 	virtual FName GetInputPinLabel(uint32 Index) const override;
-	virtual uint32 GetInputPinNum() const override;
+	virtual uint32 GetOperandNum() const override;
 
 	virtual bool IsSupportedInputType(uint16 TypeId, uint32 InputIndex, bool& bHasSpecialRequirement) const override;
 	//~End UPCGMetadataSettingsBase interface
@@ -89,5 +89,5 @@ public:
 class FPCGMetadataRotatorElement : public FPCGMetadataElementBase
 {
 protected:
-	virtual bool DoOperation(FOperationData& OperationData) const override;
+	virtual bool DoOperation(PCGMetadataOps::FOperationData& OperationData) const override;
 };
diff --git a/Engine/Plugins/Experimental/PCG/Source/PCG/Public/Elements/Metadata/PCGMetadataStringOpElement.h b/Engine/Plugins/Experimental/PCG/Source/PCG/Public/Elements/Metadata/PCGMetadataStringOpElement.h
--- a/Engine/Plugins/Experimental/PCG/Source/PCG/Public/Elements/Metadata/PCGMetadataStringOpElement.h
+++ b/Engine/Plugins/Experimental/PCG/Source/PCG/Public/Elements/Metadata/PCGMetadataStringOpElement.h
@@ -31,10 +31,10 @@ protected:
 
 public:
 	//~Begin UPCGMetadataSettingsBase interface
-	FPCGAttributePropertyInputSelector GetInputSource(uint32 Index) const override;
+	virtual FPCGAttributePropertyInputSelector GetInputSource(uint32 Index) const override;
 
 	virtual FName GetInputPinLabel(uint32 Index) const override;
-	virtual uint32 GetInputPinNum() const override;
+	virtual uint32 GetOperandNum() const override;
 
 	virtual bool IsSupportedInputType(uint16 TypeId, uint32 InputIndex, bool& bHasSpecialRequirement) const override;
 	virtual uint16 GetOutputType(uint16 InputTypeId) const override;
@@ -54,5 +54,5 @@ public:
 class FPCGMetadataStringOpElement : public FPCGMetadataElementBase
 {
 protected:
-	virtual bool DoOperation(FOperationData& OperationData) const override;
+	virtual bool DoOperation(PCGMetadataOps::FOperationData& OperationData) const override;
 };
\ No newline at end of file
diff --git a/Engine/Plugins/Experimental/PCG/Source/PCG/Public/Elements/Metadata/PCGMetadataTransformOpElement.h b/Engine/Plugins/Experimental/PCG/Source/PCG/Public/Elements/Metadata/PCGMetadataTransformOpElement.h
--- a/Engine/Plugins/Experimental/PCG/Source/PCG/Public/Elements/Metadata/PCGMetadataTransformOpElement.h
+++ b/Engine/Plugins/Experimental/PCG/Source/PCG/Public/Elements/Metadata/PCGMetadataTransformOpElement.h
@@ -50,10 +50,10 @@ public:
 	//~End UPCGSettings interface
 
 	//~Begin UPCGMetadataSettingsBase interface
-	FPCGAttributePropertyInputSelector GetInputSource(uint32 Index) const override;
+	virtual FPCGAttributePropertyInputSelector GetInputSource(uint32 Index) const override;
 
 	virtual FName GetInputPinLabel(uint32 Index) const override;
-	virtual uint32 GetInputPinNum() const override;
+	virtual uint32 GetOperandNum() const override;
 
 	virtual bool IsSupportedInputType(uint16 TypeId, uint32 InputIndex, bool& bHasSpecialRequirement) const override;
 	//~End UPCGMetadataSettingsBase interface
@@ -94,5 +94,5 @@ public:
 class FPCGMetadataTransformElement : public FPCGMetadataElementBase
 {
 protected:
-	virtual bool DoOperation(FOperationData& OperationData) const override;
+	virtual bool DoOperation(PCGMetadataOps::FOperationData& OperationData) const override;
 };
diff --git a/Engine/Plugins/Experimental/PCG/Source/PCG/Public/Elements/Metadata/PCGMetadataTrigOpElement.h b/Engine/Plugins/Experimental/PCG/Source/PCG/Public/Elements/Metadata/PCGMetadataTrigOpElement.h
--- a/Engine/Plugins/Experimental/PCG/Source/PCG/Public/Elements/Metadata/PCGMetadataTrigOpElement.h
+++ b/Engine/Plugins/Experimental/PCG/Source/PCG/Public/Elements/Metadata/PCGMetadataTrigOpElement.h
@@ -42,10 +42,10 @@ public:
 	//~End UPCGSettings interface
 
 	//~Begin UPCGMetadataSettingsBase interface
-	FPCGAttributePropertyInputSelector GetInputSource(uint32 Index) const override;
+	virtual FPCGAttributePropertyInputSelector GetInputSource(uint32 Index) const override;
 
 	virtual FName GetInputPinLabel(uint32 Index) const override;
-	virtual uint32 GetInputPinNum() const override;
+	virtual uint32 GetOperandNum() const override;
 
 	virtual bool IsSupportedInputType(uint16 TypeId, uint32 InputIndex, bool& bHasSpecialRequirement) const override;
 	virtual uint16 GetOutputType(uint16 InputTypeId) const override;
@@ -78,5 +78,5 @@ public:
 class FPCGMetadataTrigElement : public FPCGMetadataElementBase
 {
 protected:
-	virtual bool DoOperation(FOperationData& OperationData) const override;
+	virtual bool DoOperation(PCGMetadataOps::FOperationData& OperationData) const override;
 };
diff --git a/Engine/Plugins/Experimental/PCG/Source/PCG/Public/Elements/Metadata/PCGMetadataVectorOpElement.h b/Engine/Plugins/Experimental/PCG/Source/PCG/Public/Elements/Metadata/PCGMetadataVectorOpElement.h
--- a/Engine/Plugins/Experimental/PCG/Source/PCG/Public/Elements/Metadata/PCGMetadataVectorOpElement.h
+++ b/Engine/Plugins/Experimental/PCG/Source/PCG/Public/Elements/Metadata/PCGMetadataVectorOpElement.h
@@ -46,10 +46,10 @@ public:
 	//~End UPCGSettings interface
 
 	//~Begin UPCGMetadataSettingsBase interface
-	FPCGAttributePropertyInputSelector GetInputSource(uint32 Index) const override;
+	virtual FPCGAttributePropertyInputSelector GetInputSource(uint32 Index) const override;
 
 	virtual FName GetInputPinLabel(uint32 Index) const override;
-	virtual uint32 GetInputPinNum() const override;
+	virtual uint32 GetOperandNum() const override;
 
 	virtual bool IsSupportedInputType(uint16 TypeId, uint32 InputIndex, bool& bHasSpecialRequirement) const override;
 	virtual uint16 GetOutputType(uint16 InputTypeId) const override;
@@ -88,5 +88,5 @@ public:
 class FPCGMetadataVectorElement : public FPCGMetadataElementBase
 {
 protected:
-	virtual bool DoOperation(FOperationData& OperationData) const override;
+	virtual bool DoOperation(PCGMetadataOps::FOperationData& OperationData) const override;
 };
