[PCG] Increased MetadataOps Unit Test coverage

+ New Unit tests for numerous Maths Ops
+ New Multiple Input unit test
TODO - More math operations will be added at a later date, after this workflow
is approved

#rnx

diff --git a/Engine/Plugins/Experimental/PCG/Source/PCG/Private/Metadata/PCGAttributePropertySelector.cpp b/Engine/Plugins/Experimental/PCG/Source/PCG/Private/Metadata/PCGAttributePropertySelector.cpp
--- a/Engine/Plugins/Experimental/PCG/Source/PCG/Private/Metadata/PCGAttributePropertySelector.cpp
+++ b/Engine/Plugins/Experimental/PCG/Source/PCG/Private/Metadata/PCGAttributePropertySelector.cpp
@@ -273,34 +273,6 @@ bool FPCGAttributePropertySelector::IsBasicAttribute() const
 	return Selection == EPCGAttributePropertySelection::Attribute && ExtraNames.IsEmpty();
 }
 
-FPCGAttributePropertySelector FPCGAttributePropertySelector::CreateAttributeSelector(const FName AttributeName)
-{
-	FPCGAttributePropertySelector Selector;
-	Selector.SetAttributeName(AttributeName);
-	return Selector;
-}
-
-FPCGAttributePropertySelector FPCGAttributePropertySelector::CreatePointPropertySelector(EPCGPointProperties PointProperty)
-{
-	FPCGAttributePropertySelector Selector;
-	Selector.SetPointProperty(PointProperty);
-	return Selector;
-}
-
-FPCGAttributePropertySelector FPCGAttributePropertySelector::CreateExtraPropertySelector(EPCGExtraProperties ExtraProperty)
-{
-	FPCGAttributePropertySelector Selector;
-	Selector.SetExtraProperty(ExtraProperty);
-	return Selector;
-}
-
-FPCGAttributePropertySelector FPCGAttributePropertySelector::CreateSelectorFromString(const FString& String)
-{
-	FPCGAttributePropertySelector Selector;
-	Selector.Update(String);
-	return Selector;
-}
-
 ///////////////////////////////////////////////////////////////////////
 
 FPCGAttributePropertyInputSelector::FPCGAttributePropertyInputSelector()
diff --git a/Engine/Plugins/Experimental/PCG/Source/PCG/Private/Tests/Elements/Metadata/PCGMetadataMathsOpTest.cpp b/Engine/Plugins/Experimental/PCG/Source/PCG/Private/Tests/Elements/Metadata/PCGMetadataMathsOpTest.cpp
--- a/Engine/Plugins/Experimental/PCG/Source/PCG/Private/Tests/Elements/Metadata/PCGMetadataMathsOpTest.cpp
+++ b/Engine/Plugins/Experimental/PCG/Source/PCG/Private/Tests/Elements/Metadata/PCGMetadataMathsOpTest.cpp
@@ -5,98 +5,445 @@
 #include "PCGComponent.h"
 #include "PCGContext.h"
 #include "PCGParamData.h"
+#include "Data/PCGPointData.h"
 #include "Elements/Metadata/PCGMetadataMathsOpElement.h"
-#include "Metadata/PCGMetadataAccessor.h"
-#include "Metadata/PCGMetadataAttributeTraits.h"
+#include "Helpers/PCGPropertyHelpers.h"
 
 #if WITH_EDITOR
 
-IMPLEMENT_CUSTOM_SIMPLE_AUTOMATION_TEST(FPCGMetadataMathsOp_OneMinus, FPCGTestBaseClass, "pcg.tests.Metadata.MathsOp.OneMinus", PCGTestsCommon::TestFlags)
-
-namespace PCGMathsOpTest
+class FPCGMetadataMathsOpTest : public FPCGTestBaseClass
 {
-	const FName AttributeName = TEXT("Attr");
+	using FPCGTestBaseClass::FPCGTestBaseClass;
+
+protected:
+	inline static const FName DefaultAttributeName = TEXT("TestAttribute");
+	inline static const FVector DefaultTranslation = {1.0, 10.0, 100.0};
+	inline static const FPCGAttributePropertyInputSelector DefaultSelector = FPCGAttributePropertySelector::CreateAttributeSelector<FPCGAttributePropertyInputSelector>(DefaultAttributeName);
 
 	template <typename T>
-	FPCGTaggedData GenerateParamData(PCGTestsCommon::FTestData& TestData, const T& Value)
+	struct TestParams
 	{
-		FPCGTaggedData& ParamTaggedData = TestData.InputData.TaggedData.Emplace_GetRef();
-		TObjectPtr<UPCGParamData> ParamData = PCGTestsCommon::CreateEmptyParamData();
-		ParamTaggedData.Data = ParamData;
+		explicit TestParams(const EPCGMedadataMathsOperation Operation, TArray<TArray<FPCGTaggedData>>&& InputDataByPin, TArray<T>&& ExpectedDefaultValues, TArray<T>&& ExpectedResultValues, const int32 PinToForward = 0) :
+			Operation(Operation),
+			InputDataByPin(std::move(InputDataByPin)),
+			ExpectedDefaultValues(ExpectedDefaultValues),
+			ExpectedResultValues(ExpectedResultValues),
+			PinToForward(PinToForward)
+		{
+		}
 
-		PCGMetadataEntryKey EntryKey = ParamData->Metadata->AddEntry();
-		const bool bAllowsInterpolation = true;
-		const bool bOverrideParent = false;
+		EPCGMedadataMathsOperation Operation;
+		TArray<TArray<FPCGTaggedData>> InputDataByPin;
+		TArray<T> ExpectedDefaultValues;
+		TArray<T> ExpectedResultValues;
+		int32 PinToForward = 0;
+	};
 
-		FPCGMetadataAttribute<T>* Attribute = ParamData->Metadata->CreateAttribute<T>(AttributeName, PCG::Private::MetadataTraits<T>::ZeroValue(), bAllowsInterpolation, bOverrideParent);
-		check(Attribute);
-		Attribute->SetValue(EntryKey, Value);
+	// Expects Input Data to be be an array of Tagged Data on each index per pin and a single value in each Input Data (1 attribute or 1 point data)
+	template <typename T>
+	bool ExecuteTest(TestParams<T>&& Params, const FPCGAttributePropertyInputSelector InputSelector = DefaultSelector)
+	{
+		const EPCGMedadataMathsOperation& Operation = Params.Operation;
+		const TArray<TArray<FPCGTaggedData>>& InputDataByPin = Params.InputDataByPin;
+		const TArray<T>& ExpectedDefaults = Params.ExpectedDefaultValues;
+		const TArray<T>& ExpectedValues = Params.ExpectedResultValues;
+		const int32 PinToForward = Params.PinToForward;
 
-		return ParamTaggedData;
-	}
-}
+		check(!InputDataByPin.IsEmpty() && !ExpectedDefaults.IsEmpty() && !ExpectedValues.IsEmpty());
 
+		const UEnum* EnumClass = StaticEnum<EPCGMedadataMathsOperation>();
+		auto FormatText = [Operation, PinToForward, EnumClass](const FString& Text) -> FString
+		{
+			return FString::Printf(TEXT("[Pin %d] %s (%s): %s."), PinToForward, *EnumClass->GetDisplayNameTextByValue(static_cast<int64>(Operation)).ToString(), *PCG::Private::GetTypeName<T>(), *Text);
+		};
+
+		bool bSuccess = true;
 
-bool FPCGMetadataMathsOp_OneMinus::RunTest(const FString& Parameters)
-{
-	auto RunTest = [this]<typename T>(const T& Value, const T& ExpectedDefaultValue, const T& ExpectedValue)
-	{
 		PCGTestsCommon::FTestData TestData;
 		UPCGMetadataMathsSettings* Settings = PCGTestsCommon::GenerateSettings<UPCGMetadataMathsSettings>(TestData);
-		check(Settings);
+		FPCGElementPtr TestElement = TestData.Settings->GetElement();
 
-		Settings->Operation = EPCGMedadataMathsOperation::OneMinus;
-		Settings->InputSource1.SetAttributeName(PCGMathsOpTest::AttributeName);
-		Settings->OutputTarget.SetAttributeName(PCGMathsOpTest::AttributeName);
+		// Assign the element settings here. Use the same selector for all input sources, since the user can only select one
+		Settings->Operation = Operation;
+		Settings->InputSource1 = InputSelector;
+		Settings->InputSource2 = InputSelector;
+		Settings->InputSource3 = InputSelector;
+		Settings->OutputTarget.SetAttributeName(DefaultAttributeName);
+		Settings->OutputDataFromPin = Settings->GetInputPinLabel(PinToForward);
 
-		FPCGTaggedData& InputTaggedData = TestData.InputData.TaggedData.Add_GetRef(PCGMathsOpTest::GenerateParamData(TestData, Value));
-		InputTaggedData.Pin = PCGPinConstants::DefaultInputLabel;
+		// Assign pin labels based on incoming data for the execution
+		for (int PinIndex = 0; PinIndex < InputDataByPin.Num(); ++PinIndex)
+		{
+			TArray<FPCGTaggedData> InputData = InputDataByPin[PinIndex];
+			for (FPCGTaggedData& Data : InputData)
+			{
+				Data.Pin = Settings->GetInputPinLabel(PinIndex);
+			}
+
+			TestData.InputData.TaggedData.Append(std::move(InputData));
+		}
 
-		FPCGElementPtr TestElement = TestData.Settings->GetElement();
 		TUniquePtr<FPCGContext> Context = TestData.InitializeTestContext();
 
-		while (!TestElement->Execute(Context.Get())) {}
+		while (!TestElement->Execute(Context.Get()))
+		{
+		}
 
-		auto FormatText = [](const FString& In) -> FString { return FString::Printf(TEXT("OneMinus %s: %s"), *PCG::Private::GetTypeName<T>(), *In); };
+		UTEST_TRUE(*FormatText("Output created"), !Context->OutputData.TaggedData.IsEmpty());
 
-		UTEST_EQUAL(*FormatText("There is one output"), Context->OutputData.TaggedData.Num(), 1);
+		UTEST_EQUAL(*FormatText("Correct number of outputs"), Context->OutputData.TaggedData.Num(), ExpectedValues.Num());
 
-		const UPCGParamData* OutputParam = Cast<const UPCGParamData>(Context->OutputData.TaggedData[0].Data);
+		// Cycle through expected results and validate
+		for (int I = 0; I < ExpectedValues.Num(); ++I)
+		{
+			const UPCGData* OutputData = Context->OutputData.TaggedData[I].Data;
+			UTEST_TRUE("Valid output data", OutputData != nullptr);
 
-		UTEST_NOT_NULL(*FormatText("Output is an attribute set"), OutputParam);
+			FPCGAttributePropertySelector Selector = FPCGAttributePropertySelector::CreateAttributeSelector(DefaultAttributeName);
+			const TUniquePtr<const IPCGAttributeAccessor> Accessor = PCGAttributeAccessorHelpers::CreateConstAccessor(OutputData, Selector);
+			const TUniquePtr<const IPCGAttributeAccessorKeys> Keys = PCGAttributeAccessorHelpers::CreateConstKeys(OutputData, Selector);
 
-		const FPCGMetadataAttribute<T>* OutputAttribute = OutputParam->Metadata->GetConstTypedAttribute<T>(PCGMathsOpTest::AttributeName);
+			UTEST_TRUE(*FormatText("Target attribute or property was found"), Accessor.IsValid() && Keys.IsValid());
 
-		UTEST_NOT_NULL(*FormatText("Output metadata has the expected attribute"), OutputAttribute);
+			UTEST_EQUAL(*FormatText("Output has the correct type"), Accessor->GetUnderlyingType(), PCG::Private::MetadataTypes<T>::Id);
 
-		// Vec2 and Vec4 don't have "nearly equal" in the test framework
-		T DefaultValue = OutputAttribute->GetValue(PCGDefaultValueKey);
-		T FirstEntry = OutputAttribute->GetValueFromItemKey(0);
-		if constexpr (std::is_same_v<T, FVector4> || std::is_same_v<T, FVector2D>)
+			T DefaultValue{};
+			T OutputValue{};
+
+			FPCGAttributeAccessorKeysEntries DefaultKey(PCGInvalidEntryKey);
+			bSuccess &= TestTrue(*FormatText("Output default value attribute accessible"), Accessor->Get(DefaultValue, 0, DefaultKey));
+
+			// We're only testing one attribute for each iteration
+			bSuccess &= TestTrue(*FormatText("Output value attribute accessible"), Accessor->Get(OutputValue, 0, *Keys));
+
+			// Vec2 and Vec4 don't have "nearly equal" in the test framework
+			if constexpr (std::is_same_v<T, FVector2D> || std::is_same_v<T, FVector4>)
+			{
+				bSuccess &= TestTrue(*FormatText("Default value is correct"), DefaultValue.Equals(ExpectedDefaults[I]));
+				bSuccess &= TestTrue(*FormatText("Result is correct"), OutputValue.Equals(ExpectedValues[I]));
+			}
+			else
+			{
+				bSuccess &= TestEqual(*FormatText("Default value is correct"), DefaultValue, ExpectedDefaults[I]);
+				bSuccess &= TestEqual(*FormatText("Result is correct"), OutputValue, ExpectedValues[I]);
+			}
+		}
+
+		return bSuccess;
+	}
+
+	// Generate points with a default location, density, and seed based on index
+	TArray<TArray<FPCGTaggedData>> GeneratePointData(const int32 NumPins = 1, const int32 NumData = 1)
+	{
+		TArray<TArray<FPCGTaggedData>> GeneratedData;
+
+		for (int PinIndex = 0; PinIndex < NumPins; ++PinIndex)
 		{
-			UTEST_TRUE(*FormatText("Default value is correct"), DefaultValue.Equals(ExpectedDefaultValue));
-			UTEST_TRUE(*FormatText("First entry is correct"), FirstEntry.Equals(ExpectedValue));
+			TArray<FPCGTaggedData>& TaggedData = GeneratedData.Emplace_GetRef();
+
+			for (int32 InputIndex = 0; InputIndex < NumData; ++InputIndex)
+			{
+				int32 M = NumData * PinIndex + InputIndex + 1;
+				TObjectPtr<UPCGPointData> PointData = PCGTestsCommon::CreateEmptyPointData();
+				TArray<FPCGPoint>& Points = PointData->GetMutablePoints();
+				Points.Emplace(FTransform(M * DefaultTranslation), static_cast<double>(NumData * NumPins) / M, M);
+				TaggedData.Emplace_GetRef().Data = PointData;
+			}
 		}
-		else
+
+		return GeneratedData;
+	}
+
+	// Generate param data of the tested type
+	template <typename T>
+	TArray<TArray<FPCGTaggedData>> GenerateParamData(const TArray<TArray<T>> ValuesByPin)
+	{
+		TArray<TArray<FPCGTaggedData>> GeneratedData;
+
+		for (TArray<T> Values : ValuesByPin)
 		{
-			UTEST_EQUAL(*FormatText("Default value is correct"), DefaultValue, ExpectedDefaultValue);
-			UTEST_EQUAL(*FormatText("First entry is correct"), FirstEntry, ExpectedValue);
+			TArray<FPCGTaggedData>& TaggedData = GeneratedData.Emplace_GetRef();
+
+			for (const T& Value : Values)
+			{
+				TObjectPtr<UPCGParamData> ParamData = PCGTestsCommon::CreateEmptyParamData();
+				TaggedData.Emplace_GetRef().Data = ParamData;
+
+				PCGMetadataEntryKey EntryKey = ParamData->Metadata->AddEntry();
+				constexpr bool bAllowsInterpolation = true;
+				constexpr bool bOverrideParent = false;
+
+				FPCGMetadataAttribute<T>* Attribute = ParamData->Metadata->FindOrCreateAttribute<T>(DefaultAttributeName, PCG::Private::MetadataTraits<T>::ZeroValue(), bAllowsInterpolation, bOverrideParent);
+				check(Attribute);
+
+				Attribute->SetValue(EntryKey, Value);
+			}
 		}
 
-		return true;
-	};
+		return GeneratedData;
+	}
+};
+
+IMPLEMENT_CUSTOM_SIMPLE_AUTOMATION_TEST(FPCGMetadataMathsOp_OneMinus, FPCGMetadataMathsOpTest, "pcg.tests.Metadata.MathsOp.OneMinus", PCGTestsCommon::TestFlags)
 
+bool FPCGMetadataMathsOp_OneMinus::RunTest(const FString& Parameters)
+{
 	bool bSuccess = true;
 
-	bSuccess &= RunTest(5, 1, -4);
-	bSuccess &= RunTest(6ll, 1ll, -5ll);
-	bSuccess &= RunTest(0.1f, 1.0f, 0.9f);
-	bSuccess &= RunTest(0.2, 1.0, 0.8);
-	bSuccess &= RunTest(FVector2D(0.3, 0.4), FVector2D(1, 1), FVector2D(0.7, 0.6));
-	bSuccess &= RunTest(FVector(0.5, 0.6, 0.7), FVector(1, 1, 1), FVector(0.5, 0.4, 0.3));
-	bSuccess &= RunTest(FVector4(0.5, 0.6, 0.7, 0.8), FVector4(1, 1, 1, 1), FVector4(0.5, 0.4, 0.3, 0.2));
+	using EPCGMedadataMathsOperation::OneMinus;
+
+	// Param Data
+	bSuccess &= ExecuteTest(TestParams(OneMinus, GenerateParamData<int32>({{5}}), TArray({1}), TArray({-4})));
+	bSuccess &= ExecuteTest(TestParams(OneMinus, GenerateParamData<long long>({{6ll}}), TArray({1ll}), TArray({-5ll})));
+	bSuccess &= ExecuteTest(TestParams(OneMinus, GenerateParamData<float>({{0.1f}}), TArray({1.f}), TArray({0.9f})));
+	bSuccess &= ExecuteTest(TestParams(OneMinus, GenerateParamData<double>({{0.2}}), TArray({1.0}), TArray({0.8})));
+	bSuccess &= ExecuteTest(TestParams(OneMinus, GenerateParamData<FVector2D>({{FVector2D(0.3, 1.4)}}), TArray({FVector2D::One()}), TArray({FVector2D(0.7, -0.4)})));
+	bSuccess &= ExecuteTest(TestParams(OneMinus, GenerateParamData<FVector>({{FVector(0.5, 1.6, 0.7)}}), TArray({FVector::One()}), TArray({FVector(0.5, -0.6, 0.3)})));
+	bSuccess &= ExecuteTest(TestParams(OneMinus, GenerateParamData<FVector4>({{FVector4(0.5, 1.6, 0.7, 0.8)}}), TArray({FVector4::One()}), TArray({FVector4(0.5, -0.6, 0.3, 0.2)})));
+
+	// Point Data
+	FPCGAttributePropertyInputSelector Selector;
+	Selector.SetPointProperty(EPCGPointProperties::Seed);
+	bSuccess &= ExecuteTest(TestParams(OneMinus, GeneratePointData(), TArray<int64>({0}), TArray<int64>({0})), Selector);
+	Selector.SetPointProperty(EPCGPointProperties::Density);
+	bSuccess &= ExecuteTest(TestParams(OneMinus, GeneratePointData(), TArray({0.0}), TArray({0.0})), Selector);
+	Selector.SetPointProperty(EPCGPointProperties::Position);
+	bSuccess &= ExecuteTest(TestParams(OneMinus, GeneratePointData(), TArray({FVector(0.0, 0.0, 0.0)}), TArray({FVector(0.0, -9.0, -99.0)})), Selector);
 
 	return bSuccess;
 }
 
-#endif // WITH_EDITOR
+IMPLEMENT_CUSTOM_SIMPLE_AUTOMATION_TEST(FPCGMetadataMathsOp_Sign, FPCGMetadataMathsOpTest, "pcg.tests.Metadata.MathsOp.Sign", PCGTestsCommon::TestFlags)
+
+bool FPCGMetadataMathsOp_Sign::RunTest(const FString& Parameters)
+{
+	bool bSuccess = true;
+
+	using EPCGMedadataMathsOperation::Sign;
+
+	// Param Data
+	bSuccess &= ExecuteTest(TestParams(Sign, GenerateParamData<int32>({{0}}), TArray({0}), TArray({0})));
+	bSuccess &= ExecuteTest(TestParams(Sign, GenerateParamData<int32>({{5}}), TArray({0}), TArray({1})));
+	bSuccess &= ExecuteTest(TestParams(Sign, GenerateParamData<int32>({{-5}}), TArray({0}), TArray({-1})));
+	bSuccess &= ExecuteTest(TestParams(Sign, GenerateParamData<long long>({{6ll}}), TArray({0ll}), TArray({1ll})));
+	bSuccess &= ExecuteTest(TestParams(Sign, GenerateParamData<float>({{-0.1f}}), TArray({0.f}), TArray({-1.f})));
+	bSuccess &= ExecuteTest(TestParams(Sign, GenerateParamData<double>({{-4.2}}), TArray({0.0}), TArray({-1.0})));
+	bSuccess &= ExecuteTest(TestParams(Sign, GenerateParamData<FVector2D>({{FVector2D(0.3, -1.4)}}), TArray({FVector2D::Zero()}), TArray({FVector2D(1.0, -1.0)})));
+	bSuccess &= ExecuteTest(TestParams(Sign, GenerateParamData<FVector>({{FVector(0.5, -1.6, 0.0)}}), TArray({FVector::Zero()}), TArray({FVector(1.0, -1.0, 0.0)})));
+	bSuccess &= ExecuteTest(TestParams(Sign, GenerateParamData<FVector4>({{FVector4(0.5, -1.6, 0.0, -0.0)}}), TArray({FVector4::Zero()}), TArray({FVector4(1.0, -1.0, 0.0, 0.0)})));
+
+	// Point Data
+	FPCGAttributePropertyInputSelector Selector;
+	Selector.SetPointProperty(EPCGPointProperties::Seed);
+	bSuccess &= ExecuteTest(TestParams(Sign, GeneratePointData(), TArray<int64>({0}), TArray<int64>({1})), Selector);
+	Selector.SetPointProperty(EPCGPointProperties::Density);
+	bSuccess &= ExecuteTest(TestParams(Sign, GeneratePointData(), TArray({0.0}), TArray({1.0})), Selector);
+	Selector.SetPointProperty(EPCGPointProperties::Position);
+	bSuccess &= ExecuteTest(TestParams(Sign, GeneratePointData(), TArray({FVector(0.0, 0.0, 0.0)}), TArray({FVector(1.0, 1.0, 1.0)})), Selector);
+
+	return bSuccess;
+}
+
+IMPLEMENT_CUSTOM_SIMPLE_AUTOMATION_TEST(FPCGMetadataMathsOp_Frac, FPCGMetadataMathsOpTest, "pcg.tests.Metadata.MathsOp.Frac", PCGTestsCommon::TestFlags)
+
+bool FPCGMetadataMathsOp_Frac::RunTest(const FString& Parameters)
+{
+	bool bSuccess = true;
+
+	using EPCGMedadataMathsOperation::Frac;
+
+	// Param Data
+	bSuccess &= ExecuteTest(TestParams(Frac, GenerateParamData<int32>({{0}}), TArray({0}), TArray({0})));
+	bSuccess &= ExecuteTest(TestParams(Frac, GenerateParamData<long long>({{6ll}}), TArray({0ll}), TArray({0ll})));
+	bSuccess &= ExecuteTest(TestParams(Frac, GenerateParamData<float>({{-0.1f}}), TArray({0.f}), TArray({.9f})));
+	bSuccess &= ExecuteTest(TestParams(Frac, GenerateParamData<double>({{-0.5}}), TArray({0.0}), TArray({0.5})));
+	bSuccess &= ExecuteTest(TestParams(Frac, GenerateParamData<FVector2D>({{FVector2D(0.3, -1.4)}}), TArray({FVector2D::Zero()}), TArray({FVector2D(.3, .6)})));
+	bSuccess &= ExecuteTest(TestParams(Frac, GenerateParamData<FVector>({{FVector(0.5, -1.6, 0.0)}}), TArray({FVector::Zero()}), TArray({FVector(.5, .4, .0)})));
+	bSuccess &= ExecuteTest(TestParams(Frac, GenerateParamData<FVector4>({{FVector4(0.5, -1.6, 0.0, -0.0)}}), TArray({FVector4::Zero()}), TArray({FVector4(.5, .4, 0.0, 0.0)})));
+
+	// Point Data
+	FPCGAttributePropertyInputSelector Selector;
+	Selector.SetPointProperty(EPCGPointProperties::Seed);
+	bSuccess &= ExecuteTest(TestParams(Frac, GeneratePointData(), TArray<int64>({0}), TArray<int64>({0})), Selector);
+	Selector.SetPointProperty(EPCGPointProperties::Density);
+	bSuccess &= ExecuteTest(TestParams(Frac, GeneratePointData(), TArray({0.0}), TArray({0.0})), Selector);
+	Selector.SetPointProperty(EPCGPointProperties::Position);
+	bSuccess &= ExecuteTest(TestParams(Frac, GeneratePointData(), TArray({FVector(0.0, 0.0, 0.0)}), TArray({FVector(0.0, 0.0, 0.0)})), Selector);
+
+	return bSuccess;
+}
+
+IMPLEMENT_CUSTOM_SIMPLE_AUTOMATION_TEST(FPCGMetadataMathsOp_Truncate, FPCGMetadataMathsOpTest, "pcg.tests.Metadata.MathsOp.Truncate", PCGTestsCommon::TestFlags)
+
+bool FPCGMetadataMathsOp_Truncate::RunTest(const FString& Parameters)
+{
+	bool bSuccess = true;
+
+	using EPCGMedadataMathsOperation::Truncate;
+
+	// Param Data
+	bSuccess &= ExecuteTest(TestParams(Truncate, GenerateParamData<int32>({{0}}), TArray({0}), TArray({0})));
+	bSuccess &= ExecuteTest(TestParams(Truncate, GenerateParamData<long long>({{6ll}}), TArray({0ll}), TArray({6ll})));
+	bSuccess &= ExecuteTest(TestParams(Truncate, GenerateParamData<float>({{-0.1f}}), TArray({0.f}), TArray({0.0f})));
+	bSuccess &= ExecuteTest(TestParams(Truncate, GenerateParamData<double>({{-0.3}}), TArray({0.0}), TArray({0.0})));
+	bSuccess &= ExecuteTest(TestParams(Truncate, GenerateParamData<FVector2D>({{FVector2D(0.3, -1.4)}}), TArray({FVector2D::Zero()}), TArray({FVector2D(0.0, -1.0)})));
+	bSuccess &= ExecuteTest(TestParams(Truncate, GenerateParamData<FVector>({{FVector(0.5, -1.6, 0.0)}}), TArray({FVector::Zero()}), TArray({FVector(0.0, -1.0, .0)})));
+	bSuccess &= ExecuteTest(TestParams(Truncate, GenerateParamData<FVector4>({{FVector4(0.5, -1.6, 0.0, -0.0)}}), TArray({FVector4::Zero()}), TArray({FVector4(0.0, -1.0, 0.0, 0.0)})));
+
+	// Point Data
+	FPCGAttributePropertyInputSelector Selector;
+	Selector.SetPointProperty(EPCGPointProperties::Seed);
+	bSuccess &= ExecuteTest(TestParams(Truncate, GeneratePointData(), TArray<int64>({0}), TArray<int64>({1})), Selector);
+	Selector.SetPointProperty(EPCGPointProperties::Density);
+	bSuccess &= ExecuteTest(TestParams(Truncate, GeneratePointData(), TArray({0.0}), TArray({1.0})), Selector);
+	Selector.SetPointProperty(EPCGPointProperties::Position);
+	bSuccess &= ExecuteTest(TestParams(Truncate, GeneratePointData(), TArray({FVector(0.0, 0.0, 0.0)}), TArray({FVector(1.0, 10.0, 100.0)})), Selector);
+
+	return bSuccess;
+}
+
+IMPLEMENT_CUSTOM_SIMPLE_AUTOMATION_TEST(FPCGMetadataMathsOp_Add, FPCGMetadataMathsOpTest, "pcg.tests.Metadata.MathsOp.Add", PCGTestsCommon::TestFlags)
+
+bool FPCGMetadataMathsOp_Add::RunTest(const FString& Parameters)
+{
+	bool bSuccess = true;
+
+	using EPCGMedadataMathsOperation::Add;
+
+	// Param Data
+	bSuccess &= ExecuteTest(TestParams(Add, GenerateParamData<int32>({{5}, {1}}), TArray({0}), TArray({6})));
+	bSuccess &= ExecuteTest(TestParams(Add, GenerateParamData<long long>({{6ll}, {-1ll}}), TArray({0ll}), TArray({5ll})));
+	bSuccess &= ExecuteTest(TestParams(Add, GenerateParamData<float>({{0.1f}, {0.4f}}), TArray({0.f}), TArray({0.5f})));
+	bSuccess &= ExecuteTest(TestParams(Add, GenerateParamData<double>({{0.0}, {-2.4}}), TArray({0.0}), TArray({-2.4})));
+	bSuccess &= ExecuteTest(TestParams(Add, GenerateParamData<FVector2D>({{FVector2D(0.3, 1.4)}, {FVector2D(-1.2, 1.4)}}), TArray({FVector2D::Zero()}), TArray({FVector2D(-0.9, 2.8)})));
+	bSuccess &= ExecuteTest(TestParams(Add, GenerateParamData<FVector>({{FVector(0.5, 1.6, 0.7)}, {FVector(-0.5, 4.2, -0.7)}}), TArray({FVector::Zero()}), TArray({FVector(0.0, 5.8, 0.0)})));
+	bSuccess &= ExecuteTest(TestParams(Add, GenerateParamData<FVector4>({{FVector4(0.5, 1.6, 0.7, 0.8)}, {FVector4(-0.5, 4.2, -0.7, 0.3)}}), TArray({FVector4::Zero()}), TArray({FVector4(0.0, 5.8, 0.0, 1.1)})));
+
+	// Point Data
+	FPCGAttributePropertyInputSelector Selector;
+	Selector.SetPointProperty(EPCGPointProperties::Seed);
+	bSuccess &= ExecuteTest(TestParams(Add, GeneratePointData(2), TArray<int64>({0}), TArray<int64>({3})), Selector);
+	Selector.SetPointProperty(EPCGPointProperties::Density);
+	bSuccess &= ExecuteTest(TestParams(Add, GeneratePointData(2), TArray({0.0}), TArray({3.0})), Selector);
+	Selector.SetPointProperty(EPCGPointProperties::Position);
+	bSuccess &= ExecuteTest(TestParams(Add, GeneratePointData(2), TArray({FVector::Zero()}), TArray({FVector(3.0, 30.0, 300.0)})), Selector);
+
+	return bSuccess;
+}
+
+IMPLEMENT_CUSTOM_SIMPLE_AUTOMATION_TEST(FPCGMetadataMathsOp_Subtract, FPCGMetadataMathsOpTest, "pcg.tests.Metadata.MathsOp.Subtract", PCGTestsCommon::TestFlags)
+
+bool FPCGMetadataMathsOp_Subtract::RunTest(const FString& Parameters)
+{
+	bool bSuccess = true;
+
+	using EPCGMedadataMathsOperation::Subtract;
+
+	// Param Data
+	bSuccess &= ExecuteTest(TestParams(Subtract, GenerateParamData<int32>({{5}, {1}}), TArray({0}), TArray({4})));
+	bSuccess &= ExecuteTest(TestParams(Subtract, GenerateParamData<long long>({{6ll}, {-1ll}}), TArray({0ll}), TArray({7ll})));
+	bSuccess &= ExecuteTest(TestParams(Subtract, GenerateParamData<float>({{0.1f}, {0.4f}}), TArray({0.f}), TArray({-.3f})));
+	bSuccess &= ExecuteTest(TestParams(Subtract, GenerateParamData<double>({{0.0}, {-2.4}}), TArray({0.0}), TArray({2.4})));
+	bSuccess &= ExecuteTest(TestParams(Subtract, GenerateParamData<FVector2D>({{FVector2D(0.3, 1.4)}, {FVector2D(-1.2, 1.4)}}), TArray({FVector2D::Zero()}), TArray({FVector2D(1.5, 0.0)})));
+	bSuccess &= ExecuteTest(TestParams(Subtract, GenerateParamData<FVector>({{FVector(0.5, 1.6, 0.7)}, {FVector(-0.5, 4.2, -0.7)}}), TArray({FVector::Zero()}), TArray({FVector(1.0, -2.6, 1.4)})));
+	bSuccess &= ExecuteTest(TestParams(Subtract, GenerateParamData<FVector4>({{FVector4(0.5, 1.6, 0.7, 0.8)}, {FVector4(-0.5, 4.2, -0.7, 0.3)}}), TArray({FVector4::Zero()}), TArray({FVector4(1.0, -2.6, 1.4, 0.5)})));
+
+	// Point Data
+	FPCGAttributePropertyInputSelector Selector;
+	Selector.SetPointProperty(EPCGPointProperties::Seed);
+	bSuccess &= ExecuteTest(TestParams(Subtract, GeneratePointData(2), TArray<int64>({0}), TArray<int64>({-1})), Selector);
+	Selector.SetPointProperty(EPCGPointProperties::Density);
+	bSuccess &= ExecuteTest(TestParams(Subtract, GeneratePointData(2), TArray({0.0}), TArray({1.0})), Selector);
+	Selector.SetPointProperty(EPCGPointProperties::Position);
+	bSuccess &= ExecuteTest(TestParams(Subtract, GeneratePointData(2), TArray({FVector::Zero()}), TArray({FVector(-1.0, -10.0, -100.0)})), Selector);
+
+	return bSuccess;
+}
+
+IMPLEMENT_CUSTOM_SIMPLE_AUTOMATION_TEST(FPCGMetadataMathsOp_Multiply, FPCGMetadataMathsOpTest, "pcg.tests.Metadata.MathsOp.Multiply", PCGTestsCommon::TestFlags)
+
+bool FPCGMetadataMathsOp_Multiply::RunTest(const FString& Parameters)
+{
+	bool bSuccess = true;
+
+	using EPCGMedadataMathsOperation::Multiply;
+
+	// Param Data
+	bSuccess &= ExecuteTest(TestParams(Multiply, GenerateParamData<int32>({{5}, {1}}), TArray({0}), TArray({5})));
+	bSuccess &= ExecuteTest(TestParams(Multiply, GenerateParamData<long long>({{6ll}, {-1ll}}), TArray({0ll}), TArray({-6ll})));
+	bSuccess &= ExecuteTest(TestParams(Multiply, GenerateParamData<float>({{0.1f}, {0.4f}}), TArray({0.f}), TArray({0.04f})));
+	bSuccess &= ExecuteTest(TestParams(Multiply, GenerateParamData<double>({{0.0}, {-2.4}}), TArray({0.0}), TArray({0.0})));
+	bSuccess &= ExecuteTest(TestParams(Multiply, GenerateParamData<FVector2D>({{FVector2D(0.3, 1.4)}, {FVector2D(-1.2, 1.4)}}), TArray({FVector2D::Zero()}), TArray({FVector2D(-0.36, 1.96)})));
+	bSuccess &= ExecuteTest(TestParams(Multiply, GenerateParamData<FVector>({{FVector(0.5, 1.6, 0.7)}, {FVector(-0.5, 4.2, -0.7)}}), TArray({FVector::Zero()}), TArray({FVector(-0.25, 6.72, -0.49)})));
+	bSuccess &= ExecuteTest(TestParams(Multiply, GenerateParamData<FVector4>({{FVector4(0.5, 1.6, 0.7, 0.8)}, {FVector4(-0.5, 4.2, -0.7, 0.3)}}), TArray({FVector4::Zero()}), TArray({FVector4(-0.25, 6.72, -0.49, 0.24)})));
+
+	// Point Data
+	FPCGAttributePropertyInputSelector Selector;
+	Selector.SetPointProperty(EPCGPointProperties::Seed);
+	bSuccess &= ExecuteTest(TestParams(Multiply, GeneratePointData(2), TArray<int64>({0}), TArray<int64>({2})), Selector);
+	Selector.SetPointProperty(EPCGPointProperties::Density);
+	bSuccess &= ExecuteTest(TestParams(Multiply, GeneratePointData(2), TArray({0.0}), TArray({2.0})), Selector);
+	Selector.SetPointProperty(EPCGPointProperties::Position);
+	bSuccess &= ExecuteTest(TestParams(Multiply, GeneratePointData(2), TArray({FVector::Zero()}), TArray({FVector(2.0, 200.0, 20000.0)})), Selector);
+
+	return bSuccess;
+}
+
+IMPLEMENT_CUSTOM_SIMPLE_AUTOMATION_TEST(FPCGMetadataMathsOp_Divide, FPCGMetadataMathsOpTest, "pcg.tests.Metadata.MathsOp.Divide", PCGTestsCommon::TestFlags)
+
+bool FPCGMetadataMathsOp_Divide::RunTest(const FString& Parameters)
+{
+	bool bSuccess = true;
+
+	using EPCGMedadataMathsOperation::Divide;
+
+	// Param Data
+	bSuccess &= ExecuteTest(TestParams(Divide, GenerateParamData<int32>({{32}, {8}}), TArray({0}), TArray({4})));
+	bSuccess &= ExecuteTest(TestParams(Divide, GenerateParamData<long long>({{33ll}, {3ll}}), TArray({0ll}), TArray({11ll})));
+	bSuccess &= ExecuteTest(TestParams(Divide, GenerateParamData<float>({{0.1f}, {0.4f}}), TArray({0.f}), TArray({0.25f})));
+	bSuccess &= ExecuteTest(TestParams(Divide, GenerateParamData<double>({{0.0}, {-2.4}}), TArray({0.0}), TArray({0.0})));
+	bSuccess &= ExecuteTest(TestParams(Divide, GenerateParamData<FVector2D>({{FVector2D(0.3, 1.4)}, {FVector2D(-1.2, 1.4)}}), TArray({FVector2D::Zero()}), TArray({FVector2D(-0.25, 1.0)})));
+	bSuccess &= ExecuteTest(TestParams(Divide, GenerateParamData<FVector>({{FVector(0.5, 1.6, 0.7)}, {FVector(-0.5, 4.2, -0.7)}}), TArray({FVector::Zero()}), TArray({FVector(-1.0, 0.381, -1.0)})));
+	bSuccess &= ExecuteTest(TestParams(Divide, GenerateParamData<FVector4>({{FVector4(0.5, 1.6, 0.7, 0.8)}, {FVector4(-0.5, 4.2, -0.7, 0.3)}}), TArray({FVector4::Zero()}), TArray({FVector4(-1.0, 0.381, -1.0, 2.666667)})));
+
+	// Point Data
+	FPCGAttributePropertyInputSelector Selector;
+	Selector.SetPointProperty(EPCGPointProperties::Seed);
+	bSuccess &= ExecuteTest(TestParams(Divide, GeneratePointData(2), TArray<int64>({0}), TArray<int64>({0})), Selector);
+	Selector.SetPointProperty(EPCGPointProperties::Density);
+	bSuccess &= ExecuteTest(TestParams(Divide, GeneratePointData(2), TArray({0.0}), TArray({2.0})), Selector);
+	Selector.SetPointProperty(EPCGPointProperties::Position);
+	bSuccess &= ExecuteTest(TestParams(Divide, GeneratePointData(2), TArray({FVector::Zero()}), TArray({FVector(0.5, 0.5, 0.5)})), Selector);
+
+	return bSuccess;
+}
+
+IMPLEMENT_CUSTOM_SIMPLE_AUTOMATION_TEST(FPCGMetadataMathsOp_MultipleInput, FPCGMetadataMathsOpTest, "pcg.tests.Metadata.MathsOp.MultipleInput", PCGTestsCommon::TestFlags)
+
+bool FPCGMetadataMathsOp_MultipleInput::RunTest(const FString& Parameters)
+{
+	bool bSuccess = true;
+
+	using EPCGMedadataMathsOperation::Add;
+
+	// 2 x 2
+	bSuccess &= ExecuteTest(TestParams(Add, GenerateParamData<int32>({{4, 5}, {6, 7}}), TArray({0, 0}), TArray({10, 12})));
+
+	// 3 x 3
+	bSuccess &= ExecuteTest(TestParams(Add, GenerateParamData<int32>({{4, 5, 6}, {7, 8, 9}}), TArray({0, 0, 0}), TArray({11, 13, 15})));
+
+	// 2 x 1
+	bSuccess &= ExecuteTest(TestParams(Add, GenerateParamData<int32>({{4, 5}, {6}}), TArray({0, 0}), TArray({10, 11})));
+
+	// 3 x 1
+	bSuccess &= ExecuteTest(TestParams(Add, GenerateParamData<int32>({{4, 5, 6}, {7}}), TArray({0, 0, 0}), TArray({11, 12, 13})));
+
+	// 1 x 3
+	bSuccess &= ExecuteTest(TestParams(Add, GenerateParamData<int32>({{4}, {5, 6, 7}}), TArray({0, 0, 0}), TArray({9, 10, 11}), /*PinToForward=*/ 1));
+
+	return bSuccess;
+}
+
+// TODO: Round, Sqrt, Abs, Floor, Ceil, Max, Min, Pow, ClampMin, ClampMax, Modulo, Set, Clamp, Lerp
+
+#endif // WITH_EDITOR
\ No newline at end of file
diff --git a/Engine/Plugins/Experimental/PCG/Source/PCG/Private/Tests/PCGTestsCommon.cpp b/Engine/Plugins/Experimental/PCG/Source/PCG/Private/Tests/PCGTestsCommon.cpp
--- a/Engine/Plugins/Experimental/PCG/Source/PCG/Private/Tests/PCGTestsCommon.cpp
+++ b/Engine/Plugins/Experimental/PCG/Source/PCG/Private/Tests/PCGTestsCommon.cpp
@@ -90,6 +90,12 @@ namespace PCGTestsCommon
 		InputData.TaggedData.Empty();
 		OutputData.TaggedData.Empty();
 		Settings = InSettings;
+
+		if (Settings)
+		{
+			InputData.TaggedData.Emplace_GetRef().Data = Settings;
+			InputData.TaggedData.Last().Pin = FName(TEXT("Settings"));
+		}
 	}
 
 	TUniquePtr<FPCGContext> InitializeTestContext(IPCGElement* InElement, const FPCGDataCollection& InputData, UPCGComponent* InSourceComponent, const UPCGNode* InNode)
diff --git a/Engine/Plugins/Experimental/PCG/Source/PCG/Public/Metadata/PCGAttributePropertySelector.h b/Engine/Plugins/Experimental/PCG/Source/PCG/Public/Metadata/PCGAttributePropertySelector.h
--- a/Engine/Plugins/Experimental/PCG/Source/PCG/Public/Metadata/PCGAttributePropertySelector.h
+++ b/Engine/Plugins/Experimental/PCG/Source/PCG/Public/Metadata/PCGAttributePropertySelector.h
@@ -77,16 +77,48 @@ public:
 	template <typename T>
 	static T CreateFromOtherSelector(const FPCGAttributePropertySelector& InOther)
 	{
+		static_assert(std::is_base_of_v<FPCGAttributePropertySelector, T>, "The type must be of base class 'FPCGAttributePropertySelector'");
 		T OutSelector;
 		OutSelector.ImportFromOtherSelector(InOther);
 		return OutSelector;
 	}
 
-	// Convenience static constructors
-	static FPCGAttributePropertySelector CreateAttributeSelector(const FName AttributeName);
-	static FPCGAttributePropertySelector CreatePointPropertySelector(EPCGPointProperties PointProperty);
-	static FPCGAttributePropertySelector CreateExtraPropertySelector(EPCGExtraProperties ExtraProperty);
-	static FPCGAttributePropertySelector CreateSelectorFromString(const FString& String);
+	// Convenience templated static constructors
+	template <typename T = FPCGAttributePropertySelector>
+	static T CreateAttributeSelector(const FName AttributeName)
+	{
+		static_assert(std::is_base_of_v<FPCGAttributePropertySelector, T>, "The type must be of base class 'FPCGAttributePropertySelector'");
+		T Selector;
+		Selector.SetAttributeName(AttributeName);
+		return Selector;
+	}
+
+	template <typename T = FPCGAttributePropertySelector>
+	static T CreatePointPropertySelector(EPCGPointProperties PointProperty)
+	{
+		static_assert(std::is_base_of_v<FPCGAttributePropertySelector, T>, "The type must be of base class 'FPCGAttributePropertySelector'");
+		T Selector;
+		Selector.SetPointProperty(PointProperty);
+		return Selector;
+	}
+
+	template <typename T = FPCGAttributePropertySelector>
+	static T CreateExtraPropertySelector(EPCGExtraProperties ExtraProperty)
+	{
+		static_assert(std::is_base_of_v<FPCGAttributePropertySelector, T>, "The type must be of base class 'FPCGAttributePropertySelector'");
+		T Selector;
+		Selector.SetExtraProperty(ExtraProperty);
+		return Selector;
+	}
+
+	template <typename T = FPCGAttributePropertySelector>
+	static T CreateSelectorFromString(const FString& String)
+	{
+		static_assert(std::is_base_of_v<FPCGAttributePropertySelector, T>, "The type must be of base class 'FPCGAttributePropertySelector'");
+		T Selector;
+		Selector.Update(String);
+		return Selector;
+	}
 
 	void ImportFromOtherSelector(const FPCGAttributePropertySelector& InOther);
 
