Fixed caching issue when nodes have multiple inputs/outputs

diff --git a/Engine/Plugins/Experimental/PCG/Source/PCG/Private/Graph/PCGGraphExecutor.cpp b/Engine/Plugins/Experimental/PCG/Source/PCG/Private/Graph/PCGGraphExecutor.cpp
--- a/Engine/Plugins/Experimental/PCG/Source/PCG/Private/Graph/PCGGraphExecutor.cpp
+++ b/Engine/Plugins/Experimental/PCG/Source/PCG/Private/Graph/PCGGraphExecutor.cpp
@@ -1113,6 +1113,11 @@ void FPCGGraphExecutor::BuildTaskInput(const FPCGGraphTask& Task, FPCGDataCollec
 	// Initialize a Crc onto which each input Crc will be combined.
 	FPCGCrc Crc(Task.Inputs.Num());
 
+	// Random prime numbers to use as placeholders in the CRC computation when there are no defined in/out pins.
+	// Note that they aren't strictly needed, but will make sure we don't introduce issues if we rework this bit of code.
+	constexpr uint32 DefaultHashForNoInputPin = 955333;
+	constexpr uint32 DefaultHashForNoOutputPin = 999983;
+
 	for (const FPCGGraphTaskInput& Input : Task.Inputs)
 	{
 		check(OutputData.Contains(Input.TaskId));
@@ -1151,26 +1156,20 @@ void FPCGGraphExecutor::BuildTaskInput(const FPCGGraphTask& Task, FPCGDataCollec
 			{
 				TaskInput.TaggedData.Append(InputsOnPin);
 			}
-
-			// Write input pin name Crc to uniquely identify inputs per-pin.
-			Crc.Combine(GetTypeHash(Input.OutPin ? Input.OutPin->Properties.Label : FName(TEXT("MissingLabel"))));
 		}
 		else
 		{
 			TaskInput.TaggedData.Append(InputCollection.TaggedData);
 		}
 
+		// Write input pin name (e.g. name of the output pin on the dependency node) Crc to uniquely identify inputs per-pin, or use a placeholder for symmetry.
+		Crc.Combine(Input.InPin ? GetTypeHash(Input.InPin->Properties.Label) : DefaultHashForNoInputPin);
+
 		if (TaskInput.TaggedData.Num() == TaggedDataOffset && InputCollection.bCancelExecutionOnEmpty)
 		{
 			TaskInput.bCancelExecution = true;
 		}
 
-		// This chains the Crc of each input to produce a Crc that covers all of them.
-		if (InputCollection.Crc.IsValid())
-		{
-			Crc.Combine(InputCollection.Crc);
-		}
-
 		// Apply labelling on data; technically, we should ensure that we do this only for pass-through nodes,
 		// Otherwise we could also null out the label on the input...
 		if (Input.OutPin)
@@ -1180,6 +1179,15 @@ void FPCGGraphExecutor::BuildTaskInput(const FPCGGraphTask& Task, FPCGDataCollec
 				TaskInput.TaggedData[TaggedDataIndex].Pin = Input.OutPin->Properties.Label;
 			}
 		}
+
+		// Write output pin name (e.g. input pin node on this node) Crc to uniquely identify inputs per-pin, or use a placeholder for symmetry.
+		Crc.Combine(Input.OutPin ? GetTypeHash(Input.OutPin->Properties.Label) : DefaultHashForNoOutputPin);
+
+		// This chains the Crc of each input to produce a Crc that covers all of them.
+		if (InputCollection.Crc.IsValid())
+		{
+			Crc.Combine(InputCollection.Crc);
+		}
 	}
 
 	// Then combine params if needed
