Extending loop node to support feedback loops, and added a Get Loop Index node.
Added visual cue to distinguish static pins vs. loop pins vs feedback pins.
Added convenient visual update when using manual pins selection.

diff --git a/Engine/Plugins/Experimental/PCG/Source/PCG/Private/Elements/PCGGetLoopIndex.cpp b/Engine/Plugins/Experimental/PCG/Source/PCG/Private/Elements/PCGGetLoopIndex.cpp
new file mode 100644
--- /dev/null
+++ b/Engine/Plugins/Experimental/PCG/Source/PCG/Private/Elements/PCGGetLoopIndex.cpp
@@ -0,0 +1,91 @@
+// Copyright Epic Games, Inc. All Rights Reserved.
+
+#include "Elements/PCGGetLoopIndex.h"
+
+#include "PCGParamData.h"
+#include "PCGPin.h"
+#include "Graph/PCGStackContext.h"
+#include "Metadata/PCGMetadata.h"
+
+#include UE_INLINE_GENERATED_CPP_BY_NAME(PCGGetLoopIndex)
+
+#define LOCTEXT_NAMESPACE "PCGGetLoopIndexElement"
+
+namespace PCGGetLoopIndexConstants
+{
+	static const FName LoopIndexAttributeName = TEXT("LoopIndex");
+}
+
+#if WITH_EDITOR
+FText UPCGGetLoopIndexSettings::GetNodeTooltipText() const
+{
+	return LOCTEXT("GetLoopIndexTooltip", "Returns the index of the loop this subgraph is executing, if any.");
+}
+#endif // WITH_EDITOR
+
+TArray<FPCGPinProperties> UPCGGetLoopIndexSettings::OutputPinProperties() const
+{
+	TArray<FPCGPinProperties> PinProperties;
+	PinProperties.Emplace(PCGPinConstants::DefaultOutputLabel, EPCGDataType::Param);
+
+	return PinProperties;
+}
+
+FPCGElementPtr UPCGGetLoopIndexSettings::CreateElement() const
+{
+	return MakeShared<FPCGGetLoopIndexElement>();
+}
+
+bool FPCGGetLoopIndexElement::ExecuteInternal(FPCGContext* Context) const
+{
+	TRACE_CPUPROFILER_EVENT_SCOPE(FPCGGetLoopIndexElement::Execute);
+
+	check(Context);
+
+	const UPCGGetLoopIndexSettings* Settings = Context->GetInputSettings<UPCGGetLoopIndexSettings>();
+	check(Settings);
+
+	if (!ensure(Context->Stack))
+	{
+		PCGE_LOG(Error, LogOnly, LOCTEXT("ContextHasNoExecutionStack", "The execution context is malformed and has no context."));
+		return true;
+	}
+
+	const FPCGStack* Stack = Context->Stack;
+
+	int LoopIndex = INDEX_NONE;
+
+	const TArray<FPCGStackFrame>& StackFrames = Stack->GetStackFrames();
+	// Last element in the stack should be the current subgraph; Second-to-last will be the loop index, if any.
+	// WARNING: this is sensitive to stack frames changing.
+	// This will get the loop index only if the current subgraph is executed as a loop, this will NOT find the "nearest" loop index.
+	if (StackFrames.Num() >= 2)
+	{
+		LoopIndex = StackFrames.Last(1).LoopIndex;
+	}
+
+	if (LoopIndex == INDEX_NONE)
+	{
+		if (Settings->bWarnIfCalledOutsideOfLoop)
+		{
+			PCGE_LOG(Warning, GraphAndLog, LOCTEXT("NotRunningALoop", "GetLoopIndex node executed outside of a loop. If this is okay, consider turning off this warning."));
+		}
+		
+		return true;
+	}
+
+	UPCGParamData* LoopIndexParamData = NewObject<UPCGParamData>();
+	check(LoopIndexParamData && LoopIndexParamData->Metadata);
+
+	FPCGTaggedData& OutputData = Context->OutputData.TaggedData.Emplace_GetRef();
+	OutputData.Data = LoopIndexParamData;
+
+	FPCGMetadataAttribute<int>* LoopIndexAttribute = LoopIndexParamData->Metadata->CreateAttribute<int>(PCGGetLoopIndexConstants::LoopIndexAttributeName, LoopIndex, /*bAllowInterpolation=*/false, /*bOverrideParent=*/false);
+	check(LoopIndexAttribute);
+
+	LoopIndexAttribute->SetValue(LoopIndexParamData->Metadata->AddEntry(), LoopIndex);
+
+	return true;
+}
+
+#undef LOCTEXT_NAMESPACE
\ No newline at end of file
diff --git a/Engine/Plugins/Experimental/PCG/Source/PCG/Private/Elements/PCGLoopElement.cpp b/Engine/Plugins/Experimental/PCG/Source/PCG/Private/Elements/PCGLoopElement.cpp
--- a/Engine/Plugins/Experimental/PCG/Source/PCG/Private/Elements/PCGLoopElement.cpp
+++ b/Engine/Plugins/Experimental/PCG/Source/PCG/Private/Elements/PCGLoopElement.cpp
@@ -3,12 +3,15 @@
 #include "Elements/PCGLoopElement.h"
 
 #include "PCGComponent.h"
+#include "PCGCustomVersion.h"
 #include "PCGModule.h"
 #include "PCGPin.h"
 #include "PCGSubsystem.h"
 #include "Graph/PCGStackContext.h"
 #include "Helpers/PCGHelpers.h"
 
+#include "Algo/Copy.h"
+
 #include UE_INLINE_GENERATED_CPP_BY_NAME(PCGLoopElement)
 
 #define LOCTEXT_NAMESPACE "PCGLoopElement"
@@ -23,7 +26,137 @@ FText UPCGLoopSettings::GetNodeTooltipText() const
 {
 	return LOCTEXT("NodeTooltip", "Executes the specified Subgraph for each data on the loop pins (or on the first pin if no specific loop pins are provided), keeping the rest constant.");
 }
-#endif
+
+void UPCGLoopSettings::ApplyDeprecation(UPCGNode* InOutNode)
+{
+	if (DataVersion < FPCGCustomVersion::UpdateGraphSettingsLoopPins)
+	{
+		bUseGraphDefaultPinUsage = LoopPins.IsEmpty();
+	}
+
+	Super::ApplyDeprecation(InOutNode);
+}
+
+bool UPCGLoopSettings::GetPinExtraIcon(const UPCGPin* InPin, FName& OutExtraIcon, FText& OutTooltip) const
+{
+	if (bUseGraphDefaultPinUsage)
+	{
+		// Direct base class (Subgraph) removes these for readability, but here we explicitly want these
+		return UPCGSettings::GetPinExtraIcon(InPin, OutExtraIcon, OutTooltip);
+	}
+	else if(InPin)
+	{
+		// In this case, we'll check if the given pin matches our list and get the extra from there
+		TArray<FName> LoopPinNames;
+		TArray<FName> FeedbackPinNames;
+		GetLoopPinNames(nullptr, LoopPinNames, FeedbackPinNames, /*bQuiet=*/true);
+
+		FPCGPinProperties PinProperties = InPin->Properties;
+		PinProperties.Usage = EPCGPinUsage::Normal;
+
+		if (LoopPinNames.Contains(PinProperties.Label))
+		{
+			PinProperties.Usage = EPCGPinUsage::Loop;
+		}
+		else if (FeedbackPinNames.Contains(PinProperties.Label))
+		{
+			PinProperties.Usage = EPCGPinUsage::Feedback;
+		}
+
+		return PCGPinPropertiesHelpers::GetDefaultPinExtraIcon(PinProperties, OutExtraIcon, OutTooltip);
+	}
+
+	return false;
+}
+
+void UPCGLoopSettings::PostEditChangeProperty(FPropertyChangedEvent& PropertyChangedEvent)
+{
+	Super::PostEditChangeProperty(PropertyChangedEvent);
+
+	// Trigger a cosmetic change because all properties on this class might trigger a visual change on the pins.
+	OnSettingsChangedDelegate.Broadcast(this, EPCGChangeType::Cosmetic);
+}
+#endif // WITH_EDITOR
+
+void UPCGLoopSettings::GetLoopPinNames(FPCGContext* Context, TArray<FName>& LoopPinNames, TArray<FName>& FeedbackPinNames, bool bQuiet) const
+{
+	if (bUseGraphDefaultPinUsage)
+	{
+		const UPCGGraph* Subgraph = GetSubgraph();
+
+		if (!Subgraph)
+		{
+			return;
+		}
+
+		const UPCGNode* SubgraphInputNode = Subgraph->GetInputNode();
+		check(SubgraphInputNode);
+
+		TArray<FPCGPinProperties> InputPins = SubgraphInputNode->InputPinProperties();
+		for (const FPCGPinProperties& InputPin : InputPins)
+		{
+			if (InputPin.Usage == EPCGPinUsage::Loop)
+			{
+				LoopPinNames.Add(InputPin.Label);
+			}
+			else if (InputPin.Usage == EPCGPinUsage::Feedback)
+			{
+				// Ignore feedback pins that have no connection.
+				if (SubgraphInputNode->GetOutputPin(InputPin.Label)->EdgeCount() > 0)
+				{
+					FeedbackPinNames.Add(InputPin.Label);
+				}
+			}
+		}
+	}
+	else
+	{
+		TArray<FString> PinsFromSettings = PCGHelpers::GetStringArrayFromCommaSeparatedString(LoopPins);
+		for (const FString& PinLabel : PinsFromSettings)
+		{
+			LoopPinNames.Emplace(PinLabel);
+		}
+
+		PinsFromSettings = PCGHelpers::GetStringArrayFromCommaSeparatedString(FeedbackPins);
+		for (const FString& PinLabel : PinsFromSettings)
+		{
+			if (LoopPinNames.Contains(PinLabel))
+			{
+				if (!bQuiet && Context)
+				{
+					PCGE_LOG_C(Warning, GraphAndLog, Context, FText::Format(LOCTEXT("SameLabelAppearsTwice", "Label '{0}' appears in both the loop pins and the feedback pins"), FText::FromString(PinLabel)));
+				}
+			}
+			else
+			{
+				FeedbackPinNames.Emplace(PinLabel);
+			}
+		}
+	}
+
+	// If no named pins have been provided, default to the first pin
+	if (LoopPinNames.IsEmpty())
+	{
+		const UPCGSubgraphNode* LoopNode = nullptr;
+
+		if (Context)
+		{
+			LoopNode = Cast<const UPCGSubgraphNode>(Context->Node);
+		}
+		else
+		{
+			LoopNode = Cast<const UPCGSubgraphNode>(GetOuter());
+		}
+
+		if (LoopNode)
+		{
+			if (const UPCGPin* LoopPin = LoopNode->GetFirstConnectedInputPin())
+			{
+				LoopPinNames.Add(LoopPin->Properties.Label);
+			}
+		}
+	}
+}
 
 FPCGElementPtr UPCGLoopSettings::CreateElement() const
 {
@@ -77,8 +210,10 @@ bool FPCGLoopElement::ExecuteInternal(FPCGContext* InContext) const
 		}
 
 		TArray<FPCGDataCollection> LoopDataCollection;
+		TArray<FName> FeedbackPinNames;
+		FPCGDataCollection FeedbackDataCollection;
 		FPCGDataCollection FixedInputDataCollection;
-		PrepareLoopDataCollections(Context, Settings, LoopDataCollection, FixedInputDataCollection);
+		PrepareLoopDataCollections(Context, Settings, LoopDataCollection, FeedbackDataCollection, FeedbackPinNames, FixedInputDataCollection);
 
 		// Early out if there are no data on the loop pin
 		if (LoopDataCollection.IsEmpty() || LoopDataCollection[0].TaggedData.IsEmpty())
@@ -105,6 +240,9 @@ bool FPCGLoopElement::ExecuteInternal(FPCGContext* InContext) const
 		FPCGDataCollection FixedDataCollection;
 		PrepareSubgraphData(Settings, Context, FixedInputDataCollection, FixedDataCollection);
 
+		FPCGTaskId PreviousTaskId = InvalidPCGTaskId;
+		const bool bIsFeedbackLoop = !(FeedbackPinNames.IsEmpty());
+
 		// Dispatch the subgraph for each loop entries we have.
 		// Implementation note: even if the execution gets cancelled, these tasks would get cancelled because they are associated to the current source component
 		for (int EntryIndex = 0; EntryIndex < LoopDataCollection[0].TaggedData.Num(); ++EntryIndex)
@@ -116,6 +254,12 @@ bool FPCGLoopElement::ExecuteInternal(FPCGContext* InContext) const
 				InputDataCollection.TaggedData.Insert(LoopDataCollection[LoopCollectionIndex].TaggedData[EntryIndex], LoopCollectionIndex);
 			}
 
+			// If it's the first valid iteration, pass in the iteration data, otherwise it will be gathered by the loop input forwarding element
+			if (PreviousTaskId == InvalidPCGTaskId && bIsFeedbackLoop)
+			{
+				InputDataCollection.TaggedData.Append(FeedbackDataCollection.TaggedData);
+			}
+
 			// Prepare the invocation stack - which is the stack up to this node, and then this node, then a loop index
 			FPCGStack InvocationStack = ensure(Context->Stack) ? *Context->Stack : FPCGStack();
 
@@ -128,23 +272,46 @@ bool FPCGLoopElement::ExecuteInternal(FPCGContext* InContext) const
 			Subgraph->OnGraphDynamicallyExecutedDelegate.Broadcast(Subgraph, Context->SourceComponent, InvocationStack);
 #endif
 
-			FPCGTaskId SubgraphTaskId = Subsystem->ScheduleGraph(Subgraph, Context->SourceComponent.Get(), PreGraphElement, MakeShared<FPCGInputForwardingElement>(InputDataCollection), {}, &InvocationStack);
+			TArray<FPCGTaskId> Dependencies = {};
+			if (PreviousTaskId != InvalidPCGTaskId && bIsFeedbackLoop)
+			{
+				Dependencies.Add(PreviousTaskId);
+			}
+
+			FPCGElementPtr InputElement = MakeShared<FPCGLoopInputForwardingElement>(InputDataCollection, PreviousTaskId, FeedbackPinNames);
+			FPCGTaskId SubgraphTaskId = Subsystem->ScheduleGraph(Subgraph, Context->SourceComponent.Get(), PreGraphElement, InputElement, Dependencies, &InvocationStack);
+
+			if (SubgraphTaskId != InvalidPCGTaskId)
+			{
+				if (bIsFeedbackLoop)
+				{
+					PreviousTaskId = SubgraphTaskId;
+				}
 
-			Context->SubgraphTaskIds.Add(SubgraphTaskId);
+				Context->SubgraphTaskIds.Add(SubgraphTaskId);
+			}
 		}
 
-		Context->bScheduledSubgraph = true;
-		Context->bIsPaused = true;
+		// Add a trivial task after the output tasks that wakes up this task
+		if (!Context->SubgraphTaskIds.IsEmpty())
+		{
+			Context->bScheduledSubgraph = true;
+			Context->bIsPaused = true;
+
+			Subsystem->ScheduleGeneric([Context]()
+			{
+				// Wake up the current task
+				Context->bIsPaused = false;
+				return true;
+			}, Context->SourceComponent.Get(), Context->SubgraphTaskIds);
 
-		// add a trivial task after the output tasks that wakes up this task
-		Subsystem->ScheduleGeneric([Context]()
+			return false;
+		}
+		else
 		{
-			// Wake up the current task
-			Context->bIsPaused = false;
+			// Nothing to do
 			return true;
-		}, Context->SourceComponent.Get(), Context->SubgraphTaskIds);
-
-		return false;
+		}
 	}
 	else if (Context->bIsPaused)
 	{
@@ -158,13 +325,28 @@ bool FPCGLoopElement::ExecuteInternal(FPCGContext* InContext) const
 		UPCGSubsystem* Subsystem = Context->SourceComponent->GetSubsystem();
 		if (Subsystem)
 		{
-			for (FPCGTaskId SubgraphTaskId : Context->SubgraphTaskIds)
+			// If this was running iterations, we need to build a list of pins we'll ignore for non-terminal tasks
+			TArray<FName> LoopPinNames;
+			TArray<FName> FeedbackPinNames;
+			Settings->GetLoopPinNames(Context, LoopPinNames, FeedbackPinNames, /*bQuiet=*/true);
+
+			for(int SubtaskIndex = 0; SubtaskIndex < Context->SubgraphTaskIds.Num(); ++SubtaskIndex)
 			{
+				FPCGTaskId SubgraphTaskId = Context->SubgraphTaskIds[SubtaskIndex];
+				const bool bIsLastTask = (SubtaskIndex == Context->SubgraphTaskIds.Num() - 1);
+
 				FPCGDataCollection SubgraphOutput;
 				// While this should be always return true, if a scheduled task was cancelled, this can still happen because the dependency has been removed but this code isn't aware of this
 				if (Subsystem->GetOutputData(SubgraphTaskId, SubgraphOutput))
 				{
-					Context->OutputData.TaggedData.Append(SubgraphOutput.TaggedData);
+					if (FeedbackPinNames.IsEmpty() || bIsLastTask)
+					{
+						Context->OutputData.TaggedData.Append(SubgraphOutput.TaggedData);
+					}
+					else
+					{
+						Algo::CopyIf(SubgraphOutput.TaggedData, Context->OutputData.TaggedData, [&FeedbackPinNames](const FPCGTaggedData& InTaggedData) { return !FeedbackPinNames.Contains(InTaggedData.Pin); });
+					}
 				}
 			}
 		}
@@ -178,27 +360,11 @@ bool FPCGLoopElement::ExecuteInternal(FPCGContext* InContext) const
 	}
 }
 
-void FPCGLoopElement::PrepareLoopDataCollections(FPCGContext* Context, const UPCGLoopSettings* Settings, TArray<FPCGDataCollection>& LoopDataCollection, FPCGDataCollection& FixedInputDataCollection) const
+void FPCGLoopElement::PrepareLoopDataCollections(FPCGContext* Context, const UPCGLoopSettings* Settings, TArray<FPCGDataCollection>& LoopDataCollection, FPCGDataCollection& FeedbackDataCollection, TArray<FName>& FeedbackPinNames, FPCGDataCollection& FixedInputDataCollection) const
 {
 	TArray<FName> LoopPinNames;
-	TArray<FString> PinsFromSettings = PCGHelpers::GetStringArrayFromCommaSeparatedString(Settings->LoopPins);
-
-	for (const FString& PinLabel : PinsFromSettings)
-	{
-		LoopPinNames.Emplace(PinLabel);
-	}
-	
-	// If no named pins have been provided, default to the first pin
-	if (LoopPinNames.IsEmpty())
-	{
-		const UPCGSubgraphNode* LoopNode = Cast<const UPCGSubgraphNode>(Context->Node);
-		check(LoopNode);
-
-		if (const UPCGPin* LoopPin = LoopNode->GetFirstConnectedInputPin())
-		{
-			LoopPinNames.Add(LoopPin->Properties.Label);
-		}
-	}
+	check(Settings);
+	Settings->GetLoopPinNames(Context, LoopPinNames, FeedbackPinNames, /*bQuiet=*/false);
 
 	LoopDataCollection.Reset();
 	LoopDataCollection.SetNum(LoopPinNames.Num());
@@ -206,10 +372,16 @@ void FPCGLoopElement::PrepareLoopDataCollections(FPCGContext* Context, const UPC
 	for (const FPCGTaggedData& TaggedData : Context->InputData.TaggedData)
 	{
 		int LoopPinIndex = LoopPinNames.IndexOfByKey(TaggedData.Pin);
+		bool bIsFeedbackData = FeedbackPinNames.Contains(TaggedData.Pin);
+
 		if (LoopPinIndex != INDEX_NONE)
 		{
 			LoopDataCollection[LoopPinIndex].TaggedData.Add(TaggedData);
 		}
+		else if (bIsFeedbackData)
+		{
+			FeedbackDataCollection.TaggedData.Add(TaggedData);
+		}
 		else
 		{
 			FixedInputDataCollection.TaggedData.Add(TaggedData);
@@ -217,4 +389,35 @@ void FPCGLoopElement::PrepareLoopDataCollections(FPCGContext* Context, const UPC
 	}
 }
 
+FPCGLoopInputForwardingElement::FPCGLoopInputForwardingElement(const FPCGDataCollection& StaticInputToForward, FPCGTaskId InPreviousIterationTaskId, const TArray<FName>& InFeedbackPinNames)
+	: FPCGInputForwardingElement(StaticInputToForward)
+	, PreviousIterationTaskId(InPreviousIterationTaskId)
+	, FeedbackPinNames(InFeedbackPinNames)
+{
+}
+
+bool FPCGLoopInputForwardingElement::ExecuteInternal(FPCGContext* Context) const
+{
+	// Performs forwarding of static data
+	while (!FPCGInputForwardingElement::ExecuteInternal(Context))
+	{
+	}
+
+	// Retrieve data from previous task if needed
+	if (PreviousIterationTaskId != InvalidPCGTaskId && !FeedbackPinNames.IsEmpty())
+	{
+		check(Context && Context->SourceComponent.IsValid());
+		if (UPCGSubsystem* Subsystem = Context->SourceComponent->GetSubsystem())
+		{
+			FPCGDataCollection PreviousTaskOutput;
+			if (Subsystem->GetOutputData(PreviousIterationTaskId, PreviousTaskOutput))
+			{
+				Algo::CopyIf(PreviousTaskOutput.TaggedData, Context->OutputData.TaggedData, [this](const FPCGTaggedData& InTaggedData) { return FeedbackPinNames.Contains(InTaggedData.Pin); });
+			}
+		}
+	}
+
+	return true;
+}
+
 #undef LOCTEXT_NAMESPACE
\ No newline at end of file
diff --git a/Engine/Plugins/Experimental/PCG/Source/PCG/Private/PCGPin.cpp b/Engine/Plugins/Experimental/PCG/Source/PCG/Private/PCGPin.cpp
--- a/Engine/Plugins/Experimental/PCG/Source/PCG/Private/PCGPin.cpp
+++ b/Engine/Plugins/Experimental/PCG/Source/PCG/Private/PCGPin.cpp
@@ -2,6 +2,7 @@
 
 #include "PCGPin.h"
 
+#include "PCGCommon.h"
 #include "PCGEdge.h"
 #include "PCGNode.h"
 #include "PCGSettings.h"
@@ -32,6 +33,32 @@ namespace PCGPin
 	}
 }
 
+#if WITH_EDITOR
+namespace PCGPinPropertiesHelpers
+{
+	bool GetDefaultPinExtraIcon(const UPCGPin* InPin, FName& OutExtraIcon, FText& OutTooltip)
+	{
+		return InPin && PCGPinPropertiesHelpers::GetDefaultPinExtraIcon(InPin->Properties, OutExtraIcon, OutTooltip);
+	}
+
+	bool GetDefaultPinExtraIcon(const FPCGPinProperties& InPinProperties, FName& OutExtraIcon, FText& OutTooltip)
+	{
+		if (InPinProperties.Usage == EPCGPinUsage::Loop)
+		{
+			OutExtraIcon = PCGPinConstants::Icons::LoopPinIcon;
+			return true;
+		}
+		else if (InPinProperties.Usage == EPCGPinUsage::Feedback)
+		{
+			OutExtraIcon = PCGPinConstants::Icons::FeedbackPinIcon;
+			return true;
+		}
+
+		return false;
+	}
+}
+#endif // WITH_EDITOR
+
 FPCGPinProperties::FPCGPinProperties(const FName& InLabel, EPCGDataType InAllowedTypes, bool bInAllowMultipleConnections, bool bInAllowMultipleData, const FText& InTooltip)
 	: Label(InLabel), AllowedTypes(InAllowedTypes), bAllowMultipleData(bInAllowMultipleData)
 #if WITH_EDITORONLY_DATA
@@ -60,6 +87,7 @@ bool FPCGPinProperties::operator==(const FPCGPinProperties& Other) const
 		AllowedTypes == Other.AllowedTypes &&
 		bAllowMultipleConnections == Other.bAllowMultipleConnections &&
 		bAllowMultipleData == Other.bAllowMultipleData &&
+		Usage == Other.Usage &&
 		bAdvancedPin == Other.bAdvancedPin;
 }
 
diff --git a/Engine/Plugins/Experimental/PCG/Source/PCG/Private/PCGSettings.cpp b/Engine/Plugins/Experimental/PCG/Source/PCG/Private/PCGSettings.cpp
--- a/Engine/Plugins/Experimental/PCG/Source/PCG/Private/PCGSettings.cpp
+++ b/Engine/Plugins/Experimental/PCG/Source/PCG/Private/PCGSettings.cpp
@@ -222,6 +222,11 @@ void UPCGSettings::ApplyDeprecation(UPCGNode* InOutNode)
 	DataVersion = FPCGCustomVersion::LatestVersion;
 }
 
+bool UPCGSettings::GetPinExtraIcon(const UPCGPin* InPin, FName& OutExtraIcon, FText& OutTooltip) const
+{
+	return PCGPinPropertiesHelpers::GetDefaultPinExtraIcon(InPin, OutExtraIcon, OutTooltip);
+}
+
 void UPCGSettings::PostEditUndo()
 {
 	// CachedOverridableParams was reset to previous value
diff --git a/Engine/Plugins/Experimental/PCG/Source/PCG/Public/Elements/PCGGetLoopIndex.h b/Engine/Plugins/Experimental/PCG/Source/PCG/Public/Elements/PCGGetLoopIndex.h
new file mode 100644
--- /dev/null
+++ b/Engine/Plugins/Experimental/PCG/Source/PCG/Public/Elements/PCGGetLoopIndex.h
@@ -0,0 +1,46 @@
+// Copyright Epic Games, Inc. All Rights Reserved.
+
+#pragma once
+
+#include "PCGContext.h"
+#include "PCGSettings.h"
+
+#include "PCGGetLoopIndex.generated.h"
+
+/** Returns the current loop iteration index of the "nearest" subgraph in the execution stack */
+UCLASS(BlueprintType, ClassGroup=(Procedural))
+class PCG_API UPCGGetLoopIndexSettings : public UPCGSettings
+{
+	GENERATED_BODY()
+
+public:
+	//~Begin UPCGSettings interface implementation
+#if WITH_EDITOR
+	virtual FName GetDefaultNodeName() const override { return FName(TEXT("GetLoopIndex")); }
+	virtual FText GetDefaultNodeTitle() const override { return NSLOCTEXT("PCGGetLoopIndexElement", "NodeTitle", "Get Loop Index"); }
+	virtual FText GetNodeTooltipText() const override;
+	virtual EPCGSettingsType GetType() const override { return EPCGSettingsType::Param; }
+#endif // WITH_EDITOR
+
+protected:
+	virtual TArray<FPCGPinProperties> InputPinProperties() const override { return TArray<FPCGPinProperties>(); }
+	virtual TArray<FPCGPinProperties> OutputPinProperties() const override;
+
+	virtual FPCGElementPtr CreateElement() const override;
+	// ~End UPCGSettings interface implementation
+
+public:
+	/** Controls whether this node will create a warning when not called from within a loop. */
+	UPROPERTY(BlueprintReadWrite, EditAnywhere, Category = Settings)
+	bool bWarnIfCalledOutsideOfLoop = true;
+};
+
+class FPCGGetLoopIndexElement : public IPCGElement
+{
+public:
+	// A loop index node is never cacheable because it needs to execute for every iteration
+	virtual bool IsCacheable(const UPCGSettings* InSettings) const { return false; }
+
+protected:
+	virtual bool ExecuteInternal(FPCGContext* Context) const override;
+};
\ No newline at end of file
diff --git a/Engine/Plugins/Experimental/PCG/Source/PCG/Public/Elements/PCGLoopElement.h b/Engine/Plugins/Experimental/PCG/Source/PCG/Public/Elements/PCGLoopElement.h
--- a/Engine/Plugins/Experimental/PCG/Source/PCG/Public/Elements/PCGLoopElement.h
+++ b/Engine/Plugins/Experimental/PCG/Source/PCG/Public/Elements/PCGLoopElement.h
@@ -12,11 +12,19 @@ class PCG_API UPCGLoopSettings : public UPCGSubgraphSettings
 	GENERATED_BODY()
 
 public:
+	//~Begin UObject interface implementation
+#if WITH_EDITOR
+	virtual void PostEditChangeProperty(FPropertyChangedEvent& PropertyChangedEvent) override;
+#endif
+	//~End UObject interface implementation
+
 	//~Begin UPCGSettings interface implementation
 #if WITH_EDITOR
 	virtual FName GetDefaultNodeName() const override { return FName(TEXT("Loop")); }
 	virtual FText GetDefaultNodeTitle() const override;
 	virtual FText GetNodeTooltipText() const override;
+	virtual void ApplyDeprecation(UPCGNode* InOutNode) override;
+	virtual bool GetPinExtraIcon(const UPCGPin* InPin, FName& OutExtraIcon, FText& OutTooltip) const override;
 #endif
 
 	virtual FName AdditionalTaskName() const override;
@@ -30,17 +38,44 @@ public:
 	virtual bool IsDynamicGraph() const override { return true; }
 	//~End UPCGBaseSubgraphSettings interface
 
+	/** This method is used to retrieve the loop & feedback pin names either from the graph (if using the default pin usage) or from the comma-separated list otherwise. */
+	void GetLoopPinNames(FPCGContext* Context, TArray<FName>& LoopPinNames, TArray<FName>& FeedbackPinNames, bool bQuiet) const;
+
+	/** Controls whether the pin usage (normal, loop, feedback) will be taken from the subgraph to execute or from the manually provided list. */
+	UPROPERTY(BlueprintReadWrite, EditAnywhere, Category = Data)
+	bool bUseGraphDefaultPinUsage = true;
+
 	/** 
-	* Comma-separated list of pin names on which we will loop by-element in a step-wise fashion; if more than one is provided, it is expected that they all have the same number of data. 
+	* Comma-separated list of pin names on which we will loop by-element in a step-wise fashion; if more than one is provided, it is expected that they all have the same number of data.
 	* If none are provided, the first connected pin will taken as the pin to loop on. 
 	*/
-	UPROPERTY(BlueprintReadWrite, EditAnywhere, Category = Data)
+	UPROPERTY(BlueprintReadWrite, EditAnywhere, Category = Data, meta = (EditCondition = "!bUseGraphDefaultPinUsage"))
 	FString LoopPins;
+
+	/** 
+	* Comma-separated list of pin names that will act as feedback pins, namely that in a given iteration it will receive the data from the output pin of the same name of the previous loop iteration.
+	* These pins can have initial data, in which case only the first iteration will get this data.
+	*/
+	UPROPERTY(BlueprintReadWrite, EditAnywhere, Category = Data, meta = (EditCondition = "!bUseGraphDefaultPinUsage"))
+	FString FeedbackPins;
 };
 
 class PCG_API FPCGLoopElement : public FPCGSubgraphElement
 {
 protected:
 	virtual bool ExecuteInternal(FPCGContext* Context) const override;
-	void PrepareLoopDataCollections(FPCGContext* Context, const UPCGLoopSettings* Settings, TArray<FPCGDataCollection>& LoopDataCollection, FPCGDataCollection& FixedInputDataCollection) const;
+	void PrepareLoopDataCollections(FPCGContext* Context, const UPCGLoopSettings* Settings, TArray<FPCGDataCollection>& LoopDataCollection, FPCGDataCollection& FeedbackDataCollection, TArray<FName>& FeedbackPinNames, FPCGDataCollection& FixedInputDataCollection) const;
 };
+
+/** This element does the same static input forwarding mechanism as the FPCGInputForwardingElement, but it also will go fetch the previous iteration data, if any. */
+class PCG_API FPCGLoopInputForwardingElement : public FPCGInputForwardingElement
+{
+public:
+	FPCGLoopInputForwardingElement(const FPCGDataCollection& StaticInputToForward, FPCGTaskId InPreviousIterationTaskId, const TArray<FName>& InFeedbackPinNames);
+
+protected:
+	virtual bool ExecuteInternal(FPCGContext* Context) const override;
+
+	FPCGTaskId PreviousIterationTaskId = InvalidPCGTaskId;
+	TArray<FName> FeedbackPinNames;
+};
\ No newline at end of file
diff --git a/Engine/Plugins/Experimental/PCG/Source/PCG/Public/PCGCommon.h b/Engine/Plugins/Experimental/PCG/Source/PCG/Public/PCGCommon.h
--- a/Engine/Plugins/Experimental/PCG/Source/PCG/Public/PCGCommon.h
+++ b/Engine/Plugins/Experimental/PCG/Source/PCG/Public/PCGCommon.h
@@ -97,6 +97,12 @@ namespace Private
 {
 	const FName OldDefaultParamsLabel = TEXT("Params");
 }
+
+namespace Icons
+{
+	const FName LoopPinIcon = TEXT("GraphEditor.Macro.Loop_16x");
+	const FName FeedbackPinIcon = TEXT("GraphEditor.GetSequenceBinding");
+}
 }
 
 // Metadata used by PCG
diff --git a/Engine/Plugins/Experimental/PCG/Source/PCG/Public/PCGCustomVersion.h b/Engine/Plugins/Experimental/PCG/Source/PCG/Public/PCGCustomVersion.h
--- a/Engine/Plugins/Experimental/PCG/Source/PCG/Public/PCGCustomVersion.h
+++ b/Engine/Plugins/Experimental/PCG/Source/PCG/Public/PCGCustomVersion.h
@@ -59,6 +59,9 @@ struct PCG_API FPCGCustomVersion
 		// Removed by-default pins on input node. Note, this breaks cooked binary compatibility
 		UpdateInputOutputNodesDefaults = 16,
 
+		// Introduced the concept of pin usage and graph defaults around loops, which changed the default behavior otherwise
+		UpdateGraphSettingsLoopPins = 17,
+
 		// -----<new versions can be added above this line>-------------------------------------------------
 		VersionPlusOne,
 		LatestVersion = VersionPlusOne - 1
diff --git a/Engine/Plugins/Experimental/PCG/Source/PCG/Public/PCGPin.h b/Engine/Plugins/Experimental/PCG/Source/PCG/Public/PCGPin.h
--- a/Engine/Plugins/Experimental/PCG/Source/PCG/Public/PCGPin.h
+++ b/Engine/Plugins/Experimental/PCG/Source/PCG/Public/PCGPin.h
@@ -11,6 +11,15 @@
 class UPCGNode;
 class UPCGEdge;
 
+UENUM(BlueprintType)
+enum class EPCGPinUsage : uint8
+{
+	Normal = 0, // Normal usage pin, will pass all data as is.
+	Loop, // When used in a loop subgraph node, will separate each data from that pin into separate subgraph executions.
+	Feedback, // When used in a loop subgraph node, will pass data on the feedback pins to the next iteration only if the data is passed from a previous iteration (or the original subgraph call).
+	DependencyOnly UMETA(Hidden)
+};
+
 USTRUCT(BlueprintType, meta=(HasNativeBreak="/Script/PCG.PCGBlueprintPinHelpers.BreakPinProperty", HasNativeMake="/Script/PCG.PCGBlueprintPinHelpers.MakePinProperty"))
 struct PCG_API FPCGPinProperties
 {
@@ -24,6 +33,9 @@ struct PCG_API FPCGPinProperties
 	UPROPERTY(BlueprintReadWrite, EditAnywhere, Category = Settings)
 	FName Label = NAME_None;
 
+	UPROPERTY(BlueprintReadWrite, EditAnywhere, Category = Settings)
+	EPCGPinUsage Usage = EPCGPinUsage::Normal;
+
 	UPROPERTY(BlueprintReadWrite, EditAnywhere, Category = Settings)
 	EPCGDataType AllowedTypes = EPCGDataType::Any;
 
@@ -180,6 +192,14 @@ public:
 	EPCGTypeConversion GetRequiredTypeConversion(const UPCGPin* InOtherPin) const;
 };
 
+#if WITH_EDITOR
+namespace PCGPinPropertiesHelpers
+{
+	bool GetDefaultPinExtraIcon(const UPCGPin* InPin, FName& OutExtraIcon, FText& OutTooltip);
+	bool GetDefaultPinExtraIcon(const FPCGPinProperties& InPinProperties, FName& OutExtraIcon, FText& OutTooltip);
+}
+#endif // WITH_EDITOR
+
 #if UE_ENABLE_INCLUDE_ORDER_DEPRECATED_IN_5_2
 #include "CoreMinimal.h"
 #endif
diff --git a/Engine/Plugins/Experimental/PCG/Source/PCG/Public/PCGSettings.h b/Engine/Plugins/Experimental/PCG/Source/PCG/Public/PCGSettings.h
--- a/Engine/Plugins/Experimental/PCG/Source/PCG/Public/PCGSettings.h
+++ b/Engine/Plugins/Experimental/PCG/Source/PCG/Public/PCGSettings.h
@@ -273,7 +273,7 @@ public:
 	virtual bool GetPinLabelStyle(const UPCGPin* InPin, FName& OutLabelStyle) const { return false; }
 
 	/** Can override to add a custom icon next to the pin label (and an optional tooltip). Return false if no override is available. */
-	virtual bool GetPinExtraIcon(const UPCGPin* InPin, FName& OutExtraIcon, FText& OutTooltip) const { return false; }
+	virtual bool GetPinExtraIcon(const UPCGPin* InPin, FName& OutExtraIcon, FText& OutTooltip) const;
 
 	/** Derived classes must implement this to communicate dependencies on external actors */
 	virtual void GetTrackedActorKeys(FPCGActorSelectionKeyToSettingsMap& OutKeysToSettings, TArray<TObjectPtr<const UPCGGraph>>& OutVisitedGraphs) const {}
diff --git a/Engine/Plugins/Experimental/PCG/Source/PCG/Public/PCGSubgraph.h b/Engine/Plugins/Experimental/PCG/Source/PCG/Public/PCGSubgraph.h
--- a/Engine/Plugins/Experimental/PCG/Source/PCG/Public/PCGSubgraph.h
+++ b/Engine/Plugins/Experimental/PCG/Source/PCG/Public/PCGSubgraph.h
@@ -95,6 +95,7 @@ public:
 	virtual FText GetDefaultNodeTitle() const override { return NSLOCTEXT("PCGSubgraphSettings", "NodeTitle", "Subgraph"); }
 	virtual EPCGSettingsType GetType() const override { return EPCGSettingsType::Subgraph; }
 	virtual UObject* GetJumpTargetForDoubleClick() const override;
+	virtual bool GetPinExtraIcon(const UPCGPin* InPin, FName& OutExtraIcon, FText& OutTooltip) const override { return false; }
 #endif
 
 	virtual FName AdditionalTaskName() const override;
diff --git a/Engine/Plugins/Experimental/PCG/Source/PCGEditor/Private/SPCGEditorGraphNode.cpp b/Engine/Plugins/Experimental/PCG/Source/PCGEditor/Private/SPCGEditorGraphNode.cpp
--- a/Engine/Plugins/Experimental/PCG/Source/PCGEditor/Private/SPCGEditorGraphNode.cpp
+++ b/Engine/Plugins/Experimental/PCG/Source/PCGEditor/Private/SPCGEditorGraphNode.cpp
@@ -116,7 +116,8 @@ void SPCGEditorGraphNodePin::Construct(const FArguments& InArgs, UEdGraphPin* In
 	if (GetExtraIcon(ExtraPinIcon, ExtraPinIconTooltip))
 	{
 		ExtraPinIconWidget = SNew(SImage)
-			.Image(FAppStyle::GetBrush(ExtraPinIcon));
+			.Image(FAppStyle::GetBrush(ExtraPinIcon))
+			.ColorAndOpacity(this, &SPCGEditorGraphNodePin::GetPinTextColor);
 
 		if (!ExtraPinIconTooltip.IsEmpty())
 		{
@@ -132,23 +133,34 @@ void SPCGEditorGraphNodePin::Construct(const FArguments& InArgs, UEdGraphPin* In
 				PinStatusIndicator
 			];
 
+		LabelAndValue->AddSlot()
+			.VAlign(VAlign_Center)
+			[
+				LabelWidget
+			];
+
 		if (ExtraPinIconWidget.IsValid())
 		{
 			LabelAndValue->AddSlot()
+				.Padding(5, 0, 0, 0)
 				.VAlign(VAlign_Center)
 				[
 					ExtraPinIconWidget.ToSharedRef()
 				];
 		}
-
-		LabelAndValue->AddSlot()
-			.VAlign(VAlign_Center)
-			[
-				LabelWidget
-			];
 	}
 	else
 	{
+		if (ExtraPinIconWidget.IsValid())
+		{
+			LabelAndValue->AddSlot()
+				.Padding(0, 0, 5, 0)
+				.VAlign(VAlign_Center)
+				[
+					ExtraPinIconWidget.ToSharedRef()
+				];
+		}
+
 		LabelAndValue->AddSlot()
 			.VAlign(VAlign_Center)
 			[
@@ -177,15 +189,6 @@ void SPCGEditorGraphNodePin::Construct(const FArguments& InArgs, UEdGraphPin* In
 			}
 		}
 
-		if (ExtraPinIconWidget.IsValid())
-		{
-			LabelAndValue->AddSlot()
-				.VAlign(VAlign_Center)
-				[
-					ExtraPinIconWidget.ToSharedRef()
-				];
-		}
-
 		LabelAndValue->AddSlot()
 			.VAlign(VAlign_Center)
 			[
