[PCG] Quickfix for MetadataOps unintended passthrough of all inputs

#rnx

diff --git a/Engine/Plugins/Experimental/PCG/Source/PCG/Private/Elements/Metadata/PCGMetadataOpElementBase.cpp b/Engine/Plugins/Experimental/PCG/Source/PCG/Private/Elements/Metadata/PCGMetadataOpElementBase.cpp
--- a/Engine/Plugins/Experimental/PCG/Source/PCG/Private/Elements/Metadata/PCGMetadataOpElementBase.cpp
+++ b/Engine/Plugins/Experimental/PCG/Source/PCG/Private/Elements/Metadata/PCGMetadataOpElementBase.cpp
@@ -272,30 +272,6 @@ void FPCGMetadataElementBase::PassthroughInput(FPCGContext* Context, TArray<FPCG
 	}
 }
 
-void FPCGMetadataElementBase::PassthroughAllInputs(FPCGContext* Context, TArray<FPCGTaggedData>& Outputs) const
-{
-	check(Context);
-
-	const UPCGMetadataSettingsBase* Settings = Context->GetInputSettings<UPCGMetadataSettingsBase>();
-	check(Settings);
-
-	const uint32 NumberOfOutputs = Settings->GetResultNum();
-	const uint32 PrimaryPinIndex = Settings->GetInputPinToForward();
-	TArray<FPCGTaggedData> InputsToForward = Context->InputData.GetSpatialInputsByPin(Settings->GetInputPinLabel(PrimaryPinIndex));
-
-	Outputs.Empty(InputsToForward.Num() * NumberOfOutputs);
-
-	// Passthrough all inputs to all of the outputs
-	for (uint32 I = 0; I < NumberOfOutputs; ++I)
-	{
-		const FName OutputPin = Settings->GetOutputPinLabel(I);
-		for (const FPCGTaggedData& Input : InputsToForward)
-		{
-			Outputs.Emplace_GetRef(Input).Pin = OutputPin;
-		}
-	}
-}
-
 namespace PCGMetadataOpPrivate
 {
 	using ContextType = FPCGMetadataElementBase::ContextType;
@@ -700,7 +676,7 @@ bool FPCGMetadataElementBase::ExecuteInternal(FPCGContext* Context) const
 		// No operation, so skip the iteration.
 		if (Context->GetIterationStateResult(IterationIndex) == EPCGTimeSliceInitResult::NoOperation)
 		{
-			PassthroughAllInputs(Context, Outputs);
+			PassthroughInput(Context, Outputs, IterationIndex);
 			return true;
 		}
 
diff --git a/Engine/Plugins/Experimental/PCG/Source/PCG/Public/Elements/Metadata/PCGMetadataOpElementBase.h b/Engine/Plugins/Experimental/PCG/Source/PCG/Public/Elements/Metadata/PCGMetadataOpElementBase.h
--- a/Engine/Plugins/Experimental/PCG/Source/PCG/Public/Elements/Metadata/PCGMetadataOpElementBase.h
+++ b/Engine/Plugins/Experimental/PCG/Source/PCG/Public/Elements/Metadata/PCGMetadataOpElementBase.h
@@ -214,7 +214,6 @@ protected:
 
 	/** To be called if we have no data to perform any operation, it will passthrough the input. */
 	void PassthroughInput(FPCGContext* Context, TArray<FPCGTaggedData>& Outputs, const int32 Index) const;
-	void PassthroughAllInputs(FPCGContext* Context, TArray<FPCGTaggedData>& Outputs) const;
 };
 
 /**
