[PCG] Fix an array bounds issue with MetadataOpElementBase

#rnx

diff --git a/Engine/Plugins/Experimental/PCG/Source/PCG/Private/Elements/Metadata/PCGMetadataOpElementBase.cpp b/Engine/Plugins/Experimental/PCG/Source/PCG/Private/Elements/Metadata/PCGMetadataOpElementBase.cpp
--- a/Engine/Plugins/Experimental/PCG/Source/PCG/Private/Elements/Metadata/PCGMetadataOpElementBase.cpp
+++ b/Engine/Plugins/Experimental/PCG/Source/PCG/Private/Elements/Metadata/PCGMetadataOpElementBase.cpp
@@ -263,7 +263,12 @@ void FPCGMetadataElementBase::PassthroughInput(FPCGContext* Context, TArray<FPCG
 
 	const uint32 NumberOfOutputs = Settings->GetResultNum();
 	const uint32 PrimaryPinIndex = Settings->GetInputPinToForward();
-	TArray<FPCGTaggedData> InputsToForward = Context->InputData.GetSpatialInputsByPin(Settings->GetInputPinLabel(PrimaryPinIndex));
+	TArray<FPCGTaggedData> InputsToForward = Context->InputData.GetInputsByPin(Settings->GetInputPinLabel(PrimaryPinIndex));
+
+	if (!ensure(Index < InputsToForward.Num()))
+	{
+		return;
+	}
 
 	// Passthrough this single input to all of the outputs
 	for (uint32 I = 0; I < NumberOfOutputs; ++I)
