[PCG] Fix up RuntimeGen bugs
* Toggling IsPartitioned on/off could leak generated components into the world
* Changing generation radii at runtime could leak generated components into the world
* Copy/Paste of GenerateAtRuntime components hits a check failure in UPCGComponent::OnRefresh

#rnx

diff --git a/Engine/Plugins/Experimental/PCG/Source/PCG/Private/PCGActorAndComponentMapping.cpp b/Engine/Plugins/Experimental/PCG/Source/PCG/Private/PCGActorAndComponentMapping.cpp
--- a/Engine/Plugins/Experimental/PCG/Source/PCG/Private/PCGActorAndComponentMapping.cpp
+++ b/Engine/Plugins/Experimental/PCG/Source/PCG/Private/PCGActorAndComponentMapping.cpp
@@ -15,6 +15,7 @@
 #include "LandscapeProxy.h"
 #include "StaticMeshCompiler.h"
 #include "TextureCompiler.h"
+#include "Algo/AnyOf.h"
 #include "Engine/Engine.h"
 #include "Engine/Level.h"
 #include "Engine/StaticMesh.h"
@@ -418,11 +419,6 @@ void FPCGActorAndComponentMapping::UnregisterPCGComponent(UPCGComponent* InCompo
 #endif // WITH_EDITOR
 	}
 
-	if (InComponent->GetOwner() && !InComponent->GetOwner()->IsA<APCGPartitionActor>())
-	{
-		PCGSubsystem->OnOriginalComponentUnregistered(InComponent);
-	}
-
 	UnregisterPartitionedPCGComponent(InComponent);
 	UnregisterNonPartitionedPCGComponent(InComponent);
 
@@ -435,6 +431,8 @@ void FPCGActorAndComponentMapping::UnregisterPCGComponent(UPCGComponent* InCompo
 
 void FPCGActorAndComponentMapping::UnregisterPartitionedPCGComponent(UPCGComponent* InComponent)
 {
+	PCGSubsystem->OnOriginalComponentUnregistered(InComponent);
+
 	if (!PartitionedOctree.RemoveComponent(InComponent) || InComponent->IsManagedByRuntimeGenSystem())
 	{
 		return;
@@ -462,6 +460,8 @@ void FPCGActorAndComponentMapping::UnregisterPartitionedPCGComponent(UPCGCompone
 
 void FPCGActorAndComponentMapping::UnregisterNonPartitionedPCGComponent(UPCGComponent* InComponent)
 {
+	PCGSubsystem->OnOriginalComponentUnregistered(InComponent);
+
 	NonPartitionedOctree.RemoveComponent(InComponent);
 }
 
@@ -924,7 +924,7 @@ void FPCGActorAndComponentMapping::UpdateTracking(UPCGComponent* InComponent, bo
 		{
 			if (!Actor->HasActorRegisteredAllComponents() && !bDisableDelayedActorRegistering)
 			{
-				DelayedAddedActors.Emplace({ Actor, false, 0 });
+				DelayedAddedActors.Emplace(Actor, { false, 0 });
 				continue;
 			}
 
@@ -1138,12 +1138,12 @@ void FPCGActorAndComponentMapping::AddDelayedActors()
 		return;
 	}
 
-	TSet<TTuple<TObjectKey<AActor>, bool, int>> StillDelayedActors;
+	TMap<TObjectKey<AActor>, TTuple<bool, int>> StillDelayedActors;
 	const bool bDisableDelayedActorRegistering = PCGActorAndComponentMapping::CVarDisableDelayedActorRegistering.GetValueOnAnyThread();
 
-	for (TTuple<TObjectKey<AActor>, bool, int>& ActorPtrAndShouldDirty : DelayedAddedActors)
+	for (const TPair<TObjectKey<AActor>, TTuple<bool, int>>& ActorPtrAndShouldDirty : DelayedAddedActors)
 	{
-		AActor* Actor = ActorPtrAndShouldDirty.Get<0>().ResolveObjectPtr();
+		AActor* Actor = ActorPtrAndShouldDirty.Key.ResolveObjectPtr();
 		if (!Actor)
 		{
 			continue;
@@ -1156,7 +1156,7 @@ void FPCGActorAndComponentMapping::AddDelayedActors()
 		else
 		{
 			// Implementation note: since the delayed actors list is built from top-level actors (e.g. directly in the level) then depth here is 0.
-			OnActorAdded_Internal(Actor, ActorPtrAndShouldDirty.Get<1>(), ActorPtrAndShouldDirty.Get<2>());
+			OnActorAdded_Internal(Actor, ActorPtrAndShouldDirty.Value.Get<0>(), ActorPtrAndShouldDirty.Value.Get<1>(), /*bForceAddDelayedActor=*/true);
 		}
 	}
 
@@ -1204,18 +1204,27 @@ void FPCGActorAndComponentMapping::OnActorAdded(AActor* InActor)
 #endif // WITH_EDITOR
 
 	// Implementation note: since this is called only for actors directly in the current level, the depth here is 0.
-	OnActorAdded_Internal(InActor, true, LevelInstanceDepth);
+	// Another implementation note: We delay adding because OnActorAdded fires before an actor's properties are set,
+	// so the actor is not ready for processing until the next tick.
+	DelayedAddedActors.Emplace(InActor, { true, LevelInstanceDepth });
 }
 
-void FPCGActorAndComponentMapping::OnActorAdded_Internal(AActor* InActor, bool bShouldDirty, int32 LevelInstanceDepth)
+void FPCGActorAndComponentMapping::OnActorAdded_Internal(AActor* InActor, bool bShouldDirty, int32 LevelInstanceDepth, bool bForceAddDelayedActor)
 {
 	TRACE_CPUPROFILER_EVENT_SCOPE(FPCGActorAndComponentMapping::OnActorAdded);
 	check(InActor && !InActor->IsA<APCGWorldActor>() && PCGSubsystem && InActor->GetWorld() == PCGSubsystem->GetWorld());
 
+	// A delayed actor should only be added from AddDelayedActors(), because that guarantees the new actor has waited at
+	// least one tick for its properties to be popualated.
+	if (!bForceAddDelayedActor && DelayedAddedActors.Contains(InActor))
+	{
+		return;
+	}
+
 	// If the subsystem is not initialized, wait for it to be, and store all the actors to check
 	if (!PCGSubsystem->IsInitialized())
 	{
-		DelayedAddedActors.Emplace({ InActor, bShouldDirty, LevelInstanceDepth });
+		DelayedAddedActors.Emplace(InActor, { bShouldDirty, LevelInstanceDepth });
 		return;
 	}
 
@@ -1245,7 +1254,7 @@ void FPCGActorAndComponentMapping::OnActorAdded_Internal(AActor* InActor, bool b
 		}
 		else
 		{
-			DelayedAddedActors.Emplace({ InActor, bShouldDirty, LevelInstanceDepth });
+			DelayedAddedActors.Emplace(InActor, { bShouldDirty, LevelInstanceDepth });
 			return;
 		}
 	}
@@ -1580,6 +1589,12 @@ void FPCGActorAndComponentMapping::OnObjectPropertyChanged(UObject* InObject, FP
 			{
 				if (AActor* ActorToChange = TrackedActor.Key.ResolveObjectPtr())
 				{
+					// Ignore property changes on delayed actors. All their properties are still being set.
+					if (DelayedAddedActors.Contains(ActorToChange))
+					{
+						return;
+					}
+
 					OnActorChanged(ActorToChange, /*bInHasMoved=*/ false, /*InOriginatingChangeObject=*/ InObject);
 					UpdateActorDependencies(ActorToChange);
 				}
@@ -1601,6 +1616,12 @@ void FPCGActorAndComponentMapping::OnObjectPropertyChanged(UObject* InObject, FP
 	}
 #endif
 
+	// Ignore property changes on delayed actors. All their properties are still being set.
+	if (DelayedAddedActors.Contains(Actor))
+	{
+		return;
+	}
+
 	// Check if we are not tracking it or is a tag change.
 	bool bShouldChange = true;
 	if (!IsActorTracked(Actor) || bActorTagChange)
diff --git a/Engine/Plugins/Experimental/PCG/Source/PCG/Private/PCGComponent.cpp b/Engine/Plugins/Experimental/PCG/Source/PCG/Private/PCGComponent.cpp
--- a/Engine/Plugins/Experimental/PCG/Source/PCG/Private/PCGComponent.cpp
+++ b/Engine/Plugins/Experimental/PCG/Source/PCG/Private/PCGComponent.cpp
@@ -1527,6 +1527,14 @@ void UPCGComponent::PostEditChangeProperty(FPropertyChangedEvent& PropertyChange
 		return;
 	}
 
+	const FName MemberName = PropertyChangedEvent.MemberProperty->GetFName();
+
+	if (MemberName == GET_MEMBER_NAME_CHECKED(UPCGComponent, GenerationRadii))
+	{
+		// RuntimeGen will automatically pick up any changes to generation radii, we don't need to do any work here.
+		return;
+	}
+
 	// Important note: all property changes already go through the OnObjectPropertyChanged, and will be dirtied here.
 	// So where only a Refresh is needed, it goes through the "capture all" else case.
 	if (PropName == GET_MEMBER_NAME_CHECKED(UPCGComponent, bIsComponentPartitioned))
diff --git a/Engine/Plugins/Experimental/PCG/Source/PCG/Private/RuntimeGen/PCGRuntimeGenScheduler.cpp b/Engine/Plugins/Experimental/PCG/Source/PCG/Private/RuntimeGen/PCGRuntimeGenScheduler.cpp
--- a/Engine/Plugins/Experimental/PCG/Source/PCG/Private/RuntimeGen/PCGRuntimeGenScheduler.cpp
+++ b/Engine/Plugins/Experimental/PCG/Source/PCG/Private/RuntimeGen/PCGRuntimeGenScheduler.cpp
@@ -700,15 +700,13 @@ void FPCGRuntimeGenScheduler::OnOriginalComponentRegistered(UPCGComponent* InOri
 
 void FPCGRuntimeGenScheduler::OnOriginalComponentUnregistered(UPCGComponent* InOriginalComponent)
 {
-	check(ActorAndComponentMapping);
-
-	// Safety checks.
 	if (!InOriginalComponent || Cast<APCGPartitionActor>(InOriginalComponent->GetOwner()))
 	{
-		ensure(false);
 		return;
 	}
 
+	check(ActorAndComponentMapping);
+
 	// When an original/non-partitioned component is unregistered, we need to dirty the state.
 	bAnyRuntimeGenComponentsExistDirty = true;
 
diff --git a/Engine/Plugins/Experimental/PCG/Source/PCG/Public/PCGActorAndComponentMapping.h b/Engine/Plugins/Experimental/PCG/Source/PCG/Public/PCGActorAndComponentMapping.h
--- a/Engine/Plugins/Experimental/PCG/Source/PCG/Public/PCGActorAndComponentMapping.h
+++ b/Engine/Plugins/Experimental/PCG/Source/PCG/Public/PCGActorAndComponentMapping.h
@@ -138,7 +138,7 @@ private:
 
 	void OnActorAdded(AActor* InActor);
 	void OnActorLoaded(AActor& InActor);
-	void OnActorAdded_Internal(AActor* InActor, bool bShouldDirty, int32 LevelInstanceDepth);
+	void OnActorAdded_Internal(AActor* InActor, bool bShouldDirty, int32 LevelInstanceDepth, bool bForceAddDelayedActor = false);
 	void OnActorDeleted(AActor* InActor);
 	void OnActorUnloaded(AActor& InActor);
 	void OnActorDeleted_Internal(AActor* InActor, bool bShouldDirty, int32 LevelInstanceDepth);
@@ -222,7 +222,7 @@ private:
 	mutable FRWLock TrackedComponentsLock;
 
 	// Keep track of actors that aren't yet ready (or if the subsystem is not yet ready), whether we should dirty them and their instance level depth so we can add them in next tick.
-	TSet<TTuple<TObjectKey<AActor>, bool, int>> DelayedAddedActors;
+	TMap<TObjectKey<AActor>, TTuple<bool, int>> DelayedAddedActors;
 
 	/** Keep a mapping between tracked actors and their dependencies. */
 	TMap<TObjectKey<AActor>, TSet<TObjectPtr<UObject>>> TrackedActorsToDependenciesMap;
