[PCG] Pass on our elements to make sure we propagate the input tags

diff --git a/Engine/Plugins/Experimental/PCG/Source/PCG/Private/Elements/Metadata/PCGMetadataPartition.cpp b/Engine/Plugins/Experimental/PCG/Source/PCG/Private/Elements/Metadata/PCGMetadataPartition.cpp
--- a/Engine/Plugins/Experimental/PCG/Source/PCG/Private/Elements/Metadata/PCGMetadataPartition.cpp
+++ b/Engine/Plugins/Experimental/PCG/Source/PCG/Private/Elements/Metadata/PCGMetadataPartition.cpp
@@ -76,9 +76,8 @@ bool FPCGMetadataPartitionElement::ExecuteInternal(FPCGContext* Context) const
 		{
 			if (PartitionData)
 			{
-				FPCGTaggedData& Output = Outputs.Emplace_GetRef();
+				FPCGTaggedData& Output = Outputs.Add_GetRef(Input);
 				Output.Data = PartitionData;
-				Output.Tags = Input.Tags;
 			}
 		}
 	}
diff --git a/Engine/Plugins/Experimental/PCG/Source/PCG/Private/Elements/PCGAttributeFilterNamesElement.cpp b/Engine/Plugins/Experimental/PCG/Source/PCG/Private/Elements/PCGAttributeFilterNamesElement.cpp
--- a/Engine/Plugins/Experimental/PCG/Source/PCG/Private/Elements/PCGAttributeFilterNamesElement.cpp
+++ b/Engine/Plugins/Experimental/PCG/Source/PCG/Private/Elements/PCGAttributeFilterNamesElement.cpp
@@ -172,7 +172,7 @@ bool FPCGAttributeFilterNamesElement::ExecuteInternal(FPCGContext* Context) cons
 		}
 
 		TArray<FPCGTaggedData>& Outputs = Context->OutputData.TaggedData;
-		FPCGTaggedData& Output = Outputs.Emplace_GetRef();
+		FPCGTaggedData& Output = Outputs.Add_GetRef(InputTaggedData);
 		Output.Data = OutputData;
 	}
 
diff --git a/Engine/Plugins/Experimental/PCG/Source/PCG/Private/Elements/PCGAttributeGetFromPointIndexElement.cpp b/Engine/Plugins/Experimental/PCG/Source/PCG/Private/Elements/PCGAttributeGetFromPointIndexElement.cpp
--- a/Engine/Plugins/Experimental/PCG/Source/PCG/Private/Elements/PCGAttributeGetFromPointIndexElement.cpp
+++ b/Engine/Plugins/Experimental/PCG/Source/PCG/Private/Elements/PCGAttributeGetFromPointIndexElement.cpp
@@ -120,7 +120,7 @@ bool FPCGAttributeGetFromPointIndexElement::ExecuteInternal(FPCGContext* Context
 		OutputPointData->InitializeFromData(PointData);
 		OutputPointData->GetMutablePoints().Add(Point);
 
-		FPCGTaggedData& Output = Outputs.Emplace_GetRef();
+		FPCGTaggedData& Output = Outputs.Add_GetRef(Inputs[0]);
 		Output.Data = OutputPointData;
 		Output.Pin = PCGAttributeGetFromPointIndexConstants::OutputPointLabel;
 	}
@@ -164,7 +164,7 @@ bool FPCGAttributeGetFromPointIndexElement::ExecuteInternal(FPCGContext* Context
 
 		if (PCGMetadataAttribute::CallbackWithRightType(Accessor->GetUnderlyingType(), ExtractAttribute))
 		{
-			FPCGTaggedData& Output = Outputs.Emplace_GetRef();
+			FPCGTaggedData& Output = Outputs.Add_GetRef(Inputs[0]);
 			Output.Data = OutputParamData;
 			Output.Pin = PCGAttributeGetFromPointIndexConstants::OutputAttributeLabel;
 		}
diff --git a/Engine/Plugins/Experimental/PCG/Source/PCG/Private/Elements/PCGAttributeReduceElement.cpp b/Engine/Plugins/Experimental/PCG/Source/PCG/Private/Elements/PCGAttributeReduceElement.cpp
--- a/Engine/Plugins/Experimental/PCG/Source/PCG/Private/Elements/PCGAttributeReduceElement.cpp
+++ b/Engine/Plugins/Experimental/PCG/Source/PCG/Private/Elements/PCGAttributeReduceElement.cpp
@@ -302,7 +302,7 @@ bool FPCGAttributeReduceElement::ExecuteInternal(FPCGContext* Context) const
 
 		if (ensure(OutputParams) && (Outputs.IsEmpty() || !Settings->bMergeOutputAttributes))
 		{
-			FPCGTaggedData& Output = Outputs.Emplace_GetRef();
+			FPCGTaggedData& Output = Outputs.Add_GetRef(Inputs[i]);
 			Output.Data = OutputParams;
 		}
 	}
diff --git a/Engine/Plugins/Experimental/PCG/Source/PCG/Private/Elements/PCGAttributeSelectElement.cpp b/Engine/Plugins/Experimental/PCG/Source/PCG/Private/Elements/PCGAttributeSelectElement.cpp
--- a/Engine/Plugins/Experimental/PCG/Source/PCG/Private/Elements/PCGAttributeSelectElement.cpp
+++ b/Engine/Plugins/Experimental/PCG/Source/PCG/Private/Elements/PCGAttributeSelectElement.cpp
@@ -422,7 +422,7 @@ bool FPCGAttributeSelectElement::ExecuteInternal(FPCGContext* Context) const
 	}
 
 	TArray<FPCGTaggedData>& Outputs = Context->OutputData.TaggedData;
-	FPCGTaggedData& Output = Outputs.Emplace_GetRef();
+	FPCGTaggedData& Output = Outputs.Add_GetRef(Inputs[0]);
 	Output.Data = OutputParamData;
 	Output.Pin = PCGAttributeSelectConstants::OutputAttributeLabel;
 
@@ -436,7 +436,7 @@ bool FPCGAttributeSelectElement::ExecuteInternal(FPCGContext* Context) const
 		OutputPointData->InitializeFromData(PointData);
 		OutputPointData->GetMutablePoints().Add(PointData->GetPoint(OutputIndex));
 
-		FPCGTaggedData& PointOutput = Outputs.Emplace_GetRef();
+		FPCGTaggedData& PointOutput = Outputs.Add_GetRef(Inputs[0]);
 		PointOutput.Data = OutputPointData;
 		PointOutput.Pin = PCGAttributeSelectConstants::OutputPointLabel;
 	}
diff --git a/Engine/Plugins/Experimental/PCG/Source/PCG/Private/Elements/PCGAttributeTransferElement.cpp b/Engine/Plugins/Experimental/PCG/Source/PCG/Private/Elements/PCGAttributeTransferElement.cpp
--- a/Engine/Plugins/Experimental/PCG/Source/PCG/Private/Elements/PCGAttributeTransferElement.cpp
+++ b/Engine/Plugins/Experimental/PCG/Source/PCG/Private/Elements/PCGAttributeTransferElement.cpp
@@ -175,10 +175,10 @@ bool FPCGAttributeTransferElement::ExecuteInternal(FPCGContext* Context) const
 	UPCGSpatialData* OutputData = TargetData->DuplicateData();
 	check(OutputData->Metadata);
 
-	auto AppendOutputData = [Context, OutputData]()
+	auto AppendOutputData = [Context, OutputData, &TargetInputs]()
 	{
 		TArray<FPCGTaggedData>& Outputs = Context->OutputData.TaggedData;
-		FPCGTaggedData& Output = Outputs.Emplace_GetRef();
+		FPCGTaggedData& Output = Outputs.Add_GetRef(TargetInputs[0]);
 		Output.Data = OutputData;
 	};
 
