Added tooltips on the pin usage icons

#rmx

diff --git a/Engine/Plugins/Experimental/PCG/Source/PCG/Private/PCGPin.cpp b/Engine/Plugins/Experimental/PCG/Source/PCG/Private/PCGPin.cpp
--- a/Engine/Plugins/Experimental/PCG/Source/PCG/Private/PCGPin.cpp
+++ b/Engine/Plugins/Experimental/PCG/Source/PCG/Private/PCGPin.cpp
@@ -45,11 +45,13 @@ namespace PCGPinPropertiesHelpers
 	{
 		if (InPinProperties.Usage == EPCGPinUsage::Loop)
 		{
+			OutTooltip = NSLOCTEXT("PCGPins", "LoopTooltip", "Loop pin, data collection will be split to one data per execution.");
 			OutExtraIcon = PCGPinConstants::Icons::LoopPinIcon;
 			return true;
 		}
 		else if (InPinProperties.Usage == EPCGPinUsage::Feedback)
 		{
+			OutTooltip = NSLOCTEXT("PCGPins", "FeedbackTooltip", "Feedback pin, will daisy-chain results from graph entry or previous loop iteration.");
 			OutExtraIcon = PCGPinConstants::Icons::FeedbackPinIcon;
 			return true;
 		}
