Fixed potential stack overflow with recursive graphs

#rnx

diff --git a/Engine/Plugins/Experimental/PCG/Source/PCG/Private/PCGGraph.cpp b/Engine/Plugins/Experimental/PCG/Source/PCG/Private/PCGGraph.cpp
--- a/Engine/Plugins/Experimental/PCG/Source/PCG/Private/PCGGraph.cpp
+++ b/Engine/Plugins/Experimental/PCG/Source/PCG/Private/PCGGraph.cpp
@@ -334,19 +334,37 @@ void UPCGGraph::PostLoad()
 }
 
 bool UPCGGraph::IsEditorOnly() const
+{
+	TSet<const UPCGGraph*> VisitedGraphs;
+	return IsEditorOnly_Internal(VisitedGraphs);
+}
+
+bool UPCGGraph::IsEditorOnly_Internal(TSet<const UPCGGraph*>& VisitedGraphs) const
 {
 	bool bIsCurrentlyEditorOnly = (Super::IsEditorOnly() || bIsEditorOnly);
+	
+	check(!VisitedGraphs.Contains(this));
+	VisitedGraphs.Add(this);
 
 	if (!bIsCurrentlyEditorOnly)
 	{
 		// Check for any subgraphs...
-		ForEachNode([&bIsCurrentlyEditorOnly](UPCGNode* Node)
+		ForEachNode([&bIsCurrentlyEditorOnly, &VisitedGraphs](UPCGNode* Node)
 		{
+			// Early out if we already know this is going to be editor only
+			if (bIsCurrentlyEditorOnly)
+			{
+				return;
+			}
+
 			if (UPCGBaseSubgraphNode* SubgraphNode = Cast<UPCGBaseSubgraphNode>(Node))
 			{
-				if (SubgraphNode->GetSubgraph())
+				if (const UPCGGraph* Subgraph = SubgraphNode->GetSubgraph())
 				{
-					bIsCurrentlyEditorOnly |= SubgraphNode->GetSubgraph()->IsEditorOnly();
+					if (!VisitedGraphs.Contains(Subgraph))
+					{
+						bIsCurrentlyEditorOnly |= Subgraph->IsEditorOnly_Internal(VisitedGraphs);
+					}
 				}
 			}
 		});
diff --git a/Engine/Plugins/Experimental/PCG/Source/PCG/Public/PCGGraph.h b/Engine/Plugins/Experimental/PCG/Source/PCG/Public/PCGGraph.h
--- a/Engine/Plugins/Experimental/PCG/Source/PCG/Public/PCGGraph.h
+++ b/Engine/Plugins/Experimental/PCG/Source/PCG/Public/PCGGraph.h
@@ -317,6 +317,8 @@ protected:
 	void RemoveNodes_Internal(TArrayView<UPCGNode*> InNodes);
 	void AddNodes_Internal(TArrayView<UPCGNode*> InNodes);
 
+	bool IsEditorOnly_Internal(TSet<const UPCGGraph*>& VisitedGraphs) const;
+
 	/** Calculates node grid size. Not thread safe, called within write lock. */
 	uint32 CalculateNodeGridSizeRecursive_Unsafe(const UPCGNode* InNode, uint32 InDefaultGridSize) const;
 
