Fixed issue with the FilterByTag that propagated overrides as output too

#rnx

diff --git a/Engine/Plugins/Experimental/PCG/Source/PCG/Private/Elements/PCGFilterByTag.cpp b/Engine/Plugins/Experimental/PCG/Source/PCG/Private/Elements/PCGFilterByTag.cpp
--- a/Engine/Plugins/Experimental/PCG/Source/PCG/Private/Elements/PCGFilterByTag.cpp
+++ b/Engine/Plugins/Experimental/PCG/Source/PCG/Private/Elements/PCGFilterByTag.cpp
@@ -93,7 +93,7 @@ bool FPCGFilterByTagElement::ExecuteInternal(FPCGContext* Context) const
 	const bool bKeepIfTag = (Settings->Operation == EPCGFilterByTagOperation::KeepTagged);
 	const TArray<FString> Tags = PCGHelpers::GetStringArrayFromCommaSeparatedString(Settings->SelectedTags);
 
-	TArray<FPCGTaggedData> Inputs = Context->InputData.TaggedData;
+	TArray<FPCGTaggedData> Inputs = Context->InputData.GetInputsByPin(PCGPinConstants::DefaultInputLabel);
 	TArray<FPCGTaggedData>& Outputs = Context->OutputData.TaggedData;
 
 	for (const FPCGTaggedData& Input : Inputs)
