[PCG] Add constexpr IsConstructible matrix to know if a metadata type is constructible from another

* Add the new check to AttributeFilter
* Add the new check to Actor Overrides + allow constructible now for it.

#rnx

diff --git a/Engine/Plugins/Experimental/PCG/Source/PCG/Private/Elements/PCGAttributeFilter.cpp b/Engine/Plugins/Experimental/PCG/Source/PCG/Private/Elements/PCGAttributeFilter.cpp
--- a/Engine/Plugins/Experimental/PCG/Source/PCG/Private/Elements/PCGAttributeFilter.cpp
+++ b/Engine/Plugins/Experimental/PCG/Source/PCG/Private/Elements/PCGAttributeFilter.cpp
@@ -197,23 +197,20 @@ namespace PCGAttributeFilterHelpers
 		// And also validate that types are comparable/constructible. Do it all at once for a single dispatch.
 		bool bCanCompare = true;
 		bool bCanSearchString = true;
-		bool bIsConstructible = true;
-		PCGMetadataAttribute::CallbackWithRightType(InOutThresholdInfo.ThresholdAccessor->GetUnderlyingType(), [&bCanCompare, &bCanSearchString, &bIsConstructible, TargetType](auto ThresholdValue)
+		PCGMetadataAttribute::CallbackWithRightType(InOutThresholdInfo.ThresholdAccessor->GetUnderlyingType(), [&bCanCompare, &bCanSearchString, TargetType](auto ThresholdValue)
 		{
 			bCanCompare = PCG::Private::MetadataTraits<decltype(ThresholdValue)>::CanCompare;
 			bCanSearchString = PCG::Private::MetadataTraits<decltype(ThresholdValue)>::CanSearchString;
-			bIsConstructible = PCGMetadataAttribute::CallbackWithRightType(TargetType, [](auto TargetValue) { return std::is_constructible_v<decltype(TargetValue),decltype(ThresholdValue)>;});
 		});
 
 		// Comparison between threshold and target data needs to be of the same type. So we have to make sure that we can
 		// request target type from threshold type. ie. We need to make sure we can broadcast threshold type to target type, or construct a target type from a threshold type.
 		// For example: if target is double but threshold is int32, we can broadcast int32 to double, to compare a double with a double.
-		if (!PCG::Private::IsBroadcastable(InOutThresholdInfo.ThresholdAccessor->GetUnderlyingType(), TargetType) && !bIsConstructible)
+		if (!PCG::Private::IsBroadcastableOrConstructible(InOutThresholdInfo.ThresholdAccessor->GetUnderlyingType(), TargetType))
 		{
-			UEnum* PCGDataTypeEnum = StaticEnum<EPCGMetadataTypes>();
-			FText ThresholdTypeName = FText::FromString(PCGDataTypeEnum ? PCGDataTypeEnum->GetNameStringByValue(InOutThresholdInfo.ThresholdAccessor->GetUnderlyingType()) : FString(TEXT("Unknown")));
-			FText InputTypeName = FText::FromString(PCGDataTypeEnum ? PCGDataTypeEnum->GetNameStringByValue(TargetType) : FString(TEXT("Unknown")));
-			PCGE_LOG_C(Warning, GraphAndLog, InContext, FText::Format(LOCTEXT("TypeConversionFailed", "Cannot use threshold type '{0}' on input target type '{1}'"),
+			const FText InputTypeName = PCG::Private::GetTypeNameText(TargetType);
+			const FText ThresholdTypeName = PCG::Private::GetTypeNameText(InOutThresholdInfo.ThresholdAccessor->GetUnderlyingType());
+			PCGE_LOG_C(Warning, GraphAndLog, InContext, FText::Format(LOCTEXT("TypeConversionFailed", "Cannot convert threshold type '{0}' to input target type '{1}'"),
 				ThresholdTypeName,
 				InputTypeName));
 			return false;
@@ -221,13 +218,15 @@ namespace PCGAttributeFilterHelpers
 
 		if (bCheckCompare && !bCanCompare)
 		{
-			PCGE_LOG_C(Warning, GraphAndLog, InContext, LOCTEXT("TypeComparisonFailed", "Cannot compare target type"));
+			const FText InputTypeName = PCG::Private::GetTypeNameText(TargetType);
+			PCGE_LOG_C(Warning, GraphAndLog, InContext, FText::Format(LOCTEXT("TypeComparisonFailed", "Cannot compare target type '{0}'"), InputTypeName));
 			return false;
 		}
 
 		if (bCheckStringSearch && !bCanSearchString)
 		{
-			PCGE_LOG_C(Warning, GraphAndLog, InContext, LOCTEXT("TypeStringSearchFailed", "Cannot perform string operations on target type"));
+			const FText InputTypeName = PCG::Private::GetTypeNameText(TargetType);
+			PCGE_LOG_C(Warning, GraphAndLog, InContext, FText::Format(LOCTEXT("TypeStringSearchFailed", "Cannot perform string operations on target type '{0}'"), InputTypeName));
 			return false;
 		}
 
diff --git a/Engine/Plugins/Experimental/PCG/Source/PCG/Private/Metadata/PCGActorPropertyOverride.cpp b/Engine/Plugins/Experimental/PCG/Source/PCG/Private/Metadata/PCGActorPropertyOverride.cpp
--- a/Engine/Plugins/Experimental/PCG/Source/PCG/Private/Metadata/PCGActorPropertyOverride.cpp
+++ b/Engine/Plugins/Experimental/PCG/Source/PCG/Private/Metadata/PCGActorPropertyOverride.cpp
@@ -85,8 +85,7 @@ void FPCGActorSingleOverride::Initialize(const FPCGAttributePropertySelector& In
 		return;
 	}
 
-	// TODO: Replace with something less restrictive (e.g. allow Double -> Float conversion).
-	if (!PCG::Private::IsBroadcastable(ActorOverrideInputAccessor->GetUnderlyingType(), ActorOverrideOutputAccessor->GetUnderlyingType()))
+	if (!PCG::Private::IsBroadcastableOrConstructible(ActorOverrideInputAccessor->GetUnderlyingType(), ActorOverrideOutputAccessor->GetUnderlyingType()))
 	{
 		PCGLog::LogWarningOnGraph(
 			FText::Format(LOCTEXT("TypesIncompatible", "ActorOverride cannot set input '{0}' to output '{1}'. Cannot convert type '{2}' to type '{3}'. Will be skipped."),
diff --git a/Engine/Plugins/Experimental/PCG/Source/PCG/Public/Metadata/PCGActorPropertyOverride.h b/Engine/Plugins/Experimental/PCG/Source/PCG/Public/Metadata/PCGActorPropertyOverride.h
--- a/Engine/Plugins/Experimental/PCG/Source/PCG/Public/Metadata/PCGActorPropertyOverride.h
+++ b/Engine/Plugins/Experimental/PCG/Source/PCG/Public/Metadata/PCGActorPropertyOverride.h
@@ -89,7 +89,7 @@ private:
 		}
 
 		Type Value{};
-		if (ActorOverrideInputAccessor->Get<Type>(Value, InputKeyIndex, *InputKeys.Get(), EPCGAttributeAccessorFlags::AllowBroadcast))
+		if (ActorOverrideInputAccessor->Get<Type>(Value, InputKeyIndex, *InputKeys.Get(), EPCGAttributeAccessorFlags::AllowBroadcast | EPCGAttributeAccessorFlags::AllowConstructible))
 		{
 			if (ActorOverrideOutputAccessor->Set<Type>(Value, OutputKey))
 			{
diff --git a/Engine/Plugins/Experimental/PCG/Source/PCG/Public/Metadata/PCGMetadataAttributeTraits.h b/Engine/Plugins/Experimental/PCG/Source/PCG/Public/Metadata/PCGMetadataAttributeTraits.h
--- a/Engine/Plugins/Experimental/PCG/Source/PCG/Public/Metadata/PCGMetadataAttributeTraits.h
+++ b/Engine/Plugins/Experimental/PCG/Source/PCG/Public/Metadata/PCGMetadataAttributeTraits.h
@@ -50,6 +50,24 @@ enum class EPCGMetadataTypes : uint8
 	MACRO(FSoftObjectPath) \
 	MACRO(FSoftClassPath)
 
+// Need a duplicated macro for a template list, because Foo<A,B,> is syntax error. It needs Foo<A,B>
+#define PCG_FOREACH_SUPPORTEDTYPES_WITH_COMMA(MACRO) \
+	MACRO(int32),           \
+	MACRO(int64),           \
+	MACRO(float),           \
+	MACRO(double),          \
+	MACRO(FVector2D),       \
+	MACRO(FVector),         \
+	MACRO(FVector4),        \
+	MACRO(FQuat),           \
+	MACRO(FTransform),      \
+	MACRO(FString),         \
+	MACRO(bool),            \
+	MACRO(FRotator),        \
+	MACRO(FName),           \
+	MACRO(FSoftObjectPath), \
+	MACRO(FSoftClassPath)
+
 namespace PCG
 {
 	namespace Private
@@ -153,10 +171,11 @@ namespace PCG
 
 		// Wrapper around a standard 2-dimensional CArray that is constexpr, to know if a type is broadcastable to another.
 		// First index is the original type, second index is the wanted type. Returns true if we can broadcast first type into second type.
-		struct UBroadcastableTypes
+		struct FBroadcastableTypes
 		{
-			constexpr UBroadcastableTypes() : Values{{false}}
+			constexpr FBroadcastableTypes() : Values{ {false} }
 			{
+				// All types are broadcastable to themselves.
 				for (uint8 i = 0; i < (uint8)EPCGMetadataTypes::Count; ++i)
 				{
 					Values[i][i] = true;
@@ -220,7 +239,52 @@ namespace PCG
 			bool Values[(uint8)EPCGMetadataTypes::Count][(uint8)EPCGMetadataTypes::Count];
 		};
 
-		inline static constexpr UBroadcastableTypes BroadcastableTypes{};
+		// Wrapper around a standard 2-dimensional CArray that is constexpr, to know if a type is constructible from another.
+		// First index is the wanted type, second index is the original type. Returns true if we can construct first type from second type.
+		struct FConstructibleTypes
+		{
+		private:
+			// Empty struct that will hold all our types.
+			template <typename... T>
+			struct TypeHolder {};
+
+			// Use the macro to have AllTypes that is the Dummy struct templated with all our types.
+#define PCG_ALL_TYPES(T) T
+			using AllTypes = TypeHolder<PCG_FOREACH_SUPPORTEDTYPES_WITH_COMMA(PCG_ALL_TYPES)>;
+#undef PCG_ALL_TYPES
+
+			// Unroll AllTypes twice, to have all combinations of types <T, U>
+			template <typename FirstType> constexpr void UnrollSecond(TypeHolder<> InTypeHolder) {}
+
+			template <typename FirstType, typename SecondType, typename... SecondTypes>
+			constexpr void UnrollSecond(TypeHolder<SecondType, SecondTypes...> InTypeHolder)
+			{
+				Values[MetadataTypes<FirstType>::Id][MetadataTypes<SecondType>::Id] = std::is_constructible_v<FirstType, SecondType>;
+				UnrollSecond<FirstType>(TypeHolder<SecondTypes...>{});
+			}
+
+			template <typename InputType, typename... InputTypes>
+			constexpr void UnrollFirst(TypeHolder<InputType, InputTypes...> InTypeHolder)
+			{
+				UnrollSecond<InputType>(AllTypes{});
+				UnrollFirst(TypeHolder<InputTypes...>{});
+			}
+
+			constexpr void UnrollFirst(TypeHolder<> InTypeHolder) {}
+
+		public:
+			constexpr FConstructibleTypes() : Values{ {false} }
+			{
+				UnrollFirst(AllTypes{});
+			}
+
+			constexpr const bool* operator[](int i) const { return Values[i]; }
+
+			bool Values[(uint8)EPCGMetadataTypes::Count][(uint8)EPCGMetadataTypes::Count];
+		};
+
+		inline static constexpr FBroadcastableTypes BroadcastableTypes{};
+		inline static constexpr FConstructibleTypes ConstructibleTypes{};
 
 		constexpr inline bool IsBroadcastable(uint16 FirstType, uint16 SecondType)
 		{
@@ -239,6 +303,25 @@ namespace PCG
 			return IsBroadcastable(MetadataTypes<FirstType>::Id, MetadataTypes<SecondType>::Id);
 		}
 
+		constexpr inline bool IsConstructible(uint16 FirstType, uint16 SecondType)
+		{
+			// Unknown types can't be checked for constructible.
+			if (FirstType >= static_cast<uint16>(EPCGMetadataTypes::Count) || SecondType >= static_cast<uint16>(EPCGMetadataTypes::Count))
+			{
+				return false;
+			}
+
+			return ConstructibleTypes[FirstType][SecondType];
+		}
+
+		// Convenience function for accessor operation, to know if a Get will succeed with AllowBroadcast and AllowConstructible
+		// where the SecondType is the target type, and FirstType is the accessor type.
+		// ie. Accessor is an int, Accessor->Get<float>(...). First type is int, second type is float.
+		constexpr inline bool IsBroadcastableOrConstructible(uint16 FirstType, uint16 SecondType)
+		{
+			return IsBroadcastable(FirstType, SecondType) || IsConstructible(SecondType, FirstType);
+		}
+
 		/**
 		* FirstType is more complex than SecondType, in the sense of PCG types, if both types are different and valid
 		* and we can broadcast SecondType to FirstType.
@@ -989,3 +1072,6 @@ namespace PCG
 		}
 	}
 }
+
+// Undef this one since we don't use it outside this file. But the other is used outside so keep it defined.
+#undef PCG_FOREACH_SUPPORTEDTYPES_WITH_COMMA
\ No newline at end of file
