Added options around tags propagation for copy points and projection

#rnx

diff --git a/Engine/Plugins/Experimental/PCG/Source/PCG/Private/Elements/PCGCopyPoints.cpp b/Engine/Plugins/Experimental/PCG/Source/PCG/Private/Elements/PCGCopyPoints.cpp
--- a/Engine/Plugins/Experimental/PCG/Source/PCG/Private/Elements/PCGCopyPoints.cpp
+++ b/Engine/Plugins/Experimental/PCG/Source/PCG/Private/Elements/PCGCopyPoints.cpp
@@ -93,7 +93,14 @@ bool FPCGCopyPointsElement::ExecuteInternal(FPCGContext* Context) const
 			Output.Data = OutPointData;
 
 			// Make sure that output contains both collection of tags from source and target
-			Output.Tags.Append(Target.Tags);
+			if (Settings->TagInheritance == EPCGCopyPointsTagInheritanceMode::Target)
+			{
+				Output.Tags = Target.Tags;
+			}
+			else if (Settings->TagInheritance == EPCGCopyPointsTagInheritanceMode::Both)
+			{
+				Output.Tags.Append(Target.Tags);
+			}
 
 			// RootMetadata will be parent to the ouptut metadata, while NonRootMetadata will carry attributes from the input not selected for inheritance
 			// Note that this is a preference, as we can and should pick more efficiently in the trivial cases
diff --git a/Engine/Plugins/Experimental/PCG/Source/PCG/Private/Elements/PCGProjectionElement.cpp b/Engine/Plugins/Experimental/PCG/Source/PCG/Private/Elements/PCGProjectionElement.cpp
--- a/Engine/Plugins/Experimental/PCG/Source/PCG/Private/Elements/PCGProjectionElement.cpp
+++ b/Engine/Plugins/Experimental/PCG/Source/PCG/Private/Elements/PCGProjectionElement.cpp
@@ -93,7 +93,15 @@ bool FPCGProjectionElement::ExecuteInternal(FPCGContext* Context) const
 
 		FPCGTaggedData& ProjectionTaggedData = Outputs.Emplace_GetRef(Source);
 		ProjectionTaggedData.Data = ProjectionData;
-		ProjectionTaggedData.Tags.Append(Targets[0].Tags);
+
+		if (ProjectionParams.TagMergeOperation == EPCGProjectionTagMergeMode::Target)
+		{
+			ProjectionTaggedData.Tags = Targets[0].Tags;
+		}
+		else if(ProjectionParams.TagMergeOperation == EPCGProjectionTagMergeMode::Both)
+		{
+			ProjectionTaggedData.Tags.Append(Targets[0].Tags);
+		}
 	}
 
 	return true;
diff --git a/Engine/Plugins/Experimental/PCG/Source/PCG/Public/Elements/PCGCopyPoints.h b/Engine/Plugins/Experimental/PCG/Source/PCG/Public/Elements/PCGCopyPoints.h
--- a/Engine/Plugins/Experimental/PCG/Source/PCG/Public/Elements/PCGCopyPoints.h
+++ b/Engine/Plugins/Experimental/PCG/Source/PCG/Public/Elements/PCGCopyPoints.h
@@ -20,6 +20,14 @@ enum class EPCGCopyPointsInheritanceMode : uint8
 	Target
 };
 
+UENUM()
+enum class EPCGCopyPointsTagInheritanceMode : uint8
+{
+	Both,
+	Source,
+	Target,
+};
+
 UENUM()
 enum class EPCGCopyPointsMetadataInheritanceMode : uint8
 {
@@ -70,6 +78,10 @@ public:
 	/** The method used to determine output data attributes */
 	UPROPERTY(BlueprintReadWrite, EditAnywhere, Category = Settings, meta = (PCG_Overridable))
 	EPCGCopyPointsMetadataInheritanceMode AttributeInheritance = EPCGCopyPointsMetadataInheritanceMode::SourceFirst;
+
+	/** The method used to determine the output data tags */
+	UPROPERTY(BlueprintReadWrite, EditAnywhere, Category = Settings, meta = (PCG_Overridable))
+	EPCGCopyPointsTagInheritanceMode TagInheritance = EPCGCopyPointsTagInheritanceMode::Both;
 };
 
 class FPCGCopyPointsElement : public IPCGElement
diff --git a/Engine/Plugins/Experimental/PCG/Source/PCG/Public/Elements/PCGProjectionParams.h b/Engine/Plugins/Experimental/PCG/Source/PCG/Public/Elements/PCGProjectionParams.h
--- a/Engine/Plugins/Experimental/PCG/Source/PCG/Public/Elements/PCGProjectionParams.h
+++ b/Engine/Plugins/Experimental/PCG/Source/PCG/Public/Elements/PCGProjectionParams.h
@@ -16,6 +16,14 @@ enum class EPCGProjectionColorBlendMode : uint8
 	Multiply
 };
 
+UENUM()
+enum class EPCGProjectionTagMergeMode : uint8
+{
+	Source,
+	Target,
+	Both
+};
+
 /** Parameters that control projection behaviour. */
 USTRUCT(BlueprintType)
 struct PCG_API FPCGProjectionParams
@@ -72,6 +80,10 @@ struct PCG_API FPCGProjectionParams
 	UPROPERTY(BlueprintReadWrite, EditAnywhere, Category = "Apply Data")
 	EPCGMetadataOp AttributeMergeOperation = EPCGMetadataOp::TargetValue;
 
+	/** Controls whether the data tags are taken from the source, the target or both. */
+	UPROPERTY(BlueprintReadWrite, EditAnywhere, Category = "Apply Data")
+	EPCGProjectionTagMergeMode TagMergeOperation = EPCGProjectionTagMergeMode::Source;
+
 	// TODO [DEPRECATED_IN_5_4]: Remove the 'rule of five' needed above upon deprecation
 #if WITH_EDITORONLY_DATA
 	UE_DEPRECATED(5.4, "This property will be deprecated in 5.4. Please use `ColorBlendMode` instead.")
