[PCG] Made visual logs based on stacks and added component to stacks.

Remove function that attempted to build list of stacks from graph as this could not account for higen and in the future branch culling. Get stacks direct from compile cache.

Now that we use execution stacks to inform debug UI/tools, make the stack context per-compiled-graph because the stacks depend on compilation settings, like higen grid.

Display icons on tree view for warnings errors, for quick navigation. Display warnings/errors within subgraphs on the subgraph node.

Add missing refresh paths so warnings/errors update properly when changing debug object, and refresh properly after re-executing subgraphs.

Rewrite debug object tree build code to use stacks and to consolidate dynamic vs static tree building paths.

Bugs fixed along the way:

* Debug object selection sync'd bidirectionally between dropdown and tree.
* Debug object tree: Loop items ordered by loop index.
* Debug object tree: User item selection & expansion retained across refreshes.
* Supports nested combinations of static / dynamic subgraphs deeper than 1 level.
* Clear inspection flag from nodes when opening graph editor to fix zombie state when reopening inspected graph.

Open bugs:

* Profiling window is still quite broken for subgraphs

diff --git a/Engine/Plugins/Experimental/PCG/Source/PCG/Private/Elements/PCGLoopElement.cpp b/Engine/Plugins/Experimental/PCG/Source/PCG/Private/Elements/PCGLoopElement.cpp
--- a/Engine/Plugins/Experimental/PCG/Source/PCG/Private/Elements/PCGLoopElement.cpp
+++ b/Engine/Plugins/Experimental/PCG/Source/PCG/Private/Elements/PCGLoopElement.cpp
@@ -268,10 +268,6 @@ bool FPCGLoopElement::ExecuteInternal(FPCGContext* InContext) const
 			StackFrames.Emplace(Context->Node);
 			StackFrames.Emplace(EntryIndex);
 
-#if WITH_EDITOR
-			Subgraph->OnGraphDynamicallyExecutedDelegate.Broadcast(Subgraph, Context->SourceComponent, InvocationStack);
-#endif
-
 			TArray<FPCGTaskId> Dependencies = {};
 			if (PreviousTaskId != InvalidPCGTaskId && bIsFeedbackLoop)
 			{
diff --git a/Engine/Plugins/Experimental/PCG/Source/PCG/Private/Elements/PCGSpawnActor.cpp b/Engine/Plugins/Experimental/PCG/Source/PCG/Private/Elements/PCGSpawnActor.cpp
--- a/Engine/Plugins/Experimental/PCG/Source/PCG/Private/Elements/PCGSpawnActor.cpp
+++ b/Engine/Plugins/Experimental/PCG/Source/PCG/Private/Elements/PCGSpawnActor.cpp
@@ -543,10 +543,6 @@ bool FPCGSpawnActorElement::SpawnAndPrepareSubgraphs(FPCGSubgraphContext* Contex
 
 							UPCGGraph* Graph = GraphInterface->GetGraph();
 
-#if WITH_EDITOR
-							Graph->OnGraphDynamicallyExecutedDelegate.Broadcast(Graph, Context->SourceComponent, InvocationStack);
-#endif
-
 							FPCGTaskId SubgraphTaskId = Subsystem->ScheduleGraph(Graph,
 								Context->SourceComponent.Get(),
 								MakeShared<FPCGTrivialElement>(),// TODO: prepare user parameters like in subgraph/loop
diff --git a/Engine/Plugins/Experimental/PCG/Source/PCG/Private/Graph/PCGGraphCompiler.cpp b/Engine/Plugins/Experimental/PCG/Source/PCG/Private/Graph/PCGGraphCompiler.cpp
--- a/Engine/Plugins/Experimental/PCG/Source/PCG/Private/Graph/PCGGraphCompiler.cpp
+++ b/Engine/Plugins/Experimental/PCG/Source/PCG/Private/Graph/PCGGraphCompiler.cpp
@@ -1,6 +1,6 @@
 // Copyright Epic Games, Inc. All Rights Reserved.
 
-#include "PCGGraphCompiler.h"
+#include "Graph/PCGGraphCompiler.h"
 
 #include "PCGEdge.h"
 #include "PCGGraph.h"
@@ -287,21 +287,17 @@ TArray<FPCGGraphTask> FPCGGraphCompiler::GetPrecompiledTasks(const UPCGGraph* In
 		// Top graphs are optimized per grid size.
 		const TMap<uint32, TArray<FPCGGraphTask>>* GridSizeToCompiledGraph = TopGraphToTaskMap.Find(InGraph);
 		ExistingTasks = GridSizeToCompiledGraph ? GridSizeToCompiledGraph->Find(GenerationGridSize) : nullptr;
-	}
-	else
-	{
-		ExistingTasks = GraphToTaskMap.Find(InGraph);
-	}
 
-	const FPCGStackContext* StackContext = (bIsTopGraph ? TopGraphToStackContext : GraphToStackContext).Find(InGraph);
-	if (StackContext)
-	{
-		OutStackContext = *StackContext;
+		const TMap<uint32, FPCGStackContext>* GridSizeToStackContext = TopGraphToStackContextMap.Find(InGraph);
+		const FPCGStackContext* ExistingStackContext = GridSizeToStackContext ? GridSizeToStackContext->Find(GenerationGridSize) : nullptr;
+		OutStackContext = ExistingStackContext ? *ExistingStackContext : FPCGStackContext();
+
+		// Should either find both, or find neither.
+		ensure(!ExistingTasks == !ExistingStackContext);
 	}
 	else
 	{
-		// If we failed to get a context, output a dummy blank one.
-		OutStackContext = FPCGStackContext();
+		ExistingTasks = GraphToTaskMap.Find(InGraph);
 	}
 
 	return ExistingTasks ? *ExistingTasks : TArray<FPCGGraphTask>();
@@ -589,6 +585,46 @@ void FPCGGraphCompiler::CullTasks(TArray<FPCGGraphTask>& InOutCompiledTasks, boo
 	}
 }
 
+void FPCGGraphCompiler::PostCullStackCleanup(TArray<FPCGGraphTask>& InCompiledTasks, FPCGStackContext& InOutStackContext)
+{
+	// Build set of stack IDs used by tasks. Using array as set is likely small.
+	TArray<int32> ActiveStackIDs;
+	for (FPCGGraphTask& Task : InCompiledTasks)
+	{
+		ActiveStackIDs.AddUnique(Task.StackIndex);
+	}
+
+	if (ActiveStackIDs.Num() < InOutStackContext.GetNumStacks())
+	{
+		// Prepare cumulative counts for stack ID remapping.
+		TArray<int> RemovedCounts;
+		RemovedCounts.SetNumZeroed(InOutStackContext.GetNumStacks());
+
+		// Remove any stacks that are inactive.
+		TArray<FPCGStack>& Stacks = InOutStackContext.GetStacksMutable();
+		for (int StackIndex = Stacks.Num() - 1; StackIndex >= 0; --StackIndex)
+		{
+			if (!ActiveStackIDs.Contains(StackIndex))
+			{
+				Stacks.RemoveAt(StackIndex);
+				RemovedCounts[StackIndex] = 1;
+			}
+		}
+
+		// Accumulate removed counts.
+		for (int Index = 1; Index < RemovedCounts.Num(); ++Index)
+		{
+			RemovedCounts[Index] += RemovedCounts[Index - 1];
+		}
+
+		// Remap stack IDs.
+		for (FPCGGraphTask& Task : InCompiledTasks)
+		{
+			Task.StackIndex -= RemovedCounts[Task.StackIndex];
+		}
+	}
+}
+
 TArray<FPCGGraphTask> FPCGGraphCompiler::GetCompiledTasks(UPCGGraph* InGraph, uint32 GenerationGridSize, FPCGStackContext& OutStackContext, bool bIsTopGraph)
 {
 	TArray<FPCGGraphTask> CompiledTasks;
@@ -602,10 +638,17 @@ TArray<FPCGGraphTask> FPCGGraphCompiler::GetCompiledTasks(UPCGGraph* InGraph, ui
 		FReadScopeLock Lock(GraphToTaskMapLock);
 		const TMap<uint32, TArray<FPCGGraphTask>>* GridSizeToCompiledGraph = TopGraphToTaskMap.Find(InGraph);
 		const TArray<FPCGGraphTask>* Tasks = GridSizeToCompiledGraph ? GridSizeToCompiledGraph->Find(GenerationGridSize) : nullptr;
-		if (Tasks)
+
+		const TMap<uint32, FPCGStackContext>* GridSizeToStackContext = TopGraphToStackContextMap.Find(InGraph);
+		const FPCGStackContext* StackContext = GridSizeToStackContext ? GridSizeToStackContext->Find(GenerationGridSize) : nullptr;
+
+		// Should have either found both or neither.
+		ensure(!Tasks == !StackContext);
+
+		if (Tasks && StackContext)
 		{
 			CompiledTasks = *Tasks;
-			OutStackContext = TopGraphToStackContext[InGraph];
+			OutStackContext = *StackContext;
 		}
 	}
 	else
@@ -694,6 +737,10 @@ void FPCGGraphCompiler::CompileTopGraph(UPCGGraph* InGraph, uint32 GenerationGri
 		});
 	}
 
+	// Post culling - remove any stacks that are no longer part of execution. Besides being tidy this also helps
+	// debug tools discern which stacks were executed or not.
+	PostCullStackCleanup(CompiledTasks, StackContext);
+
 	const int TaskNum = CompiledTasks.Num();
 	const FPCGTaskId PreExecuteTaskId = FPCGTaskId(TaskNum);
 	const FPCGTaskId PostExecuteTaskId = PreExecuteTaskId + 1;
@@ -748,7 +795,12 @@ void FPCGGraphCompiler::CompileTopGraph(UPCGGraph* InGraph, uint32 GenerationGri
 	if (!TasksPerGenerationGrid.Contains(GenerationGridSize))
 	{
 		TasksPerGenerationGrid.Add(GenerationGridSize, MoveTemp(CompiledTasks));
-		TopGraphToStackContext.Add(InGraph, StackContext);
+	}
+
+	TMap<uint32, FPCGStackContext>& StackContextPerGenerationGrid = TopGraphToStackContextMap.FindOrAdd(InGraph);
+	if (!StackContextPerGenerationGrid.Contains(GenerationGridSize))
+	{
+		StackContextPerGenerationGrid.Add(GenerationGridSize, MoveTemp(StackContext));
 	}
 	GraphToTaskMapLock.WriteUnlock();
 }
@@ -759,7 +811,7 @@ void FPCGGraphCompiler::ClearCache()
 	GraphToTaskMap.Reset();
 	GraphToStackContext.Reset();
 	TopGraphToTaskMap.Reset();
-	TopGraphToStackContext.Reset();
+	TopGraphToStackContextMap.Reset();
 }
 
 #if WITH_EDITOR
@@ -781,7 +833,9 @@ void FPCGGraphCompiler::RemoveFromCacheRecursive(UPCGGraph* InGraph)
 {
 	GraphToTaskMapLock.WriteLock();
 	GraphToTaskMap.Remove(InGraph);
+	GraphToStackContext.Remove(InGraph);
 	TopGraphToTaskMap.Remove(InGraph);
+	TopGraphToStackContextMap.Remove(InGraph);
 	GraphToTaskMapLock.WriteUnlock();
 
 	GraphDependenciesLock.Lock();
diff --git a/Engine/Plugins/Experimental/PCG/Source/PCG/Private/Graph/PCGGraphCompiler.h b/Engine/Plugins/Experimental/PCG/Source/PCG/Private/Graph/PCGGraphCompiler.h
--- a/Engine/Plugins/Experimental/PCG/Source/PCG/Private/Graph/PCGGraphCompiler.h
+++ b/Engine/Plugins/Experimental/PCG/Source/PCG/Private/Graph/PCGGraphCompiler.h
@@ -58,12 +58,16 @@ private:
 	/** Culls tasks based on a given lambda. Never culls the first (input) task in the array. */
 	static void CullTasks(TArray<FPCGGraphTask>& InOutCompiledTasks, bool bAddPassthroughWires, TFunctionRef<bool(const FPCGGraphTask&)> CullTask);
 
+	/** Remove any stack frames that are not used by any task. */
+	static void PostCullStackCleanup(TArray<FPCGGraphTask>& InCompiledTasks, FPCGStackContext& InOutStackContext);
+
 	mutable FRWLock GraphToTaskMapLock;
 	TMap<UPCGGraph*, TArray<FPCGGraphTask>> GraphToTaskMap;
 	TMap<UPCGGraph*, FPCGStackContext> GraphToStackContext;
+
 	// Top graphs are optimized for execution grid and store one set of compiled tasks per grid size.
 	TMap<UPCGGraph*, TMap<uint32, TArray<FPCGGraphTask>>> TopGraphToTaskMap;
-	TMap<UPCGGraph*, FPCGStackContext> TopGraphToStackContext;
+	TMap<UPCGGraph*, TMap<uint32, FPCGStackContext>> TopGraphToStackContextMap;
 
 #if WITH_EDITOR
 	void RemoveFromCache(UPCGGraph* InGraph);
diff --git a/Engine/Plugins/Experimental/PCG/Source/PCG/Private/Graph/PCGGraphExecutor.cpp b/Engine/Plugins/Experimental/PCG/Source/PCG/Private/Graph/PCGGraphExecutor.cpp
--- a/Engine/Plugins/Experimental/PCG/Source/PCG/Private/Graph/PCGGraphExecutor.cpp
+++ b/Engine/Plugins/Experimental/PCG/Source/PCG/Private/Graph/PCGGraphExecutor.cpp
@@ -7,7 +7,6 @@
 #include "PCGContext.h"
 #include "PCGCrc.h"
 #include "PCGGraph.h"
-#include "PCGGraphCompiler.h"
 #include "PCGInputOutputSettings.h"
 #include "PCGModule.h"
 #include "PCGParamData.h"
@@ -15,6 +14,7 @@
 #include "PCGSubsystem.h"
 #include "PCGWorldActor.h"
 #include "Graph/PCGGraphCache.h"
+#include "Graph/PCGGraphCompiler.h"
 #include "Graph/PCGStackContext.h"
 #include "Grid/PCGPartitionActor.h"
 #include "Helpers/PCGActorHelpers.h"
@@ -111,18 +111,7 @@ FPCGTaskId FPCGGraphExecutor::Schedule(UPCGGraph* Graph, UPCGComponent* SourceCo
 	{
 		UE_LOG(LogPCG, Log, TEXT("[%s] --- SCHEDULE GRAPH ---"), *SourceComponent->GetOwner()->GetName());
 	}
-
-#if WITH_EDITOR
-	if (UPCGSubsystem* Subsystem = UPCGSubsystem::GetInstance(SourceComponent->GetWorld()))
-	{
-		for (const UPCGNode* Node : Graph->GetNodes())
-		{
-			// Always clear warnings/errors before compile regardless of connectivity
-			Subsystem->GetNodeVisualLogsMutable().ClearLogs(Node, SourceComponent);
-		}
-	}
-#endif
-
+	
 	FPCGTaskId ScheduledId = InvalidPCGTaskId;
 
 	const bool bNonPartitionedComponent = !SourceComponent->IsLocalComponent() && !SourceComponent->IsPartitioned();
@@ -137,6 +126,19 @@ FPCGTaskId FPCGGraphExecutor::Schedule(UPCGGraph* Graph, UPCGComponent* SourceCo
 	{
 		StackContextPtr->PrependParentStack(InFromStack);
 	}
+	else
+	{
+		FPCGStack ComponentStack;
+		ComponentStack.PushFrame(SourceComponent);
+		StackContextPtr->PrependParentStack(&ComponentStack);		
+	}
+
+#if WITH_EDITOR
+	if (UPCGSubsystem* Subsystem = SourceComponent ? SourceComponent->GetSubsystem() : nullptr)
+	{
+		Subsystem->OnScheduleGraph(*StackContextPtr);
+	}
+#endif
 
 	// Assign this component to the tasks
 	for (FPCGGraphTask& Task : CompiledTasks)
@@ -1554,8 +1556,7 @@ namespace PCGGraphExecutor
 				if (ToGridSize == PCGHiGenGrid::UnboundedGridSize())
 				{
 					Subsystem->GetNodeVisualLogsMutable().Log(
-						InDownstreamNode,
-						InContext->SourceComponent,
+						*InContext->Stack,
 						ELogVerbosity::Error,
 						FText::Format(
 							NSLOCTEXT("PCGGraphCompiler", "InvalidLinkageToUnbounded", "Could not read data across grid levels - cannot read from grid size {0} to Unbounded domain."),
@@ -1565,8 +1566,7 @@ namespace PCGGraphExecutor
 				else
 				{
 					Subsystem->GetNodeVisualLogsMutable().Log(
-						InDownstreamNode,
-						InContext->SourceComponent,
+						*InContext->Stack,
 						ELogVerbosity::Error,
 						FText::Format(
 							NSLOCTEXT("PCGGraphCompiler", "InvalidLinkageInvalidGridSizes", "Could not read data across grid levels - origin grid size {0} must be greater than destination grid size {1}. Graph default grid size may need increasing."),
diff --git a/Engine/Plugins/Experimental/PCG/Source/PCG/Private/Graph/PCGStackContext.cpp b/Engine/Plugins/Experimental/PCG/Source/PCG/Private/Graph/PCGStackContext.cpp
--- a/Engine/Plugins/Experimental/PCG/Source/PCG/Private/Graph/PCGStackContext.cpp
+++ b/Engine/Plugins/Experimental/PCG/Source/PCG/Private/Graph/PCGStackContext.cpp
@@ -2,11 +2,13 @@
 
 #include "Graph/PCGStackContext.h"
 
+#include "PCGComponent.h"
 #include "PCGGraph.h"
 #include "PCGNode.h"
 #include "PCGPin.h"
 #include "PCGSubgraph.h"
 
+#include "Algo/Find.h"
 #include "Containers/UnrealString.h"
 #include "Misc/StringBuilder.h"
 
@@ -58,7 +60,11 @@ bool FPCGStack::CreateStackFramePath(FString& OutString, const UPCGNode* InNode,
 				return false;
 			}
 
-			if (Object->IsA<UPCGGraph>())
+			if (Object->IsA<UPCGComponent>())
+			{
+				StringBuilder << TEXT("COMPONENT:") << Object->GetFullName();
+			}
+			else if (Object->IsA<UPCGGraph>())
 			{
 				StringBuilder << TEXT("GRAPH:") << Object->GetFullName();
 			}
@@ -101,23 +107,22 @@ uint32 FPCGStack::GetNumGraphLevels() const
 	uint32 GraphCount = 0;
 	for (const FPCGStackFrame& Frame : StackFrames)
 	{
-		GraphCount += Frame.Object->IsA<UPCGGraph>() ? 1 : 0;
+		GraphCount += (Frame.Object.IsValid() && Frame.Object->IsA<UPCGGraph>()) ? 1 : 0;
 	}
 
 	return GraphCount;
 }
 
-bool FPCGStack::operator==(const FPCGStack& Other) const
+bool FPCGStack::BeginsWith(const FPCGStack& Other) const
 {
-	// Stacks are the same if all stack frames are the same
-	if (StackFrames.Num() != Other.StackFrames.Num())
+	if (Other.GetStackFrames().Num() > GetStackFrames().Num())
 	{
 		return false;
 	}
 
-	for (int32 i = 0; i < StackFrames.Num(); i++)
+	for (int32 StackIndex = 0; StackIndex < Other.GetStackFrames().Num(); ++StackIndex)
 	{
-		if (StackFrames[i] != Other.StackFrames[i])
+		if (Other.GetStackFrames()[StackIndex] != GetStackFrames()[StackIndex])
 		{
 			return false;
 		}
@@ -126,16 +131,67 @@ bool FPCGStack::operator==(const FPCGStack& Other) const
 	return true;
 }
 
-uint32 GetTypeHash(const FPCGStack& In)
+const UPCGComponent* FPCGStack::GetRootComponent() const
+{
+	return StackFrames.IsEmpty() ? nullptr : Cast<const UPCGComponent>(StackFrames[0].Object.Get());
+}
+
+const UPCGGraph* FPCGStack::GetRootGraph() const
 {
-	uint32 Hash = 0;
+	for (int StackIndex = 0; StackIndex < GetStackFrames().Num(); ++StackIndex)
+	{
+		if (const UPCGGraph* Graph = Cast<const UPCGGraph>(StackFrames[StackIndex].Object.Get()))
+		{
+			return Graph;
+		}
+	}
+
+	return nullptr;
+}
 
-	for (const FPCGStackFrame& Frame : In.StackFrames)
+const UPCGGraph* FPCGStack::GetGraphForCurrentFrame() const
+{
+	for (int StackIndex = GetStackFrames().Num() - 1; StackIndex >= 0; --StackIndex)
 	{
-		Hash = HashCombine(Hash, GetTypeHash(Frame));
+		if (const UPCGGraph* Graph = Cast<const UPCGGraph>(StackFrames[StackIndex].Object.Get()))
+		{
+			return Graph;
+		}
 	}
 
-	return Hash;
+	return nullptr;
+}
+
+const UPCGNode* FPCGStack::GetCurrentFrameNode() const
+{
+	return StackFrames.IsEmpty() ? nullptr : Cast<const UPCGNode>(StackFrames.Last().Object.Get());
+}
+
+bool FPCGStack::HasObject(const UObject* InObject) const
+{
+	return !!Algo::FindByPredicate(StackFrames, [InObject](const FPCGStackFrame& Frame)
+	{
+		return Frame.Object == InObject;
+	});
+}
+
+bool FPCGStack::operator==(const FPCGStack& Other) const
+{
+	// Stacks are the same if all stack frames are the same
+	if (StackFrames.Num() != Other.StackFrames.Num())
+	{
+		return false;
+	}
+
+	for (int32 i = 0; i < StackFrames.Num(); i++)
+	{
+		if (StackFrames[i] != Other.StackFrames[i])
+		{
+			return false;
+		}
+	}
+
+	return true;
 }
 
 int32 FPCGStackContext::PushFrame(const UObject* InFrameObject)
@@ -144,7 +200,7 @@ int32 FPCGStackContext::PushFrame(const UObject* InFrameObject)
 	{
 		// Create first stack using the given frame.
 		FPCGStack& Stack = Stacks.Emplace_GetRef();
-		Stack.PushFrame(FPCGStackFrame(InFrameObject));
+		Stack.PushFrame(InFrameObject);
 		CurrentStackIndex = 0;
 	}
 	else
@@ -225,53 +281,3 @@ void FPCGStackContext::PrependParentStack(const FPCGStack* InParentStack)
 		Stack.StackFrames.Insert(InParentStack->StackFrames, 0);
 	}
 }
-
-FPCGStackContext FPCGStackContext::CreateStackContextFromGraph(const UPCGGraph* InPCGGraph)
-{
-	FPCGStackContext StackContext;
-
-	auto ParseGraphRecursive = [](const UPCGGraph* InPCGGraph, FPCGStackContext& InStackContext, TArray<const UPCGGraph*>& VisitedGraphStack, auto RecursiveCallback)
-	{
-		if (!InPCGGraph)
-		{
-			return;
-		}
-
-		InStackContext.PushFrame(InPCGGraph);
-
-		for (const UPCGNode* PCGNode : InPCGGraph->GetNodes())
-		{
-			if (PCGNode)
-			{
-				// TODO: GetSettings() has no execution context and therefore cannot recurse on dynamically chosen subgraphs
-				if (const UPCGBaseSubgraphSettings* SubgraphSettings = Cast<UPCGBaseSubgraphSettings>(PCGNode->GetSettings()))
-				{
-					if (const UPCGGraph* PCGSubgraph = SubgraphSettings->GetSubgraph())
-					{
-						// Skip graphs we have already visited to avoid cycles
-						// TODO: This prevents recursive subgraphs
-						if (!VisitedGraphStack.Contains(PCGSubgraph))
-						{
-							InStackContext.PushFrame(PCGNode);
-							VisitedGraphStack.Push(PCGSubgraph);
-
-							FPCGStackContext SubgraphStackContext;
-							RecursiveCallback(PCGSubgraph, SubgraphStackContext, VisitedGraphStack, RecursiveCallback);
-							InStackContext.AppendStacks(SubgraphStackContext);
-
-							VisitedGraphStack.Pop();
-							InStackContext.PopFrame();
-						}
-					}
-				}
-			}
-		}
-	};
-
-	TArray<const UPCGGraph*> VisitedGraphStack;
-	VisitedGraphStack.Push(InPCGGraph);
-
-	ParseGraphRecursive(InPCGGraph, StackContext, VisitedGraphStack, ParseGraphRecursive);
-
-	return StackContext;
-}
diff --git a/Engine/Plugins/Experimental/PCG/Source/PCG/Private/PCGComponent.cpp b/Engine/Plugins/Experimental/PCG/Source/PCG/Private/PCGComponent.cpp
--- a/Engine/Plugins/Experimental/PCG/Source/PCG/Private/PCGComponent.cpp
+++ b/Engine/Plugins/Experimental/PCG/Source/PCG/Private/PCGComponent.cpp
@@ -442,6 +442,8 @@ void UPCGComponent::PostProcessGraph(const FBox& InNewBounds, bool bInGenerated,
 	}
 
 	StopGenerationInProgress();
+
+	GetSubsystem()->OnComponentGenerationCompleteOrCancelled.Broadcast();
 #endif
 }
 
@@ -538,6 +540,8 @@ void UPCGComponent::OnProcessGraphAborted(bool bQuiet)
 	// the component is still considered dirty if we aborted processing, hence it should stay this way.
 
 	StopGenerationInProgress();
+
+	GetSubsystem()->OnComponentGenerationCompleteOrCancelled.Broadcast();
 #endif
 }
 
@@ -2826,6 +2830,20 @@ void UPCGComponent::ChangeTransientState(EPCGEditorDirtyMode NewEditingMode)
 		GEditor->Trans->Reset(LOCTEXT("ChangeEditingMode", "Changing Editing Mode"));
 	}
 }
+
+bool UPCGComponent::GetStackContext(FPCGStackContext& OutStackContext) const
+{
+	const UPCGSubsystem* Subsystem = GetSubsystem();
+	if (Subsystem && Subsystem->GetStackContext(this, OutStackContext))
+	{
+		FPCGStack ComponentStack;
+		ComponentStack.PushFrame(this);
+		OutStackContext.PrependParentStack(&ComponentStack);
+		return true;
+	}
+
+	return false;
+}
 #endif // WITH_EDITOR
 
 FPCGComponentInstanceData::FPCGComponentInstanceData(const UPCGComponent* InSourceComponent)
diff --git a/Engine/Plugins/Experimental/PCG/Source/PCG/Private/PCGContext.cpp b/Engine/Plugins/Experimental/PCG/Source/PCG/Private/PCGContext.cpp
--- a/Engine/Plugins/Experimental/PCG/Source/PCG/Private/PCGContext.cpp
+++ b/Engine/Plugins/Experimental/PCG/Source/PCG/Private/PCGContext.cpp
@@ -278,7 +278,9 @@ void FPCGContext::LogVisual(ELogVerbosity::Type InVerbosity, const FText& InMess
 
 	if (UPCGSubsystem* Subsystem = UPCGSubsystem::GetInstance(SourceComponent->GetWorld()))
 	{
-		Subsystem->GetNodeVisualLogsMutable().Log(Node, SourceComponent, InVerbosity, InMessage);
+		FPCGStack StackWithNode = *Stack;
+		StackWithNode.PushFrame(Node);
+		Subsystem->GetNodeVisualLogsMutable().Log(StackWithNode, InVerbosity, InMessage);
 	}
 }
 
@@ -289,9 +291,11 @@ bool FPCGContext::HasVisualLogs() const
 		return false;
 	}
 
-	if (UPCGSubsystem* Subsystem = UPCGSubsystem::GetInstance(SourceComponent->GetWorld()))
+	if (const UPCGSubsystem* Subsystem = UPCGSubsystem::GetInstance(SourceComponent->GetWorld()))
 	{
-		return Subsystem->GetNodeVisualLogs().HasLogs(Node, SourceComponent.Get());
+		FPCGStack StackWithNode = *Stack;
+		StackWithNode.PushFrame(Node);
+		return Subsystem->GetNodeVisualLogs().HasLogs(StackWithNode);
 	}
 
 	return false;
diff --git a/Engine/Plugins/Experimental/PCG/Source/PCG/Private/PCGSubgraph.cpp b/Engine/Plugins/Experimental/PCG/Source/PCG/Private/PCGSubgraph.cpp
--- a/Engine/Plugins/Experimental/PCG/Source/PCG/Private/PCGSubgraph.cpp
+++ b/Engine/Plugins/Experimental/PCG/Source/PCG/Private/PCGSubgraph.cpp
@@ -526,10 +526,6 @@ bool FPCGSubgraphElement::ExecuteInternal(FPCGContext* InContext) const
 				FPCGStack InvocationStack = ensure(Context->Stack) ? *Context->Stack : FPCGStack();
 				InvocationStack.GetStackFramesMutable().Emplace(Context->Node);
 
-#if WITH_EDITOR
-				Subgraph->OnGraphDynamicallyExecutedDelegate.Broadcast(Subgraph, Context->SourceComponent, InvocationStack);
-#endif
-
 				FPCGTaskId SubgraphTaskId = Subsystem->ScheduleGraph(Subgraph, Context->SourceComponent.Get(), MakeShared<FPCGInputForwardingElement>(PreSubgraphInputData), MakeShared<FPCGInputForwardingElement>(SubgraphInputData), {}, &InvocationStack);
 
 				if (SubgraphTaskId != InvalidPCGTaskId)
diff --git a/Engine/Plugins/Experimental/PCG/Source/PCG/Private/PCGSubsystem.cpp b/Engine/Plugins/Experimental/PCG/Source/PCG/Private/PCGSubsystem.cpp
--- a/Engine/Plugins/Experimental/PCG/Source/PCG/Private/PCGSubsystem.cpp
+++ b/Engine/Plugins/Experimental/PCG/Source/PCG/Private/PCGSubsystem.cpp
@@ -428,6 +428,36 @@ void UPCGSubsystem::OnOriginalComponentUnregistered(UPCGComponent* InComponent)
 	}
 }
 
+#if WITH_EDITOR
+void UPCGSubsystem::OnScheduleGraph(const FPCGStackContext& StackContext)
+{
+	// Always clear any possibly related warnings/errors on schedule.
+	for (const FPCGStack& Stack : StackContext.GetStacks())
+	{
+		GetNodeVisualLogsMutable().ClearLogs(Stack);
+	}
+
+	// Flush out all stacks that begin from the component / top graph as the existing dynamic
+	// stacks may not be occur during the next execution.
+	if (const FPCGStack* BaseStack = StackContext.GetStack(0))
+	{
+		if (BaseStack->IsCurrentFrameInRootGraph())
+		{
+			ClearExecutedStacks(*BaseStack);
+		}
+	}
+
+	// Record executed stacks.
+	{
+		FWriteScopeLock Lock(ExecutedStacksLock);
+		for (const FPCGStack& ExecutedStack : StackContext.GetStacks())
+		{
+			ExecutedStacks.AddUnique(ExecutedStack);
+		}
+	}
+}
+#endif
+
 UPCGLandscapeCache* UPCGSubsystem::GetLandscapeCache()
 {
 	APCGWorldActor* LandscapeCacheOwner = GetPCGWorldActor();
@@ -1444,7 +1474,7 @@ void UPCGSubsystem::ClearLandscapeCache()
 	}
 }
 
-const FPCGGraphCompiler* UPCGSubsystem::GetGraphCompiler() const
+FPCGGraphCompiler* UPCGSubsystem::GetGraphCompiler() const
 {
 	if (GraphExecutor)
 	{
@@ -1454,11 +1484,54 @@ const FPCGGraphCompiler* UPCGSubsystem::GetGraphCompiler() const
 	return nullptr;
 }
 
+bool UPCGSubsystem::GetStackContext(const UPCGComponent* InComponent, FPCGStackContext& OutStackContext) const
+{
+	if (InComponent && InComponent->GetGraph())
+	{
+		GetGraphCompiler()->GetCompiledTasks(InComponent->GetGraph(), InComponent->GetGenerationGridSize(), OutStackContext);
+		return true;
+	}
+
+	return false;
+}
+
 uint32 UPCGSubsystem::GetGraphCacheEntryCount(IPCGElement* InElement) const
 {
 	return GraphExecutor ? GraphExecutor->GetGraphCacheEntryCount(InElement) : 0;
 }
 
+TArray<FPCGStack> UPCGSubsystem::GetExecutedStacks(const UPCGComponent* InComponent, const UPCGGraph* InSubgraph)
+{
+	FReadScopeLock Lock(ExecutedStacksLock);
+
+	TArray<FPCGStack> MatchingStacks;
+
+	for (const FPCGStack& Stack : ExecutedStacks)
+	{
+		if (Stack.GetRootComponent() == InComponent && Stack.GetGraphForCurrentFrame() == InSubgraph)
+		{
+			MatchingStacks.Add(Stack);
+		}
+	}
+
+	return MatchingStacks;
+}
+
+void UPCGSubsystem::ClearExecutedStacks(FPCGStack BeginningWithStack)
+{
+	FWriteScopeLock Lock(ExecutedStacksLock);
+
+	for (int StackIndex = ExecutedStacks.Num() - 1; StackIndex >= 0; --StackIndex)
+	{
+		// Clear any dynamic stack that starts with the given base stack, we don't know what stacks will
+		// execute so clean out everything.
+		if (ExecutedStacks[StackIndex].BeginsWith(BeginningWithStack))
+		{
+			ExecutedStacks.RemoveAtSwap(StackIndex);
+		}
+	}
+}
+
 #endif // WITH_EDITOR
 
 void UPCGSubsystem::FlushCache()
diff --git a/Engine/Plugins/Experimental/PCG/Source/PCG/Private/Utils/PCGNodeVisualLogs.cpp b/Engine/Plugins/Experimental/PCG/Source/PCG/Private/Utils/PCGNodeVisualLogs.cpp
--- a/Engine/Plugins/Experimental/PCG/Source/PCG/Private/Utils/PCGNodeVisualLogs.cpp
+++ b/Engine/Plugins/Experimental/PCG/Source/PCG/Private/Utils/PCGNodeVisualLogs.cpp
@@ -11,187 +11,216 @@
 
 #if WITH_EDITOR
 
-void FPCGNodeVisualLogs::Log(TWeakObjectPtr<const UPCGNode> InNode, TWeakObjectPtr<UPCGComponent> InComponent, ELogVerbosity::Type InVerbosity, const FText& InMessage)
+void FPCGNodeVisualLogs::Log(const FPCGStack& InPCGStack, ELogVerbosity::Type InVerbosity, const FText& InMessage)
 {
-	// InNode can be null, in case of unit tests.
-	if (!InNode.Get() || !ensure(InComponent.Get()))
-	{
-		return;
-	}
-
 	bool bAdded = false;
 
 	{
 		FWriteScopeLock ScopedWriteLock(LogsLock);
 
-		FPCGPerNodeVisualLogs* NodeLogs = NodeToLogs.Find(InNode);
-		if (!NodeLogs)
-		{
-			NodeLogs = &NodeToLogs.Add(InNode);
-		}
+		FPCGPerNodeVisualLogs& NodeLogs = StackToLogs.FindOrAdd(InPCGStack);
 
 		constexpr int32 MaxLogged = 1024;
-		if (NodeToLogs.Num() < MaxLogged)
+		if (StackToLogs.Num() < MaxLogged)
 		{
-			NodeLogs->Emplace(InMessage, InVerbosity, InComponent);
+			NodeLogs.Emplace(InMessage, InVerbosity);
 
 			bAdded = true;
 		}
 	}
 
 	// Broadcast outside of write scope lock
-	if (bAdded && IsInGameThread())
+	if (bAdded && IsInGameThread() && !InPCGStack.GetStackFrames().IsEmpty())
 	{
-		InNode->OnNodeChangedDelegate.Broadcast(const_cast<UPCGNode*>(InNode.Get()), EPCGChangeType::Cosmetic);
+		for (const FPCGStackFrame& Frame : InPCGStack.GetStackFrames())
+		{
+			if (const UPCGNode* Node = Cast<const UPCGNode>(Frame.Object.Get()))
+			{
+				Node->OnNodeChangedDelegate.Broadcast(const_cast<UPCGNode*>(Node), EPCGChangeType::Cosmetic);
+			}
+		}
 	}
 }
 
-bool FPCGNodeVisualLogs::HasLogs(TWeakObjectPtr<const UPCGNode> InNode, const UPCGComponent* InComponent) const
+bool FPCGNodeVisualLogs::HasLogs(const FPCGStack& InPCGStack) const
 {
 	FReadScopeLock ScopedReadLock(LogsLock);
 
-	if (!ensure(InNode.Get()))
+	for (const TPair<FPCGStack, FPCGPerNodeVisualLogs>& Entry : StackToLogs)
 	{
-		return false;
-	}
-
-	const FPCGPerNodeVisualLogs* NodeLogs = NodeToLogs.Find(InNode);
-	if (!NodeLogs)
-	{
-		return false;
-	}
-
-	if (!InComponent)
-	{
-		return !NodeLogs->IsEmpty();
-	}
-	else
-	{
-		return !!Algo::FindByPredicate(*NodeLogs, [InComponent](const FPCGNodeLogEntry& Log)
+		if (Entry.Key.BeginsWith(InPCGStack) && !Entry.Value.IsEmpty())
 		{
-			return Log.Component == InComponent;
-		});
+			return true;
+		}
 	}
+
+	return false;
 }
 
-bool FPCGNodeVisualLogs::HasLogs(TWeakObjectPtr<const UPCGNode> InNode, const UPCGComponent* InComponent, ELogVerbosity::Type InVerbosity) const
+bool FPCGNodeVisualLogs::HasLogs(const FPCGStack& InPCGStack, ELogVerbosity::Type& OutMinVerbosity) const
 {
+	OutMinVerbosity = ELogVerbosity::All;
+	
 	FReadScopeLock ScopedReadLock(LogsLock);
 
-	if (!ensure(InNode.Get()))
-	{
-		return false;
-	}
-
-	const FPCGPerNodeVisualLogs* NodeLogs = NodeToLogs.Find(InNode);
-	if (!NodeLogs)
+	for (const TPair<FPCGStack, FPCGPerNodeVisualLogs>& Entry : StackToLogs)
 	{
-		return false;
+		if (Entry.Key.BeginsWith(InPCGStack))
+		{
+			for (const FPCGNodeLogEntry& Log : Entry.Value)
+			{
+				OutMinVerbosity = FMath::Min(OutMinVerbosity, Log.Verbosity);
+			}
+		}
 	}
 
-	return !!Algo::FindByPredicate(*NodeLogs, [InVerbosity, InComponent](const FPCGNodeLogEntry& Log)
-	{
-		return Log.Verbosity == InVerbosity && (!InComponent || Log.Component == InComponent);
-	});
+	return OutMinVerbosity != ELogVerbosity::All;
 }
 
-FText FPCGNodeVisualLogs::GetLogsSummaryText(TWeakObjectPtr<const UPCGNode> InNode, const UPCGComponent* InComponent) const
+bool FPCGNodeVisualLogs::HasLogsOfVerbosity(const FPCGStack& InPCGStack, ELogVerbosity::Type InVerbosity) const
 {
 	FReadScopeLock ScopedReadLock(LogsLock);
 
-	if (!ensure(InNode.Get()))
+	for (const TPair<FPCGStack, FPCGPerNodeVisualLogs>& Entry : StackToLogs)
 	{
-		return FText::GetEmpty();
+		if (Entry.Key.BeginsWith(InPCGStack))
+		{
+			if (Algo::FindByPredicate(Entry.Value, [InVerbosity](const FPCGNodeLogEntry& Log) { return Log.Verbosity == InVerbosity; }))
+			{
+				return true;
+			}
+		}
 	}
 
-	const FPCGPerNodeVisualLogs* NodeLogsPtr = NodeToLogs.Find(InNode);
-	if (!NodeLogsPtr)
-	{
-		return FText::GetEmpty();
-	}
-	const FPCGPerNodeVisualLogs& NodeLogs = *NodeLogsPtr;
+	return false;
+}
+
+FText FPCGNodeVisualLogs::GetLogsSummaryText(const UPCGNode* InNode, ELogVerbosity::Type& OutMinimumVerbosity) const
+{
+	FReadScopeLock ScopedReadLock(LogsLock);
 
 	// Compose user friendly summary of first N errors
-	FText ResultText = FText::GetEmpty();
-	int32 LogCounter = 0;
+	FText Summary = FText::GetEmpty();
+	OutMinimumVerbosity = ELogVerbosity::All;
+
+	int LogCounter = 0;
 
-	for (int32 i = 0; i < NodeLogs.Num(); ++i)
+	for (const TPair<FPCGStack, FPCGPerNodeVisualLogs>& Entry : StackToLogs)
 	{
-		if (InComponent && NodeLogs[i].Component != InComponent)
+		const FPCGStack& Stack = Entry.Key;
+
+		if (Stack.HasObject(InNode))
 		{
-			continue;
-		}
+			const FPCGPerNodeVisualLogs& NodeLogs = Entry.Value;
 
-		++LogCounter;
+			for (int LogIndex = 0; LogIndex < NodeLogs.Num() && LogCounter < MaxLogsInSummary; ++LogIndex)
+			{
+				++LogCounter;
 
-		if (LogCounter > 1)
-		{
-			// New lines after each log
-			ResultText = FText::Format(FText::FromString(TEXT("{0}\n")), ResultText);
-		}
+				const UPCGComponent* Component = Stack.GetRootComponent();
+				FText ActorName = (Component && Component->GetOwner()) ? FText::FromString(Component->GetOwner()->GetActorLabel()) : FText::FromString(TEXT("MissingComponent"));
+				
+				OutMinimumVerbosity = FMath::Min(OutMinimumVerbosity, NodeLogs[LogIndex].Verbosity);
+				const FText VerbosityText = NodeLogs[LogIndex].Verbosity == ELogVerbosity::Warning ? FText::FromString(TEXT("Warning")) : FText::FromString(TEXT("Error"));
 
-		const FText MessageVerbosity = NodeLogs[i].Verbosity == ELogVerbosity::Warning ? FText::FromString(TEXT("Warning")) : FText::FromString(TEXT("Error"));
+				if (LogCounter > 1)
+				{
+					Summary = FText::Format(FText::FromString(TEXT("{0}\n")), Summary);
+				}
 
-		if (!InComponent)
-		{
-			FText ActorName = (NodeLogs[i].Component.Get() && NodeLogs[i].Component->GetOwner()) ? FText::FromName(NodeLogs[i].Component->GetOwner()->GetFName()) : FText::FromString(TEXT("MissingComponent"));
-			ResultText = FText::Format(LOCTEXT("NodeTooltipLogWithActor", "{0}{1}/{2}: [{3}] {4}: {5}"), ResultText, i + 1, NodeLogs.Num(), ActorName, MessageVerbosity, NodeLogs[i].Message);
-		}
-		else
-		{
-			ResultText = FText::Format(LOCTEXT("NodeTooltipLog", "{0}{1}/{2} {3}: {4}"), ResultText, i + 1, NodeLogs.Num(), MessageVerbosity, NodeLogs[i].Message);
+				Summary = FText::Format(LOCTEXT("NodeTooltipLogWithActor", "{0}[{1}] {2}: {3}"), Summary, ActorName, VerbosityText, NodeLogs[LogIndex].Message);
+			}
 		}
 
-		constexpr int32 MaxLogs = 8;
-		if (LogCounter >= MaxLogs)
+		if (LogCounter >= MaxLogsInSummary)
 		{
-			ResultText = FText::Format(FText::FromString(TEXT("{0}\n...")), ResultText);
+			Summary = FText::Format(FText::FromString(TEXT("{0}\n...")), Summary);
 			break;
 		}
 	}
 
-	return ResultText;
+	return Summary;
 }
 
-void FPCGNodeVisualLogs::ClearLogs(TWeakObjectPtr<const UPCGNode> InNode, const UPCGComponent* InComponent)
+FText FPCGNodeVisualLogs::GetLogsSummaryText(const FPCGStack& InBaseStack) const
 {
-	if (!ensure(InNode.Get()) || !ensure(InComponent))
-	{
-		return;
-	}
+	FReadScopeLock ScopedReadLock(LogsLock);
+
+	FText ResultText = FText::GetEmpty();
+
+	int32 LogCounter = 0;
 
-	FPCGPerNodeVisualLogs* NodeLogs = nullptr;
+	for (const TPair<FPCGStack, FPCGPerNodeVisualLogs>& Entry : StackToLogs)
 	{
-		FReadScopeLock ScopedReadLock(LogsLock);
-		NodeLogs = NodeToLogs.Find(InNode);
+		if (LogCounter >= MaxLogsInSummary)
+		{
+			ResultText = FText::Format(FText::FromString(TEXT("{0}\n...")), ResultText);
+			break;
+		}
 
-		if (!NodeLogs || NodeLogs->IsEmpty())
+		if (!Entry.Key.BeginsWith(InBaseStack))
 		{
-			// No logs, nothing to do
-			return;
+			continue;
+		}
+
+		const FPCGPerNodeVisualLogs& NodeLogs = Entry.Value;
+
+		for (int32 i = 0; i < NodeLogs.Num(); ++i)
+		{
+			++LogCounter;
+
+			if (LogCounter > 1)
+			{
+				ResultText = FText::Format(FText::FromString(TEXT("{0}\n")), ResultText);
+			}
+
+			const FText MessageVerbosity = NodeLogs[i].Verbosity == ELogVerbosity::Warning ? FText::FromString(TEXT("Warning")) : FText::FromString(TEXT("Error"));
+			ResultText = FText::Format(LOCTEXT("NodeTooltipLog", "{0}{1}: {2}"), ResultText, /*i + 1, NodeLogs.Num(),*/ MessageVerbosity, NodeLogs[i].Message);
+
+			if (LogCounter >= MaxLogsInSummary)
+			{
+				break;
+			}
 		}
 	}
 
-	bool bAnyRemoved = false;
+	return ResultText;
+}
+
+void FPCGNodeVisualLogs::ClearLogs(const FPCGStack& InPCGStack)
+{
+	TSet<const UPCGNode*> TouchedNodes;
+
 	{
 		FWriteScopeLock ScopedWriteLock(LogsLock);
 
-		check(NodeLogs);
-		for (int32 i = NodeLogs->Num() - 1; i >= 0; --i)
+		TArray<FPCGStack> StacksToRemove;
+		for (const TPair<FPCGStack, FPCGPerNodeVisualLogs>& Entry : StackToLogs)
 		{
-			// Remove entry if it matches the given component, or if the component is no longer valid
-			if (!(*NodeLogs)[i].Component.IsValid() || (*NodeLogs)[i].Component == InComponent)
+			if (Entry.Key.BeginsWith(InPCGStack))
 			{
-				NodeLogs->RemoveAtSwap(i);
-				bAnyRemoved = true;
+				StacksToRemove.Add(Entry.Key);
+
+				for (const FPCGStackFrame& Frame : Entry.Key.GetStackFrames())
+				{
+					if (const UPCGNode* Node = Cast<const UPCGNode>(Frame.Object.Get()))
+					{
+						TouchedNodes.Add(Node);
+					}
+				}
 			}
 		}
+
+		for (const FPCGStack& StackToRemove : StacksToRemove)
+		{
+			StackToLogs.Remove(StackToRemove);
+		}
 	}
 
 	// Broadcast change notification outside of write scope lock
-	if (bAnyRemoved)
+	for (const UPCGNode* TouchedNode : TouchedNodes)
 	{
-		InNode->OnNodeChangedDelegate.Broadcast(const_cast<UPCGNode*>(InNode.Get()), EPCGChangeType::Cosmetic);
+		TouchedNode->OnNodeChangedDelegate.Broadcast(const_cast<UPCGNode*>(TouchedNode), EPCGChangeType::Cosmetic);
 	}
 }
 
diff --git a/Engine/Plugins/Experimental/PCG/Source/PCG/Public/Graph/PCGStackContext.h b/Engine/Plugins/Experimental/PCG/Source/PCG/Public/Graph/PCGStackContext.h
--- a/Engine/Plugins/Experimental/PCG/Source/PCG/Public/Graph/PCGStackContext.h
+++ b/Engine/Plugins/Experimental/PCG/Source/PCG/Public/Graph/PCGStackContext.h
@@ -9,6 +9,7 @@
 
 class FString;
 class UObject;
+class UPCGComponent;
 class UPCGGraph;
 class UPCGNode;
 class UPCGPin;
@@ -56,7 +57,9 @@ struct PCG_API FPCGStack
 
 public:
 	/** Push frame onto top of stack. */
-	void PushFrame(const FPCGStackFrame& InFrame) { StackFrames.Add(InFrame); }
+	void PushFrame(const FPCGStackFrame& Frame) { StackFrames.Add(Frame); }
+	void PushFrame(const UObject* InFrameObject) { StackFrames.Emplace(InFrameObject); }
+	void PushFrame(int32 FrameLoopIndex) { StackFrames.Emplace(FrameLoopIndex); }
 
 	/** Pop frame from the stack. */
 	void PopFrame();
@@ -67,16 +70,44 @@ public:
 	/** Returns how many graphs the stack contains (top level graph stacks will return 1). */
 	uint32 GetNumGraphLevels() const;
 
-	/** Returns true if this stack in the top level graph, rather than in a subgraph. */
-	bool IsTopLevelGraph() const { return GetNumGraphLevels() == 1; }
+	/** Returns true if given stack is a prefix of this stack. */
+	bool BeginsWith(const FPCGStack& Other) const;
 
 	const TArray<FPCGStackFrame>& GetStackFrames() const { return StackFrames; }
 	TArray<FPCGStackFrame>& GetStackFramesMutable() { return StackFrames; }
 
+	/** Component given by first stack frame. */
+	const UPCGComponent* GetRootComponent() const;
+
+	/** First (top) graph frame in stack (or null if no graph frames present). */
+	const UPCGGraph* GetRootGraph() const;
+
+	/** Returns true if this stack is the top level/root graph, rather than in a subgraph. */
+	bool IsCurrentFrameInRootGraph() const { return GetNumGraphLevels() == 1; }
+
+	/** Gets the graph from the graph frame closest to the top of the stack (most recent), or null if no such graph present. */
+	const UPCGGraph* GetGraphForCurrentFrame() const;
+
+	/** If current frame (top of stack) corresponds to a node returns that node, otherwise returns null. */
+	const UPCGNode* GetCurrentFrameNode() const;
+
+	/** Stack has a frame corresponding to the given object. */
+	bool HasObject(const UObject* InObject) const;
+
 	bool operator==(const FPCGStack& Other) const;
 	bool operator!=(const FPCGStack& Other) const { return !(*this == Other); }
 
-	friend uint32 GetTypeHash(const FPCGStack& In);
+	friend uint32 GetTypeHash(const FPCGStack& In)
+	{
+		uint32 Hash = 0;
+
+		for (const FPCGStackFrame& Frame : In.StackFrames)
+		{
+			Hash = HashCombine(Hash, GetTypeHash(Frame));
+		}
+
+		return Hash;
+	}
 
 private:
 	TArray<FPCGStackFrame> StackFrames;
@@ -103,7 +134,7 @@ public:
 	/** Called during execution when invoking a dynamic subgraph, to prepend the caller stack to form the complete callstacks. */
 	void PrependParentStack(const FPCGStack* InParentStack);
 
-	static FPCGStackContext CreateStackContextFromGraph(const UPCGGraph* InPCGGraph);
+	TArray<FPCGStack>& GetStacksMutable() { return Stacks; }
 
 private:
 	/** List of all stacks encountered top graph and all (nested) subgraphs. Order is simply order of encountering during compilation. */
diff --git a/Engine/Plugins/Experimental/PCG/Source/PCG/Public/PCGComponent.h b/Engine/Plugins/Experimental/PCG/Source/PCG/Public/PCGComponent.h
--- a/Engine/Plugins/Experimental/PCG/Source/PCG/Public/PCGComponent.h
+++ b/Engine/Plugins/Experimental/PCG/Source/PCG/Public/PCGComponent.h
@@ -17,8 +17,8 @@
 namespace EEndPlayReason { enum Type : int; }
 
 class APCGPartitionActor;
-struct FPCGContext;
 class FPCGActorAndComponentMapping;
+class FPCGStackContext;
 class UPCGComponent;
 class UPCGData;
 class IPCGGenSourceBase;
@@ -31,7 +31,7 @@ class UPCGSubsystem;
 class ALandscapeProxy;
 class FLandscapeProxyComponentDataChangedParams;
 class UClass;
-
+struct FPCGContext;
 
 #if WITH_EDITOR
 	DECLARE_MULTICAST_DELEGATE_OneParam(FOnPCGGraphGenerated, UPCGComponent*);
@@ -299,6 +299,9 @@ public:
 
 	/** Changes the transient state (preview, normal, load on preview) - public only because it needs to be accessed by APCGPartitionActor */
 	void ChangeTransientState(EPCGEditorDirtyMode NewEditingMode);
+
+	/** Get execution stack information. */
+	bool GetStackContext(FPCGStackContext& OutStackContext) const;
 #endif
 
 	/** Utility function (mostly for tests) to properly set the value of bIsComponentPartitioned.
diff --git a/Engine/Plugins/Experimental/PCG/Source/PCG/Public/PCGGraph.h b/Engine/Plugins/Experimental/PCG/Source/PCG/Public/PCGGraph.h
--- a/Engine/Plugins/Experimental/PCG/Source/PCG/Public/PCGGraph.h
+++ b/Engine/Plugins/Experimental/PCG/Source/PCG/Public/PCGGraph.h
@@ -35,7 +35,6 @@ enum class EPCGGraphParameterEvent
 DECLARE_MULTICAST_DELEGATE_TwoParams(FOnPCGGraphChanged, UPCGGraphInterface* /*Graph*/, EPCGChangeType /*ChangeType*/);
 DECLARE_MULTICAST_DELEGATE_OneParam(FOnPCGGraphGridSizesChanged, UPCGGraphInterface* /*Graph*/);
 DECLARE_MULTICAST_DELEGATE_ThreeParams(FOnPCGGraphParametersChanged, UPCGGraphInterface* /*Graph*/, EPCGGraphParameterEvent /*ChangeType*/, FName /*ChangedPropertyName*/);
-DECLARE_MULTICAST_DELEGATE_ThreeParams(FOnPCGGraphDynamicallyExecuted, UPCGGraphInterface* /*Graph*/, const TWeakObjectPtr<UPCGComponent> /*SourceComponent*/, FPCGStack /*InvocationStack*/);
 #endif // WITH_EDITOR
 
 /**
@@ -112,7 +111,6 @@ public:
 	FOnPCGGraphChanged OnGraphChangedDelegate;
 	FOnPCGGraphGridSizesChanged OnGraphGridSizesChangedDelegate;
 	FOnPCGGraphParametersChanged OnGraphParametersChangedDelegate;
-	FOnPCGGraphDynamicallyExecuted OnGraphDynamicallyExecutedDelegate;
 #endif // WITH_EDITOR
 
 	template <typename T>
diff --git a/Engine/Plugins/Experimental/PCG/Source/PCG/Public/PCGSubsystem.h b/Engine/Plugins/Experimental/PCG/Source/PCG/Public/PCGSubsystem.h
--- a/Engine/Plugins/Experimental/PCG/Source/PCG/Public/PCGSubsystem.h
+++ b/Engine/Plugins/Experimental/PCG/Source/PCG/Public/PCGSubsystem.h
@@ -32,6 +32,10 @@ typedef TSharedPtr<IPCGElement, ESPMode::ThreadSafe> FPCGElementPtr;
 
 class UWorld;
 
+#if WITH_EDITOR
+DECLARE_MULTICAST_DELEGATE(FPCGOnComponentGenerationCompleteOrCancelled);
+#endif // WITH_EDITOR
+
 /**
 * UPCGSubsystem
 */
@@ -88,6 +92,9 @@ public:
 	void OnOriginalComponentRegistered(UPCGComponent* InComponent);
 	void OnOriginalComponentUnregistered(UPCGComponent* InComponent);
 
+	/** Called by graph executor when a graph is scheduled. */
+	void OnScheduleGraph(const FPCGStackContext& StackContext);
+
 	UPCGLandscapeCache* GetLandscapeCache();
 
 	// Schedule graph (owner -> graph)
@@ -209,7 +216,10 @@ public:
 	void ClearLandscapeCache();
 
 	/** Returns the graph compiler so we can figure out task info in the profiler view **/
-	const FPCGGraphCompiler* GetGraphCompiler() const;
+	FPCGGraphCompiler* GetGraphCompiler() const;
+
+	/** Get the execution stack information for the given component. */
+	bool GetStackContext(const UPCGComponent* InComponent, FPCGStackContext& OutStackContext) const;
 
 	/** Returns how many times InElement is present in the cache. */
 	uint32 GetGraphCacheEntryCount(IPCGElement* InElement) const;
@@ -221,6 +231,12 @@ public:
 	/** Notify that we exited the Landscape edit mode. */
 	void NotifyLandscapeEditModeExited() { ActorAndComponentMapping.NotifyLandscapeEditModeExited(); }
 
+	/** Get a list of stacks that were executed during the last execution. */
+	TArray<FPCGStack> GetExecutedStacks(const UPCGComponent* InComponent, const UPCGGraph* InSubgraph);
+	void ClearExecutedStacks(FPCGStack BeginningWithStack);
+
+	FPCGOnComponentGenerationCompleteOrCancelled OnComponentGenerationCompleteOrCancelled;
+
 private:
 	enum class EOperation : uint32
 	{
@@ -241,6 +257,10 @@ private:
 	bool bHasTickedOnce = false;
 	FPCGActorAndComponentMapping ActorAndComponentMapping;
 
+	/** A record of stacks that were executed. Used to populate debugging tool UIs. */
+	TArray<FPCGStack> ExecutedStacks;
+	mutable FRWLock ExecutedStacksLock;
+
 #if WITH_EDITOR
 	FCriticalSection PCGWorldActorLock;
 #endif
diff --git a/Engine/Plugins/Experimental/PCG/Source/PCG/Public/Utils/PCGNodeVisualLogs.h b/Engine/Plugins/Experimental/PCG/Source/PCG/Public/Utils/PCGNodeVisualLogs.h
--- a/Engine/Plugins/Experimental/PCG/Source/PCG/Public/Utils/PCGNodeVisualLogs.h
+++ b/Engine/Plugins/Experimental/PCG/Source/PCG/Public/Utils/PCGNodeVisualLogs.h
@@ -2,10 +2,11 @@
 
 #pragma once
 
-#include "Logging/LogVerbosity.h"
+#include "Graph/PCGStackContext.h"
 
 #include "Containers/Map.h"
 #include "Internationalization/Text.h"
+#include "Logging/LogVerbosity.h"
 #include "UObject/WeakObjectPtr.h"
 
 #if WITH_EDITOR
@@ -15,16 +16,14 @@ class UPCGNode;
 
 struct FPCGNodeLogEntry
 {
-	explicit FPCGNodeLogEntry(const FText& InMessage, ELogVerbosity::Type InVerbosity, TWeakObjectPtr<UPCGComponent> InComponent)
+	explicit FPCGNodeLogEntry(const FText& InMessage, ELogVerbosity::Type InVerbosity)
 		: Message(InMessage)
 		, Verbosity(InVerbosity)
-		, Component(InComponent)
 	{
 	}
 
 	FText Message;
 	ELogVerbosity::Type Verbosity;
-	TWeakObjectPtr<UPCGComponent> Component;
 };
 
 typedef TArray<FPCGNodeLogEntry, TInlineAllocator<16>> FPCGPerNodeVisualLogs;
@@ -34,23 +33,33 @@ class PCG_API FPCGNodeVisualLogs
 {
 public:
 	/** Log warnings and errors to be displayed on node in graph editor. */
-	void Log(TWeakObjectPtr<const UPCGNode> InNode, TWeakObjectPtr<UPCGComponent> InComponent, ELogVerbosity::Type InVerbosity, const FText& InMessage);
+	void Log(const FPCGStack& InPCGStack, ELogVerbosity::Type InVerbosity, const FText& InMessage);
+
+	/** Returns true if any issues were logged during last execution. */
+	bool HasLogs(const FPCGStack& InPCGStack) const;
+
+	/** Returns true if an issue with given severity was logged during last execution, and writes the minimum encountered verbosity to OutMinVerbosity. */
+	bool HasLogs(const FPCGStack& InPCGStack, ELogVerbosity::Type& OutMinVerbosity) const;
 
-	/** True if any issues were logged during last execution. */
-	bool HasLogs(TWeakObjectPtr<const UPCGNode> InNode, const UPCGComponent* InComponent) const;
+	/** Returns true if an issue with given severity was logged during last execution. */
+	bool HasLogsOfVerbosity(const FPCGStack& InPCGStack, ELogVerbosity::Type InVerbosity) const;
 
-	/** True if an issue with given severity was logged during last execution. */
-	bool HasLogs(TWeakObjectPtr<const UPCGNode> InNode, const UPCGComponent* InComponent, ELogVerbosity::Type InVerbosity) const;
+	/** Summary text of all visual logs produced while executing the provided base stack, appropriate for display in graph editor tooltip. */
+	FText GetLogsSummaryText(const FPCGStack& InBaseStack) const;
 
-	/** Summary text of visual logs from recent execution, appropriate for display in graph editor tooltip. */
-	FText GetLogsSummaryText(TWeakObjectPtr<const UPCGNode> InNode, const UPCGComponent* InComponent = nullptr) const;
+	/**
+	* Returns summary text of visual logs from recent execution, appropriate for display in graph editor tooltip. Writes the minimum encountered verbosity
+	* to OutMinimumVerbosity.
+	*/
+	FText GetLogsSummaryText(const UPCGNode* InNode, ELogVerbosity::Type& OutMinimumVerbosity) const;
 
-	/** Clear all errors and warnings that occurred while executing the given component. */
-	void ClearLogs(TWeakObjectPtr<const UPCGNode> InNode, const UPCGComponent* InComponent);
+	/** Clear all errors and warnings that occurred while executing stacks beginning with the given stack. */
+	void ClearLogs(const FPCGStack& InPCGStack);
 
 private:
-	TMap<TWeakObjectPtr<const UPCGNode>, FPCGPerNodeVisualLogs> NodeToLogs;
+	const int MaxLogsInSummary = 8;
 
+	TMap<FPCGStack, FPCGPerNodeVisualLogs> StackToLogs;
 	mutable FRWLock LogsLock;
 };
 
diff --git a/Engine/Plugins/Experimental/PCG/Source/PCGEditor/Private/PCGEditor.cpp b/Engine/Plugins/Experimental/PCG/Source/PCGEditor/Private/PCGEditor.cpp
--- a/Engine/Plugins/Experimental/PCG/Source/PCGEditor/Private/PCGEditor.cpp
+++ b/Engine/Plugins/Experimental/PCG/Source/PCGEditor/Private/PCGEditor.cpp
@@ -104,7 +104,11 @@ void FPCGEditor::Initialize(const EToolkitMode::Type InMode, const TSharedPtr<cl
 	if (PCGGraphBeingEdited)
 	{
 		PCGGraphBeingEdited->OnGraphGridSizesChangedDelegate.AddRaw(this, &FPCGEditor::OnGraphGridSizesChanged);
-		PCGGraphBeingEdited->OnGraphDynamicallyExecutedDelegate.AddRaw(this, &FPCGEditor::OnGraphDynamicallyExecuted);
+	}
+
+	if (UPCGSubsystem* Subsystem = GetSubsystem())
+	{
+		Subsystem->OnComponentGenerationCompleteOrCancelled.AddRaw(this, &FPCGEditor::OnComponentGenerationCompleteOrCancelled);
 	}
 
 	// Initializes the UPCGEditorGraph if needed
@@ -203,6 +207,15 @@ void FPCGEditor::Initialize(const EToolkitMode::Type InMode, const TSharedPtr<cl
 	{
 		GEngine->OnLevelActorDeleted().AddRaw(this, &FPCGEditor::OnLevelActorDeleted);
 	}
+
+	// Clear inspection flag on all nodes.
+	for (UEdGraphNode* EdGraphNode : PCGEditorGraph->Nodes)
+	{
+		if (UPCGEditorGraphNodeBase* PCGEditorGraphNode = Cast<UPCGEditorGraphNodeBase>(EdGraphNode))
+		{
+			PCGEditorGraphNode->SetInspected(false);
+		}
+	}
 }
 
 UPCGEditorGraph* FPCGEditor::GetPCGEditorGraph()
@@ -210,12 +223,12 @@ UPCGEditorGraph* FPCGEditor::GetPCGEditorGraph()
 	return PCGEditorGraph;
 }
 
-void FPCGEditor::SetComponentAndStackBeingInspected(UPCGComponent* InPCGComponent, const FPCGStack& InPCGStack)
+void FPCGEditor::SetStackBeingInspected(const FPCGStack& FullStack, FPCGDebugObjectSelectionMethod SelectionMethod)
 {
 	UPCGComponent* OldComponent = PCGComponentBeingInspected.Get();
-	UPCGComponent* NewComponent = InPCGComponent;
+	UPCGComponent* NewComponent = const_cast<UPCGComponent*>(FullStack.GetRootComponent());
 
-	if (OldComponent != NewComponent)
+	if (StackBeingInspected != FullStack)
 	{
 		if (OldComponent)
 		{
@@ -229,7 +242,9 @@ void FPCGEditor::SetComponentAndStackBeingInspected(UPCGComponent* InPCGComponen
 		const bool bNewComponentStartedInspecting = NewComponent && !NewComponent->IsInspecting();
 
 		PCGComponentBeingInspected = NewComponent;
-		OnInspectedComponentChangedDelegate.Broadcast(NewComponent);
+
+		StackBeingInspected = FullStack;
+		OnInspectedStackChangedDelegate.Broadcast(StackBeingInspected);
 
 		if (NewComponent)
 		{
@@ -240,6 +255,15 @@ void FPCGEditor::SetComponentAndStackBeingInspected(UPCGComponent* InPCGComponen
 			}
 		}
 
+		if (SelectionMethod == FPCGDebugObjectSelectionMethod::DebugObjectTree)
+		{
+			DebugObjectWidget->SetDebugObjectSelection(StackBeingInspected);
+		}
+		else
+		{
+			DebugObjectTreeWidget->SetDebugObjectSelection(StackBeingInspected);
+		}
+
 		UpdateDebugAfterComponentSelection(OldComponent, NewComponent, bNewComponentStartedInspecting);
 
 		check(PCGEditorGraph);
@@ -249,7 +273,7 @@ void FPCGEditor::SetComponentAndStackBeingInspected(UPCGComponent* InPCGComponen
 			{
 				// Update now that component has changed. Will fire OnNodeChanged if necessary.
 				EPCGChangeType ChangeType = PCGNode->UpdateErrorsAndWarnings();
-				ChangeType |= PCGNode->UpdateGridSizeVisualization(InPCGComponent, InPCGStack);
+				ChangeType |= PCGNode->UpdateGridSizeVisualization(NewComponent, FullStack);
 
 				if (ChangeType != EPCGChangeType::None)
 				{
@@ -258,12 +282,6 @@ void FPCGEditor::SetComponentAndStackBeingInspected(UPCGComponent* InPCGComponen
 			}
 		}
 	}
-
-	if (InPCGStack != StackBeingInspected)
-	{
-		StackBeingInspected = InPCGStack;
-		OnInspectedStackChangedDelegate.Broadcast(StackBeingInspected);
-	}
 }
 
 void FPCGEditor::UpdateDebugAfterComponentSelection(UPCGComponent* InOldComponent, UPCGComponent* InNewComponent, bool bInNewComponentStartedInspecting)
@@ -336,6 +354,11 @@ void FPCGEditor::UpdateDebugAfterComponentSelection(UPCGComponent* InOldComponen
 	}
 }
 
+const FPCGStack* FPCGEditor::GetStackBeingInspected() const
+{
+	return StackBeingInspected.GetStackFrames().IsEmpty() ? nullptr : &StackBeingInspected;
+}
+
 void FPCGEditor::JumpToNode(const UEdGraphNode* InNode)
 {
 	if (GraphEditorWidget.IsValid())
@@ -2046,7 +2069,11 @@ void FPCGEditor::OnClose()
 		}
 
 		PCGGraphBeingEdited->OnGraphGridSizesChangedDelegate.RemoveAll(this);
-		PCGGraphBeingEdited->OnGraphDynamicallyExecutedDelegate.RemoveAll(this);
+	}
+
+	if (UPCGSubsystem* Subsystem = GetSubsystem())
+	{
+		Subsystem->OnComponentGenerationCompleteOrCancelled.RemoveAll(this);
 	}
 }
 
@@ -2268,25 +2295,19 @@ bool FPCGEditor::IsVisibleProperty(const FPropertyAndParent& InPropertyAndParent
 void FPCGEditor::OnGraphGridSizesChanged(UPCGGraphInterface* InGraph)
 {
 	check(PCGEditorGraph);
-	if (UPCGComponent* PCGComponent = GetPCGComponentBeingInspected())
+
+	const FPCGStack* Stack = GetStackBeingInspected();
+	UPCGComponent* PCGComponent = GetPCGComponentBeingInspected();
+	if (Stack && PCGComponent)
 	{
-		PCGEditorGraph->UpdateGridSizeVisualization(PCGComponent, GetStackBeingInspected());
+		PCGEditorGraph->UpdateGridSizeVisualization(PCGComponent, *Stack);
 	}
 }
 
-void FPCGEditor::OnGraphDynamicallyExecuted(UPCGGraphInterface* InGraphInterface, const TWeakObjectPtr<UPCGComponent> InSourceComponent, FPCGStack InvocationStack)
+void FPCGEditor::OnComponentGenerationCompleteOrCancelled()
 {
-	InvocationStack.GetStackFramesMutable().Emplace(InGraphInterface);
-
-	if (DebugObjectWidget.IsValid())
-	{
-		DebugObjectWidget->AddDynamicStack(InSourceComponent, InvocationStack);
-	}
-
-	if (DebugObjectTreeWidget.IsValid())
-	{
-		DebugObjectTreeWidget->AddDynamicStack(InSourceComponent, InvocationStack);
-	}
+	DebugObjectWidget->RefreshDebugObjects();
+	DebugObjectTreeWidget->RequestRefresh();
 }
 
 UPCGSubsystem* FPCGEditor::GetSubsystem()
@@ -2307,6 +2328,12 @@ void FPCGEditor::OnMapChanged(UWorld* InWorld, EMapChangeType InMapChangedType)
 		{
 			DebugObjectWidget->RefreshDebugObjects();
 		}
+
+		// Subsystem has been torn down and rebuilt.
+		if (UPCGSubsystem* Subsystem = GetSubsystem())
+		{
+			Subsystem->OnComponentGenerationCompleteOrCancelled.AddRaw(this, &FPCGEditor::OnComponentGenerationCompleteOrCancelled);
+		}
 	}
 }
 
diff --git a/Engine/Plugins/Experimental/PCG/Source/PCGEditor/Private/PCGEditor.h b/Engine/Plugins/Experimental/PCG/Source/PCGEditor/Private/PCGEditor.h
--- a/Engine/Plugins/Experimental/PCG/Source/PCGEditor/Private/PCGEditor.h
+++ b/Engine/Plugins/Experimental/PCG/Source/PCGEditor/Private/PCGEditor.h
@@ -35,6 +35,13 @@ DECLARE_MULTICAST_DELEGATE_OneParam(FOnInspectedComponentChanged, UPCGComponent*
 DECLARE_MULTICAST_DELEGATE_OneParam(FOnInspectedStackChanged, const FPCGStack&);
 DECLARE_MULTICAST_DELEGATE_OneParam(FOnInspectedNodeChanged, UPCGEditorGraphNodeBase*);
 
+/** Used to communicate how a debug object was selected so that other selection widgets can be synchronized. */
+enum class FPCGDebugObjectSelectionMethod : uint8
+{
+	DebugObjectTree,
+	DebugObjectDropdown,
+};
+
 class FPCGEditor : public FAssetEditorToolkit, public FGCObject, public FSelfRegisteringEditorUndoClient
 {
 public:
@@ -47,18 +54,21 @@ public:
 	/** Gets/Creates the PCG graph editor for a given PCG graph */
 	static UPCGEditorGraph* GetPCGEditorGraph(UPCGGraph* InGraph);
 
-	/** Sets the PCG component and stack that we want to inspect */
-	void SetComponentAndStackBeingInspected(UPCGComponent* InPCGComponent, const FPCGStack& InPCGStack);
+	/** Sets the execution stack that want to inspect. */
+	void SetStackBeingInspected(const FPCGStack& FullStack, FPCGDebugObjectSelectionMethod SelectionMethod);
 
 	/** Gets the PCG component we are debugging */
-	UPCGComponent* GetPCGComponentBeingInspected() const { return PCGComponentBeingInspected.Get(); }
+	UPCGComponent* GetPCGComponentBeingInspected() const { return const_cast<UPCGComponent*>(StackBeingInspected.GetRootComponent()); }
 	
 	/** Gets the PCG stack we are inspecting */
-	const FPCGStack& GetStackBeingInspected() const { return StackBeingInspected; }
+	const FPCGStack* GetStackBeingInspected() const;
 
 	/** Focus the graph view on a specific node */
 	void JumpToNode(const UEdGraphNode* InNode);
 
+	/** Helper to get to the subsystem. */
+	static class UPCGSubsystem* GetSubsystem();
+
 	// ~Begin IToolkit interface
 	virtual void RegisterTabSpawners(const TSharedRef<class FTabManager>& InTabManager) override;
 	virtual void UnregisterTabSpawners(const TSharedRef<class FTabManager>& InTabManager) override;
@@ -87,7 +97,6 @@ public:
 	virtual void InitToolMenuContext(FToolMenuContext& MenuContext) override;
 	// ~End FAssetEditorToolkit interface
 
-	FOnInspectedComponentChanged OnInspectedComponentChangedDelegate;
 	FOnInspectedStackChanged OnInspectedStackChangedDelegate;
 	FOnInspectedNodeChanged OnInspectedNodeChangedDelegate;
 
@@ -285,14 +294,13 @@ private:
 	bool IsVisibleProperty(const FPropertyAndParent& InPropertyAndParent, IDetailsView* InDetailsView) const;
 
 	void OnGraphGridSizesChanged(UPCGGraphInterface* InGraph);
-	void OnGraphDynamicallyExecuted(UPCGGraphInterface* InGraphInterface, const TWeakObjectPtr<UPCGComponent> InSourceComponent, FPCGStack InvocationStack);
+
+	/** Called when a component finishes executing. Useful for updating debugging tools/UIs. */
+	void OnComponentGenerationCompleteOrCancelled();
 
 	/** Trigger any generation required to ensure debug display is up to date. */
 	void UpdateDebugAfterComponentSelection(UPCGComponent* InOldComponent, UPCGComponent* InNewComponent, bool bNewComponentStartedInspecting);
 
-	/** Helper to get to the subsystem. */
-	static class UPCGSubsystem* GetSubsystem();
-
 	void OnMapChanged(UWorld* InWorld, EMapChangeType InMapChangedType);
 	void OnLevelActorDeleted(AActor* InActor);
 
diff --git a/Engine/Plugins/Experimental/PCG/Source/PCGEditor/Private/PCGEditorGraphNodeBase.cpp b/Engine/Plugins/Experimental/PCG/Source/PCGEditor/Private/PCGEditorGraphNodeBase.cpp
--- a/Engine/Plugins/Experimental/PCG/Source/PCGEditor/Private/PCGEditorGraphNodeBase.cpp
+++ b/Engine/Plugins/Experimental/PCG/Source/PCGEditor/Private/PCGEditorGraphNodeBase.cpp
@@ -320,44 +320,63 @@ void UPCGEditorGraphNodeBase::OnNodeChanged(UPCGNode* InNode, EPCGChangeType Cha
 
 EPCGChangeType UPCGEditorGraphNodeBase::UpdateErrorsAndWarnings()
 {
-	EPCGChangeType ChangeType = EPCGChangeType::None;
-
-	// Pull current errors/warnings state from PCG subsystem.
-	if (const UPCGSubsystem* Subsystem = UPCGSubsystem::GetActiveEditorInstance())
+	const UPCGSubsystem* Subsystem = UPCGSubsystem::GetActiveEditorInstance();
+	if (!PCGNode || !Subsystem)
 	{
-		const UPCGComponent* ComponentBeingDebugged = nullptr;
-		{
-			const UPCGEditorGraph* EditorGraph = CastChecked<UPCGEditorGraph>(GetGraph());
-			const FPCGEditor* Editor = (EditorGraph && EditorGraph->GetEditor().IsValid()) ? EditorGraph->GetEditor().Pin().Get() : nullptr;
-			ComponentBeingDebugged = Editor ? Editor->GetPCGComponentBeingInspected() : nullptr;
-		}
-
-		const bool bOldHasCompilerMessage = bHasCompilerMessage;
-		const int32 OldErrorType = ErrorType;
-		const FString OldErrorMsg = ErrorMsg;
+		return EPCGChangeType::None;
+	}
+	
+	const FPCGStack* InspectedStack = nullptr;
+	{
+		const UPCGEditorGraph* EditorGraph = CastChecked<UPCGEditorGraph>(GetGraph());
+		const FPCGEditor* Editor = (EditorGraph && EditorGraph->GetEditor().IsValid()) ? EditorGraph->GetEditor().Pin().Get() : nullptr;
+		InspectedStack = Editor ? Editor->GetStackBeingInspected() : nullptr;
+	}
 
-		bHasCompilerMessage = Subsystem->GetNodeVisualLogs().HasLogs(PCGNode, ComponentBeingDebugged);
+	const bool bOldHasCompilerMessage = bHasCompilerMessage;
+	const int32 OldErrorType = ErrorType;
+	const FString OldErrorMsg = ErrorMsg;
+		
+	if (InspectedStack)
+	{
+		// Get errors/warnings for the inspected stack.
+		FPCGStack StackWithNode = *InspectedStack;
+		StackWithNode.PushFrame(PCGNode);
+		bHasCompilerMessage = Subsystem->GetNodeVisualLogs().HasLogs(StackWithNode);
 
 		if (bHasCompilerMessage)
 		{
-			const bool bHasErrors = Subsystem->GetNodeVisualLogs().HasLogs(PCGNode, ComponentBeingDebugged, ELogVerbosity::Error);
-			ErrorType = bHasErrors ? EMessageSeverity::Error : EMessageSeverity::Warning;
+			ErrorMsg = Subsystem->GetNodeVisualLogs().GetLogsSummaryText(StackWithNode).ToString();
 
-			ErrorMsg = Subsystem->GetNodeVisualLogs().GetLogsSummaryText(PCGNode, ComponentBeingDebugged).ToString();
+			const bool bHasErrors = Subsystem->GetNodeVisualLogs().HasLogsOfVerbosity(StackWithNode, ELogVerbosity::Error);
+			ErrorType = bHasErrors ? EMessageSeverity::Error : EMessageSeverity::Warning;
 		}
 		else
 		{
 			ErrorMsg.Empty();
 			ErrorType = 0;
 		}
+	}
+	else
+	{
+		// Collect all errors/warnings for this node.
+		ELogVerbosity::Type MinimumVerbosity;
+		ErrorMsg = Subsystem->GetNodeVisualLogs().GetLogsSummaryText(PCGNode.Get(), MinimumVerbosity).ToString();
+
+		bHasCompilerMessage = !ErrorMsg.IsEmpty();
 
-		if ((bHasCompilerMessage != bOldHasCompilerMessage) || (ErrorType != OldErrorType) || (ErrorMsg != OldErrorMsg))
+		if (bHasCompilerMessage)
+		{
+			ErrorType = MinimumVerbosity < ELogVerbosity::Warning ? EMessageSeverity::Error : EMessageSeverity::Warning;
+		}
+		else
 		{
-			ChangeType = EPCGChangeType::Cosmetic;
+			ErrorType = 0;
 		}
 	}
 
-	return ChangeType;
+	const bool bStateChanged = (bHasCompilerMessage != bOldHasCompilerMessage) || (ErrorType != OldErrorType) || (ErrorMsg != OldErrorMsg);
+	return bStateChanged ? EPCGChangeType::Cosmetic : EPCGChangeType::None;
 }
 
 EPCGChangeType UPCGEditorGraphNodeBase::UpdateGridSizeVisualization(UPCGComponent* InComponentBeingDebugged, const FPCGStack& InStackBeingInspected)
@@ -375,7 +394,7 @@ EPCGChangeType UPCGEditorGraphNodeBase::UpdateGridSizeVisualization(UPCGComponen
 
 	// Disable grid size visualization if higen is disabled, or if we're not inspecting a specific grid, or if we're
 	// inspecting a subgraph since subgraphs execute at the invoked grid level.
-	if (!HiGenEnabled || (InspectingGridSize == PCGHiGenGrid::UninitializedGridSize()) || !InStackBeingInspected.IsTopLevelGraph())
+	if (!HiGenEnabled || (InspectingGridSize == PCGHiGenGrid::UninitializedGridSize()) || !InStackBeingInspected.IsCurrentFrameInRootGraph())
 	{
 		if (IsDisplayAsDisabledForced())
 		{
diff --git a/Engine/Plugins/Experimental/PCG/Source/PCGEditor/Private/SPCGEditorGraphAttributeListView.cpp b/Engine/Plugins/Experimental/PCG/Source/PCGEditor/Private/SPCGEditorGraphAttributeListView.cpp
--- a/Engine/Plugins/Experimental/PCG/Source/PCGEditor/Private/SPCGEditorGraphAttributeListView.cpp
+++ b/Engine/Plugins/Experimental/PCG/Source/PCGEditor/Private/SPCGEditorGraphAttributeListView.cpp
@@ -320,7 +320,6 @@ void SPCGEditorGraphAttributeListView::Construct(const FArguments& InArgs, TShar
 	PCGEditorPtr = InPCGEditor;
 	SortingColumn = PCGEditorGraphAttributeListView::NAME_Index;
 
-	PCGEditorPtr.Pin()->OnInspectedComponentChangedDelegate.AddSP(this, &SPCGEditorGraphAttributeListView::OnInspectedComponentChanged);
 	PCGEditorPtr.Pin()->OnInspectedStackChangedDelegate.AddSP(this, &SPCGEditorGraphAttributeListView::OnInspectedStackChanged);
 	PCGEditorPtr.Pin()->OnInspectedNodeChangedDelegate.AddSP(this, &SPCGEditorGraphAttributeListView::OnInspectedNodeChanged);
 
@@ -527,7 +526,7 @@ TSharedRef<SHeaderRow> SPCGEditorGraphAttributeListView::CreateHeaderRowWidget()
 	return SNew(SHeaderRow);
 }
 
-void SPCGEditorGraphAttributeListView::OnInspectedComponentChanged(UPCGComponent* InPCGComponent)
+void SPCGEditorGraphAttributeListView::OnInspectedStackChanged(const FPCGStack& InPCGStack)
 {
 	if (PCGComponent.IsValid())
 	{
@@ -535,7 +534,7 @@ void SPCGEditorGraphAttributeListView::OnInspectedComponentChanged(UPCGComponent
 		PCGComponent->OnPCGGraphCleanedDelegate.RemoveAll(this);
 	}
 
-	PCGComponent = InPCGComponent;
+	PCGComponent = const_cast<UPCGComponent*>(InPCGStack.GetRootComponent());
 
 	if (PCGComponent.IsValid())
 	{
@@ -555,11 +554,6 @@ void SPCGEditorGraphAttributeListView::OnInspectedComponentChanged(UPCGComponent
 	}
 }
 
-void SPCGEditorGraphAttributeListView::OnInspectedStackChanged(const FPCGStack& InPCGStack)
-{
-	RequestRefresh();
-}
-
 void SPCGEditorGraphAttributeListView::OnInspectedNodeChanged(UPCGEditorGraphNodeBase* InPCGEditorGraphNode)
 {
 	if (PCGEditorGraphNode == InPCGEditorGraphNode)
@@ -617,10 +611,14 @@ const FPCGDataCollection* SPCGEditorGraphAttributeListView::GetInspectionData()
 	}
 
 	const TSharedPtr<FPCGEditor> PCGEditor = PCGEditorPtr.Pin();
-	const FPCGStack& PCGStack = PCGEditor->GetStackBeingInspected();
+	const FPCGStack* PCGStack = PCGEditor->GetStackBeingInspected();
+	if (!PCGStack)
+	{
+		return nullptr;
+	}
 
 	// Create a temporary stack with Node+Pin to query the exact DataCollection we are inspecting
-	FPCGStack Stack = PCGStack;
+	FPCGStack Stack = *PCGStack;
 	TArray<FPCGStackFrame>& StackFrames = Stack.GetStackFramesMutable();
 	StackFrames.Reserve(StackFrames.Num() + 2);
 	StackFrames.Emplace(PCGNode);
diff --git a/Engine/Plugins/Experimental/PCG/Source/PCGEditor/Private/SPCGEditorGraphAttributeListView.h b/Engine/Plugins/Experimental/PCG/Source/PCGEditor/Private/SPCGEditorGraphAttributeListView.h
--- a/Engine/Plugins/Experimental/PCG/Source/PCGEditor/Private/SPCGEditorGraphAttributeListView.h
+++ b/Engine/Plugins/Experimental/PCG/Source/PCGEditor/Private/SPCGEditorGraphAttributeListView.h
@@ -142,7 +142,6 @@ public:
 private:
 	TSharedRef<SHeaderRow> CreateHeaderRowWidget() const;
 
-	void OnInspectedComponentChanged(UPCGComponent* InPCGComponent);
 	void OnInspectedStackChanged(const FPCGStack& InPCGStack);
 	void OnInspectedNodeChanged(UPCGEditorGraphNodeBase* InPCGEditorGraphNode);
 
diff --git a/Engine/Plugins/Experimental/PCG/Source/PCGEditor/Private/SPCGEditorGraphDebugObjectTree.cpp b/Engine/Plugins/Experimental/PCG/Source/PCGEditor/Private/SPCGEditorGraphDebugObjectTree.cpp
--- a/Engine/Plugins/Experimental/PCG/Source/PCGEditor/Private/SPCGEditorGraphDebugObjectTree.cpp
+++ b/Engine/Plugins/Experimental/PCG/Source/PCGEditor/Private/SPCGEditorGraphDebugObjectTree.cpp
@@ -3,12 +3,13 @@
 #include "SPCGEditorGraphDebugObjectTree.h"
 
 #include "PCGComponent.h"
-#include "PCGEditor.h"
-#include "PCGEditorGraph.h"
-#include "PCGEditorModule.h"
 #include "PCGGraph.h"
+#include "PCGSubsystem.h"
 #include "Elements/PCGLoopElement.h"
 
+#include "PCGEditor.h"
+#include "PCGEditorGraph.h"
+
 #include "PropertyCustomizationHelpers.h"
 #include "Selection.h"
 #include "UnrealEdGlobals.h"
@@ -42,7 +43,17 @@ void FPCGEditorGraphDebugObjectItem::SortChildren(bool bIsAscending, bool bIsRec
 {
 	Children.Sort([bIsAscending](const FPCGEditorGraphDebugObjectItemPtr& InLHS, const FPCGEditorGraphDebugObjectItemPtr& InRHS)
 	{
-		return (InLHS->GetLabel() < InRHS->GetLabel()) == bIsAscending;
+		// Support for sorting by loop index.
+		const int32 IndexLHS = InLHS->GetSortPriority();
+		const int32 IndexRHS = InRHS->GetSortPriority();
+		if (IndexLHS == INDEX_NONE || IndexRHS == INDEX_NONE)
+		{
+			return (InLHS->GetLabel() < InRHS->GetLabel()) == bIsAscending;
+		}
+		else
+		{
+			return (IndexLHS < IndexRHS) == bIsAscending;
+		}
 	});
 
 	if (bIsRecursive)
@@ -61,12 +72,12 @@ FString FPCGEditorGraphDebugObjectItem_Actor::GetLabel() const
 
 FString FPCGEditorGraphDebugObjectItem_PCGComponent::GetLabel() const
 {
-	return PCGComponent.IsValid() ? PCGComponent->GetName() : FString();
-}
+	if (PCGComponent.IsValid() && PCGGraph.IsValid())
+	{
+		return PCGComponent->GetName() + FString(TEXT(" - ") + PCGGraph->GetName());
+	}
 
-FString FPCGEditorGraphDebugObjectItem_PCGGraph::GetLabel() const
-{
-	return PCGGraph.IsValid() ? PCGGraph->GetName() : FString();
+	return FString();
 }
 
 FString FPCGEditorGraphDebugObjectItem_PCGSubgraph::GetLabel() const
@@ -81,12 +92,7 @@ FString FPCGEditorGraphDebugObjectItem_PCGSubgraph::GetLabel() const
 
 FString FPCGEditorGraphDebugObjectItem_PCGLoopIndex::GetLabel() const
 {
-	if (PCGNode.IsValid())
-	{
-		return PCGNode->GetName() + FString::Format(TEXT(" - {0}"), { LoopIndex });
-	}
-
-	return FString();
+	return FString::Format(TEXT("{0}"), { LoopIndex });
 }
 
 void SPCGEditorGraphDebugObjectItemRow::Construct(const FArguments& InArgs, const TSharedRef<STableViewBase>& InOwnerTableView, FPCGEditorGraphDebugObjectItemPtr InItem)
@@ -98,14 +104,25 @@ void SPCGEditorGraphDebugObjectItemRow::Construct(const FArguments& InArgs, cons
 		SNew(SHorizontalBox)
 		+SHorizontalBox::Slot()
 		.AutoWidth()
-		.Padding(2.0f, 0.0f)
+		.Padding(6.0f, 0.0f)
 		[
 			SNew(SImage)
 			.Visibility(EVisibility::HitTestInvisible)
+			.ColorAndOpacity_Lambda([this]() { return Item->IsDebuggable() ? FLinearColor::White : FLinearColor(1.0, 1.0f, 1.0f, 0.5f); })
 			.Image_Lambda([this]()
 			{
+				if (Item->IsDebuggable())
+				{
+					return FAppStyle::Get().GetBrush("LevelEditor.Tabs.Debug");
+				}
+
 				const UObject* ItemObject = Item->GetObject();
-				return FSlateIconFinder::FindIconBrushForClass(IsValid(ItemObject) ? ItemObject->GetClass() : nullptr); // TODO: IsValid check is done here to prevent crashing due to deleted pcg components, we need to refresh the tree items when world changes, then this might not be needed.
+				if (IsValid(ItemObject) && ItemObject->IsA<AActor>())
+				{
+					return FSlateIconFinder::FindIconBrushForClass(ItemObject->GetClass());
+				}
+				
+				return Item->IsExpanded() ? FAppStyle::Get().GetBrush("Icons.FolderOpen") : FAppStyle::Get().GetBrush("Icons.FolderClosed");
 			})
 		]
 		+SHorizontalBox::Slot()
@@ -114,6 +131,74 @@ void SPCGEditorGraphDebugObjectItemRow::Construct(const FArguments& InArgs, cons
 			SNew(STextBlock)
 			.Text(FText::FromString(Item->GetLabel()))
 		]
+		+SHorizontalBox::Slot()
+		.AutoWidth()
+		.Padding(20.0f, 0.0f)
+		[
+			SNew(SImage)
+			.Visibility_Lambda([this]()
+			{
+				if (const UPCGSubsystem* Subsystem = UPCGSubsystem::GetActiveEditorInstance())
+				{
+					if (Item->GetPCGStack())
+					{
+						if (Subsystem->GetNodeVisualLogs().HasLogs(*Item->GetPCGStack()))
+						{
+							return EVisibility::Visible;
+						}
+					}
+					else
+					{
+						for (const FPCGEditorGraphDebugObjectItemPtr& Child : Item->GetChildren())
+						{
+							if (Child->GetPCGStack() && Subsystem->GetNodeVisualLogs().HasLogs(*Child->GetPCGStack()))
+							{
+								return EVisibility::Visible;
+							}
+						}
+					}
+				}
+				
+				return EVisibility::Hidden;
+			})
+			.Image(FAppStyle::Get().GetBrush("Icons.Error"))
+			.ColorAndOpacity_Lambda([this]()
+			{
+				if (const UPCGSubsystem* Subsystem = UPCGSubsystem::GetActiveEditorInstance())
+				{
+					constexpr FLinearColor WarningColor(1.0f, 0.75f, 0.0f, 0.9f);
+					constexpr FLinearColor ErrorColor(1.0f, 0.0f, 0.0f, 0.9f);
+
+					if (Item->GetPCGStack())
+					{
+						ELogVerbosity::Type MinVerbosity;
+						if (Subsystem->GetNodeVisualLogs().HasLogs(*Item->GetPCGStack(), MinVerbosity))
+						{
+							return FSlateColor(MinVerbosity <= ELogVerbosity::Error ? ErrorColor : WarningColor);
+						}
+					}
+					else
+					{
+						ELogVerbosity::Type MinVerbosityOfAll = ELogVerbosity::All;
+						for (const FPCGEditorGraphDebugObjectItemPtr& Child : Item->GetChildren())
+						{
+							ELogVerbosity::Type MinVerbosity;
+							if (Child->GetPCGStack() && Subsystem->GetNodeVisualLogs().HasLogs(*Child->GetPCGStack(), MinVerbosity))
+							{
+								MinVerbosityOfAll = FMath::Min(MinVerbosityOfAll, MinVerbosity);
+							}
+						}
+
+						if (MinVerbosityOfAll < ELogVerbosity::All)
+						{
+							return FSlateColor(MinVerbosityOfAll <= ELogVerbosity::Error ? ErrorColor : WarningColor);
+						}
+					}
+				}
+
+				return FSlateColor(FLinearColor(1.0f, 1.0f, 1.0f, 0.0f));
+			})
+		]
 	];
 }
 
@@ -149,6 +234,7 @@ void SPCGEditorGraphDebugObjectTree::Construct(const FArguments& InArgs, TShared
 				.OnGenerateRow(this, &SPCGEditorGraphDebugObjectTree::MakeTreeRowWidget)
 				.OnGetChildren(this, &SPCGEditorGraphDebugObjectTree::OnGetChildren)
 				.OnSelectionChanged(this, &SPCGEditorGraphDebugObjectTree::OnSelectionChanged)
+				.OnExpansionChanged(this, &SPCGEditorGraphDebugObjectTree::OnExpansionChanged)
 				.OnSetExpansionRecursive(this, &SPCGEditorGraphDebugObjectTree::OnSetExpansionRecursive)
 				.ItemHeight(18)
 				.AllowOverscroll(EAllowOverscroll::No)
@@ -236,15 +322,29 @@ void SPCGEditorGraphDebugObjectTree::Tick(const FGeometry& AllottedGeometry, con
 	}
 }
 
-void SPCGEditorGraphDebugObjectTree::AddDynamicStack(const TWeakObjectPtr<UPCGComponent> InComponent, const FPCGStack& InvocationStack)
+void SPCGEditorGraphDebugObjectTree::SetDebugObjectSelection(const FPCGStack& FullStack)
 {
-	TArray<FPCGStack>& Stacks = DynamicInvocationStacks.FindOrAdd(InComponent);
+	bDisableDebugObjectChangeNotification = true;
 
-	if (!Stacks.Contains(InvocationStack))
+	for (FPCGEditorGraphDebugObjectItemPtr& Item : AllGraphItems)
 	{
-		Stacks.Add(InvocationStack);
-		RequestRefresh();
+		const FPCGStack* ItemStack = Item->GetPCGStack();
+		const bool bSelected = ItemStack && *ItemStack == FullStack;
+
+		DebugObjectTreeView->SetItemSelection(Item, bSelected);
+
+		if (bSelected)
+		{
+			FPCGEditorGraphDebugObjectItemPtr Parent = Item->GetParent();
+			while (Parent)
+			{
+				DebugObjectTreeView->SetItemExpansion(Parent, true);
+				Parent = Parent->GetParent();
+			}
+		}
 	}
+
+	bDisableDebugObjectChangeNotification = false;
 }
 
 void SPCGEditorGraphDebugObjectTree::SelectedDebugObject_OnClicked() const
@@ -300,7 +400,11 @@ void SPCGEditorGraphDebugObjectTree::SetDebugObjectFromSelection_OnClicked()
 				continue;
 			}
 
-			FPCGStackContext StackContext = FPCGStackContext::CreateStackContextFromGraph(PCGComponent->GetGraph());
+			FPCGStackContext StackContext;
+			if (!PCGComponent->GetStackContext(StackContext))
+			{
+				continue;
+			}
 
 			for (const FPCGStack& Stack : StackContext.GetStacks())
 			{
@@ -361,7 +465,11 @@ bool SPCGEditorGraphDebugObjectTree::IsSetDebugObjectFromSelectionButtonEnabled(
 				continue;
 			}
 
-			FPCGStackContext StackContext = FPCGStackContext::CreateStackContextFromGraph(PCGComponent->GetGraph());
+			FPCGStackContext StackContext;
+			if (!PCGComponent->GetStackContext(StackContext))
+			{
+				continue;
+			}
 
 			const bool bGraphFound = Algo::AnyOf(StackContext.GetStacks(), [&PCGGraph](const FPCGStack& InStack)
 			{
@@ -378,257 +486,239 @@ bool SPCGEditorGraphDebugObjectTree::IsSetDebugObjectFromSelectionButtonEnabled(
 	return false;
 }
 
-void SPCGEditorGraphDebugObjectTree::RefreshTree()
+void SPCGEditorGraphDebugObjectTree::AddStacksToTree(const TArray<FPCGStack>& Stacks,
+	TMap<AActor*, TSharedPtr<FPCGEditorGraphDebugObjectItem_Actor>>& InOutActorItems,
+	TMap<const FPCGStack, FPCGEditorGraphDebugObjectItemPtr>& InOutStackToItem)
 {
-	RootItems.Empty();
-	AllGraphItems.Empty();
-	DebugObjectTreeView->RequestTreeRefresh();
-
-	const UPCGGraph* PCGGraph = GetPCGGraph();
-	if (!PCGGraph)
+	const UPCGGraph* GraphBeingEdited = GetPCGGraph();
+	if (!GraphBeingEdited)
 	{
 		return;
 	}
 
-	TArray<UObject*> PCGComponents;
-	GetObjectsOfClass(UPCGComponent::StaticClass(), PCGComponents, /*bIncludeDerivedClasses=*/ true);
-
-	TMap<AActor*, TSharedPtr<FPCGEditorGraphDebugObjectItem_Actor>> ActorItems;
-	TMap<UPCGComponent*, TSharedPtr<FPCGEditorGraphDebugObjectItem_PCGComponent>> ComponentItems;
-
-	for (UObject* PCGComponentObject : PCGComponents)
+	for (const FPCGStack& Stack : Stacks)
 	{
-		if (!IsValid(PCGComponentObject))
+		// If the current graph is not in the stack at all then skip.
+		bool bRemainingStackContainsEditedGraph = Stack.HasObject(GraphBeingEdited);
+		if (!bRemainingStackContainsEditedGraph)
 		{
 			continue;
 		}
 
-		UPCGComponent* PCGComponent = Cast<UPCGComponent>(PCGComponentObject);
+		UPCGComponent* PCGComponent = const_cast<UPCGComponent*>(Stack.GetRootComponent());
 		if (!PCGComponent)
 		{
 			continue;
 		}
 
-		AActor* Actor = PCGComponent->GetOwner();
-		if (!Actor)
+		UPCGGraph* TopGraph = const_cast<UPCGGraph*>(Stack.GetRootGraph());
+		if (!TopGraph)
 		{
 			continue;
 		}
 
-		const UPCGGraph* PCGComponentGraph = PCGComponent->GetGraph();
-		if (!PCGComponentGraph)
+		AActor* Actor = PCGComponent->GetOwner();
+		if (!Actor)
 		{
 			continue;
 		}
 
-		FPCGStackContext StackContext = FPCGStackContext::CreateStackContextFromGraph(PCGComponentGraph);
+		// Add actor item if not already added.
+		FPCGEditorGraphDebugObjectItemPtr ActorItem;
+		if (TSharedPtr<FPCGEditorGraphDebugObjectItem_Actor>* FoundActorItem = InOutActorItems.Find(Actor))
+		{
+			ActorItem = *FoundActorItem;
+		}
+		else
+		{
+			ActorItem = InOutActorItems.Emplace(Actor, MakeShared<FPCGEditorGraphDebugObjectItem_Actor>(Actor));
+			AllGraphItems.Add(ActorItem);
+		}
 
-		TMap<const UPCGGraph*, FPCGEditorGraphDebugObjectItemPtr> GraphItems;
-		TArray<FPCGEditorGraphDebugObjectItemPtr> NodeItems;
+		const TArray<FPCGStackFrame>& StackFrames = Stack.GetStackFrames();
 
-		// Process statically generated stacks
-		for (const FPCGStack& Stack : StackContext.GetStacks())
+		// Example stack:
+		//     Component/TopGraph/SubgraphNode/Subgraph/LoopSubgraphNode/LoopIndex/LoopSubgraph
+		// 
+		// The loop below adds tree items for component & top graph, and then whenever a graph is encountered
+		// we look a previous frames to determine whether to add a subgraph item or loop subgraph item.
+		for (int FrameIndex = 1; FrameIndex < StackFrames.Num() && bRemainingStackContainsEditedGraph; FrameIndex++)
 		{
-			const UPCGGraph* TopStackGraph = Cast<const UPCGGraph>(Stack.GetStackFrames().Top().Object);
-			if (TopStackGraph == PCGGraph)
-			{
-				// Find or Create ActorItem
-				TSharedPtr<FPCGEditorGraphDebugObjectItem_Actor> ActorItem;
-				if (TSharedPtr<FPCGEditorGraphDebugObjectItem_Actor>* FoundActorItem = ActorItems.Find(Actor))
-				{
-					ActorItem = *FoundActorItem;
-				}
-				else
-				{
-					ActorItem = ActorItems.Emplace(Actor, MakeShared<FPCGEditorGraphDebugObjectItem_Actor>(Actor));
-				}
+			const FPCGStackFrame& StackFrame = StackFrames[FrameIndex];
+			const FPCGStackFrame& PreviousStackFrame = StackFrames[FrameIndex - 1];
+			FPCGEditorGraphDebugObjectItemPtr CurrentItem;
 
-				// Find or Create ComponentItem
-				TSharedPtr<FPCGEditorGraphDebugObjectItem_PCGComponent> ComponentItem;
-				if (TSharedPtr<FPCGEditorGraphDebugObjectItem_PCGComponent>* FoundComponentItem = ComponentItems.Find(PCGComponent))
-				{
-					ComponentItem = *FoundComponentItem;
-				}
-				else
-				{
-					ComponentItem = ComponentItems.Emplace(PCGComponent, MakeShared<FPCGEditorGraphDebugObjectItem_PCGComponent>(PCGComponent));
-				}
+			// When we encounter a graph, we look at the frame index and/or preceding frames to determine the graph type.
+			if (const UPCGGraph* StackGraph = Cast<const UPCGGraph>(StackFrame.Object))
+			{
+				const bool bIsDebuggable = (GraphBeingEdited == StackGraph);
 
-				if (!ComponentItem->GetParent())
+				// Top graph.
+				if (StackGraph == TopGraph)
 				{
-					ActorItem->AddChild(ComponentItem.ToSharedRef());
-				}
+					FPCGStack GraphStack = Stack;
+					GraphStack.GetStackFramesMutable().SetNum(FrameIndex + 1);
 
-				const TArray<FPCGStackFrame>& StackFrames = Stack.GetStackFrames();
+					if (!InOutStackToItem.Contains(GraphStack))
+					{
+						FPCGEditorGraphDebugObjectItemPtr TopGraphItem = InOutStackToItem.Emplace(
+							GraphStack,
+							MakeShared<FPCGEditorGraphDebugObjectItem_PCGComponent>(PCGComponent, StackGraph, GraphStack, bIsDebuggable));
 
-				FString StackPath;
-				Stack.CreateStackFramePath(StackPath);
+						AllGraphItems.Add(TopGraphItem);
 
-				FPCGEditorGraphDebugObjectItemPtr PreviousItem = ComponentItem;
-				for (int32 StackIndex = 0; StackIndex < StackFrames.Num(); StackIndex++)
+						ActorItem->AddChild(TopGraphItem.ToSharedRef());
+					}
+				}
+				// Previous stack was node, therefore subgraph.
+				else if (const UPCGNode* SubgraphNode = Cast<const UPCGNode>(PreviousStackFrame.Object))
 				{
-					const FPCGStackFrame& StackFrame = StackFrames[StackIndex];
-					FPCGEditorGraphDebugObjectItemPtr CurrentItem;
-					if (const UPCGGraph* StackGraph = Cast<const UPCGGraph>(StackFrame.Object))
+					FPCGStack GraphStack = Stack;
+					GraphStack.GetStackFramesMutable().SetNum(FrameIndex + 1);
+
+					if (!InOutStackToItem.Contains(GraphStack))
 					{
-						if (FPCGEditorGraphDebugObjectItemPtr* GraphItem = GraphItems.Find(StackGraph))
-						{
-							CurrentItem = *GraphItem;
-						}
-						else
+						FPCGEditorGraphDebugObjectItemPtr GraphItem = InOutStackToItem.Emplace(
+							GraphStack,
+							MakeShared<FPCGEditorGraphDebugObjectItem_PCGSubgraph>(SubgraphNode, StackGraph, GraphStack, bIsDebuggable));
+
+						AllGraphItems.Add(GraphItem);
+
+						TArray<FPCGStackFrame>& GraphStackFrames = GraphStack.GetStackFramesMutable();
+						while (GraphStackFrames.Num() > 0)
 						{
-							CurrentItem = GraphItems.Emplace(StackGraph, MakeShared<FPCGEditorGraphDebugObjectItem_PCGGraph>(StackGraph, StackGraph == PCGGraph ? Stack : FPCGStack()));
+							GraphStackFrames.SetNum(GraphStackFrames.Num() - 1, /*bAllowShrinking=*/false);
+
+							if (FPCGEditorGraphDebugObjectItemPtr* ParentItem = InOutStackToItem.Find(GraphStack))
+							{
+								(*ParentItem)->AddChild(GraphItem.ToSharedRef());
+								break;
+							}
 						}
 					}
-					else if (const UPCGNode* StackNode = Cast<const UPCGNode>(StackFrame.Object))
+				}
+				// Previous stack was loop index, therefore loop subgraph.
+				else if (FrameIndex >= 2 && PreviousStackFrame.LoopIndex != INDEX_NONE)
+				{
+					const UPCGNode* LoopSubgraphNode = Cast<const UPCGNode>(StackFrames[FrameIndex - 2].Object);
+					if (ensure(LoopSubgraphNode))
 					{
-						// We cannot handle dynamically executed graphs in this path. They must be provided through DynamicInvocationStacks
-						if (const UPCGBaseSubgraphSettings* SubgraphSettings = Cast<const UPCGBaseSubgraphSettings>(StackNode->GetSettings()))
+						// Take the stack up to the looped subgraph node, add a item for the node + graph.
+						FPCGStack LoopGraphStack = Stack;
+						LoopGraphStack.GetStackFramesMutable().SetNum(FrameIndex - 1);
+
+						if (!InOutStackToItem.Contains(LoopGraphStack))
 						{
-							if (SubgraphSettings->IsDynamicGraph())
+							FPCGEditorGraphDebugObjectItemPtr LoopGraphItem = InOutStackToItem.Emplace(
+								LoopGraphStack,
+								MakeShared<FPCGEditorGraphDebugObjectItem_PCGSubgraph>(LoopSubgraphNode, StackGraph, LoopGraphStack, /*bIsDebuggable=*/false));
+
+							AllGraphItems.Add(LoopGraphItem);
+
+							TArray<FPCGStackFrame>& LoopGraphStackFrames = LoopGraphStack.GetStackFramesMutable();
+							while (LoopGraphStackFrames.Num() > 0)
 							{
-								break;
+								LoopGraphStackFrames.SetNum(LoopGraphStackFrames.Num() - 1, /*bAllowShrinking=*/false);
+
+								if (FPCGEditorGraphDebugObjectItemPtr* ParentItem = InOutStackToItem.Find(LoopGraphStack))
+								{
+									(*ParentItem)->AddChild(LoopGraphItem.ToSharedRef());
+									break;
+								}
 							}
 						}
 
-						const int32 NextStackIndex = StackIndex + 1;
-						if (StackFrames.IsValidIndex(NextStackIndex))
+						// Take full stack up until this point which will be the unique stack for the loop iteration.
+						FPCGStack LoopIterationStack = Stack;
+						LoopIterationStack.GetStackFramesMutable().SetNum(FrameIndex + 1);
+
+						if (!InOutStackToItem.Contains(LoopIterationStack))
 						{
-							const FPCGStackFrame& NextStackFrame = StackFrames[NextStackIndex];
-							if (const UPCGGraph* NextStackGraph = Cast<const UPCGGraph>(NextStackFrame.Object))
+							FPCGEditorGraphDebugObjectItemPtr LoopIterationItem = InOutStackToItem.Emplace(
+								LoopIterationStack,
+								MakeShared<FPCGEditorGraphDebugObjectItem_PCGLoopIndex>(PreviousStackFrame.LoopIndex, StackGraph, LoopIterationStack, bIsDebuggable));
+
+							AllGraphItems.Add(LoopIterationItem);
+
+							TArray<FPCGStackFrame>& GraphStackFrames = LoopIterationStack.GetStackFramesMutable();
+							while (GraphStackFrames.Num() > 0)
 							{
-								const FPCGEditorGraphDebugObjectItemPtr* NodeItem = NodeItems.FindByPredicate([&StackNode,&Stack](const FPCGEditorGraphDebugObjectItemPtr& InItem)
-								{
-									return (InItem->GetObject() == StackNode) && (*InItem->GetPCGStack() == Stack);
-								});
+								GraphStackFrames.SetNum(GraphStackFrames.Num() - 1, /*bAllowShrinking=*/false);
 
-								if (NodeItem)
-								{
-									CurrentItem = *NodeItem;
-								}
-								else
+								if (FPCGEditorGraphDebugObjectItemPtr* ParentItem = InOutStackToItem.Find(LoopIterationStack))
 								{
-									CurrentItem = NodeItems.Add_GetRef(MakeShared<FPCGEditorGraphDebugObjectItem_PCGSubgraph>(StackNode, NextStackGraph, NextStackGraph == TopStackGraph ? Stack : FPCGStack()));
+									(*ParentItem)->AddChild(LoopIterationItem.ToSharedRef());
+									break;
 								}
-
-								StackIndex++;
 							}
 						}
 					}
-
-					if (!CurrentItem)
-					{
-						break;
-					}
-
-					if (!CurrentItem->GetParent())
-					{
-						PreviousItem->AddChild(CurrentItem.ToSharedRef());
-					}
-
-					AllGraphItems.Add(CurrentItem);
-					PreviousItem = CurrentItem;
 				}
 			}
-		}
 
-		if (const TArray<FPCGStack>* DynamicStacks = DynamicInvocationStacks.Find(PCGComponent))
-		{
-			for (const FPCGStack& Stack : *DynamicStacks)
+			// Check if any of the remaining stack frames contain valid debug targets - invocations of the current edited graph.
+			// If not then we will abandon this stack, pruning the tree.
+			bRemainingStackContainsEditedGraph = false;
+			for (int RemainingFrameIndex = FrameIndex + 1; RemainingFrameIndex < StackFrames.Num(); ++RemainingFrameIndex)
 			{
-				TSharedPtr<FPCGEditorGraphDebugObjectItem_Actor> ActorItem;
-				if (TSharedPtr<FPCGEditorGraphDebugObjectItem_Actor>* FoundActorItem = ActorItems.Find(Actor))
-				{
-					ActorItem = *FoundActorItem;
-				}
-				else
-				{
-					ActorItem = ActorItems.Emplace(Actor, MakeShared<FPCGEditorGraphDebugObjectItem_Actor>(Actor));
-				}
-
-				TSharedPtr<FPCGEditorGraphDebugObjectItem_PCGComponent> ComponentItem;
-				if (TSharedPtr<FPCGEditorGraphDebugObjectItem_PCGComponent>* FoundComponentItem = ComponentItems.Find(PCGComponent))
-				{
-					ComponentItem = *FoundComponentItem;
-				}
-				else
-				{
-					ComponentItem = ComponentItems.Emplace(PCGComponent, MakeShared<FPCGEditorGraphDebugObjectItem_PCGComponent>(PCGComponent));
-				}
-
-				if (!ComponentItem->GetParent())
+				if (StackFrames[RemainingFrameIndex].Object == GraphBeingEdited)
 				{
-					ActorItem->AddChild(ComponentItem.ToSharedRef());
+					bRemainingStackContainsEditedGraph = true;
+					break;
 				}
+			}
+		}
+	}
+}
 
-				const TArray<FPCGStackFrame>& StackFrames = Stack.GetStackFrames();
-
-				FString StackPath;
-				Stack.CreateStackFramePath(StackPath);
+void SPCGEditorGraphDebugObjectTree::RefreshTree()
+{
+	RootItems.Empty();
+	AllGraphItems.Empty();
+	DebugObjectTreeView->RequestTreeRefresh();
 
-				FPCGEditorGraphDebugObjectItemPtr PreviousItem = ComponentItem;
-				for (int32 StackIndex = 0; StackIndex < StackFrames.Num(); StackIndex++)
-				{
-					const FPCGStackFrame& StackFrame = StackFrames[StackIndex];
-					FPCGEditorGraphDebugObjectItemPtr CurrentItem;
-					if (const UPCGGraph* StackGraph = Cast<const UPCGGraph>(StackFrame.Object))
-					{
-						if (FPCGEditorGraphDebugObjectItemPtr* GraphItem = GraphItems.Find(StackGraph))
-						{
-							CurrentItem = *GraphItem;
-						}
-						else
-						{
-							CurrentItem = GraphItems.Emplace(StackGraph, MakeShared<FPCGEditorGraphDebugObjectItem_PCGGraph>(StackGraph, StackGraph == PCGGraph ? Stack : FPCGStack()));
-						}
-					}
-					else if (const UPCGNode* StackNode = Cast<const UPCGNode>(StackFrame.Object))
-					{
-						const int32 NextStackIndex = StackIndex + 1;
-						if (StackFrames.IsValidIndex(NextStackIndex))
-						{
-							const FPCGEditorGraphDebugObjectItemPtr* NodeItem = NodeItems.FindByPredicate([&StackNode, &Stack](const FPCGEditorGraphDebugObjectItemPtr& InItem)
-								{
-									return (InItem->GetObject() == StackNode) && (*InItem->GetPCGStack() == Stack);
-								});
+	UPCGGraph* PCGGraph = GetPCGGraph();
+	if (!PCGGraph)
+	{
+		return;
+	}
 
-							if (NodeItem)
-							{
-								CurrentItem = *NodeItem;
-							}
-							else
-							{
-								const FPCGStackFrame& NextStackFrame = StackFrames[NextStackIndex];
+	UPCGSubsystem* Subsystem = PCGEditor.Pin()->GetSubsystem();
+	if (!Subsystem)
+	{
+		return;
+	}
 
-								if (const UPCGGraph* NextStackGraph = Cast<const UPCGGraph>(NextStackFrame.Object))
-								{
-									CurrentItem = NodeItems.Add_GetRef(MakeShared<FPCGEditorGraphDebugObjectItem_PCGSubgraph>(StackNode, NextStackGraph, Stack));
-								}
-								else if (NextStackFrame.LoopIndex != INDEX_NONE)
-								{
-									CurrentItem = NodeItems.Add_GetRef(MakeShared<FPCGEditorGraphDebugObjectItem_PCGLoopIndex>(StackNode, NextStackFrame.LoopIndex, Stack));
-								}
-							}
+	TArray<UObject*> PCGComponents;
+	GetObjectsOfClass(UPCGComponent::StaticClass(), PCGComponents, /*bIncludeDerivedClasses=*/ true);
 
-							StackIndex++;
-						}
-					}
+	TMap<AActor*, TSharedPtr<FPCGEditorGraphDebugObjectItem_Actor>> ActorItems;
+	TMap<const FPCGStack, FPCGEditorGraphDebugObjectItemPtr> StackToItem;
 
-					if (!CurrentItem)
-					{
-						break;
-					}
+	for (UObject* PCGComponentObject : PCGComponents)
+	{
+		if (!IsValid(PCGComponentObject))
+		{
+			continue;
+		}
 
-					if (!CurrentItem->GetParent())
-					{
-						PreviousItem->AddChild(CurrentItem.ToSharedRef());
-					}
+		UPCGComponent* PCGComponent = Cast<UPCGComponent>(PCGComponentObject);
+		if (!PCGComponent)
+		{
+			continue;
+		}
 
-					AllGraphItems.Add(CurrentItem);
-					PreviousItem = CurrentItem;
-				}
-			}
+		// Process static stacks that can be read from the compiled graph.
+		FPCGStackContext StackContext;
+		if (!PCGComponent->GetStackContext(StackContext))
+		{
+			continue;
 		}
+		AddStacksToTree(StackContext.GetStacks(), ActorItems, StackToItem);
+
+		// Process stacks encountered during execution so far, which will include dynamic subgraphs & loop subgraphs.
+		// There will be overlaps with the static stacks but only unique entries will be added to the tree.
+		AddStacksToTree(Subsystem->GetExecutedStacks(PCGComponent, PCGGraph), ActorItems, StackToItem);
 	}
 
 	for (TPair<AActor*, TSharedPtr<FPCGEditorGraphDebugObjectItem_Actor>>& ActorItem : ActorItems)
@@ -637,6 +727,33 @@ void SPCGEditorGraphDebugObjectTree::RefreshTree()
 	}
 
 	SortTreeItems();
+
+	// Try to restore user item expansion.
+	TSet<FPCGStack> ExpandedStacksBefore = ExpandedStacks;
+	for (const FPCGStack& ExpandedStack : ExpandedStacksBefore)
+	{
+		for (FPCGEditorGraphDebugObjectItemPtr& Item : AllGraphItems)
+		{
+			const FPCGStack* ItemStack = Item->GetPCGStack();
+
+			if (ItemStack && *ItemStack == ExpandedStack)
+			{
+				DebugObjectTreeView->SetItemExpansion(Item, true);
+			}
+		}
+	}
+
+	// Try to restore user item selection.
+	for (FPCGEditorGraphDebugObjectItemPtr& Item : AllGraphItems)
+	{
+		const FPCGStack* ItemStack = Item->GetPCGStack();
+
+		if (ItemStack && SelectedStack == *ItemStack)
+		{
+			DebugObjectTreeView->SetItemSelection(Item, true);
+			break;
+		}
+	}
 }
 
 void SPCGEditorGraphDebugObjectTree::SortTreeItems(bool bIsAscending, bool bIsRecursive)
@@ -715,11 +832,25 @@ void SPCGEditorGraphDebugObjectTree::OnGetChildren(FPCGEditorGraphDebugObjectIte
 	}
 }
 
-void SPCGEditorGraphDebugObjectTree::OnSelectionChanged(FPCGEditorGraphDebugObjectItemPtr InItem, ESelectInfo::Type InSelectInfo) const
+void SPCGEditorGraphDebugObjectTree::OnSelectionChanged(FPCGEditorGraphDebugObjectItemPtr InItem, ESelectInfo::Type InSelectInfo)
 {
+	if (const FPCGStack* Stack = InItem ? InItem->GetPCGStack() : nullptr)
+	{
+		SelectedStack = *Stack;
+	}
+	else
+	{
+		SelectedStack = FPCGStack();
+	}
+
+	if (bDisableDebugObjectChangeNotification)
+	{
+		return;
+	}
+
 	if (!InItem)
 	{
-		PCGEditor.Pin()->SetComponentAndStackBeingInspected(nullptr, FPCGStack());
+		PCGEditor.Pin()->SetStackBeingInspected(FPCGStack(), FPCGDebugObjectSelectionMethod::DebugObjectTree);
 		return;
 	}
 
@@ -727,7 +858,29 @@ void SPCGEditorGraphDebugObjectTree::OnSelectionChanged(FPCGEditorGraphDebugObje
 	{
 		if (const FPCGStack* PCGStack = InItem->GetPCGStack())
 		{
-			PCGEditor.Pin()->SetComponentAndStackBeingInspected(PCGComponent, *PCGStack);
+			PCGEditor.Pin()->SetStackBeingInspected(*PCGStack, FPCGDebugObjectSelectionMethod::DebugObjectTree);
+		}
+	}
+}
+
+void SPCGEditorGraphDebugObjectTree::OnExpansionChanged(FPCGEditorGraphDebugObjectItemPtr InItem, bool bInIsExpanded)
+{
+	if (!InItem.IsValid())
+	{
+		return;
+	}
+
+	InItem->SetExpanded(bInIsExpanded);
+
+	if (const FPCGStack* Stack = InItem->GetPCGStack())
+	{
+		if (bInIsExpanded)
+		{
+			ExpandedStacks.Add(*Stack);
+		}
+		else
+		{
+			ExpandedStacks.Remove(*Stack);
 		}
 	}
 }
@@ -740,6 +893,8 @@ void SPCGEditorGraphDebugObjectTree::OnSetExpansionRecursive(FPCGEditorGraphDebu
 	}
 
 	DebugObjectTreeView->SetItemExpansion(InItem, bInExpand);
+	InItem->SetExpanded(bInExpand);
+
 	for (const FPCGEditorGraphDebugObjectItemPtr& ChildItem : InItem->GetChildren())
 	{
 		if (ChildItem.IsValid())
diff --git a/Engine/Plugins/Experimental/PCG/Source/PCGEditor/Private/SPCGEditorGraphDebugObjectTree.h b/Engine/Plugins/Experimental/PCG/Source/PCGEditor/Private/SPCGEditorGraphDebugObjectTree.h
--- a/Engine/Plugins/Experimental/PCG/Source/PCGEditor/Private/SPCGEditorGraphDebugObjectTree.h
+++ b/Engine/Plugins/Experimental/PCG/Source/PCGEditor/Private/SPCGEditorGraphDebugObjectTree.h
@@ -28,6 +28,15 @@ public:
 	FPCGEditorGraphDebugObjectItemPtr GetParent() const;
 	void SortChildren(bool bIsAscending, bool bIsRecursive);
 
+	bool IsExpanded() const { return bIsExpanded; }
+	void SetExpanded(bool bInIsExpanded) { bIsExpanded = bInIsExpanded; }
+
+	/** Optional sort priority, if returns INDEX_NONE then sort will fall back to alphabetical. */
+	virtual int32 GetSortPriority() const { return INDEX_NONE; }
+
+	/** Whether this item represents a currently debuggable object for the current edited graph. */
+	virtual bool IsDebuggable() const { return false; }
+
 	virtual FString GetLabel() const = 0;
 	virtual UPCGComponent* GetPCGComponent() const = 0;
 	virtual const FPCGStack* GetPCGStack() const { return nullptr; }
@@ -36,6 +45,7 @@ public:
 protected:
 	TWeakPtr<FPCGEditorGraphDebugObjectItem> Parent;
 	TSet<TSharedPtr<FPCGEditorGraphDebugObjectItem>> Children;
+	bool bIsExpanded = false;
 };
 
 class FPCGEditorGraphDebugObjectItem_Actor : public FPCGEditorGraphDebugObjectItem
@@ -43,88 +53,106 @@ class FPCGEditorGraphDebugObjectItem_Actor : public FPCGEditorGraphDebugObjectIt
 public:
 	FPCGEditorGraphDebugObjectItem_Actor(TWeakObjectPtr<AActor> InActor)
 		: Actor(InActor)
-	{ }
+	{
+		PCGStack.PushFrame(InActor.Get());
+	}
 
 	virtual FString GetLabel() const override;
 	virtual UPCGComponent* GetPCGComponent() const override { return nullptr; }
 	virtual const UObject* GetObject() const override { return Actor.Get(); };
+	virtual const FPCGStack* GetPCGStack() const override { return &PCGStack; }
 
 protected:
 	TWeakObjectPtr<AActor> Actor = nullptr;
+
+	FPCGStack PCGStack;
 };
 
 class FPCGEditorGraphDebugObjectItem_PCGComponent : public FPCGEditorGraphDebugObjectItem
 {
 public:
-	FPCGEditorGraphDebugObjectItem_PCGComponent(TWeakObjectPtr<UPCGComponent> InPCGComponent)
+	FPCGEditorGraphDebugObjectItem_PCGComponent(
+		TWeakObjectPtr<UPCGComponent> InPCGComponent,
+		TWeakObjectPtr<const UPCGGraph> InPCGGraph,
+		const FPCGStack& InPCGStack,
+		bool bInIsDebuggable)
 		: PCGComponent(InPCGComponent)
+		, PCGGraph(InPCGGraph)
+		, PCGStack(InPCGStack)
+		, bIsDebuggable(bInIsDebuggable)
 	{}
 
 	virtual FString GetLabel() const override;
 	virtual UPCGComponent* GetPCGComponent() const override { return PCGComponent.Get(); }
 	virtual const UObject* GetObject() const override { return PCGComponent.Get(); };
-
-protected:
-	TWeakObjectPtr<UPCGComponent> PCGComponent = nullptr;
-};
-
-class FPCGEditorGraphDebugObjectItem_PCGGraph : public FPCGEditorGraphDebugObjectItem
-{
-public:
-	FPCGEditorGraphDebugObjectItem_PCGGraph(TWeakObjectPtr<const UPCGGraph> InPCGGraph, const FPCGStack& InPCGStack)
-		: PCGGraph(InPCGGraph)
-		, PCGStack(InPCGStack)
-	{}
-
-	virtual FString GetLabel() const override;
-	virtual UPCGComponent* GetPCGComponent() const override { return GetParent() ? GetParent()->GetPCGComponent() : nullptr; }
-	virtual const UObject* GetObject() const override { return PCGGraph.Get(); };
 	virtual const FPCGStack* GetPCGStack() const override { return &PCGStack; }
+	virtual bool IsDebuggable() const override { return bIsDebuggable; }
 
 protected:
+	TWeakObjectPtr<UPCGComponent> PCGComponent = nullptr;
 	TWeakObjectPtr<const UPCGGraph> PCGGraph = nullptr;
+
 	FPCGStack PCGStack;
+	bool bIsDebuggable = false;
 };
 
 class FPCGEditorGraphDebugObjectItem_PCGSubgraph : public FPCGEditorGraphDebugObjectItem
 {
 public:
-	FPCGEditorGraphDebugObjectItem_PCGSubgraph(TWeakObjectPtr<const UPCGNode> InPCGNode, TWeakObjectPtr<const UPCGGraph> InPCGGraph, const FPCGStack& InPCGStack)
+	FPCGEditorGraphDebugObjectItem_PCGSubgraph(
+		TWeakObjectPtr<const UPCGNode> InPCGNode,
+		TWeakObjectPtr<const UPCGGraph> InPCGGraph,
+		const FPCGStack& InPCGStack,
+		bool bInIsDebuggable)
 		: PCGNode(InPCGNode)
 		, PCGGraph(InPCGGraph)
 		, PCGStack(InPCGStack)
+		, bIsDebuggable(bInIsDebuggable)
 	{}
 
 	virtual FString GetLabel() const override;
 	virtual UPCGComponent* GetPCGComponent() const override { return GetParent() ? GetParent()->GetPCGComponent() : nullptr; }
 	virtual const UObject* GetObject() const override { return PCGNode.Get(); };
 	virtual const FPCGStack* GetPCGStack() const override { return &PCGStack; }
+	virtual bool IsDebuggable() const override { return bIsDebuggable; }
 
 protected:
 	TWeakObjectPtr<const UPCGNode> PCGNode = nullptr;
 	TWeakObjectPtr<const UPCGGraph> PCGGraph = nullptr;
+
 	FPCGStack PCGStack;
+	bool bIsDebuggable = false;
 };
 
 class FPCGEditorGraphDebugObjectItem_PCGLoopIndex : public FPCGEditorGraphDebugObjectItem
 {
 public:
-	FPCGEditorGraphDebugObjectItem_PCGLoopIndex(TWeakObjectPtr<const UPCGNode> InPCGNode, int32 InLoopIndex, const FPCGStack& InPCGStack)
-		: PCGNode(InPCGNode)
-		, LoopIndex(InLoopIndex)
+	FPCGEditorGraphDebugObjectItem_PCGLoopIndex(
+		int32 InLoopIndex,
+		TWeakObjectPtr<const UObject> InLoopedPCGGraph,
+		const FPCGStack& InPCGStack,
+		bool bInIsDebuggable)
+		: LoopIndex(InLoopIndex)
+		, LoopedPCGGraph(InLoopedPCGGraph)
 		, PCGStack(InPCGStack)
+		, bIsDebuggable(bInIsDebuggable)
 	{}
 
+	virtual int32 GetLoopIndex() const { return LoopIndex; }
+
 	virtual FString GetLabel() const override;
 	virtual UPCGComponent* GetPCGComponent() const override { return GetParent() ? GetParent()->GetPCGComponent() : nullptr; }
-	virtual const UObject* GetObject() const override { return PCGNode.Get(); };
+	virtual const UObject* GetObject() const override { return nullptr; };
 	virtual const FPCGStack* GetPCGStack() const override { return &PCGStack; }
-	virtual int32 GetLoopIndex() const { return LoopIndex; }
+	virtual int32 GetSortPriority() const override { return LoopIndex; }
+	virtual bool IsDebuggable() const override { return bIsDebuggable; }
 
 protected:
-	TWeakObjectPtr<const UPCGNode> PCGNode = nullptr;
 	int32 LoopIndex = INDEX_NONE;
+	TWeakObjectPtr<const UObject> LoopedPCGGraph = nullptr;
+
 	FPCGStack PCGStack;
+	bool bIsDebuggable = false;
 };
 
 class SPCGEditorGraphDebugObjectItemRow : public SCompoundWidget
@@ -153,7 +181,7 @@ public:
 
 	void RequestRefresh() { bNeedsRefresh = true; }
 
-	void AddDynamicStack(TWeakObjectPtr<UPCGComponent> InComponent, const FPCGStack& InvocationStack);
+	void SetDebugObjectSelection(const FPCGStack& FullStack);
 
 private:
 	void SelectedDebugObject_OnClicked() const;
@@ -165,6 +193,11 @@ private:
 	void RefreshTree();
 	void SortTreeItems(bool bIsAscending = true, bool bIsRecursive = true);
 
+	void AddStacksToTree(
+		const TArray<FPCGStack>& Stacks,
+		TMap<AActor*, TSharedPtr<FPCGEditorGraphDebugObjectItem_Actor>>& InOutActorItems,
+		TMap<const FPCGStack, FPCGEditorGraphDebugObjectItemPtr>& InOutStackToItem);
+
 	void OnPreObjectPropertyChanged(UObject* InObject, const FEditPropertyChain& InPropertyChain);
 	void OnObjectPropertyChanged(UObject* InObject, FPropertyChangedEvent& InPropertyChangedEvent);
 	void OnObjectConstructed(UObject* InObject);
@@ -173,7 +206,8 @@ private:
 
 	TSharedRef<ITableRow> MakeTreeRowWidget(FPCGEditorGraphDebugObjectItemPtr InItem, const TSharedRef<STableViewBase>& InOwnerTable) const;
 	void OnGetChildren(FPCGEditorGraphDebugObjectItemPtr InItem, TArray<FPCGEditorGraphDebugObjectItemPtr>& OutChildren) const;
-	void OnSelectionChanged(FPCGEditorGraphDebugObjectItemPtr InItem, ESelectInfo::Type InSelectInfo) const;
+	void OnSelectionChanged(FPCGEditorGraphDebugObjectItemPtr InItem, ESelectInfo::Type InSelectInfo);
+	void OnExpansionChanged(FPCGEditorGraphDebugObjectItemPtr InItem, bool bInIsExpanded);
 	void OnSetExpansionRecursive(FPCGEditorGraphDebugObjectItemPtr InItem, bool bInExpand);
 
 	TWeakPtr<FPCGEditor> PCGEditor;
@@ -182,7 +216,14 @@ private:
 	TArray<FPCGEditorGraphDebugObjectItemPtr> RootItems;
 	TArray<FPCGEditorGraphDebugObjectItemPtr> AllGraphItems;
 
-	TMap<const TWeakObjectPtr<UPCGComponent>, TArray<FPCGStack>> DynamicInvocationStacks;
-
 	bool bNeedsRefresh = false;
+
+	/** Set true to avoid broadcasting debug object change notifications when setting the object from code. */
+	bool bDisableDebugObjectChangeNotification = false;
+
+	/** Used to retain item expansion state across tree refreshes. */
+	TSet<FPCGStack> ExpandedStacks;
+
+	/** Used to retain item selection state across tree refreshes. */
+	FPCGStack SelectedStack;
 };
diff --git a/Engine/Plugins/Experimental/PCG/Source/PCGEditor/Private/SPCGEditorGraphDebugObjectWidget.cpp b/Engine/Plugins/Experimental/PCG/Source/PCGEditor/Private/SPCGEditorGraphDebugObjectWidget.cpp
--- a/Engine/Plugins/Experimental/PCG/Source/PCGEditor/Private/SPCGEditorGraphDebugObjectWidget.cpp
+++ b/Engine/Plugins/Experimental/PCG/Source/PCGEditor/Private/SPCGEditorGraphDebugObjectWidget.cpp
@@ -3,6 +3,8 @@
 #include "SPCGEditorGraphDebugObjectWidget.h"
 
 #include "PCGComponent.h"
+#include "PCGSubsystem.h"
+
 #include "PCGEditor.h"
 #include "PCGEditorGraph.h"
 
@@ -45,10 +47,11 @@ void SPCGEditorGraphDebugObjectWidget::Construct(const FArguments& InArgs, TShar
 {
 	PCGEditorPtr = InPCGEditor;
 	
-	if (UPCGComponent* PCGComponent = InPCGEditor->GetPCGComponentBeingInspected())
+	UPCGComponent* PCGComponent = InPCGEditor->GetPCGComponentBeingInspected();
+	const FPCGStack* PCGStack = InPCGEditor->GetStackBeingInspected();
+	if (PCGComponent && PCGStack)
 	{
-		const FPCGStack& PCGStack = InPCGEditor->GetStackBeingInspected();
-		DebugObjects.Add(MakeShared<FPCGEditorGraphDebugObjectInstance>(PCGComponent, PCGStack));
+		DebugObjects.Add(MakeShared<FPCGEditorGraphDebugObjectInstance>(PCGComponent, *PCGStack));
 	}
 	else
 	{
@@ -102,6 +105,8 @@ void SPCGEditorGraphDebugObjectWidget::Construct(const FArguments& InArgs, TShar
 			BrowseButton
 		]
 	];
+
+	RefreshDebugObjects();
 }
 
 void SPCGEditorGraphDebugObjectWidget::RefreshDebugObjects()
@@ -110,7 +115,8 @@ void SPCGEditorGraphDebugObjectWidget::RefreshDebugObjects()
 	DebugObjectsComboBox->RefreshOptions();
 
 	const UPCGGraph* PCGGraph = GetPCGGraph();
-	if (!PCGGraph)
+	UPCGSubsystem* Subsystem = FPCGEditor::GetSubsystem();
+	if (!PCGGraph || !Subsystem)
 	{
 		return;
 	}
@@ -151,29 +157,16 @@ void SPCGEditorGraphDebugObjectWidget::RefreshDebugObjects()
 			continue;
 		}
 
-		FPCGStackContext StackContext = FPCGStackContext::CreateStackContextFromGraph(PCGComponentGraph);
-
-		for (const FPCGStack& Stack : StackContext.GetStacks())
+		FPCGStackContext StackContext;
+		if (!PCGComponent->GetStackContext(StackContext))
 		{
-			const FPCGStackFrame& StackFrame = Stack.GetStackFrames().Top();
-			if (const UPCGGraph* StackGraph = Cast<const UPCGGraph>(StackFrame.Object))
-			{
-				if (StackGraph == PCGGraph)
-				{
-					const TSharedPtr<FPCGEditorGraphDebugObjectInstance> DebugInstance = MakeShared<FPCGEditorGraphDebugObjectInstance>(PCGComponent, Stack);
-					DebugObjects.Add(DebugInstance);
-
-					if (SelectedItem.IsValid() && SelectedItem->GetPCGComponent() == PCGComponent && SelectedItem->GetStack() == Stack)
-					{
-						DebugObjectsComboBox->SetSelectedItem(DebugInstance);
-					}
-				}
-			}
+			continue;
 		}
 
-		if (const TArray<FPCGStack>* DynamicStacks = DynamicInvocationStacks.Find(PCGComponent))
+		// Process static stacks that can be read from the compiled graph.
+		for (const FPCGStack& Stack : StackContext.GetStacks())
 		{
-			for (const FPCGStack& Stack : *DynamicStacks)
+			if (Stack.GetRootGraph() == PCGGraph)
 			{
 				const TSharedPtr<FPCGEditorGraphDebugObjectInstance> DebugInstance = MakeShared<FPCGEditorGraphDebugObjectInstance>(PCGComponent, Stack);
 				DebugObjects.Add(DebugInstance);
@@ -184,6 +177,19 @@ void SPCGEditorGraphDebugObjectWidget::RefreshDebugObjects()
 				}
 			}
 		}
+
+		// Process stacks encountered during execution so far, which will include dynamic subgraphs & loop subgraphs.
+		// There will be overlaps with the static stacks but only unique entries will be added to the tree.
+		for (const FPCGStack& Stack : Subsystem->GetExecutedStacks(PCGComponent, PCGGraph))
+		{
+			const TSharedPtr<FPCGEditorGraphDebugObjectInstance> DebugInstance = MakeShared<FPCGEditorGraphDebugObjectInstance>(PCGComponent, Stack);
+			DebugObjects.Add(DebugInstance);
+
+			if (SelectedItem.IsValid() && SelectedItem->GetPCGComponent() == PCGComponent && SelectedItem->GetStack() == Stack)
+			{
+				DebugObjectsComboBox->SetSelectedItem(DebugInstance);
+			}
+		}
 	}
 }
 
@@ -194,10 +200,9 @@ void SPCGEditorGraphDebugObjectWidget::OnComboBoxOpening()
 
 void SPCGEditorGraphDebugObjectWidget::OnSelectionChanged(TSharedPtr<FPCGEditorGraphDebugObjectInstance> NewSelection, ESelectInfo::Type SelectInfo) const
 {
-	if (NewSelection.IsValid())
+	if (!bDisableDebugObjectChangeNotification && NewSelection.IsValid() && SelectInfo != ESelectInfo::Direct)
 	{
-		UPCGComponent* PCGComponent = NewSelection->GetPCGComponent().Get();
-		PCGEditorPtr.Pin()->SetComponentAndStackBeingInspected(PCGComponent, NewSelection->GetStack());
+		PCGEditorPtr.Pin()->SetStackBeingInspected(NewSelection->GetStack(), FPCGDebugObjectSelectionMethod::DebugObjectDropdown);
 	}
 }
 
@@ -264,20 +269,26 @@ void SPCGEditorGraphDebugObjectWidget::SetDebugObjectFromSelection_OnClicked()
 		const TSharedPtr<FPCGEditorGraphDebugObjectInstance> DebugInstance = MakeShared<FPCGEditorGraphDebugObjectInstance>(InPCGComponent, InStack);
 		DebugObjects.Add(DebugInstance);
 		DebugObjectsComboBox->SetSelectedItem(DebugInstance);
-		PCGEditorPtr.Pin()->SetComponentAndStackBeingInspected(InPCGComponent, InStack);
-		
+		PCGEditorPtr.Pin()->SetStackBeingInspected(InStack, FPCGDebugObjectSelectionMethod::DebugObjectDropdown);
+
 		return true;
 	});
 }
 
 bool SPCGEditorGraphDebugObjectWidget::IsSetDebugObjectFromSelectionButtonEnabled() const
-{	
+{
 	return ForEachStackInSelection([](const FPCGStack& /*InStack*/, const UPCGComponent* /*InPCGComponent*/)
 	{
 		return true;
 	});
 }
 
+void SPCGEditorGraphDebugObjectWidget::OnDebugObjectChanged(UPCGComponent* InPCGComponent, const FPCGStack& InPCGStack)
+{
+	const TSharedPtr<FPCGEditorGraphDebugObjectInstance> DebugInstance = MakeShared<FPCGEditorGraphDebugObjectInstance>(InPCGComponent, InPCGStack);
+	DebugObjectsComboBox->SetSelectedItem(DebugInstance);
+}
+
 bool SPCGEditorGraphDebugObjectWidget::ForEachStackInSelection(const TFunctionRef<bool(const FPCGStack&, UPCGComponent*)>& InOperation) const
 {
 	const UPCGGraph* PCGGraph = GetPCGGraph();
@@ -310,7 +321,12 @@ bool SPCGEditorGraphDebugObjectWidget::ForEachStackInSelection(const TFunctionRe
 				continue;
 			}
 
-			FPCGStackContext StackContext = FPCGStackContext::CreateStackContextFromGraph(PCGComponent->GetGraph());
+			FPCGStackContext StackContext;
+			if (!PCGComponent->GetStackContext(StackContext))
+			{
+				continue;
+			}
+
 			for (const FPCGStack& Stack : StackContext.GetStacks())
 			{
 				const FPCGStackFrame& StackFrame = Stack.GetStackFrames().Top();
@@ -386,15 +402,25 @@ void SPCGEditorGraphDebugObjectWidget::OnLevelActorDeleted(const AActor* InActor
 	}
 }
 
-void SPCGEditorGraphDebugObjectWidget::AddDynamicStack(const TWeakObjectPtr<UPCGComponent> InComponent, const FPCGStack& InvocationStack)
+void SPCGEditorGraphDebugObjectWidget::SetDebugObjectSelection(const FPCGStack& FullStack)
 {
-	TArray<FPCGStack>& Stacks = DynamicInvocationStacks.FindOrAdd(InComponent);
+	bDisableDebugObjectChangeNotification = true;
+	ON_SCOPE_EXIT
+	{
+		bDisableDebugObjectChangeNotification = false;
+	};
 
-	if (!Stacks.Contains(InvocationStack))
+	for (const TSharedPtr<FPCGEditorGraphDebugObjectInstance>& Object : DebugObjects)
 	{
-		Stacks.Add(InvocationStack);
-		RefreshDebugObjects();
+		if (Object.IsValid() && Object->GetStack() == FullStack)
+		{
+			DebugObjectsComboBox->SetSelectedItem(Object);
+			return;
+		}
 	}
+
+	// No matching debug object found, select nothing.
+	DebugObjectsComboBox->SetSelectedItem(DebugObjects[0]);
 }
 
 #undef LOCTEXT_NAMESPACE
diff --git a/Engine/Plugins/Experimental/PCG/Source/PCGEditor/Private/SPCGEditorGraphDebugObjectWidget.h b/Engine/Plugins/Experimental/PCG/Source/PCGEditor/Private/SPCGEditorGraphDebugObjectWidget.h
--- a/Engine/Plugins/Experimental/PCG/Source/PCGEditor/Private/SPCGEditorGraphDebugObjectWidget.h
+++ b/Engine/Plugins/Experimental/PCG/Source/PCGEditor/Private/SPCGEditorGraphDebugObjectWidget.h
@@ -33,7 +33,6 @@ private:
 	FPCGStack PCGStack;
 };
 
-
 class SPCGEditorGraphDebugObjectWidget: public SCompoundWidget
 {
 public:
@@ -46,11 +45,13 @@ public:
 
 	void OnLevelActorDeleted(const AActor* InActor);
 
-	void AddDynamicStack(const TWeakObjectPtr<UPCGComponent> InComponent, const FPCGStack& InvocationStack);
+	void SetDebugObjectSelection(const FPCGStack& FullStack);
 
 private:
 	void OnComboBoxOpening();
+
 	void OnSelectionChanged(TSharedPtr<FPCGEditorGraphDebugObjectInstance> NewSelection, ESelectInfo::Type SelectInfo) const;
+
 	TSharedRef<SWidget> OnGenerateWidget(TSharedPtr<FPCGEditorGraphDebugObjectInstance> InDebugObjectInstance) const;
 
 	UPCGGraph* GetPCGGraph() const;
@@ -62,6 +63,8 @@ private:
 	void SetDebugObjectFromSelection_OnClicked();
 	bool IsSetDebugObjectFromSelectionButtonEnabled() const;
 
+	void OnDebugObjectChanged(UPCGComponent* InPCGComponent, const FPCGStack& InPCGStack);
+
 	bool ForEachStackInSelection(const TFunctionRef<bool(const FPCGStack&, UPCGComponent*)>& InOperation) const;
 	
 	/** Pointer back to the PCG editor that owns us */
@@ -70,5 +73,6 @@ private:
 	TArray<TSharedPtr<FPCGEditorGraphDebugObjectInstance>> DebugObjects;
 	TSharedPtr<SComboBox<TSharedPtr<FPCGEditorGraphDebugObjectInstance>>> DebugObjectsComboBox;
 
-	TMap<const TWeakObjectPtr<UPCGComponent>, TArray<FPCGStack>> DynamicInvocationStacks;
+	/** Set true to avoid broadcasting debug object change notifications when setting the object from code. */
+	bool bDisableDebugObjectChangeNotification = false;
 };
diff --git a/Engine/Plugins/Experimental/PCG/Source/PCGEditor/Private/SPCGEditorGraphLogView.cpp b/Engine/Plugins/Experimental/PCG/Source/PCGEditor/Private/SPCGEditorGraphLogView.cpp
--- a/Engine/Plugins/Experimental/PCG/Source/PCGEditor/Private/SPCGEditorGraphLogView.cpp
+++ b/Engine/Plugins/Experimental/PCG/Source/PCGEditor/Private/SPCGEditorGraphLogView.cpp
@@ -2,18 +2,19 @@
 
 #include "SPCGEditorGraphLogView.h"
 
-#include "Framework/Views/TableViewMetadata.h"
 #include "PCGComponent.h"
-#include "PCGEditor.h"
-#include "PCGEditorGraph.h"
-#include "PCGEditorGraphNode.h"
 #include "PCGGraph.h"
 #include "PCGSubgraph.h"
+#include "Graph/PCGStackContext.h"
 
-#include "Widgets/Layout/SScrollBox.h"
+#include "PCGEditor.h"
+#include "PCGEditorGraph.h"
+#include "PCGEditorGraphNode.h"
 
+#include "Framework/Views/TableViewMetadata.h"
 #include "Styling/StyleColors.h"
 #include "Widgets/Input/SButton.h"
+#include "Widgets/Layout/SScrollBox.h"
 #include "Widgets/Views/SListView.h"
 
 #define LOCTEXT_NAMESPACE "SPCGEditorGraphLogView"
@@ -107,7 +108,7 @@ SPCGEditorGraphLogView::~SPCGEditorGraphLogView()
 {
 	if (PCGEditorPtr.IsValid())
 	{
-		PCGEditorPtr.Pin()->OnInspectedComponentChangedDelegate.RemoveAll(this);
+		PCGEditorPtr.Pin()->OnInspectedStackChangedDelegate.RemoveAll(this);
 	}
 }
 
@@ -121,7 +122,7 @@ void SPCGEditorGraphLogView::Construct(const FArguments& InArgs, TSharedPtr<FPCG
 		PCGEditorGraph = PCGEditor->GetPCGEditorGraph();
 		PCGComponent = PCGEditor->GetPCGComponentBeingInspected();
 
-		PCGEditor->OnInspectedComponentChangedDelegate.AddSP(this, &SPCGEditorGraphLogView::OnDebugObjectChanged);
+		PCGEditor->OnInspectedStackChangedDelegate.AddSP(this, &SPCGEditorGraphLogView::OnDebugStackChanged);
 	}
 
 	ListViewHeader = CreateHeaderRowWidget();
@@ -396,14 +397,14 @@ void SPCGEditorGraphLogView::CreateAndAddItem(const UPCGEditorGraphNode* InPCGEd
 	}
 }
 
-void SPCGEditorGraphLogView::OnDebugObjectChanged(UPCGComponent* InPCGComponent)
+void SPCGEditorGraphLogView::OnDebugStackChanged(const FPCGStack& InPCGStack)
 {
 	if (PCGComponent.IsValid())
 	{
 		PCGComponent->OnPCGGraphGeneratedDelegate.RemoveAll(this);
 	}
 
-	PCGComponent = InPCGComponent;
+	PCGComponent = const_cast<UPCGComponent*>(InPCGStack.GetRootComponent());
 
 	if (PCGComponent.IsValid())
 	{
diff --git a/Engine/Plugins/Experimental/PCG/Source/PCGEditor/Private/SPCGEditorGraphLogView.h b/Engine/Plugins/Experimental/PCG/Source/PCGEditor/Private/SPCGEditorGraphLogView.h
--- a/Engine/Plugins/Experimental/PCG/Source/PCGEditor/Private/SPCGEditorGraphLogView.h
+++ b/Engine/Plugins/Experimental/PCG/Source/PCGEditor/Private/SPCGEditorGraphLogView.h
@@ -12,6 +12,7 @@ class UPCGComponent;
 class UPCGEditorGraphNode;
 class UPCGEditorGraph;
 class UPCGNode;
+struct FPCGStack;
 
 namespace PCGEditorGraphLogView
 {	
@@ -62,7 +63,8 @@ private:
 	TSharedRef<SHeaderRow> CreateHeaderRowWidget();
 	void CreateAndAddItem(const UPCGEditorGraphNode* InPCGEditorNode, const UPCGNode* InPCGNode, const FName& InName);
 
-	void OnDebugObjectChanged(UPCGComponent* InPCGComponent);
+	void OnDebugStackChanged(const FPCGStack& InPCGStack);
+
 	void OnGenerateUpdated(UPCGComponent* InPCGComponent);
 	
 	// Callbacks
diff --git a/Engine/Plugins/Experimental/PCG/Source/PCGEditor/Private/SPCGEditorGraphProfilingView.cpp b/Engine/Plugins/Experimental/PCG/Source/PCGEditor/Private/SPCGEditorGraphProfilingView.cpp
--- a/Engine/Plugins/Experimental/PCG/Source/PCGEditor/Private/SPCGEditorGraphProfilingView.cpp
+++ b/Engine/Plugins/Experimental/PCG/Source/PCGEditor/Private/SPCGEditorGraphProfilingView.cpp
@@ -2,12 +2,14 @@
 
 #include "SPCGEditorGraphProfilingView.h"
 
-#include "Framework/Views/TableViewMetadata.h"
 #include "PCGComponent.h"
+#include "Graph/PCGStackContext.h"
+
 #include "PCGEditor.h"
 #include "PCGEditorGraph.h"
 #include "PCGEditorGraphNode.h"
 
+#include "Framework/Views/TableViewMetadata.h"
 #include "Widgets/Input/SButton.h"
 #include "Widgets/Input/SCheckBox.h"
 #include "Widgets/Layout/SScrollBox.h"
@@ -116,7 +118,7 @@ SPCGEditorGraphProfilingView::~SPCGEditorGraphProfilingView()
 {
 	if (PCGEditorPtr.IsValid())
 	{
-		PCGEditorPtr.Pin()->OnInspectedComponentChangedDelegate.RemoveAll(this);
+		PCGEditorPtr.Pin()->OnInspectedStackChangedDelegate.RemoveAll(this);
 	}
 }
 
@@ -130,7 +132,7 @@ void SPCGEditorGraphProfilingView::Construct(const FArguments& InArgs, TSharedPt
 		PCGEditorGraph = PCGEditor->GetPCGEditorGraph();
 		PCGComponent = PCGEditor->GetPCGComponentBeingInspected();
 
-		PCGEditor->OnInspectedComponentChangedDelegate.AddSP(this, &SPCGEditorGraphProfilingView::OnDebugObjectChanged);
+		PCGEditor->OnInspectedStackChangedDelegate.AddSP(this, &SPCGEditorGraphProfilingView::OnDebugStackChanged);
 	}
 
 	SortingColumn = PCGEditorGraphProfilingView::NAME_TotalExecutionTime;
@@ -484,14 +486,14 @@ FReply SPCGEditorGraphProfilingView::Refresh()
 	return FReply::Handled();
 }
 
-void SPCGEditorGraphProfilingView::OnDebugObjectChanged(UPCGComponent* InPCGComponent)
+void SPCGEditorGraphProfilingView::OnDebugStackChanged(const FPCGStack& InPCGStack)
 {
 	if (PCGComponent.IsValid())
 	{
 		PCGComponent->OnPCGGraphGeneratedDelegate.RemoveAll(this);
 	}
 
-	PCGComponent = InPCGComponent;
+	PCGComponent = const_cast<UPCGComponent*>(InPCGStack.GetRootComponent());
 
 	if (PCGComponent.IsValid())
 	{
diff --git a/Engine/Plugins/Experimental/PCG/Source/PCGEditor/Private/SPCGEditorGraphProfilingView.h b/Engine/Plugins/Experimental/PCG/Source/PCGEditor/Private/SPCGEditorGraphProfilingView.h
--- a/Engine/Plugins/Experimental/PCG/Source/PCGEditor/Private/SPCGEditorGraphProfilingView.h
+++ b/Engine/Plugins/Experimental/PCG/Source/PCGEditor/Private/SPCGEditorGraphProfilingView.h
@@ -3,16 +3,17 @@
 #pragma once
 
 #include "Utils/PCGExtraCapture.h"
+
 #include "Widgets/Views/ITableRow.h"
 #include "Widgets/Views/STableRow.h"
 
-class STableViewBase;
-
 class FPCGEditor;
+class STableViewBase;
 class UPCGComponent;
 class UPCGEditorGraphNode;
 class UPCGEditorGraph;
 class UPCGNode;
+struct FPCGStack;
 
 struct FPCGProfilingListViewItem
 {
@@ -58,7 +59,8 @@ private:
 	ECheckBoxState IsSubgraphExpanded() const;
 	void OnSubgraphExpandedChanged(ECheckBoxState InNewState);
 
-	void OnDebugObjectChanged(UPCGComponent* InPCGComponent);
+	void OnDebugStackChanged(const FPCGStack& InPCGStack);
+
 	void OnGenerateUpdated(UPCGComponent* InPCGComponent);
 	
 	// Callbacks
