[PCG] Fix duplicate debug object entries while in PIE

#rnx

diff --git a/Engine/Plugins/Experimental/PCG/Source/PCGEditor/Private/SPCGEditorGraphDebugObjectTree.cpp b/Engine/Plugins/Experimental/PCG/Source/PCGEditor/Private/SPCGEditorGraphDebugObjectTree.cpp
--- a/Engine/Plugins/Experimental/PCG/Source/PCGEditor/Private/SPCGEditorGraphDebugObjectTree.cpp
+++ b/Engine/Plugins/Experimental/PCG/Source/PCGEditor/Private/SPCGEditorGraphDebugObjectTree.cpp
@@ -9,6 +9,7 @@
 
 #include "PCGEditor.h"
 #include "PCGEditorGraph.h"
+#include "Helpers/PCGHelpers.h"
 
 #include "PropertyCustomizationHelpers.h"
 #include "Selection.h"
@@ -511,6 +512,18 @@ void SPCGEditorGraphDebugObjectTree::AddStacksToTree(const TArray<FPCGStack>& St
 			continue;
 		}
 
+		// Prevent duplicate entries from the editor world while in PIE.
+		if (PCGHelpers::IsRuntimeOrPIE())
+		{
+			if (UWorld* World = PCGComponent->GetWorld())
+			{
+				if (!World->IsGameWorld())
+				{
+					continue;
+				}
+			}
+		}
+
 		UPCGGraph* TopGraph = const_cast<UPCGGraph*>(Stack.GetRootGraph());
 		if (!TopGraph)
 		{
diff --git a/Engine/Plugins/Experimental/PCG/Source/PCGEditor/Private/SPCGEditorGraphDebugObjectWidget.cpp b/Engine/Plugins/Experimental/PCG/Source/PCGEditor/Private/SPCGEditorGraphDebugObjectWidget.cpp
--- a/Engine/Plugins/Experimental/PCG/Source/PCGEditor/Private/SPCGEditorGraphDebugObjectWidget.cpp
+++ b/Engine/Plugins/Experimental/PCG/Source/PCGEditor/Private/SPCGEditorGraphDebugObjectWidget.cpp
@@ -7,6 +7,7 @@
 
 #include "PCGEditor.h"
 #include "PCGEditorGraph.h"
+#include "Helpers/PCGHelpers.h"
 
 #include "PropertyCustomizationHelpers.h"
 #include "Selection.h"
@@ -139,6 +140,18 @@ void SPCGEditorGraphDebugObjectWidget::RefreshDebugObjects()
 			continue;
 		}
 
+		// Prevent duplicate entries from the editor world while in PIE.
+		if (PCGHelpers::IsRuntimeOrPIE())
+		{
+			if (UWorld* World = PCGComponentObject->GetWorld())
+			{
+				if (!World->IsGameWorld())
+				{
+					continue;
+				}
+			}
+		}
+
 		UPCGComponent* PCGComponent = Cast<UPCGComponent>(PCGComponentObject);
 		if (!PCGComponent)
 		{
