PCG Unit Tests Should be Run During Automation

diff --git a/Engine/Plugins/Experimental/PCG/Source/PCG/Private/Tests/Accessors/PCGAttributeAccessorTest.cpp b/Engine/Plugins/Experimental/PCG/Source/PCG/Private/Tests/Accessors/PCGAttributeAccessorTest.cpp
--- a/Engine/Plugins/Experimental/PCG/Source/PCG/Private/Tests/Accessors/PCGAttributeAccessorTest.cpp
+++ b/Engine/Plugins/Experimental/PCG/Source/PCG/Private/Tests/Accessors/PCGAttributeAccessorTest.cpp
@@ -79,26 +79,26 @@ namespace FPCGAttributeAccessorTestHelpers
 	}
 }
 
-IMPLEMENT_CUSTOM_SIMPLE_AUTOMATION_TEST(FPCGAttributeAccessorSingleGetDefaultTest, FPCGTestBaseClass, "pcg.tests.Accessor.Attribute.SingleGetDefault", PCGTestsCommon::TestFlags)
-IMPLEMENT_CUSTOM_SIMPLE_AUTOMATION_TEST(FPCGAttributeAccessorSingleGetTest, FPCGTestBaseClass, "pcg.tests.Accessor.Attribute.SingleGet", PCGTestsCommon::TestFlags)
-IMPLEMENT_CUSTOM_SIMPLE_AUTOMATION_TEST(FPCGAttributeAccessorSingleGetOutsideRangeTest, FPCGTestBaseClass, "pcg.tests.Accessor.Attribute.SingleGetOutsideRange", PCGTestsCommon::TestFlags)
+IMPLEMENT_CUSTOM_SIMPLE_AUTOMATION_TEST(FPCGAttributeAccessorSingleGetDefaultTest, FPCGTestBaseClass, "Plugins.PCG.Accessor.Attribute.SingleGetDefault", PCGTestsCommon::TestFlags)
+IMPLEMENT_CUSTOM_SIMPLE_AUTOMATION_TEST(FPCGAttributeAccessorSingleGetTest, FPCGTestBaseClass, "Plugins.PCG.Accessor.Attribute.SingleGet", PCGTestsCommon::TestFlags)
+IMPLEMENT_CUSTOM_SIMPLE_AUTOMATION_TEST(FPCGAttributeAccessorSingleGetOutsideRangeTest, FPCGTestBaseClass, "Plugins.PCG.Accessor.Attribute.SingleGetOutsideRange", PCGTestsCommon::TestFlags)
 
-IMPLEMENT_CUSTOM_SIMPLE_AUTOMATION_TEST(FPCGAttributeAccessorGetRangeDefaultTest, FPCGTestBaseClass, "pcg.tests.Accessor.Attribute.GetRangeDefault", PCGTestsCommon::TestFlags)
-IMPLEMENT_CUSTOM_SIMPLE_AUTOMATION_TEST(FPCGAttributeAccessorGetRangeTest, FPCGTestBaseClass, "pcg.tests.Accessor.Attribute.GetRange", PCGTestsCommon::TestFlags)
-IMPLEMENT_CUSTOM_SIMPLE_AUTOMATION_TEST(FPCGAttributeAccessorGetRangeOutsideRangeTest, FPCGTestBaseClass, "pcg.tests.Accessor.Attribute.GetRangeOutsideRange", PCGTestsCommon::TestFlags)
+IMPLEMENT_CUSTOM_SIMPLE_AUTOMATION_TEST(FPCGAttributeAccessorGetRangeDefaultTest, FPCGTestBaseClass, "Plugins.PCG.Accessor.Attribute.GetRangeDefault", PCGTestsCommon::TestFlags)
+IMPLEMENT_CUSTOM_SIMPLE_AUTOMATION_TEST(FPCGAttributeAccessorGetRangeTest, FPCGTestBaseClass, "Plugins.PCG.Accessor.Attribute.GetRange", PCGTestsCommon::TestFlags)
+IMPLEMENT_CUSTOM_SIMPLE_AUTOMATION_TEST(FPCGAttributeAccessorGetRangeOutsideRangeTest, FPCGTestBaseClass, "Plugins.PCG.Accessor.Attribute.GetRangeOutsideRange", PCGTestsCommon::TestFlags)
 
-IMPLEMENT_CUSTOM_SIMPLE_AUTOMATION_TEST(FPCGAttributeAccessorSingleSetDefaultTest, FPCGTestBaseClass, "pcg.tests.Accessor.Attribute.SingleSetDefault", PCGTestsCommon::TestFlags)
-IMPLEMENT_CUSTOM_SIMPLE_AUTOMATION_TEST(FPCGAttributeAccessorSingleSetTest, FPCGTestBaseClass, "pcg.tests.Accessor.Attribute.SingleSet", PCGTestsCommon::TestFlags)
-IMPLEMENT_CUSTOM_SIMPLE_AUTOMATION_TEST(FPCGAttributeAccessorSingleInvalidKeySetTest, FPCGTestBaseClass, "pcg.tests.Accessor.Attribute.SingleSetInvalidKey", PCGTestsCommon::TestFlags)
+IMPLEMENT_CUSTOM_SIMPLE_AUTOMATION_TEST(FPCGAttributeAccessorSingleSetDefaultTest, FPCGTestBaseClass, "Plugins.PCG.Accessor.Attribute.SingleSetDefault", PCGTestsCommon::TestFlags)
+IMPLEMENT_CUSTOM_SIMPLE_AUTOMATION_TEST(FPCGAttributeAccessorSingleSetTest, FPCGTestBaseClass, "Plugins.PCG.Accessor.Attribute.SingleSet", PCGTestsCommon::TestFlags)
+IMPLEMENT_CUSTOM_SIMPLE_AUTOMATION_TEST(FPCGAttributeAccessorSingleInvalidKeySetTest, FPCGTestBaseClass, "Plugins.PCG.Accessor.Attribute.SingleSetInvalidKey", PCGTestsCommon::TestFlags)
 
-IMPLEMENT_CUSTOM_SIMPLE_AUTOMATION_TEST(FPCGAttributeAccessorSetRangeTest, FPCGTestBaseClass, "pcg.tests.Accessor.Attribute.SetRange", PCGTestsCommon::TestFlags)
-IMPLEMENT_CUSTOM_SIMPLE_AUTOMATION_TEST(FPCGAttributeAccessorSetRangeOutsideRangeTest, FPCGTestBaseClass, "pcg.tests.Accessor.Attribute.SetRangeOutsideRange", PCGTestsCommon::TestFlags)
+IMPLEMENT_CUSTOM_SIMPLE_AUTOMATION_TEST(FPCGAttributeAccessorSetRangeTest, FPCGTestBaseClass, "Plugins.PCG.Accessor.Attribute.SetRange", PCGTestsCommon::TestFlags)
+IMPLEMENT_CUSTOM_SIMPLE_AUTOMATION_TEST(FPCGAttributeAccessorSetRangeOutsideRangeTest, FPCGTestBaseClass, "Plugins.PCG.Accessor.Attribute.SetRangeOutsideRange", PCGTestsCommon::TestFlags)
 
-IMPLEMENT_CUSTOM_SIMPLE_AUTOMATION_TEST(FPCGAttributeAccessorSingleGetPointsTest, FPCGTestBaseClass, "pcg.tests.Accessor.Attribute.SingleGetPoints", PCGTestsCommon::TestFlags)
-IMPLEMENT_CUSTOM_SIMPLE_AUTOMATION_TEST(FPCGAttributeAccessorGetRangePointsTest, FPCGTestBaseClass, "pcg.tests.Accessor.Attribute.GetRangePoints", PCGTestsCommon::TestFlags)
+IMPLEMENT_CUSTOM_SIMPLE_AUTOMATION_TEST(FPCGAttributeAccessorSingleGetPointsTest, FPCGTestBaseClass, "Plugins.PCG.Accessor.Attribute.SingleGetPoints", PCGTestsCommon::TestFlags)
+IMPLEMENT_CUSTOM_SIMPLE_AUTOMATION_TEST(FPCGAttributeAccessorGetRangePointsTest, FPCGTestBaseClass, "Plugins.PCG.Accessor.Attribute.GetRangePoints", PCGTestsCommon::TestFlags)
 
-IMPLEMENT_CUSTOM_SIMPLE_AUTOMATION_TEST(FPCGAttributeAccessorSingleSetPointsTest, FPCGTestBaseClass, "pcg.tests.Accessor.Attribute.SingleSetPoints", PCGTestsCommon::TestFlags)
-IMPLEMENT_CUSTOM_SIMPLE_AUTOMATION_TEST(FPCGAttributeAccessorSetRangePointsTest, FPCGTestBaseClass, "pcg.tests.Accessor.Attribute.SetRangePoints", PCGTestsCommon::TestFlags)
+IMPLEMENT_CUSTOM_SIMPLE_AUTOMATION_TEST(FPCGAttributeAccessorSingleSetPointsTest, FPCGTestBaseClass, "Plugins.PCG.Accessor.Attribute.SingleSetPoints", PCGTestsCommon::TestFlags)
+IMPLEMENT_CUSTOM_SIMPLE_AUTOMATION_TEST(FPCGAttributeAccessorSetRangePointsTest, FPCGTestBaseClass, "Plugins.PCG.Accessor.Attribute.SetRangePoints", PCGTestsCommon::TestFlags)
 
 
 bool FPCGAttributeAccessorSingleGetDefaultTest::RunTest(const FString& Parameters)
diff --git a/Engine/Plugins/Experimental/PCG/Source/PCG/Private/Tests/Accessors/PCGAttributeExtractorTest.cpp b/Engine/Plugins/Experimental/PCG/Source/PCG/Private/Tests/Accessors/PCGAttributeExtractorTest.cpp
--- a/Engine/Plugins/Experimental/PCG/Source/PCG/Private/Tests/Accessors/PCGAttributeExtractorTest.cpp
+++ b/Engine/Plugins/Experimental/PCG/Source/PCG/Private/Tests/Accessors/PCGAttributeExtractorTest.cpp
@@ -46,13 +46,13 @@ namespace FPCGAttributeExtractorTestHelpers
 
 #define PCG_STRUCT_NAME_CHECK(Name) GET_MEMBER_NAME_CHECKED(FPCGMetadataTypesConstantStruct, Name)
 
-IMPLEMENT_CUSTOM_SIMPLE_AUTOMATION_TEST(FPCGAttributeSingleGetPropertyTest, FPCGTestBaseClass, "pcg.tests.Accessor.Property.SimpleGetProperty", PCGTestsCommon::TestFlags)
-IMPLEMENT_CUSTOM_SIMPLE_AUTOMATION_TEST(FPCGAttributeVectorPropertyExtractorTest, FPCGTestBaseClass, "pcg.tests.Accessor.Property.VectorExtractor", PCGTestsCommon::TestFlags)
-IMPLEMENT_CUSTOM_SIMPLE_AUTOMATION_TEST(FPCGAttributeRotatorPropertyExtractorTest, FPCGTestBaseClass, "pcg.tests.Accessor.Property.RotatorExtractor", PCGTestsCommon::TestFlags)
-IMPLEMENT_CUSTOM_SIMPLE_AUTOMATION_TEST(FPCGAttributeTransformPropertyExtractorTest, FPCGTestBaseClass, "pcg.tests.Accessor.Property.TransformExtractor", PCGTestsCommon::TestFlags)
-IMPLEMENT_CUSTOM_SIMPLE_AUTOMATION_TEST(FPCGAttributeInvalidPropertyExtractorTest, FPCGTestBaseClass, "pcg.tests.Accessor.Property.InvalidExtractor", PCGTestsCommon::TestFlags)
-IMPLEMENT_CUSTOM_SIMPLE_AUTOMATION_TEST(FPCGAttributePropertyMultipleDepthTest, FPCGTestBaseClass, "pcg.tests.Accessor.Property.MultipleDepth", PCGTestsCommon::TestFlags)
-IMPLEMENT_CUSTOM_SIMPLE_AUTOMATION_TEST(FPCGAttributePropertyMultipleDepthRangeTest, FPCGTestBaseClass, "pcg.tests.Accessor.Property.MultipleDepthRange", PCGTestsCommon::TestFlags)
+IMPLEMENT_CUSTOM_SIMPLE_AUTOMATION_TEST(FPCGAttributeSingleGetPropertyTest, FPCGTestBaseClass, "Plugins.PCG.Accessor.Property.SimpleGetProperty", PCGTestsCommon::TestFlags)
+IMPLEMENT_CUSTOM_SIMPLE_AUTOMATION_TEST(FPCGAttributeVectorPropertyExtractorTest, FPCGTestBaseClass, "Plugins.PCG.Accessor.Property.VectorExtractor", PCGTestsCommon::TestFlags)
+IMPLEMENT_CUSTOM_SIMPLE_AUTOMATION_TEST(FPCGAttributeRotatorPropertyExtractorTest, FPCGTestBaseClass, "Plugins.PCG.Accessor.Property.RotatorExtractor", PCGTestsCommon::TestFlags)
+IMPLEMENT_CUSTOM_SIMPLE_AUTOMATION_TEST(FPCGAttributeTransformPropertyExtractorTest, FPCGTestBaseClass, "Plugins.PCG.Accessor.Property.TransformExtractor", PCGTestsCommon::TestFlags)
+IMPLEMENT_CUSTOM_SIMPLE_AUTOMATION_TEST(FPCGAttributeInvalidPropertyExtractorTest, FPCGTestBaseClass, "Plugins.PCG.Accessor.Property.InvalidExtractor", PCGTestsCommon::TestFlags)
+IMPLEMENT_CUSTOM_SIMPLE_AUTOMATION_TEST(FPCGAttributePropertyMultipleDepthTest, FPCGTestBaseClass, "Plugins.PCG.Accessor.Property.MultipleDepth", PCGTestsCommon::TestFlags)
+IMPLEMENT_CUSTOM_SIMPLE_AUTOMATION_TEST(FPCGAttributePropertyMultipleDepthRangeTest, FPCGTestBaseClass, "Plugins.PCG.Accessor.Property.MultipleDepthRange", PCGTestsCommon::TestFlags)
 
 bool FPCGAttributeSingleGetPropertyTest::RunTest(const FString& Parameters)
 {
diff --git a/Engine/Plugins/Experimental/PCG/Source/PCG/Private/Tests/Cache/PCGLandscapeCacheTest.cpp b/Engine/Plugins/Experimental/PCG/Source/PCG/Private/Tests/Cache/PCGLandscapeCacheTest.cpp
--- a/Engine/Plugins/Experimental/PCG/Source/PCG/Private/Tests/Cache/PCGLandscapeCacheTest.cpp
+++ b/Engine/Plugins/Experimental/PCG/Source/PCG/Private/Tests/Cache/PCGLandscapeCacheTest.cpp
@@ -6,7 +6,7 @@
 
 #if WITH_EDITOR
 
-IMPLEMENT_CUSTOM_SIMPLE_AUTOMATION_TEST(FPCGLandscapeCache_CalcSafeIndices, FPCGTestBaseClass, "pcg.tests.LandscapeCache.CalcSafeIndices", PCGTestsCommon::TestFlags)
+IMPLEMENT_CUSTOM_SIMPLE_AUTOMATION_TEST(FPCGLandscapeCache_CalcSafeIndices, FPCGTestBaseClass, "Plugins.PCG.LandscapeCache.CalcSafeIndices", PCGTestsCommon::TestFlags)
 
 bool FPCGLandscapeCache_CalcSafeIndices::RunTest(const FString& Parameters)
 {
diff --git a/Engine/Plugins/Experimental/PCG/Source/PCG/Private/Tests/Data/PCGDifferenceDataTest.cpp b/Engine/Plugins/Experimental/PCG/Source/PCG/Private/Tests/Data/PCGDifferenceDataTest.cpp
--- a/Engine/Plugins/Experimental/PCG/Source/PCG/Private/Tests/Data/PCGDifferenceDataTest.cpp
+++ b/Engine/Plugins/Experimental/PCG/Source/PCG/Private/Tests/Data/PCGDifferenceDataTest.cpp
@@ -5,7 +5,7 @@
 #include "Data/PCGPointData.h"
 #include "Data/PCGVolumeData.h"
 
-IMPLEMENT_CUSTOM_SIMPLE_AUTOMATION_TEST(FPCGDifferenceDataTest, FPCGTestBaseClass, "pcg.tests.Difference.Data", PCGTestsCommon::TestFlags)
+IMPLEMENT_CUSTOM_SIMPLE_AUTOMATION_TEST(FPCGDifferenceDataTest, FPCGTestBaseClass, "Plugins.PCG.Difference.Data", PCGTestsCommon::TestFlags)
 
 bool FPCGDifferenceDataTest::RunTest(const FString& Parameters)
 {
diff --git a/Engine/Plugins/Experimental/PCG/Source/PCG/Private/Tests/Data/PCGIntersectionDataTest.cpp b/Engine/Plugins/Experimental/PCG/Source/PCG/Private/Tests/Data/PCGIntersectionDataTest.cpp
--- a/Engine/Plugins/Experimental/PCG/Source/PCG/Private/Tests/Data/PCGIntersectionDataTest.cpp
+++ b/Engine/Plugins/Experimental/PCG/Source/PCG/Private/Tests/Data/PCGIntersectionDataTest.cpp
@@ -5,7 +5,7 @@
 #include "Data/PCGPointData.h"
 #include "Data/PCGVolumeData.h"
 
-IMPLEMENT_CUSTOM_SIMPLE_AUTOMATION_TEST(FPCGIntersectionDataTest, FPCGTestBaseClass, "pcg.tests.Intersection.Data", PCGTestsCommon::TestFlags)
+IMPLEMENT_CUSTOM_SIMPLE_AUTOMATION_TEST(FPCGIntersectionDataTest, FPCGTestBaseClass, "Plugins.PCG.Intersection.Data", PCGTestsCommon::TestFlags)
 
 bool FPCGIntersectionDataTest::RunTest(const FString& Parameters)
 {
diff --git a/Engine/Plugins/Experimental/PCG/Source/PCG/Private/Tests/Data/PCGTextureDataTest.cpp b/Engine/Plugins/Experimental/PCG/Source/PCG/Private/Tests/Data/PCGTextureDataTest.cpp
--- a/Engine/Plugins/Experimental/PCG/Source/PCG/Private/Tests/Data/PCGTextureDataTest.cpp
+++ b/Engine/Plugins/Experimental/PCG/Source/PCG/Private/Tests/Data/PCGTextureDataTest.cpp
@@ -9,7 +9,7 @@
 
 #if WITH_EDITOR
 
-IMPLEMENT_CUSTOM_SIMPLE_AUTOMATION_TEST(FPCGTextureDataOffsetTilingRotation, FPCGTestBaseClass, "pcg.tests.Texture.OffsetTilingRotation", PCGTestsCommon::TestFlags)
+IMPLEMENT_CUSTOM_SIMPLE_AUTOMATION_TEST(FPCGTextureDataOffsetTilingRotation, FPCGTestBaseClass, "Plugins.PCG.Texture.OffsetTilingRotation", PCGTestsCommon::TestFlags)
 
 bool FPCGTextureDataOffsetTilingRotation::RunTest(const FString& Parameters)
 {
diff --git a/Engine/Plugins/Experimental/PCG/Source/PCG/Private/Tests/Determinism/PCGCollapseDeterminismTest.cpp b/Engine/Plugins/Experimental/PCG/Source/PCG/Private/Tests/Determinism/PCGCollapseDeterminismTest.cpp
--- a/Engine/Plugins/Experimental/PCG/Source/PCG/Private/Tests/Determinism/PCGCollapseDeterminismTest.cpp
+++ b/Engine/Plugins/Experimental/PCG/Source/PCG/Private/Tests/Determinism/PCGCollapseDeterminismTest.cpp
@@ -5,11 +5,11 @@
 #include "Elements/PCGCollapseElement.h"
 #include "Tests/Determinism/PCGDeterminismTestsCommon.h"
 
-IMPLEMENT_CUSTOM_SIMPLE_AUTOMATION_TEST(FPCGCollapseDeterminismSingleSameDataTest, FPCGTestBaseClass, "pcg.tests.Collapse.Determinism.SingleSameData", PCGTestsCommon::TestFlags)
-IMPLEMENT_CUSTOM_SIMPLE_AUTOMATION_TEST(FPCGCollapseDeterminismSingleIdenticalDataTest, FPCGTestBaseClass, "pcg.tests.Collapse.Determinism.SingleIdenticalData", PCGTestsCommon::TestFlags)
-IMPLEMENT_CUSTOM_SIMPLE_AUTOMATION_TEST(FPCGCollapseDeterminismMultipleSameDataTest, FPCGTestBaseClass, "pcg.tests.Collapse.Determinism.MultipleSameData", PCGTestsCommon::TestFlags)
-IMPLEMENT_CUSTOM_SIMPLE_AUTOMATION_TEST(FPCGCollapseDeterminismMultipleIdenticalDataTest, FPCGTestBaseClass, "pcg.tests.Collapse.Determinism.MultipleIdenticalData", PCGTestsCommon::TestFlags)
-IMPLEMENT_CUSTOM_SIMPLE_AUTOMATION_TEST(FPCGCollapseDeterminismOrderIndependenceTest, FPCGTestBaseClass, "pcg.tests.Collapse.Determinism.OrderIndependence", PCGTestsCommon::TestFlags)
+IMPLEMENT_CUSTOM_SIMPLE_AUTOMATION_TEST(FPCGCollapseDeterminismSingleSameDataTest, FPCGTestBaseClass, "Plugins.PCG.Collapse.Determinism.SingleSameData", PCGTestsCommon::TestFlags)
+IMPLEMENT_CUSTOM_SIMPLE_AUTOMATION_TEST(FPCGCollapseDeterminismSingleIdenticalDataTest, FPCGTestBaseClass, "Plugins.PCG.Collapse.Determinism.SingleIdenticalData", PCGTestsCommon::TestFlags)
+IMPLEMENT_CUSTOM_SIMPLE_AUTOMATION_TEST(FPCGCollapseDeterminismMultipleSameDataTest, FPCGTestBaseClass, "Plugins.PCG.Collapse.Determinism.MultipleSameData", PCGTestsCommon::TestFlags)
+IMPLEMENT_CUSTOM_SIMPLE_AUTOMATION_TEST(FPCGCollapseDeterminismMultipleIdenticalDataTest, FPCGTestBaseClass, "Plugins.PCG.Collapse.Determinism.MultipleIdenticalData", PCGTestsCommon::TestFlags)
+IMPLEMENT_CUSTOM_SIMPLE_AUTOMATION_TEST(FPCGCollapseDeterminismOrderIndependenceTest, FPCGTestBaseClass, "Plugins.PCG.Collapse.Determinism.OrderIndependence", PCGTestsCommon::TestFlags)
 
 namespace
 {
diff --git a/Engine/Plugins/Experimental/PCG/Source/PCG/Private/Tests/Determinism/PCGDifferenceDeterminismTest.cpp b/Engine/Plugins/Experimental/PCG/Source/PCG/Private/Tests/Determinism/PCGDifferenceDeterminismTest.cpp
--- a/Engine/Plugins/Experimental/PCG/Source/PCG/Private/Tests/Determinism/PCGDifferenceDeterminismTest.cpp
+++ b/Engine/Plugins/Experimental/PCG/Source/PCG/Private/Tests/Determinism/PCGDifferenceDeterminismTest.cpp
@@ -6,11 +6,11 @@
 
 #include "Elements/PCGDifferenceElement.h"
 
-IMPLEMENT_CUSTOM_SIMPLE_AUTOMATION_TEST(FPCGDifferenceDeterminismSingleSameDataTest, FPCGTestBaseClass, "pcg.tests.Difference.Determinism.SingleSameData", PCGTestsCommon::TestFlags)
-IMPLEMENT_CUSTOM_SIMPLE_AUTOMATION_TEST(FPCGDifferenceDeterminismSingleIdenticalDataTest, FPCGTestBaseClass, "pcg.tests.Difference.Determinism.SingleMultipleData", PCGTestsCommon::TestFlags)
-IMPLEMENT_CUSTOM_SIMPLE_AUTOMATION_TEST(FPCGDifferenceDeterminismMultipleSameDataTest, FPCGTestBaseClass, "pcg.tests.Difference.Determinism.MultipleSameData", PCGTestsCommon::TestFlags)
-IMPLEMENT_CUSTOM_SIMPLE_AUTOMATION_TEST(FPCGDifferenceDeterminismMultipleIdenticalDataTest, FPCGTestBaseClass, "pcg.tests.Difference.Determinism.MultipleIdenticalData", PCGTestsCommon::TestFlags)
-IMPLEMENT_CUSTOM_SIMPLE_AUTOMATION_TEST(FPCGDifferenceDeterminismOrderIndependenceTest, FPCGTestBaseClass, "pcg.tests.Difference.Determinism.OrderIndependence", PCGTestsCommon::TestFlags)
+IMPLEMENT_CUSTOM_SIMPLE_AUTOMATION_TEST(FPCGDifferenceDeterminismSingleSameDataTest, FPCGTestBaseClass, "Plugins.PCG.Difference.Determinism.SingleSameData", PCGTestsCommon::TestFlags)
+IMPLEMENT_CUSTOM_SIMPLE_AUTOMATION_TEST(FPCGDifferenceDeterminismSingleIdenticalDataTest, FPCGTestBaseClass, "Plugins.PCG.Difference.Determinism.SingleMultipleData", PCGTestsCommon::TestFlags)
+IMPLEMENT_CUSTOM_SIMPLE_AUTOMATION_TEST(FPCGDifferenceDeterminismMultipleSameDataTest, FPCGTestBaseClass, "Plugins.PCG.Difference.Determinism.MultipleSameData", PCGTestsCommon::TestFlags)
+IMPLEMENT_CUSTOM_SIMPLE_AUTOMATION_TEST(FPCGDifferenceDeterminismMultipleIdenticalDataTest, FPCGTestBaseClass, "Plugins.PCG.Difference.Determinism.MultipleIdenticalData", PCGTestsCommon::TestFlags)
+IMPLEMENT_CUSTOM_SIMPLE_AUTOMATION_TEST(FPCGDifferenceDeterminismOrderIndependenceTest, FPCGTestBaseClass, "Plugins.PCG.Difference.Determinism.OrderIndependence", PCGTestsCommon::TestFlags)
 
 namespace
 {
diff --git a/Engine/Plugins/Experimental/PCG/Source/PCG/Private/Tests/Determinism/PCGIntersectionDeterminismTest.cpp b/Engine/Plugins/Experimental/PCG/Source/PCG/Private/Tests/Determinism/PCGIntersectionDeterminismTest.cpp
--- a/Engine/Plugins/Experimental/PCG/Source/PCG/Private/Tests/Determinism/PCGIntersectionDeterminismTest.cpp
+++ b/Engine/Plugins/Experimental/PCG/Source/PCG/Private/Tests/Determinism/PCGIntersectionDeterminismTest.cpp
@@ -5,11 +5,11 @@
 #include "Elements/PCGInnerIntersectionElement.h"
 #include "Tests/Determinism/PCGDeterminismTestsCommon.h"
 
-IMPLEMENT_CUSTOM_SIMPLE_AUTOMATION_TEST(FPCGIntersectionDeterminismSingleSameDataTest, FPCGTestBaseClass, "pcg.tests.Intersection.Determinism.SingleSameData", PCGTestsCommon::TestFlags)
-IMPLEMENT_CUSTOM_SIMPLE_AUTOMATION_TEST(FPCGIntersectionDeterminismSingleIdenticalDataTest, FPCGTestBaseClass, "pcg.tests.Intersection.Determinism.SingleMultipleData", PCGTestsCommon::TestFlags)
-IMPLEMENT_CUSTOM_SIMPLE_AUTOMATION_TEST(FPCGIntersectionDeterminismMultipleSameDataTest, FPCGTestBaseClass, "pcg.tests.Intersection.Determinism.MultipleSameData", PCGTestsCommon::TestFlags)
-IMPLEMENT_CUSTOM_SIMPLE_AUTOMATION_TEST(FPCGIntersectionDeterminismMultipleIdenticalDataTest, FPCGTestBaseClass, "pcg.tests.Intersection.Determinism.MultipleIdenticalData", PCGTestsCommon::TestFlags)
-IMPLEMENT_CUSTOM_SIMPLE_AUTOMATION_TEST(FPCGIntersectionDeterminismOrderIndependenceTest, FPCGTestBaseClass, "pcg.tests.Intersection.Determinism.OrderIndependence", PCGTestsCommon::TestFlags)
+IMPLEMENT_CUSTOM_SIMPLE_AUTOMATION_TEST(FPCGIntersectionDeterminismSingleSameDataTest, FPCGTestBaseClass, "Plugins.PCG.Intersection.Determinism.SingleSameData", PCGTestsCommon::TestFlags)
+IMPLEMENT_CUSTOM_SIMPLE_AUTOMATION_TEST(FPCGIntersectionDeterminismSingleIdenticalDataTest, FPCGTestBaseClass, "Plugins.PCG.Intersection.Determinism.SingleMultipleData", PCGTestsCommon::TestFlags)
+IMPLEMENT_CUSTOM_SIMPLE_AUTOMATION_TEST(FPCGIntersectionDeterminismMultipleSameDataTest, FPCGTestBaseClass, "Plugins.PCG.Intersection.Determinism.MultipleSameData", PCGTestsCommon::TestFlags)
+IMPLEMENT_CUSTOM_SIMPLE_AUTOMATION_TEST(FPCGIntersectionDeterminismMultipleIdenticalDataTest, FPCGTestBaseClass, "Plugins.PCG.Intersection.Determinism.MultipleIdenticalData", PCGTestsCommon::TestFlags)
+IMPLEMENT_CUSTOM_SIMPLE_AUTOMATION_TEST(FPCGIntersectionDeterminismOrderIndependenceTest, FPCGTestBaseClass, "Plugins.PCG.Intersection.Determinism.OrderIndependence", PCGTestsCommon::TestFlags)
 
 namespace
 {
diff --git a/Engine/Plugins/Experimental/PCG/Source/PCG/Private/Tests/Determinism/PCGVolumeSamplerDeterminismTest.cpp b/Engine/Plugins/Experimental/PCG/Source/PCG/Private/Tests/Determinism/PCGVolumeSamplerDeterminismTest.cpp
--- a/Engine/Plugins/Experimental/PCG/Source/PCG/Private/Tests/Determinism/PCGVolumeSamplerDeterminismTest.cpp
+++ b/Engine/Plugins/Experimental/PCG/Source/PCG/Private/Tests/Determinism/PCGVolumeSamplerDeterminismTest.cpp
@@ -5,11 +5,11 @@
 #include "Elements/PCGVolumeSampler.h"
 #include "Tests/Determinism/PCGDeterminismTestsCommon.h"
 
-IMPLEMENT_CUSTOM_SIMPLE_AUTOMATION_TEST(FPCGVolumeSamplerDeterminismSingleSameDataTest, FPCGTestBaseClass, "pcg.tests.VolumeSampler.Determinism.SingleSameData", PCGTestsCommon::TestFlags)
-IMPLEMENT_CUSTOM_SIMPLE_AUTOMATION_TEST(FPCGVolumeSamplerDeterminismSingleIdenticalDataTest, FPCGTestBaseClass, "pcg.tests.VolumeSampler.Determinism.SingleIdenticalData", PCGTestsCommon::TestFlags)
-IMPLEMENT_CUSTOM_SIMPLE_AUTOMATION_TEST(FPCGVolumeSamplerDeterminismMultipleSameDataTest, FPCGTestBaseClass, "pcg.tests.VolumeSampler.Determinism.MultipleSameData", PCGTestsCommon::TestFlags)
-IMPLEMENT_CUSTOM_SIMPLE_AUTOMATION_TEST(FPCGVolumeSamplerDeterminismMultipleIdenticalDataTest, FPCGTestBaseClass, "pcg.tests.VolumeSampler.Determinism.MultipleIdenticalData", PCGTestsCommon::TestFlags)
-IMPLEMENT_CUSTOM_SIMPLE_AUTOMATION_TEST(FPCGVolumeSamplerDeterminismOrderIndependenceTest, FPCGTestBaseClass, "pcg.tests.VolumeSampler.Determinism.OrderIndependence", PCGTestsCommon::TestFlags)
+IMPLEMENT_CUSTOM_SIMPLE_AUTOMATION_TEST(FPCGVolumeSamplerDeterminismSingleSameDataTest, FPCGTestBaseClass, "Plugins.PCG.VolumeSampler.Determinism.SingleSameData", PCGTestsCommon::TestFlags)
+IMPLEMENT_CUSTOM_SIMPLE_AUTOMATION_TEST(FPCGVolumeSamplerDeterminismSingleIdenticalDataTest, FPCGTestBaseClass, "Plugins.PCG.VolumeSampler.Determinism.SingleIdenticalData", PCGTestsCommon::TestFlags)
+IMPLEMENT_CUSTOM_SIMPLE_AUTOMATION_TEST(FPCGVolumeSamplerDeterminismMultipleSameDataTest, FPCGTestBaseClass, "Plugins.PCG.VolumeSampler.Determinism.MultipleSameData", PCGTestsCommon::TestFlags)
+IMPLEMENT_CUSTOM_SIMPLE_AUTOMATION_TEST(FPCGVolumeSamplerDeterminismMultipleIdenticalDataTest, FPCGTestBaseClass, "Plugins.PCG.VolumeSampler.Determinism.MultipleIdenticalData", PCGTestsCommon::TestFlags)
+IMPLEMENT_CUSTOM_SIMPLE_AUTOMATION_TEST(FPCGVolumeSamplerDeterminismOrderIndependenceTest, FPCGTestBaseClass, "Plugins.PCG.VolumeSampler.Determinism.OrderIndependence", PCGTestsCommon::TestFlags)
 
 namespace
 {
diff --git a/Engine/Plugins/Experimental/PCG/Source/PCG/Private/Tests/Elements/Metadata/PCGMetadataBooleanOpTest.cpp b/Engine/Plugins/Experimental/PCG/Source/PCG/Private/Tests/Elements/Metadata/PCGMetadataBooleanOpTest.cpp
--- a/Engine/Plugins/Experimental/PCG/Source/PCG/Private/Tests/Elements/Metadata/PCGMetadataBooleanOpTest.cpp
+++ b/Engine/Plugins/Experimental/PCG/Source/PCG/Private/Tests/Elements/Metadata/PCGMetadataBooleanOpTest.cpp
@@ -10,7 +10,7 @@
 
 #if WITH_EDITOR
 
-IMPLEMENT_CUSTOM_SIMPLE_AUTOMATION_TEST(FPCGMetadataBooleanOpTest, FPCGTestBaseClass, "pcg.tests.Metadata.BooleanOp", PCGTestsCommon::TestFlags)
+IMPLEMENT_CUSTOM_SIMPLE_AUTOMATION_TEST(FPCGMetadataBooleanOpTest, FPCGTestBaseClass, "Plugins.PCG.Metadata.BooleanOp", PCGTestsCommon::TestFlags)
 
 bool FPCGMetadataBooleanOpTest::RunTest(const FString& Parameters)
 {
diff --git a/Engine/Plugins/Experimental/PCG/Source/PCG/Private/Tests/Elements/Metadata/PCGMetadataBreakTransformTest.cpp b/Engine/Plugins/Experimental/PCG/Source/PCG/Private/Tests/Elements/Metadata/PCGMetadataBreakTransformTest.cpp
--- a/Engine/Plugins/Experimental/PCG/Source/PCG/Private/Tests/Elements/Metadata/PCGMetadataBreakTransformTest.cpp
+++ b/Engine/Plugins/Experimental/PCG/Source/PCG/Private/Tests/Elements/Metadata/PCGMetadataBreakTransformTest.cpp
@@ -14,7 +14,7 @@
 
 #if WITH_EDITOR
 
-IMPLEMENT_CUSTOM_SIMPLE_AUTOMATION_TEST(FPCGMetadataBreakTransformTest, FPCGTestBaseClass, "pcg.tests.Metadata.BreakTransform", PCGTestsCommon::TestFlags)
+IMPLEMENT_CUSTOM_SIMPLE_AUTOMATION_TEST(FPCGMetadataBreakTransformTest, FPCGTestBaseClass, "Plugins.PCG.Metadata.BreakTransform", PCGTestsCommon::TestFlags)
 
 namespace PCGBreakTransformTest
 {
diff --git a/Engine/Plugins/Experimental/PCG/Source/PCG/Private/Tests/Elements/Metadata/PCGMetadataBreakVectorTest.cpp b/Engine/Plugins/Experimental/PCG/Source/PCG/Private/Tests/Elements/Metadata/PCGMetadataBreakVectorTest.cpp
--- a/Engine/Plugins/Experimental/PCG/Source/PCG/Private/Tests/Elements/Metadata/PCGMetadataBreakVectorTest.cpp
+++ b/Engine/Plugins/Experimental/PCG/Source/PCG/Private/Tests/Elements/Metadata/PCGMetadataBreakVectorTest.cpp
@@ -13,7 +13,7 @@
 
 #if WITH_EDITOR
 
-IMPLEMENT_CUSTOM_SIMPLE_AUTOMATION_TEST(FPCGMetadataBreakVectorTest, FPCGTestBaseClass, "pcg.tests.Metadata.BreakVector", PCGTestsCommon::TestFlags)
+IMPLEMENT_CUSTOM_SIMPLE_AUTOMATION_TEST(FPCGMetadataBreakVectorTest, FPCGTestBaseClass, "Plugins.PCG.Metadata.BreakVector", PCGTestsCommon::TestFlags)
 
 namespace PCGBreakVectorTest
 {
diff --git a/Engine/Plugins/Experimental/PCG/Source/PCG/Private/Tests/Elements/Metadata/PCGMetadataMakeTransformTest.cpp b/Engine/Plugins/Experimental/PCG/Source/PCG/Private/Tests/Elements/Metadata/PCGMetadataMakeTransformTest.cpp
--- a/Engine/Plugins/Experimental/PCG/Source/PCG/Private/Tests/Elements/Metadata/PCGMetadataMakeTransformTest.cpp
+++ b/Engine/Plugins/Experimental/PCG/Source/PCG/Private/Tests/Elements/Metadata/PCGMetadataMakeTransformTest.cpp
@@ -13,7 +13,7 @@
 
 #if WITH_EDITOR
 
-IMPLEMENT_CUSTOM_SIMPLE_AUTOMATION_TEST(FPCGMetadataMakeTransformTest, FPCGTestBaseClass, "pcg.tests.Metadata.MakeTransform", PCGTestsCommon::TestFlags)
+IMPLEMENT_CUSTOM_SIMPLE_AUTOMATION_TEST(FPCGMetadataMakeTransformTest, FPCGTestBaseClass, "Plugins.PCG.Metadata.MakeTransform", PCGTestsCommon::TestFlags)
 
 namespace PCGMetadataMakeTransformTest
 {
diff --git a/Engine/Plugins/Experimental/PCG/Source/PCG/Private/Tests/Elements/Metadata/PCGMetadataMathsOpTest.cpp b/Engine/Plugins/Experimental/PCG/Source/PCG/Private/Tests/Elements/Metadata/PCGMetadataMathsOpTest.cpp
--- a/Engine/Plugins/Experimental/PCG/Source/PCG/Private/Tests/Elements/Metadata/PCGMetadataMathsOpTest.cpp
+++ b/Engine/Plugins/Experimental/PCG/Source/PCG/Private/Tests/Elements/Metadata/PCGMetadataMathsOpTest.cpp
@@ -184,7 +184,7 @@ protected:
 	}
 };
 
-IMPLEMENT_CUSTOM_SIMPLE_AUTOMATION_TEST(FPCGMetadataMathsOp_OneMinus, FPCGMetadataMathsOpTest, "pcg.tests.Metadata.MathsOp.OneMinus", PCGTestsCommon::TestFlags)
+IMPLEMENT_CUSTOM_SIMPLE_AUTOMATION_TEST(FPCGMetadataMathsOp_OneMinus, FPCGMetadataMathsOpTest, "Plugins.PCG.Metadata.MathsOp.OneMinus", PCGTestsCommon::TestFlags)
 
 bool FPCGMetadataMathsOp_OneMinus::RunTest(const FString& Parameters)
 {
@@ -213,7 +213,7 @@ bool FPCGMetadataMathsOp_OneMinus::RunTest(const FString& Parameters)
 	return bSuccess;
 }
 
-IMPLEMENT_CUSTOM_SIMPLE_AUTOMATION_TEST(FPCGMetadataMathsOp_Sign, FPCGMetadataMathsOpTest, "pcg.tests.Metadata.MathsOp.Sign", PCGTestsCommon::TestFlags)
+IMPLEMENT_CUSTOM_SIMPLE_AUTOMATION_TEST(FPCGMetadataMathsOp_Sign, FPCGMetadataMathsOpTest, "Plugins.PCG.Metadata.MathsOp.Sign", PCGTestsCommon::TestFlags)
 
 bool FPCGMetadataMathsOp_Sign::RunTest(const FString& Parameters)
 {
@@ -244,7 +244,7 @@ bool FPCGMetadataMathsOp_Sign::RunTest(const FString& Parameters)
 	return bSuccess;
 }
 
-IMPLEMENT_CUSTOM_SIMPLE_AUTOMATION_TEST(FPCGMetadataMathsOp_Frac, FPCGMetadataMathsOpTest, "pcg.tests.Metadata.MathsOp.Frac", PCGTestsCommon::TestFlags)
+IMPLEMENT_CUSTOM_SIMPLE_AUTOMATION_TEST(FPCGMetadataMathsOp_Frac, FPCGMetadataMathsOpTest, "Plugins.PCG.Metadata.MathsOp.Frac", PCGTestsCommon::TestFlags)
 
 bool FPCGMetadataMathsOp_Frac::RunTest(const FString& Parameters)
 {
@@ -273,7 +273,7 @@ bool FPCGMetadataMathsOp_Frac::RunTest(const FString& Parameters)
 	return bSuccess;
 }
 
-IMPLEMENT_CUSTOM_SIMPLE_AUTOMATION_TEST(FPCGMetadataMathsOp_Truncate, FPCGMetadataMathsOpTest, "pcg.tests.Metadata.MathsOp.Truncate", PCGTestsCommon::TestFlags)
+IMPLEMENT_CUSTOM_SIMPLE_AUTOMATION_TEST(FPCGMetadataMathsOp_Truncate, FPCGMetadataMathsOpTest, "Plugins.PCG.Metadata.MathsOp.Truncate", PCGTestsCommon::TestFlags)
 
 bool FPCGMetadataMathsOp_Truncate::RunTest(const FString& Parameters)
 {
@@ -302,7 +302,7 @@ bool FPCGMetadataMathsOp_Truncate::RunTest(const FString& Parameters)
 	return bSuccess;
 }
 
-IMPLEMENT_CUSTOM_SIMPLE_AUTOMATION_TEST(FPCGMetadataMathsOp_Add, FPCGMetadataMathsOpTest, "pcg.tests.Metadata.MathsOp.Add", PCGTestsCommon::TestFlags)
+IMPLEMENT_CUSTOM_SIMPLE_AUTOMATION_TEST(FPCGMetadataMathsOp_Add, FPCGMetadataMathsOpTest, "Plugins.PCG.Metadata.MathsOp.Add", PCGTestsCommon::TestFlags)
 
 bool FPCGMetadataMathsOp_Add::RunTest(const FString& Parameters)
 {
@@ -331,7 +331,7 @@ bool FPCGMetadataMathsOp_Add::RunTest(const FString& Parameters)
 	return bSuccess;
 }
 
-IMPLEMENT_CUSTOM_SIMPLE_AUTOMATION_TEST(FPCGMetadataMathsOp_Subtract, FPCGMetadataMathsOpTest, "pcg.tests.Metadata.MathsOp.Subtract", PCGTestsCommon::TestFlags)
+IMPLEMENT_CUSTOM_SIMPLE_AUTOMATION_TEST(FPCGMetadataMathsOp_Subtract, FPCGMetadataMathsOpTest, "Plugins.PCG.Metadata.MathsOp.Subtract", PCGTestsCommon::TestFlags)
 
 bool FPCGMetadataMathsOp_Subtract::RunTest(const FString& Parameters)
 {
@@ -360,7 +360,7 @@ bool FPCGMetadataMathsOp_Subtract::RunTest(const FString& Parameters)
 	return bSuccess;
 }
 
-IMPLEMENT_CUSTOM_SIMPLE_AUTOMATION_TEST(FPCGMetadataMathsOp_Multiply, FPCGMetadataMathsOpTest, "pcg.tests.Metadata.MathsOp.Multiply", PCGTestsCommon::TestFlags)
+IMPLEMENT_CUSTOM_SIMPLE_AUTOMATION_TEST(FPCGMetadataMathsOp_Multiply, FPCGMetadataMathsOpTest, "Plugins.PCG.Metadata.MathsOp.Multiply", PCGTestsCommon::TestFlags)
 
 bool FPCGMetadataMathsOp_Multiply::RunTest(const FString& Parameters)
 {
@@ -389,7 +389,7 @@ bool FPCGMetadataMathsOp_Multiply::RunTest(const FString& Parameters)
 	return bSuccess;
 }
 
-IMPLEMENT_CUSTOM_SIMPLE_AUTOMATION_TEST(FPCGMetadataMathsOp_Divide, FPCGMetadataMathsOpTest, "pcg.tests.Metadata.MathsOp.Divide", PCGTestsCommon::TestFlags)
+IMPLEMENT_CUSTOM_SIMPLE_AUTOMATION_TEST(FPCGMetadataMathsOp_Divide, FPCGMetadataMathsOpTest, "Plugins.PCG.Metadata.MathsOp.Divide", PCGTestsCommon::TestFlags)
 
 bool FPCGMetadataMathsOp_Divide::RunTest(const FString& Parameters)
 {
@@ -418,7 +418,7 @@ bool FPCGMetadataMathsOp_Divide::RunTest(const FString& Parameters)
 	return bSuccess;
 }
 
-IMPLEMENT_CUSTOM_SIMPLE_AUTOMATION_TEST(FPCGMetadataMathsOp_MultipleInput, FPCGMetadataMathsOpTest, "pcg.tests.Metadata.MathsOp.MultipleInput", PCGTestsCommon::TestFlags)
+IMPLEMENT_CUSTOM_SIMPLE_AUTOMATION_TEST(FPCGMetadataMathsOp_MultipleInput, FPCGMetadataMathsOpTest, "Plugins.PCG.Metadata.MathsOp.MultipleInput", PCGTestsCommon::TestFlags)
 
 bool FPCGMetadataMathsOp_MultipleInput::RunTest(const FString& Parameters)
 {
diff --git a/Engine/Plugins/Experimental/PCG/Source/PCG/Private/Tests/Elements/PCGAddTagTest.cpp b/Engine/Plugins/Experimental/PCG/Source/PCG/Private/Tests/Elements/PCGAddTagTest.cpp
--- a/Engine/Plugins/Experimental/PCG/Source/PCG/Private/Tests/Elements/PCGAddTagTest.cpp
+++ b/Engine/Plugins/Experimental/PCG/Source/PCG/Private/Tests/Elements/PCGAddTagTest.cpp
@@ -8,7 +8,7 @@
 
 #include "Elements/PCGAddTag.h"
 
-IMPLEMENT_CUSTOM_SIMPLE_AUTOMATION_TEST(FPCGAddTagTest_ZeroTags, FPCGTestBaseClass, "pcg.tests.AddTag.ZeroTags", PCGTestsCommon::TestFlags)
+IMPLEMENT_CUSTOM_SIMPLE_AUTOMATION_TEST(FPCGAddTagTest_ZeroTags, FPCGTestBaseClass, "Plugins.PCG.AddTag.ZeroTags", PCGTestsCommon::TestFlags)
 
 bool FPCGAddTagTest_ZeroTags::RunTest(const FString& Parameters)
 {
@@ -45,7 +45,7 @@ bool FPCGAddTagTest_ZeroTags::RunTest(const FString& Parameters)
 	return true;
 }
 
-IMPLEMENT_CUSTOM_SIMPLE_AUTOMATION_TEST(FPCGAddTagTest_OneTag, FPCGTestBaseClass, "pcg.tests.AddTag.OneTag", PCGTestsCommon::TestFlags)
+IMPLEMENT_CUSTOM_SIMPLE_AUTOMATION_TEST(FPCGAddTagTest_OneTag, FPCGTestBaseClass, "Plugins.PCG.AddTag.OneTag", PCGTestsCommon::TestFlags)
 
 bool FPCGAddTagTest_OneTag::RunTest(const FString& Parameters)
 {
@@ -92,7 +92,7 @@ bool FPCGAddTagTest_OneTag::RunTest(const FString& Parameters)
 	return true;
 }
 
-IMPLEMENT_CUSTOM_SIMPLE_AUTOMATION_TEST(FPCGAddTagTest_MultipleTags, FPCGTestBaseClass, "pcg.tests.AddTag.MultipleTags", PCGTestsCommon::TestFlags)
+IMPLEMENT_CUSTOM_SIMPLE_AUTOMATION_TEST(FPCGAddTagTest_MultipleTags, FPCGTestBaseClass, "Plugins.PCG.AddTag.MultipleTags", PCGTestsCommon::TestFlags)
 
 bool FPCGAddTagTest_MultipleTags::RunTest(const FString& Parameters)
 {
diff --git a/Engine/Plugins/Experimental/PCG/Source/PCG/Private/Tests/Elements/PCGApplyScaleToBoundsTest.cpp b/Engine/Plugins/Experimental/PCG/Source/PCG/Private/Tests/Elements/PCGApplyScaleToBoundsTest.cpp
--- a/Engine/Plugins/Experimental/PCG/Source/PCG/Private/Tests/Elements/PCGApplyScaleToBoundsTest.cpp
+++ b/Engine/Plugins/Experimental/PCG/Source/PCG/Private/Tests/Elements/PCGApplyScaleToBoundsTest.cpp
@@ -7,7 +7,7 @@
 
 #include "Elements/PCGApplyScaleToBounds.h"
 
-IMPLEMENT_CUSTOM_SIMPLE_AUTOMATION_TEST(FPCGApplyScaleToBoundsTest_Basic, FPCGTestBaseClass, "pcg.tests.ApplyScaleToBounds.Basic", PCGTestsCommon::TestFlags)
+IMPLEMENT_CUSTOM_SIMPLE_AUTOMATION_TEST(FPCGApplyScaleToBoundsTest_Basic, FPCGTestBaseClass, "Plugins.PCG.ApplyScaleToBounds.Basic", PCGTestsCommon::TestFlags)
 
 bool FPCGApplyScaleToBoundsTest_Basic::RunTest(const FString& Parameters)
 {
@@ -55,7 +55,7 @@ bool FPCGApplyScaleToBoundsTest_Basic::RunTest(const FString& Parameters)
 	return true;
 }
 
-IMPLEMENT_CUSTOM_SIMPLE_AUTOMATION_TEST(FPCGApplyScaleToBoundsTest_NegativeScale, FPCGTestBaseClass, "pcg.tests.ApplyScaleToBounds.NegativeScale", PCGTestsCommon::TestFlags)
+IMPLEMENT_CUSTOM_SIMPLE_AUTOMATION_TEST(FPCGApplyScaleToBoundsTest_NegativeScale, FPCGTestBaseClass, "Plugins.PCG.ApplyScaleToBounds.NegativeScale", PCGTestsCommon::TestFlags)
 
 bool FPCGApplyScaleToBoundsTest_NegativeScale::RunTest(const FString& Parameters)
 {
diff --git a/Engine/Plugins/Experimental/PCG/Source/PCG/Private/Tests/Elements/PCGAttributeFilterTest.cpp b/Engine/Plugins/Experimental/PCG/Source/PCG/Private/Tests/Elements/PCGAttributeFilterTest.cpp
--- a/Engine/Plugins/Experimental/PCG/Source/PCG/Private/Tests/Elements/PCGAttributeFilterTest.cpp
+++ b/Engine/Plugins/Experimental/PCG/Source/PCG/Private/Tests/Elements/PCGAttributeFilterTest.cpp
@@ -14,10 +14,10 @@
 #include "Math/RandomStream.h"
 
 
-IMPLEMENT_CUSTOM_SIMPLE_AUTOMATION_TEST(FPCGPointFilterDensity, FPCGTestBaseClass, "pcg.tests.AttributeFilter.Points.Density", PCGTestsCommon::TestFlags)
-IMPLEMENT_CUSTOM_SIMPLE_AUTOMATION_TEST(FPCGPointFilterDensityRange, FPCGTestBaseClass, "pcg.tests.AttributeFilter.Points.DensityRange", PCGTestsCommon::TestFlags)
-IMPLEMENT_CUSTOM_SIMPLE_AUTOMATION_TEST(FPCGAttributeFilterInt, FPCGTestBaseClass, "pcg.tests.AttributeFilter.Params.Int", PCGTestsCommon::TestFlags)
-IMPLEMENT_CUSTOM_SIMPLE_AUTOMATION_TEST(FPCGAttributeFilterIntRange, FPCGTestBaseClass, "pcg.tests.AttributeFilter.Params.IntRange", PCGTestsCommon::TestFlags)
+IMPLEMENT_CUSTOM_SIMPLE_AUTOMATION_TEST(FPCGPointFilterDensity, FPCGTestBaseClass, "Plugins.PCG.AttributeFilter.Points.Density", PCGTestsCommon::TestFlags)
+IMPLEMENT_CUSTOM_SIMPLE_AUTOMATION_TEST(FPCGPointFilterDensityRange, FPCGTestBaseClass, "Plugins.PCG.AttributeFilter.Points.DensityRange", PCGTestsCommon::TestFlags)
+IMPLEMENT_CUSTOM_SIMPLE_AUTOMATION_TEST(FPCGAttributeFilterInt, FPCGTestBaseClass, "Plugins.PCG.AttributeFilter.Params.Int", PCGTestsCommon::TestFlags)
+IMPLEMENT_CUSTOM_SIMPLE_AUTOMATION_TEST(FPCGAttributeFilterIntRange, FPCGTestBaseClass, "Plugins.PCG.AttributeFilter.Params.IntRange", PCGTestsCommon::TestFlags)
 
 namespace PCGPointFilterTest
 {
diff --git a/Engine/Plugins/Experimental/PCG/Source/PCG/Private/Tests/Elements/PCGAttributeNoiseTest.cpp b/Engine/Plugins/Experimental/PCG/Source/PCG/Private/Tests/Elements/PCGAttributeNoiseTest.cpp
--- a/Engine/Plugins/Experimental/PCG/Source/PCG/Private/Tests/Elements/PCGAttributeNoiseTest.cpp
+++ b/Engine/Plugins/Experimental/PCG/Source/PCG/Private/Tests/Elements/PCGAttributeNoiseTest.cpp
@@ -9,7 +9,7 @@
 #include "Elements/PCGAttributeNoise.h"
 #include "PCGContext.h"
 
-IMPLEMENT_CUSTOM_SIMPLE_AUTOMATION_TEST(FPCGAttributeNoiseDensityTest, FPCGTestBaseClass, "pcg.tests.AttributeNoise.Density", PCGTestsCommon::TestFlags)
+IMPLEMENT_CUSTOM_SIMPLE_AUTOMATION_TEST(FPCGAttributeNoiseDensityTest, FPCGTestBaseClass, "Plugins.PCG.AttributeNoise.Density", PCGTestsCommon::TestFlags)
 
 bool FPCGAttributeNoiseDensityTest::RunTest(const FString& Parameters)
 {
diff --git a/Engine/Plugins/Experimental/PCG/Source/PCG/Private/Tests/Elements/PCGAttributePartitionTest.cpp b/Engine/Plugins/Experimental/PCG/Source/PCG/Private/Tests/Elements/PCGAttributePartitionTest.cpp
--- a/Engine/Plugins/Experimental/PCG/Source/PCG/Private/Tests/Elements/PCGAttributePartitionTest.cpp
+++ b/Engine/Plugins/Experimental/PCG/Source/PCG/Private/Tests/Elements/PCGAttributePartitionTest.cpp
@@ -9,8 +9,8 @@
 
 #include "Elements/Metadata/PCGMetadataPartition.h"
 
-IMPLEMENT_CUSTOM_SIMPLE_AUTOMATION_TEST(FPCGAttributePartition_Points, FPCGTestBaseClass, "pcg.tests.AttributePartition.Points", PCGTestsCommon::TestFlags)
-IMPLEMENT_CUSTOM_SIMPLE_AUTOMATION_TEST(FPCGAttributePartition_AttributeSet, FPCGTestBaseClass, "pcg.tests.AttributePartition.AttributeSet", PCGTestsCommon::TestFlags)
+IMPLEMENT_CUSTOM_SIMPLE_AUTOMATION_TEST(FPCGAttributePartition_Points, FPCGTestBaseClass, "Plugins.PCG.AttributePartition.Points", PCGTestsCommon::TestFlags)
+IMPLEMENT_CUSTOM_SIMPLE_AUTOMATION_TEST(FPCGAttributePartition_AttributeSet, FPCGTestBaseClass, "Plugins.PCG.AttributePartition.AttributeSet", PCGTestsCommon::TestFlags)
 
 bool FPCGAttributePartition_Points::RunTest(const FString& Parameters)
 {
diff --git a/Engine/Plugins/Experimental/PCG/Source/PCG/Private/Tests/Elements/PCGAttributeReduceTest.cpp b/Engine/Plugins/Experimental/PCG/Source/PCG/Private/Tests/Elements/PCGAttributeReduceTest.cpp
--- a/Engine/Plugins/Experimental/PCG/Source/PCG/Private/Tests/Elements/PCGAttributeReduceTest.cpp
+++ b/Engine/Plugins/Experimental/PCG/Source/PCG/Private/Tests/Elements/PCGAttributeReduceTest.cpp
@@ -182,21 +182,21 @@ protected:
 };
 
 
-IMPLEMENT_CUSTOM_SIMPLE_AUTOMATION_TEST(FPCGAttributeReducePoints_Average, FPCGAttributeReduceTests, "pcg.tests.AttributeReduce.Points.Average", PCGTestsCommon::TestFlags)
-IMPLEMENT_CUSTOM_SIMPLE_AUTOMATION_TEST(FPCGAttributeReducePoints_Min, FPCGAttributeReduceTests, "pcg.tests.AttributeReduce.Points.Min", PCGTestsCommon::TestFlags)
-IMPLEMENT_CUSTOM_SIMPLE_AUTOMATION_TEST(FPCGAttributeReducePoints_Max, FPCGAttributeReduceTests, "pcg.tests.AttributeReduce.Points.Max", PCGTestsCommon::TestFlags)
+IMPLEMENT_CUSTOM_SIMPLE_AUTOMATION_TEST(FPCGAttributeReducePoints_Average, FPCGAttributeReduceTests, "Plugins.PCG.AttributeReduce.Points.Average", PCGTestsCommon::TestFlags)
+IMPLEMENT_CUSTOM_SIMPLE_AUTOMATION_TEST(FPCGAttributeReducePoints_Min, FPCGAttributeReduceTests, "Plugins.PCG.AttributeReduce.Points.Min", PCGTestsCommon::TestFlags)
+IMPLEMENT_CUSTOM_SIMPLE_AUTOMATION_TEST(FPCGAttributeReducePoints_Max, FPCGAttributeReduceTests, "Plugins.PCG.AttributeReduce.Points.Max", PCGTestsCommon::TestFlags)
 
-IMPLEMENT_CUSTOM_SIMPLE_AUTOMATION_TEST(FPCGAttributeReduceParams_Average, FPCGAttributeReduceTests, "pcg.tests.AttributeReduce.Params.Average", PCGTestsCommon::TestFlags)
-IMPLEMENT_CUSTOM_SIMPLE_AUTOMATION_TEST(FPCGAttributeReduceParams_Min, FPCGAttributeReduceTests, "pcg.tests.AttributeReduce.Params.Min", PCGTestsCommon::TestFlags)
-IMPLEMENT_CUSTOM_SIMPLE_AUTOMATION_TEST(FPCGAttributeReduceParams_Max, FPCGAttributeReduceTests, "pcg.tests.AttributeReduce.Params.Max", PCGTestsCommon::TestFlags)
+IMPLEMENT_CUSTOM_SIMPLE_AUTOMATION_TEST(FPCGAttributeReduceParams_Average, FPCGAttributeReduceTests, "Plugins.PCG.AttributeReduce.Params.Average", PCGTestsCommon::TestFlags)
+IMPLEMENT_CUSTOM_SIMPLE_AUTOMATION_TEST(FPCGAttributeReduceParams_Min, FPCGAttributeReduceTests, "Plugins.PCG.AttributeReduce.Params.Min", PCGTestsCommon::TestFlags)
+IMPLEMENT_CUSTOM_SIMPLE_AUTOMATION_TEST(FPCGAttributeReduceParams_Max, FPCGAttributeReduceTests, "Plugins.PCG.AttributeReduce.Params.Max", PCGTestsCommon::TestFlags)
 
-IMPLEMENT_CUSTOM_SIMPLE_AUTOMATION_TEST(FPCGAttributeReduceMixed_NoMerge, FPCGAttributeReduceTests, "pcg.tests.AttributeReduce.Mixed.NoMerge", PCGTestsCommon::TestFlags)
-IMPLEMENT_CUSTOM_SIMPLE_AUTOMATION_TEST(FPCGAttributeReduceMixed_Merge, FPCGAttributeReduceTests, "pcg.tests.AttributeReduce.Mixed.Merge", PCGTestsCommon::TestFlags)
+IMPLEMENT_CUSTOM_SIMPLE_AUTOMATION_TEST(FPCGAttributeReduceMixed_NoMerge, FPCGAttributeReduceTests, "Plugins.PCG.AttributeReduce.Mixed.NoMerge", PCGTestsCommon::TestFlags)
+IMPLEMENT_CUSTOM_SIMPLE_AUTOMATION_TEST(FPCGAttributeReduceMixed_Merge, FPCGAttributeReduceTests, "Plugins.PCG.AttributeReduce.Mixed.Merge", PCGTestsCommon::TestFlags)
 
-IMPLEMENT_CUSTOM_SIMPLE_AUTOMATION_TEST(FPCGAttributeReduceFirstInputInvalid_NoMerge, FPCGAttributeReduceTests, "pcg.tests.AttributeReduce.FirstInputInvalid.NoMerge", PCGTestsCommon::TestFlags)
-IMPLEMENT_CUSTOM_SIMPLE_AUTOMATION_TEST(FPCGAttributeReduceFirstInputInvalid_Merge, FPCGAttributeReduceTests, "pcg.tests.AttributeReduce.FirstInputInvalid.Merge", PCGTestsCommon::TestFlags)
+IMPLEMENT_CUSTOM_SIMPLE_AUTOMATION_TEST(FPCGAttributeReduceFirstInputInvalid_NoMerge, FPCGAttributeReduceTests, "Plugins.PCG.AttributeReduce.FirstInputInvalid.NoMerge", PCGTestsCommon::TestFlags)
+IMPLEMENT_CUSTOM_SIMPLE_AUTOMATION_TEST(FPCGAttributeReduceFirstInputInvalid_Merge, FPCGAttributeReduceTests, "Plugins.PCG.AttributeReduce.FirstInputInvalid.Merge", PCGTestsCommon::TestFlags)
 
-IMPLEMENT_CUSTOM_SIMPLE_AUTOMATION_TEST(FPCGAttributeReduceIncompatibleType, FPCGAttributeReduceTests, "pcg.tests.AttributeReduce.IncompatibleType", PCGTestsCommon::TestFlags)
+IMPLEMENT_CUSTOM_SIMPLE_AUTOMATION_TEST(FPCGAttributeReduceIncompatibleType, FPCGAttributeReduceTests, "Plugins.PCG.AttributeReduce.IncompatibleType", PCGTestsCommon::TestFlags)
 
 
 bool FPCGAttributeReducePoints_Average::RunTest(const FString& Parameters)
diff --git a/Engine/Plugins/Experimental/PCG/Source/PCG/Private/Tests/Elements/PCGCollapseElementTest.cpp b/Engine/Plugins/Experimental/PCG/Source/PCG/Private/Tests/Elements/PCGCollapseElementTest.cpp
--- a/Engine/Plugins/Experimental/PCG/Source/PCG/Private/Tests/Elements/PCGCollapseElementTest.cpp
+++ b/Engine/Plugins/Experimental/PCG/Source/PCG/Private/Tests/Elements/PCGCollapseElementTest.cpp
@@ -3,7 +3,7 @@
 #include "Tests/PCGTestsCommon.h"
 #include "Elements/PCGCollapseElement.h"
 
-IMPLEMENT_CUSTOM_SIMPLE_AUTOMATION_TEST(FPCGCollapseElementSmokeTest, FPCGTestBaseClass, "pcg.Collapse.Smoke", PCGTestsCommon::TestFlags)
+IMPLEMENT_CUSTOM_SIMPLE_AUTOMATION_TEST(FPCGCollapseElementSmokeTest, FPCGTestBaseClass, "Plugins.PCG.Collapse.Smoke", PCGTestsCommon::TestFlags)
 
 bool FPCGCollapseElementSmokeTest::RunTest(const FString& Parameters)
 {
diff --git a/Engine/Plugins/Experimental/PCG/Source/PCG/Private/Tests/Elements/PCGCopyPointsTest.cpp b/Engine/Plugins/Experimental/PCG/Source/PCG/Private/Tests/Elements/PCGCopyPointsTest.cpp
--- a/Engine/Plugins/Experimental/PCG/Source/PCG/Private/Tests/Elements/PCGCopyPointsTest.cpp
+++ b/Engine/Plugins/Experimental/PCG/Source/PCG/Private/Tests/Elements/PCGCopyPointsTest.cpp
@@ -12,7 +12,7 @@
 
 #if WITH_EDITOR
 
-IMPLEMENT_CUSTOM_SIMPLE_AUTOMATION_TEST(FPCGCopyPointsTest, FPCGTestBaseClass, "pcg.tests.CopyPoints.Basic", PCGTestsCommon::TestFlags)
+IMPLEMENT_CUSTOM_SIMPLE_AUTOMATION_TEST(FPCGCopyPointsTest, FPCGTestBaseClass, "Plugins.PCG.CopyPoints.Basic", PCGTestsCommon::TestFlags)
 
 bool FPCGCopyPointsTest::RunTest(const FString& Parameters)
 {
diff --git a/Engine/Plugins/Experimental/PCG/Source/PCG/Private/Tests/Elements/PCGCreatePointsGridTest.cpp b/Engine/Plugins/Experimental/PCG/Source/PCG/Private/Tests/Elements/PCGCreatePointsGridTest.cpp
--- a/Engine/Plugins/Experimental/PCG/Source/PCG/Private/Tests/Elements/PCGCreatePointsGridTest.cpp
+++ b/Engine/Plugins/Experimental/PCG/Source/PCG/Private/Tests/Elements/PCGCreatePointsGridTest.cpp
@@ -44,7 +44,7 @@ namespace CreatePointsGridCommonTests
 	}
 }
 
-IMPLEMENT_CUSTOM_SIMPLE_AUTOMATION_TEST(FPCGCreatePointsGridTest_CellCenter, FPCGTestBaseClass, "pcg.tests.CreatePointsGrid.CellCenter", PCGTestsCommon::TestFlags)
+IMPLEMENT_CUSTOM_SIMPLE_AUTOMATION_TEST(FPCGCreatePointsGridTest_CellCenter, FPCGTestBaseClass, "Plugins.PCG.CreatePointsGrid.CellCenter", PCGTestsCommon::TestFlags)
 
 bool FPCGCreatePointsGridTest_CellCenter::RunTest(const FString& Parameters)
 {
@@ -71,7 +71,7 @@ bool FPCGCreatePointsGridTest_CellCenter::RunTest(const FString& Parameters)
 	return true;
 }
 
-IMPLEMENT_CUSTOM_SIMPLE_AUTOMATION_TEST(FPCGCreatePointsGridTest_CellCorner, FPCGTestBaseClass, "pcg.tests.CreatePointsGrid.CellCorner", PCGTestsCommon::TestFlags)
+IMPLEMENT_CUSTOM_SIMPLE_AUTOMATION_TEST(FPCGCreatePointsGridTest_CellCorner, FPCGTestBaseClass, "Plugins.PCG.CreatePointsGrid.CellCorner", PCGTestsCommon::TestFlags)
 
 bool FPCGCreatePointsGridTest_CellCorner::RunTest(const FString& Parameters)
 {
@@ -100,7 +100,7 @@ bool FPCGCreatePointsGridTest_CellCorner::RunTest(const FString& Parameters)
 	return true;
 }
 
-IMPLEMENT_CUSTOM_SIMPLE_AUTOMATION_TEST(FPCGCreatePointsGridTest_XYCenter, FPCGTestBaseClass, "pcg.tests.CreatePointsGrid.XYCenter", PCGTestsCommon::TestFlags)
+IMPLEMENT_CUSTOM_SIMPLE_AUTOMATION_TEST(FPCGCreatePointsGridTest_XYCenter, FPCGTestBaseClass, "Plugins.PCG.CreatePointsGrid.XYCenter", PCGTestsCommon::TestFlags)
 
 bool FPCGCreatePointsGridTest_XYCenter::RunTest(const FString& Parameters)
 {
@@ -127,7 +127,7 @@ bool FPCGCreatePointsGridTest_XYCenter::RunTest(const FString& Parameters)
 	return true;
 }
 
-IMPLEMENT_CUSTOM_SIMPLE_AUTOMATION_TEST(FPCGCreatePointsGridTest_XYCorner, FPCGTestBaseClass, "pcg.tests.CreatePointsGrid.XYCorner", PCGTestsCommon::TestFlags)
+IMPLEMENT_CUSTOM_SIMPLE_AUTOMATION_TEST(FPCGCreatePointsGridTest_XYCorner, FPCGTestBaseClass, "Plugins.PCG.CreatePointsGrid.XYCorner", PCGTestsCommon::TestFlags)
 
 bool FPCGCreatePointsGridTest_XYCorner::RunTest(const FString& Parameters)
 {
@@ -154,7 +154,7 @@ bool FPCGCreatePointsGridTest_XYCorner::RunTest(const FString& Parameters)
 	return true;
 }
 
-IMPLEMENT_CUSTOM_SIMPLE_AUTOMATION_TEST(FPCGCreatePointsGridTest_XZCenter, FPCGTestBaseClass, "pcg.tests.CreatePointsGrid.XZCenter", PCGTestsCommon::TestFlags)
+IMPLEMENT_CUSTOM_SIMPLE_AUTOMATION_TEST(FPCGCreatePointsGridTest_XZCenter, FPCGTestBaseClass, "Plugins.PCG.CreatePointsGrid.XZCenter", PCGTestsCommon::TestFlags)
 
 bool FPCGCreatePointsGridTest_XZCenter::RunTest(const FString& Parameters)
 {
@@ -181,7 +181,7 @@ bool FPCGCreatePointsGridTest_XZCenter::RunTest(const FString& Parameters)
 	return true;
 }
 
-IMPLEMENT_CUSTOM_SIMPLE_AUTOMATION_TEST(FPCGCreatePointsGridTest_XZCorner, FPCGTestBaseClass, "pcg.tests.CreatePointsGrid.XZCorner", PCGTestsCommon::TestFlags)
+IMPLEMENT_CUSTOM_SIMPLE_AUTOMATION_TEST(FPCGCreatePointsGridTest_XZCorner, FPCGTestBaseClass, "Plugins.PCG.CreatePointsGrid.XZCorner", PCGTestsCommon::TestFlags)
 
 bool FPCGCreatePointsGridTest_XZCorner::RunTest(const FString& Parameters)
 {
@@ -208,7 +208,7 @@ bool FPCGCreatePointsGridTest_XZCorner::RunTest(const FString& Parameters)
 	return true;
 }
 
-IMPLEMENT_CUSTOM_SIMPLE_AUTOMATION_TEST(FPCGCreatePointsGridTest_YZCenter, FPCGTestBaseClass, "pcg.tests.CreatePointsGrid.YZCenter", PCGTestsCommon::TestFlags)
+IMPLEMENT_CUSTOM_SIMPLE_AUTOMATION_TEST(FPCGCreatePointsGridTest_YZCenter, FPCGTestBaseClass, "Plugins.PCG.CreatePointsGrid.YZCenter", PCGTestsCommon::TestFlags)
 
 bool FPCGCreatePointsGridTest_YZCenter::RunTest(const FString& Parameters)
 {
@@ -235,7 +235,7 @@ bool FPCGCreatePointsGridTest_YZCenter::RunTest(const FString& Parameters)
 	return true;
 }
 
-IMPLEMENT_CUSTOM_SIMPLE_AUTOMATION_TEST(FPCGCreatePointsGridTest_YZCorner, FPCGTestBaseClass, "pcg.tests.CreatePointsGrid.YZCorner", PCGTestsCommon::TestFlags)
+IMPLEMENT_CUSTOM_SIMPLE_AUTOMATION_TEST(FPCGCreatePointsGridTest_YZCorner, FPCGTestBaseClass, "Plugins.PCG.CreatePointsGrid.YZCorner", PCGTestsCommon::TestFlags)
 
 bool FPCGCreatePointsGridTest_YZCorner::RunTest(const FString& Parameters)
 {
@@ -262,7 +262,7 @@ bool FPCGCreatePointsGridTest_YZCorner::RunTest(const FString& Parameters)
 	return true;
 }
 
-IMPLEMENT_CUSTOM_SIMPLE_AUTOMATION_TEST(FPCGCreatePointsGridTest_NegativeValuesGrid, FPCGTestBaseClass, "pcg.tests.CreatePointsGrid.NegativeValuesGrid", PCGTestsCommon::TestFlags)
+IMPLEMENT_CUSTOM_SIMPLE_AUTOMATION_TEST(FPCGCreatePointsGridTest_NegativeValuesGrid, FPCGTestBaseClass, "Plugins.PCG.CreatePointsGrid.NegativeValuesGrid", PCGTestsCommon::TestFlags)
 
 bool FPCGCreatePointsGridTest_NegativeValuesGrid::RunTest(const FString& Parameters)
 {
@@ -278,7 +278,7 @@ bool FPCGCreatePointsGridTest_NegativeValuesGrid::RunTest(const FString& Paramet
 	return true;
 }
 
-IMPLEMENT_CUSTOM_SIMPLE_AUTOMATION_TEST(FPCGCreatePointsGridTest_NegativeValuesCell, FPCGTestBaseClass, "pcg.tests.CreatePointsGrid.NegativeValuesCell", PCGTestsCommon::TestFlags)
+IMPLEMENT_CUSTOM_SIMPLE_AUTOMATION_TEST(FPCGCreatePointsGridTest_NegativeValuesCell, FPCGTestBaseClass, "Plugins.PCG.CreatePointsGrid.NegativeValuesCell", PCGTestsCommon::TestFlags)
 
 bool FPCGCreatePointsGridTest_NegativeValuesCell::RunTest(const FString& Parameters)
 {
@@ -294,7 +294,7 @@ bool FPCGCreatePointsGridTest_NegativeValuesCell::RunTest(const FString& Paramet
 	return true;
 }
 
-IMPLEMENT_CUSTOM_SIMPLE_AUTOMATION_TEST(FPCGCreatePointsGridTest_OversizedCell, FPCGTestBaseClass, "pcg.tests.CreatePointsGrid.OversizedCell", PCGTestsCommon::TestFlags)
+IMPLEMENT_CUSTOM_SIMPLE_AUTOMATION_TEST(FPCGCreatePointsGridTest_OversizedCell, FPCGTestBaseClass, "Plugins.PCG.CreatePointsGrid.OversizedCell", PCGTestsCommon::TestFlags)
 
 bool FPCGCreatePointsGridTest_OversizedCell::RunTest(const FString& Parameters)
 {
@@ -318,7 +318,7 @@ bool FPCGCreatePointsGridTest_OversizedCell::RunTest(const FString& Parameters)
 	return true;
 }
 
-IMPLEMENT_CUSTOM_SIMPLE_AUTOMATION_TEST(FPCGCreatePointsGridTest_LocalCenter, FPCGTestBaseClass, "pcg.tests.CreatePointsGrid.LocalCenter", PCGTestsCommon::TestFlags)
+IMPLEMENT_CUSTOM_SIMPLE_AUTOMATION_TEST(FPCGCreatePointsGridTest_LocalCenter, FPCGTestBaseClass, "Plugins.PCG.CreatePointsGrid.LocalCenter", PCGTestsCommon::TestFlags)
 
 bool FPCGCreatePointsGridTest_LocalCenter::RunTest(const FString& Parameters)
 {
@@ -363,7 +363,7 @@ bool FPCGCreatePointsGridTest_LocalCenter::RunTest(const FString& Parameters)
 	return true;
 }
 
-IMPLEMENT_CUSTOM_SIMPLE_AUTOMATION_TEST(FPCGCreatePointsGridTest_LocalCorner, FPCGTestBaseClass, "pcg.tests.CreatePointsGrid.LocalCorner", PCGTestsCommon::TestFlags)
+IMPLEMENT_CUSTOM_SIMPLE_AUTOMATION_TEST(FPCGCreatePointsGridTest_LocalCorner, FPCGTestBaseClass, "Plugins.PCG.CreatePointsGrid.LocalCorner", PCGTestsCommon::TestFlags)
 
 bool FPCGCreatePointsGridTest_LocalCorner::RunTest(const FString& Parameters)
 {
diff --git a/Engine/Plugins/Experimental/PCG/Source/PCG/Private/Tests/Elements/PCGCreatePointsTest.cpp b/Engine/Plugins/Experimental/PCG/Source/PCG/Private/Tests/Elements/PCGCreatePointsTest.cpp
--- a/Engine/Plugins/Experimental/PCG/Source/PCG/Private/Tests/Elements/PCGCreatePointsTest.cpp
+++ b/Engine/Plugins/Experimental/PCG/Source/PCG/Private/Tests/Elements/PCGCreatePointsTest.cpp
@@ -7,7 +7,7 @@
 
 #include "Elements/PCGCreatePoints.h"
 
-IMPLEMENT_CUSTOM_SIMPLE_AUTOMATION_TEST(FPCGCreatePointsTest_Basic, FPCGTestBaseClass, "pcg.tests.CreatePoints.Basic", PCGTestsCommon::TestFlags)
+IMPLEMENT_CUSTOM_SIMPLE_AUTOMATION_TEST(FPCGCreatePointsTest_Basic, FPCGTestBaseClass, "Plugins.PCG.CreatePoints.Basic", PCGTestsCommon::TestFlags)
 
 bool FPCGCreatePointsTest_Basic::RunTest(const FString& Parameters)
 {
diff --git a/Engine/Plugins/Experimental/PCG/Source/PCG/Private/Tests/Elements/PCGDataTableRowToParamDataTest.cpp b/Engine/Plugins/Experimental/PCG/Source/PCG/Private/Tests/Elements/PCGDataTableRowToParamDataTest.cpp
--- a/Engine/Plugins/Experimental/PCG/Source/PCG/Private/Tests/Elements/PCGDataTableRowToParamDataTest.cpp
+++ b/Engine/Plugins/Experimental/PCG/Source/PCG/Private/Tests/Elements/PCGDataTableRowToParamDataTest.cpp
@@ -9,7 +9,7 @@
 #include "Tests/Elements/PCGDataTableRowToParamDataTestHelpers.h"
 #include "Engine/DataTable.h"
 
-IMPLEMENT_CUSTOM_SIMPLE_AUTOMATION_TEST(FPCGDataTableRowToParamDataTest, FPCGTestBaseClass, "pcg.tests.DataTableRowToParamData.Basic", PCGTestsCommon::TestFlags)
+IMPLEMENT_CUSTOM_SIMPLE_AUTOMATION_TEST(FPCGDataTableRowToParamDataTest, FPCGTestBaseClass, "Plugins.PCG.DataTableRowToParamData.Basic", PCGTestsCommon::TestFlags)
 
 bool FPCGDataTableRowToParamDataTest::RunTest(const FString& Parameters)
 {
@@ -68,7 +68,7 @@ CCC,"C Name","C String","333","3333","3.3","3.33","(X=3.0,Y=3.0)","(X=3.0,Y=3.0,
 	return bSuccess;
 }
 
-IMPLEMENT_CUSTOM_SIMPLE_AUTOMATION_TEST(FPCGDataTableRowToParamDataTest_ParamsOverride, FPCGTestBaseClass, "pcg.tests.DataTableRowToParamData.ParamsOverride", PCGTestsCommon::TestFlags)
+IMPLEMENT_CUSTOM_SIMPLE_AUTOMATION_TEST(FPCGDataTableRowToParamDataTest_ParamsOverride, FPCGTestBaseClass, "Plugins.PCG.DataTableRowToParamData.ParamsOverride", PCGTestsCommon::TestFlags)
 
 bool FPCGDataTableRowToParamDataTest_ParamsOverride::RunTest(const FString& Parameters)
 {
diff --git a/Engine/Plugins/Experimental/PCG/Source/PCG/Private/Tests/Elements/PCGDensityRemapTest.cpp b/Engine/Plugins/Experimental/PCG/Source/PCG/Private/Tests/Elements/PCGDensityRemapTest.cpp
--- a/Engine/Plugins/Experimental/PCG/Source/PCG/Private/Tests/Elements/PCGDensityRemapTest.cpp
+++ b/Engine/Plugins/Experimental/PCG/Source/PCG/Private/Tests/Elements/PCGDensityRemapTest.cpp
@@ -9,7 +9,7 @@
 #include "Elements/PCGDensityRemapElement.h"
 #include "PCGContext.h"
 
-IMPLEMENT_CUSTOM_SIMPLE_AUTOMATION_TEST(FPCGDensityRemapTest, FPCGTestBaseClass, "pcg.tests.DensityRemap.Basic", PCGTestsCommon::TestFlags)
+IMPLEMENT_CUSTOM_SIMPLE_AUTOMATION_TEST(FPCGDensityRemapTest, FPCGTestBaseClass, "Plugins.PCG.DensityRemap.Basic", PCGTestsCommon::TestFlags)
 
 bool FPCGDensityRemapTest::RunTest(const FString& Parameters)
 {
diff --git a/Engine/Plugins/Experimental/PCG/Source/PCG/Private/Tests/Elements/PCGDistanceTest.cpp b/Engine/Plugins/Experimental/PCG/Source/PCG/Private/Tests/Elements/PCGDistanceTest.cpp
--- a/Engine/Plugins/Experimental/PCG/Source/PCG/Private/Tests/Elements/PCGDistanceTest.cpp
+++ b/Engine/Plugins/Experimental/PCG/Source/PCG/Private/Tests/Elements/PCGDistanceTest.cpp
@@ -8,7 +8,7 @@
 
 #include "Elements/PCGDistance.h"
 
-IMPLEMENT_CUSTOM_SIMPLE_AUTOMATION_TEST(FPCGDistanceTest_PointToPoint, FPCGTestBaseClass, "pcg.tests.Distance.PointToPoint", PCGTestsCommon::TestFlags)
+IMPLEMENT_CUSTOM_SIMPLE_AUTOMATION_TEST(FPCGDistanceTest_PointToPoint, FPCGTestBaseClass, "Plugins.PCG.Distance.PointToPoint", PCGTestsCommon::TestFlags)
 
 bool FPCGDistanceTest_PointToPoint::RunTest(const FString& Parameters)
 {
@@ -65,7 +65,7 @@ bool FPCGDistanceTest_PointToPoint::RunTest(const FString& Parameters)
 	return true;
 }
 
-IMPLEMENT_CUSTOM_SIMPLE_AUTOMATION_TEST(FPCGDistanceTest_SetDensity, FPCGTestBaseClass, "pcg.tests.Distance.SetDensity", PCGTestsCommon::TestFlags)
+IMPLEMENT_CUSTOM_SIMPLE_AUTOMATION_TEST(FPCGDistanceTest_SetDensity, FPCGTestBaseClass, "Plugins.PCG.Distance.SetDensity", PCGTestsCommon::TestFlags)
 
 bool FPCGDistanceTest_SetDensity::RunTest(const FString& Parameters)
 {
@@ -122,7 +122,7 @@ bool FPCGDistanceTest_SetDensity::RunTest(const FString& Parameters)
 	return true;
 }
 
-IMPLEMENT_CUSTOM_SIMPLE_AUTOMATION_TEST(FPCGDistanceTest_PointToSphere, FPCGTestBaseClass, "pcg.tests.Distance.PointToSphere", PCGTestsCommon::TestFlags)
+IMPLEMENT_CUSTOM_SIMPLE_AUTOMATION_TEST(FPCGDistanceTest_PointToSphere, FPCGTestBaseClass, "Plugins.PCG.Distance.PointToSphere", PCGTestsCommon::TestFlags)
 
 bool FPCGDistanceTest_PointToSphere::RunTest(const FString& Parameters)
 {
@@ -187,7 +187,7 @@ bool FPCGDistanceTest_PointToSphere::RunTest(const FString& Parameters)
 	return true;
 }
 
-IMPLEMENT_CUSTOM_SIMPLE_AUTOMATION_TEST(FPCGDistanceTest_PointToBox, FPCGTestBaseClass, "pcg.tests.Distance.PointToBox", PCGTestsCommon::TestFlags)
+IMPLEMENT_CUSTOM_SIMPLE_AUTOMATION_TEST(FPCGDistanceTest_PointToBox, FPCGTestBaseClass, "Plugins.PCG.Distance.PointToBox", PCGTestsCommon::TestFlags)
 
 bool FPCGDistanceTest_PointToBox::RunTest(const FString& Parameters)
 {
diff --git a/Engine/Plugins/Experimental/PCG/Source/PCG/Private/Tests/Elements/PCGElementTest.cpp b/Engine/Plugins/Experimental/PCG/Source/PCG/Private/Tests/Elements/PCGElementTest.cpp
--- a/Engine/Plugins/Experimental/PCG/Source/PCG/Private/Tests/Elements/PCGElementTest.cpp
+++ b/Engine/Plugins/Experimental/PCG/Source/PCG/Private/Tests/Elements/PCGElementTest.cpp
@@ -29,7 +29,7 @@ bool FPCGBadOutputNodeElement::ExecuteInternal(FPCGContext* Context) const
 
 #if WITH_EDITOR
 
-IMPLEMENT_CUSTOM_SIMPLE_AUTOMATION_TEST(FPCGElementTest_CleanupAndValidateOutput, FPCGTestBaseClass, "pcg.tests.PCGElement.CleanupAndValidateOutput", PCGTestsCommon::TestFlags)
+IMPLEMENT_CUSTOM_SIMPLE_AUTOMATION_TEST(FPCGElementTest_CleanupAndValidateOutput, FPCGTestBaseClass, "Plugins.PCG.PCGElement.CleanupAndValidateOutput", PCGTestsCommon::TestFlags)
 
 bool FPCGElementTest_CleanupAndValidateOutput::RunTest(const FString& Parameters)
 {
diff --git a/Engine/Plugins/Experimental/PCG/Source/PCG/Private/Tests/Elements/PCGGatherTest.cpp b/Engine/Plugins/Experimental/PCG/Source/PCG/Private/Tests/Elements/PCGGatherTest.cpp
--- a/Engine/Plugins/Experimental/PCG/Source/PCG/Private/Tests/Elements/PCGGatherTest.cpp
+++ b/Engine/Plugins/Experimental/PCG/Source/PCG/Private/Tests/Elements/PCGGatherTest.cpp
@@ -8,7 +8,7 @@
 
 #include "Elements/PCGGather.h"
 
-IMPLEMENT_CUSTOM_SIMPLE_AUTOMATION_TEST(FPCGGatherTest_Basic, FPCGTestBaseClass, "pcg.tests.Gather.Basic", PCGTestsCommon::TestFlags)
+IMPLEMENT_CUSTOM_SIMPLE_AUTOMATION_TEST(FPCGGatherTest_Basic, FPCGTestBaseClass, "Plugins.PCG.Gather.Basic", PCGTestsCommon::TestFlags)
 
 bool FPCGGatherTest_Basic::RunTest(const FString& Parameters)
 {
diff --git a/Engine/Plugins/Experimental/PCG/Source/PCG/Private/Tests/Elements/PCGGetAttributeSetByIndexTest.cpp b/Engine/Plugins/Experimental/PCG/Source/PCG/Private/Tests/Elements/PCGGetAttributeSetByIndexTest.cpp
--- a/Engine/Plugins/Experimental/PCG/Source/PCG/Private/Tests/Elements/PCGGetAttributeSetByIndexTest.cpp
+++ b/Engine/Plugins/Experimental/PCG/Source/PCG/Private/Tests/Elements/PCGGetAttributeSetByIndexTest.cpp
@@ -8,7 +8,7 @@
 #include "Metadata/PCGMetadata.h"
 #include "Metadata/PCGMetadataAttributeTpl.h"
 
-IMPLEMENT_CUSTOM_SIMPLE_AUTOMATION_TEST(FPCGGetAttributeSetByIndexTest_Basic, FPCGTestBaseClass, "pcg.GetAttributeSetByIndex.Basic", PCGTestsCommon::TestFlags)
+IMPLEMENT_CUSTOM_SIMPLE_AUTOMATION_TEST(FPCGGetAttributeSetByIndexTest_Basic, FPCGTestBaseClass, "Plugins.PCG.GetAttributeSetByIndex.Basic", PCGTestsCommon::TestFlags)
 
 bool FPCGGetAttributeSetByIndexTest_Basic::RunTest(const FString& Parameters)
 {
diff --git a/Engine/Plugins/Experimental/PCG/Source/PCG/Private/Tests/Elements/PCGInstancePackerTest.cpp b/Engine/Plugins/Experimental/PCG/Source/PCG/Private/Tests/Elements/PCGInstancePackerTest.cpp
--- a/Engine/Plugins/Experimental/PCG/Source/PCG/Private/Tests/Elements/PCGInstancePackerTest.cpp
+++ b/Engine/Plugins/Experimental/PCG/Source/PCG/Private/Tests/Elements/PCGInstancePackerTest.cpp
@@ -19,8 +19,8 @@
 
 #if WITH_EDITOR
 
-IMPLEMENT_CUSTOM_SIMPLE_AUTOMATION_TEST(FPCGStaticMeshSpawnerInstancePackerByAttributeTest, FPCGTestBaseClass, "pcg.tests.StaticMeshSpawner.InstancePacker.ByAttribute", PCGTestsCommon::TestFlags)
-IMPLEMENT_CUSTOM_SIMPLE_AUTOMATION_TEST(FPCGStaticMeshSpawnerInstancePackerByRegexTest, FPCGTestBaseClass, "pcg.tests.StaticMeshSpawner.InstancePacker.ByRegex", PCGTestsCommon::TestFlags)
+IMPLEMENT_CUSTOM_SIMPLE_AUTOMATION_TEST(FPCGStaticMeshSpawnerInstancePackerByAttributeTest, FPCGTestBaseClass, "Plugins.PCG.StaticMeshSpawner.InstancePacker.ByAttribute", PCGTestsCommon::TestFlags)
+IMPLEMENT_CUSTOM_SIMPLE_AUTOMATION_TEST(FPCGStaticMeshSpawnerInstancePackerByRegexTest, FPCGTestBaseClass, "Plugins.PCG.StaticMeshSpawner.InstancePacker.ByRegex", PCGTestsCommon::TestFlags)
 
 namespace
 {
diff --git a/Engine/Plugins/Experimental/PCG/Source/PCG/Private/Tests/Elements/PCGNormalToDensityTest.cpp b/Engine/Plugins/Experimental/PCG/Source/PCG/Private/Tests/Elements/PCGNormalToDensityTest.cpp
--- a/Engine/Plugins/Experimental/PCG/Source/PCG/Private/Tests/Elements/PCGNormalToDensityTest.cpp
+++ b/Engine/Plugins/Experimental/PCG/Source/PCG/Private/Tests/Elements/PCGNormalToDensityTest.cpp
@@ -9,7 +9,7 @@
 #include "Elements/PCGNormalToDensity.h"
 #include "PCGContext.h"
 
-IMPLEMENT_CUSTOM_SIMPLE_AUTOMATION_TEST(FPCGNormalToDensityTest_Set, FPCGTestBaseClass, "pcg.tests.NormalToDensity.Set", PCGTestsCommon::TestFlags)
+IMPLEMENT_CUSTOM_SIMPLE_AUTOMATION_TEST(FPCGNormalToDensityTest_Set, FPCGTestBaseClass, "Plugins.PCG.NormalToDensity.Set", PCGTestsCommon::TestFlags)
 
 bool FPCGNormalToDensityTest_Set::RunTest(const FString& Parameters)
 {
@@ -57,7 +57,7 @@ bool FPCGNormalToDensityTest_Set::RunTest(const FString& Parameters)
 	return true;
 }
 
-IMPLEMENT_CUSTOM_SIMPLE_AUTOMATION_TEST(FPCGNormalToDensityTest_CustomNormal, FPCGTestBaseClass, "pcg.tests.NormalToDensity.CustomNormal", PCGTestsCommon::TestFlags)
+IMPLEMENT_CUSTOM_SIMPLE_AUTOMATION_TEST(FPCGNormalToDensityTest_CustomNormal, FPCGTestBaseClass, "Plugins.PCG.NormalToDensity.CustomNormal", PCGTestsCommon::TestFlags)
 
 bool FPCGNormalToDensityTest_CustomNormal::RunTest(const FString& Parameters)
 {
@@ -104,7 +104,7 @@ bool FPCGNormalToDensityTest_CustomNormal::RunTest(const FString& Parameters)
 	return true;
 }
 
-IMPLEMENT_CUSTOM_SIMPLE_AUTOMATION_TEST(FPCGNormalToDensityTest_Strength, FPCGTestBaseClass, "pcg.tests.NormalToDensity.Strength", PCGTestsCommon::TestFlags)
+IMPLEMENT_CUSTOM_SIMPLE_AUTOMATION_TEST(FPCGNormalToDensityTest_Strength, FPCGTestBaseClass, "Plugins.PCG.NormalToDensity.Strength", PCGTestsCommon::TestFlags)
 
 bool FPCGNormalToDensityTest_Strength::RunTest(const FString& Parameters)
 {
@@ -151,7 +151,7 @@ bool FPCGNormalToDensityTest_Strength::RunTest(const FString& Parameters)
 	return true;
 }
 
-IMPLEMENT_CUSTOM_SIMPLE_AUTOMATION_TEST(FPCGNormalToDensityTest_Minimum, FPCGTestBaseClass, "pcg.tests.NormalToDensity.Minimum", PCGTestsCommon::TestFlags)
+IMPLEMENT_CUSTOM_SIMPLE_AUTOMATION_TEST(FPCGNormalToDensityTest_Minimum, FPCGTestBaseClass, "Plugins.PCG.NormalToDensity.Minimum", PCGTestsCommon::TestFlags)
 
 bool FPCGNormalToDensityTest_Minimum::RunTest(const FString& Parameters)
 {
@@ -199,7 +199,7 @@ bool FPCGNormalToDensityTest_Minimum::RunTest(const FString& Parameters)
 	return true;
 }
 
-IMPLEMENT_CUSTOM_SIMPLE_AUTOMATION_TEST(FPCGNormalToDensityTest_Maximum, FPCGTestBaseClass, "pcg.tests.NormalToDensity.Maximum", PCGTestsCommon::TestFlags)
+IMPLEMENT_CUSTOM_SIMPLE_AUTOMATION_TEST(FPCGNormalToDensityTest_Maximum, FPCGTestBaseClass, "Plugins.PCG.NormalToDensity.Maximum", PCGTestsCommon::TestFlags)
 
 bool FPCGNormalToDensityTest_Maximum::RunTest(const FString& Parameters)
 {
@@ -248,7 +248,7 @@ bool FPCGNormalToDensityTest_Maximum::RunTest(const FString& Parameters)
 }
 
 
-IMPLEMENT_CUSTOM_SIMPLE_AUTOMATION_TEST(FPCGNormalToDensityTest_Add, FPCGTestBaseClass, "pcg.tests.NormalToDensity.Add", PCGTestsCommon::TestFlags)
+IMPLEMENT_CUSTOM_SIMPLE_AUTOMATION_TEST(FPCGNormalToDensityTest_Add, FPCGTestBaseClass, "Plugins.PCG.NormalToDensity.Add", PCGTestsCommon::TestFlags)
 
 bool FPCGNormalToDensityTest_Add::RunTest(const FString& Parameters)
 {
@@ -296,7 +296,7 @@ bool FPCGNormalToDensityTest_Add::RunTest(const FString& Parameters)
 	return true;
 }
 
-IMPLEMENT_CUSTOM_SIMPLE_AUTOMATION_TEST(FPCGNormalToDensityTest_Subtract, FPCGTestBaseClass, "pcg.tests.NormalToDensity.Subtract", PCGTestsCommon::TestFlags)
+IMPLEMENT_CUSTOM_SIMPLE_AUTOMATION_TEST(FPCGNormalToDensityTest_Subtract, FPCGTestBaseClass, "Plugins.PCG.NormalToDensity.Subtract", PCGTestsCommon::TestFlags)
 
 bool FPCGNormalToDensityTest_Subtract::RunTest(const FString& Parameters)
 {
@@ -344,7 +344,7 @@ bool FPCGNormalToDensityTest_Subtract::RunTest(const FString& Parameters)
 	return true;
 }
 
-IMPLEMENT_CUSTOM_SIMPLE_AUTOMATION_TEST(FPCGNormalToDensityTest_Multiply, FPCGTestBaseClass, "pcg.tests.NormalToDensity.Multiply", PCGTestsCommon::TestFlags)
+IMPLEMENT_CUSTOM_SIMPLE_AUTOMATION_TEST(FPCGNormalToDensityTest_Multiply, FPCGTestBaseClass, "Plugins.PCG.NormalToDensity.Multiply", PCGTestsCommon::TestFlags)
 
 bool FPCGNormalToDensityTest_Multiply::RunTest(const FString& Parameters)
 {
@@ -392,7 +392,7 @@ bool FPCGNormalToDensityTest_Multiply::RunTest(const FString& Parameters)
 	return true;
 }
 
-IMPLEMENT_CUSTOM_SIMPLE_AUTOMATION_TEST(FPCGNormalToDensityTest_Divide, FPCGTestBaseClass, "pcg.tests.NormalToDensity.Divide", PCGTestsCommon::TestFlags)
+IMPLEMENT_CUSTOM_SIMPLE_AUTOMATION_TEST(FPCGNormalToDensityTest_Divide, FPCGTestBaseClass, "Plugins.PCG.NormalToDensity.Divide", PCGTestsCommon::TestFlags)
 
 bool FPCGNormalToDensityTest_Divide::RunTest(const FString& Parameters)
 {
diff --git a/Engine/Plugins/Experimental/PCG/Source/PCG/Private/Tests/Elements/PCGPointExtentsModifierTest.cpp b/Engine/Plugins/Experimental/PCG/Source/PCG/Private/Tests/Elements/PCGPointExtentsModifierTest.cpp
--- a/Engine/Plugins/Experimental/PCG/Source/PCG/Private/Tests/Elements/PCGPointExtentsModifierTest.cpp
+++ b/Engine/Plugins/Experimental/PCG/Source/PCG/Private/Tests/Elements/PCGPointExtentsModifierTest.cpp
@@ -9,7 +9,7 @@
 #include "Elements/PCGPointExtentsModifier.h"
 #include "PCGContext.h"
 
-IMPLEMENT_CUSTOM_SIMPLE_AUTOMATION_TEST(FPCGPointsBoundsModiferTest, FPCGTestBaseClass, "pcg.tests.PointExtentsModifier.Basic", PCGTestsCommon::TestFlags)
+IMPLEMENT_CUSTOM_SIMPLE_AUTOMATION_TEST(FPCGPointsBoundsModiferTest, FPCGTestBaseClass, "Plugins.PCG.PointExtentsModifier.Basic", PCGTestsCommon::TestFlags)
 
 bool FPCGPointsBoundsModiferTest::RunTest(const FString& Parameters)
 {
diff --git a/Engine/Plugins/Experimental/PCG/Source/PCG/Private/Tests/Elements/PCGPropertyToParamDataElementTest.cpp b/Engine/Plugins/Experimental/PCG/Source/PCG/Private/Tests/Elements/PCGPropertyToParamDataElementTest.cpp
--- a/Engine/Plugins/Experimental/PCG/Source/PCG/Private/Tests/Elements/PCGPropertyToParamDataElementTest.cpp
+++ b/Engine/Plugins/Experimental/PCG/Source/PCG/Private/Tests/Elements/PCGPropertyToParamDataElementTest.cpp
@@ -15,8 +15,8 @@
 
 #if WITH_EDITOR
 
-IMPLEMENT_CUSTOM_SIMPLE_AUTOMATION_TEST(FPCGPropertyToParamDataPropertyTypeTest, FPCGTestBaseClass, "pcg.tests.PropertyToParamData.PropertyType", PCGTestsCommon::TestFlags)
-IMPLEMENT_CUSTOM_SIMPLE_AUTOMATION_TEST(FPCGPropertyToParamDataActorFindTest, FPCGTestBaseClass, "pcg.tests.PropertyToParamData.ActorFind", PCGTestsCommon::TestFlags)
+IMPLEMENT_CUSTOM_SIMPLE_AUTOMATION_TEST(FPCGPropertyToParamDataPropertyTypeTest, FPCGTestBaseClass, "Plugins.PCG.PropertyToParamData.PropertyType", PCGTestsCommon::TestFlags)
+IMPLEMENT_CUSTOM_SIMPLE_AUTOMATION_TEST(FPCGPropertyToParamDataActorFindTest, FPCGTestBaseClass, "Plugins.PCG.PropertyToParamData.ActorFind", PCGTestsCommon::TestFlags)
 
 
 /**
diff --git a/Engine/Plugins/Experimental/PCG/Source/PCG/Private/Tests/Elements/PCGResetPointCenterTest.cpp b/Engine/Plugins/Experimental/PCG/Source/PCG/Private/Tests/Elements/PCGResetPointCenterTest.cpp
--- a/Engine/Plugins/Experimental/PCG/Source/PCG/Private/Tests/Elements/PCGResetPointCenterTest.cpp
+++ b/Engine/Plugins/Experimental/PCG/Source/PCG/Private/Tests/Elements/PCGResetPointCenterTest.cpp
@@ -40,7 +40,7 @@ namespace ResetPointCenterCommonTests
 	}
 }
 
-IMPLEMENT_CUSTOM_SIMPLE_AUTOMATION_TEST(FPCGResetPointCenterTest_Center, FPCGTestBaseClass, "pcg.tests.ResetPointCenter.Center", PCGTestsCommon::TestFlags)
+IMPLEMENT_CUSTOM_SIMPLE_AUTOMATION_TEST(FPCGResetPointCenterTest_Center, FPCGTestBaseClass, "Plugins.PCG.ResetPointCenter.Center", PCGTestsCommon::TestFlags)
 
 bool FPCGResetPointCenterTest_Center::RunTest(const FString& Parameters)
 {
@@ -66,7 +66,7 @@ bool FPCGResetPointCenterTest_Center::RunTest(const FString& Parameters)
 	return true;
 }
 
-IMPLEMENT_CUSTOM_SIMPLE_AUTOMATION_TEST(FPCGResetPointCenterTest_Zero, FPCGTestBaseClass, "pcg.tests.ResetPointCenter.Zero", PCGTestsCommon::TestFlags)
+IMPLEMENT_CUSTOM_SIMPLE_AUTOMATION_TEST(FPCGResetPointCenterTest_Zero, FPCGTestBaseClass, "Plugins.PCG.ResetPointCenter.Zero", PCGTestsCommon::TestFlags)
 
 bool FPCGResetPointCenterTest_Zero::RunTest(const FString& Parameters)
 {
@@ -92,7 +92,7 @@ bool FPCGResetPointCenterTest_Zero::RunTest(const FString& Parameters)
 	return true;
 }
 
-IMPLEMENT_CUSTOM_SIMPLE_AUTOMATION_TEST(FPCGResetPointCenterTest_One, FPCGTestBaseClass, "pcg.tests.ResetPointCenter.One", PCGTestsCommon::TestFlags)
+IMPLEMENT_CUSTOM_SIMPLE_AUTOMATION_TEST(FPCGResetPointCenterTest_One, FPCGTestBaseClass, "Plugins.PCG.ResetPointCenter.One", PCGTestsCommon::TestFlags)
 
 bool FPCGResetPointCenterTest_One::RunTest(const FString& Parameters)
 {
@@ -118,7 +118,7 @@ bool FPCGResetPointCenterTest_One::RunTest(const FString& Parameters)
 	return true;
 }
 
-IMPLEMENT_CUSTOM_SIMPLE_AUTOMATION_TEST(FPCGResetPointCenterTest_NegativeOutside, FPCGTestBaseClass, "pcg.tests.ResetPointCenter.NegativeOutside", PCGTestsCommon::TestFlags)
+IMPLEMENT_CUSTOM_SIMPLE_AUTOMATION_TEST(FPCGResetPointCenterTest_NegativeOutside, FPCGTestBaseClass, "Plugins.PCG.ResetPointCenter.NegativeOutside", PCGTestsCommon::TestFlags)
 
 bool FPCGResetPointCenterTest_NegativeOutside::RunTest(const FString& Parameters)
 {
@@ -144,7 +144,7 @@ bool FPCGResetPointCenterTest_NegativeOutside::RunTest(const FString& Parameters
 	return true;
 }
 
-IMPLEMENT_CUSTOM_SIMPLE_AUTOMATION_TEST(FPCGResetPointCenterTest_PositiveOutside, FPCGTestBaseClass, "pcg.tests.ResetPointCenter.PositiveOutside", PCGTestsCommon::TestFlags)
+IMPLEMENT_CUSTOM_SIMPLE_AUTOMATION_TEST(FPCGResetPointCenterTest_PositiveOutside, FPCGTestBaseClass, "Plugins.PCG.ResetPointCenter.PositiveOutside", PCGTestsCommon::TestFlags)
 
 bool FPCGResetPointCenterTest_PositiveOutside::RunTest(const FString& Parameters)
 {
@@ -170,7 +170,7 @@ bool FPCGResetPointCenterTest_PositiveOutside::RunTest(const FString& Parameters
 	return true;
 }
 
-IMPLEMENT_CUSTOM_SIMPLE_AUTOMATION_TEST(FPCGResetPointCenterTest_NonUniform, FPCGTestBaseClass, "pcg.tests.ResetPointCenter.NonUniform", PCGTestsCommon::TestFlags)
+IMPLEMENT_CUSTOM_SIMPLE_AUTOMATION_TEST(FPCGResetPointCenterTest_NonUniform, FPCGTestBaseClass, "Plugins.PCG.ResetPointCenter.NonUniform", PCGTestsCommon::TestFlags)
 
 bool FPCGResetPointCenterTest_NonUniform::RunTest(const FString& Parameters)
 {
@@ -195,7 +195,7 @@ bool FPCGResetPointCenterTest_NonUniform::RunTest(const FString& Parameters)
 	return true;
 }
 
-IMPLEMENT_CUSTOM_SIMPLE_AUTOMATION_TEST(FPCGResetPointCenterTest_CenterRotation, FPCGTestBaseClass, "pcg.tests.ResetPointCenter.CenterRotation", PCGTestsCommon::TestFlags)
+IMPLEMENT_CUSTOM_SIMPLE_AUTOMATION_TEST(FPCGResetPointCenterTest_CenterRotation, FPCGTestBaseClass, "Plugins.PCG.ResetPointCenter.CenterRotation", PCGTestsCommon::TestFlags)
 
 bool FPCGResetPointCenterTest_CenterRotation::RunTest(const FString& Parameters)
 {
@@ -221,7 +221,7 @@ bool FPCGResetPointCenterTest_CenterRotation::RunTest(const FString& Parameters)
 	return true;
 }
 
-IMPLEMENT_CUSTOM_SIMPLE_AUTOMATION_TEST(FPCGResetPointCenterTest_CenterScale, FPCGTestBaseClass, "pcg.tests.ResetPointCenter.CenterScale", PCGTestsCommon::TestFlags)
+IMPLEMENT_CUSTOM_SIMPLE_AUTOMATION_TEST(FPCGResetPointCenterTest_CenterScale, FPCGTestBaseClass, "Plugins.PCG.ResetPointCenter.CenterScale", PCGTestsCommon::TestFlags)
 
 bool FPCGResetPointCenterTest_CenterScale::RunTest(const FString& Parameters)
 {
diff --git a/Engine/Plugins/Experimental/PCG/Source/PCG/Private/Tests/Elements/PCGSanityCheckPointDataTest.cpp b/Engine/Plugins/Experimental/PCG/Source/PCG/Private/Tests/Elements/PCGSanityCheckPointDataTest.cpp
--- a/Engine/Plugins/Experimental/PCG/Source/PCG/Private/Tests/Elements/PCGSanityCheckPointDataTest.cpp
+++ b/Engine/Plugins/Experimental/PCG/Source/PCG/Private/Tests/Elements/PCGSanityCheckPointDataTest.cpp
@@ -11,7 +11,7 @@
 
 #include "Elements/PCGSanityCheckPointData.h"
 
-IMPLEMENT_CUSTOM_SIMPLE_AUTOMATION_TEST(FPCGSanityCheckPointData_MinPointCount, FPCGTestBaseClass, "pcg.tests.SanityCheckPointData.MinPointCount", PCGTestsCommon::TestFlags)
+IMPLEMENT_CUSTOM_SIMPLE_AUTOMATION_TEST(FPCGSanityCheckPointData_MinPointCount, FPCGTestBaseClass, "Plugins.PCG.SanityCheckPointData.MinPointCount", PCGTestsCommon::TestFlags)
 
 bool FPCGSanityCheckPointData_MinPointCount::RunTest(const FString& Parameters)
 {
@@ -39,7 +39,7 @@ bool FPCGSanityCheckPointData_MinPointCount::RunTest(const FString& Parameters)
 	return true;
 }
 
-IMPLEMENT_CUSTOM_SIMPLE_AUTOMATION_TEST(FPCGSanityCheckPointData_MaxPointCount, FPCGTestBaseClass, "pcg.tests.SanityCheckPointData.MaxPointCount", PCGTestsCommon::TestFlags)
+IMPLEMENT_CUSTOM_SIMPLE_AUTOMATION_TEST(FPCGSanityCheckPointData_MaxPointCount, FPCGTestBaseClass, "Plugins.PCG.SanityCheckPointData.MaxPointCount", PCGTestsCommon::TestFlags)
 
 bool FPCGSanityCheckPointData_MaxPointCount::RunTest(const FString& Parameters)
 {
@@ -67,7 +67,7 @@ bool FPCGSanityCheckPointData_MaxPointCount::RunTest(const FString& Parameters)
 	return true;
 }
 
-IMPLEMENT_CUSTOM_SIMPLE_AUTOMATION_TEST(FPCGSanityCheckPointData_Success, FPCGTestBaseClass, "pcg.tests.SanityCheckPointData.Success", PCGTestsCommon::TestFlags)
+IMPLEMENT_CUSTOM_SIMPLE_AUTOMATION_TEST(FPCGSanityCheckPointData_Success, FPCGTestBaseClass, "Plugins.PCG.SanityCheckPointData.Success", PCGTestsCommon::TestFlags)
 
 bool FPCGSanityCheckPointData_Success::RunTest(const FString& Parameters)
 {
diff --git a/Engine/Plugins/Experimental/PCG/Source/PCG/Private/Tests/Elements/PCGSortPointsTest.cpp b/Engine/Plugins/Experimental/PCG/Source/PCG/Private/Tests/Elements/PCGSortPointsTest.cpp
--- a/Engine/Plugins/Experimental/PCG/Source/PCG/Private/Tests/Elements/PCGSortPointsTest.cpp
+++ b/Engine/Plugins/Experimental/PCG/Source/PCG/Private/Tests/Elements/PCGSortPointsTest.cpp
@@ -30,7 +30,7 @@ namespace SortCommonTestData
 	}
 }
 
-IMPLEMENT_CUSTOM_SIMPLE_AUTOMATION_TEST(FPCGSortPointsTest_Basic, FPCGTestBaseClass, "pcg.tests.SortPoints.Basic", PCGTestsCommon::TestFlags)
+IMPLEMENT_CUSTOM_SIMPLE_AUTOMATION_TEST(FPCGSortPointsTest_Basic, FPCGTestBaseClass, "Plugins.PCG.SortPoints.Basic", PCGTestsCommon::TestFlags)
 
 bool FPCGSortPointsTest_Basic::RunTest(const FString& Parameters)
 {
@@ -53,7 +53,7 @@ bool FPCGSortPointsTest_Basic::RunTest(const FString& Parameters)
 	return true;
 }
 
-IMPLEMENT_CUSTOM_SIMPLE_AUTOMATION_TEST(FPCGSortPointsTest_Ascending, FPCGTestBaseClass, "pcg.tests.SortPoints.Ascending", PCGTestsCommon::TestFlags)
+IMPLEMENT_CUSTOM_SIMPLE_AUTOMATION_TEST(FPCGSortPointsTest_Ascending, FPCGTestBaseClass, "Plugins.PCG.SortPoints.Ascending", PCGTestsCommon::TestFlags)
 
 bool FPCGSortPointsTest_Ascending::RunTest(const FString& Parameters)
 {
@@ -79,7 +79,7 @@ bool FPCGSortPointsTest_Ascending::RunTest(const FString& Parameters)
 	return true;
 }
 
-IMPLEMENT_CUSTOM_SIMPLE_AUTOMATION_TEST(FPCGSortPointsTest_Descending, FPCGTestBaseClass, "pcg.tests.SortPoints.Descending", PCGTestsCommon::TestFlags)
+IMPLEMENT_CUSTOM_SIMPLE_AUTOMATION_TEST(FPCGSortPointsTest_Descending, FPCGTestBaseClass, "Plugins.PCG.SortPoints.Descending", PCGTestsCommon::TestFlags)
 
 bool FPCGSortPointsTest_Descending::RunTest(const FString& Parameters)
 {
@@ -105,7 +105,7 @@ bool FPCGSortPointsTest_Descending::RunTest(const FString& Parameters)
 	return true;
 }
 
-IMPLEMENT_CUSTOM_SIMPLE_AUTOMATION_TEST(FPCGSortPointsTest_SameValues, FPCGTestBaseClass, "pcg.tests.SortPoints.SameValues", PCGTestsCommon::TestFlags)
+IMPLEMENT_CUSTOM_SIMPLE_AUTOMATION_TEST(FPCGSortPointsTest_SameValues, FPCGTestBaseClass, "Plugins.PCG.SortPoints.SameValues", PCGTestsCommon::TestFlags)
 
 bool FPCGSortPointsTest_SameValues::RunTest(const FString& Parameters)
 {
diff --git a/Engine/Plugins/Experimental/PCG/Source/PCG/Private/Tests/Elements/PCGSpatialNoiseTest.cpp b/Engine/Plugins/Experimental/PCG/Source/PCG/Private/Tests/Elements/PCGSpatialNoiseTest.cpp
--- a/Engine/Plugins/Experimental/PCG/Source/PCG/Private/Tests/Elements/PCGSpatialNoiseTest.cpp
+++ b/Engine/Plugins/Experimental/PCG/Source/PCG/Private/Tests/Elements/PCGSpatialNoiseTest.cpp
@@ -6,7 +6,7 @@
 #include "Data/PCGPointData.h"
 #include "Elements/PCGSpatialNoise.h"
 
-IMPLEMENT_SIMPLE_AUTOMATION_TEST(FPCGSpatialNoise_CalcLocalCoordinates2D, "pcg.tests.Noise.CalcLocalCoordinates2D", EAutomationTestFlags::EditorContext | EAutomationTestFlags::ProductFilter)
+IMPLEMENT_SIMPLE_AUTOMATION_TEST(FPCGSpatialNoise_CalcLocalCoordinates2D, "Plugins.PCG.Noise.CalcLocalCoordinates2D", EAutomationTestFlags::EditorContext | EAutomationTestFlags::ProductFilter)
 
 bool FPCGSpatialNoise_CalcLocalCoordinates2D::RunTest(const FString& Parameters)
 {
@@ -39,7 +39,7 @@ bool FPCGSpatialNoise_CalcLocalCoordinates2D::RunTest(const FString& Parameters)
 
 // just make sure stuff doesn't crash
 
-IMPLEMENT_CUSTOM_SIMPLE_AUTOMATION_TEST(FPCGSpatialNoise_Perlin2D, FPCGTestBaseClass, "pcg.tests.Noise.Perlin2D", PCGTestsCommon::TestFlags)
+IMPLEMENT_CUSTOM_SIMPLE_AUTOMATION_TEST(FPCGSpatialNoise_Perlin2D, FPCGTestBaseClass, "Plugins.PCG.Noise.Perlin2D", PCGTestsCommon::TestFlags)
 
 bool FPCGSpatialNoise_Perlin2D::RunTest(const FString& Parameters)
 {
@@ -77,7 +77,7 @@ bool FPCGSpatialNoise_Perlin2D::RunTest(const FString& Parameters)
 	return true;
 }
 
-IMPLEMENT_CUSTOM_SIMPLE_AUTOMATION_TEST(FPCGSpatialNoise_Caustic2D, FPCGTestBaseClass, "pcg.tests.Noise.Caustic2D", PCGTestsCommon::TestFlags)
+IMPLEMENT_CUSTOM_SIMPLE_AUTOMATION_TEST(FPCGSpatialNoise_Caustic2D, FPCGTestBaseClass, "Plugins.PCG.Noise.Caustic2D", PCGTestsCommon::TestFlags)
 
 bool FPCGSpatialNoise_Caustic2D::RunTest(const FString& Parameters)
 {
@@ -115,7 +115,7 @@ bool FPCGSpatialNoise_Caustic2D::RunTest(const FString& Parameters)
 	return true;
 }
 
-IMPLEMENT_CUSTOM_SIMPLE_AUTOMATION_TEST(FPCGSpatialNoise_Voronoi2D, FPCGTestBaseClass, "pcg.tests.Noise.Voronoi2D", PCGTestsCommon::TestFlags)
+IMPLEMENT_CUSTOM_SIMPLE_AUTOMATION_TEST(FPCGSpatialNoise_Voronoi2D, FPCGTestBaseClass, "Plugins.PCG.Noise.Voronoi2D", PCGTestsCommon::TestFlags)
 
 bool FPCGSpatialNoise_Voronoi2D::RunTest(const FString& Parameters)
 {
@@ -158,7 +158,7 @@ bool FPCGSpatialNoise_Voronoi2D::RunTest(const FString& Parameters)
 	return true;
 }
 
-IMPLEMENT_CUSTOM_SIMPLE_AUTOMATION_TEST(FPCGSpatialNoise_FractionalBrownian2D, FPCGTestBaseClass, "pcg.tests.Noise.FractionalBrownian2D", PCGTestsCommon::TestFlags)
+IMPLEMENT_CUSTOM_SIMPLE_AUTOMATION_TEST(FPCGSpatialNoise_FractionalBrownian2D, FPCGTestBaseClass, "Plugins.PCG.Noise.FractionalBrownian2D", PCGTestsCommon::TestFlags)
 
 bool FPCGSpatialNoise_FractionalBrownian2D::RunTest(const FString& Parameters)
 {
@@ -196,7 +196,7 @@ bool FPCGSpatialNoise_FractionalBrownian2D::RunTest(const FString& Parameters)
 	return true;
 }
 
-IMPLEMENT_CUSTOM_SIMPLE_AUTOMATION_TEST(FPCGSpatialNoise_EdgeMask2D, FPCGTestBaseClass, "pcg.tests.Noise.EdgeMask2D", PCGTestsCommon::TestFlags)
+IMPLEMENT_CUSTOM_SIMPLE_AUTOMATION_TEST(FPCGSpatialNoise_EdgeMask2D, FPCGTestBaseClass, "Plugins.PCG.Noise.EdgeMask2D", PCGTestsCommon::TestFlags)
 
 bool FPCGSpatialNoise_EdgeMask2D::RunTest(const FString& Parameters)
 {
@@ -234,7 +234,7 @@ bool FPCGSpatialNoise_EdgeMask2D::RunTest(const FString& Parameters)
 	return true;
 }
 
-IMPLEMENT_CUSTOM_SIMPLE_AUTOMATION_TEST(FPCGSpatialNoise_TilingPerlin2D, FPCGTestBaseClass, "pcg.tests.Noise.TilingPerlin2D", PCGTestsCommon::TestFlags)
+IMPLEMENT_CUSTOM_SIMPLE_AUTOMATION_TEST(FPCGSpatialNoise_TilingPerlin2D, FPCGTestBaseClass, "Plugins.PCG.Noise.TilingPerlin2D", PCGTestsCommon::TestFlags)
 bool FPCGSpatialNoise_TilingPerlin2D::RunTest(const FString& Parameters)
 {
 	PCGTestsCommon::FTestData TestData;
@@ -272,7 +272,7 @@ bool FPCGSpatialNoise_TilingPerlin2D::RunTest(const FString& Parameters)
 	return true;
 }
 
-IMPLEMENT_CUSTOM_SIMPLE_AUTOMATION_TEST(FPCGSpatialNoise_TilingVoronoi2D, FPCGTestBaseClass, "pcg.tests.Noise.TilingVoronoi2D", PCGTestsCommon::TestFlags)
+IMPLEMENT_CUSTOM_SIMPLE_AUTOMATION_TEST(FPCGSpatialNoise_TilingVoronoi2D, FPCGTestBaseClass, "Plugins.PCG.Noise.TilingVoronoi2D", PCGTestsCommon::TestFlags)
 bool FPCGSpatialNoise_TilingVoronoi2D::RunTest(const FString& Parameters)
 {
 	PCGTestsCommon::FTestData TestData;
diff --git a/Engine/Plugins/Experimental/PCG/Source/PCG/Private/Tests/Elements/PCGSplineSamplerTest.cpp b/Engine/Plugins/Experimental/PCG/Source/PCG/Private/Tests/Elements/PCGSplineSamplerTest.cpp
--- a/Engine/Plugins/Experimental/PCG/Source/PCG/Private/Tests/Elements/PCGSplineSamplerTest.cpp
+++ b/Engine/Plugins/Experimental/PCG/Source/PCG/Private/Tests/Elements/PCGSplineSamplerTest.cpp
@@ -5,7 +5,7 @@
 
 #if WITH_EDITOR
 
-IMPLEMENT_CUSTOM_SIMPLE_AUTOMATION_TEST(FPCGSplineSamplerTest, FPCGTestBaseClass, "pcg.tests.SplineSampler.PointInPolygon", PCGTestsCommon::TestFlags)
+IMPLEMENT_CUSTOM_SIMPLE_AUTOMATION_TEST(FPCGSplineSamplerTest, FPCGTestBaseClass, "Plugins.PCG.SplineSampler.PointInPolygon", PCGTestsCommon::TestFlags)
 
 bool FPCGSplineSamplerTest::RunTest(const FString& Parameters)
 {	
diff --git a/Engine/Plugins/Experimental/PCG/Source/PCG/Private/Tests/Elements/PCGStaticMeshSpawnerTest.cpp b/Engine/Plugins/Experimental/PCG/Source/PCG/Private/Tests/Elements/PCGStaticMeshSpawnerTest.cpp
--- a/Engine/Plugins/Experimental/PCG/Source/PCG/Private/Tests/Elements/PCGStaticMeshSpawnerTest.cpp
+++ b/Engine/Plugins/Experimental/PCG/Source/PCG/Private/Tests/Elements/PCGStaticMeshSpawnerTest.cpp
@@ -19,9 +19,9 @@
 #include "Engine/StaticMesh.h"
 
 // FIXME: Unit tests are crashing, disabled for now.
-IMPLEMENT_CUSTOM_SIMPLE_AUTOMATION_TEST(FPCGStaticMeshSpawnerByAttributeTest, FPCGTestBaseClass, "pcg.tests.StaticMeshSpawner.ByAttribute", PCGTestsCommon::TestFlags | EAutomationTestFlags::Disabled)
-IMPLEMENT_CUSTOM_SIMPLE_AUTOMATION_TEST(FPCGStaticMeshSpawnerWeightedTest, FPCGTestBaseClass, "pcg.tests.StaticMeshSpawner.Weighted", PCGTestsCommon::TestFlags | EAutomationTestFlags::Disabled)
-IMPLEMENT_CUSTOM_SIMPLE_AUTOMATION_TEST(FPCGStaticMeshSpawnerWeightedByCategoryTest, FPCGTestBaseClass, "pcg.tests.StaticMeshSpawner.WeightedByCategory", PCGTestsCommon::TestFlags | EAutomationTestFlags::Disabled)
+IMPLEMENT_CUSTOM_SIMPLE_AUTOMATION_TEST(FPCGStaticMeshSpawnerByAttributeTest, FPCGTestBaseClass, "Plugins.PCG.StaticMeshSpawner.ByAttribute", PCGTestsCommon::TestFlags | EAutomationTestFlags::Disabled)
+IMPLEMENT_CUSTOM_SIMPLE_AUTOMATION_TEST(FPCGStaticMeshSpawnerWeightedTest, FPCGTestBaseClass, "Plugins.PCG.StaticMeshSpawner.Weighted", PCGTestsCommon::TestFlags | EAutomationTestFlags::Disabled)
+IMPLEMENT_CUSTOM_SIMPLE_AUTOMATION_TEST(FPCGStaticMeshSpawnerWeightedByCategoryTest, FPCGTestBaseClass, "Plugins.PCG.StaticMeshSpawner.WeightedByCategory", PCGTestsCommon::TestFlags | EAutomationTestFlags::Disabled)
 
 namespace
 {
diff --git a/Engine/Plugins/Experimental/PCG/Source/PCG/Private/Tests/Elements/PCGTransformPointsTest.cpp b/Engine/Plugins/Experimental/PCG/Source/PCG/Private/Tests/Elements/PCGTransformPointsTest.cpp
--- a/Engine/Plugins/Experimental/PCG/Source/PCG/Private/Tests/Elements/PCGTransformPointsTest.cpp
+++ b/Engine/Plugins/Experimental/PCG/Source/PCG/Private/Tests/Elements/PCGTransformPointsTest.cpp
@@ -11,7 +11,7 @@
 
 #if WITH_EDITOR
 
-IMPLEMENT_CUSTOM_SIMPLE_AUTOMATION_TEST(FPCGTransformPointsTest, FPCGTestBaseClass, "pcg.tests.TransformPoints.Basic", PCGTestsCommon::TestFlags)
+IMPLEMENT_CUSTOM_SIMPLE_AUTOMATION_TEST(FPCGTransformPointsTest, FPCGTestBaseClass, "Plugins.PCG.TransformPoints.Basic", PCGTestsCommon::TestFlags)
 
 bool FPCGTransformPointsTest::RunTest(const FString& Parameters)
 {
diff --git a/Engine/Plugins/Experimental/PCG/Source/PCG/Private/Tests/GraphAuthoring/PCGGraphAuthoringTest.cpp b/Engine/Plugins/Experimental/PCG/Source/PCG/Private/Tests/GraphAuthoring/PCGGraphAuthoringTest.cpp
--- a/Engine/Plugins/Experimental/PCG/Source/PCG/Private/Tests/GraphAuthoring/PCGGraphAuthoringTest.cpp
+++ b/Engine/Plugins/Experimental/PCG/Source/PCG/Private/Tests/GraphAuthoring/PCGGraphAuthoringTest.cpp
@@ -10,7 +10,7 @@
 #include "Tests/PCGTestsCommon.h"
 #include "Tests/GraphAuthoring/PCGGraphAuthoringTestHelperSettings.h"
 
-IMPLEMENT_CUSTOM_SIMPLE_AUTOMATION_TEST(FPCGGraphAuthoringTypeSystemTest, FPCGTestBaseClass, "pcg.GraphAuthoring.TypeSystem", PCGTestsCommon::TestFlags)
+IMPLEMENT_CUSTOM_SIMPLE_AUTOMATION_TEST(FPCGGraphAuthoringTypeSystemTest, FPCGTestBaseClass, "Plugins.PCG.GraphAuthoring.TypeSystem", PCGTestsCommon::TestFlags)
 
 bool FPCGGraphAuthoringTypeSystemTest::RunTest(const FString& Parameters)
 {
diff --git a/Engine/Plugins/Experimental/PCG/Source/PCG/Private/Tests/Helpers/PCGAsyncTest.cpp b/Engine/Plugins/Experimental/PCG/Source/PCG/Private/Tests/Helpers/PCGAsyncTest.cpp
--- a/Engine/Plugins/Experimental/PCG/Source/PCG/Private/Tests/Helpers/PCGAsyncTest.cpp
+++ b/Engine/Plugins/Experimental/PCG/Source/PCG/Private/Tests/Helpers/PCGAsyncTest.cpp
@@ -5,7 +5,7 @@
 #include "Helpers/PCGAsync.h"
 
 
-IMPLEMENT_CUSTOM_SIMPLE_AUTOMATION_TEST(FPCGAsyncTest_AsyncProcessing, FPCGTestBaseClass, "pcg.Async.AsyncProcessing", PCGTestsCommon::TestFlags)
+IMPLEMENT_CUSTOM_SIMPLE_AUTOMATION_TEST(FPCGAsyncTest_AsyncProcessing, FPCGTestBaseClass, "Plugins.PCG.Async.AsyncProcessing", PCGTestsCommon::TestFlags)
 
 bool FPCGAsyncTest_AsyncProcessing::RunTest(const FString& Parameters)
 {
@@ -43,7 +43,7 @@ bool FPCGAsyncTest_AsyncProcessing::RunTest(const FString& Parameters)
 	return true;
 }
 
-IMPLEMENT_CUSTOM_SIMPLE_AUTOMATION_TEST(FPCGAsyncTest_AsyncProcessingEx, FPCGTestBaseClass, "pcg.Async.AsyncProcessingEx", PCGTestsCommon::TestFlags)
+IMPLEMENT_CUSTOM_SIMPLE_AUTOMATION_TEST(FPCGAsyncTest_AsyncProcessingEx, FPCGTestBaseClass, "Plugins.PCG.Async.AsyncProcessingEx", PCGTestsCommon::TestFlags)
 
 bool FPCGAsyncTest_AsyncProcessingEx::RunTest(const FString& Parameters)
 {
@@ -92,7 +92,7 @@ bool FPCGAsyncTest_AsyncProcessingEx::RunTest(const FString& Parameters)
 	return true;
 }
 
-IMPLEMENT_CUSTOM_SIMPLE_AUTOMATION_TEST(FPCGAsyncTest_AsyncProcessingOneToOneEx, FPCGTestBaseClass, "pcg.Async.AsyncProcessingOneToOneEx", PCGTestsCommon::TestFlags)
+IMPLEMENT_CUSTOM_SIMPLE_AUTOMATION_TEST(FPCGAsyncTest_AsyncProcessingOneToOneEx, FPCGTestBaseClass, "Plugins.PCG.Async.AsyncProcessingOneToOneEx", PCGTestsCommon::TestFlags)
 
 bool FPCGAsyncTest_AsyncProcessingOneToOneEx::RunTest(const FString& Parameters)
 {
diff --git a/Engine/Plugins/Experimental/PCG/Source/PCG/Private/Tests/Metadata/PCGMetadataFlattenTest.cpp b/Engine/Plugins/Experimental/PCG/Source/PCG/Private/Tests/Metadata/PCGMetadataFlattenTest.cpp
--- a/Engine/Plugins/Experimental/PCG/Source/PCG/Private/Tests/Metadata/PCGMetadataFlattenTest.cpp
+++ b/Engine/Plugins/Experimental/PCG/Source/PCG/Private/Tests/Metadata/PCGMetadataFlattenTest.cpp
@@ -8,7 +8,7 @@
 #include "Metadata/PCGMetadata.h"
 #include "Metadata/PCGMetadataAttributeTpl.h"
 
-IMPLEMENT_CUSTOM_SIMPLE_AUTOMATION_TEST(FPCGMetadataFlatten, FPCGTestBaseClass, "pcg.tests.Metadata.Flatten", PCGTestsCommon::TestFlags)
+IMPLEMENT_CUSTOM_SIMPLE_AUTOMATION_TEST(FPCGMetadataFlatten, FPCGTestBaseClass, "Plugins.PCG.Metadata.Flatten", PCGTestsCommon::TestFlags)
 
 /**
 * Series of operations to manipulate points and attributes, to validate the flatten operation.
diff --git a/Engine/Plugins/Experimental/PCG/Source/PCG/Private/Tests/PCGGraphParametersTest.cpp b/Engine/Plugins/Experimental/PCG/Source/PCG/Private/Tests/PCGGraphParametersTest.cpp
--- a/Engine/Plugins/Experimental/PCG/Source/PCG/Private/Tests/PCGGraphParametersTest.cpp
+++ b/Engine/Plugins/Experimental/PCG/Source/PCG/Private/Tests/PCGGraphParametersTest.cpp
@@ -8,19 +8,19 @@
 #include "PropertyBag.h"
 #include "StructView.h"
 
-IMPLEMENT_CUSTOM_SIMPLE_AUTOMATION_TEST(FPCGGraphNoUserParameters, FPCGTestBaseClass, "pcg.Graph.NoUserParameters", PCGTestsCommon::TestFlags)
-IMPLEMENT_CUSTOM_SIMPLE_AUTOMATION_TEST(FPCGGraphAddSingleUserParameters, FPCGTestBaseClass, "pcg.Graph.AddSingleUserParameter", PCGTestsCommon::TestFlags)
-IMPLEMENT_CUSTOM_SIMPLE_AUTOMATION_TEST(FPCGGraphAddSingleUserParametersPropagates, FPCGTestBaseClass, "pcg.Graph.AddSingleUserParameterPropagates", PCGTestsCommon::TestFlags)
-IMPLEMENT_CUSTOM_SIMPLE_AUTOMATION_TEST(FPCGGraphRemoveSingleUserParameters, FPCGTestBaseClass, "pcg.Graph.RemoveSingleUserParameter", PCGTestsCommon::TestFlags)
-IMPLEMENT_CUSTOM_SIMPLE_AUTOMATION_TEST(FPCGGraphRemoveSingleUserParametersPropagates, FPCGTestBaseClass, "pcg.Graph.RemoveSingleUserParameterPropagates", PCGTestsCommon::TestFlags)
+IMPLEMENT_CUSTOM_SIMPLE_AUTOMATION_TEST(FPCGGraphNoUserParameters, FPCGTestBaseClass, "Plugins.PCG.Graph.NoUserParameters", PCGTestsCommon::TestFlags)
+IMPLEMENT_CUSTOM_SIMPLE_AUTOMATION_TEST(FPCGGraphAddSingleUserParameters, FPCGTestBaseClass, "Plugins.PCG.Graph.AddSingleUserParameter", PCGTestsCommon::TestFlags)
+IMPLEMENT_CUSTOM_SIMPLE_AUTOMATION_TEST(FPCGGraphAddSingleUserParametersPropagates, FPCGTestBaseClass, "Plugins.PCG.Graph.AddSingleUserParameterPropagates", PCGTestsCommon::TestFlags)
+IMPLEMENT_CUSTOM_SIMPLE_AUTOMATION_TEST(FPCGGraphRemoveSingleUserParameters, FPCGTestBaseClass, "Plugins.PCG.Graph.RemoveSingleUserParameter", PCGTestsCommon::TestFlags)
+IMPLEMENT_CUSTOM_SIMPLE_AUTOMATION_TEST(FPCGGraphRemoveSingleUserParametersPropagates, FPCGTestBaseClass, "Plugins.PCG.Graph.RemoveSingleUserParameterPropagates", PCGTestsCommon::TestFlags)
 
 // We test the graph add/remove parameters just as a redundancy, but it is already tested by StructUtils. So we won't test more.
 // What is important is the propagation between graphs and graphs instances
-IMPLEMENT_CUSTOM_SIMPLE_AUTOMATION_TEST(FPCGGraphRenameUserParametersPropagates, FPCGTestBaseClass, "pcg.Graph.RenameUserParameterPropagates", PCGTestsCommon::TestFlags)
-IMPLEMENT_CUSTOM_SIMPLE_AUTOMATION_TEST(FPCGGraphChangeTypeUserParametersPropagates, FPCGTestBaseClass, "pcg.Graph.ChangeTypeUserParameterPropagates", PCGTestsCommon::TestFlags)
-IMPLEMENT_CUSTOM_SIMPLE_AUTOMATION_TEST(FPCGGraphRemovingGraphCleansUserParameters, FPCGTestBaseClass, "pcg.Graph.RemovingGraphCleansUserParameters", PCGTestsCommon::TestFlags)
-IMPLEMENT_CUSTOM_SIMPLE_AUTOMATION_TEST(FPCGGraphSetGraphUpdatesUserParameters, FPCGTestBaseClass, "pcg.Graph.SetGraphCleansUserParameters", PCGTestsCommon::TestFlags)
-IMPLEMENT_CUSTOM_SIMPLE_AUTOMATION_TEST(FPCGGraphSetValueUserParametersPropagates, FPCGTestBaseClass, "pcg.Graph.SetValueUserParametersPropagates", PCGTestsCommon::TestFlags)
+IMPLEMENT_CUSTOM_SIMPLE_AUTOMATION_TEST(FPCGGraphRenameUserParametersPropagates, FPCGTestBaseClass, "Plugins.PCG.Graph.RenameUserParameterPropagates", PCGTestsCommon::TestFlags)
+IMPLEMENT_CUSTOM_SIMPLE_AUTOMATION_TEST(FPCGGraphChangeTypeUserParametersPropagates, FPCGTestBaseClass, "Plugins.PCG.Graph.ChangeTypeUserParameterPropagates", PCGTestsCommon::TestFlags)
+IMPLEMENT_CUSTOM_SIMPLE_AUTOMATION_TEST(FPCGGraphRemovingGraphCleansUserParameters, FPCGTestBaseClass, "Plugins.PCG.Graph.RemovingGraphCleansUserParameters", PCGTestsCommon::TestFlags)
+IMPLEMENT_CUSTOM_SIMPLE_AUTOMATION_TEST(FPCGGraphSetGraphUpdatesUserParameters, FPCGTestBaseClass, "Plugins.PCG.Graph.SetGraphCleansUserParameters", PCGTestsCommon::TestFlags)
+IMPLEMENT_CUSTOM_SIMPLE_AUTOMATION_TEST(FPCGGraphSetValueUserParametersPropagates, FPCGTestBaseClass, "Plugins.PCG.Graph.SetValueUserParametersPropagates", PCGTestsCommon::TestFlags)
 
 #define PCG_FMT(Text, ...) FString::Printf(TEXT(Text), __VA_ARGS__)
 
diff --git a/Engine/Plugins/Experimental/PCG/Source/PCG/Private/Tests/Subgraphs/PCGCollapseSubgraphTest.cpp b/Engine/Plugins/Experimental/PCG/Source/PCG/Private/Tests/Subgraphs/PCGCollapseSubgraphTest.cpp
--- a/Engine/Plugins/Experimental/PCG/Source/PCG/Private/Tests/Subgraphs/PCGCollapseSubgraphTest.cpp
+++ b/Engine/Plugins/Experimental/PCG/Source/PCG/Private/Tests/Subgraphs/PCGCollapseSubgraphTest.cpp
@@ -11,8 +11,8 @@
 #include "Helpers/PCGSubgraphHelpers.h"
 #include "Tests/PCGTestsCommon.h"
 
-IMPLEMENT_CUSTOM_SIMPLE_AUTOMATION_TEST(FPCGCollapseSubgraphSimple, FPCGTestBaseClass, "pcg.tests.Subgraph.Collapse.Simple", PCGTestsCommon::TestFlags);
-IMPLEMENT_CUSTOM_SIMPLE_AUTOMATION_TEST(FPCGCollapseSubgraphWithParams, FPCGTestBaseClass, "pcg.tests.Subgraph.Collapse.WithParams", PCGTestsCommon::TestFlags);
+IMPLEMENT_CUSTOM_SIMPLE_AUTOMATION_TEST(FPCGCollapseSubgraphSimple, FPCGTestBaseClass, "Plugins.PCG.Subgraph.Collapse.Simple", PCGTestsCommon::TestFlags);
+IMPLEMENT_CUSTOM_SIMPLE_AUTOMATION_TEST(FPCGCollapseSubgraphWithParams, FPCGTestBaseClass, "Plugins.PCG.Subgraph.Collapse.WithParams", PCGTestsCommon::TestFlags);
 
 namespace PCGCollapseSubgraphTests
 {
diff --git a/Engine/Plugins/Experimental/PCG/Source/PCGEditor/Private/Tests/Display/PCGComponentFoundInRegistry.cpp b/Engine/Plugins/Experimental/PCG/Source/PCGEditor/Private/Tests/Display/PCGComponentFoundInRegistry.cpp
--- a/Engine/Plugins/Experimental/PCG/Source/PCGEditor/Private/Tests/Display/PCGComponentFoundInRegistry.cpp
+++ b/Engine/Plugins/Experimental/PCG/Source/PCGEditor/Private/Tests/Display/PCGComponentFoundInRegistry.cpp
@@ -7,7 +7,7 @@
 
 #if WITH_AUTOMATION_TESTS
 
-IMPLEMENT_CUSTOM_SIMPLE_AUTOMATION_TEST(PCGComponentFoundInRegistry, FPCGTestBaseClass, "Editor.Plugins.Tools.PCG.PCGComponentFoundInRegistry", PCGTestsCommon::TestFlags)
+IMPLEMENT_CUSTOM_SIMPLE_AUTOMATION_TEST(PCGComponentFoundInRegistry, FPCGTestBaseClass, "Plugins.PCG.Display.PCGComponentFoundInRegistry", PCGTestsCommon::TestFlags)
 
 bool PCGComponentFoundInRegistry::RunTest(const FString& Parameters)
 {
diff --git a/Engine/Plugins/Experimental/PCG/Source/PCGEditor/Private/Tests/Display/PCGComponentInNewPCGVolumes.cpp b/Engine/Plugins/Experimental/PCG/Source/PCGEditor/Private/Tests/Display/PCGComponentInNewPCGVolumes.cpp
--- a/Engine/Plugins/Experimental/PCG/Source/PCGEditor/Private/Tests/Display/PCGComponentInNewPCGVolumes.cpp
+++ b/Engine/Plugins/Experimental/PCG/Source/PCGEditor/Private/Tests/Display/PCGComponentInNewPCGVolumes.cpp
@@ -16,7 +16,7 @@
 
 #if WITH_AUTOMATION_TESTS
 
-IMPLEMENT_CUSTOM_SIMPLE_AUTOMATION_TEST(FPCGComponentInNewPCGVolumes, FPCGTestBaseClass, "Editor.Plugins.Tools.PCG.PCGComponentInNewPCGVolumes", PCGTestsCommon::TestFlags)
+IMPLEMENT_CUSTOM_SIMPLE_AUTOMATION_TEST(FPCGComponentInNewPCGVolumes, FPCGTestBaseClass, "Plugins.PCG.Display.PCGComponentInNewPCGVolumes", PCGTestsCommon::TestFlags)
 
 bool FPCGComponentInNewPCGVolumes::RunTest(const FString& Parameters)
 {
diff --git a/Engine/Plugins/Experimental/PCG/Source/PCGEditor/Private/Tests/Display/PCGGrapthAssetIsExposedInContetBrowser.cpp b/Engine/Plugins/Experimental/PCG/Source/PCGEditor/Private/Tests/Display/PCGGrapthAssetIsExposedInContetBrowser.cpp
--- a/Engine/Plugins/Experimental/PCG/Source/PCGEditor/Private/Tests/Display/PCGGrapthAssetIsExposedInContetBrowser.cpp
+++ b/Engine/Plugins/Experimental/PCG/Source/PCGEditor/Private/Tests/Display/PCGGrapthAssetIsExposedInContetBrowser.cpp
@@ -7,7 +7,7 @@
 
 #if WITH_AUTOMATION_TESTS
 
-IMPLEMENT_CUSTOM_SIMPLE_AUTOMATION_TEST(FPCGGraphExposedInContentMenuTest, FPCGTestBaseClass, "Editor.Plugins.Tools.PCG.PCGGraphExposedInContentMenuTest", PCGTestsCommon::TestFlags)
+IMPLEMENT_CUSTOM_SIMPLE_AUTOMATION_TEST(FPCGGraphExposedInContentMenuTest, FPCGTestBaseClass, "Plugins.PCG.Display.PCGGraphExposedInContentMenuTest", PCGTestsCommon::TestFlags)
 
 bool FPCGGraphExposedInContentMenuTest::RunTest(const FString& Parameters)
 {
diff --git a/Engine/Plugins/Experimental/PCG/Source/PCGEditor/Private/Tests/Display/PCGInActorAddComponentField.cpp b/Engine/Plugins/Experimental/PCG/Source/PCGEditor/Private/Tests/Display/PCGInActorAddComponentField.cpp
--- a/Engine/Plugins/Experimental/PCG/Source/PCGEditor/Private/Tests/Display/PCGInActorAddComponentField.cpp
+++ b/Engine/Plugins/Experimental/PCG/Source/PCGEditor/Private/Tests/Display/PCGInActorAddComponentField.cpp
@@ -19,7 +19,7 @@
 
 #if WITH_AUTOMATION_TESTS
 
-IMPLEMENT_CUSTOM_SIMPLE_AUTOMATION_TEST(FPCGInActorAddComponentField, FPCGTestBaseClass, "Editor.Plugins.Tools.PCG.PCGInActorAddComponentField", PCGTestsCommon::TestFlags)
+IMPLEMENT_CUSTOM_SIMPLE_AUTOMATION_TEST(FPCGInActorAddComponentField, FPCGTestBaseClass, "Plugins.PCG.Display.PCGInActorAddComponentField", PCGTestsCommon::TestFlags)
 
 ULevel* LoadLevelByTemplateName(const FString& TemplateName)
 {
diff --git a/Engine/Plugins/Experimental/PCG/Source/PCGEditor/Private/Tests/Display/PCGVolumeExposedInCreationMenu.cpp b/Engine/Plugins/Experimental/PCG/Source/PCGEditor/Private/Tests/Display/PCGVolumeExposedInCreationMenu.cpp
--- a/Engine/Plugins/Experimental/PCG/Source/PCGEditor/Private/Tests/Display/PCGVolumeExposedInCreationMenu.cpp
+++ b/Engine/Plugins/Experimental/PCG/Source/PCGEditor/Private/Tests/Display/PCGVolumeExposedInCreationMenu.cpp
@@ -5,7 +5,7 @@
 
 #if WITH_AUTOMATION_TESTS
 
-IMPLEMENT_CUSTOM_SIMPLE_AUTOMATION_TEST(FPCGVolumeExposedInCreationMenu, FPCGTestBaseClass, "Editor.Plugins.Tools.PCG.PCGVolumeExposedInCreationMenu", PCGTestsCommon::TestFlags)
+IMPLEMENT_CUSTOM_SIMPLE_AUTOMATION_TEST(FPCGVolumeExposedInCreationMenu, FPCGTestBaseClass, "Plugins.PCG.Display.PCGVolumeExposedInCreationMenu", PCGTestsCommon::TestFlags)
 
 bool FPCGVolumeExposedInCreationMenu::RunTest(const FString& Parameters)
 {
