[PCG] Remove useless macro for IsConstructible

#rnx

diff --git a/Engine/Plugins/Experimental/PCG/Source/PCG/Public/Metadata/PCGMetadataAttributeTraits.h b/Engine/Plugins/Experimental/PCG/Source/PCG/Public/Metadata/PCGMetadataAttributeTraits.h
--- a/Engine/Plugins/Experimental/PCG/Source/PCG/Public/Metadata/PCGMetadataAttributeTraits.h
+++ b/Engine/Plugins/Experimental/PCG/Source/PCG/Public/Metadata/PCGMetadataAttributeTraits.h
@@ -50,24 +50,6 @@ enum class EPCGMetadataTypes : uint8
 	MACRO(FSoftObjectPath) \
 	MACRO(FSoftClassPath)
 
-// Need a duplicated macro for a template list, because Foo<A,B,> is syntax error. It needs Foo<A,B>
-#define PCG_FOREACH_SUPPORTEDTYPES_WITH_COMMA(MACRO) \
-	MACRO(int32),           \
-	MACRO(int64),           \
-	MACRO(float),           \
-	MACRO(double),          \
-	MACRO(FVector2D),       \
-	MACRO(FVector),         \
-	MACRO(FVector4),        \
-	MACRO(FQuat),           \
-	MACRO(FTransform),      \
-	MACRO(FString),         \
-	MACRO(bool),            \
-	MACRO(FRotator),        \
-	MACRO(FName),           \
-	MACRO(FSoftObjectPath), \
-	MACRO(FSoftClassPath)
-
 namespace PCG
 {
 	namespace Private
@@ -249,12 +231,13 @@ namespace PCG
 			struct TypeHolder {};
 
 			// Use the macro to have AllTypes that is the Dummy struct templated with all our types.
-#define PCG_ALL_TYPES(T) T
-			using AllTypes = TypeHolder<PCG_FOREACH_SUPPORTEDTYPES_WITH_COMMA(PCG_ALL_TYPES)>;
+			// Add void as a sentinel, to know when to stop unrolling (and because TypeHolder<T,U,> is syntax error)
+#define PCG_ALL_TYPES(T) T,
+			using AllTypes = TypeHolder<PCG_FOREACH_SUPPORTEDTYPES(PCG_ALL_TYPES) void>;
 #undef PCG_ALL_TYPES
 
 			// Unroll AllTypes twice, to have all combinations of types <T, U>
-			template <typename FirstType> constexpr void UnrollSecond(TypeHolder<> InTypeHolder) {}
+			template <typename FirstType> constexpr void UnrollSecond(TypeHolder<void> InTypeHolder) {}
 
 			template <typename FirstType, typename SecondType, typename... SecondTypes>
 			constexpr void UnrollSecond(TypeHolder<SecondType, SecondTypes...> InTypeHolder)
@@ -270,7 +253,7 @@ namespace PCG
 				UnrollFirst(TypeHolder<InputTypes...>{});
 			}
 
-			constexpr void UnrollFirst(TypeHolder<> InTypeHolder) {}
+			constexpr void UnrollFirst(TypeHolder<void> InTypeHolder) {}
 
 		public:
 			constexpr FConstructibleTypes() : Values{ {false} }
@@ -1072,6 +1055,3 @@ namespace PCG
 		}
 	}
 }
-
-// Undef this one since we don't use it outside this file. But the other is used outside so keep it defined.
-#undef PCG_FOREACH_SUPPORTEDTYPES_WITH_COMMA
\ No newline at end of file
