[PCG] Add editor command to rename PCGEditorGraphNodes
* F2 or right-click context menu to rename a PCG node in the graph editor
* Double clicking a node title no longer begins editing the title

diff --git a/Engine/Plugins/Experimental/PCG/Source/PCGEditor/Private/PCGEditor.cpp b/Engine/Plugins/Experimental/PCG/Source/PCGEditor/Private/PCGEditor.cpp
--- a/Engine/Plugins/Experimental/PCG/Source/PCGEditor/Private/PCGEditor.cpp
+++ b/Engine/Plugins/Experimental/PCG/Source/PCGEditor/Private/PCGEditor.cpp
@@ -644,6 +644,11 @@ void FPCGEditor::BindCommands()
 		PCGEditorCommands.AddSourcePin,
 		FExecuteAction::CreateSP(this, &FPCGEditor::OnAddDynamicInputPin),
 		FCanExecuteAction::CreateSP(this, &FPCGEditor::CanAddDynamicInputPin));
+
+	GraphEditorCommands->MapAction(
+		PCGEditorCommands.RenameNode,
+		FExecuteAction::CreateSP(this, &FPCGEditor::OnRenameNode),
+		FCanExecuteAction::CreateSP(this, &FPCGEditor::CanRenameNode));
 }
 
 void FPCGEditor::OnFind()
@@ -933,6 +938,38 @@ bool FPCGEditor::CanAddDynamicInputPin() const
 	return (Node && Node->CanUserAddRemoveDynamicInputPins());
 }
 
+void FPCGEditor::OnRenameNode()
+{
+	check(GraphEditorWidget.IsValid());
+
+	const FGraphPanelSelectionSet SelectedNodes = GraphEditorWidget->GetSelectedNodes();
+
+	if (!ensure(SelectedNodes.Num() == 1))
+	{
+		UE_LOG(LogPCGEditor, Warning, TEXT("Attempting to rename multiple nodes."));
+		return;
+	}
+
+	UPCGEditorGraphNodeBase* Node = CastChecked<UPCGEditorGraphNodeBase>(*SelectedNodes.CreateConstIterator());
+	Node->EnterRenamingMode();
+}
+
+bool FPCGEditor::CanRenameNode() const
+{
+	check(GraphEditorWidget.IsValid());
+
+	const FGraphPanelSelectionSet SelectedNodes = GraphEditorWidget->GetSelectedNodes();
+
+	// You cannot enter renaming mode on multiple nodes at once, since they will not all enter synchronously.
+	// Simultaneous editing of multiple InlineEditableTextBlocks may not even be possible with default behavior.
+	if (SelectedNodes.Num() != 1)
+	{
+		return false;
+	}
+
+	return (*SelectedNodes.CreateConstIterator())->IsA<UPCGEditorGraphNodeBase>();
+}
+
 void FPCGEditor::OnCollapseNodesInSubgraph()
 {
 	if (!GraphEditorWidget.IsValid() || PCGEditorGraph == nullptr)
@@ -2168,9 +2205,17 @@ void FPCGEditor::OnNodeTitleCommitted(const FText& NewText, ETextCommit::Type Co
 {
 	if (NodeBeingChanged)
 	{
-		const FScopedTransaction Transaction(*FPCGEditorCommon::ContextIdentifier, LOCTEXT("PCGEditorRenameNode", "PCG Editor: Rename Node"), nullptr);
-		NodeBeingChanged->Modify();
-		NodeBeingChanged->OnRenameNode(NewText.ToString());
+		if (CommitInfo == ETextCommit::OnEnter)
+		{
+			const FScopedTransaction Transaction(*FPCGEditorCommon::ContextIdentifier, LOCTEXT("PCGEditorRenameNode", "PCG Editor: Rename Node"), nullptr);
+			NodeBeingChanged->Modify();
+			NodeBeingChanged->OnRenameNode(NewText.ToString());
+		}
+
+		if (UPCGEditorGraphNodeBase* PCGEditorNode = CastChecked<UPCGEditorGraphNodeBase>(NodeBeingChanged))
+		{
+			PCGEditorNode->ExitRenamingMode();
+		}
 	}
 }
 
diff --git a/Engine/Plugins/Experimental/PCG/Source/PCGEditor/Private/PCGEditor.h b/Engine/Plugins/Experimental/PCG/Source/PCGEditor/Private/PCGEditor.h
--- a/Engine/Plugins/Experimental/PCG/Source/PCGEditor/Private/PCGEditor.h
+++ b/Engine/Plugins/Experimental/PCG/Source/PCGEditor/Private/PCGEditor.h
@@ -226,6 +226,11 @@ private:
 	/** Whether the user can add a dynamic source pin to a node */
 	bool CanAddDynamicInputPin() const;
 
+	/** User is attempting to rename a node */
+	void OnRenameNode();
+	/** Whether the user can rename the selected node */
+	bool CanRenameNode() const;
+
 	void OnAlignTop();
 	void OnAlignMiddle();
 	void OnAlignBottom();
diff --git a/Engine/Plugins/Experimental/PCG/Source/PCGEditor/Private/PCGEditorCommands.cpp b/Engine/Plugins/Experimental/PCG/Source/PCGEditor/Private/PCGEditorCommands.cpp
--- a/Engine/Plugins/Experimental/PCG/Source/PCGEditor/Private/PCGEditorCommands.cpp
+++ b/Engine/Plugins/Experimental/PCG/Source/PCGEditor/Private/PCGEditorCommands.cpp
@@ -33,6 +33,7 @@ void FPCGEditorCommands::RegisterCommands()
 	UI_COMMAND(DisableDebugOnAllNodes, "Disable Debug on all nodes", "Disable debug state for all nodes", EUserInterfaceActionType::Button, FInputChord(EModifierKey::Alt, EKeys::D));
 	UI_COMMAND(ToggleInspect, "Toggle Inspection", "Toggle node inspection for selected node", EUserInterfaceActionType::ToggleButton, FInputChord(EKeys::A));
 	UI_COMMAND(AddSourcePin, "Add Source Pin", "Add new source pin to the current node", EUserInterfaceActionType::Button, FInputChord());
+	UI_COMMAND(RenameNode, "Rename Node", "Rename the selected node", EUserInterfaceActionType::Button, FInputChord(EKeys::F2));
 }
 
 #undef LOCTEXT_NAMESPACE
diff --git a/Engine/Plugins/Experimental/PCG/Source/PCGEditor/Private/PCGEditorCommands.h b/Engine/Plugins/Experimental/PCG/Source/PCGEditor/Private/PCGEditorCommands.h
--- a/Engine/Plugins/Experimental/PCG/Source/PCGEditor/Private/PCGEditorCommands.h
+++ b/Engine/Plugins/Experimental/PCG/Source/PCGEditor/Private/PCGEditorCommands.h
@@ -29,4 +29,5 @@ public:
 	TSharedPtr<FUICommandInfo> DisableDebugOnAllNodes;
 	TSharedPtr<FUICommandInfo> ToggleInspect;
 	TSharedPtr<FUICommandInfo> AddSourcePin;
+	TSharedPtr<FUICommandInfo> RenameNode;
 };
diff --git a/Engine/Plugins/Experimental/PCG/Source/PCGEditor/Private/PCGEditorGraphNodeBase.cpp b/Engine/Plugins/Experimental/PCG/Source/PCGEditor/Private/PCGEditorGraphNodeBase.cpp
--- a/Engine/Plugins/Experimental/PCG/Source/PCGEditor/Private/PCGEditorGraphNodeBase.cpp
+++ b/Engine/Plugins/Experimental/PCG/Source/PCGEditor/Private/PCGEditorGraphNodeBase.cpp
@@ -36,6 +36,7 @@ void UPCGEditorGraphNodeBase::Construct(UPCGNode* InPCGNode)
 	NodeComment = InPCGNode->NodeComment;
 	bCommentBubblePinned = InPCGNode->bCommentBubblePinned;
 	bCommentBubbleVisible = InPCGNode->bCommentBubbleVisible;
+	bCanRenameNode = false;
 
 	if (const UPCGSettingsInterface* PCGSettingsInterface = InPCGNode->GetSettingsInterface())
 	{
@@ -112,6 +113,7 @@ void UPCGEditorGraphNodeBase::GetNodeContextMenuActions(UToolMenu* Menu, class U
 		Section.AddMenuEntry(FPCGEditorCommands::Get().CollapseNodes);
 		Section.AddMenuEntry(FPCGEditorCommands::Get().ExportNodes);
 		Section.AddMenuEntry(FPCGEditorCommands::Get().ConvertToStandaloneNodes);
+		Section.AddMenuEntry(FPCGEditorCommands::Get().RenameNode, LOCTEXT("RenameNode", "Rename"));
 	}
 
 	{
@@ -454,6 +456,22 @@ bool UPCGEditorGraphNodeBase::ShouldDrawCompact() const
 	return Settings && Settings->ShouldDrawNodeCompact();
 }
 
+void UPCGEditorGraphNodeBase::EnterRenamingMode()
+{
+	bCanRenameNode = true;
+
+	// Notify the SPCGEditorGraphNode that the user is trying to rename it.
+	OnNodeRenameInitiatedDelegate.ExecuteIfBound();
+}
+
+void UPCGEditorGraphNodeBase::ExitRenamingMode()
+{
+	bCanRenameNode = false;
+
+	// Update so that the node renders with the new node title.
+	OnNodeChangedDelegate.ExecuteIfBound();
+}
+
 void UPCGEditorGraphNodeBase::OnPickColor()
 {
 	FColorPickerArgs PickerArgs;
diff --git a/Engine/Plugins/Experimental/PCG/Source/PCGEditor/Private/PCGEditorGraphNodeBase.h b/Engine/Plugins/Experimental/PCG/Source/PCGEditor/Private/PCGEditorGraphNodeBase.h
--- a/Engine/Plugins/Experimental/PCG/Source/PCGEditor/Private/PCGEditorGraphNodeBase.h
+++ b/Engine/Plugins/Experimental/PCG/Source/PCGEditor/Private/PCGEditorGraphNodeBase.h
@@ -42,6 +42,7 @@ public:
 	virtual UObject* GetJumpTargetForDoubleClick() const override;
 	virtual void OnUpdateCommentText(const FString& NewComment) override;
 	virtual void OnCommentBubbleToggled(bool bInCommentBubbleVisible) override;
+	virtual bool GetCanRenameNode() const override { return bCanRenameNode; }
 	// ~End UEdGraphNode interface
 
 	void OnUserAddDynamicInputPin();
@@ -76,9 +77,18 @@ public:
 
 	bool CanUserAddRemoveDynamicInputPins() const { return bCanUserAddRemoveSourcePins; }
 
+	/** Marks the node as re-nameable and provokes the node to update, placing the user in rename mode. */
+	void EnterRenamingMode();
+
+	/** Disables renaming, preventing the node from entering editing mode on subsequent updates. */
+	void ExitRenamingMode();
+
 	DECLARE_DELEGATE(FOnPCGEditorGraphNodeChanged);
 	FOnPCGEditorGraphNodeChanged OnNodeChangedDelegate;
 
+	DECLARE_DELEGATE(FOnPCGEditorGraphNodeRenameInitiated);
+	FOnPCGEditorGraphNodeRenameInitiated OnNodeRenameInitiatedDelegate;
+
 protected:
 	static FEdGraphPinType GetPinType(const UPCGPin* InPin);
 
diff --git a/Engine/Plugins/Experimental/PCG/Source/PCGEditor/Private/SPCGEditorGraphNode.cpp b/Engine/Plugins/Experimental/PCG/Source/PCGEditor/Private/SPCGEditorGraphNode.cpp
--- a/Engine/Plugins/Experimental/PCG/Source/PCGEditor/Private/SPCGEditorGraphNode.cpp
+++ b/Engine/Plugins/Experimental/PCG/Source/PCGEditor/Private/SPCGEditorGraphNode.cpp
@@ -358,6 +358,7 @@ void SPCGEditorGraphNode::Construct(const FArguments& InArgs, UPCGEditorGraphNod
 	if (InNode)
 	{
 		InNode->OnNodeChangedDelegate.BindSP(this, &SPCGEditorGraphNode::OnNodeChanged);
+		InNode->OnNodeRenameInitiatedDelegate.BindSP(this, &SPCGEditorGraphNode::OnNodeRenameInitiated);
 	}
 
 	UpdateGraphNode();
@@ -562,6 +563,14 @@ void SPCGEditorGraphNode::OnNodeChanged()
 	UpdateGraphNode();
 }
 
+void SPCGEditorGraphNode::OnNodeRenameInitiated()
+{
+	if (InlineEditableText.IsValid())
+	{
+		InlineEditableText->EnterEditingMode();
+	}
+}
+
 void SPCGEditorGraphNode::UpdateCompactNode()
 {
 	// Based on SGraphNodeK2Base::UpdateCompactNode. Changes:
diff --git a/Engine/Plugins/Experimental/PCG/Source/PCGEditor/Private/SPCGEditorGraphNode.h b/Engine/Plugins/Experimental/PCG/Source/PCGEditor/Private/SPCGEditorGraphNode.h
--- a/Engine/Plugins/Experimental/PCG/Source/PCGEditor/Private/SPCGEditorGraphNode.h
+++ b/Engine/Plugins/Experimental/PCG/Source/PCGEditor/Private/SPCGEditorGraphNode.h
@@ -34,6 +34,7 @@ public:
 
 protected:
 	void OnNodeChanged();
+	void OnNodeRenameInitiated();
 
 	/** Set up node in 'compact' mode */
 	void UpdateCompactNode();
