[PCG] Add a new Get Property from Object Path

* Renamed PCGPropertyToParamData to PCGGetActorProperty
* Add special case when extraction is None, to extract the actor/component/object directly
* Refactor GetOverridableParams to use a UStruct, instead of UScriptStruct and UClass (both are UStructs...)

diff --git a/Engine/Plugins/Experimental/PCG/Config/BasePCG.ini b/Engine/Plugins/Experimental/PCG/Config/BasePCG.ini
--- a/Engine/Plugins/Experimental/PCG/Config/BasePCG.ini
+++ b/Engine/Plugins/Experimental/PCG/Config/BasePCG.ini
@@ -66,6 +66,7 @@
 +ClassRedirects=(OldName="PCGPointFilterSettings",NewName="/Script/PCG.PCGAttributeFilteringSettings")
 +ClassRedirects=(OldName="PCGPointFilterRangeSettings",NewName="/Script/PCG.PCGAttributeFilteringRangeSettings")
 +ClassRedirects=(OldName="PCGAttributeFilterSettings",NewName="/Script/PCG.PCGAttributeFilterNamesSettings")
++ClassRedirects=(OldName="PCGPropertyToParamDataSettings",NewName="/Script/PCG.PCGGetActorPropertySettings")
 
 +StructRedirects=(OldName="PCGPointFilterThresholdSettings",NewName="/Script/PCG.PCGAttributeFilterThresholdSettings")
 
diff --git a/Engine/Plugins/Experimental/PCG/Source/PCG/Private/Elements/PCGGetActorProperty.cpp b/Engine/Plugins/Experimental/PCG/Source/PCG/Private/Elements/PCGGetActorProperty.cpp
new file mode 100644
--- /dev/null
+++ b/Engine/Plugins/Experimental/PCG/Source/PCG/Private/Elements/PCGGetActorProperty.cpp
@@ -0,0 +1,188 @@
+// Copyright Epic Games, Inc. All Rights Reserved.
+
+#include "Elements/PCGGetActorProperty.h"
+
+#include "GameFramework/Actor.h"
+#include "PCGContext.h"
+#include "PCGComponent.h"
+#include "PCGParamData.h"
+#include "Helpers/PCGBlueprintHelpers.h"
+#include "Helpers/PCGPropertyHelpers.h"
+#include "Helpers/PCGHelpers.h"
+#include "Metadata/PCGAttributePropertySelector.h"
+#include "Metadata/Accessors/PCGAttributeAccessorHelpers.h"
+
+#define LOCTEXT_NAMESPACE "PCGPropertyToParamDataElement"
+
+#if WITH_EDITOR
+void UPCGGetActorPropertySettings::GetTrackedActorKeys(FPCGActorSelectionKeyToSettingsMap& OutKeysToSettings, TArray<TObjectPtr<const UPCGGraph>>& OutVisitedGraphs) const
+{
+	OutKeysToSettings.FindOrAdd(ActorSelector.GetAssociatedKey()).Emplace(this, bTrackActorsOnlyWithinBounds);
+}
+
+void UPCGGetActorPropertySettings::PostEditChangeProperty(FPropertyChangedEvent& PropertyChangedEvent)
+{
+	Super::PostEditChangeProperty(PropertyChangedEvent);
+
+	if (PropertyChangedEvent.GetMemberPropertyName() == GET_MEMBER_NAME_CHECKED(UPCGGetActorPropertySettings, ActorSelector))
+	{
+		if (PropertyChangedEvent.GetPropertyName() == GET_MEMBER_NAME_CHECKED(FPCGActorSelectorSettings, ActorSelection))
+		{
+			if (ActorSelector.ActorSelection != EPCGActorSelection::ByClass)
+			{
+				ActorSelector.ActorSelectionClass = TSubclassOf<AActor>();
+			}
+		}
+	}
+}
+#endif // WITH_EDITOR
+
+void UPCGGetActorPropertySettings::PostLoad()
+{
+	Super::PostLoad();
+
+	// Migrate deprecated actor selection settings to struct if needed
+	if (ActorSelection_DEPRECATED != EPCGActorSelection::ByTag ||
+		ActorSelectionTag_DEPRECATED != NAME_None ||
+		ActorSelectionName_DEPRECATED != NAME_None ||
+		ActorSelectionClass_DEPRECATED != TSubclassOf<AActor>() ||
+		ActorFilter_DEPRECATED != EPCGActorFilter::Self ||
+		bIncludeChildren_DEPRECATED != false)
+	{
+		ActorSelector.ActorSelection = ActorSelection_DEPRECATED;
+		ActorSelector.ActorSelectionTag = ActorSelectionTag_DEPRECATED;
+		ActorSelector.ActorSelectionClass = ActorSelectionClass_DEPRECATED;
+		ActorSelector.ActorFilter = ActorFilter_DEPRECATED;
+		ActorSelector.bIncludeChildren = bIncludeChildren_DEPRECATED;
+
+		ActorSelection_DEPRECATED = EPCGActorSelection::ByTag;
+		ActorSelectionTag_DEPRECATED = NAME_None;
+		ActorSelectionName_DEPRECATED = NAME_None;
+		ActorSelectionClass_DEPRECATED = TSubclassOf<AActor>();
+		ActorFilter_DEPRECATED = EPCGActorFilter::Self;
+		bIncludeChildren_DEPRECATED = false;
+	}
+
+	if (ActorSelector.ActorSelection != EPCGActorSelection::ByClass)
+	{
+		ActorSelector.ActorSelectionClass = TSubclassOf<AActor>();
+	}
+}
+
+FName UPCGGetActorPropertySettings::AdditionalTaskName() const
+{
+#if WITH_EDITOR
+	return ActorSelector.GetTaskName(GetDefaultNodeTitle());
+#else
+	return Super::AdditionalTaskName();
+#endif
+}
+
+
+TArray<FPCGPinProperties> UPCGGetActorPropertySettings::OutputPinProperties() const
+{
+	TArray<FPCGPinProperties> PinProperties;
+	PinProperties.Emplace(PCGPinConstants::DefaultOutputLabel, EPCGDataType::Param);
+
+	return PinProperties;
+}
+
+FPCGElementPtr UPCGGetActorPropertySettings::CreateElement() const
+{
+	return MakeShared<FPCGGetActorPropertyElement>();
+}
+
+bool FPCGGetActorPropertyElement::ExecuteInternal(FPCGContext* Context) const
+{
+	TRACE_CPUPROFILER_EVENT_SCOPE(FPCGGetActorPropertyElement::Execute);
+
+	check(Context);
+
+	const UPCGGetActorPropertySettings* Settings = Context->GetInputSettings<UPCGGetActorPropertySettings>();
+	check(Settings);
+
+	// Early out if arguments are not specified
+	if (Settings->PropertyName == NAME_None || (Settings->bSelectComponent && !Settings->ComponentClass))
+	{
+		PCGE_LOG(Warning, GraphAndLog, LOCTEXT("ParametersMissing", "Some parameters are missing, aborting"));
+		return true;
+	}
+
+#if !WITH_EDITOR
+	// If we have no output connected, nothing to do
+	// Optimization possibly only in non-editor builds, otherwise we could poison the input-driven cache
+	if (!Context->Node || !Context->Node->IsOutputPinConnected(PCGPinConstants::DefaultOutputLabel))
+	{
+		PCGE_LOG(Verbose, LogOnly, LOCTEXT("UnconnectedNode", "Node is not connected, nothing to do"));
+		return true;
+	}
+#endif
+
+	// First find the actor depending on the selection
+	UPCGComponent* OriginalComponent = UPCGBlueprintHelpers::GetOriginalComponent(*Context);
+	TFunction<bool(const AActor*)> BoundsCheck = [](const AActor*) -> bool { return true; };
+	auto NoSelfIgnoreCheck = [](const AActor*) -> bool { return true; };
+
+	if (OriginalComponent && OriginalComponent->GetOwner() && Settings->ActorSelector.bMustOverlapSelf)
+	{
+		const FBox ActorBounds = PCGHelpers::GetActorBounds(OriginalComponent->GetOwner());
+		BoundsCheck = [Settings, ActorBounds, OriginalComponent](const AActor* OtherActor) -> bool
+		{
+			const FBox OtherActorBounds = OtherActor ? PCGHelpers::GetGridBounds(OtherActor, OriginalComponent) : FBox(EForceInit::ForceInit);
+			return ActorBounds.Intersect(OtherActorBounds);
+		};
+	}
+
+	TArray<AActor*> FoundActors = PCGActorSelector::FindActors(Settings->ActorSelector, OriginalComponent, BoundsCheck, NoSelfIgnoreCheck);
+
+	if (FoundActors.IsEmpty())
+	{
+		PCGE_LOG(Verbose, LogOnly, LOCTEXT("NoActorFound", "No matching actor was found"));
+		return true;
+	}
+
+	for (AActor* FoundActor : FoundActors)
+	{
+		if (!FoundActor)
+		{
+			continue;
+		}
+
+		// From there, we either check the actor, or the component attached to it.
+		UObject* ObjectToInspect = FoundActor;
+		if (Settings->bSelectComponent)
+		{
+			ObjectToInspect = FoundActor->GetComponentByClass(Settings->ComponentClass);
+			if (!ObjectToInspect)
+			{
+				PCGE_LOG(Error, GraphAndLog, FText::Format(LOCTEXT("ComponentDoesNotExist", "Component class '{0}' does not exist in the found actor {1}"), FText::FromString(Settings->ComponentClass->GetName()), FText::FromString(FoundActor->GetName())));
+				return true;
+			}
+		}
+
+		const FPCGAttributePropertySelector Selector = FPCGAttributePropertySelector::CreateSelectorFromString(Settings->PropertyName.ToString());
+
+		PCGPropertyHelpers::FExtractorParameters Parameters{ ObjectToInspect, ObjectToInspect->GetClass(), Selector, Settings->OutputAttributeName, Settings->bForceObjectAndStructExtraction, /*bPropertyNeedsToBeVisible=*/true };
+		if (UPCGParamData* ParamData = PCGPropertyHelpers::ExtractPropertyAsAttributeSet(Parameters, Context))
+		{
+			TArray<FPCGTaggedData>& Outputs = Context->OutputData.TaggedData;
+			FPCGTaggedData& Output = Outputs.Emplace_GetRef();
+			Output.Data = ParamData;
+		}
+		else
+		{
+			if (Selector.GetName() == NAME_None)
+			{
+				PCGE_LOG(Error, GraphAndLog, FText::Format(LOCTEXT("FailedToExtractActor", "Fail to extract actor {0}."), FText::FromString(FoundActor->GetName())));
+			}
+			else
+			{
+				PCGE_LOG(Error, GraphAndLog, FText::Format(LOCTEXT("FailedToExtract", "Fail to extract the property '{0}' on actor {1}."), Selector.GetDisplayText(), FText::FromString(FoundActor->GetName())));
+			}
+		}
+	}
+
+	return true;
+}
+
+#undef LOCTEXT_NAMESPACE
diff --git a/Engine/Plugins/Experimental/PCG/Source/PCG/Private/Elements/PCGGetPropertyFromObjectPath.cpp b/Engine/Plugins/Experimental/PCG/Source/PCG/Private/Elements/PCGGetPropertyFromObjectPath.cpp
new file mode 100644
--- /dev/null
+++ b/Engine/Plugins/Experimental/PCG/Source/PCG/Private/Elements/PCGGetPropertyFromObjectPath.cpp
@@ -0,0 +1,244 @@
+// Copyright Epic Games, Inc. All Rights Reserved.
+
+#include "Elements/PCGGetPropertyFromObjectPath.h"
+
+#include "PCGParamData.h"
+#include "Helpers/PCGPropertyHelpers.h"
+#include "Metadata/Accessors/IPCGAttributeAccessor.h"
+#include "Metadata/Accessors/PCGAttributeAccessorKeys.h"
+#include "Metadata/Accessors/PCGAttributeAccessorHelpers.h"
+
+#include "Engine/AssetManager.h"
+
+#define LOCTEXT_NAMESPACE "PCGGetPropertyFromObjectPathElement"
+
+#if WITH_EDITOR
+FName UPCGGetPropertyFromObjectPathSettings::GetDefaultNodeName() const
+{
+	return FName(TEXT("GetPropertyFromObjectPath"));
+}
+
+FText UPCGGetPropertyFromObjectPathSettings::GetDefaultNodeTitle() const
+{
+	return LOCTEXT("NodeTitle", "Get Property From Object Path");
+}
+
+bool UPCGGetPropertyFromObjectPathSettings::IsPinUsedByNodeExecution(const UPCGPin* InPin) const
+{
+	return !InPin || (InPin->Properties.Label != PCGPinConstants::DefaultInputLabel) || InPin->IsConnected();
+}
+
+bool UPCGGetPropertyFromObjectPathSettings::CanEditChange(const FProperty* InProperty) const
+{
+	if (!InProperty || !Super::CanEditChange(InProperty))
+	{
+		return false;
+	}
+
+	const UPCGNode* Node = Cast<UPCGNode>(GetOuter());
+	const bool InPinIsConnected = Node ? Node->IsInputPinConnected(PCGPinConstants::DefaultInputLabel) : false;
+
+	if (InProperty->GetFName() == GET_MEMBER_NAME_CHECKED(UPCGGetPropertyFromObjectPathSettings, InputSource))
+	{
+		return InPinIsConnected;
+	}
+	else if (InProperty->GetFName() == GET_MEMBER_NAME_CHECKED(UPCGGetPropertyFromObjectPathSettings, ObjectPathsToExtract))
+	{
+		return !InPinIsConnected;
+	}
+
+	return true;
+}
+#endif // WITH_EDITOR
+
+FPCGElementPtr UPCGGetPropertyFromObjectPathSettings::CreateElement() const
+{
+	return MakeShared<FPCGGetPropertyFromObjectPathElement>();
+}
+
+TArray<FPCGPinProperties> UPCGGetPropertyFromObjectPathSettings::InputPinProperties() const
+{
+	TArray<FPCGPinProperties> PinProperties;
+	PinProperties.Emplace(PCGPinConstants::DefaultInputLabel, EPCGDataType::Param);
+
+	return PinProperties;
+}
+
+TArray<FPCGPinProperties> UPCGGetPropertyFromObjectPathSettings::OutputPinProperties() const
+{
+	TArray<FPCGPinProperties> PinProperties;
+	PinProperties.Emplace(PCGPinConstants::DefaultOutputLabel, EPCGDataType::Param);
+
+	return PinProperties;
+}
+
+FPCGContext* FPCGGetPropertyFromObjectPathElement::CreateContext()
+{
+	return new FPCGGetPropertyFromObjectPathContext();
+}
+
+bool FPCGGetPropertyFromObjectPathElement::PrepareDataInternal(FPCGContext* Context) const
+{
+	TRACE_CPUPROFILER_EVENT_SCOPE(FPCGGetPropertyFromObjectPathElement::PrepareData);
+
+	check(Context);
+
+	const UPCGGetPropertyFromObjectPathSettings* Settings = Context->GetInputSettings<UPCGGetPropertyFromObjectPathSettings>();
+	check(Settings);
+
+	FPCGGetPropertyFromObjectPathContext* ThisContext = static_cast<FPCGGetPropertyFromObjectPathContext*>(Context);
+
+	bool bIsDone = true;
+
+	if (!ThisContext->bRequestSent)
+	{
+		const UPCGNode* Node = Context->Node;
+		const bool InPinIsConnected = Node ? Node->IsInputPinConnected(PCGPinConstants::DefaultInputLabel) : false;
+
+		// First gather all the soft objects paths to extract
+		TArray<FSoftObjectPath> ObjectsToLoad;
+		if (!InPinIsConnected)
+		{
+			ThisContext->PathsToObjectsToExtractAndIncomingDataIndex.Reserve(Settings->ObjectPathsToExtract.Num());
+			for (const FSoftObjectPath& Path : Settings->ObjectPathsToExtract)
+			{
+				if (!Path.IsNull())
+				{
+					ObjectsToLoad.AddUnique(Path);
+				}
+
+				ThisContext->PathsToObjectsToExtractAndIncomingDataIndex.Emplace(Path, -1);
+			}
+		}
+		else
+		{
+			const TArray<FPCGTaggedData> Inputs = Context->InputData.GetInputsByPin(PCGPinConstants::DefaultInputLabel);
+			for (int32 Index = 0; Index < Inputs.Num(); ++Index)
+			{
+				const FPCGTaggedData& Input = Inputs[Index];
+
+				if (!Input.Data)
+				{
+					PCGE_LOG(Error, GraphAndLog, FText::Format(LOCTEXT("InvalidData", "Invalid data for input {0}"), FText::AsNumber(Index)));
+					continue;
+				}
+
+				const FPCGAttributePropertyInputSelector AttributeSelector = Settings->InputSource.CopyAndFixLast(Input.Data);
+				const TUniquePtr<const IPCGAttributeAccessor> Accessor = PCGAttributeAccessorHelpers::CreateConstAccessor(Input.Data, AttributeSelector);
+				const TUniquePtr<const IPCGAttributeAccessorKeys> Keys = PCGAttributeAccessorHelpers::CreateConstKeys(Input.Data, AttributeSelector);
+
+				if (!Accessor.IsValid() || !Keys.IsValid())
+				{
+					PCGE_LOG(Error, GraphAndLog, FText::Format(LOCTEXT("AttributeNotFound", "Attribute/Property '{0}' does not exist on input {1}"), AttributeSelector.GetDisplayText(), FText::AsNumber(Index)));
+					continue;
+				}
+
+				const int32 NumElementsToAdd = Keys->GetNum();
+				if (NumElementsToAdd == 0)
+				{
+					continue;
+				}
+
+				TArray<FSoftObjectPath> InputValues;
+				InputValues.SetNum(NumElementsToAdd);
+				if (Accessor->GetRange(MakeArrayView(InputValues), 0, *Keys, EPCGAttributeAccessorFlags::AllowConstructible))
+				{
+					ThisContext->PathsToObjectsToExtractAndIncomingDataIndex.Reserve(ThisContext->PathsToObjectsToExtractAndIncomingDataIndex.Num() + NumElementsToAdd);
+					for (int32 i = 0; i < InputValues.Num(); ++i)
+					{
+						FSoftObjectPath& Path = InputValues[i];
+						if (!Path.IsNull())
+						{
+							ObjectsToLoad.AddUnique(Path);
+							ThisContext->PathsToObjectsToExtractAndIncomingDataIndex.Emplace(std::move(Path), Index);
+						}
+						else
+						{
+							PCGE_LOG(Error, GraphAndLog, FText::Format(LOCTEXT("InvalidPath", "Value number {0} for Attribute/Property '{1}' on input {2} is not a valid path or is null. Will be ignored."), FText::AsNumber(i), AttributeSelector.GetDisplayText(), FText::AsNumber(Index)));
+							continue;
+						}
+					}
+				}
+				else
+				{
+					PCGE_LOG(Error, GraphAndLog, FText::Format(LOCTEXT("InvalidAttribute", "Attribute/Property '{0}'({1}) is not convertible to a SoftObjectPath on input {2}"), AttributeSelector.GetDisplayText(), PCG::Private::GetTypeNameText(Accessor->GetUnderlyingType()), FText::AsNumber(Index)));
+					continue;
+				}
+			}
+		}
+
+		if (!ObjectsToLoad.IsEmpty())
+		{
+			if (Settings->bSynchronousLoad)
+			{
+				ThisContext->LoadHandle = UAssetManager::GetStreamableManager().RequestSyncLoad(std::move(ObjectsToLoad));
+			}
+			else
+			{
+				bIsDone = false;
+				ThisContext->bIsPaused = true;
+				ThisContext->LoadHandle = UAssetManager::GetStreamableManager().RequestAsyncLoad(std::move(ObjectsToLoad), [Context]() { Context->bIsPaused = false; });
+			}
+		}
+
+		ThisContext->bRequestSent = true;
+	}
+
+	return bIsDone;
+}
+
+bool FPCGGetPropertyFromObjectPathElement::ExecuteInternal(FPCGContext* Context) const
+{
+	TRACE_CPUPROFILER_EVENT_SCOPE(FPCGGetPropertyFromObjectPathElement::Execute);
+
+	check(Context);
+	FPCGGetPropertyFromObjectPathContext* ThisContext = static_cast<FPCGGetPropertyFromObjectPathContext*>(Context);
+
+	const UPCGGetPropertyFromObjectPathSettings* Settings = Context->GetInputSettings<UPCGGetPropertyFromObjectPathSettings>();
+	check(Settings);
+
+	const TArray<FPCGTaggedData> Inputs = Context->InputData.GetInputsByPin(PCGPinConstants::DefaultInputLabel);
+
+	for (const TTuple<FSoftObjectPath, int32>& SoftPathAndIndex : ThisContext->PathsToObjectsToExtractAndIncomingDataIndex)
+	{
+		const FSoftObjectPath& SoftPath = SoftPathAndIndex.Get<FSoftObjectPath>();
+		const int32 Index = SoftPathAndIndex.Get<int32>();
+
+		const UObject* Object = SoftPath.ResolveObject();
+		if (!Object)
+		{
+			PCGE_LOG(Error, GraphAndLog, FText::Format(LOCTEXT("FailedToLoad", "Failed to load object {0}"), FText::FromString(SoftPath.ToString())));
+			continue;
+		}
+
+		const FPCGAttributePropertySelector Selector = FPCGAttributePropertySelector::CreateSelectorFromString(Settings->PropertyName.ToString());
+
+		PCGPropertyHelpers::FExtractorParameters Parameters{ Object, Object->GetClass(), Selector, Settings->OutputAttributeName, Settings->bForceObjectAndStructExtraction, /*bPropertyNeedsToBeVisible=*/true };
+		if (UPCGParamData* ParamData = PCGPropertyHelpers::ExtractPropertyAsAttributeSet(Parameters, Context))
+		{
+			TArray<FPCGTaggedData>& Outputs = Context->OutputData.TaggedData;
+			FPCGTaggedData& Output = Outputs.Emplace_GetRef();
+			Output.Data = ParamData;
+			Output.Pin = PCGPinConstants::DefaultOutputLabel;
+			if (Index >= 0)
+			{
+				Output.Tags = Inputs[Index].Tags;
+			}
+		}
+		else
+		{
+			if (Selector.GetName() == NAME_None)
+			{
+				PCGE_LOG(Error, GraphAndLog, FText::Format(LOCTEXT("FailedToExtractObject", "Fail to extract object {0}."), FText::FromString(Object->GetName())));
+			}
+			else
+			{
+				PCGE_LOG(Error, GraphAndLog, FText::Format(LOCTEXT("FailedToExtract", "Fail to extract the property '{0}' on object {1}."), Selector.GetDisplayText(), FText::FromString(Object->GetName())));
+			}
+		}
+	}
+
+	return true;
+}
+
+#undef LOCTEXT_NAMESPACE
diff --git a/Engine/Plugins/Experimental/PCG/Source/PCG/Private/Elements/PCGPropertyToParamData.cpp b/Engine/Plugins/Experimental/PCG/Source/PCG/Private/Elements/PCGPropertyToParamData.cpp
deleted file mode 100644
diff --git a/Engine/Plugins/Experimental/PCG/Source/PCG/Private/Helpers/PCGPropertyHelpers.cpp b/Engine/Plugins/Experimental/PCG/Source/PCG/Private/Helpers/PCGPropertyHelpers.cpp
--- a/Engine/Plugins/Experimental/PCG/Source/PCG/Private/Helpers/PCGPropertyHelpers.cpp
+++ b/Engine/Plugins/Experimental/PCG/Source/PCG/Private/Helpers/PCGPropertyHelpers.cpp
@@ -130,11 +130,17 @@ UPCGParamData* PCGPropertyHelpers::ExtractPropertyAsAttributeSet(const PCGProper
 	check(Parameters.Container && Parameters.Class);
 
 	const void* Container = Parameters.Container;
-	const FProperty* Property = ExtractPropertyChain(Parameters.Class, Parameters.PropertySelector.GetName(), Parameters.PropertySelector.GetExtraNames(), Parameters.bPropertyNeedsToBeVisible, Container, InOptionalContext);
-
-	if (!Property)
+	const FProperty* Property = nullptr;
+	const FName PropertyName = Parameters.PropertySelector.GetName();
+	const bool ExtractRoot = (PropertyName == NAME_None);
+	// If Name is none, extract the container as-is, using Parameters.Class, otherwise, extract the chain.
+	if (!ExtractRoot)
 	{
-		return nullptr;
+		Property = ExtractPropertyChain(Parameters.Class, PropertyName, Parameters.PropertySelector.GetExtraNames(), Parameters.bPropertyNeedsToBeVisible, Container, InOptionalContext);
+		if (!Property)
+		{
+			return nullptr;
+		}
 	}
 
 	// If the property is an array, we will work on the underlying property, and extract each element as an entry in the param data
@@ -154,14 +160,19 @@ UPCGParamData* PCGPropertyHelpers::ExtractPropertyAsAttributeSet(const PCGProper
 	const bool bShouldExtract = Parameters.bShouldExtract || !PCGAttributeAccessorHelpers::IsPropertyAccessorSupported(Property);
 
 	// Special case where the property is a struct/object, that is not supported by our metadata, we will try to break it down to multiple attributes in the resulting param data, if asked.
-	if ((Property->IsA<FStructProperty>() || Property->IsA<FObjectProperty>()) && bShouldExtract)
+	if (ExtractRoot || ((Property->IsA<FStructProperty>() || Property->IsA<FObjectProperty>()) && bShouldExtract))
 	{
-		const UScriptStruct* UnderlyingStruct = nullptr;
-		const UClass* UnderlyingClass = nullptr;
+		const UStruct* UnderlyingClass = nullptr;
 
-		if (const FStructProperty* StructProperty = CastField<FStructProperty>(Property))
+		if (ExtractRoot)
+		{
+			UnderlyingClass = Parameters.Class;
+			// Identity
+			AddressFunc = [](const void* InAddress) { return InAddress; };
+		}
+		else if (const FStructProperty* StructProperty = CastField<FStructProperty>(Property))
 		{
-			UnderlyingStruct = StructProperty->Struct;
+			UnderlyingClass = StructProperty->Struct;
 			AddressFunc = [StructProperty](const void* InAddress) { return StructProperty->ContainerPtrToValuePtr<void>(InAddress); };
 		}
 		else if (const FObjectProperty* ObjectProperty = CastField<FObjectProperty>(Property))
@@ -170,7 +181,7 @@ UPCGParamData* PCGPropertyHelpers::ExtractPropertyAsAttributeSet(const PCGProper
 			AddressFunc = [ObjectProperty](const void* InAddress) { return ObjectProperty->GetObjectPropertyValue_InContainer(InAddress); };
 		}
 
-		check(UnderlyingStruct || UnderlyingClass);
+		check(UnderlyingClass);
 		check(!!AddressFunc);
 
 		// Re-use code from overridable params
@@ -185,18 +196,18 @@ UPCGParamData* PCGPropertyHelpers::ExtractPropertyAsAttributeSet(const PCGProper
 			Config.ExcludePropertyFlags = ExcludePropertyFlags;
 			Config.IncludePropertyFlags = IncludePropertyFlags;
 		}
-		TArray<FPCGSettingsOverridableParam> AllChildProperties = UnderlyingStruct ? PCGSettingsHelpers::GetAllOverridableParams(UnderlyingStruct, Config) : PCGSettingsHelpers::GetAllOverridableParams(UnderlyingClass, Config);
+		TArray<FPCGSettingsOverridableParam> AllChildProperties = PCGSettingsHelpers::GetAllOverridableParams(UnderlyingClass, Config);
 
 		for (const FPCGSettingsOverridableParam& Param : AllChildProperties)
 		{
 			if (ensure(!Param.PropertiesNames.IsEmpty()))
 			{
 				const FName ChildPropertyName = Param.PropertiesNames[0];
-				if (const FProperty* ChildProperty = (UnderlyingStruct ? UnderlyingStruct->FindPropertyByName(ChildPropertyName) : UnderlyingClass->FindPropertyByName(ChildPropertyName)))
+				if (const FProperty* ChildProperty = UnderlyingClass->FindPropertyByName(ChildPropertyName))
 				{
 					// We use authored name as attribute name to avoid issue with noisy property names, like in UUserDefinedStructs, where some random number is appended to the property name.
 					// By default, it will just return the property name anyway.
-					const FString AuthoredName = UnderlyingStruct ? UnderlyingStruct->GetAuthoredNameForField(ChildProperty) : UnderlyingClass->GetAuthoredNameForField(ChildProperty);
+					const FString AuthoredName = UnderlyingClass->GetAuthoredNameForField(ChildProperty);
 					ExtractableProperties.Emplace(FName(AuthoredName), ChildProperty);
 				}
 			}
diff --git a/Engine/Plugins/Experimental/PCG/Source/PCG/Private/Helpers/PCGSettingsHelpers.cpp b/Engine/Plugins/Experimental/PCG/Source/PCG/Private/Helpers/PCGSettingsHelpers.cpp
--- a/Engine/Plugins/Experimental/PCG/Source/PCG/Private/Helpers/PCGSettingsHelpers.cpp
+++ b/Engine/Plugins/Experimental/PCG/Source/PCG/Private/Helpers/PCGSettingsHelpers.cpp
@@ -273,8 +273,7 @@
 		}
 	}
 
-	template <typename ClassType>
-	TArray<FPCGSettingsOverridableParam> GetAllOverridableParamsImpl(const ClassType* InClass, const FPCGGetAllOverridableParamsConfig& InConfig)
+	TArray<FPCGSettingsOverridableParam> GetAllOverridableParams(const UStruct* InClass, const FPCGGetAllOverridableParamsConfig& InConfig)
 	{
 		TArray<FName> LabelCache;
 
@@ -375,42 +374,12 @@
 				continue;
 			}
 
-			// Validate that the property can be overriden by params
-			if (PCGAttributeAccessorHelpers::IsPropertyAccessorSupported(Property))
+			auto RecursiveExtraction = [&InConfig, &Res, &LabelCache, Property, InClass, &GatherAliases](const UStruct* NextClass)
 			{
-				FName Label = NAME_None;
-				bool bHasNameClash = false;
-#if WITH_EDITOR
-				// GetDisplayNameText is not available in non-editor build.
-				Label = *Property->GetDisplayNameText().ToString();
-				const int32 CachedLabelIndex = LabelCache.IndexOfByKey(Label);
-				if (CachedLabelIndex != INDEX_NONE)
-				{
-					// If we have a clash, we will use the full path, so mark this param and the other that clashed to use the full path.
-					Res[CachedLabelIndex].bHasNameClash = true;
-					Res[CachedLabelIndex].Label = FName(Res[CachedLabelIndex].GetDisplayPropertyPath());
-					bHasNameClash = true;
-				}
-
-				LabelCache.AddUnique(Label);
-#endif // WITH_EDITOR
-
-				FPCGSettingsOverridableParam& Param = Res.Emplace_GetRef();
-				Param.PropertiesNames.Add(Property->GetFName());
-				Param.Properties.Add(Property);
-				Param.PropertyClass = InClass;
-				Param.bHasNameClash = bHasNameClash;
-#if WITH_EDITOR
-				Param.Label = bHasNameClash ? FName(Param.GetDisplayPropertyPath()) : Label;
-				GatherAliases(Property, Param);
-#endif // WITH_EDITOR
-			}
-			else if (const FStructProperty* StructProperty = CastField<FStructProperty>(Property))
-			{
 				// Reached max depth
 				if (InConfig.MaxStructDepth == 0)
 				{
-					continue;
+					return;
 				}
 
 				// Use the seed, and don't check metadata for PCG overridable.
@@ -425,7 +394,7 @@
 					RecurseConfig.MaxStructDepth--;
 				}
 
-				for (FPCGSettingsOverridableParam& ChildParam : GetAllOverridableParams(StructProperty->Struct, RecurseConfig))
+				for (FPCGSettingsOverridableParam& ChildParam : GetAllOverridableParams(NextClass, RecurseConfig))
 				{
 					FName Label = ChildParam.Label;
 					bool bHasNameClash = false;
@@ -470,19 +439,54 @@
 					}
 #endif // WITH_EDITOR
 				}
+			};
+
+			const FProperty* PropertyToCheck = Property;
+			if (InConfig.bExtractArrays && Property->IsA<FArrayProperty>())
+			{
+				PropertyToCheck = CastFieldChecked<FArrayProperty>(Property)->Inner;
 			}
+
+			if (PropertyToCheck->IsA<FObjectProperty>() && InConfig.bExtractObjects)
+			{
+				RecursiveExtraction(CastFieldChecked<FObjectProperty>(PropertyToCheck)->PropertyClass);
+			}
+			// Validate that the property can be overriden by params
+			else if (PCGAttributeAccessorHelpers::IsPropertyAccessorSupported(PropertyToCheck))
+			{
+				FName Label = NAME_None;
+				bool bHasNameClash = false;
+#if WITH_EDITOR
+				// GetDisplayNameText is not available in non-editor build.
+				Label = *Property->GetDisplayNameText().ToString();
+				const int32 CachedLabelIndex = LabelCache.IndexOfByKey(Label);
+				if (CachedLabelIndex != INDEX_NONE)
+				{
+					// If we have a clash, we will use the full path, so mark this param and the other that clashed to use the full path.
+					Res[CachedLabelIndex].bHasNameClash = true;
+					Res[CachedLabelIndex].Label = FName(Res[CachedLabelIndex].GetDisplayPropertyPath());
+					bHasNameClash = true;
+				}
+
+				LabelCache.AddUnique(Label);
+#endif // WITH_EDITOR
+
+				FPCGSettingsOverridableParam& Param = Res.Emplace_GetRef();
+				Param.PropertiesNames.Add(Property->GetFName());
+				Param.Properties.Add(Property);
+				Param.PropertyClass = InClass;
+				Param.bHasNameClash = bHasNameClash;
+#if WITH_EDITOR
+				Param.Label = bHasNameClash ? FName(Param.GetDisplayPropertyPath()) : Label;
+				GatherAliases(Property, Param);
+#endif // WITH_EDITOR
+			}
+			else if (const FStructProperty* StructProperty = CastField<FStructProperty>(PropertyToCheck))
+			{
+				RecursiveExtraction(StructProperty->Struct);
+			}
 		}
 
 		return Res;
 	}
-
-	TArray<FPCGSettingsOverridableParam> GetAllOverridableParams(const UClass* InClass, const FPCGGetAllOverridableParamsConfig& InConfig)
-	{
-		return GetAllOverridableParamsImpl(InClass, InConfig);
-	}
-
-	TArray<FPCGSettingsOverridableParam> GetAllOverridableParams(const UScriptStruct* InStruct, const FPCGGetAllOverridableParamsConfig& InConfig)
-	{
-		return GetAllOverridableParamsImpl(InStruct, InConfig);
-	}
 }
diff --git a/Engine/Plugins/Experimental/PCG/Source/PCG/Private/Tests/Elements/PCGGetActorPropertyTest.cpp b/Engine/Plugins/Experimental/PCG/Source/PCG/Private/Tests/Elements/PCGGetActorPropertyTest.cpp
new file mode 100644
--- /dev/null
+++ b/Engine/Plugins/Experimental/PCG/Source/PCG/Private/Tests/Elements/PCGGetActorPropertyTest.cpp
@@ -0,0 +1,444 @@
+// Copyright Epic Games, Inc. All Rights Reserved.
+
+#include "Tests/Elements/PCGPropertyToParamDataElementTest.h"
+#include "PCGComponent.h"
+#include "PCGContext.h"
+#include "PCGGraph.h"
+#include "PCGParamData.h"
+#include "PCGVolume.h"
+#include "Elements/PCGGetActorProperty.h"
+#include "Tests/PCGTestsCommon.h"
+
+#include UE_INLINE_GENERATED_CPP_BY_NAME(PCGPropertyToParamDataElementTest)
+
+// To run: automation runtest pcg.tests.PropertyToParamData
+
+#if WITH_EDITOR
+
+IMPLEMENT_CUSTOM_SIMPLE_AUTOMATION_TEST(FPCGPropertyToParamDataPropertyTypeTest, FPCGTestBaseClass, "Plugins.PCG.PropertyToParamData.PropertyType", PCGTestsCommon::TestFlags)
+IMPLEMENT_CUSTOM_SIMPLE_AUTOMATION_TEST(FPCGPropertyToParamDataActorFindTest, FPCGTestBaseClass, "Plugins.PCG.PropertyToParamData.ActorFind", PCGTestsCommon::TestFlags)
+
+
+/**
+* Generic method to test the extraction and then verify that the values matches the one passed as input.
+* Use variadic template for expected values, allowing to test the extraction of structs/objects that doesn't have the same type for all their members.
+* bShouldFail is used as a constexpr bool to cut the branch where we try to compare the values, if it is expected to fail.
+* Be careful, if you extract structures, all int properties are int64 (even if they were like u8) and all float will be double. Take that into account in the ExpectedValues types.
+* 
+* @param TestInstance - Instance of the test, used for our test conditions
+* @param TestData - Test data used to run the element
+* @param PropertyName - Property name to extract
+* @param AttributeNames - List of all the attributes to test. AttributeNames count MUST divide the count of expected values. If there is more than one attribute name, it will be marked as extracting structures
+* @param ExtraTestWhat - Simple string to append to all test strings, to know which test we are currently running.
+* @param ExpectedValues - List of all values to test. If there is more than one, it can be because it is an extraction, or testing an array. If it is an array of struct, and attribute names is {'X', 'Y'}, it must follow this order: A[0].X, A[0].Y, A[1].X, A[1].Y, ...
+* @return if the test succeeded or not.
+*/
+template <bool bShouldFail, typename ...AttributeTypes>
+bool VerifyAttributeValue(FPCGTestBaseClass* TestInstance, PCGTestsCommon::FTestData& TestData, FName PropertyName, TArray<FName> AttributeNames, FString ExtraTestWhat, AttributeTypes&& ...ExpectedValues)
+{
+	// Use TestData settings to set the property name
+	UPCGGetActorPropertySettings* Settings = Cast<UPCGGetActorPropertySettings>(TestData.Settings);
+	if (!TestInstance->TestNotNull(TEXT("CastToUPCGPropertyToParamDataSettings"), Settings))
+	{
+		return false;
+	}
+
+	// Make sure this call is not ill-formed
+	check(!AttributeNames.IsEmpty());
+	if (!ensureMsgf(sizeof...(ExpectedValues) % AttributeNames.Num() == 0, TEXT("%s"), TEXT("Call is ill-formed, there is a mismatch between the number of attribute names and the number of expected values.")))
+	{
+		return false;
+	}
+
+	Settings->PropertyName = PropertyName;
+	Settings->OutputAttributeName = AttributeNames[0];
+	Settings->bForceObjectAndStructExtraction = AttributeNames.Num() > 1;
+
+	// Add 2 nodes, PropertyToParamDataNode and a Trivial node (just there for the connection)
+	UPCGNode* TestNode = TestData.TestPCGComponent->GetGraph()->AddNode(Settings);
+
+	UPCGTrivialSettings* TrivialSettings = NewObject<UPCGTrivialSettings>();
+	UPCGNode* TrivialNode = TestData.TestPCGComponent->GetGraph()->AddNode(TrivialSettings);
+
+	TestNode->GetOutputPins()[0]->AddEdgeTo(TrivialNode->GetInputPins()[0]);
+
+	FPCGElementPtr Element = Settings->GetElement();
+	TUniquePtr<FPCGContext> Context = PCGTestsCommon::InitializeTestContext(Element.Get(), FPCGDataCollection(), TestData.TestPCGComponent, TestNode);
+
+	// Run the element, can throw an error.
+	while (!Element->Execute(Context.Get()))
+	{
+	};
+
+	auto FormatWithPropertyName = [&ExtraTestWhat, PropertyName](const FString& InText) -> FString
+	{
+		return FString::Format(TEXT("{0}_{1}_{2}"), { ExtraTestWhat, PropertyName.ToString(), InText });
+	};
+
+	auto FormatWithPropertyNameAndSubNames = [&ExtraTestWhat, PropertyName](const FString& InText, const FName InSubName, const int InItemKey) -> FString
+	{
+		return FString::Format(TEXT("{0}_{1}_SubName:{2}_ItemKey:{3} - {4}"), { ExtraTestWhat, PropertyName.ToString(), InSubName.ToString(), FString::FromInt(InItemKey), InText });
+	};
+
+	bool bSuccess = true;
+
+	if constexpr (bShouldFail)
+	{
+		// If it should fail, it will have no output.
+		bSuccess = TestInstance->TestEqual(FormatWithPropertyName("NumOutput"), Context->OutputData.GetAllParams().Num(), 0);
+	}
+	else
+	{
+		// If it should succeed, we check the number of output.
+		bSuccess = TestInstance->TestEqual(FormatWithPropertyName("NumOutput"), Context->OutputData.GetAllParams().Num(), 1);
+
+		if (bSuccess)
+		{
+			const UPCGParamData* ParamData = CastChecked<UPCGParamData>(Context->OutputData.GetAllParams()[0].Data);
+			const UPCGMetadata* Metadata = ParamData->ConstMetadata();
+			check(Metadata);
+
+			int32 Index = 0;
+			int32 ItemKey = 0;
+
+			// Make sure that we have exactly the same number of attributes as expected
+			if (!TestInstance->TestEqual(FormatWithPropertyName("Attribute number is matching"), Metadata->GetAttributeCount(), AttributeNames.Num()))
+			{
+				// Cleanup behind ourselves.
+				TestData.TestPCGComponent->GetGraph()->RemoveNode(TrivialNode);
+				TestData.TestPCGComponent->GetGraph()->RemoveNode(TestNode);
+
+				return false;
+			}
+
+			// Fold expression that will iterate on all Expected Values, to check their types and their values.
+			// We keep track of the current Index, to know which attribute to check, and an ItemKey to know where to look at in the value array of the attribute.
+			// cf. comment at the top to know in which order ExpectedValues should be.
+			([&Index, &AttributeNames, Metadata, &FormatWithPropertyNameAndSubNames, TestInstance, &bSuccess, &ItemKey, &ExpectedValues]()
+			{
+				if (bSuccess)
+				{
+					// In a fold expression, ExpectedValues will be a single value. We know its underlying type by decaying the type of the value (using decltype)
+					using AttributeType = typename std::decay_t<decltype(ExpectedValues)>;
+
+					// We use Index to track the current attribute to check
+					const FName AttributeName = AttributeNames[Index];
+					const FPCGMetadataAttributeBase* Attribute = Metadata->GetConstAttribute(AttributeName);
+
+					bSuccess = TestInstance->TestNotNull(FormatWithPropertyNameAndSubNames("Attribute should not be null", AttributeName, ItemKey), Attribute);
+
+					if (bSuccess)
+					{
+						bSuccess = TestInstance->TestEqual(FormatWithPropertyNameAndSubNames("Attribute type is matching", AttributeName, ItemKey), Attribute->GetTypeId(), PCG::Private::MetadataTypes<AttributeType>::Id);
+					}
+
+					if (bSuccess)
+					{
+						AttributeType Value = static_cast<const FPCGMetadataAttribute<AttributeType>*>(Attribute)->GetValueFromItemKey(ItemKey);
+						bSuccess = TestInstance->TestEqual(FormatWithPropertyNameAndSubNames("Attribute value is the same as expected", AttributeName, ItemKey), Value, ExpectedValues);
+					}
+
+					// At each iteration on the ExpectedValues, we increment the index and if we reached the end of the Attribute Names we reset it to 0
+					// and we increment ItemKey, to check the next value in the attribute value array.
+					if (++Index >= AttributeNames.Num())
+					{
+						Index = 0;
+						ItemKey++;
+					}
+				}
+			}(), ...);
+		}
+	}
+
+	// Cleanup behind ourselves.
+	TestData.TestPCGComponent->GetGraph()->RemoveNode(TrivialNode);
+	TestData.TestPCGComponent->GetGraph()->RemoveNode(TestNode);
+
+	return bSuccess;
+}
+
+// Use aliases for convenience
+
+// Single valid attribute, Property Name will be attribute name
+template <typename AttributeType>
+bool VerifyAttributeValueValid(FPCGTestBaseClass* TestInstance, PCGTestsCommon::FTestData& TestData, FName PropertyName, AttributeType&& ExpectedValue, FString ExtraTestWhat)
+{
+	return VerifyAttributeValue</*bShouldFail=*/false>(TestInstance, TestData, PropertyName, { PropertyName }, ExtraTestWhat, std::forward<AttributeType>(ExpectedValue));
+}
+
+// Single invalid attribute, Property Name will be attribute name
+template <typename AttributeType>
+bool VerifyAttributeValueInvalid(FPCGTestBaseClass* TestInstance, PCGTestsCommon::FTestData& TestData, FName PropertyName, AttributeType&& ExpectedValue, FString ExtraTestWhat)
+{
+	return VerifyAttributeValue</*bShouldFail=*/true>(TestInstance, TestData, PropertyName, { PropertyName }, ExtraTestWhat, std::forward<AttributeType>(ExpectedValue));
+}
+
+// Valid Single or multi attributes with single or multiple values.
+template <typename ...AttributeTypes>
+bool VerifyAttributeValuesValid(FPCGTestBaseClass* TestInstance, PCGTestsCommon::FTestData& TestData, FName PropertyName, const TArray<FName>& AttributeNames, FString ExtraTestWhat, AttributeTypes&& ...ExpectedValues)
+{
+	return VerifyAttributeValue</*bShouldFail=*/false>(TestInstance, TestData, PropertyName, AttributeNames, ExtraTestWhat, std::forward<AttributeTypes>(ExpectedValues)...);
+}
+
+bool FPCGPropertyToParamDataPropertyTypeTest::RunTest(const FString& Parameters)
+{
+	bool bSuccess = true;
+
+	UPCGGetActorPropertySettings* Settings = NewObject<UPCGGetActorPropertySettings>();
+	Settings->ActorSelector.ActorSelection = EPCGActorSelection::ByClass;
+	Settings->ActorSelector.ActorSelectionClass = APCGUnitTestDummyActor::StaticClass();
+	Settings->ActorSelector.ActorFilter = EPCGActorFilter::Self;
+
+	static constexpr int32 Seed = 42;
+	const FString ExtraTestWhat = "PropertyToParamDataPropertyTypeTest";
+
+	PCGTestsCommon::FTestData TestData(Seed, Settings, APCGUnitTestDummyActor::StaticClass());
+
+	// Set all properties
+	const FName NameValue = TEXT("HelloWorld");
+	const FString StringValue = TEXT("HelloWorld");
+
+	const FVector VectorValue{ 1.0, 2.0, 3.0 };
+	const FVector SecondVectorValue{ 4.0, 5.0, 6.0 };
+	const FVector4 Vector4Value{ 1.0, 2.0, 3.0, 4.0 };
+	const FRotator RotatorValue{ 45.0, 45.0, 45.0 };
+	const FQuat QuatValue = RotatorValue.Quaternion();
+	const FTransform TransformValue{ QuatValue, VectorValue, VectorValue };
+
+	UPCGDummyGetPropertyTest* ObjectValue = NewObject<UPCGDummyGetPropertyTest>();
+	ObjectValue->SetFlags(RF_Transient);
+	ObjectValue->Int64Property = 42ll;
+	ObjectValue->DoubleProperty = 1.0;
+
+	UPCGDummyGetPropertyTest* SecondObjectValue = NewObject<UPCGDummyGetPropertyTest>();
+	SecondObjectValue->SetFlags(RF_Transient);
+	SecondObjectValue->Int64Property = 43ll;
+	SecondObjectValue->DoubleProperty = 2.0;
+
+	const FSoftObjectPath SoftObjectPathValue{ ObjectValue };
+	const FSoftClassPath SoftClassPathValue{ UPCGDummyGetPropertyTest::StaticClass() };
+
+	const FVector2D Vector2Value = { 1.0, 2.0 };
+	const FColor ColorValue = FColor::Cyan;
+	const FColor SecondColorValue = FColor::Yellow;
+
+	APCGUnitTestDummyActor* Actor = Cast<APCGUnitTestDummyActor>(TestData.TestActor);
+	Actor->IntProperty = 42;
+	Actor->Int64Property = 42ll;
+	Actor->FloatProperty = 1.0f;
+	Actor->DoubleProperty = 1.0;
+	Actor->BoolProperty = true;
+	Actor->NameProperty = NameValue;
+	Actor->StringProperty = StringValue;
+	Actor->EnumProperty = EPCGUnitTestDummyEnum::Three;
+	Actor->VectorProperty = VectorValue;
+	Actor->Vector4Property = Vector4Value;
+	Actor->RotatorProperty = RotatorValue;
+	Actor->QuatProperty = QuatValue;
+	Actor->TransformProperty = TransformValue;
+	Actor->SoftObjectPathProperty = SoftObjectPathValue;
+	Actor->SoftClassPathProperty = SoftClassPathValue;
+	Actor->ClassProperty = UPCGDummyGetPropertyTest::StaticClass();
+	Actor->ObjectProperty = ObjectValue;
+	Actor->Vector2Property = Vector2Value;
+	Actor->ColorProperty = ColorValue;
+	Actor->ArrayOfIntsProperty = { 42, 43, 44 };
+	Actor->ArrayOfVectorsProperty = { VectorValue, SecondVectorValue };
+	Actor->ArrayOfStructsProperty = { ColorValue, SecondColorValue };
+	Actor->ArrayOfObjectsProperty = { ObjectValue, SecondObjectValue };
+	Actor->DummyStruct.FloatProperty = 1.2f;
+	Actor->DummyStruct.IntArrayProperty = { 5, 6, 7 };
+	Actor->DummyStruct.Level2Struct.DoubleArrayProperty = { 0.1, 0.2, 0.3 };
+
+	// Basic properties
+	bSuccess &= VerifyAttributeValueValid(this, TestData, GET_MEMBER_NAME_CHECKED(APCGUnitTestDummyActor, IntProperty), 42ll, ExtraTestWhat);
+	bSuccess &= VerifyAttributeValueValid(this, TestData, GET_MEMBER_NAME_CHECKED(APCGUnitTestDummyActor, Int64Property), 42ll, ExtraTestWhat);
+	bSuccess &= VerifyAttributeValueValid(this, TestData, GET_MEMBER_NAME_CHECKED(APCGUnitTestDummyActor, FloatProperty), 1.0, ExtraTestWhat);
+	bSuccess &= VerifyAttributeValueValid(this, TestData, GET_MEMBER_NAME_CHECKED(APCGUnitTestDummyActor, DoubleProperty), 1.0, ExtraTestWhat);
+	bSuccess &= VerifyAttributeValueValid(this, TestData, GET_MEMBER_NAME_CHECKED(APCGUnitTestDummyActor, BoolProperty), true, ExtraTestWhat);
+
+	// String/Name Properties
+	bSuccess &= VerifyAttributeValueValid(this, TestData, GET_MEMBER_NAME_CHECKED(APCGUnitTestDummyActor, NameProperty), NameValue, ExtraTestWhat);
+	bSuccess &= VerifyAttributeValueValid(this, TestData, GET_MEMBER_NAME_CHECKED(APCGUnitTestDummyActor, StringProperty), StringValue, ExtraTestWhat);
+
+	// Enum Property
+	bSuccess &= VerifyAttributeValueValid(this, TestData, GET_MEMBER_NAME_CHECKED(APCGUnitTestDummyActor, EnumProperty), (int64)EPCGUnitTestDummyEnum::Three, ExtraTestWhat);
+
+	// Supported struct Properties
+	bSuccess &= VerifyAttributeValueValid(this, TestData, GET_MEMBER_NAME_CHECKED(APCGUnitTestDummyActor, Vector2Property), Vector2Value, ExtraTestWhat);
+	bSuccess &= VerifyAttributeValueValid(this, TestData, GET_MEMBER_NAME_CHECKED(APCGUnitTestDummyActor, VectorProperty), VectorValue, ExtraTestWhat);
+	bSuccess &= VerifyAttributeValueValid(this, TestData, GET_MEMBER_NAME_CHECKED(APCGUnitTestDummyActor, Vector4Property), Vector4Value, ExtraTestWhat);
+	bSuccess &= VerifyAttributeValueValid(this, TestData, GET_MEMBER_NAME_CHECKED(APCGUnitTestDummyActor, RotatorProperty), RotatorValue, ExtraTestWhat);
+	bSuccess &= VerifyAttributeValueValid(this, TestData, GET_MEMBER_NAME_CHECKED(APCGUnitTestDummyActor, QuatProperty), QuatValue, ExtraTestWhat);
+	bSuccess &= VerifyAttributeValueValid(this, TestData, GET_MEMBER_NAME_CHECKED(APCGUnitTestDummyActor, TransformProperty), TransformValue, ExtraTestWhat);
+	bSuccess &= VerifyAttributeValueValid(this, TestData, GET_MEMBER_NAME_CHECKED(APCGUnitTestDummyActor, SoftObjectPathProperty), SoftObjectPathValue, ExtraTestWhat);
+	bSuccess &= VerifyAttributeValueValid(this, TestData, GET_MEMBER_NAME_CHECKED(APCGUnitTestDummyActor, SoftClassPathProperty), SoftClassPathValue, ExtraTestWhat);
+
+	// Objects properties as String
+	bSuccess &= VerifyAttributeValueValid(this, TestData, GET_MEMBER_NAME_CHECKED(APCGUnitTestDummyActor, ClassProperty), UPCGDummyGetPropertyTest::StaticClass()->GetPathName(), ExtraTestWhat);
+	bSuccess &= VerifyAttributeValueValid(this, TestData, GET_MEMBER_NAME_CHECKED(APCGUnitTestDummyActor, ObjectProperty), ObjectValue->GetPathName(), ExtraTestWhat);
+
+	// Struct Property Extracted - Colors
+	// Extracting int will always yield a int64 and extracting floats with yield doubles. Here color is u8, so cast all of them to int64
+	const TArray<FName> ColorPropertyNames = { TEXT("R"), TEXT("G"), TEXT("B"), TEXT("A") };
+	bSuccess &= VerifyAttributeValuesValid(this, TestData, GET_MEMBER_NAME_CHECKED(APCGUnitTestDummyActor, ColorProperty), ColorPropertyNames, ExtraTestWhat, (int64)ColorValue.R, (int64)ColorValue.G, (int64)ColorValue.B, (int64)ColorValue.A);
+
+	// Object Property Extracted - UPCGDummyGetPropertyTest
+	const TArray<FName> ObjectPropertyNames = { GET_MEMBER_NAME_CHECKED(UPCGDummyGetPropertyTest, Int64Property), GET_MEMBER_NAME_CHECKED(UPCGDummyGetPropertyTest, DoubleProperty) };
+	bSuccess &= VerifyAttributeValuesValid(this, TestData, GET_MEMBER_NAME_CHECKED(APCGUnitTestDummyActor, ObjectProperty), ObjectPropertyNames, ExtraTestWhat, 42ll, 1.0);
+
+	// Arrays of supported properties
+	bSuccess &= VerifyAttributeValuesValid(this, TestData, GET_MEMBER_NAME_CHECKED(APCGUnitTestDummyActor, ArrayOfIntsProperty), { GET_MEMBER_NAME_CHECKED(APCGUnitTestDummyActor, ArrayOfIntsProperty) }, ExtraTestWhat, 42ll, 43ll, 44ll);
+	bSuccess &= VerifyAttributeValuesValid(this, TestData, GET_MEMBER_NAME_CHECKED(APCGUnitTestDummyActor, ArrayOfVectorsProperty), { GET_MEMBER_NAME_CHECKED(APCGUnitTestDummyActor, ArrayOfVectorsProperty) }, ExtraTestWhat, VectorValue, SecondVectorValue);
+	bSuccess &= VerifyAttributeValuesValid(this, TestData, GET_MEMBER_NAME_CHECKED(APCGUnitTestDummyActor, ArrayOfObjectsProperty), { GET_MEMBER_NAME_CHECKED(APCGUnitTestDummyActor, ArrayOfObjectsProperty) }, ExtraTestWhat, ObjectValue->GetPathName(), SecondObjectValue->GetPathName());
+
+	// Arrays of extracted properties
+	bSuccess &= VerifyAttributeValuesValid(this, TestData, GET_MEMBER_NAME_CHECKED(APCGUnitTestDummyActor, ArrayOfStructsProperty), ColorPropertyNames, ExtraTestWhat, 
+		(int64)ColorValue.R, (int64)ColorValue.G, (int64)ColorValue.B, (int64)ColorValue.A, (int64)SecondColorValue.R, (int64)SecondColorValue.G, (int64)SecondColorValue.B, (int64)SecondColorValue.A);
+	bSuccess &= VerifyAttributeValuesValid(this, TestData, GET_MEMBER_NAME_CHECKED(APCGUnitTestDummyActor, ArrayOfObjectsProperty), ObjectPropertyNames, ExtraTestWhat, ObjectValue->Int64Property, ObjectValue->DoubleProperty, SecondObjectValue->Int64Property, SecondObjectValue->DoubleProperty);
+
+	// Extractors
+	bSuccess &= VerifyAttributeValuesValid(this, TestData, TEXT("DummyStruct.FloatProperty"), { GET_MEMBER_NAME_CHECKED(FPCGDummyGetPropertyStruct, FloatProperty) }, ExtraTestWhat, 1.2);
+	bSuccess &= VerifyAttributeValuesValid(this, TestData, TEXT("DummyStruct.IntArrayProperty"), { GET_MEMBER_NAME_CHECKED(FPCGDummyGetPropertyStruct, IntArrayProperty) }, ExtraTestWhat, 5ll, 6ll, 7ll);
+	bSuccess &= VerifyAttributeValuesValid(this, TestData, TEXT("DummyStruct.Level2Struct.DoubleArrayProperty"), { GET_MEMBER_NAME_CHECKED(FPCGDummyGetPropertyLevel2Struct, DoubleArrayProperty) }, ExtraTestWhat, 0.1, 0.2, 0.3);
+
+	// Extracting the DummyStruct should only extract the float, as arrays and deeper structs are discarded
+	bSuccess &= VerifyAttributeValuesValid(this, TestData, TEXT("DummyStruct"), { GET_MEMBER_NAME_CHECKED(FPCGDummyGetPropertyStruct, FloatProperty) }, ExtraTestWhat, 1.2);
+
+	// Unknown property
+	AddExpectedError(TEXT("Property 'DummyMissingProperty' does not exist"), EAutomationExpectedErrorFlags::Contains, 1);
+	AddExpectedError(TEXT("Fail to extract the property 'DummyMissingProperty' on actor"), EAutomationExpectedErrorFlags::Contains, 1);
+	bSuccess &= VerifyAttributeValueInvalid(this, TestData, TEXT("DummyMissingProperty"), 42, ExtraTestWhat);
+
+	// Missing property
+	AddExpectedError(TEXT("Some parameters are missing, abort."), EAutomationExpectedErrorFlags::Contains, 1);
+	bSuccess &= VerifyAttributeValueInvalid(this, TestData, NAME_None, 42, ExtraTestWhat);
+
+	ObjectValue->MarkAsGarbage();
+	SecondObjectValue->MarkAsGarbage();
+
+	return bSuccess;
+}
+
+bool FPCGPropertyToParamDataActorFindTest::RunTest(const FString& Parameters)
+{
+	bool bSuccess = true;
+
+	static constexpr int32 Seed = 42;
+	static const FName Tag = TEXT("PCGUnitTestDummyMyTag");
+	static const FString Name = TEXT("MyPCGUnitTestDummyActor");
+
+	UPCGGetActorPropertySettings* Settings = NewObject<UPCGGetActorPropertySettings>();
+	Settings->ActorSelector.ActorSelectionClass = APCGUnitTestDummyActor::StaticClass();
+	Settings->ActorSelector.ActorSelectionTag = Tag;
+	Settings->ComponentClass = UPCGUnitTestDummyComponent::StaticClass();
+
+	// Self by class
+	{
+		Settings->ActorSelector.ActorSelection = EPCGActorSelection::ByClass;
+		Settings->ActorSelector.ActorFilter = EPCGActorFilter::Self;
+
+		PCGTestsCommon::FTestData TestData(Seed, Settings, APCGUnitTestDummyActor::StaticClass());
+		Cast<APCGUnitTestDummyActor>(TestData.TestActor)->IntProperty = 42;
+
+		bSuccess &= VerifyAttributeValueValid(this, TestData, GET_MEMBER_NAME_CHECKED(APCGUnitTestDummyActor, IntProperty), 42ll, "PropertyToParamDataActorFindTest_Self_Class");
+	}
+
+	// Self by tag
+	{
+		Settings->ActorSelector.ActorSelection = EPCGActorSelection::ByTag;
+		Settings->ActorSelector.ActorFilter = EPCGActorFilter::Self;
+
+		PCGTestsCommon::FTestData TestData(Seed, Settings, APCGUnitTestDummyActor::StaticClass());
+		TestData.TestActor->Tags.Add(Tag);
+		Cast<APCGUnitTestDummyActor>(TestData.TestActor)->IntProperty = 42;
+
+		bSuccess &= VerifyAttributeValueValid(this, TestData, GET_MEMBER_NAME_CHECKED(APCGUnitTestDummyActor, IntProperty), 42ll, "PropertyToParamDataActorFindTest_Self_Tag");
+	}
+
+	// TODO: Need a good way to spawn actors with parenting relation between them
+	//// Parent by class
+	//{
+	//	Settings->ActorSelection = EPCGActorSelection::ByClass;
+	//	Settings->ActorFilter = EPCGActorFilter::Parent;
+
+	//	PCGTestsCommon::FTestData TestData(Seed, Settings, AActor::StaticClass());
+	//	TestData.AddActor(APCGUnitTestDummyActor::StaticClass(), true);
+	//	Cast<APCGUnitTestDummyActor>(TestData.TestActor)->IntProperty = 42;
+
+	//	bSuccess &= VerifyAttributeValueValid(this, TestData, GET_MEMBER_NAME_CHECKED(APCGUnitTestDummyActor, IntProperty), 42ll, "PropertyToParamDataActorFindTest_Parent_Class");
+	//}
+
+	//// Root by tag
+	//{
+	//	Settings->ActorSelection = EPCGActorSelection::ByTag;
+	//	Settings->ActorFilter = EPCGActorFilter::Root;
+
+	//	PCGTestsCommon::FTestData TestData(Seed, Settings, AActor::StaticClass());
+	//	TestData.AddActor(AActor::StaticClass(), true);
+	//	TestData.AddActor(APCGUnitTestDummyActor::StaticClass(), true);
+	//	TestData.TestActor->Tags.Add(Tag);
+	//	Cast<APCGUnitTestDummyActor>(TestData.TestActor)->IntProperty = 42;
+
+	//	bSuccess &= VerifyAttributeValueValid(this, TestData, GET_MEMBER_NAME_CHECKED(APCGUnitTestDummyActor, IntProperty), 42ll, "PropertyToParamDataActorFindTest_Root_Tag");
+	//}
+
+	//// Root by tag including children
+	//{
+	//	Settings->ActorSelection = EPCGActorSelection::ByTag;
+	//	Settings->ActorFilter = EPCGActorFilter::Root;
+	//	Settings->bIncludeChildren = true;
+
+	//	PCGTestsCommon::FTestData TestData(Seed, Settings, AActor::StaticClass());
+	//	TestData.AddActor(APCGUnitTestDummyActor::StaticClass(), true);
+	//	TestData.TestActor->Tags.Add(Tag);
+	//	Cast<APCGUnitTestDummyActor>(TestData.TestActor)->IntProperty = 42;
+
+	//	TestData.AddActor(AActor::StaticClass(), true);
+
+	//	bSuccess &= VerifyAttributeValueValid(this, TestData, GET_MEMBER_NAME_CHECKED(APCGUnitTestDummyActor, IntProperty), 42ll, "PropertyToParamDataActorFindTest_Root_Tag_Children");
+	//	Settings->bIncludeChildren = false;
+	//}
+
+	// Self by class including component
+	{
+		Settings->ActorSelector.ActorSelection = EPCGActorSelection::ByTag;
+		Settings->ActorSelector.ActorFilter = EPCGActorFilter::Self;
+		Settings->bSelectComponent = true;
+		
+		PCGTestsCommon::FTestData TestData(Seed, Settings, AActor::StaticClass());
+		TestData.TestActor->Tags.Add(Tag);
+		UPCGUnitTestDummyComponent* Component = Cast<UPCGUnitTestDummyComponent>(TestData.TestActor->AddComponentByClass(Settings->ComponentClass, false, FTransform::Identity, false));
+		Component->IntProperty = 42;
+
+		bSuccess &= VerifyAttributeValueValid(this, TestData, GET_MEMBER_NAME_CHECKED(UPCGUnitTestDummyComponent, IntProperty), 42ll, "PropertyToParamDataActorFindTest_Self_Tag_Component");
+		Settings->bSelectComponent = false;
+	}
+
+	// World by class
+	{
+		Settings->ActorSelector.ActorSelection = EPCGActorSelection::ByClass;
+		Settings->ActorSelector.ActorFilter = EPCGActorFilter::AllWorldActors;
+
+		PCGTestsCommon::FTestData TestData(Seed, Settings, APCGUnitTestDummyActor::StaticClass());
+		Cast<APCGUnitTestDummyActor>(TestData.TestActor)->IntProperty = 42;
+
+		bSuccess &= VerifyAttributeValueValid(this, TestData, GET_MEMBER_NAME_CHECKED(APCGUnitTestDummyActor, IntProperty), 42ll, "PropertyToParamDataActorFindTest_World_Class");
+	}
+
+	// World by tag
+	{
+		Settings->ActorSelector.ActorSelection = EPCGActorSelection::ByTag;
+		Settings->ActorSelector.ActorFilter = EPCGActorFilter::AllWorldActors;
+
+		PCGTestsCommon::FTestData TestData(Seed, Settings, APCGUnitTestDummyActor::StaticClass());
+		TestData.TestActor->Tags.Add(Tag);
+		Cast<APCGUnitTestDummyActor>(TestData.TestActor)->IntProperty = 42;
+
+		bSuccess &= VerifyAttributeValueValid(this, TestData, GET_MEMBER_NAME_CHECKED(APCGUnitTestDummyActor, IntProperty), 42ll, "PropertyToParamDataActorFindTest_World_Tag");
+	}
+
+	return bSuccess;
+}
+
+#endif // WITH_EDITOR
diff --git a/Engine/Plugins/Experimental/PCG/Source/PCG/Private/Tests/Elements/PCGPropertyToParamDataElementTest.cpp b/Engine/Plugins/Experimental/PCG/Source/PCG/Private/Tests/Elements/PCGPropertyToParamDataElementTest.cpp
deleted file mode 100644
diff --git a/Engine/Plugins/Experimental/PCG/Source/PCG/Public/Elements/PCGGetActorProperty.h b/Engine/Plugins/Experimental/PCG/Source/PCG/Public/Elements/PCGGetActorProperty.h
new file mode 100644
--- /dev/null
+++ b/Engine/Plugins/Experimental/PCG/Source/PCG/Public/Elements/PCGGetActorProperty.h
@@ -0,0 +1,106 @@
+// Copyright Epic Games, Inc. All Rights Reserved.
+
+#pragma once
+
+#include "PCGSettings.h"
+#include "Elements/PCGActorSelector.h"
+#include "Metadata/PCGMetadataAttribute.h"
+
+#include "PCGGetActorProperty.generated.h"
+
+class UActorComponent;
+
+/**
+* Extract a property value from an actor/component into a ParamData.
+*/
+UCLASS(BlueprintType, ClassGroup = (Procedural))
+class PCG_API UPCGGetActorPropertySettings : public UPCGSettings
+{
+	GENERATED_BODY()
+
+public:
+	//~Begin UObject interface
+	virtual void PostLoad() override;
+#if WITH_EDITOR
+	virtual void PostEditChangeProperty(FPropertyChangedEvent& PropertyChangedEvent) override;
+#endif
+	//~End UObject interface
+
+	//~Begin UPCGSettings interface
+#if WITH_EDITOR
+	virtual FName GetDefaultNodeName() const override { return FName(TEXT("GetActorProperty")); }
+	virtual FText GetDefaultNodeTitle() const override { return NSLOCTEXT("PCGPropertyToParamDataSettings", "NodeTitle", "Get Actor Property"); }
+	virtual EPCGSettingsType GetType() const override { return EPCGSettingsType::Param; }
+	virtual void GetTrackedActorKeys(FPCGActorSelectionKeyToSettingsMap& OutKeysToSettings, TArray<TObjectPtr<const UPCGGraph>>& OutVisitedGraphs) const override;
+#endif
+
+	virtual FName AdditionalTaskName() const override;
+	virtual TArray<FPCGPinProperties> OutputPinProperties() const override;
+	virtual TArray<FPCGPinProperties> InputPinProperties() const override { return TArray<FPCGPinProperties>(); }
+
+protected:
+	virtual FPCGElementPtr CreateElement() const override;
+	//~End UPCGSettings interface
+
+public:
+	UPROPERTY(BlueprintReadWrite, EditAnywhere, Category = Settings, meta = (ShowOnlyInnerProperties))
+	FPCGActorSelectorSettings ActorSelector;
+
+	/** Allow to look for an actor component instead of an actor. It will need to be attached to the found actor. */
+	UPROPERTY(BlueprintReadWrite, EditAnywhere, Category = Settings)
+	bool bSelectComponent = false;
+
+	/** If we are looking for an actor component, the class can be specified here. */
+	UPROPERTY(BlueprintReadWrite, EditAnywhere, Category = Settings, meta = (EditCondition = "bSelectComponent", EditConditionHides))
+	TSubclassOf<UActorComponent> ComponentClass;
+
+	/** Property name to extract. Can only extract properties that are compatible with metadata types. If None, extract the actor/component directly.*/
+	UPROPERTY(BlueprintReadWrite, EditAnywhere, Category = Settings)
+	FName PropertyName = NAME_None;
+
+	/** If the property is a struct/object supported by metadata, this option can be toggled to force extracting all (compatible) properties contained in this property. Automatically true if unsupported by metadata. For now, only supports direct child properties (and not deeper). */
+	UPROPERTY(BlueprintReadWrite, EditAnywhere, Category = Settings)
+	bool bForceObjectAndStructExtraction = false;
+
+	/** By default, attribute name will be None, but it can be overridden by this name. Use @SourceName to use the property name (only works when not extracting). */
+	UPROPERTY(BlueprintReadWrite, EditAnywhere, Category = Settings, meta = (EditCondition = "!bExtractObjectAndStruct", EditConditionHides))
+	FName OutputAttributeName = NAME_None;
+
+	/** If this is true, we will never put this element in cache, and will always try to re-query the actors and read the latest properties from them. */
+	UPROPERTY(BlueprintReadWrite, EditAnywhere, Category = Settings)
+	bool bAlwaysRequeryActors = true;
+
+#if WITH_EDITORONLY_DATA
+	/** If this is checked, found actors that are outside component bounds will not trigger a refresh. Only works for tags for now in editor. */
+	UPROPERTY(BlueprintReadOnly, EditAnywhere, Category = Settings)
+	bool bTrackActorsOnlyWithinBounds = false;
+#endif // WITH_EDITORONLY_DATA
+
+private:
+	UPROPERTY()
+	EPCGActorSelection ActorSelection_DEPRECATED;
+
+	UPROPERTY()
+	FName ActorSelectionTag_DEPRECATED;
+
+	UPROPERTY()
+	FName ActorSelectionName_DEPRECATED;
+
+	UPROPERTY()
+	TSubclassOf<AActor> ActorSelectionClass_DEPRECATED;
+
+	UPROPERTY()
+	EPCGActorFilter ActorFilter_DEPRECATED = EPCGActorFilter::Self;
+
+	UPROPERTY()
+	bool bIncludeChildren_DEPRECATED = false;
+};
+
+class FPCGGetActorPropertyElement : public IPCGElement
+{
+public:
+	virtual bool CanExecuteOnlyOnMainThread(FPCGContext* Context) const override { return true; }
+	virtual bool IsCacheable(const UPCGSettings* InSettings) const override { return !CastChecked<UPCGGetActorPropertySettings>(InSettings)->bAlwaysRequeryActors; }
+protected:
+	virtual bool ExecuteInternal(FPCGContext* Context) const override;
+};
diff --git a/Engine/Plugins/Experimental/PCG/Source/PCG/Public/Elements/PCGGetPropertyFromObjectPath.h b/Engine/Plugins/Experimental/PCG/Source/PCG/Public/Elements/PCGGetPropertyFromObjectPath.h
new file mode 100644
--- /dev/null
+++ b/Engine/Plugins/Experimental/PCG/Source/PCG/Public/Elements/PCGGetPropertyFromObjectPath.h
@@ -0,0 +1,76 @@
+// Copyright Epic Games, Inc. All Rights Reserved.
+
+#pragma once
+
+#include "PCGContext.h"
+#include "PCGSettings.h"
+
+#include "Engine/StreamableManager.h"
+
+#include "PCGGetPropertyFromObjectPath.generated.h"
+
+/**
+* Extract property from a list of soft object paths.
+*/
+UCLASS(BlueprintType, ClassGroup = (Procedural))
+class PCG_API UPCGGetPropertyFromObjectPathSettings : public UPCGSettings
+{
+	GENERATED_BODY()
+
+public:
+	//~Begin UPCGSettings interface
+#if WITH_EDITOR
+	virtual FName GetDefaultNodeName() const override;
+	virtual FText GetDefaultNodeTitle() const override;
+	virtual EPCGSettingsType GetType() const override { return EPCGSettingsType::Param; }
+	virtual bool CanEditChange(const FProperty* InProperty) const override;
+	virtual bool IsPinUsedByNodeExecution(const UPCGPin* InPin) const override;
+#endif
+
+protected:
+	virtual TArray<FPCGPinProperties> InputPinProperties() const override;
+	virtual TArray<FPCGPinProperties> OutputPinProperties() const override;
+	virtual FPCGElementPtr CreateElement() const override;
+	//~End UPCGSettings interface
+
+public:
+	/** If nothing is connected in the In pin, will use those static paths to load.*/
+	UPROPERTY(BlueprintReadWrite, EditAnywhere, Category = "Settings")
+	TArray<FSoftObjectPath> ObjectPathsToExtract;
+
+	/** If something is connected in the In pin, will look for this attribute values to load.*/
+	UPROPERTY(BlueprintReadWrite, EditAnywhere, Category = "Settings", meta = (PCG_Overridable, PCG_DiscardPropertySelection, PCG_DiscardExtraSelection))
+	FPCGAttributePropertyInputSelector InputSource;
+
+	/** Property name to extract. Can only extract properties that are compatible with metadata types. If None, extract the object. */
+	UPROPERTY(BlueprintReadWrite, EditAnywhere, Category = Settings)
+	FName PropertyName = NAME_None;
+
+	/** If the property is a struct/object supported by metadata, this option can be toggled to force extracting all (compatible) properties contained in this property. Automatically true if unsupported by metadata. For now, only supports direct child properties (and not deeper). */
+	UPROPERTY(BlueprintReadWrite, EditAnywhere, Category = Settings)
+	bool bForceObjectAndStructExtraction = false;
+
+	/** By default, attribute name will be None, but it can be overridden by this name. Use @SourceName to use the property name (only works when not extracting). */
+	UPROPERTY(BlueprintReadWrite, EditAnywhere, Category = Settings, meta = (EditCondition = "!bExtractObjectAndStruct", EditConditionHides))
+	FName OutputAttributeName = NAME_None;
+
+	/** By default, object loading is asynchronous, can force it synchronous if needed. */
+	UPROPERTY(BlueprintReadWrite, EditAnywhere, Category = "Settings|Debug")
+	bool bSynchronousLoad = false;
+};
+
+struct FPCGGetPropertyFromObjectPathContext : public FPCGContext
+{
+	TSharedPtr<FStreamableHandle> LoadHandle;
+	TArray<TTuple<FSoftObjectPath, int32>> PathsToObjectsToExtractAndIncomingDataIndex;
+	bool bRequestSent = false;
+};
+
+class FPCGGetPropertyFromObjectPathElement : public IPCGElement
+{
+protected:
+	virtual FPCGContext* CreateContext() override;
+	virtual bool PrepareDataInternal(FPCGContext* Context) const override;
+	virtual bool ExecuteInternal(FPCGContext* Context) const override;
+};
+
diff --git a/Engine/Plugins/Experimental/PCG/Source/PCG/Public/Elements/PCGPropertyToParamData.h b/Engine/Plugins/Experimental/PCG/Source/PCG/Public/Elements/PCGPropertyToParamData.h
deleted file mode 100644
diff --git a/Engine/Plugins/Experimental/PCG/Source/PCG/Public/Helpers/PCGSettingsHelpers.h b/Engine/Plugins/Experimental/PCG/Source/PCG/Public/Helpers/PCGSettingsHelpers.h
--- a/Engine/Plugins/Experimental/PCG/Source/PCG/Public/Helpers/PCGSettingsHelpers.h
+++ b/Engine/Plugins/Experimental/PCG/Source/PCG/Public/Helpers/PCGSettingsHelpers.h
@@ -240,10 +240,15 @@ PRAGMA_ENABLE_DEPRECATION_WARNINGS
 
 		// Max depth for structs of structs. -1 = no limit
 		int32 MaxStructDepth = -1;
+
+		// If you want to go through objects too
+		bool bExtractObjects = false;
+
+		// If you want to also extract arrays
+		bool bExtractArrays = false;
 	};
 
-	TArray<FPCGSettingsOverridableParam> GetAllOverridableParams(const UClass* InClass, const FPCGGetAllOverridableParamsConfig& InConfig);
-	TArray<FPCGSettingsOverridableParam> GetAllOverridableParams(const UScriptStruct* InStruct, const FPCGGetAllOverridableParamsConfig& InConfig);
+	TArray<FPCGSettingsOverridableParam> GetAllOverridableParams(const UStruct* InClass, const FPCGGetAllOverridableParamsConfig& InConfig);
 }
 
 // Deprecated macro, not necessary anymore. Cf. GetValue
