[PCG] Add double-clicking row in the DebugObjectTreeView to auto-expand and select the first entry as the debug object

diff --git a/Engine/Plugins/Experimental/PCG/Source/PCGEditor/Private/SPCGEditorGraphDebugObjectTree.cpp b/Engine/Plugins/Experimental/PCG/Source/PCGEditor/Private/SPCGEditorGraphDebugObjectTree.cpp
--- a/Engine/Plugins/Experimental/PCG/Source/PCGEditor/Private/SPCGEditorGraphDebugObjectTree.cpp
+++ b/Engine/Plugins/Experimental/PCG/Source/PCGEditor/Private/SPCGEditorGraphDebugObjectTree.cpp
@@ -99,6 +99,8 @@ FString FPCGEditorGraphDebugObjectItem_PCGLoopIndex::GetLabel() const
 void SPCGEditorGraphDebugObjectItemRow::Construct(const FArguments& InArgs, const TSharedRef<STableViewBase>& InOwnerTableView, FPCGEditorGraphDebugObjectItemPtr InItem)
 {
 	Item = InItem;
+	// This function should auto-expand the row and select the deepest entry as the debug object if it is unambiguous (the only entry at its level in the tree).
+	DoubleClickFunc = InArgs._OnDoubleClickFunc;
 
 	ChildSlot
 	[
@@ -203,6 +205,16 @@ void SPCGEditorGraphDebugObjectItemRow::Construct(const FArguments& InArgs, cons
 	];
 }
 
+FReply SPCGEditorGraphDebugObjectItemRow::OnMouseButtonDoubleClick(const FGeometry& InMyGeometry, const FPointerEvent& InMouseEvent)
+{
+	if (Item.IsValid() && DoubleClickFunc)
+	{
+		DoubleClickFunc(Item);
+	}
+
+	return FReply::Handled();
+}
+
 SPCGEditorGraphDebugObjectTree::~SPCGEditorGraphDebugObjectTree()
 {
 	FCoreUObjectDelegates::OnPreObjectPropertyChanged.RemoveAll(this);
@@ -825,11 +837,12 @@ UPCGGraph* SPCGEditorGraphDebugObjectTree::GetPCGGraph() const
 	return PCGEditorGraph ? PCGEditorGraph->GetPCGGraph() : nullptr;
 }
 
-TSharedRef<ITableRow> SPCGEditorGraphDebugObjectTree::MakeTreeRowWidget(FPCGEditorGraphDebugObjectItemPtr InItem, const TSharedRef<STableViewBase>& InOwnerTable) const
+TSharedRef<ITableRow> SPCGEditorGraphDebugObjectTree::MakeTreeRowWidget(FPCGEditorGraphDebugObjectItemPtr InItem, const TSharedRef<STableViewBase>& InOwnerTable)
 {
 	return SNew( STableRow<TSharedPtr<SPCGEditorGraphDebugObjectItemRow>>, InOwnerTable)
 		[
 			SNew(SPCGEditorGraphDebugObjectItemRow, InOwnerTable, InItem)
+				.OnDoubleClickFunc([this](const FPCGEditorGraphDebugObjectItemPtr& Item) { ExpandAndSelectDebugObject(Item); } )
 		];
 }
 
@@ -898,7 +911,7 @@ void SPCGEditorGraphDebugObjectTree::OnExpansionChanged(FPCGEditorGraphDebugObje
 	}
 }
 
-void SPCGEditorGraphDebugObjectTree::OnSetExpansionRecursive(FPCGEditorGraphDebugObjectItemPtr InItem, bool bInExpand)
+void SPCGEditorGraphDebugObjectTree::OnSetExpansionRecursive(FPCGEditorGraphDebugObjectItemPtr InItem, bool bInExpand) const
 {
 	if (!InItem.IsValid() || !DebugObjectTreeView.IsValid())
 	{
@@ -917,4 +930,41 @@ void SPCGEditorGraphDebugObjectTree::OnSetExpansionRecursive(FPCGEditorGraphDebu
 	}
 }
 
+void SPCGEditorGraphDebugObjectTree::ExpandAndSelectDebugObject(FPCGEditorGraphDebugObjectItemPtr InItem)
+{
+	if (!InItem.IsValid())
+	{
+		return;
+	}
+
+	// Expand this item in the tree view.
+	OnSetExpansionRecursive(InItem, /*bInExpand=*/true);
+
+	FPCGEditorGraphDebugObjectItemPtr Item = InItem;
+	int NumChildren = 1; // If we land on the deepest item already, we should just select it regardless of how many siblings it has.
+
+	// Find the deepest entry in this branch of the tree view.
+	while (true)
+	{
+		const TSet<TSharedPtr<FPCGEditorGraphDebugObjectItem>>& Children = Item->GetChildren();
+		TSet<FPCGEditorGraphDebugObjectItemPtr>::TConstIterator It = Children.CreateConstIterator();
+
+		if (!It)
+		{
+			break;
+		}
+		else
+		{
+			Item = *It;
+			NumChildren = Children.Num();
+		}
+	}
+
+	// Set the discovered item as the debug object if it is the only child.
+	if (NumChildren == 1)
+	{
+		OnSelectionChanged(Item, ESelectInfo::Direct);
+	}
+}
+
 #undef LOCTEXT_NAMESPACE
diff --git a/Engine/Plugins/Experimental/PCG/Source/PCGEditor/Private/SPCGEditorGraphDebugObjectTree.h b/Engine/Plugins/Experimental/PCG/Source/PCGEditor/Private/SPCGEditorGraphDebugObjectTree.h
--- a/Engine/Plugins/Experimental/PCG/Source/PCGEditor/Private/SPCGEditorGraphDebugObjectTree.h
+++ b/Engine/Plugins/Experimental/PCG/Source/PCGEditor/Private/SPCGEditorGraphDebugObjectTree.h
@@ -158,13 +158,27 @@ protected:
 class SPCGEditorGraphDebugObjectItemRow : public SCompoundWidget
 {
 public:
-	SLATE_BEGIN_ARGS(SPCGEditorGraphDebugObjectItemRow) {}
+	using FDoubleClickFunc = TFunction<void(FPCGEditorGraphDebugObjectItemPtr)>;
+
+	SLATE_BEGIN_ARGS(SPCGEditorGraphDebugObjectItemRow)
+		: _OnDoubleClickFunc()
+	{}
+
+		SLATE_ARGUMENT(FDoubleClickFunc, OnDoubleClickFunc)
+
 	SLATE_END_ARGS()
 
 	void Construct(const FArguments& InArgs, const TSharedRef<STableViewBase>& InOwnerTableView, FPCGEditorGraphDebugObjectItemPtr InItem);
 
+	//~Begin SWidget Interface
+	virtual FReply OnMouseButtonDoubleClick(const FGeometry& InMyGeometry, const FPointerEvent& InMouseEvent) override;
+	//~End SWidget Interface
+
 private:
 	FPCGEditorGraphDebugObjectItemPtr Item;
+	
+	/** Invoked when the user double clicks on the row. */
+	FDoubleClickFunc DoubleClickFunc;
 };
 
 class SPCGEditorGraphDebugObjectTree : public SCompoundWidget
@@ -204,11 +218,14 @@ private:
 
 	UPCGGraph* GetPCGGraph() const;
 
-	TSharedRef<ITableRow> MakeTreeRowWidget(FPCGEditorGraphDebugObjectItemPtr InItem, const TSharedRef<STableViewBase>& InOwnerTable) const;
+	TSharedRef<ITableRow> MakeTreeRowWidget(FPCGEditorGraphDebugObjectItemPtr InItem, const TSharedRef<STableViewBase>& InOwnerTable);
 	void OnGetChildren(FPCGEditorGraphDebugObjectItemPtr InItem, TArray<FPCGEditorGraphDebugObjectItemPtr>& OutChildren) const;
 	void OnSelectionChanged(FPCGEditorGraphDebugObjectItemPtr InItem, ESelectInfo::Type InSelectInfo);
 	void OnExpansionChanged(FPCGEditorGraphDebugObjectItemPtr InItem, bool bInIsExpanded);
-	void OnSetExpansionRecursive(FPCGEditorGraphDebugObjectItemPtr InItem, bool bInExpand);
+	void OnSetExpansionRecursive(FPCGEditorGraphDebugObjectItemPtr InItem, bool bInExpand) const;
+
+	/** Expand the given row and select the deepest entry as the debug object if it is unambiguous (the only entry at its level in the tree). */
+	void ExpandAndSelectDebugObject(FPCGEditorGraphDebugObjectItemPtr InItem);
 
 	TWeakPtr<FPCGEditor> PCGEditor;
 
