[PCG] Remaining feedback from inspection workflows review - simplify node visual log summary logic

#rnx

diff --git a/Engine/Plugins/Experimental/PCG/Source/PCG/Private/Utils/PCGNodeVisualLogs.cpp b/Engine/Plugins/Experimental/PCG/Source/PCG/Private/Utils/PCGNodeVisualLogs.cpp
--- a/Engine/Plugins/Experimental/PCG/Source/PCG/Private/Utils/PCGNodeVisualLogs.cpp
+++ b/Engine/Plugins/Experimental/PCG/Source/PCG/Private/Utils/PCGNodeVisualLogs.cpp
@@ -113,30 +113,30 @@ FText FPCGNodeVisualLogs::GetLogsSummaryText(const UPCGNode* InNode, ELogVerbosi
 		{
 			const FPCGPerNodeVisualLogs& NodeLogs = Entry.Value;
 
-			for (int LogIndex = 0; LogIndex < NodeLogs.Num() && LogCounter < MaxLogsInSummary; ++LogIndex)
+			for (int LogIndex = 0; LogIndex < NodeLogs.Num(); ++LogIndex)
 			{
 				++LogCounter;
 
-				const UPCGComponent* Component = Stack.GetRootComponent();
-				FText ActorName = (Component && Component->GetOwner()) ? FText::FromString(Component->GetOwner()->GetActorLabel()) : FText::FromString(TEXT("MissingComponent"));
-				
-				OutMinimumVerbosity = FMath::Min(OutMinimumVerbosity, NodeLogs[LogIndex].Verbosity);
-				const FText VerbosityText = NodeLogs[LogIndex].Verbosity == ELogVerbosity::Warning ? FText::FromString(TEXT("Warning")) : FText::FromString(TEXT("Error"));
+				if (LogCounter >= MaxLogsInSummary)
+				{
+					Summary = FText::Format(FText::FromString(TEXT("{0}\n...")), Summary);
+					return Summary;
+				}
 
 				if (LogCounter > 1)
 				{
 					Summary = FText::Format(FText::FromString(TEXT("{0}\n")), Summary);
 				}
 
+				const UPCGComponent* Component = Stack.GetRootComponent();
+				FText ActorName = (Component && Component->GetOwner()) ? FText::FromString(Component->GetOwner()->GetActorLabel()) : FText::FromString(TEXT("MissingComponent"));
+				
+				OutMinimumVerbosity = FMath::Min(OutMinimumVerbosity, NodeLogs[LogIndex].Verbosity);
+				const FText VerbosityText = NodeLogs[LogIndex].Verbosity == ELogVerbosity::Warning ? FText::FromString(TEXT("Warning")) : FText::FromString(TEXT("Error"));
+
 				Summary = FText::Format(LOCTEXT("NodeTooltipLogWithActor", "{0}[{1}] {2}: {3}"), Summary, ActorName, VerbosityText, NodeLogs[LogIndex].Message);
 			}
 		}
-
-		if (LogCounter >= MaxLogsInSummary)
-		{
-			Summary = FText::Format(FText::FromString(TEXT("{0}\n...")), Summary);
-			break;
-		}
 	}
 
 	return Summary;
@@ -152,12 +152,6 @@ FText FPCGNodeVisualLogs::GetLogsSummaryText(const FPCGStack& InBaseStack) const
 
 	for (const TPair<FPCGStack, FPCGPerNodeVisualLogs>& Entry : StackToLogs)
 	{
-		if (LogCounter >= MaxLogsInSummary)
-		{
-			ResultText = FText::Format(FText::FromString(TEXT("{0}\n...")), ResultText);
-			break;
-		}
-
 		if (!Entry.Key.BeginsWith(InBaseStack))
 		{
 			continue;
@@ -169,6 +163,12 @@ FText FPCGNodeVisualLogs::GetLogsSummaryText(const FPCGStack& InBaseStack) const
 		{
 			++LogCounter;
 
+			if (LogCounter > MaxLogsInSummary)
+			{
+				ResultText = FText::Format(FText::FromString(TEXT("{0}\n...")), ResultText);
+				return ResultText;
+			}
+
 			if (LogCounter > 1)
 			{
 				ResultText = FText::Format(FText::FromString(TEXT("{0}\n")), ResultText);
@@ -176,11 +176,6 @@ FText FPCGNodeVisualLogs::GetLogsSummaryText(const FPCGStack& InBaseStack) const
 
 			const FText MessageVerbosity = NodeLogs[i].Verbosity == ELogVerbosity::Warning ? FText::FromString(TEXT("Warning")) : FText::FromString(TEXT("Error"));
 			ResultText = FText::Format(LOCTEXT("NodeTooltipLog", "{0}{1}: {2}"), ResultText, /*i + 1, NodeLogs.Num(),*/ MessageVerbosity, NodeLogs[i].Message);
-
-			if (LogCounter >= MaxLogsInSummary)
-			{
-				break;
-			}
 		}
 	}
 
