PCG Unit tests crash fix

#rnx

diff --git a/Engine/Plugins/Experimental/PCG/Source/PCG/Private/PCGContext.cpp b/Engine/Plugins/Experimental/PCG/Source/PCG/Private/PCGContext.cpp
--- a/Engine/Plugins/Experimental/PCG/Source/PCG/Private/PCGContext.cpp
+++ b/Engine/Plugins/Experimental/PCG/Source/PCG/Private/PCGContext.cpp
@@ -278,7 +278,7 @@ void FPCGContext::LogVisual(ELogVerbosity::Type InVerbosity, const FText& InMess
 
 	if (UPCGSubsystem* Subsystem = UPCGSubsystem::GetInstance(SourceComponent->GetWorld()))
 	{
-		FPCGStack StackWithNode = *Stack;
+		FPCGStack StackWithNode = Stack ? FPCGStack(*Stack) : FPCGStack();
 		StackWithNode.PushFrame(Node);
 		Subsystem->GetNodeVisualLogsMutable().Log(StackWithNode, InVerbosity, InMessage);
 	}
@@ -293,7 +293,7 @@ bool FPCGContext::HasVisualLogs() const
 
 	if (const UPCGSubsystem* Subsystem = UPCGSubsystem::GetInstance(SourceComponent->GetWorld()))
 	{
-		FPCGStack StackWithNode = *Stack;
+		FPCGStack StackWithNode = Stack ? FPCGStack(*Stack) : FPCGStack();
 		StackWithNode.PushFrame(Node);
 		return Subsystem->GetNodeVisualLogs().HasLogs(StackWithNode);
 	}
