[PCG] Nativize SplitPoint node - Splits each point into two separate points, and sets bounds based on the position and axis of the cut
--Includes unit tests

diff --git a/Engine/Plugins/Experimental/PCG/Source/PCG/Private/Elements/PCGSplitPoints.cpp b/Engine/Plugins/Experimental/PCG/Source/PCG/Private/Elements/PCGSplitPoints.cpp
new file mode 100644
--- /dev/null
+++ b/Engine/Plugins/Experimental/PCG/Source/PCG/Private/Elements/PCGSplitPoints.cpp
@@ -0,0 +1,115 @@
+// Copyright Epic Games, Inc. All Rights Reserved.
+
+#include "Elements/PCGSplitPoints.h"
+
+#include "PCGContext.h"
+#include "PCGPin.h"
+#include "Data/PCGPointData.h"
+#include "Helpers/PCGAsync.h"
+#include "Metadata/PCGMetadataAccessor.h"
+
+#define LOCTEXT_NAMESPACE "PCGSplitPointsElement"
+
+namespace PCGSplitPointsConstants
+{
+	const FName OutputALabel = TEXT("Before Split");
+	const FName OutputBLabel = TEXT("After Split");
+}
+
+TArray<FPCGPinProperties> UPCGSplitPointsSettings::OutputPinProperties() const
+{
+	TArray<FPCGPinProperties> PinProperties;
+	PinProperties.Emplace(PCGSplitPointsConstants::OutputALabel, 
+		EPCGDataType::Point, 
+		/*bAllowMultipleConnections=*/true, 
+		/*bAllowMultipleData=*/true, 
+		LOCTEXT("PinATooltip", "The portion of each point before the split plane."));
+
+	PinProperties.Emplace(PCGSplitPointsConstants::OutputBLabel, 
+		EPCGDataType::Point, 
+		/*bAllowMultipleConnections=*/true, 
+		/*bAllowMultipleData=*/true, 
+		LOCTEXT("PinBTooltip", "The portion of each point after the split plane."));
+
+	return PinProperties;
+}
+
+FPCGElementPtr UPCGSplitPointsSettings::CreateElement() const
+{
+	return MakeShared<FPCGSplitPointsElement>();
+}
+
+bool FPCGSplitPointsElement::ExecuteInternal(FPCGContext* Context) const
+{
+	TRACE_CPUPROFILER_EVENT_SCOPE(FPCGSplitPointsElement::Execute);
+	check(Context);
+
+	const UPCGSplitPointsSettings* Settings = Context->GetInputSettings<UPCGSplitPointsSettings>();
+	check(Settings);
+
+	const TArray<FPCGTaggedData> Inputs = Context->InputData.GetInputs();
+	TArray<FPCGTaggedData>& Outputs = Context->OutputData.TaggedData;
+
+	const float SplitPosition = FMath::Clamp(Settings->SplitPosition, 0.0f, 1.0f);
+
+	for (int i = 0; i < Inputs.Num(); ++i)
+	{
+		const UPCGPointData* InputPointData = Cast<UPCGPointData>(Inputs[i].Data);
+		if (!InputPointData)
+		{
+			PCGE_LOG(Verbose, GraphAndLog, FText::Format(LOCTEXT("InvalidPointData", "Input {0} is not point data"), i));
+			continue;
+		}
+
+		const TArray<FPCGPoint>& InputPoints = InputPointData->GetPoints();
+
+		FPCGTaggedData& OutputA = Outputs.Add_GetRef(Inputs[i]);
+		UPCGPointData* OutPointDataA = NewObject<UPCGPointData>();
+		OutPointDataA->InitializeFromData(InputPointData);
+		TArray<FPCGPoint>& PointsA = OutPointDataA->GetMutablePoints();
+		PointsA.SetNumUninitialized(InputPoints.Num());
+		OutputA.Data = OutPointDataA;
+		OutputA.Pin = PCGSplitPointsConstants::OutputALabel;
+
+		FPCGTaggedData& OutputB = Outputs.Add_GetRef(Inputs[i]);
+		UPCGPointData* OutPointDataB = NewObject<UPCGPointData>();
+		OutPointDataB->InitializeFromData(InputPointData);
+		TArray<FPCGPoint>& PointsB = OutPointDataB->GetMutablePoints();
+		PointsB.SetNumUninitialized(InputPoints.Num());
+		OutputB.Data = OutPointDataB;
+		OutputB.Pin = PCGSplitPointsConstants::OutputBLabel;
+
+		auto ProcessPoint = [Settings, SplitPosition, &InputPoints, &InputPointData, &OutPointDataA, &OutPointDataB, &PointsA, &PointsB](const int32 ReadIndex, const int32 WriteIndex)
+		{
+			const FPCGPoint& InPoint = InputPoints[ReadIndex];
+			FPCGPoint PointA = InPoint;
+			FPCGPoint PointB = InPoint;
+
+			const FVector SplitterPosition = (InPoint.BoundsMax - InPoint.BoundsMin) * SplitPosition;
+			const FVector MinPlusSplit = InPoint.BoundsMin + SplitterPosition;
+
+			FVector SplitValues = FVector::ZeroVector;
+			const int AxisIndex = static_cast<int>(Settings->SplitAxis);
+			if (ensure(AxisIndex >= 0 && AxisIndex <= 2))
+			{
+				SplitValues[AxisIndex] = 1.0;
+			}
+
+			// Execution for PointsA portion
+			UPCGMetadataAccessorHelpers::CopyPoint(InPoint, PointA, false);
+			PointA.BoundsMax = PointA.BoundsMax + SplitValues * (MinPlusSplit - PointA.BoundsMax);
+			PointsA[WriteIndex] = PointA;
+
+			// Execution of the PointsB portion
+			UPCGMetadataAccessorHelpers::CopyPoint(InPoint, PointB, false);
+			PointB.BoundsMin = PointB.BoundsMin + SplitValues * (MinPlusSplit - PointB.BoundsMin);
+			PointsB[WriteIndex] = PointB;
+		};
+
+		FPCGAsync::AsyncProcessingOneToOneEx(&Context->AsyncState, InputPoints.Num(), /*Initialize=*/[]() {}, ProcessPoint, /*bEnableTimeSlicing=*/false);
+	}
+
+	return true;
+}
+
+#undef LOCTEXT_NAMESPACE
\ No newline at end of file
diff --git a/Engine/Plugins/Experimental/PCG/Source/PCG/Private/Tests/Elements/PCGSplitPointsTest.cpp b/Engine/Plugins/Experimental/PCG/Source/PCG/Private/Tests/Elements/PCGSplitPointsTest.cpp
new file mode 100644
--- /dev/null
+++ b/Engine/Plugins/Experimental/PCG/Source/PCG/Private/Tests/Elements/PCGSplitPointsTest.cpp
@@ -0,0 +1,221 @@
+// Copyright Epic Games, Inc. All Rights Reserved.
+
+#include "Tests/PCGTestsCommon.h"
+
+#include "PCGContext.h"
+#include "Data/PCGPointData.h"
+
+#include "Elements/PCGSplitPoints.h"
+
+class FPCGSplitPointTestBase : public FPCGTestBaseClass
+{
+public:
+	using FPCGTestBaseClass::FPCGTestBaseClass;
+
+protected:
+	struct FTestParameters
+	{
+		// All your test parameters, like the expected value of BoundsMin
+		float SplitPosition = 0.0f;
+		EPCGSplitAxis SplitAxis = EPCGSplitAxis::Z;
+
+		FVector ExpectedABoundsMin = FVector();
+		FVector ExpectedABoundsMax = FVector();
+		FVector ExpectedBBoundsMin = FVector();
+		FVector ExpectedBBoundsMax = FVector();
+	};
+
+	bool GenerateTestDataRunAndValidate(const FTestParameters& Parameters)
+	{
+		PCGTestsCommon::FTestData TestData;
+		PCGTestsCommon::GenerateSettings<UPCGSplitPointsSettings>(TestData);
+		UPCGSplitPointsSettings* Settings = CastChecked<UPCGSplitPointsSettings>(TestData.Settings);
+		Settings->SplitPosition = Parameters.SplitPosition;
+		Settings->SplitAxis = Parameters.SplitAxis;
+
+		FPCGTaggedData& Inputs = TestData.InputData.TaggedData.Emplace_GetRef();
+		Inputs.Pin = PCGPinConstants::DefaultInputLabel;
+		Inputs.Data = PCGTestsCommon::CreatePointData();
+
+		FPCGElementPtr TestElement = TestData.Settings->GetElement();
+
+		TUniquePtr<FPCGContext> Context = TestData.InitializeTestContext();
+
+		while (!TestElement->Execute(Context.Get())) {}
+
+		// test our point data
+		const TArray<FPCGTaggedData>& Outputs = Context->OutputData.TaggedData;
+
+		UTEST_EQUAL("Output count", Outputs.Num(), 2);
+
+		// checking the PointA portion
+		const UPCGPointData* OutPointDataA = Cast<UPCGPointData>(Outputs[0].Data);
+
+		UTEST_NOT_NULL("OutputA point data", OutPointDataA);
+
+		const TArray<FPCGPoint>& OutPointsA = OutPointDataA->GetPoints();
+
+		UTEST_EQUAL("OutputA point count", OutPointsA.Num(), 1);
+
+		UTEST_EQUAL("OutputA Bounds Min", OutPointsA[0].BoundsMin, Parameters.ExpectedABoundsMin);
+		UTEST_EQUAL("OutputA Bounds Max", OutPointsA[0].BoundsMax, Parameters.ExpectedABoundsMax);
+
+		// checking the PointB portion
+		const UPCGPointData* OutPointDataB = Cast<UPCGPointData>(Outputs[1].Data);
+
+		UTEST_NOT_NULL("OutputB point data", OutPointDataB);
+
+		const TArray<FPCGPoint>& OutPointsB = OutPointDataB->GetPoints();
+
+		UTEST_EQUAL("OutputB point count", OutPointsB.Num(), 1);
+
+		UTEST_EQUAL("OutputB Bounds Min", OutPointsB[0].BoundsMin, Parameters.ExpectedBBoundsMin);
+		UTEST_EQUAL("OutputB Bounds Max", OutPointsB[0].BoundsMax, Parameters.ExpectedBBoundsMax);
+
+		return true;
+	}
+};
+
+IMPLEMENT_CUSTOM_SIMPLE_AUTOMATION_TEST(FPCGSplitPointsTest_DefaultX, FPCGSplitPointTestBase, "pcg.tests.SplitPoints.DefaultX", PCGTestsCommon::TestFlags)
+
+bool FPCGSplitPointsTest_DefaultX::RunTest(const FString& Parameters)
+{
+	FTestParameters DefaultXParameters{};
+
+	DefaultXParameters.SplitPosition = 0.5f;
+	DefaultXParameters.SplitAxis = EPCGSplitAxis::X;
+	DefaultXParameters.ExpectedABoundsMin = FVector(-1.0, -1.0, -1.0);
+	DefaultXParameters.ExpectedABoundsMax = FVector(0.0, 1.0, 1.0);
+	DefaultXParameters.ExpectedBBoundsMin = FVector(0.0, -1.0, -1.0);
+	DefaultXParameters.ExpectedBBoundsMax = FVector(1.0, 1.0, 1.0);
+
+	return GenerateTestDataRunAndValidate(DefaultXParameters);
+}
+
+IMPLEMENT_CUSTOM_SIMPLE_AUTOMATION_TEST(FPCGSplitPointsTest_DefaultY, FPCGSplitPointTestBase, "pcg.tests.SplitPoints.DefaultY", PCGTestsCommon::TestFlags)
+
+bool FPCGSplitPointsTest_DefaultY::RunTest(const FString& Parameters)
+{
+	FTestParameters DefaultYParameters{};
+
+	DefaultYParameters.SplitPosition = 0.5f;
+	DefaultYParameters.SplitAxis = EPCGSplitAxis::Y;
+	DefaultYParameters.ExpectedABoundsMin = FVector(-1.0, -1.0, -1.0);
+	DefaultYParameters.ExpectedABoundsMax = FVector(1.0, 0.0, 1.0);
+	DefaultYParameters.ExpectedBBoundsMin = FVector(-1.0, 0.0, -1.0);
+	DefaultYParameters.ExpectedBBoundsMax = FVector(1.0, 1.0, 1.0);
+
+	return GenerateTestDataRunAndValidate(DefaultYParameters);
+}
+
+IMPLEMENT_CUSTOM_SIMPLE_AUTOMATION_TEST(FPCGSplitPointsTest_DefaultZ, FPCGSplitPointTestBase, "pcg.tests.SplitPoints.DefaultZ", PCGTestsCommon::TestFlags)
+
+bool FPCGSplitPointsTest_DefaultZ::RunTest(const FString& Parameters)
+{
+	FTestParameters DefaultZParameters{};
+
+	DefaultZParameters.SplitPosition = 0.5f;
+	DefaultZParameters.SplitAxis = EPCGSplitAxis::Z;
+	DefaultZParameters.ExpectedABoundsMin = FVector(-1.0, -1.0, -1.0);
+	DefaultZParameters.ExpectedABoundsMax = FVector(1.0, 1.0, 0.0);
+	DefaultZParameters.ExpectedBBoundsMin = FVector(-1.0, -1.0, 0.0);
+	DefaultZParameters.ExpectedBBoundsMax = FVector(1.0, 1.0, 1.0);
+
+	return GenerateTestDataRunAndValidate(DefaultZParameters);
+}
+
+IMPLEMENT_CUSTOM_SIMPLE_AUTOMATION_TEST(FPCGSplitPointsTest_ZeroX, FPCGSplitPointTestBase, "pcg.tests.SplitPoints.ZeroX", PCGTestsCommon::TestFlags)
+
+bool FPCGSplitPointsTest_ZeroX::RunTest(const FString& Parameters)
+{
+	FTestParameters ZeroXParameters{};
+
+	ZeroXParameters.SplitPosition = 0.0f;
+	ZeroXParameters.SplitAxis = EPCGSplitAxis::X;
+	ZeroXParameters.ExpectedABoundsMin = FVector(-1.0, -1.0, -1.0);
+	ZeroXParameters.ExpectedABoundsMax = FVector(-1.0, 1.0, 1.0);
+	ZeroXParameters.ExpectedBBoundsMin = FVector(-1.0, -1.0, -1.0);
+	ZeroXParameters.ExpectedBBoundsMax = FVector(1.0, 1.0, 1.0);
+
+	return GenerateTestDataRunAndValidate(ZeroXParameters);
+}
+
+IMPLEMENT_CUSTOM_SIMPLE_AUTOMATION_TEST(FPCGSplitPointsTest_ZeroY, FPCGSplitPointTestBase, "pcg.tests.SplitPoints.ZeroY", PCGTestsCommon::TestFlags)
+
+bool FPCGSplitPointsTest_ZeroY::RunTest(const FString& Parameters)
+{
+	FTestParameters ZeroYParameters{};
+
+	ZeroYParameters.SplitPosition = 0.0f;
+	ZeroYParameters.SplitAxis = EPCGSplitAxis::Y;
+	ZeroYParameters.ExpectedABoundsMin = FVector(-1.0, -1.0, -1.0);
+	ZeroYParameters.ExpectedABoundsMax = FVector(1.0, -1.0, 1.0);
+	ZeroYParameters.ExpectedBBoundsMin = FVector(-1.0, -1.0, -1.0);
+	ZeroYParameters.ExpectedBBoundsMax = FVector(1.0, 1.0, 1.0);
+
+	return GenerateTestDataRunAndValidate(ZeroYParameters);
+}
+
+IMPLEMENT_CUSTOM_SIMPLE_AUTOMATION_TEST(FPCGSplitPointsTest_ZeroZ, FPCGSplitPointTestBase, "pcg.tests.SplitPoints.ZeroZ", PCGTestsCommon::TestFlags)
+
+bool FPCGSplitPointsTest_ZeroZ::RunTest(const FString& Parameters)
+{
+	FTestParameters ZeroZParameters{};
+
+	ZeroZParameters.SplitPosition = 0.0f;
+	ZeroZParameters.SplitAxis = EPCGSplitAxis::Z;
+	ZeroZParameters.ExpectedABoundsMin = FVector(-1.0, -1.0, -1.0);
+	ZeroZParameters.ExpectedABoundsMax = FVector(1.0, 1.0, -1.0);
+	ZeroZParameters.ExpectedBBoundsMin = FVector(-1.0, -1.0, -1.0);
+	ZeroZParameters.ExpectedBBoundsMax = FVector(1.0, 1.0, 1.0);
+
+	return GenerateTestDataRunAndValidate(ZeroZParameters);
+}
+
+IMPLEMENT_CUSTOM_SIMPLE_AUTOMATION_TEST(FPCGSplitPointsTest_OneX, FPCGSplitPointTestBase, "pcg.tests.SplitPoints.OneX", PCGTestsCommon::TestFlags)
+
+bool FPCGSplitPointsTest_OneX::RunTest(const FString& Parameters)
+{
+	FTestParameters OneXParameters{};
+
+	OneXParameters.SplitPosition = 1.0f;
+	OneXParameters.SplitAxis = EPCGSplitAxis::X;
+	OneXParameters.ExpectedABoundsMin = FVector(-1.0, -1.0, -1.0);
+	OneXParameters.ExpectedABoundsMax = FVector(1.0, 1.0, 1.0);
+	OneXParameters.ExpectedBBoundsMin = FVector(1.0, -1.0, -1.0);
+	OneXParameters.ExpectedBBoundsMax = FVector(1.0, 1.0, 1.0);
+
+	return GenerateTestDataRunAndValidate(OneXParameters);
+}
+
+IMPLEMENT_CUSTOM_SIMPLE_AUTOMATION_TEST(FPCGSplitPointsTest_OneY, FPCGSplitPointTestBase, "pcg.tests.SplitPoints.OneY", PCGTestsCommon::TestFlags)
+
+bool FPCGSplitPointsTest_OneY::RunTest(const FString& Parameters)
+{
+	FTestParameters OneYParameters{};
+
+	OneYParameters.SplitPosition = 1.0f;
+	OneYParameters.SplitAxis = EPCGSplitAxis::Y;
+	OneYParameters.ExpectedABoundsMin = FVector(-1.0, -1.0, -1.0);
+	OneYParameters.ExpectedABoundsMax = FVector(1.0, 1.0, 1.0);
+	OneYParameters.ExpectedBBoundsMin = FVector(-1.0, 1.0, -1.0);
+	OneYParameters.ExpectedBBoundsMax = FVector(1.0, 1.0, 1.0);
+
+	return GenerateTestDataRunAndValidate(OneYParameters);
+}
+
+IMPLEMENT_CUSTOM_SIMPLE_AUTOMATION_TEST(FPCGSplitPointsTest_OneZ, FPCGSplitPointTestBase, "pcg.tests.SplitPoints.OneZ", PCGTestsCommon::TestFlags)
+
+bool FPCGSplitPointsTest_OneZ::RunTest(const FString& Parameters)
+{
+	FTestParameters OneZParameters{};
+
+	OneZParameters.SplitPosition = 1.0f;
+	OneZParameters.SplitAxis = EPCGSplitAxis::Z;
+	OneZParameters.ExpectedABoundsMin = FVector(-1.0, -1.0, -1.0);
+	OneZParameters.ExpectedABoundsMax = FVector(1.0, 1.0, 1.0);
+	OneZParameters.ExpectedBBoundsMin = FVector(-1.0, -1.0, 1.0);
+	OneZParameters.ExpectedBBoundsMax = FVector(1.0, 1.0, 1.0);
+
+	return GenerateTestDataRunAndValidate(OneZParameters);
+}
\ No newline at end of file
diff --git a/Engine/Plugins/Experimental/PCG/Source/PCG/Public/Elements/PCGSplitPoints.h b/Engine/Plugins/Experimental/PCG/Source/PCG/Public/Elements/PCGSplitPoints.h
new file mode 100644
--- /dev/null
+++ b/Engine/Plugins/Experimental/PCG/Source/PCG/Public/Elements/PCGSplitPoints.h
@@ -0,0 +1,52 @@
+// Copyright Epic Games, Inc. All Rights Reserved.
+
+#pragma once
+
+#include "PCGSettings.h"
+
+#include "PCGSplitPoints.generated.h"
+
+UENUM()
+enum class EPCGSplitAxis
+{
+	X = 0,
+	Y,
+	Z
+};
+
+/**
+ * Splits each point into two separate points, and sets bounds based on the position and axis of the cut.
+ */
+UCLASS(BlueprintType, ClassGroup = (Procedural))
+class UPCGSplitPointsSettings : public UPCGSettings
+{
+	GENERATED_BODY()
+
+public:
+	//~Begin UPCGSettings interface
+#if WITH_EDITOR
+	virtual FName GetDefaultNodeName() const override { return FName(TEXT("SplitPoints")); }
+	virtual FText GetDefaultNodeTitle() const override { return NSLOCTEXT("PCGSplitPointsElement", "NodeTitle", "Split Points"); }
+	virtual FText GetNodeTooltipText() const override { return NSLOCTEXT("PCGSplitPointsElement", "NodeTooltip", "Splits each input point into two separate points and sets bounds based on the position and axis of the cut."); }
+	virtual EPCGSettingsType GetType() const override { return EPCGSettingsType::Spatial; }
+#endif
+
+protected:
+	virtual TArray<FPCGPinProperties> InputPinProperties() const override { return Super::DefaultPointInputPinProperties();  }
+	virtual TArray<FPCGPinProperties> OutputPinProperties() const override;
+	virtual FPCGElementPtr CreateElement() const override;
+	//~End UPCGSettings interface
+
+public:
+	UPROPERTY(BlueprintReadWrite, EditAnywhere, Category = Settings, meta = (ClampMin = "0", ClampMax = "1", PCG_Overridable))
+	float SplitPosition = 0.5f;
+
+	UPROPERTY(BlueprintReadWrite, EditAnywhere, Category = Settings, meta = (PCG_Overridable))
+	EPCGSplitAxis SplitAxis = EPCGSplitAxis::Z;
+};
+
+class FPCGSplitPointsElement : public IPCGElement
+{
+protected:
+	virtual bool ExecuteInternal(FPCGContext* Context) const override;
+};
\ No newline at end of file
