[PCG] Add an extra pointer guard to hopefully appease static analysis. Issue 399266.

diff --git a/Engine/Plugins/Experimental/PCG/Source/PCG/Private/Graph/PCGGraphExecutor.cpp b/Engine/Plugins/Experimental/PCG/Source/PCG/Private/Graph/PCGGraphExecutor.cpp
--- a/Engine/Plugins/Experimental/PCG/Source/PCG/Private/Graph/PCGGraphExecutor.cpp
+++ b/Engine/Plugins/Experimental/PCG/Source/PCG/Private/Graph/PCGGraphExecutor.cpp
@@ -109,7 +109,8 @@ FPCGTaskId FPCGGraphExecutor::Schedule(UPCGGraph* Graph, UPCGComponent* SourceCo
 
 	if (IsGraphCacheDebuggingEnabled())
 	{
-		UE_LOG(LogPCG, Log, TEXT("[%s] --- SCHEDULE GRAPH ---"), *SourceComponent->GetOwner()->GetName());
+		UE_LOG(LogPCG, Log, TEXT("[%s] --- SCHEDULE GRAPH ---"),
+			(SourceComponent && SourceComponent->GetOwner()) ? *SourceComponent->GetOwner()->GetName() : TEXT("MISSINGCOMPONENT"));
 	}
 	
 	FPCGTaskId ScheduledId = InvalidPCGTaskId;
@@ -134,7 +135,7 @@ FPCGTaskId FPCGGraphExecutor::Schedule(UPCGGraph* Graph, UPCGComponent* SourceCo
 	}
 
 #if WITH_EDITOR
-	if (UPCGSubsystem* Subsystem = SourceComponent ? SourceComponent->GetSubsystem() : nullptr)
+	if (UPCGSubsystem* Subsystem = SourceComponent->GetSubsystem())
 	{
 		Subsystem->OnScheduleGraph(*StackContextPtr);
 	}
