Renamed PCGAttributeFilterNamesElement to PCGDeleteAttributesElement

diff --git a/Engine/Plugins/Experimental/PCG/Config/BasePCG.ini b/Engine/Plugins/Experimental/PCG/Config/BasePCG.ini
--- a/Engine/Plugins/Experimental/PCG/Config/BasePCG.ini
+++ b/Engine/Plugins/Experimental/PCG/Config/BasePCG.ini
@@ -65,7 +65,8 @@
 +ClassRedirects=(OldName="PCGCreateAttributeSettings",NewName="/Script/PCG.PCGAddAttributeSettings")
 +ClassRedirects=(OldName="PCGPointFilterSettings",NewName="/Script/PCG.PCGAttributeFilteringSettings")
 +ClassRedirects=(OldName="PCGPointFilterRangeSettings",NewName="/Script/PCG.PCGAttributeFilteringRangeSettings")
-+ClassRedirects=(OldName="PCGAttributeFilterSettings",NewName="/Script/PCG.PCGAttributeFilterNamesSettings")
++ClassRedirects=(OldName="PCGAttributeFilterSettings",NewName="/Script/PCG.PCGDeleteAttributesSettings")
++ClassRedirects=(OldName="PCGAttributeFilterNamesSettings",NewName="/Script/PCG.PCGDeleteAttributesSettings")
 +ClassRedirects=(OldName="PCGPropertyToParamDataSettings",NewName="/Script/PCG.PCGGetActorPropertySettings")
 
 +StructRedirects=(OldName="PCGPointFilterThresholdSettings",NewName="/Script/PCG.PCGAttributeFilterThresholdSettings")
diff --git a/Engine/Plugins/Experimental/PCG/Source/PCG/Private/Elements/PCGAttributeFilterNamesElement.cpp b/Engine/Plugins/Experimental/PCG/Source/PCG/Private/Elements/PCGAttributeFilterNamesElement.cpp
deleted file mode 100644
diff --git a/Engine/Plugins/Experimental/PCG/Source/PCG/Private/Elements/PCGDeleteAttributesElement.cpp b/Engine/Plugins/Experimental/PCG/Source/PCG/Private/Elements/PCGDeleteAttributesElement.cpp
new file mode 100644
--- /dev/null
+++ b/Engine/Plugins/Experimental/PCG/Source/PCG/Private/Elements/PCGDeleteAttributesElement.cpp
@@ -0,0 +1,191 @@
+// Copyright Epic Games, Inc. All Rights Reserved.
+
+#include "Elements/PCGDeleteAttributesElement.h"
+
+#include "PCGContext.h"
+#include "PCGParamData.h"
+#include "PCGPin.h"
+#include "Data/PCGSpatialData.h"
+#include "Helpers/PCGHelpers.h"
+
+#define LOCTEXT_NAMESPACE "PCGAttributeFilterElement"
+
+namespace PCGAttributeFilterConstants
+{
+	const FName NodeName = TEXT("DeleteAttributes");
+	const FText NodeTitle = LOCTEXT("NodeTitle", "Delete Attributes");
+	const FText NodeTitleAlias = LOCTEXT("NodeTitleAlias", "Filter Attributes By Name");
+}
+
+UPCGDeleteAttributesSettings::UPCGDeleteAttributesSettings()
+{
+	if (PCGHelpers::IsNewObjectAndNotDefault(this))
+	{
+		Operation = EPCGAttributeFilterOperation::DeleteSelectedAttributes;
+	}
+}
+
+void UPCGDeleteAttributesSettings::PostLoad()
+{
+	Super::PostLoad();
+
+#if WITH_EDITOR
+	if (!AttributesToKeep_DEPRECATED.IsEmpty())
+	{
+		SelectedAttributes.Empty();
+		Operation = EPCGAttributeFilterOperation::KeepSelectedAttributes;
+		// Can't use FString::Join since it is an array of FName
+		for (int i = 0; i < AttributesToKeep_DEPRECATED.Num(); ++i)
+		{
+			if (i != 0)
+			{
+				SelectedAttributes += TEXT(",");
+			}
+
+			SelectedAttributes += AttributesToKeep_DEPRECATED[i].ToString();
+		}
+
+		AttributesToKeep_DEPRECATED.Empty();
+	}
+#endif // WITH_EDITOR
+}
+
+#if WITH_EDITOR
+FName UPCGDeleteAttributesSettings::GetDefaultNodeName() const
+{
+	return PCGAttributeFilterConstants::NodeName;
+}
+
+FText UPCGDeleteAttributesSettings::GetDefaultNodeTitle() const
+{
+	return PCGAttributeFilterConstants::NodeTitle;
+}
+
+TArray<FText> UPCGDeleteAttributesSettings::GetNodeTitleAliases() const
+{
+	return { PCGAttributeFilterConstants::NodeTitleAlias };
+}
+#endif
+
+EPCGDataType UPCGDeleteAttributesSettings::GetCurrentPinTypes(const UPCGPin* InPin) const
+{
+	check(InPin);
+	if (!InPin->IsOutputPin())
+	{
+		return Super::GetCurrentPinTypes(InPin);
+	}
+
+	// Output pin narrows to union of inputs on first pin
+	const EPCGDataType InputTypeUnion = GetTypeUnionOfIncidentEdges(PCGPinConstants::DefaultInputLabel);
+	return (InputTypeUnion != EPCGDataType::None) ? InputTypeUnion : EPCGDataType::Any;
+}
+
+FName UPCGDeleteAttributesSettings::AdditionalTaskName() const
+{
+	if (const UEnum* SelectionEnum = StaticEnum<EPCGAttributeFilterOperation>())
+	{
+		FText OperationText = SelectionEnum->GetDisplayNameTextByValue(static_cast<int64>(Operation));
+		TArray<FString> AttributesToKeep = PCGHelpers::GetStringArrayFromCommaSeparatedString(SelectedAttributes);
+
+		if (AttributesToKeep.Num() == 1)
+		{
+			return FName(FString::Printf(TEXT("%s (%s)"), *OperationText.ToString(), *AttributesToKeep[0]));
+		}
+		else
+		{
+			return FName(OperationText.ToString());
+		}
+	}
+	else
+	{
+		return NAME_None;
+	}
+}
+
+TArray<FPCGPinProperties> UPCGDeleteAttributesSettings::InputPinProperties() const
+{
+	TArray<FPCGPinProperties> PinProperties;
+	PinProperties.Emplace(PCGPinConstants::DefaultInputLabel, EPCGDataType::Any);
+
+	return PinProperties;
+}
+
+TArray<FPCGPinProperties> UPCGDeleteAttributesSettings::OutputPinProperties() const
+{
+	TArray<FPCGPinProperties> PinProperties;
+	PinProperties.Emplace(PCGPinConstants::DefaultOutputLabel, EPCGDataType::Any);
+
+	return PinProperties;
+}
+
+FPCGElementPtr UPCGDeleteAttributesSettings::CreateElement() const
+{
+	return MakeShared<FPCGDeleteAttributesElement>();
+}
+
+bool FPCGDeleteAttributesElement::ExecuteInternal(FPCGContext* Context) const
+{
+	TRACE_CPUPROFILER_EVENT_SCOPE(FPCGDeleteAttributesElement::Execute);
+
+	check(Context);
+
+	const UPCGDeleteAttributesSettings* Settings = Context->GetInputSettings<UPCGDeleteAttributesSettings>();
+
+	const bool bAddAttributesFromParent = (Settings->Operation == EPCGAttributeFilterOperation::DeleteSelectedAttributes);
+	const EPCGMetadataFilterMode FilterMode = bAddAttributesFromParent ? EPCGMetadataFilterMode::ExcludeAttributes : EPCGMetadataFilterMode::IncludeAttributes;
+
+	TSet<FName> AttributesToFilter; 
+	const TArray<FString> FilterAttributes = PCGHelpers::GetStringArrayFromCommaSeparatedString(Settings->SelectedAttributes);
+	for (const FString& FilterAttribute : FilterAttributes)
+	{
+		AttributesToFilter.Add(FName(*FilterAttribute));
+	}
+
+	TArray<FPCGTaggedData> Inputs = Context->InputData.GetInputsByPin(PCGPinConstants::DefaultInputLabel);
+
+	for (const FPCGTaggedData& InputTaggedData : Inputs)
+	{
+		const UPCGData* InputData = InputTaggedData.Data;
+		UPCGData* OutputData = nullptr;
+
+		const UPCGMetadata* ParentMetadata = nullptr;
+		UPCGMetadata* Metadata = nullptr;
+
+		if (const UPCGSpatialData* InputSpatialData = Cast<UPCGSpatialData>(InputData))
+		{
+			ParentMetadata = InputSpatialData->Metadata;
+
+			UPCGSpatialData* NewSpatialData = InputSpatialData->DuplicateData(/*bInitializeFromThisData=*/false);
+			Metadata = NewSpatialData->Metadata;
+			NewSpatialData->Metadata->InitializeWithAttributeFilter(ParentMetadata, AttributesToFilter, FilterMode);
+
+			// No need to inherit metadata since we already initialized it.
+			NewSpatialData->InitializeFromData(InputSpatialData, /*InMetadataParentOverride=*/ nullptr, /*bInheritMetadata=*/ false);
+
+			OutputData = NewSpatialData;
+		}
+		else if (const UPCGParamData* InputParamData = Cast<UPCGParamData>(InputData))
+		{
+			ParentMetadata = InputParamData->Metadata;
+
+			UPCGParamData* NewParamData = NewObject<UPCGParamData>();
+			Metadata = NewParamData->Metadata;
+			Metadata->InitializeAsCopyWithAttributeFilter(ParentMetadata, AttributesToFilter, FilterMode);
+
+			OutputData = NewParamData;
+		}
+		else
+		{
+			PCGE_LOG(Error, GraphAndLog, LOCTEXT("InvalidInputData", "Invalid data as input. Only Spatial and Attribute Set data are supported."));
+			continue;
+		}
+
+		TArray<FPCGTaggedData>& Outputs = Context->OutputData.TaggedData;
+		FPCGTaggedData& Output = Outputs.Add_GetRef(InputTaggedData);
+		Output.Data = OutputData;
+	}
+
+	return true;
+}
+
+#undef LOCTEXT_NAMESPACE
diff --git a/Engine/Plugins/Experimental/PCG/Source/PCG/Public/Elements/PCGAttributeFilterNamesElement.h b/Engine/Plugins/Experimental/PCG/Source/PCG/Public/Elements/PCGAttributeFilterNamesElement.h
deleted file mode 100644
diff --git a/Engine/Plugins/Experimental/PCG/Source/PCG/Public/Elements/PCGDeleteAttributesElement.h b/Engine/Plugins/Experimental/PCG/Source/PCG/Public/Elements/PCGDeleteAttributesElement.h
new file mode 100644
--- /dev/null
+++ b/Engine/Plugins/Experimental/PCG/Source/PCG/Public/Elements/PCGDeleteAttributesElement.h
@@ -0,0 +1,71 @@
+// Copyright Epic Games, Inc. All Rights Reserved.
+
+#pragma once
+
+#include "PCGSettings.h"
+
+#include "PCGDeleteAttributesElement.generated.h"
+
+UENUM()
+enum class EPCGAttributeFilterOperation
+{
+	KeepSelectedAttributes UMETA(DisplayName = "Keep Only Selected Attributes"),
+	DeleteSelectedAttributes
+};
+
+/**
+* Removes attributes from a given input metadata.
+* Either removes specifically named attributes or remove all attributes not in a given list.
+* 
+* The output will be the original data with the updated metadata.
+*/
+UCLASS(BlueprintType, ClassGroup = (Procedural))
+class PCG_API UPCGDeleteAttributesSettings : public UPCGSettings
+{
+	GENERATED_BODY()
+
+public:
+	UPCGDeleteAttributesSettings();
+
+	//~Begin UObject interface
+	virtual void PostLoad() override;
+	//~End UObject interface
+
+	//~Begin UPCGSettings interface
+#if WITH_EDITOR
+	virtual FName GetDefaultNodeName() const override;
+	virtual FText GetDefaultNodeTitle() const override;
+	virtual TArray<FText> GetNodeTitleAliases() const override;
+	virtual EPCGSettingsType GetType() const override { return EPCGSettingsType::Metadata; }
+	virtual bool HasDynamicPins() const override { return true; }
+#endif
+	virtual EPCGDataType GetCurrentPinTypes(const UPCGPin* InPin) const override;
+	virtual FName AdditionalTaskName() const override;
+
+protected:
+	virtual TArray<FPCGPinProperties> InputPinProperties() const override;
+	virtual TArray<FPCGPinProperties> OutputPinProperties() const override;
+	virtual FPCGElementPtr CreateElement() const override;
+	//~End UPCGSettings interface
+
+public:
+	// Implementation note: the default has been changed to DeleteSelected for new objects
+	UPROPERTY(BlueprintReadWrite, EditAnywhere, Category = Settings, meta = (PCG_Overridable))
+	EPCGAttributeFilterOperation Operation = EPCGAttributeFilterOperation::KeepSelectedAttributes;
+
+	/** Comma-separated list of attributes to keep or remove from the input data. */
+	UPROPERTY(BlueprintReadWrite, EditAnywhere, Category = Settings, meta = (PCG_Overridable))
+	FString SelectedAttributes;
+
+#if WITH_EDITORONLY_DATA
+	UPROPERTY()
+	TArray<FName> AttributesToKeep_DEPRECATED;
+#endif // WITH_EDITORONLY_DATA
+};
+
+
+class FPCGDeleteAttributesElement : public IPCGElement
+{
+protected:
+	virtual bool ExecuteInternal(FPCGContext* Context) const override;
+};
