Added node to allow data tags management

diff --git a/Engine/Plugins/Experimental/PCG/Source/PCG/Private/Elements/PCGDeleteTags.cpp b/Engine/Plugins/Experimental/PCG/Source/PCG/Private/Elements/PCGDeleteTags.cpp
new file mode 100644
--- /dev/null
+++ b/Engine/Plugins/Experimental/PCG/Source/PCG/Private/Elements/PCGDeleteTags.cpp
@@ -0,0 +1,106 @@
+// Copyright Epic Games, Inc. All Rights Reserved.
+
+#include "Elements/PCGDeleteTags.h"
+
+#include "PCGContext.h"
+#include "PCGData.h"
+#include "Helpers/PCGHelpers.h"
+
+#define LOCTEXT_NAMESPACE "PCGDeleteTagsElement"
+
+#if WITH_EDITOR
+FName UPCGDeleteTagsSettings::GetDefaultNodeName() const
+{
+	return TEXT("DeleteTags");
+}
+
+FText UPCGDeleteTagsSettings::GetDefaultNodeTitle() const
+{
+	return LOCTEXT("NodeTitle", "Delete Tags");
+}
+#endif // WITH_EDITOR
+
+EPCGDataType UPCGDeleteTagsSettings::GetCurrentPinTypes(const UPCGPin* InPin) const
+{
+	check(InPin);
+	if (!InPin->IsOutputPin())
+	{
+		return Super::GetCurrentPinTypes(InPin);
+	}
+
+	// Output pin narrows to union of inputs on first pin
+	const EPCGDataType InputTypeUnion = GetTypeUnionOfIncidentEdges(PCGPinConstants::DefaultInputLabel);
+	return (InputTypeUnion != EPCGDataType::None) ? InputTypeUnion : EPCGDataType::Any;
+}
+
+FName UPCGDeleteTagsSettings::AdditionalTaskName() const
+{
+	if (const UEnum* SelectionEnum = StaticEnum< EPCGTagFilterOperation>())
+	{
+		FText OperationText = SelectionEnum->GetDisplayNameTextByValue(static_cast<int64>(Operation));
+		TArray<FString> TagsToProcess = PCGHelpers::GetStringArrayFromCommaSeparatedString(SelectedTags);
+
+		if (TagsToProcess.Num() == 1)
+		{
+			return FName(FString::Printf(TEXT("%s (%s)"), *OperationText.ToString(), *TagsToProcess[0]));
+		}
+		else
+		{
+			return FName(OperationText.ToString());
+		}
+	}
+	else
+	{
+		return NAME_None;
+	}
+}
+
+TArray<FPCGPinProperties> UPCGDeleteTagsSettings::OutputPinProperties() const
+{
+	TArray<FPCGPinProperties> PinProperties;
+	PinProperties.Emplace(PCGPinConstants::DefaultOutputLabel, EPCGDataType::Any);
+
+	return PinProperties;
+}
+
+FPCGElementPtr UPCGDeleteTagsSettings::CreateElement() const
+{
+	return MakeShared<FPCGDeleteTagsElement>();
+}
+
+bool FPCGDeleteTagsElement::ExecuteInternal(FPCGContext* Context) const
+{
+	TRACE_CPUPROFILER_EVENT_SCOPE(FPCGManageTagsElement::Execute);
+	check(Context);
+
+	const UPCGDeleteTagsSettings* Settings = Context->GetInputSettings<UPCGDeleteTagsSettings>();
+	check(Settings);
+
+	const TArray<FString> FilterTags = PCGHelpers::GetStringArrayFromCommaSeparatedString(Settings->SelectedTags);
+	TSet<FString> TagsToFilter(FilterTags);
+
+	const bool bKeepInFilter = (Settings->Operation == EPCGTagFilterOperation::KeepOnlySelectedTags);
+
+	TArray<FPCGTaggedData> Inputs = Context->InputData.GetInputsByPin(PCGPinConstants::DefaultInputLabel);
+	TArray<FPCGTaggedData>& Outputs = Context->OutputData.TaggedData;
+
+	for (const FPCGTaggedData& Input : Inputs)
+	{
+		FPCGTaggedData& Output = Outputs.Add_GetRef(Input);
+
+		Output.Tags.Reset();
+		Output.Tags.Reserve(Input.Tags.Num());
+
+		for (const FString& Tag : Input.Tags)
+		{
+			if (TagsToFilter.Contains(Tag) == bKeepInFilter)
+			{
+				Output.Tags.Add(Tag);
+			}
+		}
+	}
+
+	return true;
+}
+
+#undef LOCTEXT_NAMESPACE
\ No newline at end of file
diff --git a/Engine/Plugins/Experimental/PCG/Source/PCG/Public/Elements/PCGDeleteTags.h b/Engine/Plugins/Experimental/PCG/Source/PCG/Public/Elements/PCGDeleteTags.h
new file mode 100644
--- /dev/null
+++ b/Engine/Plugins/Experimental/PCG/Source/PCG/Public/Elements/PCGDeleteTags.h
@@ -0,0 +1,51 @@
+// Copyright Epic Games, Inc. All Rights Reserved.
+
+#pragma once
+
+#include "PCGSettings.h"
+
+#include "PCGDeleteTags.generated.h"
+
+UENUM()
+enum class EPCGTagFilterOperation
+{
+	KeepOnlySelectedTags,
+	DeleteSelectedTags
+};
+
+/** Filters the tags on the input data. */
+UCLASS(BlueprintType, ClassGroup = (Procedural))
+class PCG_API UPCGDeleteTagsSettings : public UPCGSettings
+{
+	GENERATED_BODY()
+
+public:
+	//~Begin UPCGSettings interface
+#if WITH_EDITOR
+	virtual FName GetDefaultNodeName() const override;
+	virtual FText GetDefaultNodeTitle() const override;
+	virtual EPCGSettingsType GetType() const override { return EPCGSettingsType::Generic; }
+	virtual bool HasDynamicPins() const override { return true; }
+#endif
+	virtual EPCGDataType GetCurrentPinTypes(const UPCGPin* InPin) const override;
+	virtual FName AdditionalTaskName() const override;
+
+protected:
+	virtual TArray<FPCGPinProperties> OutputPinProperties() const override;
+	virtual FPCGElementPtr CreateElement() const override;
+	//~End UPCGSettings interface
+
+public:
+	UPROPERTY(BlueprintReadWrite, EditAnywhere, Category = Settings, meta = (PCG_Overridable))
+	EPCGTagFilterOperation Operation = EPCGTagFilterOperation::DeleteSelectedTags;
+
+	/** Comma-separated list of tags to add or remove from the input data. */
+	UPROPERTY(BlueprintReadWrite, EditAnywhere, Category = Settings, meta = (PCG_Overridable))
+	FString SelectedTags;
+};
+
+class FPCGDeleteTagsElement : public IPCGElement
+{
+protected:
+	virtual bool ExecuteInternal(FPCGContext* Context) const override;
+};
\ No newline at end of file
