[PCG] Nativization of DuplicatePoint node - Creates duplicates of each point with optional transform offsets.
--Include Unit Tests

diff --git a/Engine/Plugins/Experimental/PCG/Source/PCG/Private/Elements/PCGDuplicatePoint.cpp b/Engine/Plugins/Experimental/PCG/Source/PCG/Private/Elements/PCGDuplicatePoint.cpp
new file mode 100644
--- /dev/null
+++ b/Engine/Plugins/Experimental/PCG/Source/PCG/Private/Elements/PCGDuplicatePoint.cpp
@@ -0,0 +1,109 @@
+// Copyright Epic Games, Inc. All Rights Reserved.
+
+#include "Elements/PCGDuplicatePoint.h"
+
+#include "PCGContext.h"
+#include "PCGPin.h"
+#include "Data/PCGPointData.h"
+#include "Helpers/PCGAsync.h"
+#include "Helpers/PCGHelpers.h"
+#include "Metadata/PCGMetadataAccessor.h"
+
+#define LOCTEXT_NAMESPACE "PCGDuplicatePointElement"
+
+FPCGElementPtr UPCGDuplicatePointSettings::CreateElement() const
+{
+	return MakeShared<FPCGDuplicatePointElement>();
+}
+
+bool FPCGDuplicatePointElement::ExecuteInternal(FPCGContext* Context) const
+{
+	TRACE_CPUPROFILER_EVENT_SCOPE(FPCGDuplicatePointElement::Execute);
+	check(Context);
+
+	const UPCGDuplicatePointSettings* Settings = Context->GetInputSettings<UPCGDuplicatePointSettings>();
+	check(Settings);
+
+	const TArray<FPCGTaggedData> Inputs = Context->InputData.GetInputs();
+	TArray<FPCGTaggedData>& Outputs = Context->OutputData.TaggedData;
+
+	if (Settings->Iterations <= 0)
+	{
+		PCGE_LOG(Verbose, GraphAndLog, LOCTEXT("InvalidNumberOfIterations", "The number of interations must be at least 1."));
+		return true;
+	}
+
+	const int Iterations = Settings->Iterations;
+
+	const FVector Direction(FMath::Clamp(Settings->Direction.X, -1.0, 1.0), FMath::Clamp(Settings->Direction.Y, -1.0, 1.0), FMath::Clamp(Settings->Direction.Z, -1.0, 1.0));
+
+	for (int i = 0; i < Inputs.Num(); ++i)
+	{
+		const UPCGPointData* InputPointData = Cast<UPCGPointData>(Inputs[i].Data);
+		if (!InputPointData)
+		{
+			PCGE_LOG(Warning, GraphAndLog, LOCTEXT("InvalidInputPointData", "The input is not point data, skipped."));
+			continue;
+		}
+
+		const TArray<FPCGPoint>& InputPoints = InputPointData->GetPoints();
+
+		// Determines whether or not to include the source point in data
+		const int DuplicatesPerPoint = Iterations + (Settings->bOutputSourcePoint ? 1 : 0);
+		const int NumIterations = DuplicatesPerPoint * InputPoints.Num();
+		const int FirstDuplicateIndex = Settings->bOutputSourcePoint ? 0 : 1;
+
+		FPCGTaggedData& Output = Outputs.Add_GetRef(Inputs[i]);
+		if (InputPoints.IsEmpty())
+		{
+			Output.Data = InputPointData;
+			return true;
+		}
+
+		UPCGPointData* OutPointData = NewObject<UPCGPointData>();
+		OutPointData->InitializeFromData(InputPointData);
+		TArray<FPCGPoint>& OutputPoints = OutPointData->GetMutablePoints();
+		OutputPoints.SetNumUninitialized(NumIterations);
+		Output.Data = OutPointData;
+
+		auto ProcessPoint = [Settings, FirstDuplicateIndex, OutPointData, InputPointData, &Direction, &InputPoints, &OutputPoints](const int32 ReadIndex, const int32 WriteIndex)
+		{
+			const FPCGPoint& InPoint = InputPoints[ReadIndex % InputPoints.Num()];
+			FPCGPoint& OutputPoint = OutputPoints[WriteIndex];
+			const int DuplicateIndex = FirstDuplicateIndex + ReadIndex / InputPoints.Num();
+
+			if (DuplicateIndex == 0)
+			{
+				OutputPoint = InPoint;
+			}
+			else
+			{
+				const UPCGMetadata* InMetadata = InputPointData->Metadata;
+				UPCGMetadata* OutMetadata = OutPointData->Metadata;
+
+				UPCGMetadataAccessorHelpers::CopyPoint(InPoint, OutputPoint, /*bCopyMetadata=*/true, InMetadata, OutMetadata);
+
+				OutputPoint.Transform = FTransform(InPoint.Transform.GetRotation(), InPoint.Transform.GetLocation(), FVector::One());
+				OutputPoint.BoundsMin *= InPoint.Transform.GetScale3D();
+				OutputPoint.BoundsMax *= InPoint.Transform.GetScale3D();
+
+				const FVector DuplicateLocationOffset = ((OutputPoint.BoundsMax - OutputPoint.BoundsMin) * Direction + Settings->PointTransform.GetLocation()) * DuplicateIndex;
+				const FRotator DuplicateRotationOffset = Settings->PointTransform.Rotator() * DuplicateIndex;
+				const FVector DuplicateScaleMultiplier = FVector(
+					FMath::Pow(Settings->PointTransform.GetScale3D().X, DuplicateIndex), 
+					FMath::Pow(Settings->PointTransform.GetScale3D().Y, DuplicateIndex), 
+					FMath::Pow(Settings->PointTransform.GetScale3D().Z, DuplicateIndex));
+				const FTransform DuplicateTransformOffset(DuplicateRotationOffset, DuplicateLocationOffset, DuplicateScaleMultiplier);
+
+				OutputPoint.Transform *= DuplicateTransformOffset;
+				OutputPoint.Seed = PCGHelpers::ComputeSeedFromPosition(OutputPoint.Transform.GetLocation());
+			}
+		};
+
+		FPCGAsync::AsyncProcessingOneToOneEx(&Context->AsyncState, NumIterations, /*Initialize=*/[]() {}, ProcessPoint, /*bEnableTimeSlicing=*/false);
+	}
+
+	return true;
+}
+
+#undef LOCTEXT_NAMESPACE
\ No newline at end of file
diff --git a/Engine/Plugins/Experimental/PCG/Source/PCG/Private/Tests/Elements/PCGDuplicatePointTest.cpp b/Engine/Plugins/Experimental/PCG/Source/PCG/Private/Tests/Elements/PCGDuplicatePointTest.cpp
new file mode 100644
--- /dev/null
+++ b/Engine/Plugins/Experimental/PCG/Source/PCG/Private/Tests/Elements/PCGDuplicatePointTest.cpp
@@ -0,0 +1,340 @@
+// Copyright Epic Games, Inc. All Rights Reserved.
+
+#include "Tests/PCGTestsCommon.h"
+
+#include "PCGContext.h"
+#include "Data/PCGPointData.h"
+
+#include "Elements/PCGDuplicatePoint.h"
+
+class DuplicatePointTestBase : public FPCGTestBaseClass
+{
+
+public:
+	using FPCGTestBaseClass::FPCGTestBaseClass;
+	struct FTestParameters
+	{
+		int Iterations = 0;
+		FVector Direction = FVector();
+		bool IncludeSource = true;
+		FTransform PtTransform = FTransform();
+
+		int ExpectedPointCount = 0;
+		bool bCheckLocation = false;
+		bool bCheckRotation = false;
+		bool bCheckScale = false;
+	};
+
+protected:
+	bool GenerateTestDataRunAndValidate(FTestParameters Parameters)
+	{
+		PCGTestsCommon::FTestData TestData;
+		PCGTestsCommon::GenerateSettings<UPCGDuplicatePointSettings>(TestData);
+		UPCGDuplicatePointSettings* Settings = CastChecked<UPCGDuplicatePointSettings>(TestData.Settings);
+		Settings->Iterations = Parameters.Iterations;
+		Settings->Direction = Parameters.Direction;
+		Settings->bOutputSourcePoint = Parameters.IncludeSource;
+		Settings->PointTransform = Parameters.PtTransform;
+
+		FPCGTaggedData& Inputs = TestData.InputData.TaggedData.Emplace_GetRef();
+		Inputs.Pin = PCGPinConstants::DefaultInputLabel;
+		Inputs.Data = PCGTestsCommon::CreatePointData();
+
+		FPCGElementPtr TestElement = TestData.Settings->GetElement();
+
+		TUniquePtr<FPCGContext> Context = TestData.InitializeTestContext();
+
+		while (!TestElement->Execute(Context.Get())) {}
+
+		const TArray<FPCGTaggedData>& Outputs = Context->OutputData.TaggedData;
+
+		UTEST_EQUAL("Output count", Outputs.Num(), 1);
+
+		const UPCGPointData* OutPointData = Cast<UPCGPointData>(Outputs[0].Data);
+
+		UTEST_NOT_NULL("OutputA point data", OutPointData);
+
+		const TArray<FPCGPoint>& OutPoints = OutPointData->GetPoints();
+
+		UTEST_EQUAL("OutputA point count", OutPoints.Num(), Parameters.ExpectedPointCount);
+
+		const FVector LocationGetter = (OutPoints[0].BoundsMax - OutPoints[0].BoundsMin) * Parameters.Direction + Parameters.PtTransform.GetLocation();
+
+		if (Parameters.bCheckLocation)
+		{
+			for (int i = 0; i < OutPoints.Num(); ++i)
+			{
+				UTEST_EQUAL("Output Location", OutPoints[i].Transform.GetLocation(), LocationGetter * i);
+			}
+		}
+		if (Parameters.bCheckRotation)
+		{
+			for (int i = 0; i < OutPoints.Num(); ++i)
+			{
+				UTEST_EQUAL("Output Rotation", OutPoints[i].Transform.GetRotation().Rotator(), Parameters.PtTransform.GetRotation().Rotator() * i);
+			}
+		}
+		if (Parameters.bCheckScale)
+		{
+			const FVector ParameterTestScale = Parameters.PtTransform.GetScale3D();
+			for (int i = 0; i < OutPoints.Num(); ++i)
+			{
+				UTEST_EQUAL("Output Scale", OutPoints[i].Transform.GetScale3D(), FVector(FMath::Pow(ParameterTestScale.X, i), FMath::Pow(ParameterTestScale.Y, i), FMath::Pow(ParameterTestScale.Z, i)));
+			}
+		}
+
+		return true;
+	}
+};
+
+IMPLEMENT_CUSTOM_SIMPLE_AUTOMATION_TEST(FPCGDuplicatePointTest_WithSource, DuplicatePointTestBase, "pcg.tests.DuplicatePoint.WithSource", PCGTestsCommon::TestFlags)
+
+bool FPCGDuplicatePointTest_WithSource::RunTest(const FString& Parameters)
+{
+	FTestParameters SourceParameters{};
+
+	SourceParameters.Iterations = 1;
+	SourceParameters.Direction = FVector(0.0, 0.0, 1.0);
+	SourceParameters.IncludeSource = true;
+	SourceParameters.PtTransform = FTransform();
+	SourceParameters.ExpectedPointCount = 2;
+	SourceParameters.bCheckLocation = false;
+	SourceParameters.bCheckRotation = false;
+	SourceParameters.bCheckScale = false;
+	
+	return GenerateTestDataRunAndValidate(SourceParameters);
+}
+
+IMPLEMENT_CUSTOM_SIMPLE_AUTOMATION_TEST(FPCGDuplicatePointTest_WithoutSource, DuplicatePointTestBase, "pcg.tests.DuplicatePoint.WithoutSource", PCGTestsCommon::TestFlags)
+
+bool FPCGDuplicatePointTest_WithoutSource::RunTest(const FString& Parameters)
+{
+	FTestParameters NoSourceParameters{};
+
+	NoSourceParameters.Iterations = 1;
+	NoSourceParameters.Direction = FVector(0.0, 0.0, 1.0);
+	NoSourceParameters.IncludeSource = false;
+	NoSourceParameters.PtTransform = FTransform();
+	NoSourceParameters.ExpectedPointCount = 1;
+	NoSourceParameters.bCheckLocation = false;
+	NoSourceParameters.bCheckRotation = false;
+	NoSourceParameters.bCheckScale = false;
+
+	return GenerateTestDataRunAndValidate(NoSourceParameters);
+}
+
+IMPLEMENT_CUSTOM_SIMPLE_AUTOMATION_TEST(FPCGDuplicatePointTest_DirectionX, DuplicatePointTestBase, "pcg.tests.DuplicatePoint.DirectionX", PCGTestsCommon::TestFlags)
+
+bool FPCGDuplicatePointTest_DirectionX::RunTest(const FString& Parameters)
+{
+	FTestParameters XParameters{};
+
+	XParameters.Iterations = 2;
+	XParameters.Direction = FVector(1.0, 0.0, 0.0);
+	XParameters.IncludeSource = true;
+	XParameters.PtTransform = FTransform();
+	XParameters.ExpectedPointCount = 3;
+	XParameters.bCheckLocation = true;
+	XParameters.bCheckRotation = false;
+	XParameters.bCheckScale = false;
+
+	return GenerateTestDataRunAndValidate(XParameters);
+}
+
+IMPLEMENT_CUSTOM_SIMPLE_AUTOMATION_TEST(FPCGDuplicatePointTest_DirectionNegX, DuplicatePointTestBase, "pcg.tests.DuplicatePoint.DirectionNegX", PCGTestsCommon::TestFlags)
+
+bool FPCGDuplicatePointTest_DirectionNegX::RunTest(const FString& Parameters)
+{
+	FTestParameters NegXParameters{};
+
+	NegXParameters.Iterations = 2;
+	NegXParameters.Direction = FVector(-1.0, 0.0, 0.0);
+	NegXParameters.IncludeSource = true;
+	NegXParameters.PtTransform = FTransform();
+	NegXParameters.ExpectedPointCount = 3;
+	NegXParameters.bCheckLocation = true;
+	NegXParameters.bCheckRotation = false;
+	NegXParameters.bCheckScale = false;
+
+	return GenerateTestDataRunAndValidate(NegXParameters);
+}
+
+IMPLEMENT_CUSTOM_SIMPLE_AUTOMATION_TEST(FPCGDuplicatePointTest_DirectionY, DuplicatePointTestBase, "pcg.tests.DuplicatePoint.DirectionY", PCGTestsCommon::TestFlags)
+
+bool FPCGDuplicatePointTest_DirectionY::RunTest(const FString& Parameters)
+{
+	FTestParameters YParameters{};
+
+	YParameters.Iterations = 2;
+	YParameters.Direction = FVector(0.0, 1.0, 0.0);
+	YParameters.IncludeSource = true;
+	YParameters.PtTransform = FTransform();
+	YParameters.ExpectedPointCount = 3;
+	YParameters.bCheckLocation = true;
+	YParameters.bCheckRotation = false;
+	YParameters.bCheckScale = false;
+
+	return GenerateTestDataRunAndValidate(YParameters);
+}
+
+IMPLEMENT_CUSTOM_SIMPLE_AUTOMATION_TEST(FPCGDuplicatePointTest_DirectionNegY, DuplicatePointTestBase, "pcg.tests.DuplicatePoint.DirectionNegY", PCGTestsCommon::TestFlags)
+
+bool FPCGDuplicatePointTest_DirectionNegY::RunTest(const FString& Parameters)
+{
+	FTestParameters NegYParameters{};
+
+	NegYParameters.Iterations = 2;
+	NegYParameters.Direction = FVector(0.0, -1.0, 0.0);
+	NegYParameters.IncludeSource = true;
+	NegYParameters.PtTransform = FTransform();
+	NegYParameters.ExpectedPointCount = 3;
+	NegYParameters.bCheckLocation = true;
+	NegYParameters.bCheckRotation = false;
+	NegYParameters.bCheckScale = false;
+
+	return GenerateTestDataRunAndValidate(NegYParameters);
+}
+
+IMPLEMENT_CUSTOM_SIMPLE_AUTOMATION_TEST(FPCGDuplicatePointTest_DirectionZ, DuplicatePointTestBase, "pcg.tests.DuplicatePoint.DirectionZ", PCGTestsCommon::TestFlags)
+
+bool FPCGDuplicatePointTest_DirectionZ::RunTest(const FString& Parameters)
+{
+	FTestParameters ZParameters{};
+
+	ZParameters.Iterations = 2;
+	ZParameters.Direction = FVector(0.0, 0.0, 1.0);
+	ZParameters.IncludeSource = true;
+	ZParameters.PtTransform = FTransform();
+	ZParameters.ExpectedPointCount = 3;
+	ZParameters.bCheckLocation = true;
+	ZParameters.bCheckRotation = false;
+	ZParameters.bCheckScale = false;
+
+	return GenerateTestDataRunAndValidate(ZParameters);
+}
+
+IMPLEMENT_CUSTOM_SIMPLE_AUTOMATION_TEST(FPCGDuplicatePointTest_DirectionNegZ, DuplicatePointTestBase, "pcg.tests.DuplicatePoint.DirectionNegZ", PCGTestsCommon::TestFlags)
+
+bool FPCGDuplicatePointTest_DirectionNegZ::RunTest(const FString& Parameters)
+{
+	FTestParameters NegZParameters{};
+
+	NegZParameters.Iterations = 2;
+	NegZParameters.Direction = FVector(0.0, 0.0, -1.0);
+	NegZParameters.IncludeSource = true;
+	NegZParameters.PtTransform = FTransform();
+	NegZParameters.ExpectedPointCount = 3;
+	NegZParameters.bCheckLocation = true;
+	NegZParameters.bCheckRotation = false;
+	NegZParameters.bCheckScale = false;
+
+	return GenerateTestDataRunAndValidate(NegZParameters);
+}
+
+IMPLEMENT_CUSTOM_SIMPLE_AUTOMATION_TEST(FPCGDuplicatePointTest_ScaleOne, DuplicatePointTestBase, "pcg.tests.DuplicatePoint.ScaleOne", PCGTestsCommon::TestFlags)
+
+bool FPCGDuplicatePointTest_ScaleOne::RunTest(const FString& Parameters)
+{
+	FTestParameters ScaleParameters{};
+
+	ScaleParameters.Iterations = 2;
+	ScaleParameters.Direction = FVector(0.0, 0.0, 1.0);
+	ScaleParameters.IncludeSource = true;
+	ScaleParameters.PtTransform = FTransform();
+	ScaleParameters.ExpectedPointCount = 3;
+	ScaleParameters.bCheckLocation = false;
+	ScaleParameters.bCheckRotation = false;
+	ScaleParameters.bCheckScale = true;
+
+	return GenerateTestDataRunAndValidate(ScaleParameters);
+}
+
+IMPLEMENT_CUSTOM_SIMPLE_AUTOMATION_TEST(FPCGDuplicatePointTest_ScaleNegative, DuplicatePointTestBase, "pcg.tests.DuplicatePoint.ScaleNegative", PCGTestsCommon::TestFlags)
+
+bool FPCGDuplicatePointTest_ScaleNegative::RunTest(const FString& Parameters)
+{
+	FTestParameters NegScaleParameters{};
+
+	NegScaleParameters.Iterations = 2;
+	NegScaleParameters.Direction = FVector(0.0, 0.0, 1.0);
+	NegScaleParameters.IncludeSource = true;
+	NegScaleParameters.PtTransform = FTransform(FRotator::ZeroRotator, FVector::ZeroVector, FVector(-1.0));
+	NegScaleParameters.ExpectedPointCount = 3;
+	NegScaleParameters.bCheckLocation = false;
+	NegScaleParameters.bCheckRotation = false;
+	NegScaleParameters.bCheckScale = true;
+
+	return GenerateTestDataRunAndValidate(NegScaleParameters);
+}
+
+IMPLEMENT_CUSTOM_SIMPLE_AUTOMATION_TEST(FPCGDuplicatePointTest_Rotation, DuplicatePointTestBase, "pcg.tests.DuplicatePoint.Rotation", PCGTestsCommon::TestFlags)
+
+bool FPCGDuplicatePointTest_Rotation::RunTest(const FString& Parameters)
+{
+	FTestParameters RotationParameters{};
+
+	RotationParameters.Iterations = 2;
+	RotationParameters.Direction = FVector(0.0, 0.0, 1.0);
+	RotationParameters.IncludeSource = true;
+	RotationParameters.PtTransform = FTransform(FRotator(20.0, 0.0, 0.0), FVector::ZeroVector, FVector::One());
+	RotationParameters.ExpectedPointCount = 3;
+	RotationParameters.bCheckLocation = false;
+	RotationParameters.bCheckRotation = true;
+	RotationParameters.bCheckScale = false;
+
+	return GenerateTestDataRunAndValidate(RotationParameters);
+}
+
+IMPLEMENT_CUSTOM_SIMPLE_AUTOMATION_TEST(FPCGDuplicatePointTest_LocationX, DuplicatePointTestBase, "pcg.tests.DuplicatePoint.LocationX", PCGTestsCommon::TestFlags)
+
+bool FPCGDuplicatePointTest_LocationX::RunTest(const FString& Parameters)
+{
+	FTestParameters LocationXParameters{};
+
+	LocationXParameters.Iterations = 2;
+	LocationXParameters.Direction = FVector(0.0, 0.0, 1.0);
+	LocationXParameters.IncludeSource = true;
+	LocationXParameters.PtTransform = FTransform(FRotator::ZeroRotator, FVector(5.0, 0.0, 0.0), FVector::One());
+	LocationXParameters.ExpectedPointCount = 3;
+	LocationXParameters.bCheckLocation = true;
+	LocationXParameters.bCheckRotation = false;
+	LocationXParameters.bCheckScale = false;
+
+	return GenerateTestDataRunAndValidate(LocationXParameters);
+}
+
+IMPLEMENT_CUSTOM_SIMPLE_AUTOMATION_TEST(FPCGDuplicatePointTest_LocationY, DuplicatePointTestBase, "pcg.tests.DuplicatePoint.LocationY", PCGTestsCommon::TestFlags)
+
+bool FPCGDuplicatePointTest_LocationY::RunTest(const FString& Parameters)
+{
+	FTestParameters LocationYParameters{};
+
+	LocationYParameters.Iterations = 2;
+	LocationYParameters.Direction = FVector(0.0, 0.0, 1.0);
+	LocationYParameters.IncludeSource = true;
+	LocationYParameters.PtTransform = FTransform(FRotator::ZeroRotator, FVector(0.0, 5.0, 0.0), FVector::One());
+	LocationYParameters.ExpectedPointCount = 3;
+	LocationYParameters.bCheckLocation = true;
+	LocationYParameters.bCheckRotation = false;
+	LocationYParameters.bCheckScale = false;
+
+	return GenerateTestDataRunAndValidate(LocationYParameters);
+}
+
+IMPLEMENT_CUSTOM_SIMPLE_AUTOMATION_TEST(FPCGDuplicatePointTest_LocationZ, DuplicatePointTestBase, "pcg.tests.DuplicatePoint.LocationZ", PCGTestsCommon::TestFlags)
+
+bool FPCGDuplicatePointTest_LocationZ::RunTest(const FString& Parameters)
+{
+	FTestParameters LocationZParameters{};
+
+	LocationZParameters.Iterations = 2;
+	LocationZParameters.Direction = FVector(0.0, 0.0, 1.0);
+	LocationZParameters.IncludeSource = true;
+	LocationZParameters.PtTransform = FTransform(FRotator::ZeroRotator, FVector(0.0, 0.0, 5.0), FVector::One());
+	LocationZParameters.ExpectedPointCount = 3;
+	LocationZParameters.bCheckLocation = true;
+	LocationZParameters.bCheckRotation = false;
+	LocationZParameters.bCheckScale = false;
+
+	return GenerateTestDataRunAndValidate(LocationZParameters);
+}
\ No newline at end of file
diff --git a/Engine/Plugins/Experimental/PCG/Source/PCG/Public/Elements/PCGDuplicatePoint.h b/Engine/Plugins/Experimental/PCG/Source/PCG/Public/Elements/PCGDuplicatePoint.h
new file mode 100644
--- /dev/null
+++ b/Engine/Plugins/Experimental/PCG/Source/PCG/Public/Elements/PCGDuplicatePoint.h
@@ -0,0 +1,54 @@
+// Copyright Epic Games, Inc. All Rights Reserved.
+
+#pragma once
+
+#include "PCGSettings.h"
+
+#include "PCGDuplicatePoint.generated.h"
+
+/**
+ * Creates duplicates of each point with optional transform offsets.
+ */
+UCLASS(BlueprintType, ClassGroup = (Procedural))
+class UPCGDuplicatePointSettings : public UPCGSettings
+{
+	GENERATED_BODY()
+
+public:
+	//~Begin UPCGSettings interface
+#if WITH_EDITOR
+	virtual FName GetDefaultNodeName() const override { return FName(TEXT("DuplicatePoint")); }
+	virtual FText GetDefaultNodeTitle() const override { return NSLOCTEXT("PCGDuplicatePointElement", "NodeTitle", "Duplicate Point"); }
+	virtual FText GetNodeTooltipText() const override { return NSLOCTEXT("PCGDuplicatePointElement", "NodeTooltip", "Creates duplicates of each point with optional transform offsets."); }
+	virtual EPCGSettingsType GetType() const override { return EPCGSettingsType::Spatial; }
+#endif
+
+protected:
+	virtual TArray<FPCGPinProperties> InputPinProperties() const override { return Super::DefaultPointInputPinProperties(); }
+	virtual TArray<FPCGPinProperties> OutputPinProperties() const override { return Super::DefaultPointOutputPinProperties(); }
+	virtual FPCGElementPtr CreateElement() const override;
+	//~End UPCGSettings interface
+
+public:
+	/** Number of duplicates to produce. */
+	UPROPERTY(BlueprintReadWrite, EditAnywhere, Category = Settings, meta = (ClampMin = "1", PCG_Overridable))
+	int Iterations = 1;
+	
+	/** Direction to stack point duplicates. */
+	UPROPERTY(BlueprintReadWrite, EditAnywhere, Category = Settings, meta = (ClampMin = "-1.0", ClampMax = "1.0", PCG_Overridable))
+	FVector Direction = FVector(0.0, 0.0, 1.0);
+	
+	/** Include the source point. */
+	UPROPERTY(BlueprintReadWrite, EditAnywhere, Category = Settings, meta = (PCG_Overridable))
+	bool bOutputSourcePoint = true;
+
+	/** Transform offset for each point duplicate */
+	UPROPERTY(BlueprintReadWrite, EditAnywhere, Category = Settings, meta = (PCG_Overridable))
+	FTransform PointTransform;
+};
+
+class FPCGDuplicatePointElement : public IPCGElement
+{
+protected:
+	virtual bool ExecuteInternal(FPCGContext* Context) const override;
+};
\ No newline at end of file
