[PCG] Fix crash on naming comment boxes in PCG graphs
#rnx

diff --git a/Engine/Plugins/Experimental/PCG/Source/PCGEditor/Private/PCGEditor.cpp b/Engine/Plugins/Experimental/PCG/Source/PCGEditor/Private/PCGEditor.cpp
--- a/Engine/Plugins/Experimental/PCG/Source/PCGEditor/Private/PCGEditor.cpp
+++ b/Engine/Plugins/Experimental/PCG/Source/PCGEditor/Private/PCGEditor.cpp
@@ -2212,7 +2212,7 @@ void FPCGEditor::OnNodeTitleCommitted(const FText& NewText, ETextCommit::Type Co
 			NodeBeingChanged->OnRenameNode(NewText.ToString());
 		}
 
-		if (UPCGEditorGraphNodeBase* PCGEditorNode = CastChecked<UPCGEditorGraphNodeBase>(NodeBeingChanged))
+		if (UPCGEditorGraphNodeBase* PCGEditorNode = Cast<UPCGEditorGraphNodeBase>(NodeBeingChanged))
 		{
 			PCGEditorNode->ExitRenamingMode();
 		}
