Improved some texts in the UI and made empty subgraphs behave like forwarding nodes

diff --git a/Engine/Plugins/Experimental/PCG/Source/PCG/Private/Elements/PCGLoopElement.cpp b/Engine/Plugins/Experimental/PCG/Source/PCG/Private/Elements/PCGLoopElement.cpp
--- a/Engine/Plugins/Experimental/PCG/Source/PCG/Private/Elements/PCGLoopElement.cpp
+++ b/Engine/Plugins/Experimental/PCG/Source/PCG/Private/Elements/PCGLoopElement.cpp
@@ -171,7 +171,7 @@ FName UPCGLoopSettings::AdditionalTaskName() const
 	}
 	else
 	{
-		return FName(LOCTEXT("NodeTitleExtendedInvalidSubgraph", "Loop Subgraph - Invalid Subgraph").ToString());
+		return FName(LOCTEXT("NodeTitleExtendedInvalidSubgraph", "Loop Subgraph - Empty Subgraph").ToString());
 	}
 }
 
@@ -197,14 +197,15 @@ bool FPCGLoopElement::ExecuteInternal(FPCGContext* InContext) const
 		UPCGGraph* Subgraph = Settings->GetSubgraph();
 		UPCGSubsystem* Subsystem = Context->SourceComponent.IsValid() ? Context->SourceComponent->GetSubsystem() : nullptr;
 
-		if (!Subsystem || !Subgraph)
+		if (!Subgraph)
+		{
+			// No subgraph is equivalent to disabling the node
+			Context->OutputData.TaggedData = Context->InputData.GetInputsByPin(PCGPinConstants::DefaultInputLabel);
+			return true;
+		}
+		else if (!Subsystem)
 		{
 			// Job cannot run; cancel
-			if (!Subgraph)
-			{
-				PCGE_LOG(Error, GraphAndLog, LOCTEXT("InvalidSubgraph", "Cannot loop on an invalid subgraph"));
-			}
-			
 			Context->OutputData.bCancelExecution = true;
 			return true;
 		}
@@ -218,7 +219,7 @@ bool FPCGLoopElement::ExecuteInternal(FPCGContext* InContext) const
 		// Early out if there are no data on the loop pin
 		if (LoopDataCollection.IsEmpty() || LoopDataCollection[0].TaggedData.IsEmpty())
 		{
-			PCGE_LOG(Verbose, LogOnly, LOCTEXT("EmptyLoopCollection", "Loop data is empty - will not do anything"));
+			PCGE_LOG(Verbose, LogOnly, LOCTEXT("EmptyLoopCollection", "Loop data is empty - will not do anything."));
 			return true;
 		}
 
diff --git a/Engine/Plugins/Experimental/PCG/Source/PCG/Private/PCGSubgraph.cpp b/Engine/Plugins/Experimental/PCG/Source/PCG/Private/PCGSubgraph.cpp
--- a/Engine/Plugins/Experimental/PCG/Source/PCG/Private/PCGSubgraph.cpp
+++ b/Engine/Plugins/Experimental/PCG/Source/PCG/Private/PCGSubgraph.cpp
@@ -204,7 +204,10 @@ TArray<FPCGPinProperties> UPCGBaseSubgraphSettings::OutputPinProperties() const
 	}
 	else
 	{
-		return Super::OutputPinProperties();
+		// Here we do not want the base class implementation as it forces Spatial but that might not be the case here
+		TArray<FPCGPinProperties> OutputPins;
+		OutputPins.Emplace(PCGPinConstants::DefaultOutputLabel, EPCGDataType::Any);
+		return OutputPins;
 	}
 }
 
@@ -275,7 +278,7 @@ FName UPCGSubgraphSettings::AdditionalTaskName() const
 	}
 	else
 	{
-		return TEXT("Invalid subgraph");
+		return TEXT("Empty subgraph");
 	}
 }
 
@@ -482,6 +485,12 @@ bool FPCGSubgraphElement::IsCacheable(const UPCGSettings* InSettings) const
 	return (!Settings || !Settings->IsDynamicGraph());
 }
 
+bool FPCGSubgraphElement::IsPassthrough(const UPCGSettings* InSettings) const
+{
+	const UPCGSubgraphSettings* Settings = Cast<UPCGSubgraphSettings>(InSettings);
+	return (!Settings || (Settings->bEnabled && Settings->GetSubgraph()));
+}
+
 bool FPCGSubgraphElement::ExecuteInternal(FPCGContext* InContext) const
 {
 	FPCGSubgraphContext* Context = static_cast<FPCGSubgraphContext*>(InContext);
@@ -550,6 +559,12 @@ bool FPCGSubgraphElement::ExecuteInternal(FPCGContext* InContext) const
 					return true;
 				}
 			}
+			else if (!Subgraph)
+			{
+				// Simple pass-through
+				Context->OutputData.TaggedData = Context->InputData.GetInputsByPin(PCGPinConstants::DefaultInputLabel);
+				return true;
+			}
 			else
 			{
 				// Job cannot run; cancel
diff --git a/Engine/Plugins/Experimental/PCG/Source/PCG/Public/PCGSubgraph.h b/Engine/Plugins/Experimental/PCG/Source/PCG/Public/PCGSubgraph.h
--- a/Engine/Plugins/Experimental/PCG/Source/PCG/Public/PCGSubgraph.h
+++ b/Engine/Plugins/Experimental/PCG/Source/PCG/Public/PCGSubgraph.h
@@ -167,7 +167,7 @@ public:
 
 protected:
 	virtual bool ExecuteInternal(FPCGContext* Context) const override;
-	virtual bool IsPassthrough(const UPCGSettings* InSettings) const override { return !InSettings || InSettings->bEnabled; }
+	virtual bool IsPassthrough(const UPCGSettings* InSettings) const override;
 	void PrepareSubgraphData(const UPCGSubgraphSettings* Settings, FPCGSubgraphContext* Context, const FPCGDataCollection& InputData, FPCGDataCollection& OutputData) const;
 	void PrepareSubgraphUserParameters(const UPCGSubgraphSettings* Settings, FPCGSubgraphContext* Context, FPCGDataCollection& OutputData) const;
 };
