[PCG] Created a new "Filter By Index" node

+ PCGParser namespace and utilities included
+ PCGIndexing namespace and utilities included
+ Unit Tests Included

diff --git a/Engine/Plugins/Experimental/PCG/Source/PCG/Private/Elements/PCGAttributeFilter.cpp b/Engine/Plugins/Experimental/PCG/Source/PCG/Private/Elements/PCGAttributeFilter.cpp
--- a/Engine/Plugins/Experimental/PCG/Source/PCG/Private/Elements/PCGAttributeFilter.cpp
+++ b/Engine/Plugins/Experimental/PCG/Source/PCG/Private/Elements/PCGAttributeFilter.cpp
@@ -20,8 +20,6 @@ namespace PCGAttributeFilterConstants
 	const FName FilterLabel = TEXT("Filter");
 	const FName FilterMinLabel = TEXT("FilterMin");
 	const FName FilterMaxLabel = TEXT("FilterMax");
-	const FName InFilterLabel = TEXT("InsideFilter");
-	const FName OutFilterLabel = TEXT("OutsideFilter");
 
 	constexpr int32 ChunkSize = 256;
 
@@ -297,8 +295,8 @@ TArray<FPCGPinProperties> UPCGAttributeFilteringSettings::InputPinProperties() c
 TArray<FPCGPinProperties> UPCGAttributeFilteringSettings::OutputPinProperties() const
 {
 	TArray<FPCGPinProperties> PinProperties;
-	PinProperties.Emplace(PCGAttributeFilterConstants::InFilterLabel, EPCGDataType::Any);
-	PinProperties.Emplace(PCGAttributeFilterConstants::OutFilterLabel, EPCGDataType::Any);
+	PinProperties.Emplace(PCGPinConstants::DefaultInFilterLabel, EPCGDataType::Any);
+	PinProperties.Emplace(PCGPinConstants::DefaultOutFilterLabel, EPCGDataType::Any);
 
 	return PinProperties;
 }
@@ -407,8 +405,8 @@ TArray<FPCGPinProperties> UPCGAttributeFilteringRangeSettings::InputPinPropertie
 TArray<FPCGPinProperties> UPCGAttributeFilteringRangeSettings::OutputPinProperties() const
 {
 	TArray<FPCGPinProperties> PinProperties;
-	PinProperties.Emplace(PCGAttributeFilterConstants::InFilterLabel, EPCGDataType::Any);
-	PinProperties.Emplace(PCGAttributeFilterConstants::OutFilterLabel, EPCGDataType::Any);
+	PinProperties.Emplace(PCGPinConstants::DefaultInFilterLabel, EPCGDataType::Any);
+	PinProperties.Emplace(PCGPinConstants::DefaultOutFilterLabel, EPCGDataType::Any);
 
 	return PinProperties;
 }
@@ -501,14 +499,14 @@ bool FPCGAttributeFilterElementBase::DoFiltering(FPCGContext* Context, EPCGAttri
 		auto ForwardInputToInFilterPin = [&Outputs, Input]()
 		{
 			FPCGTaggedData& InFilterOutput = Outputs.Add_GetRef(Input);
-			InFilterOutput.Pin = PCGAttributeFilterConstants::InFilterLabel;
+			InFilterOutput.Pin = PCGPinConstants::DefaultInFilterLabel;
 		};
 
 		// If there is a problem with target -> forward to OutFilter
 		auto ForwardInputToOutFilterPin = [&Outputs, Input]()
 		{
 			FPCGTaggedData& OutFilterOutput = Outputs.Add_GetRef(Input);
-			OutFilterOutput.Pin = PCGAttributeFilterConstants::OutFilterLabel;
+			OutFilterOutput.Pin = PCGPinConstants::DefaultOutFilterLabel;
 		};
 
 		FPCGAttributePropertyInputSelector TargetAttribute = InTargetAttribute.CopyAndFixLast(OriginalData);
@@ -704,12 +702,12 @@ bool FPCGAttributeFilterElementBase::DoFiltering(FPCGContext* Context, EPCGAttri
 		if (PCGMetadataAttribute::CallbackWithRightType(TargetAccessor->GetUnderlyingType(), Operation))
 		{
 			FPCGTaggedData& InFilterOutput = Outputs.Add_GetRef(Input);
-			InFilterOutput.Pin = PCGAttributeFilterConstants::InFilterLabel;
+			InFilterOutput.Pin = PCGPinConstants::DefaultInFilterLabel;
 			InFilterOutput.Data = InFilterData;
 			InFilterOutput.Tags = Input.Tags;
 
 			FPCGTaggedData& OutFilterOutput = Outputs.Add_GetRef(Input);
-			OutFilterOutput.Pin = PCGAttributeFilterConstants::OutFilterLabel;
+			OutFilterOutput.Pin = PCGPinConstants::DefaultOutFilterLabel;
 			OutFilterOutput.Data = OutFilterData;
 			OutFilterOutput.Tags = Input.Tags;
 		}
@@ -729,8 +727,8 @@ bool FPCGAttributeFilterElement::ExecuteInternal(FPCGContext* Context) const
 	TRACE_CPUPROFILER_EVENT_SCOPE(FPCGAttributeFilterElement::Execute);
 
 #if !WITH_EDITOR
-	const bool bHasInFilterOutputPin = Context->Node && Context->Node->IsOutputPinConnected(PCGAttributeFilterConstants::InFilterLabel);
-	const bool bHasOutsideFilterOutputPin = Context->Node && Context->Node->IsOutputPinConnected(PCGAttributeFilterConstants::OutFilterLabel);
+	const bool bHasInFilterOutputPin = Context->Node && Context->Node->IsOutputPinConnected(PCGPinConstants::DefaultInFilterLabel);
+	const bool bHasOutsideFilterOutputPin = Context->Node && Context->Node->IsOutputPinConnected(PCGPinConstants::DefaultOutFilterLabel);
 
 	// Early out - only in non-editor builds, otherwise we will potentially poison the cache, since it is input-driven
 	if (!bHasInFilterOutputPin && !bHasOutsideFilterOutputPin)
@@ -756,8 +754,8 @@ bool FPCGAttributeFilterRangeElement::ExecuteInternal(FPCGContext* Context) cons
 	TRACE_CPUPROFILER_EVENT_SCOPE(FPCGAttributeFilterRangeElement::Execute);
 
 #if !WITH_EDITOR
-	const bool bHasInFilterOutputPin = Context->Node && Context->Node->IsOutputPinConnected(PCGAttributeFilterConstants::InFilterLabel);
-	const bool bHasOutsideFilterOutputPin = Context->Node && Context->Node->IsOutputPinConnected(PCGAttributeFilterConstants::OutFilterLabel);
+	const bool bHasInFilterOutputPin = Context->Node && Context->Node->IsOutputPinConnected(PCGPinConstants::DefaultInFilterLabel);
+	const bool bHasOutsideFilterOutputPin = Context->Node && Context->Node->IsOutputPinConnected(PCGPinConstants::DefaultOutFilterLabel);
 
 	// Early out - only in non-editor builds, otherwise we will potentially poison the cache, since it is input-driven
 	if (!bHasInFilterOutputPin && !bHasOutsideFilterOutputPin)
diff --git a/Engine/Plugins/Experimental/PCG/Source/PCG/Private/Elements/PCGFilterByIndex.cpp b/Engine/Plugins/Experimental/PCG/Source/PCG/Private/Elements/PCGFilterByIndex.cpp
new file mode 100644
--- /dev/null
+++ b/Engine/Plugins/Experimental/PCG/Source/PCG/Private/Elements/PCGFilterByIndex.cpp
@@ -0,0 +1,136 @@
+// Copyright Epic Games, Inc. All Rights Reserved.
+
+#include "Elements/PCGFilterByIndex.h"
+
+#include "PCGContext.h"
+#include "Data/PCGPointData.h"
+#include "Helpers/Parsing/PCGIndexing.h"
+#include "Helpers/Parsing/PCGParsing.h"
+
+#include UE_INLINE_GENERATED_CPP_BY_NAME(PCGFilterByIndex)
+
+#define LOCTEXT_NAMESPACE "PCGFilterByIndexElement"
+
+namespace PCGFilterByIndexConstants
+{
+	const FName NodeName = FName(TEXT("FilterByIndex"));
+	const FText NodeTitle = LOCTEXT("NodeTitle", "Filter By Index");
+
+	// After so many characters on the node, it will truncate with an ellipsis ...
+	static constexpr int32 IndexExpressionTruncation = 12;
+}
+
+#if WITH_EDITOR
+FName UPCGFilterByIndexSettings::GetDefaultNodeName() const
+{
+	return PCGFilterByIndexConstants::NodeName;
+}
+
+FText UPCGFilterByIndexSettings::GetDefaultNodeTitle() const
+{
+	return PCGFilterByIndexConstants::NodeTitle;
+}
+
+FText UPCGFilterByIndexSettings::GetNodeTooltipText() const
+{
+	return LOCTEXT("FilterByIndexNodeTooltip", "Filters data in the collection according to user selected indices");
+}
+#endif
+
+EPCGDataType UPCGFilterByIndexSettings::GetCurrentPinTypes(const UPCGPin* InPin) const
+{
+	check(InPin);
+
+	if (!InPin->IsOutputPin())
+	{
+		return Super::GetCurrentPinTypes(InPin);
+	}
+
+	// Output pin narrows to union of inputs on first pin
+	const EPCGDataType InputTypeUnion = GetTypeUnionOfIncidentEdges(PCGPinConstants::DefaultInputLabel);
+	return InputTypeUnion != EPCGDataType::None ? InputTypeUnion : EPCGDataType::Any;
+}
+
+FName UPCGFilterByIndexSettings::AdditionalTaskName() const
+{
+	FString NodeName = PCGFilterByIndexConstants::NodeTitle.ToString();
+	NodeName += TEXT(" : ");
+	NodeName += SelectedIndices.Len() <= PCGFilterByIndexConstants::IndexExpressionTruncation ? SelectedIndices : SelectedIndices.Left(PCGFilterByIndexConstants::IndexExpressionTruncation - 3) + TEXT("...");
+	return FName(NodeName);
+}
+
+TArray<FPCGPinProperties> UPCGFilterByIndexSettings::OutputPinProperties() const
+{
+	TArray<FPCGPinProperties> PinProperties;
+	PinProperties.Emplace(PCGPinConstants::DefaultInFilterLabel, EPCGDataType::Any);
+	PinProperties.Emplace(PCGPinConstants::DefaultOutFilterLabel, EPCGDataType::Any);
+
+	return PinProperties;
+}
+
+FPCGElementPtr UPCGFilterByIndexSettings::CreateElement() const
+{
+	return MakeShared<FPCGFilterByIndexElement>();
+}
+
+bool FPCGFilterByIndexElement::ExecuteInternal(FPCGContext* Context) const
+{
+	TRACE_CPUPROFILER_EVENT_SCOPE(FPCGFilterByIndexElement::ExecuteInternal);
+
+	const UPCGFilterByIndexSettings* Settings = Context->GetInputSettings<UPCGFilterByIndexSettings>();
+	check(Settings);
+
+	if (Context->InputData.GetInputs().IsEmpty())
+	{
+		return true;
+	}
+
+	PCGIndexing::FPCGIndexCollection FilteredIndices(Context->InputData.GetInputs().Num());
+	// Parse the indices and switch through possible issues.
+	switch (PCGParser::ParseIndexRanges(FilteredIndices, Settings->SelectedIndices))
+	{
+		// Error cases are caught for early out.
+		case PCGParser::EPCGParserResult::InvalidCharacter:
+			PCGE_LOG_C(Error, GraphAndLog, Context, LOCTEXT("ErrorInvalidCharacter", "Invalid character in index selection string."));
+			return true;
+
+		case PCGParser::EPCGParserResult::InvalidExpression:
+			PCGE_LOG_C(Error, GraphAndLog, Context, LOCTEXT("ErrorInvalidExpression", "Invalid expression in index selection string."));
+			return true;
+
+		// If the expression is empty, treat it as though there are simply no selected indices. But, also log.
+		case PCGParser::EPCGParserResult::EmptyExpression:
+			PCGE_LOG_C(Log, LogOnly, Context, LOCTEXT("WarningEmptyExpression", "Empty expression in index selection string."));
+			break;
+
+		case PCGParser::EPCGParserResult::Success:
+			break;
+
+		default:
+			checkNoEntry();
+			return true;
+	}
+
+	TArray<FPCGTaggedData>& Outputs = Context->OutputData.TaggedData;
+	Outputs = Context->InputData.GetInputs();
+
+	for (int Index = 0; Index < Outputs.Num(); ++Index)
+	{
+		bool bIndexIsFiltered = FilteredIndices.ContainsIndex(Index);
+
+		bIndexIsFiltered = Settings->bInvertFilter ? !bIndexIsFiltered : bIndexIsFiltered;
+
+		if (bIndexIsFiltered)
+		{
+			Outputs[Index].Pin = PCGPinConstants::DefaultInFilterLabel;
+		}
+		else
+		{
+			Outputs[Index].Pin = PCGPinConstants::DefaultOutFilterLabel;
+		}
+	}
+
+	return true;
+}
+
+#undef LOCTEXT_NAMESPACE
\ No newline at end of file
diff --git a/Engine/Plugins/Experimental/PCG/Source/PCG/Private/Helpers/Parsing/PCGIndexing.cpp b/Engine/Plugins/Experimental/PCG/Source/PCG/Private/Helpers/Parsing/PCGIndexing.cpp
new file mode 100644
--- /dev/null
+++ b/Engine/Plugins/Experimental/PCG/Source/PCG/Private/Helpers/Parsing/PCGIndexing.cpp
@@ -0,0 +1,155 @@
+// Copyright Epic Games, Inc. All Rights Reserved.
+
+#include "Helpers/Parsing/PCGIndexing.h"
+
+namespace PCGIndexing
+{
+	bool FPCGIndexRange::ContainsIndex(const int32 Index) const
+	{
+		return Index >= StartIndex && Index < EndIndex;
+	}
+
+	int32 FPCGIndexRange::GetIndexCount() const
+	{
+		return FMath::Max(EndIndex - StartIndex, 0);
+	}
+
+	bool FPCGIndexRange::operator<(const FPCGIndexRange& OtherRange) const
+	{
+		return StartIndex < OtherRange.StartIndex;
+	}
+
+	bool FPCGIndexRange::operator==(const FPCGIndexRange& OtherRange) const
+	{
+		return StartIndex == OtherRange.StartIndex && EndIndex == OtherRange.EndIndex;
+	}
+
+	bool FPCGIndexRange::operator!=(const FPCGIndexRange& OtherRange) const
+	{
+		return !operator==(OtherRange);
+	}
+
+	bool FPCGIndexCollection::AddRange(const int32 StartIndex, const int32 EndIndex)
+	{
+		const FPCGIndexRange NewRange(AdjustIndex(StartIndex), AdjustIndex(EndIndex));
+
+		if (!RangeIsValid(NewRange))
+		{
+			return false;
+		}
+
+		// First element added, early out
+		if (IndexRanges.IsEmpty())
+		{
+			IndexRanges.Add(NewRange);
+			return true;
+		}
+
+		// Insert after the new start index is larger, so we can start overlap comparing at the index prior
+		int32 InsertionIndex = 0;
+		const int32 PreviousSize = IndexRanges.Num();
+		while (InsertionIndex < PreviousSize && NewRange.StartIndex > IndexRanges[InsertionIndex].StartIndex)
+		{
+			++InsertionIndex;
+		}
+
+		IndexRanges.Insert(NewRange, InsertionIndex);
+
+		// Start of the overlap checking
+		const int32 OverlapIndex = FMath::Max(InsertionIndex - 1, 0);
+
+		// Since they inserted in sorting order, insert and resolve overlaps until there are none
+		do
+		{
+			if (!CheckOverlap(IndexRanges[OverlapIndex], IndexRanges[OverlapIndex + 1]))
+			{
+				break;
+			}
+
+			IndexRanges[OverlapIndex] = MergeRanges(IndexRanges[OverlapIndex], IndexRanges[OverlapIndex + 1]);
+			IndexRanges.RemoveAt(OverlapIndex + 1);
+		}
+		while (OverlapIndex + 1 < IndexRanges.Num());
+
+		return true;
+	}
+
+	bool FPCGIndexCollection::AddRange(const FPCGIndexRange& NewRange)
+	{
+		return AddRange(NewRange.StartIndex, NewRange.EndIndex);
+	}
+
+	bool FPCGIndexCollection::RangeIsValid(const FPCGIndexRange& Range) const
+	{
+		return (Range.StartIndex >= 0 && Range.StartIndex <= ArraySize && Range.EndIndex >= 0 && Range.EndIndex <= ArraySize && Range.StartIndex < Range.EndIndex);
+	}
+
+	bool FPCGIndexCollection::ContainsIndex(const int32 Index) const
+	{
+		for (const FPCGIndexRange& Range : IndexRanges)
+		{
+			if (Range.ContainsIndex(Index))
+			{
+				return true;
+			}
+		}
+
+		return false;
+	}
+
+	int32 FPCGIndexCollection::GetArraySize() const
+	{
+		return ArraySize;
+	}
+
+	int32 FPCGIndexCollection::GetTotalRangeCount() const
+	{
+		return IndexRanges.Num();
+	}
+
+	int32 FPCGIndexCollection::GetTotalIndexCount() const
+	{
+		int32 TotalIndexCount = 0;
+
+		for (const FPCGIndexRange& Range : IndexRanges)
+		{
+			check(Range.GetIndexCount() > 0);
+			TotalIndexCount += Range.GetIndexCount();
+		}
+
+		return TotalIndexCount;
+	}
+
+	bool FPCGIndexCollection::operator==(const FPCGIndexCollection& Other) const
+	{
+		return (ArraySize == Other.ArraySize) && (IndexRanges == Other.IndexRanges);
+	}
+
+	int32 FPCGIndexCollection::AdjustIndex(const int32 Index) const
+	{
+		if (Index >= 0)
+		{
+			return FMath::Min(Index, ArraySize);
+		}
+
+		return FMath::Max(ArraySize + Index, 0);
+	}
+
+	bool FPCGIndexCollection::CheckOverlap(const FPCGIndexRange& FirstRange, const FPCGIndexRange& SecondRange) const
+	{
+		check(RangeIsValid(FirstRange) && RangeIsValid(SecondRange));
+
+		return FirstRange.StartIndex <= SecondRange.EndIndex && SecondRange.StartIndex <= FirstRange.EndIndex;
+	}
+
+	FPCGIndexRange FPCGIndexCollection::MergeRanges(const FPCGIndexRange& FirstRange, const FPCGIndexRange& SecondRange) const
+	{
+		// Safety check
+		check(CheckOverlap(FirstRange, SecondRange));
+
+		const int32 MergedStart = FMath::Min(FirstRange.StartIndex, SecondRange.StartIndex);
+		const int32 MergedEnd = FMath::Max(FirstRange.EndIndex, SecondRange.EndIndex);
+
+		return FPCGIndexRange(MergedStart, MergedEnd);
+	}
+}
diff --git a/Engine/Plugins/Experimental/PCG/Source/PCG/Private/Helpers/Parsing/PCGParsing.cpp b/Engine/Plugins/Experimental/PCG/Source/PCG/Private/Helpers/Parsing/PCGParsing.cpp
new file mode 100644
--- /dev/null
+++ b/Engine/Plugins/Experimental/PCG/Source/PCG/Private/Helpers/Parsing/PCGParsing.cpp
@@ -0,0 +1,81 @@
+// Copyright Epic Games, Inc. All Rights Reserved.
+
+#include "Helpers/Parsing/PCGParsing.h"
+
+#include "Internationalization/Regex.h"
+#include "Templates/ValueOrError.h"
+
+namespace PCGParser
+{
+	namespace Constants
+	{
+		static constexpr const TCHAR* IndexElementToken = TEXT(",");
+		static constexpr const TCHAR* IndexRangeToken = TEXT(":");
+		static constexpr const TCHAR* IndexRangePattern = TEXT("^(-?\\d*):(-?\\d*)$");
+	}
+
+	// TODO: Consider passing error messages as well for more precise feedback
+	EPCGParserResult ParseIndexRanges(PCGIndexing::FPCGIndexCollection& OutIndexCollection, const FString& InputString)
+	{
+		FString PurgedInputString(InputString);
+		PurgedInputString.RemoveSpacesInline();
+
+		TArray<FString> ElementSubstrings;
+		// First parse the string into comma separated array
+		PurgedInputString.ParseIntoArray(ElementSubstrings, Constants::IndexElementToken, /*InCullEmpty=*/true);
+
+		if (ElementSubstrings.IsEmpty())
+		{
+			return EPCGParserResult::EmptyExpression;
+		}
+
+		for (const FString& ElementSubstring : ElementSubstrings)
+		{
+			// Range case with the ':' character
+			if (ElementSubstring.Contains(Constants::IndexRangeToken))
+			{
+				FRegexMatcher RangeMatcher(FRegexPattern(Constants::IndexRangePattern), ElementSubstring);
+				if (RangeMatcher.FindNext())
+				{
+					FString StartIndexString = RangeMatcher.GetCaptureGroup(1);
+					FString EndIndexString = RangeMatcher.GetCaptureGroup(2);
+
+					if ((!StartIndexString.IsEmpty() && !StartIndexString.IsNumeric()) ||
+						(!EndIndexString.IsEmpty() && !EndIndexString.IsNumeric()))
+					{
+						return EPCGParserResult::InvalidCharacter;
+					}
+
+					// Empty start index defaults to the array start
+					const int32 StartRangeIndex = StartIndexString.IsEmpty() ? 0 : FCString::Atoi(*StartIndexString);
+					// Empty end index defaults to the max value, to be truncated down to the array max
+					const int32 EndRangeIndex = EndIndexString.IsEmpty() ? std::numeric_limits<int32>::max() : FCString::Atoi(*EndIndexString);
+
+					if (!OutIndexCollection.AddRange(StartRangeIndex, EndRangeIndex))
+					{
+						return EPCGParserResult::InvalidExpression;
+					}
+				}
+				else
+				{
+					return EPCGParserResult::InvalidExpression;
+				}
+			}
+			else // Single index element case
+			{
+				if (!ElementSubstring.IsNumeric())
+				{
+					return EPCGParserResult::InvalidCharacter;
+				}
+
+				const int32 Index = FCString::Atoi(*ElementSubstring);
+				if (!OutIndexCollection.AddRange(Index, Index + 1))
+				{
+					return EPCGParserResult::InvalidExpression;
+				}
+			}
+		}
+
+		return EPCGParserResult::Success;
+	}
+}
diff --git a/Engine/Plugins/Experimental/PCG/Source/PCG/Private/Tests/Elements/PCGFilterByIndexTest.cpp b/Engine/Plugins/Experimental/PCG/Source/PCG/Private/Tests/Elements/PCGFilterByIndexTest.cpp
new file mode 100644
--- /dev/null
+++ b/Engine/Plugins/Experimental/PCG/Source/PCG/Private/Tests/Elements/PCGFilterByIndexTest.cpp
@@ -0,0 +1,217 @@
+// Copyright Epic Games, Inc. All Rights Reserved.
+
+#include "Tests/PCGTestsCommon.h"
+
+#include "PCGContext.h"
+#include "Data/PCGPointData.h"
+
+#include "Elements/PCGFilterByIndex.h"
+
+namespace PCGFilterByIndexTestConstants
+{
+	static constexpr int32 InputNum = 5;
+}
+
+namespace PCGFilterByIndexTestPrivate
+{
+	TArray<FPCGTaggedData> ResetAndExecute(PCGTestsCommon::FTestData& TestData, UPCGFilterByIndexSettings* Settings)
+	{
+		TestData.Reset(Settings);
+		const FPCGElementPtr TestElement = Settings->GetElement();
+
+		for (int Index = 0; Index < PCGFilterByIndexTestConstants::InputNum; ++Index)
+		{
+			FPCGTaggedData& InputData = TestData.InputData.TaggedData.Emplace_GetRef();
+			InputData.Pin = PCGPinConstants::DefaultInputLabel;
+			InputData.Data = PCGTestsCommon::CreateRandomPointData(Index + 1, 42);
+		}
+
+		TUniquePtr<FPCGContext> Context = TestData.InitializeTestContext();
+
+		while (!TestElement->Execute(Context.Get()))
+		{
+		}
+
+		return Context->OutputData.GetInputs();
+	}
+
+	bool ValidateData(const TArray<FPCGTaggedData>& Outputs)
+	{
+		check(Outputs.Num() == PCGFilterByIndexTestConstants::InputNum);
+
+		for (int I = 0; I < PCGFilterByIndexTestConstants::InputNum; ++I)
+		{
+			const UPCGPointData* OutPointData = Cast<UPCGPointData>(Outputs[I].Data);
+			if (!OutPointData || (OutPointData->GetPoints().Num() != I + 1))
+			{
+				return false;
+			}
+		}
+
+		return true;
+	}
+}
+
+IMPLEMENT_CUSTOM_SIMPLE_AUTOMATION_TEST(FPCGFilterByIndexTest_Basic, FPCGTestBaseClass, "pcg.tests.FilterByIndex.Basic", PCGTestsCommon::TestFlags)
+
+bool FPCGFilterByIndexTest_Basic::RunTest(const FString& Parameters)
+{
+	using PCGFilterByIndexTestConstants::InputNum;
+
+	PCGTestsCommon::FTestData TestData;
+	UPCGFilterByIndexSettings* Settings = PCGTestsCommon::GenerateSettings<UPCGFilterByIndexSettings>(TestData);
+
+	// Empty index input string
+	{
+		Settings->SelectedIndices.Empty();
+		const TArray<FPCGTaggedData> Outputs = PCGFilterByIndexTestPrivate::ResetAndExecute(TestData, Settings);
+
+		UTEST_EQUAL("Output count", Outputs.Num(), InputNum);
+		UTEST_TRUE("Output is valid", PCGFilterByIndexTestPrivate::ValidateData(Outputs));
+
+		UTEST_EQUAL("Output 1 filtered exclusively", Outputs[0].Pin, PCGPinConstants::DefaultOutFilterLabel);
+		UTEST_EQUAL("Output 2 filtered exclusively", Outputs[1].Pin, PCGPinConstants::DefaultOutFilterLabel);
+		UTEST_EQUAL("Output 3 filtered exclusively", Outputs[2].Pin, PCGPinConstants::DefaultOutFilterLabel);
+		UTEST_EQUAL("Output 4 filtered exclusively", Outputs[3].Pin, PCGPinConstants::DefaultOutFilterLabel);
+		UTEST_EQUAL("Output 5 filtered exclusively", Outputs[4].Pin, PCGPinConstants::DefaultOutFilterLabel);
+	}
+
+	// Explicit index selection
+	{
+		Settings->SelectedIndices = FString("0,2,4");
+		const TArray<FPCGTaggedData> Outputs = PCGFilterByIndexTestPrivate::ResetAndExecute(TestData, Settings);
+
+		UTEST_EQUAL("Output count", Outputs.Num(), InputNum);
+		UTEST_TRUE("Output is valid", PCGFilterByIndexTestPrivate::ValidateData(Outputs));
+
+		UTEST_EQUAL("Output 1 filtered inclusively", Outputs[0].Pin, PCGPinConstants::DefaultInFilterLabel);
+		UTEST_EQUAL("Output 2 filtered exclusively", Outputs[1].Pin, PCGPinConstants::DefaultOutFilterLabel);
+		UTEST_EQUAL("Output 3 filtered inclusively", Outputs[2].Pin, PCGPinConstants::DefaultInFilterLabel);
+		UTEST_EQUAL("Output 4 filtered exclusively", Outputs[3].Pin, PCGPinConstants::DefaultOutFilterLabel);
+		UTEST_EQUAL("Output 5 filtered inclusively", Outputs[4].Pin, PCGPinConstants::DefaultInFilterLabel);
+	}
+
+	// Negative index selection. 5 - 2 = 3
+	{
+		Settings->SelectedIndices = FString("-2");
+		const TArray<FPCGTaggedData> Outputs = PCGFilterByIndexTestPrivate::ResetAndExecute(TestData, Settings);
+
+		UTEST_EQUAL("Output count", Outputs.Num(), InputNum);
+		UTEST_TRUE("Output is valid", PCGFilterByIndexTestPrivate::ValidateData(Outputs));
+
+		UTEST_EQUAL("Output 1 filtered exclusively", Outputs[0].Pin, PCGPinConstants::DefaultOutFilterLabel);
+		UTEST_EQUAL("Output 2 filtered exclusively", Outputs[1].Pin, PCGPinConstants::DefaultOutFilterLabel);
+		UTEST_EQUAL("Output 3 filtered exclusively", Outputs[2].Pin, PCGPinConstants::DefaultOutFilterLabel);
+		UTEST_EQUAL("Output 4 filtered inclusively", Outputs[3].Pin, PCGPinConstants::DefaultInFilterLabel);
+		UTEST_EQUAL("Output 5 filtered exclusively", Outputs[4].Pin, PCGPinConstants::DefaultOutFilterLabel);
+	}
+
+	// Range index selection [1:3) = [1:2] since range slicing doesn't include last index
+	{
+		Settings->SelectedIndices = FString("1:3");
+		const TArray<FPCGTaggedData> Outputs = PCGFilterByIndexTestPrivate::ResetAndExecute(TestData, Settings);
+
+		UTEST_EQUAL("Output count", Outputs.Num(), InputNum);
+		UTEST_TRUE("Output is valid", PCGFilterByIndexTestPrivate::ValidateData(Outputs));
+
+		UTEST_EQUAL("Output 1 filtered exclusively", Outputs[0].Pin, PCGPinConstants::DefaultOutFilterLabel);
+		UTEST_EQUAL("Output 2 filtered inclusively", Outputs[1].Pin, PCGPinConstants::DefaultInFilterLabel);
+		UTEST_EQUAL("Output 3 filtered inclusively", Outputs[2].Pin, PCGPinConstants::DefaultInFilterLabel);
+		UTEST_EQUAL("Output 4 filtered exclusively", Outputs[3].Pin, PCGPinConstants::DefaultOutFilterLabel);
+		UTEST_EQUAL("Output 5 filtered exclusively", Outputs[4].Pin, PCGPinConstants::DefaultOutFilterLabel);
+	}
+
+	// Range index selection (should be [5-3,5-1) = [2,4) = [2:3]
+	{
+		Settings->SelectedIndices = FString("-3:-1");
+		const TArray<FPCGTaggedData> Outputs = PCGFilterByIndexTestPrivate::ResetAndExecute(TestData, Settings);
+
+		UTEST_EQUAL("Output count", Outputs.Num(), InputNum);
+		UTEST_TRUE("Output is valid", PCGFilterByIndexTestPrivate::ValidateData(Outputs));
+
+		UTEST_EQUAL("Output 1 filtered exclusively", Outputs[0].Pin, PCGPinConstants::DefaultOutFilterLabel);
+		UTEST_EQUAL("Output 2 filtered exclusively", Outputs[1].Pin, PCGPinConstants::DefaultOutFilterLabel);
+		UTEST_EQUAL("Output 3 filtered inclusively", Outputs[2].Pin, PCGPinConstants::DefaultInFilterLabel);
+		UTEST_EQUAL("Output 4 filtered inclusively", Outputs[3].Pin, PCGPinConstants::DefaultInFilterLabel);
+		UTEST_EQUAL("Output 5 filtered exclusively", Outputs[4].Pin, PCGPinConstants::DefaultOutFilterLabel);
+	}
+
+	// Combination index selection
+	{
+		Settings->SelectedIndices = FString("0,2:3,4");
+		const TArray<FPCGTaggedData> Outputs = PCGFilterByIndexTestPrivate::ResetAndExecute(TestData, Settings);
+
+		UTEST_EQUAL("Output count", Outputs.Num(), InputNum);
+		UTEST_TRUE("Output is valid", PCGFilterByIndexTestPrivate::ValidateData(Outputs));
+
+		UTEST_EQUAL("Output 1 filtered inclusively", Outputs[0].Pin, PCGPinConstants::DefaultInFilterLabel);
+		UTEST_EQUAL("Output 2 filtered exclusively", Outputs[1].Pin, PCGPinConstants::DefaultOutFilterLabel);
+		UTEST_EQUAL("Output 3 filtered inclusively", Outputs[2].Pin, PCGPinConstants::DefaultInFilterLabel);
+		UTEST_EQUAL("Output 4 filtered exclusively", Outputs[3].Pin, PCGPinConstants::DefaultOutFilterLabel);
+		UTEST_EQUAL("Output 5 filtered inclusively", Outputs[4].Pin, PCGPinConstants::DefaultInFilterLabel);
+	}
+
+	// Combination index selection, with inverted filter
+	{
+		Settings->bInvertFilter = true;
+		Settings->SelectedIndices = FString("0,2:3,4");
+		const TArray<FPCGTaggedData> Outputs = PCGFilterByIndexTestPrivate::ResetAndExecute(TestData, Settings);
+
+		UTEST_EQUAL("Output count", Outputs.Num(), InputNum);
+		UTEST_TRUE("Output is valid", PCGFilterByIndexTestPrivate::ValidateData(Outputs));
+
+		UTEST_EQUAL("Output 1 filtered exclusively", Outputs[0].Pin, PCGPinConstants::DefaultOutFilterLabel);
+		UTEST_EQUAL("Output 2 filtered inclusively", Outputs[1].Pin, PCGPinConstants::DefaultInFilterLabel);
+		UTEST_EQUAL("Output 3 filtered exclusively", Outputs[2].Pin, PCGPinConstants::DefaultOutFilterLabel);
+		UTEST_EQUAL("Output 4 filtered inclusively", Outputs[3].Pin, PCGPinConstants::DefaultInFilterLabel);
+		UTEST_EQUAL("Output 5 filtered exclusively", Outputs[4].Pin, PCGPinConstants::DefaultOutFilterLabel);
+	}
+
+	return true;
+}
+
+IMPLEMENT_CUSTOM_SIMPLE_AUTOMATION_TEST(FPCGFilterByIndexTest_InvalidSelection, FPCGTestBaseClass, "pcg.tests.FilterByIndex.InvalidSelection", PCGTestsCommon::TestFlags)
+
+bool FPCGFilterByIndexTest_InvalidSelection::RunTest(const FString& Parameters)
+{
+	PCGTestsCommon::FTestData TestData;
+	UPCGFilterByIndexSettings* Settings = PCGTestsCommon::GenerateSettings<UPCGFilterByIndexSettings>(TestData);
+
+	AddExpectedError("Invalid expression in index selection string", EAutomationExpectedMessageFlags::Contains, 3);
+
+	// Test inverted range
+	{
+		Settings->SelectedIndices = FString("4:0");
+		const TArray<FPCGTaggedData> Outputs = PCGFilterByIndexTestPrivate::ResetAndExecute(TestData, Settings);
+
+		UTEST_TRUE("Output count", Outputs.IsEmpty());
+	}
+
+	// Test inverted range with negative indices
+	{
+		Settings->SelectedIndices = FString("-1:-3");
+		const TArray<FPCGTaggedData> Outputs = PCGFilterByIndexTestPrivate::ResetAndExecute(TestData, Settings);
+
+		UTEST_TRUE("Output count", Outputs.IsEmpty());
+	}
+
+	// Duplicate selection for range
+	{
+		Settings->SelectedIndices = FString("2:2");
+		const TArray<FPCGTaggedData> Outputs = PCGFilterByIndexTestPrivate::ResetAndExecute(TestData, Settings);
+
+		UTEST_TRUE("Output count", Outputs.IsEmpty());
+	}
+
+	AddExpectedError("Invalid character in index selection string");
+
+	// Test invalid character
+	{
+		Settings->SelectedIndices = FString("0,2:3,4*");
+		const TArray<FPCGTaggedData> Outputs = PCGFilterByIndexTestPrivate::ResetAndExecute(TestData, Settings);
+
+		UTEST_TRUE("Output count", Outputs.IsEmpty());
+	}
+
+	return true;
+}
\ No newline at end of file
diff --git a/Engine/Plugins/Experimental/PCG/Source/PCG/Private/Tests/Parsing/PCGParsingTest.cpp b/Engine/Plugins/Experimental/PCG/Source/PCG/Private/Tests/Parsing/PCGParsingTest.cpp
new file mode 100644
--- /dev/null
+++ b/Engine/Plugins/Experimental/PCG/Source/PCG/Private/Tests/Parsing/PCGParsingTest.cpp
@@ -0,0 +1,287 @@
+// Copyright Epic Games, Inc. All Rights Reserved.
+
+#include "Helpers/Parsing/PCGParsing.h"
+#include "Tests/PCGTestsCommon.h"
+
+namespace PCGParsingTestConstants
+{
+	// Arbitrary constant used to determine the hypothetical array size
+	static constexpr int32 ArraySize = 512;
+}
+
+namespace PCGParsingTestPrivate
+{
+	PCGIndexing::FPCGIndexCollection ParseAndRetrieveIndices(const FString& Expression, const int32 ArraySize)
+	{
+		PCGIndexing::FPCGIndexCollection IndexCollection(ArraySize);
+		PCGParser::ParseIndexRanges(IndexCollection, Expression);
+		return IndexCollection;
+	}
+}
+
+IMPLEMENT_CUSTOM_SIMPLE_AUTOMATION_TEST(FPCGParsingTest_Basic, FPCGTestBaseClass, "pcg.tests.Parsing.Basic", PCGTestsCommon::TestFlags)
+
+bool FPCGParsingTest_Basic::RunTest(const FString& Parameters)
+{
+	using PCGParsingTestConstants::ArraySize;
+
+	// Simple index
+	{
+		const PCGIndexing::FPCGIndexCollection Results = PCGParsingTestPrivate::ParseAndRetrieveIndices("4", ArraySize);
+		UTEST_EQUAL("Singular non-range index has 1 index", Results.GetTotalIndexCount(), 1);
+		UTEST_EQUAL("Singular non-range index has 1 range", Results.GetTotalRangeCount(), 1);
+		UTEST_TRUE("Singular non-range index is correct", Results.ContainsIndex(4));
+	}
+
+	// Simple negative index
+	{
+		const PCGIndexing::FPCGIndexCollection Results = PCGParsingTestPrivate::ParseAndRetrieveIndices("-4", ArraySize);
+		UTEST_EQUAL("Negative non-range index has 1 index", Results.GetTotalIndexCount(), 1);
+		UTEST_EQUAL("Negative non-range index has 1 range", Results.GetTotalRangeCount(), 1);
+		UTEST_TRUE("Negative non-range index is correct", Results.ContainsIndex(ArraySize - 4));
+	}
+
+	// Two simple indices
+	{
+		const PCGIndexing::FPCGIndexCollection Results = PCGParsingTestPrivate::ParseAndRetrieveIndices("4,-4", ArraySize);
+		UTEST_EQUAL("Multiple index has 2 index", Results.GetTotalIndexCount(), 2);
+		UTEST_EQUAL("Multiple index has 2 ranges", Results.GetTotalRangeCount(), 2);
+		UTEST_TRUE("Multiple index (first) is correct", Results.ContainsIndex(4) && Results.ContainsIndex(ArraySize - 4));
+	}
+
+	// More complex use case
+	{
+		TArray<int32> TestArray;
+		TestArray.SetNumUninitialized(ArraySize);
+		for (int I = 0; I < ArraySize; ++I)
+		{
+			TestArray[I] = I;
+		}
+
+		// Should be 0,2,3,4,5 and from 80 to 502
+		const PCGIndexing::FPCGIndexCollection ForwardsResults = PCGParsingTestPrivate::ParseAndRetrieveIndices("0,2:5,80:-10", ArraySize);
+		const PCGIndexing::FPCGIndexCollection BackwardsResults = PCGParsingTestPrivate::ParseAndRetrieveIndices("80:-10,2:5,0", ArraySize);
+
+		UTEST_TRUE("Valid single index", ForwardsResults.ContainsIndex(0));
+		UTEST_TRUE("Valid single index", BackwardsResults.ContainsIndex(0));
+		UTEST_TRUE("Missing single index", !ForwardsResults.ContainsIndex(1));
+		UTEST_TRUE("Missing single index", !BackwardsResults.ContainsIndex(1));
+
+		for (int I = 2; I < 5; ++I)
+		{
+			UTEST_TRUE("Valid range index", ForwardsResults.ContainsIndex(I));
+			UTEST_TRUE("Valid range index", BackwardsResults.ContainsIndex(I));
+		}
+
+		for (int I = 6; I < 80; ++I)
+		{
+			UTEST_TRUE("Missing range index", !ForwardsResults.ContainsIndex(I));
+			UTEST_TRUE("Missing range index", !BackwardsResults.ContainsIndex(I));
+		}
+
+		for (int I = 80; I < ArraySize - 10; ++I)
+		{
+			UTEST_TRUE("Valid range negative index", ForwardsResults.ContainsIndex(I));
+			UTEST_TRUE("Valid range negative index", BackwardsResults.ContainsIndex(I));
+		}
+
+		for (int I = ArraySize - 9; I < ArraySize; ++I)
+		{
+			UTEST_TRUE("Missing range negative index", !ForwardsResults.ContainsIndex(I));
+			UTEST_TRUE("Missing range negative index", !BackwardsResults.ContainsIndex(I));
+		}
+	}
+
+	return true;
+}
+
+IMPLEMENT_CUSTOM_SIMPLE_AUTOMATION_TEST(FPCGParsingTest_Ranges, FPCGTestBaseClass, "pcg.tests.Parsing.Ranges", PCGTestsCommon::TestFlags)
+
+bool FPCGParsingTest_Ranges::RunTest(const FString& Parameters)
+{
+	using PCGParsingTestConstants::ArraySize;
+	PCGIndexing::FPCGIndexCollection Results(ArraySize);
+
+	auto CheckSplitTestArray = [this, &Results](const int32 SplitIndex)
+	{
+		for (int I = 0; I < ArraySize; ++I)
+		{
+			if (I < SplitIndex != Results.ContainsIndex(I))
+			{
+				return false;
+			}
+		}
+
+		return true;
+	};
+
+	// Basic ranges
+	// ------------
+
+	Results = PCGParsingTestPrivate::ParseAndRetrieveIndices("0:1", ArraySize);
+	UTEST_EQUAL("Repeated single index results in a single range", Results.GetTotalRangeCount(), 1);
+	UTEST_EQUAL("Repeated single index results in a single index", Results.GetTotalIndexCount(), 1);
+	UTEST_TRUE("Included indices verified", CheckSplitTestArray(1));
+
+	Results = PCGParsingTestPrivate::ParseAndRetrieveIndices("0:10", ArraySize);
+	UTEST_EQUAL("Single index range results in a single range", Results.GetTotalRangeCount(), 1);
+	UTEST_EQUAL("Single index range results in correct number of indices", Results.GetTotalIndexCount(), 10);
+	UTEST_TRUE("Included indices verified", CheckSplitTestArray(10));
+
+	Results = PCGParsingTestPrivate::ParseAndRetrieveIndices(":10", ArraySize);
+	UTEST_EQUAL("Inferred start index results in a single range", Results.GetTotalRangeCount(), 1);
+	UTEST_EQUAL("Inferred start index results in correct number of indices", Results.GetTotalIndexCount(), 10);
+	UTEST_TRUE("Included indices verified", CheckSplitTestArray(10));
+
+	Results = PCGParsingTestPrivate::ParseAndRetrieveIndices(":", ArraySize);
+	UTEST_EQUAL("Full range has single range", Results.GetTotalRangeCount(), 1);
+	UTEST_EQUAL("Full range has all indices", Results.GetTotalIndexCount(), ArraySize);
+	UTEST_TRUE("Included indices verified", CheckSplitTestArray(ArraySize));
+
+	Results = PCGParsingTestPrivate::ParseAndRetrieveIndices("0:", ArraySize);
+	UTEST_EQUAL("Inferred end index results in a single range", Results.GetTotalRangeCount(), 1);
+	UTEST_EQUAL("Inferred end index results in correct number of indices", Results.GetTotalIndexCount(), ArraySize);
+	UTEST_TRUE("Included indices verified", CheckSplitTestArray(ArraySize));
+
+	Results = PCGParsingTestPrivate::ParseAndRetrieveIndices(":-4", ArraySize);
+	UTEST_EQUAL("Inferred start index, negative end index results in a single range", Results.GetTotalRangeCount(), 1);
+	UTEST_EQUAL("Inferred start index, negative end results in correct number of indices", Results.GetTotalIndexCount(), ArraySize - 4);
+	UTEST_TRUE("Included indices verified", CheckSplitTestArray(ArraySize - 4));
+
+	// Complete overlap
+	// ----------------
+
+	Results = PCGParsingTestPrivate::ParseAndRetrieveIndices("0:10,2:5", ArraySize);
+	UTEST_EQUAL("Complete preceeding overlap collapsed into single index range", Results.GetTotalRangeCount(), 1);
+	UTEST_EQUAL("Complete preceeding overlap collapsed into correct index count", Results.GetTotalIndexCount(), 10);
+	UTEST_TRUE("Included indices verified", CheckSplitTestArray(10));
+
+	Results = PCGParsingTestPrivate::ParseAndRetrieveIndices("2:5,0:10", ArraySize);
+	UTEST_EQUAL("Complete subsequent overlap collapsed into single index range", Results.GetTotalRangeCount(), 1);
+	UTEST_EQUAL("Complete subsequent overlap collapsed into correct index count", Results.GetTotalIndexCount(), 10);
+	UTEST_TRUE("Included indices verified", CheckSplitTestArray(10));
+
+	Results = PCGParsingTestPrivate::ParseAndRetrieveIndices("0,1,2,3:50,0:512,80:250", ArraySize);
+	UTEST_EQUAL("Multiple complete overlaps collapsed into single index range", Results.GetTotalRangeCount(), 1);
+	UTEST_EQUAL("Multiple complete overlaps collapsed into correct index count", Results.GetTotalIndexCount(), 512);
+	UTEST_TRUE("Included indices verified", CheckSplitTestArray(ArraySize));
+
+	// Partial overlap
+	// ---------------
+
+	Results = PCGParsingTestPrivate::ParseAndRetrieveIndices("0:50,25:75", ArraySize);
+	UTEST_EQUAL("Single partial overlap is collapsed with single range", Results.GetTotalRangeCount(), 1);
+	UTEST_EQUAL("Single partial overlap is collapsed with correct index count", Results.GetTotalIndexCount(), 75);
+	UTEST_TRUE("Included indices verified", CheckSplitTestArray(75));
+
+	Results = PCGParsingTestPrivate::ParseAndRetrieveIndices("0:50,25:75,1:5,13:72,30:40,40:50", ArraySize);
+	UTEST_EQUAL("Multiple partial overlap is collapsed with single range", Results.GetTotalRangeCount(), 1);
+	UTEST_EQUAL("Multiple partial overlap is collapsed with correct index count", Results.GetTotalIndexCount(), 75);
+	UTEST_TRUE("Included indices verified", CheckSplitTestArray(75));
+
+	Results = PCGParsingTestPrivate::ParseAndRetrieveIndices("0:2,1:3,2:4,3:5,4:6,5:7", ArraySize);
+	UTEST_EQUAL("Multiple partial overlap is collapsed with single range", Results.GetTotalRangeCount(), 1);
+	UTEST_EQUAL("Multiple partial overlap is collapsed with correct index count", Results.GetTotalIndexCount(), 7);
+	UTEST_TRUE("Included indices verified", CheckSplitTestArray(7));
+
+	Results = PCGParsingTestPrivate::ParseAndRetrieveIndices("0:26,26:50", ArraySize);
+	UTEST_EQUAL("Multiple partial overlap is collapsed with one range", Results.GetTotalRangeCount(), 1);
+	UTEST_EQUAL("Multiple partial overlap is collapsed with correct index count", Results.GetTotalIndexCount(), 50);
+	UTEST_TRUE("Included indices verified", CheckSplitTestArray(50));
+
+	Results = PCGParsingTestPrivate::ParseAndRetrieveIndices("0:20,10:31,31:51,41:62,62:82,72:93", ArraySize);
+	UTEST_EQUAL("Multiple partial overlap is collapsed with one range", Results.GetTotalRangeCount(), 1);
+	UTEST_EQUAL("Multiple partial overlap is collapsed with correct index count", Results.GetTotalIndexCount(), 93);
+	UTEST_TRUE("Included indices verified", CheckSplitTestArray(93));
+
+	return true;
+}
+
+IMPLEMENT_CUSTOM_SIMPLE_AUTOMATION_TEST(FPCGParsingTest_Robustness, FPCGTestBaseClass, "pcg.tests.Parsing.Robustness", PCGTestsCommon::TestFlags)
+
+bool FPCGParsingTest_Robustness::RunTest(const FString& Parameters)
+{
+	using PCGParsingTestConstants::ArraySize;
+	PCGIndexing::FPCGIndexCollection Results(ArraySize);
+
+	// Copy
+	{
+		PCGIndexing::FPCGIndexCollection Original(42);
+		Original.AddRange(21, 41);
+		PCGIndexing::FPCGIndexCollection Copy(Original);
+		UTEST_EQUAL("Copy construction", Original, Copy);
+	}
+
+	// Assignment
+	{
+		PCGIndexing::FPCGIndexCollection Original(42);
+		Original.AddRange(21, 41);
+		PCGIndexing::FPCGIndexCollection Copy = Original;
+		UTEST_EQUAL("Assignment operator", Original, Copy);
+	}
+
+	// Forgivable user mistakes
+	// ------------------------
+
+	Results = PCGParsingTestPrivate::ParseAndRetrieveIndices("00000000000000000", ArraySize);
+	UTEST_EQUAL("Repeated 0s results in 0 index only", Results.GetTotalIndexCount(), 1);
+
+	Results = PCGParsingTestPrivate::ParseAndRetrieveIndices("1:2", ArraySize);
+	UTEST_EQUAL("Single index range created", Results.GetTotalIndexCount(), 1);
+
+	Results = PCGParsingTestPrivate::ParseAndRetrieveIndices("2,", ArraySize);
+	UTEST_EQUAL("Empty value ignored, after", Results.GetTotalIndexCount(), 1);
+
+	Results = PCGParsingTestPrivate::ParseAndRetrieveIndices(",2", ArraySize);
+	UTEST_EQUAL("Empty value ignored, before", Results.GetTotalIndexCount(), 1);
+
+	Results = PCGParsingTestPrivate::ParseAndRetrieveIndices(",,,,2,,,", ArraySize);
+	UTEST_EQUAL("Empty value ignored, multiple", Results.GetTotalIndexCount(), 1);
+
+	Results = PCGParsingTestPrivate::ParseAndRetrieveIndices("0,*", ArraySize);
+	UTEST_EQUAL("Invalid character has no result", Results.GetTotalIndexCount(), 1);
+
+	// Unforgivable mistakes
+	// ---------------------
+
+	Results = PCGParsingTestPrivate::ParseAndRetrieveIndices("", ArraySize);
+	UTEST_EQUAL("Empty expression has no result", Results.GetTotalIndexCount(), 0);
+
+	Results = PCGParsingTestPrivate::ParseAndRetrieveIndices(",", ArraySize);
+	UTEST_EQUAL("Element delimiter only has no result", Results.GetTotalIndexCount(), 0);
+
+	Results = PCGParsingTestPrivate::ParseAndRetrieveIndices("#", ArraySize);
+	UTEST_EQUAL("Invalid character has no result", Results.GetTotalIndexCount(), 0);
+
+	Results = PCGParsingTestPrivate::ParseAndRetrieveIndices("0(,5:6", ArraySize);
+	UTEST_EQUAL("Invalid character has no result", Results.GetTotalIndexCount(), 0);
+
+	Results = PCGParsingTestPrivate::ParseAndRetrieveIndices("3:", ArraySize);
+	UTEST_EQUAL("Range end index missing has no result", Results.GetTotalIndexCount(), ArraySize - 3);
+
+	Results = PCGParsingTestPrivate::ParseAndRetrieveIndices("0,1,2,5:2", ArraySize);
+	UTEST_EQUAL("Inverted range is excluded in the results", Results.GetTotalIndexCount(), 3);
+
+	Results = PCGParsingTestPrivate::ParseAndRetrieveIndices("5:2", ArraySize);
+	UTEST_EQUAL("Inverted range has no result", Results.GetTotalIndexCount(), 0);
+
+	Results = PCGParsingTestPrivate::ParseAndRetrieveIndices(FString::FromInt(ArraySize + 1), ArraySize);
+	UTEST_EQUAL("Index outside the bounds of the array", Results.GetTotalIndexCount(), 0);
+
+	Results = PCGParsingTestPrivate::ParseAndRetrieveIndices("65465465498475309485730495873409587654684986146162342348438", std::numeric_limits<int32>::max());
+	UTEST_EQUAL("Large value converts to int32 max which is the max size of the array which fails with no result", Results.GetTotalIndexCount(), 0);
+
+	Results = PCGParsingTestPrivate::ParseAndRetrieveIndices(":::::::::", ArraySize);
+	UTEST_EQUAL("Range delimiter only has no result", Results.GetTotalIndexCount(), 0);
+
+	Results = PCGParsingTestPrivate::ParseAndRetrieveIndices("1:::::::::", ArraySize);
+	UTEST_EQUAL("Multiple range delimiter with start index has no result", Results.GetTotalIndexCount(), 0);
+
+	Results = PCGParsingTestPrivate::ParseAndRetrieveIndices("::::::::::1", ArraySize);
+	UTEST_EQUAL("Multiple range delimiter with end index has no result", Results.GetTotalIndexCount(), 0);
+
+	Results = PCGParsingTestPrivate::ParseAndRetrieveIndices("1::::::::::1", ArraySize);
+	UTEST_EQUAL("Redundant delimiter between indices has no result", Results.GetTotalIndexCount(), 0);
+
+	return true;
+}
\ No newline at end of file
diff --git a/Engine/Plugins/Experimental/PCG/Source/PCG/Public/Elements/PCGFilterByIndex.h b/Engine/Plugins/Experimental/PCG/Source/PCG/Public/Elements/PCGFilterByIndex.h
new file mode 100644
--- /dev/null
+++ b/Engine/Plugins/Experimental/PCG/Source/PCG/Public/Elements/PCGFilterByIndex.h
@@ -0,0 +1,48 @@
+// Copyright Epic Games, Inc. All Rights Reserved.
+
+#pragma once
+
+#include "PCGSettings.h"
+
+#include "PCGFilterByIndex.generated.h"
+
+/** Filters a data collection based on a user defined index range expression. */
+UCLASS(BlueprintType, ClassGroup = (Procedural))
+class PCG_API UPCGFilterByIndexSettings : public UPCGSettings
+{
+	GENERATED_BODY()
+
+public:
+	//~Begin UPCGSettings interface
+#if WITH_EDITOR
+	virtual FName GetDefaultNodeName() const override;
+	virtual FText GetDefaultNodeTitle() const override;
+	virtual FText GetNodeTooltipText() const override;
+	virtual EPCGSettingsType GetType() const override { return EPCGSettingsType::Filter; }
+	virtual bool HasDynamicPins() const override { return true; }
+#endif
+	virtual EPCGDataType GetCurrentPinTypes(const UPCGPin* InPin) const override;
+	virtual FName AdditionalTaskName() const override;
+
+protected:
+	virtual TArray<FPCGPinProperties> OutputPinProperties() const override;
+	virtual FPCGElementPtr CreateElement() const override;
+	//~End UPCGSettings interface
+
+public:
+	/** Will invert which indices will be included and excluded. */
+	UPROPERTY(BlueprintReadWrite, EditAnywhere, Category = Settings, meta = (PCG_Overridable))
+	bool bInvertFilter = false;
+
+	/**
+	 * Selected individual indices or index ranges to include or exclude. Negative end indices allowed.
+	 * For example, on an array of size 10: '0,2,4:5,7:-1' will include indices: 0,2,4,7,8 */
+	UPROPERTY(BlueprintReadWrite, EditAnywhere, Category = Settings, meta = (PCG_Overridable))
+	FString SelectedIndices;
+};
+
+class FPCGFilterByIndexElement : public IPCGElement
+{
+protected:
+	virtual bool ExecuteInternal(FPCGContext* Context) const override;
+};
diff --git a/Engine/Plugins/Experimental/PCG/Source/PCG/Public/Helpers/Parsing/PCGIndexing.h b/Engine/Plugins/Experimental/PCG/Source/PCG/Public/Helpers/Parsing/PCGIndexing.h
new file mode 100644
--- /dev/null
+++ b/Engine/Plugins/Experimental/PCG/Source/PCG/Public/Helpers/Parsing/PCGIndexing.h
@@ -0,0 +1,79 @@
+// Copyright Epic Games, Inc. All Rights Reserved.
+
+#pragma once
+
+#include "CoreMinimal.h"
+
+/** Helper namespace for handling index based evaluations. */
+namespace PCGIndexing
+{
+	/** A simple data structure to represent a range of indices [X,Y]. Validating ranges are client responsibility. */
+	struct FPCGIndexRange
+	{
+		FPCGIndexRange(const int32 InStartIndex, const int32 InEndIndex) : StartIndex(InStartIndex), EndIndex(InEndIndex)
+		{
+		}
+
+		/** Returns true if the index can be found in this range. */
+		bool ContainsIndex(int32 Index) const;
+
+		/** Returns the total discrete indices within the range. */
+		int32 GetIndexCount() const;
+
+		bool operator<(const FPCGIndexRange& OtherRange) const;
+		bool operator==(const FPCGIndexRange& OtherRange) const;
+		bool operator!=(const FPCGIndexRange& OtherRange) const;
+
+		int32 StartIndex;
+		int32 EndIndex;
+	};
+
+	/** An abstract collection of FPCGIndexRange data that represents a concrete set of indices. */
+	class FPCGIndexCollection
+	{
+	public:
+		/** The constructor must accept the size of the array to support negative terminating indices. */
+		explicit FPCGIndexCollection(const int32 InArraySize) : ArraySize(InArraySize)
+		{
+			check(InArraySize > 0);
+		}
+
+		/** Add a new index range to the collection directly via start and end indices. */
+		bool AddRange(int32 StartIndex, int32 EndIndex);
+
+		/** Add a new index range to the collection directly via an index range. */
+		bool AddRange(const FPCGIndexRange& NewRange);
+
+		/** Validates that a range is acceptable for this collection. Returns true if valid, false otherwise */
+		bool RangeIsValid(const FPCGIndexRange& Range) const;
+
+		/** Returns true if this collection contains the given index. */
+		bool ContainsIndex(int32 Index) const;
+
+		/** Returns the abstract size of the array tied to the collection. */
+		int32 GetArraySize() const;
+
+		/** Gets the number of range structures currently in the collection. */
+		int32 GetTotalRangeCount() const;
+
+		/** Gets the total number of concrete indices within the collection. */
+		int32 GetTotalIndexCount() const;
+
+		bool operator==(const FPCGIndexCollection& Other) const;
+
+	private:
+		int32 AdjustIndex(int32 Index) const;
+
+		/** Checks two ranges for an overlap and returns true if they overlap. */
+		bool CheckOverlap(const FPCGIndexRange& FirstRange, const FPCGIndexRange& SecondRange) const;
+
+		/** Returns the resultant range after merging two overlapping ranges. */
+		FPCGIndexRange MergeRanges(const FPCGIndexRange& FirstRange, const FPCGIndexRange& SecondRange) const;
+
+		/** The size of the representative array associated with this collection. Needed for negative terminating indices. */
+		int32 ArraySize = 0;
+
+		/** A collection of the abstract index ranges in the collection. */
+		TArray<FPCGIndexRange> IndexRanges;
+	};
+}
diff --git a/Engine/Plugins/Experimental/PCG/Source/PCG/Public/Helpers/Parsing/PCGParsing.h b/Engine/Plugins/Experimental/PCG/Source/PCG/Public/Helpers/Parsing/PCGParsing.h
new file mode 100644
--- /dev/null
+++ b/Engine/Plugins/Experimental/PCG/Source/PCG/Public/Helpers/Parsing/PCGParsing.h
@@ -0,0 +1,27 @@
+// Copyright Epic Games, Inc. All Rights Reserved.
+
+#pragma once
+
+#include "PCGIndexing.h"
+
+namespace PCGIndexing
+{
+	class FPCGIndexCollection;
+}
+
+/** This helper namespace to be used for helper functions dedicated to parsing various input strings. */
+namespace PCGParser
+{
+	enum class EPCGParserResult
+	{
+		Success,
+		InvalidCharacter,
+		InvalidExpression,
+		EmptyExpression,
+	};
+
+	/** Indices should be included individually or with ranges using a delimiter. Negative terminating ranges are accepted.
+	* For example, on an array of size 10: '0,2,4:5,7:-1' will include indices: 0,2,4,5,7,8
+	*/
+	PCG_API EPCGParserResult ParseIndexRanges(PCGIndexing::FPCGIndexCollection& OutIndexCollection, const FString& InputString);
+}
diff --git a/Engine/Plugins/Experimental/PCG/Source/PCG/Public/PCGCommon.h b/Engine/Plugins/Experimental/PCG/Source/PCG/Public/PCGCommon.h
--- a/Engine/Plugins/Experimental/PCG/Source/PCG/Public/PCGCommon.h
+++ b/Engine/Plugins/Experimental/PCG/Source/PCG/Public/PCGCommon.h
@@ -93,6 +93,9 @@ namespace PCGPinConstants
 	const FName DefaultParamsLabel = TEXT("Overrides");
 	const FName DefaultDependencyOnlyLabel = TEXT("Dependency Only");
 
+	const FName DefaultInFilterLabel = TEXT("InsideFilter");
+	const FName DefaultOutFilterLabel = TEXT("OutsideFilter");
+
 namespace Private
 {
 	const FName OldDefaultParamsLabel = TEXT("Params");
