Added param data to points conversion and points to param data

diff --git a/Engine/Plugins/Experimental/PCG/Source/PCG/Private/Elements/PCGCollapseElement.cpp b/Engine/Plugins/Experimental/PCG/Source/PCG/Private/Elements/PCGCollapseElement.cpp
--- a/Engine/Plugins/Experimental/PCG/Source/PCG/Private/Elements/PCGCollapseElement.cpp
+++ b/Engine/Plugins/Experimental/PCG/Source/PCG/Private/Elements/PCGCollapseElement.cpp
@@ -3,6 +3,7 @@
 #include "Elements/PCGCollapseElement.h"
 
 #include "PCGContext.h"
+#include "PCGParamData.h"
 #include "Data/PCGPointData.h"
 #include "Data/PCGSpatialData.h"
 
@@ -21,24 +22,60 @@ FPCGElementPtr UPCGCollapseSettings::CreateElement() const
 	return MakeShared<FPCGCollapseElement>();
 }
 
+TArray<FPCGPinProperties> UPCGConvertToPointDataSettings::InputPinProperties() const
+{
+	TArray<FPCGPinProperties> PinProperties;
+	PinProperties.Emplace(PCGPinConstants::DefaultInputLabel, EPCGDataType::Param);
+
+	return PinProperties;
+}
+
 bool FPCGCollapseElement::ExecuteInternal(FPCGContext* Context) const
 {
 	TRACE_CPUPROFILER_EVENT_SCOPE(FPCGCollapseElement::Execute);
+	check(Context);
 
-	TArray<FPCGTaggedData> Inputs = Context->InputData.GetInputs();
+	TArray<FPCGTaggedData> Inputs = Context->InputData.GetInputsByPin(PCGPinConstants::DefaultInputLabel);
 	TArray<FPCGTaggedData>& Outputs = Context->OutputData.TaggedData;
 
 	for (const FPCGTaggedData& Input : Inputs)
 	{
+		if (!Input.Data)
+		{
+			continue;
+		}
+
 		FPCGTaggedData& Output = Outputs.Add_GetRef(Input);
 
-		if (!Input.Data || Cast<const UPCGSpatialData>(Input.Data) == nullptr)
+		if (const UPCGSpatialData* SpatialData = Cast<UPCGSpatialData>(Input.Data))
 		{
-			continue;
+			// Currently we support collapsing to point data only, but at some point in the future that might be different
+			Output.Data = Cast<const UPCGSpatialData>(Input.Data)->ToPointData(Context);
 		}
+		else if (const UPCGParamData* ParamData = Cast<UPCGParamData>(Input.Data))
+		{
+			const UPCGMetadata* ParamMetadata = ParamData->Metadata;
+			const int64 ParamItemCount = ParamMetadata->GetLocalItemCount();
+
+			if (ParamItemCount == 0)
+			{
+				continue;
+			}
 
-		// Currently we support collapsing to point data only, but at some point in the future that might be different
-		Output.Data = Cast<const UPCGSpatialData>(Input.Data)->ToPointData(Context);
+			UPCGPointData* PointData = NewObject<UPCGPointData>();
+			check(PointData->Metadata);
+			PointData->Metadata->Initialize(ParamMetadata);
+
+			TArray<FPCGPoint>& Points = PointData->GetMutablePoints();
+			Points.SetNum(ParamItemCount);
+
+			for (int PointIndex = 0; PointIndex < ParamItemCount; ++PointIndex)
+			{
+				Points[PointIndex].MetadataEntry = PointIndex;
+			}
+
+			Output.Data = PointData;
+		}
 	}
 
 	return true;
diff --git a/Engine/Plugins/Experimental/PCG/Source/PCG/Private/Elements/PCGConvertToAttributeSet.cpp b/Engine/Plugins/Experimental/PCG/Source/PCG/Private/Elements/PCGConvertToAttributeSet.cpp
new file mode 100644
--- /dev/null
+++ b/Engine/Plugins/Experimental/PCG/Source/PCG/Private/Elements/PCGConvertToAttributeSet.cpp
@@ -0,0 +1,74 @@
+// Copyright Epic Games, Inc. All Rights Reserved.
+
+#include "Elements/PCGConvertToAttributeSet.h"
+
+#include "PCGContext.h"
+#include "PCGParamData.h"
+#include "Data/PCGPointData.h"
+
+#include UE_INLINE_GENERATED_CPP_BY_NAME(PCGConvertToAttributeSet)
+
+TArray<FPCGPinProperties> UPCGConvertToAttributeSetSettings::InputPinProperties() const
+{
+	return Super::DefaultPointInputPinProperties();
+}
+
+TArray<FPCGPinProperties> UPCGConvertToAttributeSetSettings::OutputPinProperties() const
+{
+	TArray<FPCGPinProperties> PinProperties;
+	PinProperties.Emplace(PCGPinConstants::DefaultOutputLabel, EPCGDataType::Param);
+
+	return PinProperties;
+}
+
+FPCGElementPtr UPCGConvertToAttributeSetSettings::CreateElement() const
+{
+	return MakeShared<FPCGConvertToAttributeSetElement>();
+}
+
+bool FPCGConvertToAttributeSetElement::ExecuteInternal(FPCGContext* Context) const
+{
+	TRACE_CPUPROFILER_EVENT_SCOPE(FPCGConvertToAttributeSetElement::Execute);
+	check(Context);
+
+	TArray<FPCGTaggedData> Inputs = Context->InputData.GetInputs();
+	TArray<FPCGTaggedData>& Outputs = Context->OutputData.TaggedData;
+
+	for (const FPCGTaggedData& Input : Inputs)
+	{
+		if (Cast<UPCGPointData>(Input.Data) == nullptr)
+		{
+			continue;
+		}
+
+		const UPCGPointData* PointData = Cast<UPCGPointData>(Input.Data);
+		const TArray<FPCGPoint>& Points = PointData->GetPoints();
+		const UPCGMetadata* SourceMetadata = PointData->Metadata;
+		check(SourceMetadata);
+
+		// Note: this is very similar in idea to UPCGPointData::Flatten
+		if (SourceMetadata->GetAttributeCount() == 0 || Points.Num() == 0)
+		{
+			continue;
+		}
+
+		UPCGParamData* ParamData = NewObject<UPCGParamData>();
+		check(ParamData->Metadata);
+		ParamData->Metadata->Initialize(SourceMetadata);
+
+		// Create new entry that will map (or not) to the original parent
+		TArray<PCGMetadataEntryKey> EntryKeys;
+		EntryKeys.Reserve(Points.Num());
+		for (const FPCGPoint& Point : Points)
+		{
+			EntryKeys.Add(ParamData->Metadata->AddEntry(Point.MetadataEntry));
+		}
+
+		ParamData->Metadata->FlattenAndCompress(EntryKeys);
+
+		FPCGTaggedData& Output = Outputs.Emplace_GetRef(Input);
+		Output.Data = ParamData;
+	}
+
+	return true;
+}
\ No newline at end of file
diff --git a/Engine/Plugins/Experimental/PCG/Source/PCG/Public/Elements/PCGCollapseElement.h b/Engine/Plugins/Experimental/PCG/Source/PCG/Public/Elements/PCGCollapseElement.h
--- a/Engine/Plugins/Experimental/PCG/Source/PCG/Public/Elements/PCGCollapseElement.h
+++ b/Engine/Plugins/Experimental/PCG/Source/PCG/Public/Elements/PCGCollapseElement.h
@@ -28,6 +28,26 @@ protected:
 	//~End UPCGSettings interface
 };
 
+/** Converts attribute sets to point data */
+UCLASS(BlueprintType, ClassGroup = (Procedural))
+class PCG_API UPCGConvertToPointDataSettings : public UPCGCollapseSettings
+{
+	GENERATED_BODY()
+
+public:
+	//~Begin UPCGSettings interface
+#if WITH_EDITOR
+	virtual FName GetDefaultNodeName() const override { return FName(TEXT("AttributeSetToPoint")); }
+	virtual FText GetDefaultNodeTitle() const override { return NSLOCTEXT("PCGConvertToPointElement", "NodeTitle", "Attribute Set To Point"); }
+	virtual EPCGSettingsType GetType() const override { return EPCGSettingsType::Spatial; }
+	virtual bool ShouldDrawNodeCompact() const override { return true; }
+#endif
+
+protected:
+	virtual TArray<FPCGPinProperties> InputPinProperties() const override;
+	//~End UPCGSettings interface
+};
+
 class FPCGCollapseElement : public IPCGElement
 {
 public:
diff --git a/Engine/Plugins/Experimental/PCG/Source/PCG/Public/Elements/PCGConvertToAttributeSet.h b/Engine/Plugins/Experimental/PCG/Source/PCG/Public/Elements/PCGConvertToAttributeSet.h
new file mode 100644
--- /dev/null
+++ b/Engine/Plugins/Experimental/PCG/Source/PCG/Public/Elements/PCGConvertToAttributeSet.h
@@ -0,0 +1,40 @@
+// Copyright Epic Games, Inc. All Rights Reserved.
+
+#pragma once
+
+#include "PCGSettings.h"
+
+#include "PCGConvertToAttributeSet.generated.h"
+
+/** Converts point data to an attribute set with one entry per point and the same attributes. */
+UCLASS(BlueprintType, ClassGroup = (Procedural))
+class PCG_API UPCGConvertToAttributeSetSettings : public UPCGSettings
+{
+	GENERATED_BODY()
+
+public:
+	//~Begin UPCGSettings interface
+#if WITH_EDITOR
+	virtual FName GetDefaultNodeName() const override { return FName(TEXT("PointToAttributeSet")); }
+	virtual FText GetDefaultNodeTitle() const override { return NSLOCTEXT("PCGConvertToAttributeSetElement", "NodeTitle", "Point To Attribute Set"); }
+	virtual FText GetNodeTooltipText() const override { return NSLOCTEXT("PCGConvertToAttributeSetElement", "NodeTooltip", "Converts point data to an attribute set with one entry per point and the same attributes."); }
+	virtual EPCGSettingsType GetType() const override { return EPCGSettingsType::Param; }
+	virtual bool ShouldDrawNodeCompact() const override { return true; }
+#endif
+
+protected:
+	virtual TArray<FPCGPinProperties> InputPinProperties() const override;
+	virtual TArray<FPCGPinProperties> OutputPinProperties() const override;
+	virtual FPCGElementPtr CreateElement() const override;
+	//~End UPCGSettings interface
+};
+
+class FPCGConvertToAttributeSetElement : public IPCGElement
+{
+public:
+	// Since this is downstream from points, possibly sampling, it's worth to compute a full CRC so if there are no changes to the attributes, we can maximize cache usage
+	virtual bool ShouldComputeFullOutputDataCrc(FPCGContext* Context) const override { return true; }
+
+protected:
+	virtual bool ExecuteInternal(FPCGContext* Context) const;
+};
\ No newline at end of file
