Fix: Rename PCG Unit Tests command to Plugin.PCG

diff --git a/Engine/Plugins/Experimental/PCG/Source/PCG/Private/Tests/Elements/PCGDuplicatePointTest.cpp b/Engine/Plugins/Experimental/PCG/Source/PCG/Private/Tests/Elements/PCGDuplicatePointTest.cpp
--- a/Engine/Plugins/Experimental/PCG/Source/PCG/Private/Tests/Elements/PCGDuplicatePointTest.cpp
+++ b/Engine/Plugins/Experimental/PCG/Source/PCG/Private/Tests/Elements/PCGDuplicatePointTest.cpp
@@ -87,7 +87,7 @@ protected:
 	}
 };
 
-IMPLEMENT_CUSTOM_SIMPLE_AUTOMATION_TEST(FPCGDuplicatePointTest_WithSource, DuplicatePointTestBase, "pcg.tests.DuplicatePoint.WithSource", PCGTestsCommon::TestFlags)
+IMPLEMENT_CUSTOM_SIMPLE_AUTOMATION_TEST(FPCGDuplicatePointTest_WithSource, DuplicatePointTestBase, "Plugins.PCG.DuplicatePoint.WithSource", PCGTestsCommon::TestFlags)
 
 bool FPCGDuplicatePointTest_WithSource::RunTest(const FString& Parameters)
 {
@@ -105,7 +105,7 @@ bool FPCGDuplicatePointTest_WithSource::RunTest(const FString& Parameters)
 	return GenerateTestDataRunAndValidate(SourceParameters);
 }
 
-IMPLEMENT_CUSTOM_SIMPLE_AUTOMATION_TEST(FPCGDuplicatePointTest_WithoutSource, DuplicatePointTestBase, "pcg.tests.DuplicatePoint.WithoutSource", PCGTestsCommon::TestFlags)
+IMPLEMENT_CUSTOM_SIMPLE_AUTOMATION_TEST(FPCGDuplicatePointTest_WithoutSource, DuplicatePointTestBase, "Plugins.PCG.DuplicatePoint.WithoutSource", PCGTestsCommon::TestFlags)
 
 bool FPCGDuplicatePointTest_WithoutSource::RunTest(const FString& Parameters)
 {
@@ -123,7 +123,7 @@ bool FPCGDuplicatePointTest_WithoutSource::RunTest(const FString& Parameters)
 	return GenerateTestDataRunAndValidate(NoSourceParameters);
 }
 
-IMPLEMENT_CUSTOM_SIMPLE_AUTOMATION_TEST(FPCGDuplicatePointTest_DirectionX, DuplicatePointTestBase, "pcg.tests.DuplicatePoint.DirectionX", PCGTestsCommon::TestFlags)
+IMPLEMENT_CUSTOM_SIMPLE_AUTOMATION_TEST(FPCGDuplicatePointTest_DirectionX, DuplicatePointTestBase, "Plugins.PCG.DuplicatePoint.DirectionX", PCGTestsCommon::TestFlags)
 
 bool FPCGDuplicatePointTest_DirectionX::RunTest(const FString& Parameters)
 {
@@ -141,7 +141,7 @@ bool FPCGDuplicatePointTest_DirectionX::RunTest(const FString& Parameters)
 	return GenerateTestDataRunAndValidate(XParameters);
 }
 
-IMPLEMENT_CUSTOM_SIMPLE_AUTOMATION_TEST(FPCGDuplicatePointTest_DirectionNegX, DuplicatePointTestBase, "pcg.tests.DuplicatePoint.DirectionNegX", PCGTestsCommon::TestFlags)
+IMPLEMENT_CUSTOM_SIMPLE_AUTOMATION_TEST(FPCGDuplicatePointTest_DirectionNegX, DuplicatePointTestBase, "Plugins.PCG.DuplicatePoint.DirectionNegX", PCGTestsCommon::TestFlags)
 
 bool FPCGDuplicatePointTest_DirectionNegX::RunTest(const FString& Parameters)
 {
@@ -159,7 +159,7 @@ bool FPCGDuplicatePointTest_DirectionNegX::RunTest(const FString& Parameters)
 	return GenerateTestDataRunAndValidate(NegXParameters);
 }
 
-IMPLEMENT_CUSTOM_SIMPLE_AUTOMATION_TEST(FPCGDuplicatePointTest_DirectionY, DuplicatePointTestBase, "pcg.tests.DuplicatePoint.DirectionY", PCGTestsCommon::TestFlags)
+IMPLEMENT_CUSTOM_SIMPLE_AUTOMATION_TEST(FPCGDuplicatePointTest_DirectionY, DuplicatePointTestBase, "Plugins.PCG.DuplicatePoint.DirectionY", PCGTestsCommon::TestFlags)
 
 bool FPCGDuplicatePointTest_DirectionY::RunTest(const FString& Parameters)
 {
@@ -177,7 +177,7 @@ bool FPCGDuplicatePointTest_DirectionY::RunTest(const FString& Parameters)
 	return GenerateTestDataRunAndValidate(YParameters);
 }
 
-IMPLEMENT_CUSTOM_SIMPLE_AUTOMATION_TEST(FPCGDuplicatePointTest_DirectionNegY, DuplicatePointTestBase, "pcg.tests.DuplicatePoint.DirectionNegY", PCGTestsCommon::TestFlags)
+IMPLEMENT_CUSTOM_SIMPLE_AUTOMATION_TEST(FPCGDuplicatePointTest_DirectionNegY, DuplicatePointTestBase, "Plugins.PCG.DuplicatePoint.DirectionNegY", PCGTestsCommon::TestFlags)
 
 bool FPCGDuplicatePointTest_DirectionNegY::RunTest(const FString& Parameters)
 {
@@ -195,7 +195,7 @@ bool FPCGDuplicatePointTest_DirectionNegY::RunTest(const FString& Parameters)
 	return GenerateTestDataRunAndValidate(NegYParameters);
 }
 
-IMPLEMENT_CUSTOM_SIMPLE_AUTOMATION_TEST(FPCGDuplicatePointTest_DirectionZ, DuplicatePointTestBase, "pcg.tests.DuplicatePoint.DirectionZ", PCGTestsCommon::TestFlags)
+IMPLEMENT_CUSTOM_SIMPLE_AUTOMATION_TEST(FPCGDuplicatePointTest_DirectionZ, DuplicatePointTestBase, "Plugins.PCG.DuplicatePoint.DirectionZ", PCGTestsCommon::TestFlags)
 
 bool FPCGDuplicatePointTest_DirectionZ::RunTest(const FString& Parameters)
 {
@@ -213,7 +213,7 @@ bool FPCGDuplicatePointTest_DirectionZ::RunTest(const FString& Parameters)
 	return GenerateTestDataRunAndValidate(ZParameters);
 }
 
-IMPLEMENT_CUSTOM_SIMPLE_AUTOMATION_TEST(FPCGDuplicatePointTest_DirectionNegZ, DuplicatePointTestBase, "pcg.tests.DuplicatePoint.DirectionNegZ", PCGTestsCommon::TestFlags)
+IMPLEMENT_CUSTOM_SIMPLE_AUTOMATION_TEST(FPCGDuplicatePointTest_DirectionNegZ, DuplicatePointTestBase, "Plugins.PCG.DuplicatePoint.DirectionNegZ", PCGTestsCommon::TestFlags)
 
 bool FPCGDuplicatePointTest_DirectionNegZ::RunTest(const FString& Parameters)
 {
@@ -231,7 +231,7 @@ bool FPCGDuplicatePointTest_DirectionNegZ::RunTest(const FString& Parameters)
 	return GenerateTestDataRunAndValidate(NegZParameters);
 }
 
-IMPLEMENT_CUSTOM_SIMPLE_AUTOMATION_TEST(FPCGDuplicatePointTest_ScaleOne, DuplicatePointTestBase, "pcg.tests.DuplicatePoint.ScaleOne", PCGTestsCommon::TestFlags)
+IMPLEMENT_CUSTOM_SIMPLE_AUTOMATION_TEST(FPCGDuplicatePointTest_ScaleOne, DuplicatePointTestBase, "Plugins.PCG.DuplicatePoint.ScaleOne", PCGTestsCommon::TestFlags)
 
 bool FPCGDuplicatePointTest_ScaleOne::RunTest(const FString& Parameters)
 {
@@ -249,7 +249,7 @@ bool FPCGDuplicatePointTest_ScaleOne::RunTest(const FString& Parameters)
 	return GenerateTestDataRunAndValidate(ScaleParameters);
 }
 
-IMPLEMENT_CUSTOM_SIMPLE_AUTOMATION_TEST(FPCGDuplicatePointTest_ScaleNegative, DuplicatePointTestBase, "pcg.tests.DuplicatePoint.ScaleNegative", PCGTestsCommon::TestFlags)
+IMPLEMENT_CUSTOM_SIMPLE_AUTOMATION_TEST(FPCGDuplicatePointTest_ScaleNegative, DuplicatePointTestBase, "Plugins.PCG.DuplicatePoint.ScaleNegative", PCGTestsCommon::TestFlags)
 
 bool FPCGDuplicatePointTest_ScaleNegative::RunTest(const FString& Parameters)
 {
@@ -267,7 +267,7 @@ bool FPCGDuplicatePointTest_ScaleNegative::RunTest(const FString& Parameters)
 	return GenerateTestDataRunAndValidate(NegScaleParameters);
 }
 
-IMPLEMENT_CUSTOM_SIMPLE_AUTOMATION_TEST(FPCGDuplicatePointTest_Rotation, DuplicatePointTestBase, "pcg.tests.DuplicatePoint.Rotation", PCGTestsCommon::TestFlags)
+IMPLEMENT_CUSTOM_SIMPLE_AUTOMATION_TEST(FPCGDuplicatePointTest_Rotation, DuplicatePointTestBase, "Plugins.PCG.DuplicatePoint.Rotation", PCGTestsCommon::TestFlags)
 
 bool FPCGDuplicatePointTest_Rotation::RunTest(const FString& Parameters)
 {
@@ -285,7 +285,7 @@ bool FPCGDuplicatePointTest_Rotation::RunTest(const FString& Parameters)
 	return GenerateTestDataRunAndValidate(RotationParameters);
 }
 
-IMPLEMENT_CUSTOM_SIMPLE_AUTOMATION_TEST(FPCGDuplicatePointTest_LocationX, DuplicatePointTestBase, "pcg.tests.DuplicatePoint.LocationX", PCGTestsCommon::TestFlags)
+IMPLEMENT_CUSTOM_SIMPLE_AUTOMATION_TEST(FPCGDuplicatePointTest_LocationX, DuplicatePointTestBase, "Plugins.PCG.DuplicatePoint.LocationX", PCGTestsCommon::TestFlags)
 
 bool FPCGDuplicatePointTest_LocationX::RunTest(const FString& Parameters)
 {
@@ -303,7 +303,7 @@ bool FPCGDuplicatePointTest_LocationX::RunTest(const FString& Parameters)
 	return GenerateTestDataRunAndValidate(LocationXParameters);
 }
 
-IMPLEMENT_CUSTOM_SIMPLE_AUTOMATION_TEST(FPCGDuplicatePointTest_LocationY, DuplicatePointTestBase, "pcg.tests.DuplicatePoint.LocationY", PCGTestsCommon::TestFlags)
+IMPLEMENT_CUSTOM_SIMPLE_AUTOMATION_TEST(FPCGDuplicatePointTest_LocationY, DuplicatePointTestBase, "Plugins.PCG.DuplicatePoint.LocationY", PCGTestsCommon::TestFlags)
 
 bool FPCGDuplicatePointTest_LocationY::RunTest(const FString& Parameters)
 {
@@ -321,7 +321,7 @@ bool FPCGDuplicatePointTest_LocationY::RunTest(const FString& Parameters)
 	return GenerateTestDataRunAndValidate(LocationYParameters);
 }
 
-IMPLEMENT_CUSTOM_SIMPLE_AUTOMATION_TEST(FPCGDuplicatePointTest_LocationZ, DuplicatePointTestBase, "pcg.tests.DuplicatePoint.LocationZ", PCGTestsCommon::TestFlags)
+IMPLEMENT_CUSTOM_SIMPLE_AUTOMATION_TEST(FPCGDuplicatePointTest_LocationZ, DuplicatePointTestBase, "Plugins.PCG.DuplicatePoint.LocationZ", PCGTestsCommon::TestFlags)
 
 bool FPCGDuplicatePointTest_LocationZ::RunTest(const FString& Parameters)
 {
diff --git a/Engine/Plugins/Experimental/PCG/Source/PCG/Private/Tests/Elements/PCGSplitPointsTest.cpp b/Engine/Plugins/Experimental/PCG/Source/PCG/Private/Tests/Elements/PCGSplitPointsTest.cpp
--- a/Engine/Plugins/Experimental/PCG/Source/PCG/Private/Tests/Elements/PCGSplitPointsTest.cpp
+++ b/Engine/Plugins/Experimental/PCG/Source/PCG/Private/Tests/Elements/PCGSplitPointsTest.cpp
@@ -76,7 +76,7 @@ protected:
 	}
 };
 
-IMPLEMENT_CUSTOM_SIMPLE_AUTOMATION_TEST(FPCGSplitPointsTest_DefaultX, FPCGSplitPointTestBase, "pcg.tests.SplitPoints.DefaultX", PCGTestsCommon::TestFlags)
+IMPLEMENT_CUSTOM_SIMPLE_AUTOMATION_TEST(FPCGSplitPointsTest_DefaultX, FPCGSplitPointTestBase, "Plugins.PCG.SplitPoints.DefaultX", PCGTestsCommon::TestFlags)
 
 bool FPCGSplitPointsTest_DefaultX::RunTest(const FString& Parameters)
 {
@@ -92,7 +92,7 @@ bool FPCGSplitPointsTest_DefaultX::RunTest(const FString& Parameters)
 	return GenerateTestDataRunAndValidate(DefaultXParameters);
 }
 
-IMPLEMENT_CUSTOM_SIMPLE_AUTOMATION_TEST(FPCGSplitPointsTest_DefaultY, FPCGSplitPointTestBase, "pcg.tests.SplitPoints.DefaultY", PCGTestsCommon::TestFlags)
+IMPLEMENT_CUSTOM_SIMPLE_AUTOMATION_TEST(FPCGSplitPointsTest_DefaultY, FPCGSplitPointTestBase, "Plugins.PCG.SplitPoints.DefaultY", PCGTestsCommon::TestFlags)
 
 bool FPCGSplitPointsTest_DefaultY::RunTest(const FString& Parameters)
 {
@@ -108,7 +108,7 @@ bool FPCGSplitPointsTest_DefaultY::RunTest(const FString& Parameters)
 	return GenerateTestDataRunAndValidate(DefaultYParameters);
 }
 
-IMPLEMENT_CUSTOM_SIMPLE_AUTOMATION_TEST(FPCGSplitPointsTest_DefaultZ, FPCGSplitPointTestBase, "pcg.tests.SplitPoints.DefaultZ", PCGTestsCommon::TestFlags)
+IMPLEMENT_CUSTOM_SIMPLE_AUTOMATION_TEST(FPCGSplitPointsTest_DefaultZ, FPCGSplitPointTestBase, "Plugins.PCG.SplitPoints.DefaultZ", PCGTestsCommon::TestFlags)
 
 bool FPCGSplitPointsTest_DefaultZ::RunTest(const FString& Parameters)
 {
@@ -124,7 +124,7 @@ bool FPCGSplitPointsTest_DefaultZ::RunTest(const FString& Parameters)
 	return GenerateTestDataRunAndValidate(DefaultZParameters);
 }
 
-IMPLEMENT_CUSTOM_SIMPLE_AUTOMATION_TEST(FPCGSplitPointsTest_ZeroX, FPCGSplitPointTestBase, "pcg.tests.SplitPoints.ZeroX", PCGTestsCommon::TestFlags)
+IMPLEMENT_CUSTOM_SIMPLE_AUTOMATION_TEST(FPCGSplitPointsTest_ZeroX, FPCGSplitPointTestBase, "Plugins.PCG.SplitPoints.ZeroX", PCGTestsCommon::TestFlags)
 
 bool FPCGSplitPointsTest_ZeroX::RunTest(const FString& Parameters)
 {
@@ -140,7 +140,7 @@ bool FPCGSplitPointsTest_ZeroX::RunTest(const FString& Parameters)
 	return GenerateTestDataRunAndValidate(ZeroXParameters);
 }
 
-IMPLEMENT_CUSTOM_SIMPLE_AUTOMATION_TEST(FPCGSplitPointsTest_ZeroY, FPCGSplitPointTestBase, "pcg.tests.SplitPoints.ZeroY", PCGTestsCommon::TestFlags)
+IMPLEMENT_CUSTOM_SIMPLE_AUTOMATION_TEST(FPCGSplitPointsTest_ZeroY, FPCGSplitPointTestBase, "Plugins.PCG.SplitPoints.ZeroY", PCGTestsCommon::TestFlags)
 
 bool FPCGSplitPointsTest_ZeroY::RunTest(const FString& Parameters)
 {
@@ -156,7 +156,7 @@ bool FPCGSplitPointsTest_ZeroY::RunTest(const FString& Parameters)
 	return GenerateTestDataRunAndValidate(ZeroYParameters);
 }
 
-IMPLEMENT_CUSTOM_SIMPLE_AUTOMATION_TEST(FPCGSplitPointsTest_ZeroZ, FPCGSplitPointTestBase, "pcg.tests.SplitPoints.ZeroZ", PCGTestsCommon::TestFlags)
+IMPLEMENT_CUSTOM_SIMPLE_AUTOMATION_TEST(FPCGSplitPointsTest_ZeroZ, FPCGSplitPointTestBase, "Plugins.PCG.SplitPoints.ZeroZ", PCGTestsCommon::TestFlags)
 
 bool FPCGSplitPointsTest_ZeroZ::RunTest(const FString& Parameters)
 {
@@ -172,7 +172,7 @@ bool FPCGSplitPointsTest_ZeroZ::RunTest(const FString& Parameters)
 	return GenerateTestDataRunAndValidate(ZeroZParameters);
 }
 
-IMPLEMENT_CUSTOM_SIMPLE_AUTOMATION_TEST(FPCGSplitPointsTest_OneX, FPCGSplitPointTestBase, "pcg.tests.SplitPoints.OneX", PCGTestsCommon::TestFlags)
+IMPLEMENT_CUSTOM_SIMPLE_AUTOMATION_TEST(FPCGSplitPointsTest_OneX, FPCGSplitPointTestBase, "Plugins.PCG.SplitPoints.OneX", PCGTestsCommon::TestFlags)
 
 bool FPCGSplitPointsTest_OneX::RunTest(const FString& Parameters)
 {
@@ -188,7 +188,7 @@ bool FPCGSplitPointsTest_OneX::RunTest(const FString& Parameters)
 	return GenerateTestDataRunAndValidate(OneXParameters);
 }
 
-IMPLEMENT_CUSTOM_SIMPLE_AUTOMATION_TEST(FPCGSplitPointsTest_OneY, FPCGSplitPointTestBase, "pcg.tests.SplitPoints.OneY", PCGTestsCommon::TestFlags)
+IMPLEMENT_CUSTOM_SIMPLE_AUTOMATION_TEST(FPCGSplitPointsTest_OneY, FPCGSplitPointTestBase, "Plugins.PCG.SplitPoints.OneY", PCGTestsCommon::TestFlags)
 
 bool FPCGSplitPointsTest_OneY::RunTest(const FString& Parameters)
 {
@@ -204,7 +204,7 @@ bool FPCGSplitPointsTest_OneY::RunTest(const FString& Parameters)
 	return GenerateTestDataRunAndValidate(OneYParameters);
 }
 
-IMPLEMENT_CUSTOM_SIMPLE_AUTOMATION_TEST(FPCGSplitPointsTest_OneZ, FPCGSplitPointTestBase, "pcg.tests.SplitPoints.OneZ", PCGTestsCommon::TestFlags)
+IMPLEMENT_CUSTOM_SIMPLE_AUTOMATION_TEST(FPCGSplitPointsTest_OneZ, FPCGSplitPointTestBase, "Plugins.PCG.SplitPoints.OneZ", PCGTestsCommon::TestFlags)
 
 bool FPCGSplitPointsTest_OneZ::RunTest(const FString& Parameters)
 {
