Added node to filter data based on the existence of some metadata attributes.
Refactored filter classes to commonalize standard behavior and pins
Added deprecation for the nodes that now have multiple outputs
Fixed a small error in the Get Data By Index that could have gathered the overrides as well

diff --git a/Engine/Plugins/Experimental/PCG/Source/PCG/Private/Elements/PCGFilterByAttribute.cpp b/Engine/Plugins/Experimental/PCG/Source/PCG/Private/Elements/PCGFilterByAttribute.cpp
new file mode 100644
--- /dev/null
+++ b/Engine/Plugins/Experimental/PCG/Source/PCG/Private/Elements/PCGFilterByAttribute.cpp
@@ -0,0 +1,71 @@
+// Copyright Epic Games, Inc. All Rights Reserved.
+
+#include "Elements/PCGFilterByAttribute.h"
+
+#include "PCGContext.h"
+#include "PCGData.h"
+#include "Metadata/PCGMetadata.h"
+
+#define LOCTEXT_NAMESPACE "PCGFilterByAttributeElement"
+
+namespace PCGFilterByAttributeConstants
+{
+	const FText NodeTitle = LOCTEXT("NodeTitle", "Filter Data By Attribute");
+}
+
+#if WITH_EDITOR
+FText UPCGFilterByAttributeSettings::GetDefaultNodeTitle() const
+{
+	return PCGFilterByAttributeConstants::NodeTitle;
+}
+#endif // WITH_EDITOR
+
+FName UPCGFilterByAttributeSettings::AdditionalTaskName() const
+{
+#if WITH_EDITOR
+	FProperty* AttributeProperty = GetClass() ? FindFProperty<FProperty>(GetClass(), GET_MEMBER_NAME_CHECKED(UPCGFilterByAttributeSettings, Attribute)) : nullptr;
+	if (AttributeProperty && IsPropertyOverriddenByPin(AttributeProperty))
+	{
+		return NAME_None;
+	}
+	else
+#endif
+	{
+		return FName(FString::Printf(TEXT("%s (%s)"), *PCGFilterByAttributeConstants::NodeTitle.ToString(), *Attribute.ToString()));
+	}
+}
+
+FPCGElementPtr UPCGFilterByAttributeSettings::CreateElement() const
+{
+	return MakeShared<FPCGFilterByAttributeElement>();
+}
+
+bool FPCGFilterByAttributeElement::ExecuteInternal(FPCGContext* Context) const
+{
+	TRACE_CPUPROFILER_EVENT_SCOPE(FPCGFilterByAttributeElement::Execute);
+	check(Context);
+
+	const UPCGFilterByAttributeSettings* Settings = Context->GetInputSettings<UPCGFilterByAttributeSettings>();
+	check(Settings);
+
+	const FName Attribute = FName(Settings->Attribute);
+
+	TArray<FPCGTaggedData> Inputs = Context->InputData.GetInputsByPin(PCGPinConstants::DefaultInputLabel);
+	TArray<FPCGTaggedData>& Outputs = Context->OutputData.TaggedData;
+
+	for (const FPCGTaggedData& Input : Inputs)
+	{
+		FPCGTaggedData& Output = Outputs.Add_GetRef(Input);
+		Output.Pin = PCGPinConstants::DefaultOutFilterLabel;
+		
+		const UPCGMetadata* Metadata = Input.Data ? Input.Data->ConstMetadata() : nullptr;
+		if (Metadata && Metadata->HasAttribute(Attribute))
+		{
+			Output.Pin = PCGPinConstants::DefaultInFilterLabel;
+		}
+	}
+
+	return true;
+}
+
+#undef LOCTEXT_NAMESPACE
\ No newline at end of file
diff --git a/Engine/Plugins/Experimental/PCG/Source/PCG/Private/Elements/PCGFilterByIndex.cpp b/Engine/Plugins/Experimental/PCG/Source/PCG/Private/Elements/PCGFilterByIndex.cpp
--- a/Engine/Plugins/Experimental/PCG/Source/PCG/Private/Elements/PCGFilterByIndex.cpp
+++ b/Engine/Plugins/Experimental/PCG/Source/PCG/Private/Elements/PCGFilterByIndex.cpp
@@ -13,8 +13,8 @@
 
 namespace PCGFilterByIndexConstants
 {
-	const FName NodeName = FName(TEXT("FilterByIndex"));
-	const FText NodeTitle = LOCTEXT("NodeTitle", "Filter By Index");
+	const FName NodeName = FName(TEXT("FilterDataByIndex"));
+	const FText NodeTitle = LOCTEXT("NodeTitle", "Filter Data By Index");
 
 	// After so many characters on the node, it will truncate with an ellipsis ...
 	static constexpr int32 IndexExpressionTruncation = 12;
@@ -37,20 +37,6 @@ FText UPCGFilterByIndexSettings::GetNodeTooltipText() const
 }
 #endif
 
-EPCGDataType UPCGFilterByIndexSettings::GetCurrentPinTypes(const UPCGPin* InPin) const
-{
-	check(InPin);
-
-	if (!InPin->IsOutputPin())
-	{
-		return Super::GetCurrentPinTypes(InPin);
-	}
-
-	// Output pin narrows to union of inputs on first pin
-	const EPCGDataType InputTypeUnion = GetTypeUnionOfIncidentEdges(PCGPinConstants::DefaultInputLabel);
-	return InputTypeUnion != EPCGDataType::None ? InputTypeUnion : EPCGDataType::Any;
-}
-
 FName UPCGFilterByIndexSettings::AdditionalTaskName() const
 {
 	FString NodeName = PCGFilterByIndexConstants::NodeTitle.ToString();
@@ -59,15 +45,6 @@ FName UPCGFilterByIndexSettings::AdditionalTaskName() const
 	return FName(NodeName);
 }
 
-TArray<FPCGPinProperties> UPCGFilterByIndexSettings::OutputPinProperties() const
-{
-	TArray<FPCGPinProperties> PinProperties;
-	PinProperties.Emplace(PCGPinConstants::DefaultInFilterLabel, EPCGDataType::Any);
-	PinProperties.Emplace(PCGPinConstants::DefaultOutFilterLabel, EPCGDataType::Any);
-
-	return PinProperties;
-}
-
 FPCGElementPtr UPCGFilterByIndexSettings::CreateElement() const
 {
 	return MakeShared<FPCGFilterByIndexElement>();
@@ -85,7 +62,9 @@ bool FPCGFilterByIndexElement::ExecuteInternal(FPCGContext* Context) const
 		return true;
 	}
 
-	PCGIndexing::FPCGIndexCollection FilteredIndices(Context->InputData.GetInputs().Num());
+	TArray<FPCGTaggedData> Inputs = Context->InputData.GetInputsByPin(PCGPinConstants::DefaultInputLabel);
+
+	PCGIndexing::FPCGIndexCollection FilteredIndices(Inputs.Num());
 	// Parse the indices and switch through possible issues.
 	switch (PCGParser::ParseIndexRanges(FilteredIndices, Settings->SelectedIndices))
 	{
@@ -112,7 +91,7 @@ bool FPCGFilterByIndexElement::ExecuteInternal(FPCGContext* Context) const
 	}
 
 	TArray<FPCGTaggedData>& Outputs = Context->OutputData.TaggedData;
-	Outputs = Context->InputData.GetInputs();
+	Outputs = Inputs;
 
 	for (int Index = 0; Index < Outputs.Num(); ++Index)
 	{
diff --git a/Engine/Plugins/Experimental/PCG/Source/PCG/Private/Elements/PCGFilterByTag.cpp b/Engine/Plugins/Experimental/PCG/Source/PCG/Private/Elements/PCGFilterByTag.cpp
--- a/Engine/Plugins/Experimental/PCG/Source/PCG/Private/Elements/PCGFilterByTag.cpp
+++ b/Engine/Plugins/Experimental/PCG/Source/PCG/Private/Elements/PCGFilterByTag.cpp
@@ -11,8 +11,8 @@
 
 namespace PCGFilterByTagConstants
 {
-	const FName NodeName = FName(TEXT("FilterByTag"));
-	const FText NodeTitle = LOCTEXT("NodeTitle", "Filter By Tag");
+	const FName NodeName = FName(TEXT("FilterDataByTag"));
+	const FText NodeTitle = LOCTEXT("NodeTitle", "Filter Data By Tag");
 }
 
 #if WITH_EDITOR
@@ -32,19 +32,6 @@ FText UPCGFilterByTagSettings::GetNodeTooltipText() const
 }
 #endif
 
-EPCGDataType UPCGFilterByTagSettings::GetCurrentPinTypes(const UPCGPin* InPin) const
-{
-	check(InPin);
-	if (!InPin->IsOutputPin())
-	{
-		return Super::GetCurrentPinTypes(InPin);
-	}
-
-	// Output pin narrows to union of inputs on first pin
-	const EPCGDataType InputTypeUnion = GetTypeUnionOfIncidentEdges(PCGPinConstants::DefaultInputLabel);
-	return InputTypeUnion != EPCGDataType::None ? InputTypeUnion : EPCGDataType::Any;
-}
-
 FName UPCGFilterByTagSettings::AdditionalTaskName() const
 {
 	const TArray<FString> Tags = PCGHelpers::GetStringArrayFromCommaSeparatedString(SelectedTags);
@@ -62,22 +49,6 @@ FName UPCGFilterByTagSettings::AdditionalTaskName() const
 	}
 }
 
-TArray<FPCGPinProperties> UPCGFilterByTagSettings::InputPinProperties() const
-{
-	TArray<FPCGPinProperties> PinProperties;
-	PinProperties.Emplace(PCGPinConstants::DefaultInputLabel, EPCGDataType::Any);
-
-	return PinProperties;
-}
-
-TArray<FPCGPinProperties> UPCGFilterByTagSettings::OutputPinProperties() const
-{
-	TArray<FPCGPinProperties> PinProperties;
-	PinProperties.Emplace(PCGPinConstants::DefaultOutputLabel, EPCGDataType::Any);
-
-	return PinProperties;
-}
-
 FPCGElementPtr UPCGFilterByTagSettings::CreateElement() const
 {
 	return MakeShared<FPCGFilterByTagElement>();
@@ -98,6 +69,9 @@ bool FPCGFilterByTagElement::ExecuteInternal(FPCGContext* Context) const
 
 	for (const FPCGTaggedData& Input : Inputs)
 	{
+		FPCGTaggedData& Output = Outputs.Add_GetRef(Input);
+		Output.Pin = PCGPinConstants::DefaultOutFilterLabel;
+
 		bool bHasCommonTags = false;
 
 		for (const FString& Tag : Tags)
@@ -111,7 +85,7 @@ bool FPCGFilterByTagElement::ExecuteInternal(FPCGContext* Context) const
 
 		if (bKeepIfTag == bHasCommonTags)
 		{
-			Outputs.Add(Input);
+			Output.Pin = PCGPinConstants::DefaultInFilterLabel;
 		}
 	}
 
diff --git a/Engine/Plugins/Experimental/PCG/Source/PCG/Private/Elements/PCGFilterByType.cpp b/Engine/Plugins/Experimental/PCG/Source/PCG/Private/Elements/PCGFilterByType.cpp
--- a/Engine/Plugins/Experimental/PCG/Source/PCG/Private/Elements/PCGFilterByType.cpp
+++ b/Engine/Plugins/Experimental/PCG/Source/PCG/Private/Elements/PCGFilterByType.cpp
@@ -14,11 +14,11 @@ FText UPCGFilterByTypeSettings::GetDefaultNodeTitle() const
 {
 	if (this == GetClass()->GetDefaultObject())
 	{
-		return LOCTEXT("DefaultNodeTitle", "Filter By Type");
+		return LOCTEXT("DefaultNodeTitle", "Filter Data By Type");
 	}
 
 	const FText TypeText = FText::FromString(StaticEnum<EPCGDataType>()->GetDisplayNameTextByValue(static_cast<int64>(TargetType)).ToString());
-	return FText::Format(LOCTEXT("NodeTitleFormat", "Filter - {0}"), TypeText);
+	return FText::Format(LOCTEXT("NodeTitleFormat", "Filter Data - {0}"), TypeText);
 }
 
 FText UPCGFilterByTypeSettings::GetNodeTooltipText() const
@@ -30,7 +30,7 @@ FText UPCGFilterByTypeSettings::GetNodeTooltipText() const
 EPCGDataType UPCGFilterByTypeSettings::GetCurrentPinTypes(const UPCGPin* InPin) const
 {
 	check(InPin);
-	if (!InPin->IsOutputPin())
+	if (!InPin->IsOutputPin() || InPin->Properties.Label != PCGPinConstants::DefaultInFilterLabel)
 	{
 		return Super::GetCurrentPinTypes(InPin);
 	}
@@ -38,6 +38,18 @@ EPCGDataType UPCGFilterByTypeSettings::GetCurrentPinTypes(const UPCGPin* InPin)
 	return TargetType;
 }
 
+TArray<FPCGPinProperties> UPCGFilterByTypeSettings::OutputPinProperties() const
+{
+	TArray<FPCGPinProperties> PinProperties = Super::OutputPinProperties();
+
+	if (!bShowOutsideFilter)
+	{
+		PinProperties.RemoveAll([](const FPCGPinProperties& InProperties) { return InProperties.Label == PCGPinConstants::DefaultOutFilterLabel; });
+	}
+
+	return PinProperties;
+}
+
 FPCGElementPtr UPCGFilterByTypeSettings::CreateElement() const
 {
 	return MakeShared<FPCGFilterByTypeElement>();
@@ -51,25 +63,32 @@ bool FPCGFilterByTypeElement::ExecuteInternal(FPCGContext* InContext) const
 	const UPCGFilterByTypeSettings* Settings = InContext->GetInputSettings<UPCGFilterByTypeSettings>();
 	check(Settings);
 
-	TArray<FPCGTaggedData> Inputs = InContext->InputData.TaggedData;
+	TArray<FPCGTaggedData> Inputs = InContext->InputData.GetInputsByPin(PCGPinConstants::DefaultInputLabel);
 	TArray<FPCGTaggedData>& Outputs = InContext->OutputData.TaggedData;
 
 	for (const FPCGTaggedData& Input : Inputs)
 	{
-		if (!Input.Data)
+		bool bDataIsInFilter = false;
+
+		if (Input.Data)
 		{
-			continue;
+			// No match if data does not overlap in type with the target, or if data is broader than the target
+			const bool bDataTypeOverlapsWithTarget = !!(Input.Data->GetDataType() & Settings->TargetType);
+			const bool bDataTypeBroaderThanTarget = !!(Input.Data->GetDataType() & ~Settings->TargetType);
+
+			if (bDataTypeOverlapsWithTarget && !bDataTypeBroaderThanTarget)
+			{
+				bDataIsInFilter = true;
+			}
 		}
 
-		// No match if data does not overlap in type with the target, or if data is broader than the target
-		const bool bDataTypeOverlapsWithTarget = !!(Input.Data->GetDataType() & Settings->TargetType);
-		const bool bDataTypeBroaderThanTarget = !!(Input.Data->GetDataType() & ~Settings->TargetType);
-		if (!bDataTypeOverlapsWithTarget || bDataTypeBroaderThanTarget)
+		if (!bDataIsInFilter && !Settings->bShowOutsideFilter)
 		{
 			continue;
 		}
 
-		Outputs.Add(Input);
+		FPCGTaggedData& Output = Outputs.Add_GetRef(Input);
+		Output.Pin = bDataIsInFilter ? PCGPinConstants::DefaultInFilterLabel : PCGPinConstants::DefaultOutFilterLabel;
 	}
 
 	return true;
diff --git a/Engine/Plugins/Experimental/PCG/Source/PCG/Private/Elements/PCGFilterDataBase.cpp b/Engine/Plugins/Experimental/PCG/Source/PCG/Private/Elements/PCGFilterDataBase.cpp
new file mode 100644
--- /dev/null
+++ b/Engine/Plugins/Experimental/PCG/Source/PCG/Private/Elements/PCGFilterDataBase.cpp
@@ -0,0 +1,53 @@
+// Copyright Epic Games, Inc. All Rights Reserved.
+
+#include "Elements/PCGFilterDataBase.h"
+
+#include "PCGCustomVersion.h"
+#include "PCGNode.h"
+#include "PCGPin.h"
+#include "Helpers/PCGSettingsHelpers.h"
+
+#include UE_INLINE_GENERATED_CPP_BY_NAME(PCGFilterDataBase)
+
+#if WITH_EDITOR
+void UPCGFilterDataBaseSettings::ApplyDeprecationBeforeUpdatePins(UPCGNode* InOutNode, TArray<TObjectPtr<UPCGPin>>& InputPins, TArray<TObjectPtr<UPCGPin>>& OutputPins)
+{
+	Super::ApplyDeprecationBeforeUpdatePins(InOutNode, InputPins, OutputPins);
+
+	if (DataVersion < FPCGCustomVersion::UpdateFilterNodeOutputPins)
+	{
+		InOutNode->RenameOutputPin(PCGPinConstants::DefaultOutputLabel, PCGPinConstants::DefaultInFilterLabel);
+	}
+}
+
+#endif // WITH_EDITOR
+
+EPCGDataType UPCGFilterDataBaseSettings::GetCurrentPinTypes(const UPCGPin* InPin) const
+{
+	check(InPin);
+	if (!InPin->IsOutputPin())
+	{
+		return Super::GetCurrentPinTypes(InPin);
+	}
+
+	// Output pin narrows to union of inputs on first pin
+	const EPCGDataType InputTypeUnion = GetTypeUnionOfIncidentEdges(PCGPinConstants::DefaultInputLabel);
+	return InputTypeUnion != EPCGDataType::None ? InputTypeUnion : EPCGDataType::Any;
+}
+
+TArray<FPCGPinProperties> UPCGFilterDataBaseSettings::InputPinProperties() const
+{
+	TArray<FPCGPinProperties> PinProperties;
+	PinProperties.Emplace(PCGPinConstants::DefaultInputLabel, EPCGDataType::Any);
+
+	return PinProperties;
+}
+
+TArray<FPCGPinProperties> UPCGFilterDataBaseSettings::OutputPinProperties() const
+{
+	TArray<FPCGPinProperties> PinProperties;
+	PinProperties.Emplace(PCGPinConstants::DefaultInFilterLabel, EPCGDataType::Any);
+	PinProperties.Emplace(PCGPinConstants::DefaultOutFilterLabel, EPCGDataType::Any);
+
+	return PinProperties;
+}
\ No newline at end of file
diff --git a/Engine/Plugins/Experimental/PCG/Source/PCG/Public/Elements/PCGFilterByAttribute.h b/Engine/Plugins/Experimental/PCG/Source/PCG/Public/Elements/PCGFilterByAttribute.h
new file mode 100644
--- /dev/null
+++ b/Engine/Plugins/Experimental/PCG/Source/PCG/Public/Elements/PCGFilterByAttribute.h
@@ -0,0 +1,39 @@
+// Copyright Epic Games, Inc. All Rights Reserved.
+
+#pragma once
+
+#include "Elements/PCGFilterDataBase.h"
+
+#include "PCGFilterByAttribute.generated.h"
+
+/** Separates data on whether they have a specific metadata attribute (not by value) */
+UCLASS(BlueprintType, ClassGroup = (Procedural))
+class PCG_API UPCGFilterByAttributeSettings : public UPCGFilterDataBaseSettings
+{
+	GENERATED_BODY()
+
+public:
+	//~Begin UPCGSettings interface
+#if WITH_EDITOR
+	virtual FName GetDefaultNodeName() const override { return FName(TEXT("FilterDataByAttribute")); }
+	virtual FText GetDefaultNodeTitle() const override;
+	virtual FText GetNodeTooltipText() const override { return NSLOCTEXT("PCGFilterByAttributeElement", "NodeTooltip", "Separates input data by whether they have the specified attribute or not."); }
+#endif
+
+	virtual FName AdditionalTaskName() const override;
+
+protected:
+	virtual FPCGElementPtr CreateElement() const override;
+	//~End UPCGSettings interface
+
+public:
+	/** Attribute to look for */
+	UPROPERTY(BlueprintReadWrite, EditAnywhere, Category = Settings, meta = (PCG_Overridable))
+	FName Attribute;
+};
+
+class FPCGFilterByAttributeElement : public IPCGElement
+{
+protected:
+	virtual bool ExecuteInternal(FPCGContext* Context) const override;
+};
diff --git a/Engine/Plugins/Experimental/PCG/Source/PCG/Public/Elements/PCGFilterByIndex.h b/Engine/Plugins/Experimental/PCG/Source/PCG/Public/Elements/PCGFilterByIndex.h
--- a/Engine/Plugins/Experimental/PCG/Source/PCG/Public/Elements/PCGFilterByIndex.h
+++ b/Engine/Plugins/Experimental/PCG/Source/PCG/Public/Elements/PCGFilterByIndex.h
@@ -2,13 +2,13 @@
 
 #pragma once
 
-#include "PCGSettings.h"
+#include "Elements/PCGFilterDataBase.h"
 
 #include "PCGFilterByIndex.generated.h"
 
 /** Filters a data collection based on a user defined index range expression. */
 UCLASS(BlueprintType, ClassGroup = (Procedural))
-class PCG_API UPCGFilterByIndexSettings : public UPCGSettings
+class PCG_API UPCGFilterByIndexSettings : public UPCGFilterDataBaseSettings
 {
 	GENERATED_BODY()
 
@@ -18,14 +18,10 @@ public:
 	virtual FName GetDefaultNodeName() const override;
 	virtual FText GetDefaultNodeTitle() const override;
 	virtual FText GetNodeTooltipText() const override;
-	virtual EPCGSettingsType GetType() const override { return EPCGSettingsType::Filter; }
-	virtual bool HasDynamicPins() const override { return true; }
 #endif
-	virtual EPCGDataType GetCurrentPinTypes(const UPCGPin* InPin) const override;
 	virtual FName AdditionalTaskName() const override;
 
 protected:
-	virtual TArray<FPCGPinProperties> OutputPinProperties() const override;
 	virtual FPCGElementPtr CreateElement() const override;
 	//~End UPCGSettings interface
 
diff --git a/Engine/Plugins/Experimental/PCG/Source/PCG/Public/Elements/PCGFilterByTag.h b/Engine/Plugins/Experimental/PCG/Source/PCG/Public/Elements/PCGFilterByTag.h
--- a/Engine/Plugins/Experimental/PCG/Source/PCG/Public/Elements/PCGFilterByTag.h
+++ b/Engine/Plugins/Experimental/PCG/Source/PCG/Public/Elements/PCGFilterByTag.h
@@ -2,7 +2,7 @@
 
 #pragma once
 
-#include "PCGSettings.h"
+#include "Elements/PCGFilterDataBase.h"
 
 #include "PCGFilterByTag.generated.h"
 
@@ -15,7 +15,7 @@ enum class EPCGFilterByTagOperation
 
 /** Filters a data collection based on some tag criterion */
 UCLASS(BlueprintType, ClassGroup = (Procedural))
-class PCG_API UPCGFilterByTagSettings : public UPCGSettings
+class PCG_API UPCGFilterByTagSettings : public UPCGFilterDataBaseSettings
 {
 	GENERATED_BODY()
 
@@ -25,15 +25,10 @@ public:
 	virtual FName GetDefaultNodeName() const override;
 	virtual FText GetDefaultNodeTitle() const override;
 	virtual FText GetNodeTooltipText() const override;
-	virtual EPCGSettingsType GetType() const override { return EPCGSettingsType::Filter; }
-	virtual bool HasDynamicPins() const override { return true; }
 #endif
-	virtual EPCGDataType GetCurrentPinTypes(const UPCGPin* InPin) const override;
 	virtual FName AdditionalTaskName() const override;
 
 protected:
-	virtual TArray<FPCGPinProperties> InputPinProperties() const override;
-	virtual TArray<FPCGPinProperties> OutputPinProperties() const override;
 	virtual FPCGElementPtr CreateElement() const override;
 	//~End UPCGSettings interface
 
diff --git a/Engine/Plugins/Experimental/PCG/Source/PCG/Public/Elements/PCGFilterByType.h b/Engine/Plugins/Experimental/PCG/Source/PCG/Public/Elements/PCGFilterByType.h
--- a/Engine/Plugins/Experimental/PCG/Source/PCG/Public/Elements/PCGFilterByType.h
+++ b/Engine/Plugins/Experimental/PCG/Source/PCG/Public/Elements/PCGFilterByType.h
@@ -2,35 +2,37 @@
 
 #pragma once
 
-#include "PCGSettings.h"
+#include "PCGFilterDataBase.h"
 
 #include "PCGFilterByType.generated.h"
 
 /** Filters an input collection based on data type. */
 UCLASS(BlueprintType, ClassGroup = (Procedural))
-class PCG_API UPCGFilterByTypeSettings : public UPCGSettings
+class PCG_API UPCGFilterByTypeSettings : public UPCGFilterDataBaseSettings
 {
 	GENERATED_BODY()
 
 public:
 	//~Begin UPCGSettings interface
 #if WITH_EDITOR
-	virtual FName GetDefaultNodeName() const override { return TEXT("FilterByType"); }
+	virtual FName GetDefaultNodeName() const override { return TEXT("FilterDataByType"); }
 	virtual FText GetDefaultNodeTitle() const override;
 	virtual FText GetNodeTooltipText() const override;
-	virtual EPCGSettingsType GetType() const override { return EPCGSettingsType::Filter; }
-	virtual bool HasDynamicPins() const override { return true; }
 	virtual bool ShouldDrawNodeCompact() const override { return true; }
 #endif
 	virtual EPCGDataType GetCurrentPinTypes(const UPCGPin* InPin) const override;
 
 protected:
+	virtual TArray<FPCGPinProperties> OutputPinProperties() const override;
 	virtual FPCGElementPtr CreateElement() const override;
 	//~End UPCGSettings interface
 
 public:
 	UPROPERTY(BlueprintReadWrite, EditAnywhere, Category = Settings)
 	EPCGDataType TargetType = EPCGDataType::Any;
+
+	UPROPERTY(BlueprintReadWrite, EditAnywhere, Category = Settings)
+	bool bShowOutsideFilter = false;
 };
 
 class FPCGFilterByTypeElement : public IPCGElement
diff --git a/Engine/Plugins/Experimental/PCG/Source/PCG/Public/Elements/PCGFilterDataBase.h b/Engine/Plugins/Experimental/PCG/Source/PCG/Public/Elements/PCGFilterDataBase.h
new file mode 100644
--- /dev/null
+++ b/Engine/Plugins/Experimental/PCG/Source/PCG/Public/Elements/PCGFilterDataBase.h
@@ -0,0 +1,28 @@
+// Copyright Epic Games, Inc. All Rights Reserved.
+
+#pragma once
+
+#include "PCGSettings.h"
+
+#include "PCGFilterDataBase.generated.h"
+
+UCLASS(BlueprintType, Abstract, ClassGroup = (Procedural))
+class PCG_API UPCGFilterDataBaseSettings : public UPCGSettings
+{
+	GENERATED_BODY()
+
+public:
+	//~Begin UPCGSettings interface
+#if WITH_EDITOR
+	virtual EPCGSettingsType GetType() const override { return EPCGSettingsType::Filter; }
+	virtual bool HasDynamicPins() const override { return true; }
+	virtual void ApplyDeprecationBeforeUpdatePins(UPCGNode* InOutNode, TArray<TObjectPtr<UPCGPin>>& InputPins, TArray<TObjectPtr<UPCGPin>>& OutputPins) override;
+#endif // WITH_EDITOR
+
+	virtual EPCGDataType GetCurrentPinTypes(const UPCGPin* InPin) const override;
+
+protected:
+	virtual TArray<FPCGPinProperties> InputPinProperties() const override;
+	virtual TArray<FPCGPinProperties> OutputPinProperties() const override;
+	//~End UPCGSettings interface
+};
\ No newline at end of file
diff --git a/Engine/Plugins/Experimental/PCG/Source/PCG/Public/PCGCustomVersion.h b/Engine/Plugins/Experimental/PCG/Source/PCG/Public/PCGCustomVersion.h
--- a/Engine/Plugins/Experimental/PCG/Source/PCG/Public/PCGCustomVersion.h
+++ b/Engine/Plugins/Experimental/PCG/Source/PCG/Public/PCGCustomVersion.h
@@ -62,6 +62,9 @@ struct PCG_API FPCGCustomVersion
 		// Introduced the concept of pin usage and graph defaults around loops, which changed the default behavior otherwise
 		UpdateGraphSettingsLoopPins = 17,
 
+		// Added 'out' filter pins on filter by tag & by type
+		UpdateFilterNodeOutputPins = 18,
+
 		// -----<new versions can be added above this line>-------------------------------------------------
 		VersionPlusOne,
 		LatestVersion = VersionPlusOne - 1
diff --git a/Engine/Plugins/Experimental/PCG/Source/PCGEditor/Private/PCGEditorGraphSchema.cpp b/Engine/Plugins/Experimental/PCG/Source/PCGEditor/Private/PCGEditorGraphSchema.cpp
--- a/Engine/Plugins/Experimental/PCG/Source/PCGEditor/Private/PCGEditorGraphSchema.cpp
+++ b/Engine/Plugins/Experimental/PCG/Source/PCGEditor/Private/PCGEditorGraphSchema.cpp
@@ -232,7 +232,7 @@ bool UPCGEditorGraphSchema::TryCreateConnectionInternal(UEdGraphPin* InA, UEdGra
 		// which can refresh the node, so we must re-query the pins after each connection is made.
 		UEdGraphPin*const* ConversionOutputPin = ConversionNode->GetAllPins().FindByPredicate([](const UEdGraphPin* InPin)
 		{
-			return InPin->Direction == EGPD_Output && InPin->GetFName() == PCGPinConstants::DefaultOutputLabel;
+			return InPin->Direction == EGPD_Output && (InPin->GetFName() == PCGPinConstants::DefaultOutputLabel || InPin->GetFName() == PCGPinConstants::DefaultInFilterLabel);
 		});
 
 		if (ensure(ConversionOutputPin && *ConversionOutputPin))
