[PCG] Renamed unit tests to match new pattern

#rnx

diff --git a/Engine/Plugins/Experimental/PCG/Source/PCG/Private/Tests/Elements/PCGFilterByIndexTest.cpp b/Engine/Plugins/Experimental/PCG/Source/PCG/Private/Tests/Elements/PCGFilterByIndexTest.cpp
--- a/Engine/Plugins/Experimental/PCG/Source/PCG/Private/Tests/Elements/PCGFilterByIndexTest.cpp
+++ b/Engine/Plugins/Experimental/PCG/Source/PCG/Private/Tests/Elements/PCGFilterByIndexTest.cpp
@@ -52,7 +52,7 @@ namespace PCGFilterByIndexTestPrivate
 	}
 }
 
-IMPLEMENT_CUSTOM_SIMPLE_AUTOMATION_TEST(FPCGFilterByIndexTest_Basic, FPCGTestBaseClass, "pcg.tests.FilterByIndex.Basic", PCGTestsCommon::TestFlags)
+IMPLEMENT_CUSTOM_SIMPLE_AUTOMATION_TEST(FPCGFilterByIndexTest_Basic, FPCGTestBaseClass, "Plugins.PCG.FilterByIndex.Basic", PCGTestsCommon::TestFlags)
 
 bool FPCGFilterByIndexTest_Basic::RunTest(const FString& Parameters)
 {
@@ -170,7 +170,7 @@ bool FPCGFilterByIndexTest_Basic::RunTest(const FString& Parameters)
 	return true;
 }
 
-IMPLEMENT_CUSTOM_SIMPLE_AUTOMATION_TEST(FPCGFilterByIndexTest_InvalidSelection, FPCGTestBaseClass, "pcg.tests.FilterByIndex.InvalidSelection", PCGTestsCommon::TestFlags)
+IMPLEMENT_CUSTOM_SIMPLE_AUTOMATION_TEST(FPCGFilterByIndexTest_InvalidSelection, FPCGTestBaseClass, "Plugins.PCG.FilterByIndex.InvalidSelection", PCGTestsCommon::TestFlags)
 
 bool FPCGFilterByIndexTest_InvalidSelection::RunTest(const FString& Parameters)
 {
diff --git a/Engine/Plugins/Experimental/PCG/Source/PCG/Private/Tests/Parsing/PCGParsingTest.cpp b/Engine/Plugins/Experimental/PCG/Source/PCG/Private/Tests/Parsing/PCGParsingTest.cpp
--- a/Engine/Plugins/Experimental/PCG/Source/PCG/Private/Tests/Parsing/PCGParsingTest.cpp
+++ b/Engine/Plugins/Experimental/PCG/Source/PCG/Private/Tests/Parsing/PCGParsingTest.cpp
@@ -19,7 +19,7 @@ namespace PCGParsingTestPrivate
 	}
 }
 
-IMPLEMENT_CUSTOM_SIMPLE_AUTOMATION_TEST(FPCGParsingTest_Basic, FPCGTestBaseClass, "pcg.tests.Parsing.Basic", PCGTestsCommon::TestFlags)
+IMPLEMENT_CUSTOM_SIMPLE_AUTOMATION_TEST(FPCGParsingTest_Basic, FPCGTestBaseClass, "Plugins.PCG.Parsing.Basic", PCGTestsCommon::TestFlags)
 
 bool FPCGParsingTest_Basic::RunTest(const FString& Parameters)
 {
@@ -95,7 +95,7 @@ bool FPCGParsingTest_Basic::RunTest(const FString& Parameters)
 	return true;
 }
 
-IMPLEMENT_CUSTOM_SIMPLE_AUTOMATION_TEST(FPCGParsingTest_Ranges, FPCGTestBaseClass, "pcg.tests.Parsing.Ranges", PCGTestsCommon::TestFlags)
+IMPLEMENT_CUSTOM_SIMPLE_AUTOMATION_TEST(FPCGParsingTest_Ranges, FPCGTestBaseClass, "Plugins.PCG.Parsing.Ranges", PCGTestsCommon::TestFlags)
 
 bool FPCGParsingTest_Ranges::RunTest(const FString& Parameters)
 {
@@ -197,7 +197,7 @@ bool FPCGParsingTest_Ranges::RunTest(const FString& Parameters)
 	return true;
 }
 
-IMPLEMENT_CUSTOM_SIMPLE_AUTOMATION_TEST(FPCGParsingTest_Robustness, FPCGTestBaseClass, "pcg.tests.Parsing.Robustness", PCGTestsCommon::TestFlags)
+IMPLEMENT_CUSTOM_SIMPLE_AUTOMATION_TEST(FPCGParsingTest_Robustness, FPCGTestBaseClass, "Plugins.PCG.Parsing.Robustness", PCGTestsCommon::TestFlags)
 
 bool FPCGParsingTest_Robustness::RunTest(const FString& Parameters)
 {
