[PCG] Factorize the default case of dynamic pins, used almost everywhere

#rnx

diff --git a/Engine/Plugins/Experimental/PCG/Source/PCG/Private/Elements/Metadata/PCGMetadataPartition.cpp b/Engine/Plugins/Experimental/PCG/Source/PCG/Private/Elements/Metadata/PCGMetadataPartition.cpp
--- a/Engine/Plugins/Experimental/PCG/Source/PCG/Private/Elements/Metadata/PCGMetadataPartition.cpp
+++ b/Engine/Plugins/Experimental/PCG/Source/PCG/Private/Elements/Metadata/PCGMetadataPartition.cpp
@@ -9,19 +9,6 @@
 
 #define LOCTEXT_NAMESPACE "PCGMetadataPartitionElement"
 
-EPCGDataType UPCGMetadataPartitionSettings::GetCurrentPinTypes(const UPCGPin* InPin) const
-{
-	check(InPin);
-	if (!InPin->IsOutputPin())
-	{
-		return Super::GetCurrentPinTypes(InPin);
-	}
-
-	// Output pin narrows to union of inputs on first pin
-	const EPCGDataType InputTypeUnion = GetTypeUnionOfIncidentEdges(PCGPinConstants::DefaultInputLabel);
-	return InputTypeUnion != EPCGDataType::None ? InputTypeUnion : EPCGDataType::Any;
-}
-
 TArray<FPCGPinProperties> UPCGMetadataPartitionSettings::InputPinProperties() const
 {
 	TArray<FPCGPinProperties> PinProperties;
diff --git a/Engine/Plugins/Experimental/PCG/Source/PCG/Private/Elements/Metadata/PCGMetadataRenameElement.cpp b/Engine/Plugins/Experimental/PCG/Source/PCG/Private/Elements/Metadata/PCGMetadataRenameElement.cpp
--- a/Engine/Plugins/Experimental/PCG/Source/PCG/Private/Elements/Metadata/PCGMetadataRenameElement.cpp
+++ b/Engine/Plugins/Experimental/PCG/Source/PCG/Private/Elements/Metadata/PCGMetadataRenameElement.cpp
@@ -12,18 +12,6 @@
 
 #define LOCTEXT_NAMESPACE "PCGMetadataRenameElement"
 
-EPCGDataType UPCGMetadataRenameSettings::GetCurrentPinTypes(const UPCGPin* InPin) const
-{
-	check(InPin);
-	if (!InPin->IsOutputPin())
-	{
-		return Super::GetCurrentPinTypes(InPin);
-	}
-
-	// Output pin narrows to union of inputs on first pin
-	const EPCGDataType InputTypeUnion = GetTypeUnionOfIncidentEdges(PCGPinConstants::DefaultInputLabel);
-	return (InputTypeUnion != EPCGDataType::None) ? InputTypeUnion : EPCGDataType::Any;
-}
 
 TArray<FPCGPinProperties> UPCGMetadataRenameSettings::InputPinProperties() const
 {
diff --git a/Engine/Plugins/Experimental/PCG/Source/PCG/Private/Elements/PCGAddTag.cpp b/Engine/Plugins/Experimental/PCG/Source/PCG/Private/Elements/PCGAddTag.cpp
--- a/Engine/Plugins/Experimental/PCG/Source/PCG/Private/Elements/PCGAddTag.cpp
+++ b/Engine/Plugins/Experimental/PCG/Source/PCG/Private/Elements/PCGAddTag.cpp
@@ -8,19 +8,6 @@
 
 #define LOCTEXT_NAMESPACE "PCGAddTagElement"
 
-EPCGDataType UPCGAddTagSettings::GetCurrentPinTypes(const UPCGPin* InPin) const
-{
-	check(InPin);
-	if (!InPin->IsOutputPin())
-	{
-		return Super::GetCurrentPinTypes(InPin);
-	}
-
-	// Output pin narrows to union of inputs on first pin
-	const EPCGDataType InputTypeUnion = GetTypeUnionOfIncidentEdges(PCGPinConstants::DefaultInputLabel);
-	return InputTypeUnion != EPCGDataType::None ? InputTypeUnion : EPCGDataType::Any;
-}
-
 TArray<FPCGPinProperties> UPCGAddTagSettings::InputPinProperties() const
 {
 	TArray<FPCGPinProperties> PinProperties;
diff --git a/Engine/Plugins/Experimental/PCG/Source/PCG/Private/Elements/PCGBranch.cpp b/Engine/Plugins/Experimental/PCG/Source/PCG/Private/Elements/PCGBranch.cpp
--- a/Engine/Plugins/Experimental/PCG/Source/PCG/Private/Elements/PCGBranch.cpp
+++ b/Engine/Plugins/Experimental/PCG/Source/PCG/Private/Elements/PCGBranch.cpp
@@ -29,18 +29,6 @@ FText UPCGBranchSettings::GetNodeTooltipText() const
 }
 #endif // WITH_EDITOR
 
-EPCGDataType UPCGBranchSettings::GetCurrentPinTypes(const UPCGPin* Pin) const
-{
-	check(Pin);
-	if (Pin->IsOutputPin())
-	{
-		const EPCGDataType InputTypeUnion = GetTypeUnionOfIncidentEdges(PCGPinConstants::DefaultInputLabel);
-		return InputTypeUnion != EPCGDataType::None ? InputTypeUnion : EPCGDataType::Any;
-	}
-
-	return Super::GetCurrentPinTypes(Pin);
-}
-
 TArray<FPCGPinProperties> UPCGBranchSettings::OutputPinProperties() const
 {
 	TArray<FPCGPinProperties> PinProperties;
diff --git a/Engine/Plugins/Experimental/PCG/Source/PCG/Private/Elements/PCGDeleteAttributesElement.cpp b/Engine/Plugins/Experimental/PCG/Source/PCG/Private/Elements/PCGDeleteAttributesElement.cpp
--- a/Engine/Plugins/Experimental/PCG/Source/PCG/Private/Elements/PCGDeleteAttributesElement.cpp
+++ b/Engine/Plugins/Experimental/PCG/Source/PCG/Private/Elements/PCGDeleteAttributesElement.cpp
@@ -67,19 +67,6 @@ TArray<FText> UPCGDeleteAttributesSettings::GetNodeTitleAliases() const
 }
 #endif
 
-EPCGDataType UPCGDeleteAttributesSettings::GetCurrentPinTypes(const UPCGPin* InPin) const
-{
-	check(InPin);
-	if (!InPin->IsOutputPin())
-	{
-		return Super::GetCurrentPinTypes(InPin);
-	}
-
-	// Output pin narrows to union of inputs on first pin
-	const EPCGDataType InputTypeUnion = GetTypeUnionOfIncidentEdges(PCGPinConstants::DefaultInputLabel);
-	return (InputTypeUnion != EPCGDataType::None) ? InputTypeUnion : EPCGDataType::Any;
-}
-
 FName UPCGDeleteAttributesSettings::AdditionalTaskName() const
 {
 	if (const UEnum* SelectionEnum = StaticEnum<EPCGAttributeFilterOperation>())
diff --git a/Engine/Plugins/Experimental/PCG/Source/PCG/Private/Elements/PCGDeleteTags.cpp b/Engine/Plugins/Experimental/PCG/Source/PCG/Private/Elements/PCGDeleteTags.cpp
--- a/Engine/Plugins/Experimental/PCG/Source/PCG/Private/Elements/PCGDeleteTags.cpp
+++ b/Engine/Plugins/Experimental/PCG/Source/PCG/Private/Elements/PCGDeleteTags.cpp
@@ -20,19 +20,6 @@ FText UPCGDeleteTagsSettings::GetDefaultNodeTitle() const
 }
 #endif // WITH_EDITOR
 
-EPCGDataType UPCGDeleteTagsSettings::GetCurrentPinTypes(const UPCGPin* InPin) const
-{
-	check(InPin);
-	if (!InPin->IsOutputPin())
-	{
-		return Super::GetCurrentPinTypes(InPin);
-	}
-
-	// Output pin narrows to union of inputs on first pin
-	const EPCGDataType InputTypeUnion = GetTypeUnionOfIncidentEdges(PCGPinConstants::DefaultInputLabel);
-	return (InputTypeUnion != EPCGDataType::None) ? InputTypeUnion : EPCGDataType::Any;
-}
-
 FName UPCGDeleteTagsSettings::AdditionalTaskName() const
 {
 	if (const UEnum* SelectionEnum = StaticEnum< EPCGTagFilterOperation>())
diff --git a/Engine/Plugins/Experimental/PCG/Source/PCG/Private/Elements/PCGFilterByType.cpp b/Engine/Plugins/Experimental/PCG/Source/PCG/Private/Elements/PCGFilterByType.cpp
--- a/Engine/Plugins/Experimental/PCG/Source/PCG/Private/Elements/PCGFilterByType.cpp
+++ b/Engine/Plugins/Experimental/PCG/Source/PCG/Private/Elements/PCGFilterByType.cpp
@@ -30,12 +30,14 @@ FText UPCGFilterByTypeSettings::GetNodeTooltipText() const
 EPCGDataType UPCGFilterByTypeSettings::GetCurrentPinTypes(const UPCGPin* InPin) const
 {
 	check(InPin);
-	if (!InPin->IsOutputPin() || InPin->Properties.Label != PCGPinConstants::DefaultInFilterLabel)
+	if (InPin->Properties.Label == PCGPinConstants::DefaultInFilterLabel)
+	{
+		return TargetType;
+	}
+	else
 	{
 		return Super::GetCurrentPinTypes(InPin);
 	}
-
-	return TargetType;
 }
 
 TArray<FPCGPinProperties> UPCGFilterByTypeSettings::OutputPinProperties() const
diff --git a/Engine/Plugins/Experimental/PCG/Source/PCG/Private/Elements/PCGFilterDataBase.cpp b/Engine/Plugins/Experimental/PCG/Source/PCG/Private/Elements/PCGFilterDataBase.cpp
--- a/Engine/Plugins/Experimental/PCG/Source/PCG/Private/Elements/PCGFilterDataBase.cpp
+++ b/Engine/Plugins/Experimental/PCG/Source/PCG/Private/Elements/PCGFilterDataBase.cpp
@@ -22,19 +22,6 @@ void UPCGFilterDataBaseSettings::ApplyDeprecationBeforeUpdatePins(UPCGNode* InOu
 
 #endif // WITH_EDITOR
 
-EPCGDataType UPCGFilterDataBaseSettings::GetCurrentPinTypes(const UPCGPin* InPin) const
-{
-	check(InPin);
-	if (!InPin->IsOutputPin())
-	{
-		return Super::GetCurrentPinTypes(InPin);
-	}
-
-	// Output pin narrows to union of inputs on first pin
-	const EPCGDataType InputTypeUnion = GetTypeUnionOfIncidentEdges(PCGPinConstants::DefaultInputLabel);
-	return InputTypeUnion != EPCGDataType::None ? InputTypeUnion : EPCGDataType::Any;
-}
-
 TArray<FPCGPinProperties> UPCGFilterDataBaseSettings::InputPinProperties() const
 {
 	TArray<FPCGPinProperties> PinProperties;
diff --git a/Engine/Plugins/Experimental/PCG/Source/PCG/Private/Elements/PCGGather.cpp b/Engine/Plugins/Experimental/PCG/Source/PCG/Private/Elements/PCGGather.cpp
--- a/Engine/Plugins/Experimental/PCG/Source/PCG/Private/Elements/PCGGather.cpp
+++ b/Engine/Plugins/Experimental/PCG/Source/PCG/Private/Elements/PCGGather.cpp
@@ -7,21 +7,6 @@
 
 #define LOCTEXT_NAMESPACE "PCGGatherElement"
 
-EPCGDataType UPCGGatherSettings::GetCurrentPinTypes(const UPCGPin* InPin) const
-{
-	check(InPin);
-
-	if (InPin->Properties.Label == PCGPinConstants::DefaultDependencyOnlyLabel || !InPin->IsOutputPin())
-	{
-		return Super::GetCurrentPinTypes(InPin);
-	}
-	else
-	{
-		const EPCGDataType InputTypeUnion = GetTypeUnionOfIncidentEdges(PCGPinConstants::DefaultInputLabel);
-		return InputTypeUnion != EPCGDataType::None ? InputTypeUnion : EPCGDataType::Any;
-	}
-}
-
 TArray<FPCGPinProperties> UPCGGatherSettings::InputPinProperties() const
 {
 	TArray<FPCGPinProperties> PinProperties;
diff --git a/Engine/Plugins/Experimental/PCG/Source/PCG/Private/Elements/PCGHiGenGridSize.cpp b/Engine/Plugins/Experimental/PCG/Source/PCG/Private/Elements/PCGHiGenGridSize.cpp
--- a/Engine/Plugins/Experimental/PCG/Source/PCG/Private/Elements/PCGHiGenGridSize.cpp
+++ b/Engine/Plugins/Experimental/PCG/Source/PCG/Private/Elements/PCGHiGenGridSize.cpp
@@ -41,7 +41,7 @@ EPCGDataType UPCGHiGenGridSizeSettings::GetCurrentPinTypes(const UPCGPin* InPin)
 	// Non-dynamically-typed pins
 	if (!InPin->IsOutputPin() || InPin->Properties.Label == PCGHiGenGridSizeConstants::CellVolumeOutputLabel)
 	{
-		return Super::GetCurrentPinTypes(InPin);
+		return InPin->Properties.AllowedTypes;
 	}
 
 	// Output pin narrows to union of inputs on first pin
diff --git a/Engine/Plugins/Experimental/PCG/Source/PCG/Private/Elements/PCGReroute.cpp b/Engine/Plugins/Experimental/PCG/Source/PCG/Private/Elements/PCGReroute.cpp
--- a/Engine/Plugins/Experimental/PCG/Source/PCG/Private/Elements/PCGReroute.cpp
+++ b/Engine/Plugins/Experimental/PCG/Source/PCG/Private/Elements/PCGReroute.cpp
@@ -20,30 +20,6 @@ UPCGRerouteSettings::UPCGRerouteSettings()
 #endif
 }
 
-EPCGDataType UPCGRerouteSettings::GetCurrentPinTypes(const UPCGPin* InPin) const
-{
-	// All pins have same type
-	EPCGDataType PinTypes = EPCGDataType::Any;
-	if (UPCGNode* PCGNode = Cast<UPCGNode>(GetOuter()))
-	{
-		if (UPCGPin* InputPin = PCGNode->GetInputPin(PCGPinConstants::DefaultInputLabel))
-		{
-			if (InputPin->EdgeCount() > 0)
-			{
-				if (UPCGEdge* Edge = InputPin->Edges[0])
-				{
-					if (const UPCGPin* OtherOutputPin = Edge->GetOtherPin(InputPin))
-					{
-						PinTypes = OtherOutputPin->GetCurrentTypes();
-					}
-				}
-			}
-		}
-	}
-
-	return PinTypes;
-}
-
 TArray<FPCGPinProperties> UPCGRerouteSettings::InputPinProperties() const
 {
 	FPCGPinProperties PinProperties;
diff --git a/Engine/Plugins/Experimental/PCG/Source/PCG/Private/PCGSettings.cpp b/Engine/Plugins/Experimental/PCG/Source/PCG/Private/PCGSettings.cpp
--- a/Engine/Plugins/Experimental/PCG/Source/PCG/Private/PCGSettings.cpp
+++ b/Engine/Plugins/Experimental/PCG/Source/PCG/Private/PCGSettings.cpp
@@ -645,6 +645,17 @@ void UPCGSettings::ApplyDeprecationBeforeUpdatePins(UPCGNode* InOutNode, TArray<
 EPCGDataType UPCGSettings::GetCurrentPinTypes(const UPCGPin* InPin) const
 {
 	check(InPin);
+
+	if (HasDynamicPins() && InPin->IsOutputPin())
+	{
+		const UPCGNode* Node = Cast<const UPCGNode>(GetOuter());
+		if (Node && Node->GetInputPin(PCGPinConstants::DefaultInputLabel) != nullptr)
+		{
+			const EPCGDataType InputTypeUnion = GetTypeUnionOfIncidentEdges(PCGPinConstants::DefaultInputLabel);
+			return InputTypeUnion != EPCGDataType::None ? InputTypeUnion : EPCGDataType::Any;
+		}
+	}
+
 	return InPin->Properties.AllowedTypes;
 }
 
diff --git a/Engine/Plugins/Experimental/PCG/Source/PCG/Public/Elements/Metadata/PCGMetadataPartition.h b/Engine/Plugins/Experimental/PCG/Source/PCG/Public/Elements/Metadata/PCGMetadataPartition.h
--- a/Engine/Plugins/Experimental/PCG/Source/PCG/Public/Elements/Metadata/PCGMetadataPartition.h
+++ b/Engine/Plugins/Experimental/PCG/Source/PCG/Public/Elements/Metadata/PCGMetadataPartition.h
@@ -24,7 +24,6 @@ public:
 	virtual EPCGSettingsType GetType() const override { return EPCGSettingsType::Metadata; }
 	virtual bool HasDynamicPins() const override { return true; }
 #endif
-	virtual EPCGDataType GetCurrentPinTypes(const UPCGPin* InPin) const override;
 
 protected:
 	virtual TArray<FPCGPinProperties> InputPinProperties() const override;
diff --git a/Engine/Plugins/Experimental/PCG/Source/PCG/Public/Elements/Metadata/PCGMetadataRenameElement.h b/Engine/Plugins/Experimental/PCG/Source/PCG/Public/Elements/Metadata/PCGMetadataRenameElement.h
--- a/Engine/Plugins/Experimental/PCG/Source/PCG/Public/Elements/Metadata/PCGMetadataRenameElement.h
+++ b/Engine/Plugins/Experimental/PCG/Source/PCG/Public/Elements/Metadata/PCGMetadataRenameElement.h
@@ -19,7 +19,6 @@ public:
 	virtual EPCGSettingsType GetType() const override { return EPCGSettingsType::Metadata; }
 	virtual bool HasDynamicPins() const override { return true; }
 #endif
-	virtual EPCGDataType GetCurrentPinTypes(const UPCGPin* InPin) const override;
 
 protected:
 	virtual TArray<FPCGPinProperties> InputPinProperties() const override;
diff --git a/Engine/Plugins/Experimental/PCG/Source/PCG/Public/Elements/PCGAddTag.h b/Engine/Plugins/Experimental/PCG/Source/PCG/Public/Elements/PCGAddTag.h
--- a/Engine/Plugins/Experimental/PCG/Source/PCG/Public/Elements/PCGAddTag.h
+++ b/Engine/Plugins/Experimental/PCG/Source/PCG/Public/Elements/PCGAddTag.h
@@ -23,7 +23,6 @@ public:
 	virtual EPCGSettingsType GetType() const override { return EPCGSettingsType::Generic; }
 	virtual bool HasDynamicPins() const override { return true; }
 #endif
-	virtual EPCGDataType GetCurrentPinTypes(const UPCGPin* InPin) const override;
 
 protected:
 	virtual TArray<FPCGPinProperties> InputPinProperties() const override;
diff --git a/Engine/Plugins/Experimental/PCG/Source/PCG/Public/Elements/PCGBranch.h b/Engine/Plugins/Experimental/PCG/Source/PCG/Public/Elements/PCGBranch.h
--- a/Engine/Plugins/Experimental/PCG/Source/PCG/Public/Elements/PCGBranch.h
+++ b/Engine/Plugins/Experimental/PCG/Source/PCG/Public/Elements/PCGBranch.h
@@ -23,8 +23,6 @@ public:
 	virtual bool HasDynamicPins() const override { return true; }
 #endif
 
-	virtual EPCGDataType GetCurrentPinTypes(const UPCGPin* Pin) const override;
-
 protected:
 	virtual TArray<FPCGPinProperties> OutputPinProperties() const override;
 	virtual FPCGElementPtr CreateElement() const override;
diff --git a/Engine/Plugins/Experimental/PCG/Source/PCG/Public/Elements/PCGDeleteAttributesElement.h b/Engine/Plugins/Experimental/PCG/Source/PCG/Public/Elements/PCGDeleteAttributesElement.h
--- a/Engine/Plugins/Experimental/PCG/Source/PCG/Public/Elements/PCGDeleteAttributesElement.h
+++ b/Engine/Plugins/Experimental/PCG/Source/PCG/Public/Elements/PCGDeleteAttributesElement.h
@@ -39,7 +39,6 @@ public:
 	virtual EPCGSettingsType GetType() const override { return EPCGSettingsType::Metadata; }
 	virtual bool HasDynamicPins() const override { return true; }
 #endif
-	virtual EPCGDataType GetCurrentPinTypes(const UPCGPin* InPin) const override;
 	virtual FName AdditionalTaskName() const override;
 
 protected:
diff --git a/Engine/Plugins/Experimental/PCG/Source/PCG/Public/Elements/PCGDeleteTags.h b/Engine/Plugins/Experimental/PCG/Source/PCG/Public/Elements/PCGDeleteTags.h
--- a/Engine/Plugins/Experimental/PCG/Source/PCG/Public/Elements/PCGDeleteTags.h
+++ b/Engine/Plugins/Experimental/PCG/Source/PCG/Public/Elements/PCGDeleteTags.h
@@ -27,7 +27,6 @@ public:
 	virtual EPCGSettingsType GetType() const override { return EPCGSettingsType::Generic; }
 	virtual bool HasDynamicPins() const override { return true; }
 #endif
-	virtual EPCGDataType GetCurrentPinTypes(const UPCGPin* InPin) const override;
 	virtual FName AdditionalTaskName() const override;
 
 protected:
diff --git a/Engine/Plugins/Experimental/PCG/Source/PCG/Public/Elements/PCGFilterDataBase.h b/Engine/Plugins/Experimental/PCG/Source/PCG/Public/Elements/PCGFilterDataBase.h
--- a/Engine/Plugins/Experimental/PCG/Source/PCG/Public/Elements/PCGFilterDataBase.h
+++ b/Engine/Plugins/Experimental/PCG/Source/PCG/Public/Elements/PCGFilterDataBase.h
@@ -19,8 +19,6 @@ public:
 	virtual void ApplyDeprecationBeforeUpdatePins(UPCGNode* InOutNode, TArray<TObjectPtr<UPCGPin>>& InputPins, TArray<TObjectPtr<UPCGPin>>& OutputPins) override;
 #endif // WITH_EDITOR
 
-	virtual EPCGDataType GetCurrentPinTypes(const UPCGPin* InPin) const override;
-
 protected:
 	virtual TArray<FPCGPinProperties> InputPinProperties() const override;
 	virtual TArray<FPCGPinProperties> OutputPinProperties() const override;
diff --git a/Engine/Plugins/Experimental/PCG/Source/PCG/Public/Elements/PCGGather.h b/Engine/Plugins/Experimental/PCG/Source/PCG/Public/Elements/PCGGather.h
--- a/Engine/Plugins/Experimental/PCG/Source/PCG/Public/Elements/PCGGather.h
+++ b/Engine/Plugins/Experimental/PCG/Source/PCG/Public/Elements/PCGGather.h
@@ -22,7 +22,6 @@ public:
 	virtual EPCGSettingsType GetType() const override { return EPCGSettingsType::Spatial; }
 	virtual bool HasDynamicPins() const override { return true; }
 #endif
-	virtual EPCGDataType GetCurrentPinTypes(const UPCGPin* InPin) const override;
 
 protected:
 	virtual TArray<FPCGPinProperties> InputPinProperties() const override;
diff --git a/Engine/Plugins/Experimental/PCG/Source/PCG/Public/Elements/PCGReroute.h b/Engine/Plugins/Experimental/PCG/Source/PCG/Public/Elements/PCGReroute.h
--- a/Engine/Plugins/Experimental/PCG/Source/PCG/Public/Elements/PCGReroute.h
+++ b/Engine/Plugins/Experimental/PCG/Source/PCG/Public/Elements/PCGReroute.h
@@ -24,7 +24,6 @@ public:
 
 	//~Begin UPCGSettings interface
 	virtual bool HasDynamicPins() const override { return true; }
-	virtual EPCGDataType GetCurrentPinTypes(const UPCGPin* InPin) const override;
 
 #if WITH_EDITOR
 	virtual FName GetDefaultNodeName() const override { return FName("Reroute"); }
diff --git a/Engine/Plugins/Experimental/PCG/Source/PCG/Public/PCGSettings.h b/Engine/Plugins/Experimental/PCG/Source/PCG/Public/PCGSettings.h
--- a/Engine/Plugins/Experimental/PCG/Source/PCG/Public/PCGSettings.h
+++ b/Engine/Plugins/Experimental/PCG/Source/PCG/Public/PCGSettings.h
@@ -306,7 +306,9 @@ public:
 	/** Returns true if only the first input edge is used from the primary pin when the node is disabled. */
 	virtual bool OnlyPassThroughOneEdgeWhenDisabled() const { return false; }
 
-	/** Returns the current pin types, which can either be the static types from the pin properties, or a dynamic type based on connected edges. */
+	/** Returns the current pin types, which can either be the static types from the pin properties, or a dynamic type based on connected edges.
+	* By default we set output pin types to the union of the default input pin incident edge types, if it is dynamic and the default input exists.
+	*/
 	virtual EPCGDataType GetCurrentPinTypes(const UPCGPin* InPin) const;
 
 	UPROPERTY(EditAnywhere, BlueprintReadWrite, Category = Settings, meta=(EditCondition=bUseSeed, EditConditionHides, PCG_Overridable))
