[PCG] Nativization of DistanceToNeighbors node
--Now named Point Neighborhood
--Includes Unit Tests
--Internal function for Project Points
   *Using bounds for weighted averages now optional

diff --git a/Engine/Plugins/Experimental/PCG/Source/PCG/Private/Data/PCGPointData.cpp b/Engine/Plugins/Experimental/PCG/Source/PCG/Private/Data/PCGPointData.cpp
--- a/Engine/Plugins/Experimental/PCG/Source/PCG/Private/Data/PCGPointData.cpp
+++ b/Engine/Plugins/Experimental/PCG/Source/PCG/Private/Data/PCGPointData.cpp
@@ -489,6 +489,11 @@ bool UPCGPointData::SamplePoint(const FTransform& InTransform, const FBox& InBou
 }
 
 bool UPCGPointData::ProjectPoint(const FTransform& InTransform, const FBox& InBounds, const FPCGProjectionParams& InParams, FPCGPoint& OutPoint, UPCGMetadata* OutMetadata) const
+{
+	return ProjectPoint(InTransform, InBounds, InParams, OutPoint, OutMetadata, true);
+}
+
+bool UPCGPointData::ProjectPoint(const FTransform& InTransform, const FBox& InBounds, const FPCGProjectionParams& InParams, FPCGPoint& OutPoint, UPCGMetadata* OutMetadata, bool bUseBounds) const
 {
 	//TRACE_CPUPROFILER_EVENT_SCOPE(UPCGPointData::SamplePoint);
 	if (bOctreeIsDirty)
@@ -502,7 +507,8 @@ bool UPCGPointData::ProjectPoint(const FTransform& InTransform, const FBox& InBo
 	if (!bSampleInVolume)
 	{
 		const FVector InPosition = InTransform.GetLocation();
-		Octree.FindElementsWithBoundsTest(FBoxCenterAndExtent(InPosition, FVector::Zero()), [&InPosition, &Contributions](const FPCGPointRef& InPointRef) {
+		Octree.FindElementsWithBoundsTest(FBoxCenterAndExtent(InPosition, FVector::Zero()), [&InPosition, &Contributions](const FPCGPointRef& InPointRef) 
+		{
 			Contributions.Emplace(InPointRef.Point, PCGPointHelpers::InverseEuclidianDistance(*InPointRef.Point, InPosition));
 		});
 	}
@@ -511,8 +517,9 @@ bool UPCGPointData::ProjectPoint(const FTransform& InTransform, const FBox& InBo
 		FBox TransformedBounds = InBounds.TransformBy(InTransform);
 		FMatrix InTransformInverseMatrix = InTransform.ToMatrixWithScale().Inverse();
 
-		Octree.FindElementsWithBoundsTest(FBoxCenterAndExtent(TransformedBounds.GetCenter(), TransformedBounds.GetExtent()), [&InBounds, &InTransformInverseMatrix, &Contributions](const FPCGPointRef& InPointRef) {
-			float Contribution = PCGPointHelpers::VolumeOverlap(*InPointRef.Point, InBounds, InTransformInverseMatrix);
+		Octree.FindElementsWithBoundsTest(FBoxCenterAndExtent(TransformedBounds.GetCenter(), TransformedBounds.GetExtent()), [bUseBounds, &InBounds, &InTransformInverseMatrix, &Contributions](const FPCGPointRef& InPointRef) 
+		{
+			const FVector::FReal Contribution = bUseBounds ? PCGPointHelpers::VolumeOverlap(*InPointRef.Point, InBounds, InTransformInverseMatrix) : 1.0;
 			if (Contribution > 0)
 			{
 				Contributions.Emplace(InPointRef.Point, Contribution);
@@ -576,7 +583,7 @@ bool UPCGPointData::ProjectPoint(const FTransform& InTransform, const FBox& InBo
 		}
 		else
 		{
-			WeightedDensity += SourcePoint.Density * Contribution.Value * DensityNormalizationFactor;
+			WeightedDensity += SourcePoint.Density * (bUseBounds ? (Contribution.Value * DensityNormalizationFactor) : Weight);
 		}
 
 		WeightedBoundsMin += SourcePoint.BoundsMin * Weight;
diff --git a/Engine/Plugins/Experimental/PCG/Source/PCG/Private/Elements/PCGPointNeighborhood.cpp b/Engine/Plugins/Experimental/PCG/Source/PCG/Private/Elements/PCGPointNeighborhood.cpp
new file mode 100644
--- /dev/null
+++ b/Engine/Plugins/Experimental/PCG/Source/PCG/Private/Elements/PCGPointNeighborhood.cpp
@@ -0,0 +1,166 @@
+// Copyright Epic Games, Inc. All Rights Reserved.
+
+#include "Elements/PCGPointNeighborhood.h"
+
+#include "PCGContext.h"
+#include "Data/PCGPointData.h"
+#include "Helpers/PCGAsync.h"
+#include "Metadata/Accessors/PCGAttributeAccessorHelpers.h"
+
+#define LOCTEXT_NAMESPACE "PCGPointNeighborhoodElement"
+
+namespace PCGPointNeighborhood
+{
+	template<typename T>
+	void SetAttributeHelper(UPCGPointData* Data, const FName& AttributeName, const TArrayView<const T> Values)
+	{
+		// Discard None attributes
+		if (AttributeName == NAME_None)
+		{
+			return;
+		}
+
+		ensure(Data->Metadata->FindOrCreateAttribute<T>(AttributeName));
+
+		FPCGAttributePropertySelector AttributeSelector = FPCGAttributePropertySelector::CreateAttributeSelector(AttributeName);
+
+		TUniquePtr<IPCGAttributeAccessor> Accessor = PCGAttributeAccessorHelpers::CreateAccessor(Data, AttributeSelector);
+		TUniquePtr<IPCGAttributeAccessorKeys> Keys = PCGAttributeAccessorHelpers::CreateKeys(Data, AttributeSelector);
+		
+		if (Accessor.IsValid() && Keys.IsValid())
+		{
+			Accessor->SetRange(Values, 0, *Keys);
+		}
+		else
+		{
+			return;
+		}
+	}
+}
+
+FPCGElementPtr UPCGPointNeighborhoodSettings::CreateElement() const
+{
+	return MakeShared<FPCGPointNeighborhoodElement>();
+}
+
+bool FPCGPointNeighborhoodElement::ExecuteInternal(FPCGContext* Context) const
+{
+	TRACE_CPUPROFILER_EVENT_SCOPE(FPCGPointNeighborhoodElement::Execute);
+
+	const UPCGPointNeighborhoodSettings* Settings = Context->GetInputSettings<UPCGPointNeighborhoodSettings>();
+	check(Settings);
+
+	TArray<FPCGTaggedData> Inputs = Context->InputData.GetInputs();
+	TArray<FPCGTaggedData>& Outputs = Context->OutputData.TaggedData;
+
+	const double SearchDistance = Settings->SearchDistance;
+	if (SearchDistance < UE_DOUBLE_SMALL_NUMBER)
+	{
+		PCGE_LOG(Error, GraphAndLog, LOCTEXT("InvalidSearchDistance", "Search Distance must be greater than 0."));
+		return true;
+	}
+
+	struct FProcessResults
+	{
+		TArray<double> Distances;
+		TArray<FVector> AveragePositions;
+	};
+
+	FProcessResults OutputDataBuffers;
+
+	for (int i = 0; i < Inputs.Num(); ++i)
+	{
+		const UPCGPointData* InputPointData = Cast<UPCGPointData>(Inputs[i].Data);
+		if (!InputPointData)
+		{
+			PCGE_LOG(Verbose, GraphAndLog, FText::Format(LOCTEXT("InvalidPointData", "Input {0} is not point data"), i));
+			continue;
+		}
+
+		const TArray<FPCGPoint>& SrcPoints = InputPointData->GetPoints();
+		FPCGTaggedData& Output = Outputs.Add_GetRef(Inputs[i]);
+		UPCGPointData* OutputPointData = NewObject<UPCGPointData>();
+		TArray<FPCGPoint>& OutputPoints = OutputPointData->GetMutablePoints();
+		OutputPoints.SetNumUninitialized(SrcPoints.Num());
+		Output.Data = OutputPointData;
+		
+		const FBox SearchBounds(FVector(SearchDistance * -1.0), FVector(SearchDistance));
+
+		FPCGProjectionParams Params{};
+		Params.bProjectRotations = Params.bProjectScales = false;
+		Params.ColorBlendMode = Settings->bSetAverageColor ? EPCGProjectionColorBlendMode::TargetValue : EPCGProjectionColorBlendMode::SourceValue;
+
+		auto InitializeBuffers = [Settings, &OutputDataBuffers, Count = SrcPoints.Num()]()
+		{
+			if (Settings->bSetDistanceToAttribute)
+			{
+				OutputDataBuffers.Distances.SetNumUninitialized(Count);
+			}
+
+			if (Settings->bSetAveragePositionToAttribute)
+			{
+				OutputDataBuffers.AveragePositions.SetNumUninitialized(Count);
+			}
+		};
+
+		auto ProcessPoint = [Settings, &OutputPoints, &OutputDataBuffers, &SrcPoints, &SearchBounds, &SearchDistance, &Params, &InputPointData, &OutputPointData](const int32 ReadIndex, const int32 WriteIndex)
+		{
+
+			const FPCGPoint& InPoint = SrcPoints[ReadIndex];
+			FPCGPoint& OutPoint = OutputPoints[WriteIndex];
+			OutPoint = InPoint;
+
+			const FTransform InTransform(InPoint.Transform.GetLocation());
+			UPCGMetadata* OutMetadata = OutputPointData->Metadata;
+
+			FPCGPoint ProjectionPoint = FPCGPoint();
+			InputPointData->ProjectPoint(InTransform, SearchBounds, Params, ProjectionPoint, OutMetadata, Settings->bWeightedAverage);
+
+			const double NormalizedDistance = FVector::Distance(InPoint.Transform.GetLocation(), ProjectionPoint.Transform.GetLocation()) / SearchDistance;
+
+			if (Settings->SetDensity == EPCGPointNeighborhoodDensityMode::SetNormalizedDistanceToDensity)
+			{
+				OutPoint.Density = FMath::Clamp(NormalizedDistance, 0.0, 1.0);
+			}
+			else if (Settings->SetDensity == EPCGPointNeighborhoodDensityMode::SetAverageDensity)
+			{
+				OutPoint.Density = ProjectionPoint.Density;
+			}
+
+			if (Settings->bSetDistanceToAttribute)
+			{
+				OutputDataBuffers.Distances[WriteIndex] = FVector::Distance(InPoint.Transform.GetLocation(), ProjectionPoint.Transform.GetLocation());
+			}
+
+			if (Settings->bSetAveragePosition)
+			{
+				OutPoint.Transform.SetLocation(ProjectionPoint.Transform.GetLocation());
+			}
+
+			if (Settings->bSetAveragePositionToAttribute)
+			{
+				OutputDataBuffers.AveragePositions[WriteIndex] = ProjectionPoint.Transform.GetLocation();
+			}
+
+			if (Settings->bSetAverageColor)
+			{
+				OutPoint.Color = ProjectionPoint.Color;
+			}
+		};
+
+		FPCGAsync::AsyncProcessingOneToOneEx(&Context->AsyncState, SrcPoints.Num(), InitializeBuffers, ProcessPoint, /* bEnableTimeSlicing=*/false);
+
+		if (Settings->bSetDistanceToAttribute)
+		{
+			PCGPointNeighborhood::SetAttributeHelper<double>(OutputPointData, Settings->DistanceAttribute, OutputDataBuffers.Distances);
+		}
+		if (Settings->bSetAveragePositionToAttribute)
+		{
+			PCGPointNeighborhood::SetAttributeHelper<FVector>(OutputPointData, Settings->AveragePositionAttribute, OutputDataBuffers.AveragePositions);
+		}
+	}
+
+	return true;
+}
+
+#undef LOCTEXT_NAMESPACE
diff --git a/Engine/Plugins/Experimental/PCG/Source/PCG/Private/Tests/Elements/PCGPointNeighborhoodTest.cpp b/Engine/Plugins/Experimental/PCG/Source/PCG/Private/Tests/Elements/PCGPointNeighborhoodTest.cpp
new file mode 100644
--- /dev/null
+++ b/Engine/Plugins/Experimental/PCG/Source/PCG/Private/Tests/Elements/PCGPointNeighborhoodTest.cpp
@@ -0,0 +1,334 @@
+// Copyright Epic Games, Inc. All Rights Reserved.
+
+#include "Tests/PCGTestsCommon.h"
+
+#include "PCGContext.h"
+#include "Data/PCGPointData.h"
+
+#include "Elements/PCGPointNeighborhood.h"
+
+class PCGPointNeighborhoodTestBase : public FPCGTestBaseClass
+{
+public:
+	using FPCGTestBaseClass::FPCGTestBaseClass;
+
+	struct FTestParameters
+	{
+		double SearchDistance = 0.0;
+		bool bDistanceToAttr = false;
+		FName DistanceName;
+		bool bPositionToAttr = false;
+		FName PositionName;
+		EPCGPointNeighborhoodDensityMode SetDensity = EPCGPointNeighborhoodDensityMode::None;
+		bool bAvgPosition = false;
+		bool bAvgColor = false;
+		bool bUseBounds = true;
+	};
+
+protected:
+	TUniquePtr<FPCGContext> GenerateTestDataAndRun(const FTestParameters& Parameters)
+	{
+		PCGTestsCommon::FTestData TestData;
+		PCGTestsCommon::GenerateSettings<UPCGPointNeighborhoodSettings>(TestData);
+		UPCGPointNeighborhoodSettings* Settings = CastChecked<UPCGPointNeighborhoodSettings>(TestData.Settings);
+		Settings->SearchDistance = Parameters.SearchDistance;
+		Settings->bSetDistanceToAttribute = Parameters.bDistanceToAttr;
+		Settings->DistanceAttribute = Parameters.DistanceName;
+		Settings->bSetAveragePositionToAttribute = Parameters.bPositionToAttr;
+		Settings->AveragePositionAttribute = Parameters.PositionName;
+		Settings->SetDensity = Parameters.SetDensity;
+		Settings->bSetAveragePosition = Parameters.bAvgPosition;
+		Settings->bSetAverageColor = Parameters.bAvgColor;
+		Settings->bWeightedAverage = Parameters.bUseBounds;
+
+		FPCGTaggedData& Inputs = TestData.InputData.TaggedData.Emplace_GetRef();
+		Inputs.Pin = PCGPinConstants::DefaultInputLabel;
+		UPCGPointData* InData = PCGTestsCommon::CreatePointData();
+		Inputs.Data = InData;
+
+		TArray<FPCGPoint>& InPoints = InData->GetMutablePoints();
+		InPoints.SetNum(9);
+
+		for (int i = 0; i < 3; ++i)
+		{
+			for (int j = 0; j < 3; ++j)
+			{
+				InPoints[(i * 3) + j].Transform.SetLocation(FVector(j * 100.0, i * 100.0, 0.0));
+			}
+		}
+
+		for (FPCGPoint& InPoint : InPoints)
+		{
+			InPoint.BoundsMin = FVector(-10.0);
+			InPoint.BoundsMax = FVector(10.0);
+		}
+
+		FPCGElementPtr TestElement = TestData.Settings->GetElement();
+
+		TUniquePtr<FPCGContext> Context = TestData.InitializeTestContext();
+
+		while (!TestElement->Execute(Context.Get())) {}
+
+		return Context;
+	}
+};
+
+IMPLEMENT_CUSTOM_SIMPLE_AUTOMATION_TEST(FPCGPointNeighborhoodTest_AttrDistance, PCGPointNeighborhoodTestBase, "Plugins.PCG.PointNeighborhood.AttrDistance", PCGTestsCommon::TestFlags)
+
+bool FPCGPointNeighborhoodTest_AttrDistance::RunTest(const FString& Parameters)
+{
+	FTestParameters DistanceAttrParameters{};
+	DistanceAttrParameters.SearchDistance = 225.0;
+	DistanceAttrParameters.bDistanceToAttr = true;
+	DistanceAttrParameters.DistanceName = TEXT("Distance");
+	DistanceAttrParameters.bPositionToAttr = false;
+	DistanceAttrParameters.PositionName = NAME_None;
+	DistanceAttrParameters.SetDensity = EPCGPointNeighborhoodDensityMode::None;
+	DistanceAttrParameters.bAvgPosition = false;
+	DistanceAttrParameters.bAvgColor = false;
+	DistanceAttrParameters.bUseBounds = true;
+
+	TUniquePtr<FPCGContext> Context = GenerateTestDataAndRun(DistanceAttrParameters);
+
+	const TArray<FPCGTaggedData>& Outputs = Context->OutputData.TaggedData;
+
+	UTEST_EQUAL("Output count", Outputs.Num(), 1);
+
+	const UPCGPointData* OutPointData = Cast<UPCGPointData>(Outputs[0].Data);
+
+	UTEST_NOT_NULL("Output point data", OutPointData);
+
+	const TArray<FPCGPoint>& OutPoints = OutPointData->GetPoints();
+
+	UTEST_EQUAL("Output point count", OutPoints.Num(), 9);
+
+	const FPCGMetadataAttribute<double>* DistanceAttr = OutPointData->Metadata->GetConstTypedAttribute<double>(DistanceAttrParameters.DistanceName);
+	UTEST_NOT_NULL("Distance Attribute exists", DistanceAttr);
+	UTEST_EQUAL("Point 4 of Distance Attribute", DistanceAttr->GetValueFromItemKey(OutPoints[4].MetadataEntry), 0.0)
+
+	return true; 
+}
+
+IMPLEMENT_CUSTOM_SIMPLE_AUTOMATION_TEST(FPCGPointNeighborhoodTest_AttrPosition, PCGPointNeighborhoodTestBase, "Plugins.PCG.PointNeighborhood.AttrPosition", PCGTestsCommon::TestFlags)
+
+bool FPCGPointNeighborhoodTest_AttrPosition::RunTest(const FString& Parameters)
+{
+	FTestParameters PositionAttrParameters{};
+	PositionAttrParameters.SearchDistance = 225.0;
+	PositionAttrParameters.bDistanceToAttr = false;
+	PositionAttrParameters.DistanceName = NAME_None;
+	PositionAttrParameters.bPositionToAttr = true;
+	PositionAttrParameters.PositionName = TEXT("AvgPosition");
+	PositionAttrParameters.SetDensity = EPCGPointNeighborhoodDensityMode::None;
+	PositionAttrParameters.bAvgPosition = false;
+	PositionAttrParameters.bAvgColor = false;
+	PositionAttrParameters.bUseBounds = true;
+
+	TUniquePtr<FPCGContext> Context = GenerateTestDataAndRun(PositionAttrParameters);
+
+	const TArray<FPCGTaggedData>& Outputs = Context->OutputData.TaggedData;
+
+	UTEST_EQUAL("Output count", Outputs.Num(), 1);
+
+	const UPCGPointData* OutPointData = Cast<UPCGPointData>(Outputs[0].Data);
+
+	UTEST_NOT_NULL("Output point data", OutPointData);
+
+	const TArray<FPCGPoint>& OutPoints = OutPointData->GetPoints();
+
+	UTEST_EQUAL("Output point count", OutPoints.Num(), 9);
+
+	const FPCGMetadataAttribute<FVector>* PositionAttr = OutPointData->Metadata->GetConstTypedAttribute<FVector>(PositionAttrParameters.PositionName);
+	UTEST_NOT_NULL("Average Position Attribute exists", PositionAttr);
+	UTEST_EQUAL("Point 4 of Position Attribute", PositionAttr->GetValueFromItemKey(OutPoints[4].MetadataEntry), FVector(100.0, 100.0, 0.0));
+
+	return true;
+}
+
+IMPLEMENT_CUSTOM_SIMPLE_AUTOMATION_TEST(FPCGPointNeighborhoodTest_DensityAvg, PCGPointNeighborhoodTestBase, "Plugins.PCG.PointNeighborhood.DensityAvg", PCGTestsCommon::TestFlags)
+
+bool FPCGPointNeighborhoodTest_DensityAvg::RunTest(const FString& Parameters)
+{
+	FTestParameters AvgDensityParameters{};
+	AvgDensityParameters.SearchDistance = 225.0;
+	AvgDensityParameters.bDistanceToAttr = false;
+	AvgDensityParameters.DistanceName = NAME_None;
+	AvgDensityParameters.bPositionToAttr = false;
+	AvgDensityParameters.PositionName = NAME_None;
+	AvgDensityParameters.SetDensity = EPCGPointNeighborhoodDensityMode::SetAverageDensity;
+	AvgDensityParameters.bAvgPosition = false;
+	AvgDensityParameters.bAvgColor = false;
+	AvgDensityParameters.bUseBounds = true;
+
+	TUniquePtr<FPCGContext> Context = GenerateTestDataAndRun(AvgDensityParameters);
+
+	const TArray<FPCGTaggedData>& Outputs = Context->OutputData.TaggedData;
+
+	UTEST_EQUAL("Output count", Outputs.Num(), 1);
+
+	const UPCGPointData* OutPointData = Cast<UPCGPointData>(Outputs[0].Data);
+
+	UTEST_NOT_NULL("Output point data", OutPointData);
+
+	const TArray<FPCGPoint>& OutPoints = OutPointData->GetPoints();
+
+	UTEST_EQUAL("Output point count", OutPoints.Num(), 9);
+
+	/** Steepness taken into account:
+	 * Point Bounds = 30 -> (20cm * (1 + 0.5))
+	 * Search Bounds = 450 (|-225| + 225)
+	 * Given 9 contributions:
+	 * 9 * (30 ^ 3) / (450 ^ 3) == 0.00266...*/
+	UTEST_EQUAL("Average Density in Point 0 Density Attribute", OutPoints[0].Density, 9.0f * FMath::Pow(30.0f, 3.0f) / FMath::Pow(450.0f, 3.0f));
+	UTEST_EQUAL("Average Density in Point 4 Density Attribute", OutPoints[4].Density, 9.0f * FMath::Pow(30.0f, 3.0f) / FMath::Pow(450.0f, 3.0f));
+	
+	return true;
+}
+
+IMPLEMENT_CUSTOM_SIMPLE_AUTOMATION_TEST(FPCGPointNeighborhoodTest_DensityAvgNoBounds, PCGPointNeighborhoodTestBase, "Plugins.PCG.PointNeighborhood.DensityAvgNoBounds", PCGTestsCommon::TestFlags)
+
+bool FPCGPointNeighborhoodTest_DensityAvgNoBounds::RunTest(const FString& Parameters)
+{
+	FTestParameters NoBoundsAvgParameters{};
+	NoBoundsAvgParameters.SearchDistance = 225.0;
+	NoBoundsAvgParameters.bDistanceToAttr = false;
+	NoBoundsAvgParameters.DistanceName = NAME_None;
+	NoBoundsAvgParameters.bPositionToAttr = false;
+	NoBoundsAvgParameters.PositionName = NAME_None;
+	NoBoundsAvgParameters.SetDensity = EPCGPointNeighborhoodDensityMode::SetAverageDensity;
+	NoBoundsAvgParameters.bAvgPosition = false;
+	NoBoundsAvgParameters.bAvgColor = false;
+	NoBoundsAvgParameters.bUseBounds = false;
+
+	TUniquePtr<FPCGContext> Context = GenerateTestDataAndRun(NoBoundsAvgParameters);
+
+	const TArray<FPCGTaggedData>& Outputs = Context->OutputData.TaggedData;
+
+	UTEST_EQUAL("Output count", Outputs.Num(), 1);
+
+	const UPCGPointData* OutPointData = Cast<UPCGPointData>(Outputs[0].Data);
+
+	UTEST_NOT_NULL("Output point data", OutPointData);
+
+	const TArray<FPCGPoint>& OutPoints = OutPointData->GetPoints();
+
+	UTEST_EQUAL("Output point count", OutPoints.Num(), 9);
+
+	/** Density / num points total = 1 / 9 contribution per point
+	 * 9 * 1/9 == 1 average density */ 
+	UTEST_EQUAL("Average Density in Point 0 Density Attribute", OutPoints[0].Density, 1.0f);
+	UTEST_EQUAL("Average Density in Point 4 Density Attribute", OutPoints[4].Density, 1.0f);
+
+	return true;
+}
+
+IMPLEMENT_CUSTOM_SIMPLE_AUTOMATION_TEST(FPCGPointNeighborhoodTest_DensityNormal, PCGPointNeighborhoodTestBase, "Plugins.PCG.PointNeighborhood.DensityNormal", PCGTestsCommon::TestFlags)
+
+bool FPCGPointNeighborhoodTest_DensityNormal::RunTest(const FString& Parameters)
+{
+	FTestParameters NormalDensityParameters{};
+	NormalDensityParameters.SearchDistance = 225.0;
+	NormalDensityParameters.bDistanceToAttr = false;
+	NormalDensityParameters.DistanceName = NAME_None;
+	NormalDensityParameters.bPositionToAttr = false;
+	NormalDensityParameters.PositionName = NAME_None;
+	NormalDensityParameters.SetDensity = EPCGPointNeighborhoodDensityMode::SetNormalizedDistanceToDensity;
+	NormalDensityParameters.bAvgPosition = false;
+	NormalDensityParameters.bAvgColor = false;
+	NormalDensityParameters.bUseBounds = true;
+
+	TUniquePtr<FPCGContext> Context = GenerateTestDataAndRun(NormalDensityParameters);
+
+	const TArray<FPCGTaggedData>& Outputs = Context->OutputData.TaggedData;
+
+	UTEST_EQUAL("Output count", Outputs.Num(), 1);
+
+	const UPCGPointData* OutPointData = Cast<UPCGPointData>(Outputs[0].Data);
+
+	UTEST_NOT_NULL("Output point data", OutPointData);
+
+	const TArray<FPCGPoint>& OutPoints = OutPointData->GetPoints();
+
+	UTEST_EQUAL("Output point count", OutPoints.Num(), 9);
+
+	/** Distance between each point / search distance (radius)
+	 * sqrt(100^2 + 100^2) == 100*sqrt(2) / 225
+	 * == 0.629...
+	*/
+	UTEST_EQUAL("Normalized Distance in Point 0 Density Attribute", OutPoints[0].Density, (100.0f * FMath::Sqrt(2.0f)) / 225.0f);
+	UTEST_EQUAL("Normalized Distance in Point 4 Density Attribute", OutPoints[4].Density, 0.0f);
+
+	return true;
+}
+
+IMPLEMENT_CUSTOM_SIMPLE_AUTOMATION_TEST(FPCGPointNeighborhoodTest_AvgPosition, PCGPointNeighborhoodTestBase, "Plugins.PCG.PointNeighborhood.AvgPosition", PCGTestsCommon::TestFlags)
+
+bool FPCGPointNeighborhoodTest_AvgPosition::RunTest(const FString& Parameters)
+{
+	FTestParameters AveragePositionParameters{};
+	AveragePositionParameters.SearchDistance = 225.0;
+	AveragePositionParameters.bDistanceToAttr = false;
+	AveragePositionParameters.DistanceName = NAME_None;
+	AveragePositionParameters.bPositionToAttr = false;
+	AveragePositionParameters.PositionName = NAME_None;
+	AveragePositionParameters.SetDensity = EPCGPointNeighborhoodDensityMode::None;
+	AveragePositionParameters.bAvgPosition = true;
+	AveragePositionParameters.bAvgColor = false;
+	AveragePositionParameters.bUseBounds = true;
+
+	TUniquePtr<FPCGContext> Context = GenerateTestDataAndRun(AveragePositionParameters);
+
+	const TArray<FPCGTaggedData>& Outputs = Context->OutputData.TaggedData;
+
+	UTEST_EQUAL("Output count", Outputs.Num(), 1);
+
+	const UPCGPointData* OutPointData = Cast<UPCGPointData>(Outputs[0].Data);
+
+	UTEST_NOT_NULL("Output point data", OutPointData);
+
+	const TArray<FPCGPoint>& OutPoints = OutPointData->GetPoints();
+
+	UTEST_EQUAL("Output point count", OutPoints.Num(), 9);
+
+	UTEST_EQUAL("Average Position in Point 0 Location Attribute", OutPoints[0].Transform.GetLocation(), FVector(100.0, 100.0, 0.0));
+	UTEST_EQUAL("Average Position in Point 4 Location Attribute", OutPoints[4].Transform.GetLocation(), FVector(100.0, 100.0, 0.0));
+
+	return true;
+}
+
+IMPLEMENT_CUSTOM_SIMPLE_AUTOMATION_TEST(FPCGPointNeighborhoodTest_AvgColor, PCGPointNeighborhoodTestBase, "Plugins.PCG.PointNeighborhood.AvgColor", PCGTestsCommon::TestFlags)
+
+bool FPCGPointNeighborhoodTest_AvgColor::RunTest(const FString& Parameters)
+{
+	FTestParameters AverageColorParameters{};
+	AverageColorParameters.SearchDistance = 225.0;
+	AverageColorParameters.bDistanceToAttr = false;
+	AverageColorParameters.DistanceName = NAME_None;
+	AverageColorParameters.bPositionToAttr = false;
+	AverageColorParameters.PositionName = NAME_None;
+	AverageColorParameters.SetDensity = EPCGPointNeighborhoodDensityMode::None;
+	AverageColorParameters.bAvgPosition = true;
+	AverageColorParameters.bAvgColor = false;
+	AverageColorParameters.bUseBounds = true;
+
+	TUniquePtr<FPCGContext> Context = GenerateTestDataAndRun(AverageColorParameters);
+
+	const TArray<FPCGTaggedData>& Outputs = Context->OutputData.TaggedData;
+
+	UTEST_EQUAL("Output count", Outputs.Num(), 1);
+
+	const UPCGPointData* OutPointData = Cast<UPCGPointData>(Outputs[0].Data);
+
+	UTEST_NOT_NULL("Output point data", OutPointData);
+
+	const TArray<FPCGPoint>& OutPoints = OutPointData->GetPoints();
+
+	UTEST_EQUAL("Output point count", OutPoints.Num(), 9);
+
+	for (int i = 0; i < OutPoints.Num(); ++i)
+	{
+		UTEST_EQUAL("Average Color in Location Attribute", OutPoints[1].Color, FVector4(1.0, 1.0, 1.0, 1.0));
+	}
+	
+	return true;
+}
\ No newline at end of file
diff --git a/Engine/Plugins/Experimental/PCG/Source/PCG/Public/Data/PCGPointData.h b/Engine/Plugins/Experimental/PCG/Source/PCG/Public/Data/PCGPointData.h
--- a/Engine/Plugins/Experimental/PCG/Source/PCG/Public/Data/PCGPointData.h
+++ b/Engine/Plugins/Experimental/PCG/Source/PCG/Public/Data/PCGPointData.h
@@ -120,6 +120,8 @@ public:
 
 	const PointOctree& GetOctree() const;
 
+	bool ProjectPoint(const FTransform& InTransform, const FBox& InBounds, const FPCGProjectionParams& InParams, FPCGPoint& OutPoint, UPCGMetadata* OutMetadata, bool bUseBounds) const;
+
 protected:
 	void RebuildOctree() const;
 	void RecomputeBounds() const;
diff --git a/Engine/Plugins/Experimental/PCG/Source/PCG/Public/Elements/PCGPointNeighborhood.h b/Engine/Plugins/Experimental/PCG/Source/PCG/Public/Elements/PCGPointNeighborhood.h
new file mode 100644
--- /dev/null
+++ b/Engine/Plugins/Experimental/PCG/Source/PCG/Public/Elements/PCGPointNeighborhood.h
@@ -0,0 +1,80 @@
+// Copyright Epic Games, Inc. All Rights Reserved.
+#pragma once
+
+#include "PCGSettings.h"
+
+#include "PCGPointNeighborhood.generated.h"
+
+UENUM()
+enum class EPCGPointNeighborhoodDensityMode
+{
+	None,
+	SetNormalizedDistanceToDensity,
+	SetAverageDensity
+};
+
+/**
+* Computes quantities from nearby neighbor points, such as average density, color, and position.
+ */
+UCLASS(BlueprintType, ClassGroup = (Procedural))
+class UPCGPointNeighborhoodSettings : public UPCGSettings
+{
+	GENERATED_BODY()
+
+public:
+	//~Begin UPCGSettings interface
+#if WITH_EDITOR
+	virtual FName GetDefaultNodeName() const override { return FName(TEXT("PointNeighborhood")); }
+	virtual FText GetDefaultNodeTitle() const override { return NSLOCTEXT("PCGPointNeighborhoodElement", "NodeTitle", "Point Neighborhood"); }
+	virtual FText GetNodeTooltipText() const override { return NSLOCTEXT("PCGPointNeighborhoodElement", "NodeTooltip", "Computes quantities from nearby neighbor points, such as average density, color, and position."); }
+	virtual EPCGSettingsType GetType() const override { return EPCGSettingsType::Spatial; }
+#endif
+
+protected:
+	virtual TArray<FPCGPinProperties> InputPinProperties() const override { return Super::DefaultPointInputPinProperties(); }
+	virtual TArray<FPCGPinProperties> OutputPinProperties() const override { return Super::DefaultPointOutputPinProperties(); }
+	virtual FPCGElementPtr CreateElement() const override;
+	//~End UPCGSettings interface
+
+public:
+	UPROPERTY(BlueprintReadWrite, EditAnywhere, Category = Settings, meta = (ClampMin = "0", PCG_Overridable))
+	double SearchDistance = 500.0;
+
+	/** Allows the non-normalized distance to be output into a user-generated attribute. */
+	UPROPERTY(BlueprintReadWrite, EditAnywhere, Category = Settings, meta = (PCG_Overridable, InlineEditConditionToggle))
+	bool bSetDistanceToAttribute = false;
+	
+	/** The output attribute name to write the non-normalized distance, if not "None". */
+	UPROPERTY(BlueprintReadWrite, EditAnywhere, Category = Settings, meta = (EditCondition = "bSetDistanceToAttribute", PCG_Overridable))
+	FName DistanceAttribute = TEXT("Distance");
+
+	/** Allows the average position to be output into a user-generated attribute. */
+	UPROPERTY(BlueprintReadWrite, EditAnywhere, Category = Settings, meta = (PCG_Overridable, InlineEditConditionToggle))
+	bool bSetAveragePositionToAttribute = false;
+
+	/** The output attribute name to write the average positions, if not "None". */
+	UPROPERTY(BlueprintReadWrite, EditAnywhere, Category = Settings, meta = (EditCondition = "bSetAveragePositionToAttribute", PCG_Overridable))
+	FName AveragePositionAttribute = TEXT("AvgPosition");
+
+	/** Writes either the normalized distance or the average density to the point density. */
+	UPROPERTY(BlueprintReadWrite, EditAnywhere, Category = Settings, meta = (PCG_Overridable))
+	EPCGPointNeighborhoodDensityMode SetDensity = EPCGPointNeighborhoodDensityMode::None;
+
+	/** Writes the average position to the point transform. */
+	UPROPERTY(BlueprintReadWrite, EditAnywhere, Category = Settings, meta = (PCG_Overridable))
+	bool bSetAveragePosition = false;
+
+	/** Writes the target color to the point color if true, otherwise keeps the source color. */
+	UPROPERTY(BlueprintReadWrite, EditAnywhere, Category = Settings, meta = (PCG_Overridable))
+	bool bSetAverageColor = false;
+
+	/** Takes the bounds into account when projecting points. */
+	UPROPERTY(BlueprintReadWrite, EditAnywhere, Category = Settings, meta = (PCG_Overridable))
+	bool bWeightedAverage = false;
+};
+
+class FPCGPointNeighborhoodElement : public IPCGElement
+{
+protected:
+	virtual bool ExecuteInternal(FPCGContext* Context) const override;
+};
