[PCG] Cull nodes based on static Branch nodes

Generalizes the current functionality that visualizes which higen grids are present to a more general "structural visualization".

diff --git a/Engine/Plugins/Experimental/PCG/Source/PCG/Private/Elements/PCGBranch.cpp b/Engine/Plugins/Experimental/PCG/Source/PCG/Private/Elements/PCGBranch.cpp
--- a/Engine/Plugins/Experimental/PCG/Source/PCG/Private/Elements/PCGBranch.cpp
+++ b/Engine/Plugins/Experimental/PCG/Source/PCG/Private/Elements/PCGBranch.cpp
@@ -51,6 +51,62 @@ FPCGElementPtr UPCGBranchSettings::CreateElement() const
 	return MakeShared<FPCGBranchElement>();
 }
 
+#if WITH_EDITOR
+bool UPCGBranchSettings::IsStructuralProperty(const FName& InPropertyName) const
+{
+	// Static branches are processed during graph compilation and are part of the graph structure.
+	return InPropertyName == GET_MEMBER_NAME_CHECKED(UPCGBranchSettings, bEnabled)
+		|| InPropertyName == GET_MEMBER_NAME_CHECKED(UPCGBranchSettings, bOutputToB)
+		|| Super::IsStructuralProperty(InPropertyName);
+}
+#endif
+
+bool UPCGBranchSettings::IsDynamicBranch() const
+{
+	UPCGNode* Node = Cast<UPCGNode>(GetOuter());
+
+	if (!Node && OriginalSettings)
+	{
+		Node = Cast<UPCGNode>(OriginalSettings->GetOuter());
+	}
+
+	if (!Node)
+	{
+		return false;
+	}
+
+	const FName PropertyName = GET_MEMBER_NAME_CHECKED(UPCGBranchSettings, bOutputToB);
+
+	const FPCGSettingsOverridableParam* Param = CachedOverridableParams.FindByPredicate([PropertyName](const FPCGSettingsOverridableParam& ParamToCheck)
+	{
+		return !ParamToCheck.PropertiesNames.IsEmpty() && ParamToCheck.PropertiesNames.Last() == PropertyName;
+	});
+
+	if (Param)
+	{
+		if (const UPCGPin* Pin = Node->GetInputPin(Param->Label))
+		{
+			return Pin->IsConnected();
+		}
+	}
+
+	return false;
+}
+
+bool UPCGBranchSettings::IsPinStaticallyActive(const FName& PinLabel) const
+{
+	// Dynamic branches are never known in advance - assume all branches are active prior to execution.
+	if (IsDynamicBranch())
+	{
+		return true;
+	}
+
+	// Branch must be both enabled and set to B in order for output pin B to be active.
+	const FName ActiveOutputPinLabel = (bEnabled && bOutputToB) ? PCGBranchConstants::OutputLabelB : PCGBranchConstants::OutputLabelA;
+
+	return PinLabel == ActiveOutputPinLabel;
+}
+
 bool FPCGBranchElement::ExecuteInternal(FPCGContext* Context) const
 {
 	TRACE_CPUPROFILER_EVENT_SCOPE(FPCGBranchElement::ExecuteInternal);
@@ -66,4 +122,4 @@ bool FPCGBranchElement::ExecuteInternal(FPCGContext* Context) const
 	return true;
 }
 
-#undef LOCTEXT_NAMESPACE
\ No newline at end of file
+#undef LOCTEXT_NAMESPACE
diff --git a/Engine/Plugins/Experimental/PCG/Source/PCG/Private/Graph/PCGGraphCompiler.cpp b/Engine/Plugins/Experimental/PCG/Source/PCG/Private/Graph/PCGGraphCompiler.cpp
--- a/Engine/Plugins/Experimental/PCG/Source/PCG/Private/Graph/PCGGraphCompiler.cpp
+++ b/Engine/Plugins/Experimental/PCG/Source/PCG/Private/Graph/PCGGraphCompiler.cpp
@@ -475,6 +475,68 @@ EPCGHiGenGrid FPCGGraphCompiler::CalculateGridRecursive(
 	return Grid;
 }
 
+bool FPCGGraphCompiler::CalculateActiveRecursive(FPCGTaskId InTaskId, const TArray<FPCGGraphTask>& InCompiledTasks, TMap<int32, bool>& InTaskIdToActiveFlag)
+{
+	if (const bool* bEntry = InTaskIdToActiveFlag.Find(InTaskId))
+	{
+		return *bEntry;
+	}
+
+	bool bAnyInputActive = false;
+
+	for (FPCGGraphTaskInput Input : InCompiledTasks[InTaskId].Inputs)
+	{
+		// Default to tasks being active unless proved otherwise.
+		bool bInputActive = true;
+
+		// If we are connected to an upstream branch node, evaluate if the output pin is active.
+		const UPCGNode* UpstreamNode = InCompiledTasks[Input.TaskId].Node;
+		if (const UPCGSettings* UpstreamSettings = UpstreamNode ? UpstreamNode->GetSettings() : nullptr)
+		{
+			bInputActive &= UpstreamSettings->IsPinStaticallyActive(Input.InPin->Properties.Label);
+		}
+
+		if (bInputActive)
+		{
+			bInputActive &= CalculateActiveRecursive(Input.TaskId, InCompiledTasks, InTaskIdToActiveFlag);
+		}
+
+		if (bInputActive)
+		{
+			bAnyInputActive = true;
+			break;
+		}
+	}
+
+	const bool bActive = bAnyInputActive || InCompiledTasks[InTaskId].Inputs.IsEmpty();
+	InTaskIdToActiveFlag.Add(InTaskId, bActive);
+	return bActive;
+}
+
+void FPCGGraphCompiler::CullTasksStaticBranchNodes(TArray<FPCGGraphTask>& InOutCompiledTasks)
+{
+	if (InOutCompiledTasks.IsEmpty())
+	{
+		return;
+	}
+
+	TMap<int32, bool> NodeIdToActiveFlag;
+	// First task is input node task which is active
+	NodeIdToActiveFlag.Add(InOutCompiledTasks[0].NodeId, true);
+
+	for (int32 i = 1; i < InOutCompiledTasks.Num(); ++i)
+	{
+		CalculateActiveRecursive(InOutCompiledTasks[i].NodeId, InOutCompiledTasks, NodeIdToActiveFlag);
+	}
+
+	auto ShouldCull = [&NodeIdToActiveFlag](const FPCGGraphTask& InTask)
+	{
+		const bool* bActive = NodeIdToActiveFlag.Find(InTask.NodeId);
+		return bActive && !*bActive;
+	};
+	CullTasks(InOutCompiledTasks, /*bAddPassthroughWires=*/false, ShouldCull);
+}
+
 void FPCGGraphCompiler::CullTasks(TArray<FPCGGraphTask>& InOutCompiledTasks, bool bAddPassthroughWires, TFunctionRef<bool(const FPCGGraphTask&)> CullTask)
 {
 	TRACE_CPUPROFILER_EVENT_SCOPE(FPCGGraphCompiler::CullTasks);
@@ -715,6 +777,9 @@ void FPCGGraphCompiler::CompileTopGraph(UPCGGraph* InGraph, uint32 GenerationGri
 	// Remove reroute nodes before execution grid setup, as grid linkages need final nodes to connect from/to.
 	CullTasks(CompiledTasks, /*bAddPassthroughWires=*/true, [](const FPCGGraphTask& InTask) { return InTask.Node && Cast<UPCGRerouteSettings>(InTask.Node->GetSettings()); });
 
+	// Cull inactive branches downstream of branch nodes with static selection values.
+	CullTasksStaticBranchNodes(CompiledTasks);
+
 	// For hierarchical generation resolve the execution grid for each task and cull any tasks that won't execute.
 	if (InGraph->IsHierarchicalGenerationEnabled() && GenerationGridSize != PCGHiGenGrid::UninitializedGridSize())
 	{
diff --git a/Engine/Plugins/Experimental/PCG/Source/PCG/Private/Graph/PCGGraphCompiler.h b/Engine/Plugins/Experimental/PCG/Source/PCG/Private/Graph/PCGGraphCompiler.h
--- a/Engine/Plugins/Experimental/PCG/Source/PCG/Private/Graph/PCGGraphCompiler.h
+++ b/Engine/Plugins/Experimental/PCG/Source/PCG/Private/Graph/PCGGraphCompiler.h
@@ -55,6 +55,12 @@ private:
 		TArray<FPCGGraphTask>& InOutCompiledTasks,
 		const FPCGStackContext& InStackContext);
 
+	/** Discovers whether task is on a statically active branch (and needs to be passed to graph executor). */
+	static bool CalculateActiveRecursive(FPCGTaskId InTaskId, const TArray<FPCGGraphTask>& InCompiledTasks, TMap<int32, bool>& InTaskIdToActiveFlag);
+
+	/** Culls nodes that are not on statically active branches. */
+	static void CullTasksStaticBranchNodes(TArray<FPCGGraphTask>& InOutCompiledTasks);
+
 	/** Culls tasks based on a given lambda. Never culls the first (input) task in the array. */
 	static void CullTasks(TArray<FPCGGraphTask>& InOutCompiledTasks, bool bAddPassthroughWires, TFunctionRef<bool(const FPCGGraphTask&)> CullTask);
 
diff --git a/Engine/Plugins/Experimental/PCG/Source/PCG/Private/PCGGraph.cpp b/Engine/Plugins/Experimental/PCG/Source/PCG/Private/PCGGraph.cpp
--- a/Engine/Plugins/Experimental/PCG/Source/PCG/Private/PCGGraph.cpp
+++ b/Engine/Plugins/Experimental/PCG/Source/PCG/Private/PCGGraph.cpp
@@ -82,23 +82,25 @@ namespace PCGGraphUtils
 		InSourcePropertyDesc->CachedProperty->CopyCompleteValue(TargetValueAddress, SourceValueAddress);
 	}
 
-	EPCGChangeType NotifyTouchedNodes(const TSet<UPCGNode*>& InTouchedNodes)
+	EPCGChangeType NotifyTouchedNodes(const TSet<UPCGNode*>& InTouchedNodes, EPCGChangeType ChangeType)
 	{
-		EPCGChangeType ChangeType = EPCGChangeType::None;
+		EPCGChangeType FinalChangeType = EPCGChangeType::None;
 
 		for (UPCGNode* TouchedNode : InTouchedNodes)
 		{
 			if (TouchedNode)
 			{
-				ChangeType |= TouchedNode->PropagateDynamicPinTypes();
+				const EPCGChangeType NodeChangeType = ChangeType | TouchedNode->PropagateDynamicPinTypes();
 
 #if WITH_EDITOR
-				TouchedNode->OnNodeChangedDelegate.Broadcast(TouchedNode, EPCGChangeType::Node);
+				TouchedNode->OnNodeChangedDelegate.Broadcast(TouchedNode, NodeChangeType | EPCGChangeType::Node | ChangeType);
 #endif
+
+				FinalChangeType |= NodeChangeType;
 			}
 		}
 
-		return ChangeType;
+		return FinalChangeType;
 	}
 }
 
@@ -609,10 +611,10 @@ bool UPCGGraph::AddLabeledEdge(UPCGNode* From, const FName& FromPinLabel, UPCGNo
 		bToPinBrokeOtherEdges = ToPin->BreakAllIncompatibleEdges(&TouchedNodes);
 	}
 
-	const EPCGChangeType ChangeType = PCGGraphUtils::NotifyTouchedNodes(TouchedNodes);
+	const EPCGChangeType ChangeType = PCGGraphUtils::NotifyTouchedNodes(TouchedNodes, EPCGChangeType::Structural);
 
 #if WITH_EDITOR
-	NotifyGraphChanged(ChangeType | EPCGChangeType::Structural);
+	NotifyGraphChanged(ChangeType);
 #endif
 
 	return bToPinBrokeOtherEdges;
@@ -715,7 +717,7 @@ void UPCGGraph::RemoveNodes_Internal(TArrayView<UPCGNode*> InNodes)
 		Nodes.Remove(Node);
 	}
 
-	PCGGraphUtils::NotifyTouchedNodes(TouchedNodes);
+	PCGGraphUtils::NotifyTouchedNodes(TouchedNodes, EPCGChangeType::Structural);
 
 	OnNodesRemoved(InNodes);
 }
@@ -737,12 +739,12 @@ bool UPCGGraph::RemoveEdge(UPCGNode* From, const FName& FromLabel, UPCGNode* To,
 		OutPin->BreakEdgeTo(InPin, &TouchedNodes);
 	}
 
-	const EPCGChangeType ChangeType = PCGGraphUtils::NotifyTouchedNodes(TouchedNodes);
+	const EPCGChangeType ChangeType = PCGGraphUtils::NotifyTouchedNodes(TouchedNodes, EPCGChangeType::Structural);
 
 	if (TouchedNodes.Num() > 0)
 	{
 #if WITH_EDITOR
-		NotifyGraphChanged(ChangeType | EPCGChangeType::Structural);
+		NotifyGraphChanged(ChangeType);
 #endif
 	}
 
@@ -770,12 +772,12 @@ bool UPCGGraph::RemoveInboundEdges(UPCGNode* InNode, const FName& InboundLabel)
 		InputPin->BreakAllEdges(&TouchedNodes);
 	}
 
-	const EPCGChangeType ChangeType = PCGGraphUtils::NotifyTouchedNodes(TouchedNodes);
+	const EPCGChangeType ChangeType = PCGGraphUtils::NotifyTouchedNodes(TouchedNodes, EPCGChangeType::Structural);
 
 #if WITH_EDITOR
 	if (TouchedNodes.Num() > 0)
 	{
-		NotifyGraphChanged(ChangeType | EPCGChangeType::Structural);
+		NotifyGraphChanged(ChangeType);
 	}
 #endif
 
@@ -793,12 +795,12 @@ bool UPCGGraph::RemoveOutboundEdges(UPCGNode* InNode, const FName& OutboundLabel
 		OutputPin->BreakAllEdges(&TouchedNodes);
 	}
 
-	const EPCGChangeType ChangeType = PCGGraphUtils::NotifyTouchedNodes(TouchedNodes);
+	const EPCGChangeType ChangeType = PCGGraphUtils::NotifyTouchedNodes(TouchedNodes, EPCGChangeType::Structural);
 
 #if WITH_EDITOR
 	if (TouchedNodes.Num() > 0)
 	{
-		NotifyGraphChanged(ChangeType | EPCGChangeType::Structural);
+		NotifyGraphChanged(ChangeType);
 	}
 #endif
 
@@ -1013,7 +1015,9 @@ void UPCGGraph::NotifyGraphChanged(EPCGChangeType ChangeType)
 	// Graph settings, nodes, graph structure can all change the higen grid sizes.
 	if (ChangeType != EPCGChangeType::Cosmetic)
 	{
-		FWriteScopeLock ScopedWriteLock(NodeToGridSizeLock);
+		ResetNodeToOnActiveBranchMap();
+
+		FWriteScopeLock GridSizeLock(NodeToGridSizeLock);
 		NodeToGridSize.Reset();
 	}
 
@@ -1038,16 +1042,20 @@ void UPCGGraph::NotifyGraphParametersChanged(EPCGGraphParameterEvent InChangeTyp
 
 void UPCGGraph::OnNodeChanged(UPCGNode* InNode, EPCGChangeType ChangeType)
 {
-	if (!!(ChangeType & EPCGChangeType::Structural) && Cast<UPCGHiGenGridSizeSettings>(InNode->GetSettings()))
+	if (!!(ChangeType & EPCGChangeType::Structural))
 	{
 		// Update node to grid size map for grid size changes.
+		if (Cast<UPCGHiGenGridSizeSettings>(InNode->GetSettings()))
 		{
-			FWriteScopeLock ScopedWriteLock(NodeToGridSizeLock);
+			FWriteScopeLock Lock(NodeToGridSizeLock);
 			NodeToGridSize.Reset();
 		}
 
+		// Any node/edge change can affect which branches are statically active/inactive.
+		ResetNodeToOnActiveBranchMap();
+
 		// Broadcast so that grid size visualization can be updated editor-side.
-		OnGraphGridSizesChangedDelegate.Broadcast(this);
+		OnGraphStructureChangedDelegate.Broadcast(this);
 	}
 
 	if ((ChangeType & ~EPCGChangeType::Cosmetic) != EPCGChangeType::None)
@@ -1270,6 +1278,24 @@ FInstancedPropertyBag* UPCGGraph::GetMutableUserParametersStruct()
 	return &UserParameters;
 }
 
+#if WITH_EDITOR
+bool UPCGGraph::IsNodeOnActiveBranch(const UPCGNode* InNode) const
+{
+	{
+		FReadScopeLock Lock(NodeToOnActiveBranchLock);
+		if (const bool* CachedValue = NodeToOnActiveBranch.Find(InNode))
+		{
+			return *CachedValue;
+		}
+	}
+
+	{
+		FWriteScopeLock Lock(NodeToOnActiveBranchLock);
+		return CalculateNodeOnActiveBranchRecursive_Unsafe(InNode);
+	}
+}
+#endif // WITH_EDITOR
+
 uint32 UPCGGraph::GetNodeGenerationGridSize(const UPCGNode* InNode, uint32 InDefaultGridSize) const
 {
 	{
@@ -1334,6 +1360,69 @@ uint32 UPCGGraph::CalculateNodeGridSizeRecursive_Unsafe(const UPCGNode* InNode,
 	return GridSize;
 }
 
+#if WITH_EDITOR
+bool UPCGGraph::CalculateNodeOnActiveBranchRecursive_Unsafe(const UPCGNode* InNode) const
+{
+	if (const bool* CachedValue = NodeToOnActiveBranch.Find(InNode))
+	{
+		return *CachedValue;
+	}
+
+	bool bAnyInputActive = false;
+	bool bAnyEdgesPresent = false;
+
+	for (const UPCGPin* Pin : InNode->GetInputPins())
+	{
+		if (!Pin)
+		{
+			continue;
+		}
+
+		for (const UPCGEdge* Edge : Pin->Edges)
+		{
+			bAnyEdgesPresent = true;
+
+			bool bEdgeActive = true;
+
+			const UPCGPin* OtherPin = Edge ? Edge->InputPin : nullptr;
+			if (OtherPin && OtherPin->Node.Get())
+			{
+				if (const UPCGSettings* UpstreamSettings = OtherPin->Node->GetSettings())
+				{
+					bEdgeActive &= UpstreamSettings->IsPinStaticallyActive(OtherPin->Properties.Label);
+				}
+
+				bEdgeActive = bEdgeActive && CalculateNodeOnActiveBranchRecursive_Unsafe(OtherPin->Node);
+			}
+
+			if (bEdgeActive)
+			{
+				bAnyInputActive = true;
+				break;
+			}
+		}
+
+		if (bAnyInputActive)
+		{
+			break;
+		}
+	}
+
+	// Active if any input is active, or if there are no edges.
+	const bool bActive = bAnyInputActive || !bAnyEdgesPresent;
+
+	NodeToOnActiveBranch.Add(InNode, bActive);
+
+	return bActive;
+}
+
+void UPCGGraph::ResetNodeToOnActiveBranchMap()
+{
+	FWriteScopeLock Lock(NodeToOnActiveBranchLock);
+	NodeToOnActiveBranch.Reset();
+}
+#endif
+
 void UPCGGraph::AddUserParameters(const TArray<FPropertyBagPropertyDesc>& InDescs, const UPCGGraph* InOptionalOriginalGraph)
 {
 	UserParameters.AddProperties(InDescs);
@@ -1461,11 +1550,11 @@ void UPCGGraphInstance::OnGraphChanged(UPCGGraphInterface* InGraph, EPCGChangeTy
 	}
 }
 
-void UPCGGraphInstance::OnGraphGridSizesChanged(UPCGGraphInterface* InGraph)
+void UPCGGraphInstance::OnGraphStructureChanged(UPCGGraphInterface* InGraph)
 {
 	if (InGraph == Graph)
 	{
-		OnGraphGridSizesChangedDelegate.Broadcast(this);
+		OnGraphStructureChangedDelegate.Broadcast(this);
 	}
 }
 
@@ -1494,7 +1583,7 @@ void UPCGGraphInstance::TeardownCallbacks()
 	if (Graph)
 	{
 		Graph->OnGraphChangedDelegate.RemoveAll(this);
-		Graph->OnGraphGridSizesChangedDelegate.RemoveAll(this);
+		Graph->OnGraphStructureChangedDelegate.RemoveAll(this);
 		Graph->OnGraphParametersChangedDelegate.RemoveAll(this);
 	}
 }
@@ -1504,7 +1593,7 @@ void UPCGGraphInstance::SetupCallbacks()
 	if (Graph && !Graph->OnGraphChangedDelegate.IsBoundToObject(this))
 	{
 		Graph->OnGraphChangedDelegate.AddUObject(this, &UPCGGraphInstance::OnGraphChanged);
-		Graph->OnGraphGridSizesChangedDelegate.AddUObject(this, &UPCGGraphInstance::OnGraphGridSizesChanged);
+		Graph->OnGraphStructureChangedDelegate.AddUObject(this, &UPCGGraphInstance::OnGraphStructureChanged);
 		Graph->OnGraphParametersChangedDelegate.AddUObject(this, &UPCGGraphInstance::OnGraphParametersChanged);
 	}
 }
diff --git a/Engine/Plugins/Experimental/PCG/Source/PCG/Private/PCGSubgraph.cpp b/Engine/Plugins/Experimental/PCG/Source/PCG/Private/PCGSubgraph.cpp
--- a/Engine/Plugins/Experimental/PCG/Source/PCG/Private/PCGSubgraph.cpp
+++ b/Engine/Plugins/Experimental/PCG/Source/PCG/Private/PCGSubgraph.cpp
@@ -352,7 +352,7 @@ bool UPCGSubgraphSettings::IsDynamicGraph() const
 
 	const FPCGSettingsOverridableParam* Param = CachedOverridableParams.FindByPredicate([PropertyName](const FPCGSettingsOverridableParam& ParamToCheck)
 	{
-		return !ParamToCheck.PropertiesNames.IsEmpty() && ParamToCheck.PropertiesNames[0] == PropertyName;
+		return !ParamToCheck.PropertiesNames.IsEmpty() && ParamToCheck.PropertiesNames.Last() == PropertyName;
 	});
 
 	if (Param)
diff --git a/Engine/Plugins/Experimental/PCG/Source/PCG/Public/Elements/PCGBranch.h b/Engine/Plugins/Experimental/PCG/Source/PCG/Public/Elements/PCGBranch.h
--- a/Engine/Plugins/Experimental/PCG/Source/PCG/Public/Elements/PCGBranch.h
+++ b/Engine/Plugins/Experimental/PCG/Source/PCG/Public/Elements/PCGBranch.h
@@ -6,7 +6,7 @@
 #include "PCGBranch.generated.h"
 
 /**
- * Routes input data to one of two outputs,based on a boolean condition.
+ * Routes input data to one of two outputs, based on a boolean condition.
  */
 UCLASS(BlueprintType, ClassGroup = (Procedural), meta=(Keywords = "if bool branch"))
 class UPCGBranchSettings : public UPCGSettings
@@ -23,11 +23,19 @@ public:
 	virtual bool HasDynamicPins() const override { return true; }
 #endif
 
+	virtual bool IsPinStaticallyActive(const FName& PinLabel) const override;
+
 protected:
 	virtual TArray<FPCGPinProperties> OutputPinProperties() const override;
 	virtual FPCGElementPtr CreateElement() const override;
+#if WITH_EDITOR
+	virtual bool IsStructuralProperty(const FName& InPropertyName) const override;
+#endif
 	//~End UPCGSettings interface
 
+	/** Branch is dynamic if selection value is overridden / cannot be determined prior to execution. */
+	virtual bool IsDynamicBranch() const;
+
 public:
 	UPROPERTY(BlueprintReadWrite, EditAnywhere, Category=Settings, meta=(PCG_Overridable))
 	bool bOutputToB = false;
diff --git a/Engine/Plugins/Experimental/PCG/Source/PCG/Public/PCGGraph.h b/Engine/Plugins/Experimental/PCG/Source/PCG/Public/PCGGraph.h
--- a/Engine/Plugins/Experimental/PCG/Source/PCG/Public/PCGGraph.h
+++ b/Engine/Plugins/Experimental/PCG/Source/PCG/Public/PCGGraph.h
@@ -33,7 +33,7 @@ enum class EPCGGraphParameterEvent
 
 #if WITH_EDITOR
 DECLARE_MULTICAST_DELEGATE_TwoParams(FOnPCGGraphChanged, UPCGGraphInterface* /*Graph*/, EPCGChangeType /*ChangeType*/);
-DECLARE_MULTICAST_DELEGATE_OneParam(FOnPCGGraphGridSizesChanged, UPCGGraphInterface* /*Graph*/);
+DECLARE_MULTICAST_DELEGATE_OneParam(FOnPCGGraphStructureChanged, UPCGGraphInterface* /*Graph*/);
 DECLARE_MULTICAST_DELEGATE_ThreeParams(FOnPCGGraphParametersChanged, UPCGGraphInterface* /*Graph*/, EPCGGraphParameterEvent /*ChangeType*/, FName /*ChangedPropertyName*/);
 #endif // WITH_EDITOR
 
@@ -109,8 +109,10 @@ public:
 
 #if WITH_EDITOR
 	FOnPCGGraphChanged OnGraphChangedDelegate;
-	FOnPCGGraphGridSizesChanged OnGraphGridSizesChangedDelegate;
 	FOnPCGGraphParametersChanged OnGraphParametersChangedDelegate;
+
+	/** A structural change has been made such as changing a static branch or a higen grid size node. */
+	FOnPCGGraphStructureChanged OnGraphStructureChangedDelegate;
 #endif // WITH_EDITOR
 
 	template <typename T>
@@ -296,6 +298,9 @@ public:
 	void DisableInspection() { bIsInspecting = false; }
 	bool DebugFlagAppliesToIndividualComponents() const { return bDebugFlagAppliesToIndividualComponents; }
 	void RemoveExtraEditorNode(const UObject* InNode);
+
+	/** Is this node on an active branch (and will therefore execute). */
+	bool IsNodeOnActiveBranch(const UPCGNode* InNode) const;
 #endif
 
 #if WITH_EDITOR
@@ -317,7 +322,7 @@ protected:
 
 	bool IsEditorOnly_Internal(TSet<const UPCGGraph*>& VisitedGraphs) const;
 
-	/** Calculates node grid size. Not thread safe, called within write lock. */
+	/** Calculates node grid size. Not thread safe, must be called within write lock. */
 	uint32 CalculateNodeGridSizeRecursive_Unsafe(const UPCGNode* InNode, uint32 InDefaultGridSize) const;
 
 	UPROPERTY(BlueprintReadOnly, VisibleAnywhere, Category = Graph, meta = (NoResetToDefault))
@@ -400,6 +405,12 @@ private:
 	/** Remove invalid edges and edges to nodes that are not present in the node array. */
 	void FixInvalidEdges();
 
+	/** Evaluates whether node is on active branch and will therefore execute. Not thread safe, must be called within write lock. */
+	bool CalculateNodeOnActiveBranchRecursive_Unsafe(const UPCGNode* InNode) const;
+
+	/** Clear all active branch flags. */
+	void ResetNodeToOnActiveBranchMap();
+
 	int32 GraphChangeNotificationsDisableCounter = 0;
 	bool bDelayedChangeNotification = false;
 	EPCGChangeType DelayedChangeType = EPCGChangeType::None;
@@ -407,6 +418,10 @@ private:
 	bool bUserPausedNotificationsInGraphEditor = false;
 	int32 NumberOfUserParametersPreEdit = 0;
 	bool bIsInspecting = false;
+
+	/** Cache for "on active branch" flag for nodes. */
+	mutable TMap<const UPCGNode*, bool> NodeToOnActiveBranch;
+	mutable FRWLock NodeToOnActiveBranchLock;
 #endif // WITH_EDITOR
 };
 
@@ -442,7 +457,7 @@ public:
 protected:
 #if WITH_EDITOR
 	void OnGraphChanged(UPCGGraphInterface* InGraph, EPCGChangeType ChangeType);
-	void OnGraphGridSizesChanged(UPCGGraphInterface* InGraph);
+	void OnGraphStructureChanged(UPCGGraphInterface* InGraph);
 	void NotifyGraphParametersChanged(EPCGGraphParameterEvent InChangeType, FName InChangedPropertyName);
 #endif
 	void OnGraphParametersChanged(UPCGGraphInterface* InGraph, EPCGGraphParameterEvent InChangeType, FName InChangedPropertyName);
diff --git a/Engine/Plugins/Experimental/PCG/Source/PCG/Public/PCGSettings.h b/Engine/Plugins/Experimental/PCG/Source/PCG/Public/PCGSettings.h
--- a/Engine/Plugins/Experimental/PCG/Source/PCG/Public/PCGSettings.h
+++ b/Engine/Plugins/Experimental/PCG/Source/PCG/Public/PCGSettings.h
@@ -303,6 +303,9 @@ public:
 	/** Returns true if InPin is in use by node (assuming node enabled). Can be used to communicate when a pin is not in use to user. */
 	virtual bool IsPinUsedByNodeExecution(const UPCGPin* InPin) const { return true; }
 
+	/** True if we know prior to execution that the given pin will be active (not on inactive branch). */
+	virtual bool IsPinStaticallyActive(const FName& OutputPinLabel) const { return true; }
+
 	/** Returns true if only the first input edge is used from the primary pin when the node is disabled. */
 	virtual bool OnlyPassThroughOneEdgeWhenDisabled() const { return false; }
 
diff --git a/Engine/Plugins/Experimental/PCG/Source/PCGEditor/Private/PCGEditor.cpp b/Engine/Plugins/Experimental/PCG/Source/PCGEditor/Private/PCGEditor.cpp
--- a/Engine/Plugins/Experimental/PCG/Source/PCGEditor/Private/PCGEditor.cpp
+++ b/Engine/Plugins/Experimental/PCG/Source/PCGEditor/Private/PCGEditor.cpp
@@ -103,7 +103,7 @@ void FPCGEditor::Initialize(const EToolkitMode::Type InMode, const TSharedPtr<cl
 
 	if (PCGGraphBeingEdited)
 	{
-		PCGGraphBeingEdited->OnGraphGridSizesChangedDelegate.AddRaw(this, &FPCGEditor::OnGraphGridSizesChanged);
+		PCGGraphBeingEdited->OnGraphStructureChangedDelegate.AddRaw(this, &FPCGEditor::OnGraphStructureChanged);
 	}
 
 	if (UPCGSubsystem* Subsystem = GetSubsystem())
@@ -273,7 +273,7 @@ void FPCGEditor::SetStackBeingInspected(const FPCGStack& FullStack, FPCGDebugObj
 			{
 				// Update now that component has changed. Will fire OnNodeChanged if necessary.
 				EPCGChangeType ChangeType = PCGNode->UpdateErrorsAndWarnings();
-				ChangeType |= PCGNode->UpdateGridSizeVisualization(NewComponent, FullStack);
+				ChangeType |= PCGNode->UpdateStructuralVisualization(NewComponent, &StackBeingInspected);
 
 				if (ChangeType != EPCGChangeType::None)
 				{
@@ -2105,7 +2105,7 @@ void FPCGEditor::OnClose()
 			PCGGraphBeingEdited->ToggleUserPausedNotificationsForEditor();
 		}
 
-		PCGGraphBeingEdited->OnGraphGridSizesChangedDelegate.RemoveAll(this);
+		PCGGraphBeingEdited->OnGraphStructureChangedDelegate.RemoveAll(this);
 	}
 
 	if (UPCGSubsystem* Subsystem = GetSubsystem())
@@ -2337,16 +2337,11 @@ bool FPCGEditor::IsVisibleProperty(const FPropertyAndParent& InPropertyAndParent
 	return true;
 }
 
-void FPCGEditor::OnGraphGridSizesChanged(UPCGGraphInterface* InGraph)
+void FPCGEditor::OnGraphStructureChanged(UPCGGraphInterface* InGraph)
 {
 	check(PCGEditorGraph);
 
-	const FPCGStack* Stack = GetStackBeingInspected();
-	UPCGComponent* PCGComponent = GetPCGComponentBeingInspected();
-	if (Stack && PCGComponent)
-	{
-		PCGEditorGraph->UpdateGridSizeVisualization(PCGComponent, *Stack);
-	}
+	PCGEditorGraph->UpdateStructuralVisualization(GetPCGComponentBeingInspected(), GetStackBeingInspected());
 }
 
 void FPCGEditor::OnComponentGenerationCompleteOrCancelled()
diff --git a/Engine/Plugins/Experimental/PCG/Source/PCGEditor/Private/PCGEditor.h b/Engine/Plugins/Experimental/PCG/Source/PCGEditor/Private/PCGEditor.h
--- a/Engine/Plugins/Experimental/PCG/Source/PCGEditor/Private/PCGEditor.h
+++ b/Engine/Plugins/Experimental/PCG/Source/PCGEditor/Private/PCGEditor.h
@@ -298,7 +298,7 @@ private:
 	/** Returns whether a property should be visible (used for instance vs. settings properties) */
 	bool IsVisibleProperty(const FPropertyAndParent& InPropertyAndParent, IDetailsView* InDetailsView) const;
 
-	void OnGraphGridSizesChanged(UPCGGraphInterface* InGraph);
+	void OnGraphStructureChanged(UPCGGraphInterface* InGraph);
 
 	/** Called when a component finishes executing. Useful for updating debugging tools/UIs. */
 	void OnComponentGenerationCompleteOrCancelled();
diff --git a/Engine/Plugins/Experimental/PCG/Source/PCGEditor/Private/PCGEditorGraph.cpp b/Engine/Plugins/Experimental/PCG/Source/PCGEditor/Private/PCGEditorGraph.cpp
--- a/Engine/Plugins/Experimental/PCG/Source/PCGEditor/Private/PCGEditorGraph.cpp
+++ b/Engine/Plugins/Experimental/PCG/Source/PCGEditor/Private/PCGEditorGraph.cpp
@@ -99,6 +99,9 @@ void UPCGEditorGraph::ReconstructGraph()
 			AddNode(NewNode, bIsUserAction, bSelectNewNode);
 		}
 	}
+
+	// Ensure graph structure visualization is nice and fresh upon opening.
+	UpdateStructuralVisualization(nullptr, nullptr);
 }
 
 void UPCGEditorGraph::BeginDestroy()
@@ -134,12 +137,12 @@ void UPCGEditorGraph::CreateLinks(UPCGEditorGraphNodeBase* GraphNode, bool bCrea
 	CreateLinks(GraphNode, bCreateInbound, bCreateOutbound, PCGNodeToPCGEditorNodeMap);
 }
 
-void UPCGEditorGraph::UpdateGridSizeVisualization(UPCGComponent* PCGComponentBeingInspected, const FPCGStack& PCGStackBeingInspected)
+void UPCGEditorGraph::UpdateStructuralVisualization(UPCGComponent* PCGComponentBeingInspected, const FPCGStack* PCGStackBeingInspected)
 {
 	for (UEdGraphNode* EditorNode : Nodes)
 	{
 		UPCGEditorGraphNodeBase* PCGEditorNode = Cast<UPCGEditorGraphNodeBase>(EditorNode);
-		if (PCGEditorNode && PCGEditorNode->UpdateGridSizeVisualization(PCGComponentBeingInspected, PCGStackBeingInspected) != EPCGChangeType::None)
+		if (PCGEditorNode && PCGEditorNode->UpdateStructuralVisualization(PCGComponentBeingInspected, PCGStackBeingInspected) != EPCGChangeType::None)
 		{
 			PCGEditorNode->ReconstructNode();
 		}
diff --git a/Engine/Plugins/Experimental/PCG/Source/PCGEditor/Private/PCGEditorGraph.h b/Engine/Plugins/Experimental/PCG/Source/PCGEditor/Private/PCGEditorGraph.h
--- a/Engine/Plugins/Experimental/PCG/Source/PCGEditor/Private/PCGEditorGraph.h
+++ b/Engine/Plugins/Experimental/PCG/Source/PCGEditor/Private/PCGEditorGraph.h
@@ -41,7 +41,7 @@ public:
 	TWeakPtr<FPCGEditor> GetEditor() const { return PCGEditor; }
 
 	/** Updates the grid size visualization in the editor. */
-	void UpdateGridSizeVisualization(UPCGComponent* PCGComponentBeingInspected, const FPCGStack& PCGStackBeingInspected);
+	void UpdateStructuralVisualization(UPCGComponent* PCGComponentBeingInspected, const FPCGStack* PCGStackBeingInspected);
 
 protected:
 	void CreateLinks(UPCGEditorGraphNodeBase* InGraphNode, bool bCreateInbound, bool bCreateOutbound, const TMap<UPCGNode*, UPCGEditorGraphNodeBase*>& InPCGNodeToPCGEditorNodeMap);
diff --git a/Engine/Plugins/Experimental/PCG/Source/PCGEditor/Private/PCGEditorGraphNodeBase.cpp b/Engine/Plugins/Experimental/PCG/Source/PCGEditor/Private/PCGEditorGraphNodeBase.cpp
--- a/Engine/Plugins/Experimental/PCG/Source/PCGEditor/Private/PCGEditorGraphNodeBase.cpp
+++ b/Engine/Plugins/Experimental/PCG/Source/PCGEditor/Private/PCGEditorGraphNodeBase.cpp
@@ -313,6 +313,18 @@ void UPCGEditorGraphNodeBase::OnNodeChanged(UPCGNode* InNode, EPCGChangeType Cha
 
 		ChangeType |= UpdateErrorsAndWarnings();
 
+		UPCGComponent* ComponentBeingDebugged = nullptr;
+		const FPCGStack* StackBeingDebugged = nullptr;
+
+		{
+			const UPCGEditorGraph* EditorGraph = CastChecked<UPCGEditorGraph>(GetGraph());
+			const FPCGEditor* Editor = EditorGraph ? EditorGraph->GetEditor().Pin().Get() : nullptr;
+			ComponentBeingDebugged = Editor ? Editor->GetPCGComponentBeingInspected() : nullptr;
+			StackBeingDebugged = Editor ? Editor->GetStackBeingInspected() : nullptr;
+		}
+
+		ChangeType |= UpdateStructuralVisualization(ComponentBeingDebugged, StackBeingDebugged);
+
 		if (!!(ChangeType & (EPCGChangeType::Structural | EPCGChangeType::Node | EPCGChangeType::Edge | EPCGChangeType::Cosmetic)))
 		{
 			ReconstructNode();
@@ -381,7 +393,7 @@ EPCGChangeType UPCGEditorGraphNodeBase::UpdateErrorsAndWarnings()
 	return bStateChanged ? EPCGChangeType::Cosmetic : EPCGChangeType::None;
 }
 
-EPCGChangeType UPCGEditorGraphNodeBase::UpdateGridSizeVisualization(UPCGComponent* InComponentBeingDebugged, const FPCGStack& InStackBeingInspected)
+EPCGChangeType UPCGEditorGraphNodeBase::UpdateStructuralVisualization(UPCGComponent* InComponentBeingDebugged, const FPCGStack* InStackBeingInspected)
 {
 	const UPCGGraph* Graph = PCGNode ? PCGNode->GetGraph() : nullptr;
 	if (!Graph)
@@ -391,26 +403,17 @@ EPCGChangeType UPCGEditorGraphNodeBase::UpdateGridSizeVisualization(UPCGComponen
 
 	EPCGChangeType ChangeType = EPCGChangeType::None;
 
+	SetOnActiveBranch(Graph->IsNodeOnActiveBranch(PCGNode));
+
+	bool bShouldDisplayAsDisabled = !IsOnActiveBranch();
+	bool bShouldDisplayAsHighlighted = false;
+
 	const bool HiGenEnabled = Graph->IsHierarchicalGenerationEnabled();
 	const uint32 InspectingGridSize = InComponentBeingDebugged ? InComponentBeingDebugged->GetGenerationGridSize() : PCGHiGenGrid::UninitializedGridSize();
 
-	// Disable grid size visualization if higen is disabled, or if we're not inspecting a specific grid, or if we're
-	// inspecting a subgraph since subgraphs execute at the invoked grid level.
-	if (!HiGenEnabled || (InspectingGridSize == PCGHiGenGrid::UninitializedGridSize()) || !InStackBeingInspected.IsCurrentFrameInRootGraph())
-	{
-		if (IsDisplayAsDisabledForced())
-		{
-			SetForceDisplayAsDisabled(false);
-			ChangeType |= EPCGChangeType::Cosmetic;
-		}
-
-		if (IsHighlighted())
-		{
-			SetIsHighlighted(false);
-			ChangeType |= EPCGChangeType::Cosmetic;
-		}
-	}
-	else
+	// Show grid size visualization if higen is enabled and if we're inspecting a specific grid, and we're inspecting a subgraph since subgraphs
+	// execute at the invoked grid level.
+	if (HiGenEnabled && InStackBeingInspected && InStackBeingInspected->IsCurrentFrameInRootGraph() && InspectingGridSize != PCGHiGenGrid::UninitializedGridSize())
 	{
 		uint32 DefaultGridSize;
 		UPCGSubsystem* Subsystem = UPCGSubsystem::GetActiveEditorInstance();
@@ -431,20 +434,22 @@ EPCGChangeType UPCGEditorGraphNodeBase::UpdateGridSizeVisualization(UPCGComponen
 		const uint32 NodeGridSize = Graph->GetNodeGenerationGridSize(PCGNode, DefaultGridSize);
 
 		// Disable nodes that are on a smaller grid
-		const bool bForceDisable = NodeGridSize < InspectingGridSize;
-		if (IsDisplayAsDisabledForced() != bForceDisable)
-		{
-			SetForceDisplayAsDisabled(bForceDisable);
-			ChangeType |= EPCGChangeType::Cosmetic;
-		}
+		bShouldDisplayAsDisabled |= NodeGridSize < InspectingGridSize;
 
 		// If node is on larger grid than current, highlight it to indicate that its data is available for use
-		const bool bHighlight = NodeGridSize > InspectingGridSize;
-		if (IsHighlighted() != bHighlight)
-		{
-			SetIsHighlighted(bHighlight);
-			ChangeType |= EPCGChangeType::Cosmetic;
-		}
+		bShouldDisplayAsHighlighted |= NodeGridSize > InspectingGridSize;
+	}
+
+	if (IsDisplayAsDisabledForced() != bShouldDisplayAsDisabled)
+	{
+		SetForceDisplayAsDisabled(bShouldDisplayAsDisabled);
+		ChangeType |= EPCGChangeType::Cosmetic;
+	}
+
+	if (IsHighlighted() != bShouldDisplayAsHighlighted)
+	{
+		SetIsHighlighted(bShouldDisplayAsHighlighted);
+		ChangeType |= EPCGChangeType::Cosmetic;
 	}
 
 	return ChangeType;
diff --git a/Engine/Plugins/Experimental/PCG/Source/PCGEditor/Private/PCGEditorGraphNodeBase.h b/Engine/Plugins/Experimental/PCG/Source/PCGEditor/Private/PCGEditorGraphNodeBase.h
--- a/Engine/Plugins/Experimental/PCG/Source/PCGEditor/Private/PCGEditorGraphNodeBase.h
+++ b/Engine/Plugins/Experimental/PCG/Source/PCGEditor/Private/PCGEditorGraphNodeBase.h
@@ -54,9 +54,13 @@ public:
 	void PostCopy();
 	void RebuildAfterPaste();
 	void PostPaste();
+
 	void SetInspected(bool InIsInspecting) { bIsInspected = InIsInspecting; }
 	bool GetInspected() const { return bIsInspected; }
-	
+
+	void SetOnActiveBranch(bool bInIsOnActiveBranch) { bIsOnActiveBranch = bInIsOnActiveBranch; }
+	bool IsOnActiveBranch() const { return bIsOnActiveBranch; }
+
 	/** Increase deferred reconstruct counter, calls to ReconstructNode will flag reconstruct to happen when count hits zero */
 	void EnableDeferredReconstruct();
 	/** Decrease deferred reconstruct counter, ReconstructNode will be called if counter hits zero and the node is flagged for reconstruction  */
@@ -66,7 +70,7 @@ public:
 	EPCGChangeType UpdateErrorsAndWarnings();
 
 	/** If the currently inspected grid size is smaller than the grid size of this node, display transparent. */
-	EPCGChangeType UpdateGridSizeVisualization(class UPCGComponent* InComponentBeingDebugged, const FPCGStack& InStackBeingInspected);
+	EPCGChangeType UpdateStructuralVisualization(class UPCGComponent* InComponentBeingDebugged, const FPCGStack* InStackBeingInspected);
 
 	/** Puts node title on node body, reducing overall node size */
 	bool ShouldDrawCompact() const;
@@ -116,4 +120,5 @@ protected:
 	bool bIsInspected = false;
 	bool bIsHighlighted = false;
 	bool bCanUserAddRemoveSourcePins = false;
+	bool bIsOnActiveBranch = true;
 };
diff --git a/Engine/Plugins/Experimental/PCG/Source/PCGEditor/Private/PCGEditorStyle.cpp b/Engine/Plugins/Experimental/PCG/Source/PCGEditor/Private/PCGEditorStyle.cpp
--- a/Engine/Plugins/Experimental/PCG/Source/PCGEditor/Private/PCGEditorStyle.cpp
+++ b/Engine/Plugins/Experimental/PCG/Source/PCGEditor/Private/PCGEditorStyle.cpp
@@ -29,6 +29,7 @@ FPCGEditorStyle::FPCGEditorStyle() : FSlateStyleSet("PCGEditorStyle")
 
 	Set("PCG.NodeOverlay.Debug", new CORE_IMAGE_BRUSH_SVG("Starship/Blueprints/Breakpoint_Valid", Icon20x20, FSlateColor(FColor::Cyan)));
 	Set("PCG.NodeOverlay.Inspect", new CORE_IMAGE_BRUSH_SVG("Starship/Blueprints/Breakpoint_Valid", Icon20x20, FSlateColor(FColor::Orange)));
+	Set("PCG.NodeOverlay.OnInactiveBranch", new CORE_IMAGE_BRUSH_SVG("Starship/Blueprints/CompileStatus_Fail_Badge", Icon20x20, FSlateColor(FColor::White)));
 
 	static const FVector2D PinSize(22.0f, 22.0f);
 	Set(PCGEditorStyleConstants::Pin_SD_SC_IN_C, new IMAGE_BRUSH_SVG("Style/PCG_Graph_SD_SC_IN", PinSize));
@@ -61,6 +62,8 @@ FPCGEditorStyle::FPCGEditorStyle() : FSlateStyleSet("PCGEditorStyle")
 	Set(PCGEditorStyleConstants::Pin_Composite_OUT_C, new IMAGE_BRUSH_SVG("Style/PCG_Graph_Composite_OUT", PinSize));
 	Set(PCGEditorStyleConstants::Pin_Composite_OUT_DC, new IMAGE_BRUSH_SVG("Style/PCG_Graph_Composite_OUT_Unplugged", PinSize));
 
+	Set(PCGEditorStyleConstants::Node_Overlay_Inactive, new IMAGE_BRUSH_SVG("Style/PCG_Node_Overlay_Inactive", Icon20x20));
+
 	FInlineEditableTextBlockStyle NodeTitleStyle = FInlineEditableTextBlockStyle(FAppStyle::Get().GetWidgetStyle<FInlineEditableTextBlockStyle>("Graph.Node.NodeTitleInlineEditableText"));
 	FTextBlockStyle GraphNodeItalicTitle = FTextBlockStyle(FAppStyle::Get().GetWidgetStyle<FTextBlockStyle>("Graph.Node.NodeTitle"))
 		.SetFont(DEFAULT_FONT("BoldCondensedItalic", 10));
diff --git a/Engine/Plugins/Experimental/PCG/Source/PCGEditor/Private/PCGEditorStyle.h b/Engine/Plugins/Experimental/PCG/Source/PCGEditor/Private/PCGEditorStyle.h
--- a/Engine/Plugins/Experimental/PCG/Source/PCGEditor/Private/PCGEditorStyle.h
+++ b/Engine/Plugins/Experimental/PCG/Source/PCGEditor/Private/PCGEditorStyle.h
@@ -47,4 +47,6 @@ namespace PCGEditorStyleConstants
 	const FName Pin_Composite_IN_DC = TEXT("PCG.Graph.Composite.In.Disconnected");
 	const FName Pin_Composite_OUT_C = TEXT("PCG.Graph.Composite.Out.Connected");
 	const FName Pin_Composite_OUT_DC = TEXT("PCG.Graph.Composite.Out.Disconnected");
+
+	const FName Node_Overlay_Inactive = TEXT("PCG.Node.Overlay.Inactive");
 }
\ No newline at end of file
diff --git a/Engine/Plugins/Experimental/PCG/Source/PCGEditor/Private/SPCGEditorGraphNode.cpp b/Engine/Plugins/Experimental/PCG/Source/PCGEditor/Private/SPCGEditorGraphNode.cpp
--- a/Engine/Plugins/Experimental/PCG/Source/PCGEditor/Private/SPCGEditorGraphNode.cpp
+++ b/Engine/Plugins/Experimental/PCG/Source/PCGEditor/Private/SPCGEditorGraphNode.cpp
@@ -528,34 +528,46 @@ void SPCGEditorGraphNode::GetOverlayBrushes(bool bSelected, const FVector2D Widg
 {
 	check(PCGEditorGraphNode);
 	
-	const FSlateBrush* DebugBrush = FPCGEditorStyle::Get().GetBrush(TEXT("PCG.NodeOverlay.Debug"));
-	const FSlateBrush* InspectBrush = FPCGEditorStyle::Get().GetBrush(TEXT("PCG.NodeOverlay.Inspect"));
-	
-	const FVector2D HalfDebugBrushSize = DebugBrush->GetImageSize() / 2.0;
-	const FVector2D HalfInspectBrushSize = InspectBrush->GetImageSize() / 2.0;
-	
-	FVector2D OverlayOffset(0.0, 0.0);
+	if (!PCGEditorGraphNode->IsOnActiveBranch())
+	{
+		const FSlateBrush* InactiveBranchBrush = FPCGEditorStyle::Get().GetBrush(PCGEditorStyleConstants::Node_Overlay_Inactive);
 
-	if (const UPCGNode* PCGNode = PCGEditorGraphNode->GetPCGNode())
+		FOverlayBrushInfo BrushInfo;
+		BrushInfo.Brush = InactiveBranchBrush;
+		BrushInfo.OverlayOffset = -InactiveBranchBrush->GetImageSize() / 2.0;
+		Brushes.Add(BrushInfo);
+	}
+	else
 	{
-		if (PCGNode->GetSettingsInterface() && PCGNode->GetSettingsInterface()->bDebug)
+		const FSlateBrush* DebugBrush = FPCGEditorStyle::Get().GetBrush(TEXT("PCG.NodeOverlay.Debug"));
+		const FSlateBrush* InspectBrush = FPCGEditorStyle::Get().GetBrush(TEXT("PCG.NodeOverlay.Inspect"));
+
+		const FVector2D HalfDebugBrushSize = DebugBrush->GetImageSize() / 2.0;
+		const FVector2D HalfInspectBrushSize = InspectBrush->GetImageSize() / 2.0;
+
+		FVector2D OverlayOffset(0.0, 0.0);
+
+		if (const UPCGNode* PCGNode = PCGEditorGraphNode->GetPCGNode())
+		{
+			if (PCGNode->GetSettingsInterface() && PCGNode->GetSettingsInterface()->bDebug)
+			{
+				FOverlayBrushInfo BrushInfo;
+				BrushInfo.Brush = DebugBrush;
+				BrushInfo.OverlayOffset = OverlayOffset - HalfDebugBrushSize;
+				Brushes.Add(BrushInfo);
+
+				OverlayOffset.Y += HalfDebugBrushSize.Y + HalfInspectBrushSize.Y;
+			}
+		}
+
+		if (PCGEditorGraphNode->GetInspected())
 		{
 			FOverlayBrushInfo BrushInfo;
-			BrushInfo.Brush = DebugBrush;
-			BrushInfo.OverlayOffset = OverlayOffset - HalfDebugBrushSize;
+			BrushInfo.Brush = InspectBrush;
+			BrushInfo.OverlayOffset = OverlayOffset - HalfInspectBrushSize;
 			Brushes.Add(BrushInfo);
-			
-			OverlayOffset.Y += HalfDebugBrushSize.Y + HalfInspectBrushSize.Y;
 		}
 	}
-
-	if (PCGEditorGraphNode->GetInspected())
-	{
-		FOverlayBrushInfo BrushInfo;
-		BrushInfo.Brush = InspectBrush;
-		BrushInfo.OverlayOffset = OverlayOffset - HalfInspectBrushSize;
-		Brushes.Add(BrushInfo);	
-	}
 }
 
 void SPCGEditorGraphNode::OnNodeChanged()
