[PCG] If projection source is Point, collapse the output to Point (execute projection), and add option to always sample to Point.

diff --git a/Engine/Plugins/Experimental/PCG/Source/PCG/Private/Elements/PCGProjectionElement.cpp b/Engine/Plugins/Experimental/PCG/Source/PCG/Private/Elements/PCGProjectionElement.cpp
--- a/Engine/Plugins/Experimental/PCG/Source/PCG/Private/Elements/PCGProjectionElement.cpp
+++ b/Engine/Plugins/Experimental/PCG/Source/PCG/Private/Elements/PCGProjectionElement.cpp
@@ -2,12 +2,13 @@
 
 #include "Elements/PCGProjectionElement.h"
 
-#include "Data/PCGProjectionData.h"
+#include "PCGContext.h"
 #include "PCGCustomVersion.h"
-#include "Data/PCGSpatialData.h"
 #include "PCGEdge.h"
-#include "PCGContext.h"
 #include "PCGPin.h"
+#include "Data/PCGPointData.h"
+#include "Data/PCGProjectionData.h"
+#include "Data/PCGSpatialData.h"
 
 #include UE_INLINE_GENERATED_CPP_BY_NAME(PCGProjectionElement)
 
@@ -36,6 +37,19 @@ FText UPCGProjectionSettings::GetNodeTooltipText() const
 }
 #endif
 
+EPCGDataType UPCGProjectionSettings::GetCurrentPinTypes(const UPCGPin* InPin) const
+{
+	if (InPin->IsOutputPin())
+	{
+		if (bForceCollapseToPoint || GetTypeUnionOfIncidentEdges(PCGPinConstants::DefaultInputLabel) == EPCGDataType::Point)
+		{
+			return EPCGDataType::Point;
+		}
+	}
+
+	return Super::GetCurrentPinTypes(InPin);
+}
+
 FPCGElementPtr UPCGProjectionSettings::CreateElement() const
 {
 	return MakeShared<FPCGProjectionElement>();
@@ -84,15 +98,21 @@ bool FPCGProjectionElement::ExecuteInternal(FPCGContext* Context) const
 		UPCGSpatialData* ProjectionData = ProjectionSource->ProjectOn(ProjectionTarget, ProjectionParams);
 		ProjectionData->bKeepZeroDensityPoints = bKeepZeroDensityPoints;
 
-		if (ProjectionData->RequiresCollapseToSample())
+		const UPCGSpatialData* OutputData = ProjectionData;
+
+		if (Settings->bForceCollapseToPoint || ProjectionSource->IsA<UPCGPointData>())
+		{
+			OutputData = OutputData->ToPointData(Context);
+		}
+		else if (OutputData->RequiresCollapseToSample())
 		{
 			// Calling ToPointData will populate the point cache. Doing so here means we can pass in the Context object, which
 			// means the operation will be multi-threaded. This primes the cache in the most efficient way.
-			ProjectionData->ToPointData(Context);
+			OutputData->ToPointData(Context);
 		}
 
 		FPCGTaggedData& ProjectionTaggedData = Outputs.Emplace_GetRef(Source);
-		ProjectionTaggedData.Data = ProjectionData;
+		ProjectionTaggedData.Data = OutputData;
 
 		if (ProjectionParams.TagMergeOperation == EPCGProjectionTagMergeMode::Target)
 		{
diff --git a/Engine/Plugins/Experimental/PCG/Source/PCG/Public/Elements/PCGProjectionElement.h b/Engine/Plugins/Experimental/PCG/Source/PCG/Public/Elements/PCGProjectionElement.h
--- a/Engine/Plugins/Experimental/PCG/Source/PCG/Public/Elements/PCGProjectionElement.h
+++ b/Engine/Plugins/Experimental/PCG/Source/PCG/Public/Elements/PCGProjectionElement.h
@@ -28,13 +28,19 @@ public:
 	virtual FText GetDefaultNodeTitle() const override { return NSLOCTEXT("PCGProjectionSettings", "NodeTitle", "Projection"); }
 	virtual FText GetNodeTooltipText() const override;
 	virtual EPCGSettingsType GetType() const override { return EPCGSettingsType::Spatial; }
+	virtual bool HasDynamicPins() const override { return true; }
 #endif
+	virtual EPCGDataType GetCurrentPinTypes(const UPCGPin* InPin) const override;
 
 	virtual TArray<FPCGPinProperties> InputPinProperties() const override;
 
 	UPROPERTY(BlueprintReadWrite, EditAnywhere, Category = Settings, meta = (ShowOnlyInnerProperties, PCG_Overridable))
 	FPCGProjectionParams ProjectionParams;
 
+	/** Force the result to be sampled to points, equivalent to having a To Point node after the projection node. */
+	UPROPERTY(BlueprintReadWrite, EditAnywhere, Category = Settings, meta = (PCG_Overridable))
+	bool bForceCollapseToPoint = false;
+
 	UPROPERTY(BlueprintReadWrite, EditAnywhere, Category = Settings, meta = (PCG_Overridable))
 	bool bKeepZeroDensityPoints = false;
 
