[PCG] Implement higen+subgraph rule - disable higen for subgraphs, execute subgraph on the grid that the subgraph was invoked.

diff --git a/Engine/Plugins/Experimental/PCG/Source/PCG/Private/Elements/PCGLoopElement.cpp b/Engine/Plugins/Experimental/PCG/Source/PCG/Private/Elements/PCGLoopElement.cpp
--- a/Engine/Plugins/Experimental/PCG/Source/PCG/Private/Elements/PCGLoopElement.cpp
+++ b/Engine/Plugins/Experimental/PCG/Source/PCG/Private/Elements/PCGLoopElement.cpp
@@ -276,7 +276,14 @@ bool FPCGLoopElement::ExecuteInternal(FPCGContext* InContext) const
 			}
 
 			FPCGElementPtr InputElement = MakeShared<FPCGLoopInputForwardingElement>(InputDataCollection, PreviousTaskId, FeedbackPinNames);
-			FPCGTaskId SubgraphTaskId = Subsystem->ScheduleGraph(Subgraph, Context->SourceComponent.Get(), PreGraphElement, InputElement, Dependencies, &InvocationStack);
+			FPCGTaskId SubgraphTaskId = Subsystem->ScheduleGraph(
+				Subgraph,
+				Context->SourceComponent.Get(),
+				PreGraphElement,
+				InputElement,
+				Dependencies,
+				&InvocationStack,
+				/*bAllowHierarchicalGeneration=*/false);
 
 			if (SubgraphTaskId != InvalidPCGTaskId)
 			{
diff --git a/Engine/Plugins/Experimental/PCG/Source/PCG/Private/Elements/PCGSpawnActor.cpp b/Engine/Plugins/Experimental/PCG/Source/PCG/Private/Elements/PCGSpawnActor.cpp
--- a/Engine/Plugins/Experimental/PCG/Source/PCG/Private/Elements/PCGSpawnActor.cpp
+++ b/Engine/Plugins/Experimental/PCG/Source/PCG/Private/Elements/PCGSpawnActor.cpp
@@ -547,8 +547,9 @@ bool FPCGSpawnActorElement::SpawnAndPrepareSubgraphs(FPCGSubgraphContext* Contex
 								Context->SourceComponent.Get(),
 								MakeShared<FPCGTrivialElement>(),// TODO: prepare user parameters like in subgraph/loop
 								MakeShared<FPCGInputForwardingElement>(SubgraphInputData),
-								{},
-								&InvocationStack);
+								/*Dependencies=*/{},
+								&InvocationStack,
+								/*bAllowHierarchicalGeneration=*/false);
 
 							if (SubgraphTaskId != InvalidPCGTaskId)
 							{
diff --git a/Engine/Plugins/Experimental/PCG/Source/PCG/Private/Graph/PCGGraphCompiler.cpp b/Engine/Plugins/Experimental/PCG/Source/PCG/Private/Graph/PCGGraphCompiler.cpp
--- a/Engine/Plugins/Experimental/PCG/Source/PCG/Private/Graph/PCGGraphCompiler.cpp
+++ b/Engine/Plugins/Experimental/PCG/Source/PCG/Private/Graph/PCGGraphCompiler.cpp
@@ -781,6 +781,7 @@ void FPCGGraphCompiler::CompileTopGraph(UPCGGraph* InGraph, uint32 GenerationGri
 	CullTasksStaticBranchNodes(CompiledTasks);
 
 	// For hierarchical generation resolve the execution grid for each task and cull any tasks that won't execute.
+	// TODO - we could add an else branch for higen disabled that culls the grid size nodes.
 	if (InGraph->IsHierarchicalGenerationEnabled() && GenerationGridSize != PCGHiGenGrid::UninitializedGridSize())
 	{
 		const EPCGHiGenGrid GenerationGrid = PCGHiGenGrid::GridSizeToGrid(GenerationGridSize);
diff --git a/Engine/Plugins/Experimental/PCG/Source/PCG/Private/Graph/PCGGraphCompiler.h b/Engine/Plugins/Experimental/PCG/Source/PCG/Private/Graph/PCGGraphCompiler.h
--- a/Engine/Plugins/Experimental/PCG/Source/PCG/Private/Graph/PCGGraphCompiler.h
+++ b/Engine/Plugins/Experimental/PCG/Source/PCG/Private/Graph/PCGGraphCompiler.h
@@ -30,6 +30,8 @@ public:
 
 private:
 	TArray<FPCGGraphTask> CompileGraph(UPCGGraph* InGraph, FPCGTaskId& NextId, FPCGStackContext& InOutStackContext);
+
+	/** Compiles the top graph and applies culling optimizations if a non-unitialized grid size is provided. */
 	void CompileTopGraph(UPCGGraph* InGraph, uint32 GenerationGridSize);
 
 	/** Propagates grid sizes through a graph's compiled tasks. */
diff --git a/Engine/Plugins/Experimental/PCG/Source/PCG/Private/Graph/PCGGraphExecutor.cpp b/Engine/Plugins/Experimental/PCG/Source/PCG/Private/Graph/PCGGraphExecutor.cpp
--- a/Engine/Plugins/Experimental/PCG/Source/PCG/Private/Graph/PCGGraphExecutor.cpp
+++ b/Engine/Plugins/Experimental/PCG/Source/PCG/Private/Graph/PCGGraphExecutor.cpp
@@ -100,10 +100,17 @@ FPCGTaskId FPCGGraphExecutor::Schedule(UPCGComponent* Component, const TArray<FP
 	check(Component);
 	UPCGGraph* Graph = Component->GetGraph();
 
-	return Schedule(Graph, Component, FPCGElementPtr(), GetFetchInputElement(), ExternalDependencies, InFromStack);
+	return Schedule(Graph, Component, FPCGElementPtr(), GetFetchInputElement(), ExternalDependencies, InFromStack, /*bAllowHierarchicalGeneration=*/true);
 }
 
-FPCGTaskId FPCGGraphExecutor::Schedule(UPCGGraph* Graph, UPCGComponent* SourceComponent, FPCGElementPtr PreGraphElement, FPCGElementPtr InputElement, const TArray<FPCGTaskId>& ExternalDependencies, const FPCGStack* InFromStack)
+FPCGTaskId FPCGGraphExecutor::Schedule(
+	UPCGGraph* Graph,
+	UPCGComponent* SourceComponent,
+	FPCGElementPtr PreGraphElement,
+	FPCGElementPtr InputElement,
+	const TArray<FPCGTaskId>& ExternalDependencies,
+	const FPCGStack* InFromStack,
+	bool bAllowHierarchicalGeneration)
 {
 	check(SourceComponent);
 
@@ -115,8 +122,12 @@ FPCGTaskId FPCGGraphExecutor::Schedule(UPCGGraph* Graph, UPCGComponent* SourceCo
 	
 	FPCGTaskId ScheduledId = InvalidPCGTaskId;
 
-	const bool bNonPartitionedComponent = !SourceComponent->IsLocalComponent() && !SourceComponent->IsPartitioned();
-	const uint32 GenerationGridSize = bNonPartitionedComponent ? PCGHiGenGrid::UninitializedGridSize() : SourceComponent->GetGenerationGridSize();
+	uint32 GenerationGridSize = PCGHiGenGrid::UninitializedGridSize();
+	const bool bComponentIsPartOfHierarchy = SourceComponent->IsLocalComponent() || SourceComponent->IsPartitioned();
+	if (bAllowHierarchicalGeneration && bComponentIsPartOfHierarchy)
+	{
+		GenerationGridSize = SourceComponent->GetGenerationGridSize();
+	}
 
 	// Get compiled tasks from compiler
 	TSharedPtr<FPCGStackContext> StackContextPtr = MakeShared<FPCGStackContext>();
diff --git a/Engine/Plugins/Experimental/PCG/Source/PCG/Private/Graph/PCGGraphExecutor.h b/Engine/Plugins/Experimental/PCG/Source/PCG/Private/Graph/PCGGraphExecutor.h
--- a/Engine/Plugins/Experimental/PCG/Source/PCG/Private/Graph/PCGGraphExecutor.h
+++ b/Engine/Plugins/Experimental/PCG/Source/PCG/Private/Graph/PCGGraphExecutor.h
@@ -93,7 +93,16 @@ public:
 
 	/** Schedules the execution of a given graph with specified inputs. This call is threadsafe */
 	FPCGTaskId Schedule(UPCGComponent* InComponent, const TArray<FPCGTaskId>& TaskDependency = TArray<FPCGTaskId>(), const FPCGStack* InFromStack = nullptr);
-	FPCGTaskId Schedule(UPCGGraph* Graph, UPCGComponent* InSourceComponent, FPCGElementPtr PreGraphElement, FPCGElementPtr InputElement, const TArray<FPCGTaskId>& TaskDependency, const FPCGStack* InFromStack);
+
+	/** Schedules the execution of a given graph with specified inputs. This call is threadsafe */
+	FPCGTaskId Schedule(
+		UPCGGraph* Graph,
+		UPCGComponent* InSourceComponent,
+		FPCGElementPtr PreGraphElement,
+		FPCGElementPtr InputElement,
+		const TArray<FPCGTaskId>& TaskDependency,
+		const FPCGStack* InFromStack,
+		bool bAllowHierarchicalGeneration);
 
 	/** Cancels all tasks originating from the given component */
 	TArray<UPCGComponent*> Cancel(UPCGComponent* InComponent);
diff --git a/Engine/Plugins/Experimental/PCG/Source/PCG/Private/PCGSubgraph.cpp b/Engine/Plugins/Experimental/PCG/Source/PCG/Private/PCGSubgraph.cpp
--- a/Engine/Plugins/Experimental/PCG/Source/PCG/Private/PCGSubgraph.cpp
+++ b/Engine/Plugins/Experimental/PCG/Source/PCG/Private/PCGSubgraph.cpp
@@ -535,7 +535,15 @@ bool FPCGSubgraphElement::ExecuteInternal(FPCGContext* InContext) const
 				FPCGStack InvocationStack = ensure(Context->Stack) ? *Context->Stack : FPCGStack();
 				InvocationStack.GetStackFramesMutable().Emplace(Context->Node);
 
-				FPCGTaskId SubgraphTaskId = Subsystem->ScheduleGraph(Subgraph, Context->SourceComponent.Get(), MakeShared<FPCGInputForwardingElement>(PreSubgraphInputData), MakeShared<FPCGInputForwardingElement>(SubgraphInputData), {}, &InvocationStack);
+				// Higen is not allowed in dynamic subgraphs, entire subgraph is executed on the same grid as this subgraph node.
+				FPCGTaskId SubgraphTaskId = Subsystem->ScheduleGraph(
+					Subgraph,
+					Context->SourceComponent.Get(),
+					MakeShared<FPCGInputForwardingElement>(PreSubgraphInputData),
+					MakeShared<FPCGInputForwardingElement>(SubgraphInputData),
+					/*Dependencies=*/{},
+					&InvocationStack,
+					/*bAllowHierarchicalGeneration=*/false);
 
 				if (SubgraphTaskId != InvalidPCGTaskId)
 				{
diff --git a/Engine/Plugins/Experimental/PCG/Source/PCG/Private/PCGSubsystem.cpp b/Engine/Plugins/Experimental/PCG/Source/PCG/Private/PCGSubsystem.cpp
--- a/Engine/Plugins/Experimental/PCG/Source/PCG/Private/PCGSubsystem.cpp
+++ b/Engine/Plugins/Experimental/PCG/Source/PCG/Private/PCGSubsystem.cpp
@@ -692,11 +692,18 @@ FPCGTaskId UPCGSubsystem::ScheduleGraph(UPCGComponent* SourceComponent, const TA
 	}
 }
 
-FPCGTaskId UPCGSubsystem::ScheduleGraph(UPCGGraph* Graph, UPCGComponent* SourceComponent, FPCGElementPtr PreGraphElement, FPCGElementPtr InputElement, const TArray<FPCGTaskId>& Dependencies, const FPCGStack* InFromStack)
+FPCGTaskId UPCGSubsystem::ScheduleGraph(
+	UPCGGraph* Graph,
+	UPCGComponent* SourceComponent,
+	FPCGElementPtr PreGraphElement, 
+	FPCGElementPtr InputElement, 
+	const TArray<FPCGTaskId>& Dependencies, 
+	const FPCGStack* InFromStack, 
+	bool bAllowHierarchicalGeneration)
 {
 	if (SourceComponent)
 	{
-		return GraphExecutor->Schedule(Graph, SourceComponent, PreGraphElement, InputElement, Dependencies, InFromStack);
+		return GraphExecutor->Schedule(Graph, SourceComponent, PreGraphElement, InputElement, Dependencies, InFromStack, bAllowHierarchicalGeneration);
 	}
 	else
 	{
diff --git a/Engine/Plugins/Experimental/PCG/Source/PCG/Public/PCGSubsystem.h b/Engine/Plugins/Experimental/PCG/Source/PCG/Public/PCGSubsystem.h
--- a/Engine/Plugins/Experimental/PCG/Source/PCG/Public/PCGSubsystem.h
+++ b/Engine/Plugins/Experimental/PCG/Source/PCG/Public/PCGSubsystem.h
@@ -104,7 +104,14 @@ public:
 	FPCGTaskId ScheduleCleanup(UPCGComponent* PCGComponent, bool bRemoveComponents, bool bSave, const TArray<FPCGTaskId>& Dependencies);
 
 	// Schedule graph (used internally for dynamic subgraph execution)
-	FPCGTaskId ScheduleGraph(UPCGGraph* Graph, UPCGComponent* SourceComponent, FPCGElementPtr PreGraphElement, FPCGElementPtr InputElement, const TArray<FPCGTaskId>& Dependencies, const FPCGStack* InFromStack);
+	FPCGTaskId ScheduleGraph(
+		UPCGGraph* Graph,
+		UPCGComponent* SourceComponent,
+		FPCGElementPtr PreGraphElement,
+		FPCGElementPtr InputElement,
+		const TArray<FPCGTaskId>& Dependencies, 
+		const FPCGStack* InFromStack,
+		bool bAllowHierarchicalGeneration);
 
 	// Schedule graph (used internally for dynamic subgraph execution)
 	FPCGTaskId ScheduleGraph(UPCGComponent* SourceComponent, const TArray<FPCGTaskId>& Dependencies);
