[PCG] Extended MetadataOpsBase to now include 1:N options for multiple inputs in addition to 1:1, N:1, N:N

#rnx

diff --git a/Engine/Plugins/Experimental/PCG/Source/PCG/Private/Elements/Metadata/PCGMetadataOpElementBase.cpp b/Engine/Plugins/Experimental/PCG/Source/PCG/Private/Elements/Metadata/PCGMetadataOpElementBase.cpp
--- a/Engine/Plugins/Experimental/PCG/Source/PCG/Private/Elements/Metadata/PCGMetadataOpElementBase.cpp
+++ b/Engine/Plugins/Experimental/PCG/Source/PCG/Private/Elements/Metadata/PCGMetadataOpElementBase.cpp
@@ -265,15 +265,18 @@ void FPCGMetadataElementBase::PassthroughInput(FPCGContext* Context, TArray<FPCG
 	const uint32 PrimaryPinIndex = Settings->GetInputPinToForward();
 	TArray<FPCGTaggedData> InputsToForward = Context->InputData.GetInputsByPin(Settings->GetInputPinLabel(PrimaryPinIndex));
 
-	if (!ensure(Index < InputsToForward.Num()))
+	if (InputsToForward.IsEmpty())
 	{
 		return;
 	}
 
+	// Take the index of the iteration, except for the 1:N case, where we just grab the first index
+	const int32 AdjustedIndex = (Index <= InputsToForward.Num()) ? Index : 0;
+
 	// Passthrough this single input to all of the outputs
 	for (uint32 I = 0; I < NumberOfOutputs; ++I)
 	{
-		Outputs.Emplace_GetRef(InputsToForward[Index]).Pin = Settings->GetOutputPinLabel(I);
+		Outputs.Emplace_GetRef(InputsToForward[AdjustedIndex]).Pin = Settings->GetOutputPinLabel(I);
 	}
 }
 
@@ -380,11 +383,30 @@ bool FPCGMetadataElementBase::PrepareDataInternal(FPCGContext* Context) const
 		return true;
 	}
 
+	int32 OperandInputNumMax = 0;
+
 	// There's no execution state, so just flag that it is ready to continue
-	TimeSlicedContext->InitializePerExecutionState();
+	TimeSlicedContext->InitializePerExecutionState([this, Settings, &OperandInputNumMax, OperandNum](ContextType* Context, PCGTimeSlice::FEmptyStruct& OutState) -> EPCGTimeSliceInitResult
+	{
+		for (uint32 i = 0; i < OperandNum; ++i)
+		{
+			const int32 CurrentInputNum = Context->InputData.GetInputsByPin(Settings->GetInputPinLabel(i)).Num();
+
+			OperandInputNumMax = FMath::Max(OperandInputNumMax, CurrentInputNum);
+
+			// For the current input, no input (0) could be default value and we support N:1 and 1:N
+			if (CurrentInputNum > 1 && CurrentInputNum != OperandInputNumMax)
+			{
+				PCGE_LOG(Error, GraphAndLog, LOCTEXT("MismatchedOperandDataCount", "Number of data elements provided on inputs must be 1:N, N:1, or N:N."));
+				return EPCGTimeSliceInitResult::AbortExecution;
+			}
+		}
+
+		return EPCGTimeSliceInitResult::Success;
+	});
 
 	// Set up the iterations on the multiple inputs of the primary pin
-	TimeSlicedContext->InitializePerIterationStates(PrimaryInputs.Num(), [this, Context, Settings, OperandNum, &PrimaryInputs, &PrimaryPinLabel](IterStateType& OutState, const ExecStateType& ExecState, const uint32 IterationIndex)
+	TimeSlicedContext->InitializePerIterationStates(OperandInputNumMax, [this, Context, OperandNum, Settings, OperandInputNumMax](IterStateType& OutState, const ExecStateType& ExecState, const uint32 IterationIndex)
 	{
 		TArray<FPCGTaggedData>& Outputs = Context->OutputData.TaggedData;
 		const uint32 NumberOfResults = Settings->GetResultNum();
@@ -426,18 +448,17 @@ bool FPCGMetadataElementBase::PrepareDataInternal(FPCGContext* Context) const
 				PCGE_LOG(Verbose, LogOnly, FText::Format(LOCTEXT("MissingInputDataForPin", "No data provided on pin '{0}'."), FText::FromName(CurrentPinLabel)));
 				return EPCGTimeSliceInitResult::NoOperation;
 			}
-			else if (CurrentPinInputData.Num() != 1 && CurrentPinInputData.Num() != PrimaryInputs.Num())
+			else if (CurrentPinInputData.Num() != 1 && CurrentPinInputData.Num() != OperandInputNumMax)
 			{
 				PCGE_LOG(Error, GraphAndLog,
-					FText::Format(LOCTEXT("MismatchedDataCountForPin", "Number of data elements ({0}) provided on pin '{1}' doesn't match number of elements ({2}) on primary pin ({3}). Only 1 input or {2} are supported."),
+					FText::Format(LOCTEXT("MismatchedDataCountForPin", "Number of data elements ({0}) provided on pin '{1}' doesn't match number of expected elements ({2}). Only 1 input or {2} are supported."),
 						CurrentPinInputData.Num(),
 						FText::FromName(CurrentPinLabel),
-						PrimaryInputs.Num(),
-						FText::FromName(PrimaryPinLabel)));
+						OperandInputNumMax));
 				return EPCGTimeSliceInitResult::AbortExecution;
 			}
 
-			// The operand inputs must either be N:1 or N:N
+			// The operand inputs must either be N:1 or N:N or 1:N
 			InputTaggedData[OperandPinIndex] = CurrentPinInputData.Num() == 1 ? CurrentPinInputData[0] : MoveTemp(CurrentPinInputData[IterationIndex]);
 
 			// Check if we have any points
diff --git a/Engine/Plugins/Experimental/PCG/Source/PCG/Private/Tests/Elements/Metadata/PCGMetadataMathsOpTest.cpp b/Engine/Plugins/Experimental/PCG/Source/PCG/Private/Tests/Elements/Metadata/PCGMetadataMathsOpTest.cpp
--- a/Engine/Plugins/Experimental/PCG/Source/PCG/Private/Tests/Elements/Metadata/PCGMetadataMathsOpTest.cpp
+++ b/Engine/Plugins/Experimental/PCG/Source/PCG/Private/Tests/Elements/Metadata/PCGMetadataMathsOpTest.cpp
@@ -426,21 +426,24 @@ bool FPCGMetadataMathsOp_MultipleInput::RunTest(const FString& Parameters)
 
 	using EPCGMedadataMathsOperation::Add;
 
-	// 2 x 2
+	// 2 x 2 - N:N
 	bSuccess &= ExecuteTest(TestParams(Add, GenerateParamData<int32>({{4, 5}, {6, 7}}), TArray({0, 0}), TArray({10, 12})));
 
-	// 3 x 3
+	// 3 x 3 - N:N
 	bSuccess &= ExecuteTest(TestParams(Add, GenerateParamData<int32>({{4, 5, 6}, {7, 8, 9}}), TArray({0, 0, 0}), TArray({11, 13, 15})));
 
-	// 2 x 1
+	// 2 x 1 - N:1
 	bSuccess &= ExecuteTest(TestParams(Add, GenerateParamData<int32>({{4, 5}, {6}}), TArray({0, 0}), TArray({10, 11})));
 
-	// 3 x 1
+	// 3 x 1 - N:1
 	bSuccess &= ExecuteTest(TestParams(Add, GenerateParamData<int32>({{4, 5, 6}, {7}}), TArray({0, 0, 0}), TArray({11, 12, 13})));
 
-	// 1 x 3
+	// 1 x 3 - N:1 with second pin
 	bSuccess &= ExecuteTest(TestParams(Add, GenerateParamData<int32>({{4}, {5, 6, 7}}), TArray({0, 0, 0}), TArray({9, 10, 11}), /*PinToForward=*/ 1));
 
+	// 1 x 3 - 1:N
+	bSuccess &= ExecuteTest(TestParams(Add, GenerateParamData<int32>({{4}, {5, 6, 7}}), TArray({0, 0, 0}), TArray({9, 10, 11})));
+
 	return bSuccess;
 }
 
