Tentative fix for Non-unity build

diff --git a/Engine/Plugins/Experimental/PCG/Source/PCG/Private/Tests/Elements/PCGPointNeighborhoodTest.cpp b/Engine/Plugins/Experimental/PCG/Source/PCG/Private/Tests/Elements/PCGPointNeighborhoodTest.cpp
--- a/Engine/Plugins/Experimental/PCG/Source/PCG/Private/Tests/Elements/PCGPointNeighborhoodTest.cpp
+++ b/Engine/Plugins/Experimental/PCG/Source/PCG/Private/Tests/Elements/PCGPointNeighborhoodTest.cpp
@@ -327,7 +327,7 @@ bool FPCGPointNeighborhoodTest_AvgColor::RunTest(const FString& Parameters)
 
 	for (int i = 0; i < OutPoints.Num(); ++i)
 	{
-		UTEST_EQUAL("Average Color in Location Attribute", OutPoints[1].Color, FVector4(1.0, 1.0, 1.0, 1.0));
+		UTEST_EQUAL("Average Color in Location Attribute", OutPoints[i].Color, FVector4(1.0, 1.0, 1.0, 1.0));
 	}
 	
 	return true;
