[PCG] Add options for data parity for Get Object From Property Path

* Add an "PersistAllData" advanced option to make sure we have as many output data (even if they are empty) than input data
* Add option to silence those cases (also advanced)
* Change underlying types for Object Property Accessor (FString -> FSoftObjectPath)

#rnx

diff --git a/Engine/Plugins/Experimental/PCG/Source/PCG/Private/Elements/PCGGetPropertyFromObjectPath.cpp b/Engine/Plugins/Experimental/PCG/Source/PCG/Private/Elements/PCGGetPropertyFromObjectPath.cpp
--- a/Engine/Plugins/Experimental/PCG/Source/PCG/Private/Elements/PCGGetPropertyFromObjectPath.cpp
+++ b/Engine/Plugins/Experimental/PCG/Source/PCG/Private/Elements/PCGGetPropertyFromObjectPath.cpp
@@ -129,7 +129,22 @@ bool FPCGGetPropertyFromObjectPathElement::PrepareDataInternal(FPCGContext* Cont
 
 				if (!Accessor.IsValid() || !Keys.IsValid())
 				{
-					PCGE_LOG(Error, GraphAndLog, FText::Format(LOCTEXT("AttributeNotFound", "Attribute/Property '{0}' does not exist on input {1}"), AttributeSelector.GetDisplayText(), FText::AsNumber(Index)));
+					if (Settings->bPersistAllData)
+					{
+						// Special case for empty data. We need this case if we ever chain this node multiple times. An empty param (with no attributes and no entries) will generate another empty param.
+						const UPCGMetadata* Metadata = Input.Data->ConstMetadata();
+						if (Metadata && Metadata->GetAttributeCount() == 0 && Metadata->GetLocalItemCount() == 0)
+						{
+							// Emplace empty path for this input. Will generate an empty param data.
+							ThisContext->PathsToObjectsToExtractAndIncomingDataIndex.Emplace(FSoftObjectPath(), Index);
+						}
+					}
+
+					if (!Settings->bSilenceErrorOnEmptyObjectPath)
+					{
+						PCGE_LOG(Error, GraphAndLog, FText::Format(LOCTEXT("AttributeNotFound", "Attribute/Property '{0}' does not exist on input {1}"), AttributeSelector.GetDisplayText(), FText::AsNumber(Index)));
+					}
+
 					continue;
 				}
 
@@ -139,14 +154,24 @@ bool FPCGGetPropertyFromObjectPathElement::PrepareDataInternal(FPCGContext* Cont
 					continue;
 				}
 
-				TArray<FSoftObjectPath> InputValues;
+				// Extract value as String to validate that a path is empty or ill-formed (because any ill-formed path will be null).
+				TArray<FString> InputValues;
 				InputValues.SetNum(NumElementsToAdd);
-				if (Accessor->GetRange(MakeArrayView(InputValues), 0, *Keys, EPCGAttributeAccessorFlags::AllowConstructible))
+				if (Accessor->GetRange(MakeArrayView(InputValues), 0, *Keys, EPCGAttributeAccessorFlags::AllowConstructible | EPCGAttributeAccessorFlags::AllowBroadcast))
 				{
 					ThisContext->PathsToObjectsToExtractAndIncomingDataIndex.Reserve(ThisContext->PathsToObjectsToExtractAndIncomingDataIndex.Num() + NumElementsToAdd);
 					for (int32 i = 0; i < InputValues.Num(); ++i)
 					{
-						FSoftObjectPath& Path = InputValues[i];
+						FString& StringPath = InputValues[i];
+						// Empty SoftObjectPath can convert to string to None and is treated as empty, so check that one too.
+						const bool PathIsEmpty = StringPath.IsEmpty() || StringPath.Equals(TEXT("None"), ESearchCase::CaseSensitive);
+						if (PathIsEmpty && Settings->bPersistAllData)
+						{
+							ThisContext->PathsToObjectsToExtractAndIncomingDataIndex.Emplace(FSoftObjectPath(), Index);
+						}
+
+						FSoftObjectPath Path(std::move(StringPath));
+
 						if (!Path.IsNull())
 						{
 							ObjectsToLoad.AddUnique(Path);
@@ -154,7 +179,11 @@ bool FPCGGetPropertyFromObjectPathElement::PrepareDataInternal(FPCGContext* Cont
 						}
 						else
 						{
-							PCGE_LOG(Error, GraphAndLog, FText::Format(LOCTEXT("InvalidPath", "Value number {0} for Attribute/Property '{1}' on input {2} is not a valid path or is null. Will be ignored."), FText::AsNumber(i), AttributeSelector.GetDisplayText(), FText::AsNumber(Index)));
+							if (!PathIsEmpty || !Settings->bSilenceErrorOnEmptyObjectPath)
+							{
+								PCGE_LOG(Error, GraphAndLog, FText::Format(LOCTEXT("InvalidPath", "Value number {0} for Attribute/Property '{1}' on input {2} is not a valid path or is null. Will be ignored."), FText::AsNumber(i), AttributeSelector.GetDisplayText(), FText::AsNumber(Index)));
+							}
+
 							continue;
 						}
 					}
@@ -199,11 +228,42 @@ bool FPCGGetPropertyFromObjectPathElement::ExecuteInternal(FPCGContext* Context)
 
 	const TArray<FPCGTaggedData> Inputs = Context->InputData.GetInputsByPin(PCGPinConstants::DefaultInputLabel);
 
+	// For any "invalid" entry that needs an empty param, just allocate one and use this one for all the invalids.
+	const UPCGParamData* EmptyParam = nullptr;
+
+	auto AddToOutput = [Context, &Inputs](const UPCGData* Data, int32 Index)
+	{
+		TArray<FPCGTaggedData>& Outputs = Context->OutputData.TaggedData;
+		FPCGTaggedData& Output = Outputs.Emplace_GetRef();
+		Output.Data = Data;
+		Output.Pin = PCGPinConstants::DefaultOutputLabel;
+		if (Index >= 0)
+		{
+			Output.Tags = Inputs[Index].Tags;
+		}
+	};
+
 	for (const TTuple<FSoftObjectPath, int32>& SoftPathAndIndex : ThisContext->PathsToObjectsToExtractAndIncomingDataIndex)
 	{
 		const FSoftObjectPath& SoftPath = SoftPathAndIndex.Get<FSoftObjectPath>();
 		const int32 Index = SoftPathAndIndex.Get<int32>();
 
+		if (SoftPath.IsNull() && Settings->bPersistAllData)
+		{
+			if (!Settings->bSilenceErrorOnEmptyObjectPath)
+			{
+				PCGE_LOG(Error, GraphAndLog, FText::Format(LOCTEXT("EmptyData", "Empty data on index {0}"), FText::AsNumber(Index)));
+			}
+
+			if (!EmptyParam)
+			{
+				EmptyParam = NewObject<UPCGParamData>();
+			}
+
+			AddToOutput(EmptyParam, Index);
+			continue;
+		}
+
 		const UObject* Object = SoftPath.ResolveObject();
 		if (!Object)
 		{
@@ -216,14 +276,7 @@ bool FPCGGetPropertyFromObjectPathElement::ExecuteInternal(FPCGContext* Context)
 		PCGPropertyHelpers::FExtractorParameters Parameters{ Object, Object->GetClass(), Selector, Settings->OutputAttributeName, Settings->bForceObjectAndStructExtraction, /*bPropertyNeedsToBeVisible=*/true };
 		if (UPCGParamData* ParamData = PCGPropertyHelpers::ExtractPropertyAsAttributeSet(Parameters, Context))
 		{
-			TArray<FPCGTaggedData>& Outputs = Context->OutputData.TaggedData;
-			FPCGTaggedData& Output = Outputs.Emplace_GetRef();
-			Output.Data = ParamData;
-			Output.Pin = PCGPinConstants::DefaultOutputLabel;
-			if (Index >= 0)
-			{
-				Output.Tags = Inputs[Index].Tags;
-			}
+			AddToOutput(ParamData, Index);
 		}
 		else
 		{
diff --git a/Engine/Plugins/Experimental/PCG/Source/PCG/Public/Elements/PCGGetPropertyFromObjectPath.h b/Engine/Plugins/Experimental/PCG/Source/PCG/Public/Elements/PCGGetPropertyFromObjectPath.h
--- a/Engine/Plugins/Experimental/PCG/Source/PCG/Public/Elements/PCGGetPropertyFromObjectPath.h
+++ b/Engine/Plugins/Experimental/PCG/Source/PCG/Public/Elements/PCGGetPropertyFromObjectPath.h
@@ -57,6 +57,14 @@ public:
 	/** By default, object loading is asynchronous, can force it synchronous if needed. */
 	UPROPERTY(BlueprintReadWrite, EditAnywhere, Category = "Settings|Debug")
 	bool bSynchronousLoad = false;
+
+	/** Opt-in option to create empty data when there is nothing to extract or property is not found, to have the same number of inputs than outputs. */
+	UPROPERTY(BlueprintReadWrite, EditAnywhere, Category = "Settings|Advanced")
+	bool bPersistAllData = false;
+
+	/** Opt-in option to silence errors when the path is Empty or nothing to extract. */
+	UPROPERTY(BlueprintReadWrite, EditAnywhere, Category = "Settings|Advanced")
+	bool bSilenceErrorOnEmptyObjectPath = false;
 };
 
 struct FPCGGetPropertyFromObjectPathContext : public FPCGContext
diff --git a/Engine/Plugins/Experimental/PCG/Source/PCG/Public/Metadata/Accessors/PCGPropertyAccessor.h b/Engine/Plugins/Experimental/PCG/Source/PCG/Public/Metadata/Accessors/PCGPropertyAccessor.h
--- a/Engine/Plugins/Experimental/PCG/Source/PCG/Public/Metadata/Accessors/PCGPropertyAccessor.h
+++ b/Engine/Plugins/Experimental/PCG/Source/PCG/Public/Metadata/Accessors/PCGPropertyAccessor.h
@@ -392,13 +392,13 @@ private:
 /**
 * Templated accessor class for object/class ptr properties.
 * Do not instantiate it manually, use PCGAttributeAccessorHelpers::CreatePropertyAccessor.
-* Will always convert to FString for PCG
+* Will always convert to FSoftObjectPath for PCG
 * Key supported: Generic object
 */
 class FPCGPropertyObjectPtrAccessor : public IPCGAttributeAccessorT<FPCGPropertyObjectPtrAccessor>, IPCGPropertyChainAccessor
 {
 public:
-	using Type = FString;
+	using Type = FSoftObjectPath;
 	using Super = IPCGAttributeAccessorT<FPCGPropertyObjectPtrAccessor>;
 
 	FPCGPropertyObjectPtrAccessor(const FObjectProperty* InProperty, TArray<const FProperty*>&& ExtraProperties = {})
@@ -421,7 +421,7 @@ public:
 	{
 		return PCGPropertyAccessor::IterateSet(GetPropertyChain(), InValues, Index, Keys, [this](void* PropertyAddressData, const Type& Value) -> void
 		{
-			Property->SetPropertyValue(PropertyAddressData, FSoftObjectPath(Value).TryLoad());
+			Property->SetPropertyValue(PropertyAddressData, Value.TryLoad());
 		});
 	}
 
