[PCG] Add Set Object/Class Parameter for BP

#rnx

diff --git a/Engine/Plugins/Experimental/PCG/Source/PCG/Private/Helpers/PCGGraphParameterExtension.cpp b/Engine/Plugins/Experimental/PCG/Source/PCG/Private/Helpers/PCGGraphParameterExtension.cpp
--- a/Engine/Plugins/Experimental/PCG/Source/PCG/Private/Helpers/PCGGraphParameterExtension.cpp
+++ b/Engine/Plugins/Experimental/PCG/Source/PCG/Private/Helpers/PCGGraphParameterExtension.cpp
@@ -63,6 +63,14 @@ template<> TValueOrError<TSoftClassPtr<UObject>, EPropertyBagResult> FPCGGraphPa
 	TValueOrError<TSoftClassPtr<UObject>, EPropertyBagResult> FinalResult(MakeValue(TSoftClassPtr<UObject>(InitialResult.GetValue())));
 	return FinalResult;
 }
+template<> TValueOrError<UObject*, EPropertyBagResult> FPCGGraphParameterExtension::GetGraphParameter(const FInstancedPropertyBag& PropertyBag, const FName PropertyName)
+{
+	return PropertyBag.GetValueObject<UObject>(PropertyName);
+}
+template<> TValueOrError<UClass*, EPropertyBagResult> FPCGGraphParameterExtension::GetGraphParameter(const FInstancedPropertyBag& PropertyBag, const FName PropertyName)
+{
+	return PropertyBag.GetValueClass(PropertyName);
+}
 template<> TValueOrError<FVector*, EPropertyBagResult> FPCGGraphParameterExtension::GetGraphParameter(const FInstancedPropertyBag& PropertyBag, const FName PropertyName)
 {
 	return PropertyBag.GetValueStruct<FVector>(PropertyName);
@@ -147,6 +155,14 @@ template<> EPropertyBagResult FPCGGraphParameterExtension::SetGraphParameter(FIn
 {
 	return PropertyBag.SetValueSoftPath(PropertyName, Value.ToSoftObjectPath());
 }
+template<> EPropertyBagResult FPCGGraphParameterExtension::SetGraphParameter(FInstancedPropertyBag& PropertyBag, const FName PropertyName, UObject* const& Value)
+{
+	return PropertyBag.SetValueObject<UObject>(PropertyName, Value);
+}
+template<> EPropertyBagResult FPCGGraphParameterExtension::SetGraphParameter(FInstancedPropertyBag& PropertyBag, const FName PropertyName, UClass* const& Value)
+{
+	return PropertyBag.SetValueClass(PropertyName, Value);
+}
 template<> EPropertyBagResult FPCGGraphParameterExtension::SetGraphParameter(FInstancedPropertyBag& PropertyBag, const FName PropertyName, const FVector& Value)
 {
 	return PropertyBag.SetValueStruct(PropertyName, Value);
diff --git a/Engine/Plugins/Experimental/PCG/Source/PCG/Private/Helpers/PCGGraphParametersHelpers.cpp b/Engine/Plugins/Experimental/PCG/Source/PCG/Private/Helpers/PCGGraphParametersHelpers.cpp
--- a/Engine/Plugins/Experimental/PCG/Source/PCG/Private/Helpers/PCGGraphParametersHelpers.cpp
+++ b/Engine/Plugins/Experimental/PCG/Source/PCG/Private/Helpers/PCGGraphParametersHelpers.cpp
@@ -67,7 +67,7 @@ namespace PCGGraphParametersHelpersPrivate
 			return T{};
 		}
 
-		if constexpr (std::is_pointer_v<U>)
+		if constexpr (!std::is_same_v<T, U> && std::is_pointer_v<U>)
 		{
 			static_assert(std::is_same_v<T, std::remove_pointer_t<U>>);
 			return Result.GetValue() ? *Result.GetValue() : T{};
@@ -187,6 +187,16 @@ TSoftClassPtr<UObject> UPCGGraphParametersHelpers::GetSoftClassParameter(const U
 	return PCGGraphParametersHelpersPrivate::ValidateAndReturnResult<TSoftClassPtr<UObject>>(GraphInterface, Name);
 }
 
+UObject* UPCGGraphParametersHelpers::GetObjectParameter(const UPCGGraphInterface* GraphInterface, const FName Name)
+{
+	return PCGGraphParametersHelpersPrivate::ValidateAndReturnResult<UObject*>(GraphInterface, Name);
+}
+
+UClass* UPCGGraphParametersHelpers::GetClassParameter(const UPCGGraphInterface* GraphInterface, const FName Name)
+{
+	return PCGGraphParametersHelpersPrivate::ValidateAndReturnResult<UClass*>(GraphInterface, Name);
+}
+
 FVector UPCGGraphParametersHelpers::GetVectorParameter(const UPCGGraphInterface* GraphInterface, const FName Name)
 {
 	return PCGGraphParametersHelpersPrivate::ValidateAndReturnResult<FVector, FVector*>(GraphInterface, Name);
@@ -281,6 +291,16 @@ void UPCGGraphParametersHelpers::SetSoftClassParameter(UPCGGraphInterface* Graph
 	PCGGraphParametersHelpersPrivate::ValidateAndSetValue(GraphInterface, Name, Value);
 }
 
+void UPCGGraphParametersHelpers::SetObjectParameter(UPCGGraphInterface* GraphInterface, const FName Name, UObject* Value)
+{
+	PCGGraphParametersHelpersPrivate::ValidateAndSetValue(GraphInterface, Name, Value);
+}
+
+void UPCGGraphParametersHelpers::SetClassParameter(UPCGGraphInterface* GraphInterface, const FName Name, UClass* Value)
+{
+	PCGGraphParametersHelpersPrivate::ValidateAndSetValue(GraphInterface, Name, Value);
+}
+
 void UPCGGraphParametersHelpers::SetVectorParameter(UPCGGraphInterface* GraphInterface, const FName Name, const FVector& Value)
 {
 	PCGGraphParametersHelpersPrivate::ValidateAndSetValue(GraphInterface, Name, Value);
diff --git a/Engine/Plugins/Experimental/PCG/Source/PCG/Public/Helpers/PCGGraphParameterExtension.h b/Engine/Plugins/Experimental/PCG/Source/PCG/Public/Helpers/PCGGraphParameterExtension.h
--- a/Engine/Plugins/Experimental/PCG/Source/PCG/Public/Helpers/PCGGraphParameterExtension.h
+++ b/Engine/Plugins/Experimental/PCG/Source/PCG/Public/Helpers/PCGGraphParameterExtension.h
@@ -1,4 +1,4 @@
-﻿// Copyright Epic Games, Inc. All Rights Reserved.
+// Copyright Epic Games, Inc. All Rights Reserved.
 
 #pragma once
 
@@ -33,6 +33,8 @@ template<> PCG_API TValueOrError<FString, EPropertyBagResult> FPCGGraphParameter
 template<> PCG_API TValueOrError<FSoftObjectPath*, EPropertyBagResult> FPCGGraphParameterExtension::GetGraphParameter(const FInstancedPropertyBag& PropertyBag, const FName PropertyName);
 template<> PCG_API TValueOrError<TSoftObjectPtr<UObject>, EPropertyBagResult> FPCGGraphParameterExtension::GetGraphParameter(const FInstancedPropertyBag& PropertyBag, const FName PropertyName);
 template<> PCG_API TValueOrError<TSoftClassPtr<UObject>, EPropertyBagResult> FPCGGraphParameterExtension::GetGraphParameter(const FInstancedPropertyBag& PropertyBag, const FName PropertyName);
+template<> PCG_API TValueOrError<UObject*, EPropertyBagResult> FPCGGraphParameterExtension::GetGraphParameter(const FInstancedPropertyBag& PropertyBag, const FName PropertyName);
+template<> PCG_API TValueOrError<UClass*, EPropertyBagResult> FPCGGraphParameterExtension::GetGraphParameter(const FInstancedPropertyBag& PropertyBag, const FName PropertyName);
 template<> PCG_API TValueOrError<FVector*, EPropertyBagResult> FPCGGraphParameterExtension::GetGraphParameter(const FInstancedPropertyBag& PropertyBag, const FName PropertyName);
 template<> PCG_API TValueOrError<FRotator*, EPropertyBagResult> FPCGGraphParameterExtension::GetGraphParameter(const FInstancedPropertyBag& PropertyBag, const FName PropertyName);
 template<> PCG_API TValueOrError<FTransform*, EPropertyBagResult> FPCGGraphParameterExtension::GetGraphParameter(const FInstancedPropertyBag& PropertyBag, const FName PropertyName);
@@ -51,6 +53,8 @@ template<> PCG_API EPropertyBagResult FPCGGraphParameterExtension::SetGraphParam
 template<> PCG_API EPropertyBagResult FPCGGraphParameterExtension::SetGraphParameter(FInstancedPropertyBag& PropertyBag, const FName PropertyName, const FSoftObjectPath& Value);
 template<> PCG_API EPropertyBagResult FPCGGraphParameterExtension::SetGraphParameter(FInstancedPropertyBag& PropertyBag, const FName PropertyName, const TSoftObjectPtr<UObject>& Value);
 template<> PCG_API EPropertyBagResult FPCGGraphParameterExtension::SetGraphParameter(FInstancedPropertyBag& PropertyBag, const FName PropertyName, const TSoftClassPtr<UObject>& Value);
+template<> PCG_API EPropertyBagResult FPCGGraphParameterExtension::SetGraphParameter(FInstancedPropertyBag& PropertyBag, const FName PropertyName, UObject* const& Value); // UObject* const& to respect template specialization
+template<> PCG_API EPropertyBagResult FPCGGraphParameterExtension::SetGraphParameter(FInstancedPropertyBag& PropertyBag, const FName PropertyName, UClass* const& Value); // UClass* const& to respect template specialization
 template<> PCG_API EPropertyBagResult FPCGGraphParameterExtension::SetGraphParameter(FInstancedPropertyBag& PropertyBag, const FName PropertyName, const FVector& Value);
 template<> PCG_API EPropertyBagResult FPCGGraphParameterExtension::SetGraphParameter(FInstancedPropertyBag& PropertyBag, const FName PropertyName, const FRotator& Value);
 template<> PCG_API EPropertyBagResult FPCGGraphParameterExtension::SetGraphParameter(FInstancedPropertyBag& PropertyBag, const FName PropertyName, const FTransform& Value);
diff --git a/Engine/Plugins/Experimental/PCG/Source/PCG/Public/Helpers/PCGGraphParametersHelpers.h b/Engine/Plugins/Experimental/PCG/Source/PCG/Public/Helpers/PCGGraphParametersHelpers.h
--- a/Engine/Plugins/Experimental/PCG/Source/PCG/Public/Helpers/PCGGraphParametersHelpers.h
+++ b/Engine/Plugins/Experimental/PCG/Source/PCG/Public/Helpers/PCGGraphParametersHelpers.h
@@ -51,6 +51,9 @@ public:
 	UFUNCTION(BlueprintCallable, Category = "PCG|Graph Parameters")
 	static uint8 GetEnumParameter(const UPCGGraphInterface* GraphInterface, const FName Name);
 
+	UFUNCTION(BlueprintCallable, Category = "PCG|Graph Parameters")
+	static FSoftObjectPath GetSoftObjectPathParameter(const UPCGGraphInterface* GraphInterface, const FName Name);
+
 	UFUNCTION(BlueprintCallable, Category = "PCG|Graph Parameters")
 	static TSoftObjectPtr<UObject> GetSoftObjectParameter(const UPCGGraphInterface* GraphInterface, const FName Name);
 
@@ -58,7 +61,10 @@ public:
 	static TSoftClassPtr<UObject> GetSoftClassParameter(const UPCGGraphInterface* GraphInterface, const FName Name);
 
 	UFUNCTION(BlueprintCallable, Category = "PCG|Graph Parameters")
-	static FSoftObjectPath GetSoftObjectPathParameter(const UPCGGraphInterface* GraphInterface, const FName Name);
+	static UObject* GetObjectParameter(const UPCGGraphInterface* GraphInterface, const FName Name);
+
+	UFUNCTION(BlueprintCallable, Category = "PCG|Graph Parameters")
+	static UClass* GetClassParameter(const UPCGGraphInterface* GraphInterface, const FName Name);
 
 	UFUNCTION(BlueprintCallable, Category = "PCG|Graph Parameters")
 	static FVector GetVectorParameter(const UPCGGraphInterface* GraphInterface, const FName Name);
@@ -118,6 +124,12 @@ public:
 	UFUNCTION(BlueprintCallable, Category = "PCG|Graph Parameters")
 	static void SetSoftClassParameter(UPCGGraphInterface* GraphInterface, const FName Name, UPARAM(ref) const TSoftClassPtr<UObject>& Value);
 
+	UFUNCTION(BlueprintCallable, Category = "PCG|Graph Parameters")
+	static void SetObjectParameter(UPCGGraphInterface* GraphInterface, const FName Name, UObject* Value);
+
+	UFUNCTION(BlueprintCallable, Category = "PCG|Graph Parameters")
+	static void SetClassParameter(UPCGGraphInterface* GraphInterface, const FName Name, UClass* Value);
+
 	UFUNCTION(BlueprintCallable, Category = "PCG|Graph Parameters")
 	static void SetVectorParameter(UPCGGraphInterface* GraphInterface, const FName Name, UPARAM(ref) const FVector& Value);
 
