[PCG] Replace Rotation.XYZ by [REDACTED_PERSON]/Pitch/Yaw in Attribute list view

* Will fix the issue that was showing the XYZ component of the Quat instead of the Roll/Pitch/Yaw values
* Will be easier for the users to use the right extractor

#rnx

diff --git a/Engine/Plugins/Experimental/PCG/Source/PCGEditor/Private/SPCGEditorGraphAttributeListView.cpp b/Engine/Plugins/Experimental/PCG/Source/PCGEditor/Private/SPCGEditorGraphAttributeListView.cpp
--- a/Engine/Plugins/Experimental/PCG/Source/PCGEditor/Private/SPCGEditorGraphAttributeListView.cpp
+++ b/Engine/Plugins/Experimental/PCG/Source/PCGEditor/Private/SPCGEditorGraphAttributeListView.cpp
@@ -51,9 +51,9 @@ namespace PCGEditorGraphAttributeListView
 	const FName NAME_PointPositionX = FName(TEXT("$Position.X"));
 	const FName NAME_PointPositionY = FName(TEXT("$Position.Y"));
 	const FName NAME_PointPositionZ = FName(TEXT("$Position.Z"));
-	const FName NAME_PointRotationX = FName(TEXT("$Rotation.X"));
-	const FName NAME_PointRotationY = FName(TEXT("$Rotation.Y"));
-	const FName NAME_PointRotationZ = FName(TEXT("$Rotation.Z"));
+	const FName NAME_PointRotationRoll = FName(TEXT("$Rotation.Roll"));
+	const FName NAME_PointRotationPitch = FName(TEXT("$Rotation.Pitch"));
+	const FName NAME_PointRotationYaw = FName(TEXT("$Rotation.Yaw"));
 	const FName NAME_PointScaleX = FName(TEXT("$Scale.X"));
 	const FName NAME_PointScaleY = FName(TEXT("$Scale.Y"));
 	const FName NAME_PointScaleZ = FName(TEXT("$Scale.Z"));
@@ -78,9 +78,9 @@ namespace PCGEditorGraphAttributeListView
 	const FText TEXT_PointPositionLabelX = LOCTEXT("PointPositionLabelX", "Position.X");
 	const FText TEXT_PointPositionLabelY = LOCTEXT("PointPositionLabelY", "Position.Y");
 	const FText TEXT_PointPositionLabelZ = LOCTEXT("PointPositionLabelZ", "Position.Z");
-	const FText TEXT_PointRotationLabelX = LOCTEXT("PointRotationLabelX", "Rotation.X");
-	const FText TEXT_PointRotationLabelY = LOCTEXT("PointRotationLabelY", "Rotation.Y");
-	const FText TEXT_PointRotationLabelZ = LOCTEXT("PointRotationLabelZ", "Rotation.Z");
+	const FText TEXT_PointRotationLabelRoll = LOCTEXT("PointRotationLabelRoll", "Rotation.Roll");
+	const FText TEXT_PointRotationLabelPitch = LOCTEXT("PointRotationLabelPitch", "Rotation.Pitch");
+	const FText TEXT_PointRotationLabelYaw = LOCTEXT("PointRotationLabelYaw", "Rotation.Yaw");
 	const FText TEXT_PointScaleLabelX = LOCTEXT("PointScaleLabelX", "Scale.X");
 	const FText TEXT_PointScaleLabelY = LOCTEXT("PointScaleLabelY", "Scale.Y");
 	const FText TEXT_PointScaleLabelZ = LOCTEXT("PointScaleLabelZ", "Scale.Z");
@@ -1227,9 +1227,9 @@ void SPCGEditorGraphAttributeListView::AddPointDataColumns(const UPCGPointData*
 	AddColumn(InPCGPointData, PCGEditorGraphAttributeListView::NAME_PointPositionX, PCGEditorGraphAttributeListView::TEXT_PointPositionLabelX);
 	AddColumn(InPCGPointData, PCGEditorGraphAttributeListView::NAME_PointPositionY, PCGEditorGraphAttributeListView::TEXT_PointPositionLabelY);
 	AddColumn(InPCGPointData, PCGEditorGraphAttributeListView::NAME_PointPositionZ, PCGEditorGraphAttributeListView::TEXT_PointPositionLabelZ);
-	AddColumn(InPCGPointData, PCGEditorGraphAttributeListView::NAME_PointRotationX, PCGEditorGraphAttributeListView::TEXT_PointRotationLabelX);
-	AddColumn(InPCGPointData, PCGEditorGraphAttributeListView::NAME_PointRotationY, PCGEditorGraphAttributeListView::TEXT_PointRotationLabelY);
-	AddColumn(InPCGPointData, PCGEditorGraphAttributeListView::NAME_PointRotationZ, PCGEditorGraphAttributeListView::TEXT_PointRotationLabelZ);
+	AddColumn(InPCGPointData, PCGEditorGraphAttributeListView::NAME_PointRotationRoll, PCGEditorGraphAttributeListView::TEXT_PointRotationLabelRoll);
+	AddColumn(InPCGPointData, PCGEditorGraphAttributeListView::NAME_PointRotationPitch, PCGEditorGraphAttributeListView::TEXT_PointRotationLabelPitch);
+	AddColumn(InPCGPointData, PCGEditorGraphAttributeListView::NAME_PointRotationYaw, PCGEditorGraphAttributeListView::TEXT_PointRotationLabelYaw);
 	AddColumn(InPCGPointData, PCGEditorGraphAttributeListView::NAME_PointScaleX, PCGEditorGraphAttributeListView::TEXT_PointScaleLabelX);
 	AddColumn(InPCGPointData, PCGEditorGraphAttributeListView::NAME_PointScaleY, PCGEditorGraphAttributeListView::TEXT_PointScaleLabelY);
 	AddColumn(InPCGPointData, PCGEditorGraphAttributeListView::NAME_PointScaleZ, PCGEditorGraphAttributeListView::TEXT_PointScaleLabelZ);
