[PCG] Normalize how we get Landscape transform when sampling points on the landscape

* Should fix issue with the WorldRayHit Query getting the wrong landscape proxy when sampling metadata.

#rnx

diff --git a/Engine/Plugins/Experimental/PCG/Source/PCG/Private/Data/PCGLandscapeData.cpp b/Engine/Plugins/Experimental/PCG/Source/PCG/Private/Data/PCGLandscapeData.cpp
--- a/Engine/Plugins/Experimental/PCG/Source/PCG/Private/Data/PCGLandscapeData.cpp
+++ b/Engine/Plugins/Experimental/PCG/Source/PCG/Private/Data/PCGLandscapeData.cpp
@@ -179,11 +179,7 @@ bool UPCGLandscapeData::SamplePoint(const FTransform& InTransform, const FBox& I
 		return false;
 	}
 
-#if WITH_EDITOR
-	const FTransform& LandscapeTransform = LandscapeInfo->GetLandscapeProxy()->GetTransform();
-#else
 	const FTransform LandscapeTransform = LandscapeInfo->GetLandscapeProxy()->LandscapeActorToWorld();
-#endif
 
 	// Box in local space -> box in world space -> box in landscape space
 	const FTransform BoundsTransformInLanscapeSpace = InTransform.GetRelativeTransform(LandscapeTransform);
@@ -240,16 +236,13 @@ bool UPCGLandscapeData::ProjectPoint(const FTransform& InTransform, const FBox&
 	}
 
 	const ULandscapeInfo* LandscapeInfo = GetLandscapeInfo(InTransform.GetLocation());
-	if (!LandscapeInfo || !LandscapeInfo->GetLandscapeProxy())
+	ALandscapeProxy* LandscapeProxy = LandscapeInfo ? LandscapeInfo->GetLandscapeProxy() : nullptr;
+	if (!LandscapeProxy)
 	{
 		return false;
 	}
 
-#if WITH_EDITOR
-	const FTransform& LandscapeTransform = LandscapeInfo->GetLandscapeProxy()->GetTransform();
-#else
-	const FTransform LandscapeTransform = LandscapeInfo->GetLandscapeProxy()->LandscapeActorToWorld();
-#endif
+	const FTransform LandscapeTransform = LandscapeProxy->LandscapeActorToWorld();
 
 	// TODO: compute full transform when we want to support bounds
 	const FVector LocalPoint = LandscapeTransform.InverseTransformPosition(InTransform.GetLocation());
@@ -259,7 +252,7 @@ bool UPCGLandscapeData::ProjectPoint(const FTransform& InTransform, const FBox&
 	ULandscapeComponent* LandscapeComponent = LandscapeInfo->XYtoComponentMap.FindRef(ComponentMapKey);
 	const FPCGLandscapeCacheEntry * LandscapeCacheEntry = (LandscapeComponent ? LandscapeCache->GetCacheEntry(LandscapeComponent, ComponentMapKey) : nullptr);
 #else
-	const FPCGLandscapeCacheEntry* LandscapeCacheEntry = LandscapeCache->GetCacheEntry(LandscapeInfo->GetLandscapeProxy(), LandscapeInfo->GetLandscapeProxy()->GetOriginalLandscapeGuid(), ComponentMapKey);
+	const FPCGLandscapeCacheEntry* LandscapeCacheEntry = LandscapeCache->GetCacheEntry(LandscapeProxy, LandscapeProxy->GetOriginalLandscapeGuid(), ComponentMapKey);
 #endif
 
 	if (!LandscapeCacheEntry)
@@ -381,16 +374,13 @@ const UPCGPointData* UPCGLandscapeData::CreatePointData(FPCGContext* Context, co
 
 	for(ULandscapeInfo* LandscapeInfo : AllLandscapeInfos)
 	{
-		if (!LandscapeInfo || !LandscapeInfo->GetLandscapeProxy())
+		ALandscapeProxy* LandscapeProxy = LandscapeInfo ? LandscapeInfo->GetLandscapeProxy() : nullptr;
+		if (!LandscapeProxy)
 		{
 			continue;
 		}
 
-#if WITH_EDITOR
-		const FTransform& LandscapeTransform = LandscapeInfo->GetLandscapeProxy()->GetTransform();
-#else
-		const FTransform LandscapeTransform = LandscapeInfo->GetLandscapeProxy()->LandscapeActorToWorld();
-#endif
+		const FTransform LandscapeTransform = LandscapeProxy->LandscapeActorToWorld();
 		const int32 ComponentSizeQuads = LandscapeInfo->ComponentSizeQuads;
 
 		// TODO: add offset to nearest edge, will have an impact if the grid size doesn't match the landscape size
@@ -420,7 +410,7 @@ const UPCGPointData* UPCGLandscapeData::CreatePointData(FPCGContext* Context, co
 		const int32 MaxComponentX = MaxX / ComponentSizeQuads;
 		const int32 MinComponentY = MinY / ComponentSizeQuads;
 		const int32 MaxComponentY = MaxY / ComponentSizeQuads;
-		const FGuid LandscapeGuid = LandscapeInfo->GetLandscapeProxy()->GetOriginalLandscapeGuid();
+		const FGuid LandscapeGuid = LandscapeProxy->GetOriginalLandscapeGuid();
 
 		for (int32 ComponentX = MinComponentX; ComponentX <= MaxComponentX; ++ComponentX)
 		{
@@ -431,7 +421,7 @@ const UPCGPointData* UPCGLandscapeData::CreatePointData(FPCGContext* Context, co
 				ULandscapeComponent* LandscapeComponent = LandscapeInfo->XYtoComponentMap.FindRef(ComponentMapKey);
 				const FPCGLandscapeCacheEntry* LandscapeCacheEntry = LandscapeComponent ? LandscapeCache->GetCacheEntry(LandscapeComponent, ComponentMapKey) : nullptr;
 #else
-				const FPCGLandscapeCacheEntry* LandscapeCacheEntry = LandscapeCache->GetCacheEntry(LandscapeInfo->GetLandscapeProxy(), LandscapeGuid, ComponentMapKey);
+				const FPCGLandscapeCacheEntry* LandscapeCacheEntry = LandscapeCache->GetCacheEntry(LandscapeProxy, LandscapeGuid, ComponentMapKey);
 #endif
 
 				if (!LandscapeCacheEntry)
@@ -469,7 +459,7 @@ const UPCGPointData* UPCGLandscapeData::CreatePointData(FPCGContext* Context, co
 		}
 
 		check(Points.Num() - PointsBeforeNum <= PointCountUpperBound);
-		UE_LOG(LogPCG, Verbose, TEXT("Landscape %s extracted %d of %d potential points"), *LandscapeInfo->GetLandscapeProxy()->GetFName().ToString(), Points.Num() - PointsBeforeNum, PointCountUpperBound);
+		UE_LOG(LogPCG, Verbose, TEXT("Landscape %s extracted %d of %d potential points"), *LandscapeProxy->GetFName().ToString(), Points.Num() - PointsBeforeNum, PointCountUpperBound);
 	}
 
 	return Data;
diff --git a/Engine/Plugins/Experimental/PCG/Source/PCG/Private/Grid/PCGLandscapeCache.cpp b/Engine/Plugins/Experimental/PCG/Source/PCG/Private/Grid/PCGLandscapeCache.cpp
--- a/Engine/Plugins/Experimental/PCG/Source/PCG/Private/Grid/PCGLandscapeCache.cpp
+++ b/Engine/Plugins/Experimental/PCG/Source/PCG/Private/Grid/PCGLandscapeCache.cpp
@@ -863,7 +863,8 @@ void UPCGLandscapeCache::SampleMetadataOnPoint(ALandscapeProxy* Landscape, FPCGP
 		return;
 	}
 
-	const FVector LocalPoint = Landscape->GetTransform().InverseTransformPosition(InOutPoint.Transform.GetLocation());
+	const FTransform LandscapeTransform = Landscape->LandscapeActorToWorld();
+	const FVector LocalPoint = LandscapeTransform.InverseTransformPosition(InOutPoint.Transform.GetLocation());
 	const FIntPoint ComponentMapKey(FMath::FloorToInt(LocalPoint.X / LandscapeInfo->ComponentSizeQuads), FMath::FloorToInt(LocalPoint.Y / LandscapeInfo->ComponentSizeQuads));
 
 #if WITH_EDITOR
diff --git a/Engine/Plugins/Experimental/PCG/Source/PCG/Private/Helpers/PCGBlueprintHelpers.cpp b/Engine/Plugins/Experimental/PCG/Source/PCG/Private/Helpers/PCGBlueprintHelpers.cpp
--- a/Engine/Plugins/Experimental/PCG/Source/PCG/Private/Helpers/PCGBlueprintHelpers.cpp
+++ b/Engine/Plugins/Experimental/PCG/Source/PCG/Private/Helpers/PCGBlueprintHelpers.cpp
@@ -191,11 +191,7 @@ TArray<FPCGLandscapeLayerWeight> UPCGBlueprintHelpers::GetInterpolatedPCGLandsca
 			continue;
 		}
 
-#if WITH_EDITOR
-		const FVector LocalPoint = Landscape->GetTransform().InverseTransformPosition(Location);
-#else
 		const FVector LocalPoint = Landscape->LandscapeActorToWorld().InverseTransformPosition(Location);
-#endif
 		const FIntPoint ComponentMapKey(FMath::FloorToInt(LocalPoint.X / LandscapeInfo->ComponentSizeQuads), FMath::FloorToInt(LocalPoint.Y / LandscapeInfo->ComponentSizeQuads));
 
 		#if WITH_EDITOR
