[PCG] Make Transform overridable on Spatial Noise

#rnx

diff --git a/Engine/Plugins/Experimental/PCG/Source/PCG/Public/Elements/PCGSpatialNoise.h b/Engine/Plugins/Experimental/PCG/Source/PCG/Public/Elements/PCGSpatialNoise.h
--- a/Engine/Plugins/Experimental/PCG/Source/PCG/Public/Elements/PCGSpatialNoise.h
+++ b/Engine/Plugins/Experimental/PCG/Source/PCG/Public/Elements/PCGSpatialNoise.h
@@ -102,7 +102,7 @@ public:
 	FVector RandomOffset = FVector(100000.0);
 
 	// this will apply a transform to the points before calculating noise
-	UPROPERTY(BlueprintReadWrite, EditAnywhere, Category = Settings, meta = (EditCondition = "Mode != PCGSpatialNoiseMode::Voronoi2D || !bTiling", EditConditionHides))
+	UPROPERTY(BlueprintReadWrite, EditAnywhere, Category = Settings, meta = (EditCondition = "Mode != PCGSpatialNoiseMode::Voronoi2D || !bTiling", EditConditionHides, PCG_Overridable))
 	FTransform Transform = FTransform::Identity;
 
 	// the less random this is, the more it returns to being a grid
