[PCG] Add Iteration value in all PCG BP loops

#rnx

diff --git a/Engine/Plugins/Experimental/PCG/Source/PCG/Private/Elements/PCGExecuteBlueprint.cpp b/Engine/Plugins/Experimental/PCG/Source/PCG/Private/Elements/PCGExecuteBlueprint.cpp
--- a/Engine/Plugins/Experimental/PCG/Source/PCG/Private/Elements/PCGExecuteBlueprint.cpp
+++ b/Engine/Plugins/Experimental/PCG/Source/PCG/Private/Elements/PCGExecuteBlueprint.cpp
@@ -791,7 +791,7 @@ void UPCGBlueprintElement::PointLoop(FPCGContext& InContext, const UPCGPointData
 
 	FPCGAsync::AsyncPointProcessing(&InContext, InPoints.Num(), OutPoints, [this, &InContext, InData, OutData, &InPoints](int32 Index, FPCGPoint& OutPoint)
 	{
-		return PointLoopBody(InContext, InData, InPoints[Index], OutPoint, OutData->Metadata);
+		return PointLoopBody(InContext, InData, InPoints[Index], OutPoint, OutData->Metadata, Index);
 	});
 }
 
@@ -818,7 +818,7 @@ void UPCGBlueprintElement::VariableLoop(FPCGContext& InContext, const UPCGPointD
 
 	FPCGAsync::AsyncMultiPointProcessing(&InContext, InPoints.Num(), OutPoints, [this, &InContext, InData, OutData, &InPoints](int32 Index)
 	{
-		return VariableLoopBody(InContext, InData, InPoints[Index], OutData->Metadata);
+		return VariableLoopBody(InContext, InData, InPoints[Index], OutData->Metadata, Index);
 	});
 }
 
@@ -847,7 +847,9 @@ void UPCGBlueprintElement::NestedLoop(FPCGContext& InContext, const UPCGPointDat
 
 	FPCGAsync::AsyncPointProcessing(&InContext, InOuterPoints.Num() * InInnerPoints.Num(), OutPoints, [this, &InContext, InOuterData, InInnerData, OutData, &InOuterPoints, &InInnerPoints](int32 Index, FPCGPoint& OutPoint)
 	{
-		return NestedLoopBody(InContext, InOuterData, InInnerData, InOuterPoints[Index / InInnerPoints.Num()], InInnerPoints[Index % InInnerPoints.Num()], OutPoint, OutData->Metadata);
+		const int32 OuterIndex = Index / InInnerPoints.Num();
+		const int32 InnerIndex = Index % InInnerPoints.Num();
+		return NestedLoopBody(InContext, InOuterData, InInnerData, InOuterPoints[OuterIndex], InInnerPoints[InnerIndex], OutPoint, OutData->Metadata, OuterIndex, InnerIndex);
 	});
 }
 
diff --git a/Engine/Plugins/Experimental/PCG/Source/PCG/Public/Elements/PCGExecuteBlueprint.h b/Engine/Plugins/Experimental/PCG/Source/PCG/Public/Elements/PCGExecuteBlueprint.h
--- a/Engine/Plugins/Experimental/PCG/Source/PCG/Public/Elements/PCGExecuteBlueprint.h
+++ b/Engine/Plugins/Experimental/PCG/Source/PCG/Public/Elements/PCGExecuteBlueprint.h
@@ -37,21 +37,77 @@ public:
 	virtual void BeginDestroy() override;
 	// ~End UObject interface
 
+	/**
+	 * Main execution function that will contain the logic for this PCG Element, with the context as parameter.
+	 * @param InContext - Context of the execution
+	 * @param Input     - Input collection containing all the data passed as input to the node.
+	 * @param Output    - Data collection that will be passed as the output of the node, with pins matching the ones provided during the execution.
+	 */
 	UFUNCTION(BlueprintNativeEvent, Category = "PCG|Execution")
 	void ExecuteWithContext(UPARAM(ref)FPCGContext& InContext, const FPCGDataCollection& Input, FPCGDataCollection& Output);
 
+	/**
+	 * Main execution function that will contain the logic for this PCG Element. Use GetContext to have access to the context.
+	 * @param Input  - Input collection containing all the data passed as input to the node.
+	 * @param Output - Data collection that will be passed as the output of the node, with pins matching the ones provided during the execution.
+	 */
 	UFUNCTION(BlueprintImplementableEvent, BlueprintCallable, Category = "PCG|Execution")
 	void Execute(const FPCGDataCollection& Input, FPCGDataCollection& Output);
 
+	/**
+	 * Multi-threaded loop that will iterate on all points in InData. All points will be added in the same order than in input. Will be called by Point Loop function.
+	 * @param InContext   - Context of the execution
+	 * @param InData      - Input point data. Constant, must not be modified.
+	 * @param InPoint     - Point for the current iteration. Constant, must not be modified.
+	 * @param OutPoint    - Point that will be added to the output data. Can be modified.
+	 * @param OutMetadata - Output metadata to write attribute to. Can be modified.
+	 * @param Iteration   - Index of the current point. Must only be used to access input data, as this call is multi-threaded. It is not safe to access output data.
+	 * @returns True if the point should be kept, False if not.
+	 */
 	UFUNCTION(BlueprintImplementableEvent, Category = "PCG|Flow Control")
-	bool PointLoopBody(const FPCGContext& InContext, const UPCGPointData* InData, const FPCGPoint& InPoint, FPCGPoint& OutPoint, UPCGMetadata* OutMetadata) const;
-
+	bool PointLoopBody(const FPCGContext& InContext, const UPCGPointData* InData, const FPCGPoint& InPoint, FPCGPoint& OutPoint, UPCGMetadata* OutMetadata, int64 Iteration) const;
+
+	/**
+	 * Multi-threaded loop that will be called on all points in InData. Can return a variable number of output points. 
+	 * All points will be added in the same order than in input. Will be called by Variable Loop function. 
+	 * @param InContext   - Context of the execution
+	 * @param InData      - Input point data. Constant, must not be modified.
+	 * @param InPoint     - Point for the current iteration. Constant, must not be modified.
+	 * @param OutMetadata - Output metadata to write attribute to. Can be modified.
+	 * @param Iteration   - Index of the current point. Must only be used to access input data, as this call is multi-threaded. It is not safe to access output data.
+	 * @returns Array of new points that will be added to the output point data.
+	 */
 	UFUNCTION(BlueprintImplementableEvent, Category = "PCG|Flow Control")
-	TArray<FPCGPoint> VariableLoopBody(const FPCGContext& InContext, const UPCGPointData* InData, const FPCGPoint& InPoint, UPCGMetadata* OutMetadata) const;
-
+	TArray<FPCGPoint> VariableLoopBody(const FPCGContext& InContext, const UPCGPointData* InData, const FPCGPoint& InPoint, UPCGMetadata* OutMetadata, int64 Iteration) const;
+
+	/**
+	 * Multi-threaded loop that will iterate on all nested loop pairs (e.g. (o, i) for all o in Outer, i in Inner).
+	 * All points will be added in the same order than in input (e.g: (0,0), (0,1), (0,2), ...). Will be called by Nested Loop function.
+	 * @param InContext      - Context of the execution
+	 * @param InOuterData    - Outer point data. Constant, must not be modified.
+	 * @param InInnerData    - Inner point data. Constant, must not be modified.
+	 * @param InOuterPoint   - Outer Point for the current iteration. Constant, must not be modified.
+	 * @param InInnerPoint   - Inner Point for the current iteration. Constant, must not be modified.
+	 * @param OutPoint       - Point that will be added to the output data. Can be modified.
+	 * @param OutMetadata    - Output metadata to write attribute to. Can be modified.
+	 * @param OuterIteration - Index of the current point in outer data. Must only be used to access input data, as this call is multi-threaded. It is not safe to access output data.
+	 * @param InnerIteration - Index of the current point in inner data. Must only be used to access input data, as this call is multi-threaded. It is not safe to access output data.
+	 * @returns True if the point should be kept, False if not.
+	 */
 	UFUNCTION(BlueprintImplementableEvent, Category = "PCG|Flow Control")
-	bool NestedLoopBody(const FPCGContext& InContext, const UPCGPointData* InOuterData, const UPCGPointData* InInnerData, const FPCGPoint& InOuterPoint, const FPCGPoint& InInnerPoint, FPCGPoint& OutPoint, UPCGMetadata* OutMetadata) const;
-
+	bool NestedLoopBody(const FPCGContext& InContext, const UPCGPointData* InOuterData, const UPCGPointData* InInnerData, const FPCGPoint& InOuterPoint, const FPCGPoint& InInnerPoint, FPCGPoint& OutPoint, UPCGMetadata* OutMetadata, int64 OuterIteration, int64 InnerIteration) const;
+
+	/**
+	 * Multi-threaded loop that will be called N number of times (defined by Iteration Loop parameter NumIterations). 
+	 * All points will be added in order (iteration 0 will be before iteration 1 in the final array).
+	 * @param InContext   - Context of the execution
+	 * @param Iteration   - Index of the current iteration. Must only be used to access input data, as this call is multi-threaded. It is not safe to access output data.
+	 * @param InA         - Optional input spatial data, can be null. Constant, must not be modified.
+	 * @param InB         - Optional input spatial data, can be null. Constant, must not be modified.
+	 * @param OutPoint    - Point that will be added to the output data. Can be modified.
+	 * @param OutMetadata - Output metadata to write attribute to. Can be modified.
+	 * @returns True if the point should be kept, False if not.
+	 */
 	UFUNCTION(BlueprintImplementableEvent, Category = "PCG|Flow Control")
 	bool IterationLoopBody(const FPCGContext& InContext, int64 Iteration, const UPCGSpatialData* InA, const UPCGSpatialData* InB, FPCGPoint& OutPoint, UPCGMetadata* OutMetadata) const;
 
@@ -75,16 +131,19 @@ public:
 	UFUNCTION(BlueprintNativeEvent, Category = "PCG|Node Customization")
 	FName NodeTitleOverride() const;
 
+	/** Override for the default node color. */
 	UFUNCTION(BlueprintNativeEvent, Category = "PCG|Node Customization")
 	FLinearColor NodeColorOverride() const;
 
+	/** Override to change the node type. */
 	UFUNCTION(BlueprintNativeEvent, Category = "PCG|Node Customization")
 	EPCGSettingsType NodeTypeOverride() const;
 
-	/** Override for the IsCacheable node property when it depends on the settings in your node */
+	/** Override for the IsCacheable node property when it depends on the settings in your node. If true, the node will be cached, if not it will always be executed. */
 	UFUNCTION(BlueprintNativeEvent, Category = "PCG|Execution")
 	bool IsCacheableOverride() const;
 
+	/** Apply the preconfigured settings specified in the class default. Used to create nodes that are configured with pre-defined settings. Use InPreconfigureInfo index to know which settings it is. */
 	UFUNCTION(BlueprintImplementableEvent, Category = "PCG|Preconfigure Settings", meta = (ForceAsFunction))
 	void ApplyPreconfiguredSettings(UPARAM(ref) const FPCGPreConfiguredSettingsInfo& InPreconfigureInfo);
 
