Allowed graph recursion in subgraph nodes - will stop recursing when there are no more inputs

diff --git a/Engine/Plugins/Experimental/PCG/Source/PCG/Private/Graph/PCGGraphCompiler.cpp b/Engine/Plugins/Experimental/PCG/Source/PCG/Private/Graph/PCGGraphCompiler.cpp
--- a/Engine/Plugins/Experimental/PCG/Source/PCG/Private/Graph/PCGGraphCompiler.cpp
+++ b/Engine/Plugins/Experimental/PCG/Source/PCG/Private/Graph/PCGGraphCompiler.cpp
@@ -47,21 +47,13 @@ TArray<FPCGGraphTask> FPCGGraphCompiler::CompileGraph(UPCGGraph* InGraph, FPCGTa
 		UPCGGraph* Subgraph = SubgraphNode ? SubgraphNode->GetSubgraph() : nullptr;
 		const UPCGBaseSubgraphSettings* SubgraphSettings = SubgraphNode ? Cast<const UPCGBaseSubgraphSettings>(SubgraphNode->GetSettings()) : nullptr;
 
-		// Only catch immediate infinite recursion.
-		// TODO: Add a better mechanism for detecting more complex cyclic recursions
-		// Like Graph A has subgraph node with Graph B, and Graph B has a subgraph node with Graph A (A -> B -> A)
-		// or A -> B -> C -> A, etc...
-		// NOTE (for the person that would work on it), keeping a stack of all the subgraphs is not enough, as we could already have compiled graph
-		// A, with an inclusion to graph B, but B has not yet a subgraph to A. When adding subgraph node to A in B, we will only recompile B, since A is already
-		// compiled and cached... We need to store a subgraph dependency chain to detect those more complex cases.
-		if (Subgraph == InGraph)
-		{
-			UE_LOG(LogPCG, Error, TEXT("[FPCGGraphCompiler::CompileGraph] %s cannot include itself as a subgraph, subgraph will not be executed."), *InGraph->GetName());
-			return TArray<FPCGGraphTask>();
-		}
+		// Note that recursive graphs must be dynamic by definition, as we aren't able to 'finish' emitting compiled tasks during compilation otherwise.
+		// Implementation note: it is very important that the predicate here is symmetrical with the one in the subgraph node otherwise some tasks could be missing.
+		const bool bIsRecursiveGraph = Subgraph && Subgraph->Contains(InGraph);
+		const bool bIsNonDynamic = SubgraphSettings && !SubgraphSettings->IsDynamicGraph();
+		const bool bIsNotDisabled = SubgraphSettings && SubgraphSettings->bEnabled;
 
-		const bool bIsNonDynamicAndNonDisabledSubgraphNode = (SubgraphNode && Subgraph && SubgraphSettings && !SubgraphSettings->IsDynamicGraph() && SubgraphSettings->bEnabled);
-		if (bIsNonDynamicAndNonDisabledSubgraphNode)
+		if(Subgraph && bIsNonDynamic && bIsNotDisabled && !bIsRecursiveGraph)
 		{
 			const FPCGTaskId PreId = NextId++;
 
diff --git a/Engine/Plugins/Experimental/PCG/Source/PCG/Private/Graph/PCGStackContext.cpp b/Engine/Plugins/Experimental/PCG/Source/PCG/Private/Graph/PCGStackContext.cpp
--- a/Engine/Plugins/Experimental/PCG/Source/PCG/Private/Graph/PCGStackContext.cpp
+++ b/Engine/Plugins/Experimental/PCG/Source/PCG/Private/Graph/PCGStackContext.cpp
@@ -136,12 +136,17 @@ const UPCGComponent* FPCGStack::GetRootComponent() const
 	return StackFrames.IsEmpty() ? nullptr : Cast<const UPCGComponent>(StackFrames[0].Object.Get());
 }
 
-const UPCGGraph* FPCGStack::GetRootGraph() const
+const UPCGGraph* FPCGStack::GetRootGraph(int32* OutRootFrameIndex) const
 {
 	for (int StackIndex = 0; StackIndex < GetStackFrames().Num(); ++StackIndex)
 	{
 		if (const UPCGGraph* Graph = Cast<const UPCGGraph>(StackFrames[StackIndex].Object.Get()))
 		{
+			if (OutRootFrameIndex)
+			{
+				*OutRootFrameIndex = StackIndex;
+			}
+
 			return Graph;
 		}
 	}
diff --git a/Engine/Plugins/Experimental/PCG/Source/PCG/Private/PCGGraph.cpp b/Engine/Plugins/Experimental/PCG/Source/PCG/Private/PCGGraph.cpp
--- a/Engine/Plugins/Experimental/PCG/Source/PCG/Private/PCGGraph.cpp
+++ b/Engine/Plugins/Experimental/PCG/Source/PCG/Private/PCGGraph.cpp
@@ -303,13 +303,13 @@ void UPCGGraph::PostLoad()
 	}
 
 	// Finally, apply deprecation that changes edges/rebinds
-	ForEachNode([](UPCGNode* InNode) { return InNode->ApplyDeprecationBeforeUpdatePins(); });
+	ForEachNode([](UPCGNode* InNode) { InNode->ApplyDeprecationBeforeUpdatePins(); return true; });
 
 	// Update pins on all nodes
-	ForEachNode([](UPCGNode* InNode) { return InNode->UpdatePins(); });
+	ForEachNode([](UPCGNode* InNode) { InNode->UpdatePins(); return true; });
 
 	// Finally, apply deprecation that changes edges/rebinds
-	ForEachNode([](UPCGNode* InNode) { return InNode->ApplyDeprecation(); });
+	ForEachNode([](UPCGNode* InNode) { InNode->ApplyDeprecation(); return true; });
 
 	InputNode->OnNodeChangedDelegate.AddUObject(this, &UPCGGraph::OnNodeChanged);
 	OutputNode->OnNodeChangedDelegate.AddUObject(this, &UPCGGraph::OnNodeChanged);
@@ -336,45 +336,60 @@ void UPCGGraph::PostLoad()
 }
 
 bool UPCGGraph::IsEditorOnly() const
-{
-	TSet<const UPCGGraph*> VisitedGraphs;
-	return IsEditorOnly_Internal(VisitedGraphs);
-}
-
-bool UPCGGraph::IsEditorOnly_Internal(TSet<const UPCGGraph*>& VisitedGraphs) const
 {
 	bool bIsCurrentlyEditorOnly = (Super::IsEditorOnly() || bIsEditorOnly);
-	
-	check(!VisitedGraphs.Contains(this));
-	VisitedGraphs.Add(this);
 
 	if (!bIsCurrentlyEditorOnly)
 	{
-		// Check for any subgraphs...
-		ForEachNode([&bIsCurrentlyEditorOnly, &VisitedGraphs](UPCGNode* Node)
+		auto IsSubgraphEditorOnly = [&bIsCurrentlyEditorOnly](UPCGNode* Node)
 		{
-			// Early out if we already know this is going to be editor only
-			if (bIsCurrentlyEditorOnly)
-			{
-				return;
-			}
-
-			if (UPCGBaseSubgraphNode* SubgraphNode = Cast<UPCGBaseSubgraphNode>(Node))
+			if (const UPCGBaseSubgraphNode* SubgraphNode = Cast<UPCGBaseSubgraphNode>(Node))
 			{
 				if (const UPCGGraph* Subgraph = SubgraphNode->GetSubgraph())
 				{
-					if (!VisitedGraphs.Contains(Subgraph))
+					if (Subgraph->IsEditorOnly_Internal())
 					{
-						bIsCurrentlyEditorOnly |= Subgraph->IsEditorOnly_Internal(VisitedGraphs);
+						bIsCurrentlyEditorOnly = true;
+						return false;
 					}
 				}
 			}
-		});
+
+			return true;
+		};
+
+		ForEachNodeRecursively(IsSubgraphEditorOnly);
 	}
 
 	return bIsCurrentlyEditorOnly;
 }
 
+bool UPCGGraph::IsEditorOnly_Internal() const
+{
+	return Super::IsEditorOnly() || bIsEditorOnly;
+}
+
+bool UPCGGraph::Contains(const UPCGGraph* InGraph) const
+{
+	bool bContains = false;
+	auto ContainsSelectedGraph = [InGraph, &bContains](UPCGNode* Node)
+	{
+		if (UPCGBaseSubgraphNode* SubgraphNode = Cast<UPCGBaseSubgraphNode>(Node))
+		{
+			if (InGraph == SubgraphNode->GetSubgraph())
+			{
+				bContains = true;
+				return false; // stop execution
+			}
+		}
+
+		return true;
+	};
+
+	ForEachNodeRecursively(ContainsSelectedGraph);
+	return bContains;
+}
+
 #if WITH_EDITOR
 void UPCGGraph::DeclareConstructClasses(TArray<FTopLevelAssetPath>& OutConstructClasses, const UClass* SpecificSubclass)
 {
@@ -751,15 +766,58 @@ bool UPCGGraph::RemoveEdge(UPCGNode* From, const FName& FromLabel, UPCGNode* To,
 	return TouchedNodes.Num() > 0;
 }
 
-void UPCGGraph::ForEachNode(const TFunction<void(UPCGNode*)>& Action) const
+bool UPCGGraph::ForEachNode(TFunctionRef<bool(UPCGNode*)> Action) const
 {
-	Action(InputNode);
-	Action(OutputNode);
+	if (!Action(InputNode) ||
+		!Action(OutputNode))
+	{
+		return false;
+	}
 
 	for (UPCGNode* Node : Nodes)
 	{
-		Action(Node);
+		if (!Action(Node))
+		{
+			return false;
+		}
 	}
+
+	return true;
+}
+
+bool UPCGGraph::ForEachNodeRecursively(TFunctionRef<bool(UPCGNode*)> Action) const
+{
+	TSet<const UPCGGraph*> VisitedGraphs;
+	return ForEachNodeRecursively_Internal(Action, VisitedGraphs);
+}
+
+bool UPCGGraph::ForEachNodeRecursively_Internal(TFunctionRef<bool(UPCGNode*)> Action, TSet<const UPCGGraph*>& VisitedGraphs) const
+{
+	check(!VisitedGraphs.Contains(this));
+	VisitedGraphs.Add(this);
+
+	auto RecursiveCall = [&Action, &VisitedGraphs](UPCGNode* Node) -> bool
+	{
+		if (!Action(Node))
+		{
+			return false;
+		}
+
+		if (UPCGBaseSubgraphNode* SubgraphNode = Cast<UPCGBaseSubgraphNode>(Node))
+		{
+			if (const UPCGGraph* Subgraph = SubgraphNode->GetSubgraph())
+			{
+				if (!VisitedGraphs.Contains(Subgraph))
+				{
+					return Subgraph->ForEachNodeRecursively_Internal(Action, VisitedGraphs);
+				}
+			}
+		}
+
+		return true;
+	};
+
+	return ForEachNode(RecursiveCall);
 }
 
 bool UPCGGraph::RemoveInboundEdges(UPCGNode* InNode, const FName& InboundLabel)
@@ -1170,6 +1228,7 @@ void UPCGGraph::FixInvalidEdges()
 	{
 		ValidatePins(InNode->GetInputPins(), /*bPinsAreInputs=*/true);
 		ValidatePins(InNode->GetOutputPins(), /*bPinsAreInputs=*/false);
+		return true;
 	});
 }
 
diff --git a/Engine/Plugins/Experimental/PCG/Source/PCG/Private/PCGSubgraph.cpp b/Engine/Plugins/Experimental/PCG/Source/PCG/Private/PCGSubgraph.cpp
--- a/Engine/Plugins/Experimental/PCG/Source/PCG/Private/PCGSubgraph.cpp
+++ b/Engine/Plugins/Experimental/PCG/Source/PCG/Private/PCGSubgraph.cpp
@@ -498,27 +498,41 @@ bool FPCGSubgraphElement::ExecuteInternal(FPCGContext* InContext) const
 	const UPCGSubgraphSettings* Settings = Context->GetInputSettings<UPCGSubgraphSettings>();
 	check(Settings);
 
-	// TODO: only prevents A->A recursion, not A->B->A or cycles in general
-	if (Context->SourceComponent->GetGraph() == Settings->GetSubgraph())
+	const bool bIsDynamic = Settings->IsDynamicGraph();
+
+	if (bIsDynamic && !Context->bScheduledSubgraph)
 	{
-		PCGE_LOG(Error, GraphAndLog, LOCTEXT("FailedRecursiveSubgraph", "PCGGraph cannot include itself as a subgraph, subgraph will not be executed."));
-		return true;
+		if (Settings->SubgraphInstance && Settings->OriginalSettings)
+		{
+			// If OriginalSettings is null, then we ARE the original settings, and writing over the existing graph is incorrect (and potentially a race condition)
+			Settings->SubgraphInstance->SetGraph(Settings->SubgraphOverride);
+		}
 	}
 
-	if (Settings->IsDynamicGraph())
+	UPCGGraph* Subgraph = Settings->GetSubgraph();
+
+	// Implementation note: recursivity test here must be consequential with the way the compilation has been done,
+	// otherwise the other tasks will not behave as expected.
+	// If the current graph is present in the subgraph downstream, then this must be a dynamic graph execution
+	bool bIsRecursive = false;
+	if (Subgraph)
+	{
+		if (Context->Stack)
+		{
+			bIsRecursive = Subgraph->Contains(Context->Stack->GetGraphForCurrentFrame());
+		}
+		else if(Context->SourceComponent.Get())
+		{
+			bIsRecursive = Subgraph->Contains(Context->SourceComponent->GetGraph());
+		}
+	}
+
+	if (bIsDynamic || bIsRecursive)
 	{
 		TRACE_CPUPROFILER_EVENT_SCOPE(FPCGSubgraphElement::Execute);
 
 		if (!Context->bScheduledSubgraph)
 		{
-			if (Settings->SubgraphInstance)
-			{
-				// If OriginalSettings is null, then we ARE the original settings, and writing over the existing graph is incorrect (and potentially a race condition)
-				check(Settings->OriginalSettings);
-				Settings->SubgraphInstance->SetGraph(Settings->SubgraphOverride);
-			}
-
-			UPCGGraph* Subgraph = Settings->GetSubgraph();
 			UPCGSubsystem* Subsystem = Context->SourceComponent.IsValid() ? Context->SourceComponent->GetSubsystem() : nullptr;
 
 			if (Subsystem && Subgraph)
@@ -531,6 +545,12 @@ bool FPCGSubgraphElement::ExecuteInternal(FPCGContext* InContext) const
 				FPCGDataCollection SubgraphInputData;
 				PrepareSubgraphData(Settings, Context, Context->InputData, SubgraphInputData);
 
+				// At this point, if we're in a recursive context and we have no input, we must terminate execution
+				if (bIsRecursive && SubgraphInputData.TaggedData.IsEmpty())
+				{
+					return true;
+				}
+
 				// Prepare the invocation stack - which is the stack up to this node, and then this node
 				FPCGStack InvocationStack = ensure(Context->Stack) ? *Context->Stack : FPCGStack();
 				InvocationStack.GetStackFramesMutable().Emplace(Context->Node);
diff --git a/Engine/Plugins/Experimental/PCG/Source/PCG/Public/Graph/PCGStackContext.h b/Engine/Plugins/Experimental/PCG/Source/PCG/Public/Graph/PCGStackContext.h
--- a/Engine/Plugins/Experimental/PCG/Source/PCG/Public/Graph/PCGStackContext.h
+++ b/Engine/Plugins/Experimental/PCG/Source/PCG/Public/Graph/PCGStackContext.h
@@ -80,7 +80,7 @@ public:
 	const UPCGComponent* GetRootComponent() const;
 
 	/** First (top) graph frame in stack (or null if no graph frames present). */
-	const UPCGGraph* GetRootGraph() const;
+	const UPCGGraph* GetRootGraph(int32* OutRootFrameIndex = nullptr) const;
 
 	/** Returns true if this stack is the top level/root graph, rather than in a subgraph. */
 	bool IsCurrentFrameInRootGraph() const { return GetNumGraphLevels() == 1; }
diff --git a/Engine/Plugins/Experimental/PCG/Source/PCG/Public/PCGGraph.h b/Engine/Plugins/Experimental/PCG/Source/PCG/Public/PCGGraph.h
--- a/Engine/Plugins/Experimental/PCG/Source/PCG/Public/PCGGraph.h
+++ b/Engine/Plugins/Experimental/PCG/Source/PCG/Public/PCGGraph.h
@@ -264,13 +264,21 @@ public:
 	/** Creates an edge between two nodes/pins based on the labels. Returns true if the To node has removed other edges (happens with single pins) */
 	bool AddLabeledEdge(UPCGNode* From, const FName& InboundLabel, UPCGNode* To, const FName& OutboundLabel);
 
+	/** Returns true if the current graph contains directly the specified node. This does not query recursively (through subgraphs). */
 	bool Contains(UPCGNode* Node) const;
+
+	/** Returns true if the current graph contains a subgraph node using statically the specified graph, recursively. */
+	bool Contains(const UPCGGraph* InGraph) const;
+
 	const TArray<UPCGNode*>& GetNodes() const { return Nodes; }
 	void AddNode(UPCGNode* InNode);
 	void AddNodes(TArray<UPCGNode*>& InNodes);
 
-	/** Calls the lambda on every node in graph. */
-	void ForEachNode(const TFunction<void(UPCGNode*)>& Action) const;
+	/** Calls the lambda on every node in the graph or until the Action call returns false */
+	bool ForEachNode(TFunctionRef<bool(UPCGNode*)> Action) const;
+
+	/** Calls the lambda on every node (going through subgraphs too) or until the Action call returns false */
+	bool ForEachNodeRecursively(TFunctionRef<bool(UPCGNode*)> Action) const;
 
 	bool RemoveInboundEdges(UPCGNode* InNode, const FName& InboundLabel);
 	bool RemoveOutboundEdges(UPCGNode* InNode, const FName& OutboundLabel);
@@ -320,7 +328,9 @@ protected:
 	void RemoveNodes_Internal(TArrayView<UPCGNode*> InNodes);
 	void AddNodes_Internal(TArrayView<UPCGNode*> InNodes);
 
-	bool IsEditorOnly_Internal(TSet<const UPCGGraph*>& VisitedGraphs) const;
+	bool IsEditorOnly_Internal() const;
+
+	bool ForEachNodeRecursively_Internal(TFunctionRef<bool(UPCGNode*)> Action, TSet<const UPCGGraph*>& VisitedGraphs) const;
 
 	/** Calculates node grid size. Not thread safe, must be called within write lock. */
 	uint32 CalculateNodeGridSizeRecursive_Unsafe(const UPCGNode* InNode, uint32 InDefaultGridSize) const;
diff --git a/Engine/Plugins/Experimental/PCG/Source/PCGEditor/Private/PCGEditorGraphSchemaActions.cpp b/Engine/Plugins/Experimental/PCG/Source/PCGEditor/Private/PCGEditorGraphSchemaActions.cpp
--- a/Engine/Plugins/Experimental/PCG/Source/PCGEditor/Private/PCGEditorGraphSchemaActions.cpp
+++ b/Engine/Plugins/Experimental/PCG/Source/PCGEditor/Private/PCGEditorGraphSchemaActions.cpp
@@ -324,11 +324,6 @@ UEdGraphNode* FPCGEditorGraphSchemaAction_NewSubgraphElement::PerformAction(UEdG
 	}
 
 	UPCGGraph* Subgraph = CastChecked<UPCGGraph>(SubgraphObjectPath.TryLoad());
-	if (Subgraph == PCGGraph)
-	{
-		UE_LOG(LogPCGEditor, Error, TEXT("Invalid Subgraph"));
-		return nullptr;
-	}
 
 	const FScopedTransaction Transaction(*FPCGEditorCommon::ContextIdentifier, LOCTEXT("PCGEditorNewSubgraphElement", "PCG Editor: New Subgraph Element"), nullptr);
 	EditorGraph->Modify();
diff --git a/Engine/Plugins/Experimental/PCG/Source/PCGEditor/Private/SPCGEditorGraphDebugObjectTree.cpp b/Engine/Plugins/Experimental/PCG/Source/PCGEditor/Private/SPCGEditorGraphDebugObjectTree.cpp
--- a/Engine/Plugins/Experimental/PCG/Source/PCGEditor/Private/SPCGEditorGraphDebugObjectTree.cpp
+++ b/Engine/Plugins/Experimental/PCG/Source/PCGEditor/Private/SPCGEditorGraphDebugObjectTree.cpp
@@ -536,7 +536,8 @@ void SPCGEditorGraphDebugObjectTree::AddStacksToTree(const TArray<FPCGStack>& St
 			}
 		}
 
-		UPCGGraph* TopGraph = const_cast<UPCGGraph*>(Stack.GetRootGraph());
+		int32 TopGraphIndex = INDEX_NONE;
+		UPCGGraph* TopGraph = const_cast<UPCGGraph*>(Stack.GetRootGraph(&TopGraphIndex));
 		if (!TopGraph)
 		{
 			continue;
@@ -579,7 +580,7 @@ void SPCGEditorGraphDebugObjectTree::AddStacksToTree(const TArray<FPCGStack>& St
 				const bool bIsDebuggable = (GraphBeingEdited == StackGraph);
 
 				// Top graph.
-				if (StackGraph == TopGraph)
+				if (FrameIndex == TopGraphIndex && StackGraph == TopGraph)
 				{
 					FPCGStack GraphStack = Stack;
 					GraphStack.GetStackFramesMutable().SetNum(FrameIndex + 1);
