[PCG] Revamp Sort Points into Sort Attributes

* Renamed Sort points to Sort Attributes.
* Now supports any pcg data that supports metadata.
* Add functions in Metadata to re-order entries on copy.
* Need to dance for default arrays values in UFunctions, hence the need of K2_* new functions.
* Add redirectors for impacted class/functions.
* Add bulk getter/setters for ValueKeys.
* Bubble up the param Initialize Metadata on DuplicateData to PCGData (was only for spatial before).
* Also remove the necessity of passing an empty array to AddToCollection in BP.

diff --git a/Engine/Plugins/Experimental/PCG/Config/BasePCG.ini b/Engine/Plugins/Experimental/PCG/Config/BasePCG.ini
--- a/Engine/Plugins/Experimental/PCG/Config/BasePCG.ini
+++ b/Engine/Plugins/Experimental/PCG/Config/BasePCG.ini
@@ -12,6 +12,8 @@
 +FunctionRedirects=(OldName="PCGDataFunctionLibrary.GetTaggedInputs",NewName="PCGDataFunctionLibrary.GetInputsByTag")
 +FunctionRedirects=(OldName="PCGDataFunctionLibrary.GetParamsByPin",NewName="PCGDataFunctionLibrary.GetParamsByPinLabel")
 +FunctionRedirects=(OldName="PCGDataFunctionLibrary.GetTaggedParams",NewName="PCGDataFunctionLibrary.GetParamsByTag")
++FunctionRedirects=(OldName="PCGMetadata.InitializeAsCopy",NewName="PCGMetadata.K2_InitializeAsCopy")
++FunctionRedirects=(OldName="PCGMetadata.InitializeAsCopyWithAttributeFilter",NewName="PCGMetadata.K2_InitializeAsCopyWithAttributeFilter")
 
 +PropertyRedirects=(OldName="PCGSplineSamplerSettings.Params",NewName="PCGSplineSamplerSettings.SamplerParams")
 +PropertyRedirects=(OldName="PCGProjectionSettings.Params",NewName="PCGProjectionSettings.ProjectionParams")
@@ -68,6 +70,7 @@
 +ClassRedirects=(OldName="PCGAttributeFilterSettings",NewName="/Script/PCG.PCGDeleteAttributesSettings")
 +ClassRedirects=(OldName="PCGAttributeFilterNamesSettings",NewName="/Script/PCG.PCGDeleteAttributesSettings")
 +ClassRedirects=(OldName="PCGPropertyToParamDataSettings",NewName="/Script/PCG.PCGGetActorPropertySettings")
++ClassRedirects=(OldName="PCGSortPointsSettings",NewName="/Script/PCG.PCGSortAttributesSettings")
 
 +StructRedirects=(OldName="PCGPointFilterThresholdSettings",NewName="/Script/PCG.PCGAttributeFilterThresholdSettings")
 
diff --git a/Engine/Plugins/Experimental/PCG/Source/PCG/Private/Data/PCGSpatialData.cpp b/Engine/Plugins/Experimental/PCG/Source/PCG/Private/Data/PCGSpatialData.cpp
--- a/Engine/Plugins/Experimental/PCG/Source/PCG/Private/Data/PCGSpatialData.cpp
+++ b/Engine/Plugins/Experimental/PCG/Source/PCG/Private/Data/PCGSpatialData.cpp
@@ -219,7 +219,7 @@ void UPCGSpatialData::InitializeFromData(const UPCGSpatialData* InSource, const
 	}
 }
 
-UPCGSpatialData* UPCGSpatialData::DuplicateData(const bool bInitializeMetadata) const
+UPCGSpatialData* UPCGSpatialData::DuplicateData(bool bInitializeMetadata) const
 {
 	UPCGSpatialData* NewSpatialData = CopyInternal();
 	check(NewSpatialData);
diff --git a/Engine/Plugins/Experimental/PCG/Source/PCG/Private/Elements/PCGSortAttributes.cpp b/Engine/Plugins/Experimental/PCG/Source/PCG/Private/Elements/PCGSortAttributes.cpp
new file mode 100644
--- /dev/null
+++ b/Engine/Plugins/Experimental/PCG/Source/PCG/Private/Elements/PCGSortAttributes.cpp
@@ -0,0 +1,127 @@
+// Copyright Epic Games, Inc. All Rights Reserved.
+
+#include "Elements/PCGSortAttributes.h"
+
+#include "PCGContext.h"
+#include "Data/PCGPointData.h"
+#include "Metadata/Accessors/IPCGAttributeAccessor.h"
+#include "Metadata/Accessors/PCGAttributeAccessorHelpers.h"
+#include "Metadata/Accessors/PCGAttributeAccessorKeys.h"
+
+#define LOCTEXT_NAMESPACE "PCGSortAttributesElement"
+
+#if WITH_EDITOR
+FName UPCGSortAttributesSettings::GetDefaultNodeName() const
+{
+	return FName(TEXT("SortAttributes"));
+}
+
+FText UPCGSortAttributesSettings::GetDefaultNodeTitle() const
+{
+	return LOCTEXT("NodeTitle", "Sort Attributes");
+}
+
+FText UPCGSortAttributesSettings::GetNodeTooltipText() const
+{
+	return LOCTEXT("NodeTooltip", "Sorts data based on an attribute.");
+}
+
+TArray<FText> UPCGSortAttributesSettings::GetNodeTitleAliases() const
+{
+	// Re-use old LOCTEXT name if it was already localized.
+	return { NSLOCTEXT("PCGSortPointsElement", "NodeTitle", "Sort Points") };
+}
+#endif // WITH_EDITOR
+
+TArray<FPCGPinProperties> UPCGSortAttributesSettings::OutputPinProperties() const
+{
+	TArray<FPCGPinProperties> PinProperties;
+	PinProperties.Emplace(PCGPinConstants::DefaultOutputLabel, EPCGDataType::Any);
+
+	return PinProperties;
+}
+
+FPCGElementPtr UPCGSortAttributesSettings::CreateElement() const
+{
+	return MakeShared<FPCGSortAttributesElement>();
+}
+
+bool FPCGSortAttributesElement::ExecuteInternal(FPCGContext* Context) const
+{
+	TRACE_CPUPROFILER_EVENT_SCOPE(FPCGSortAttributesElement::Execute);
+	
+	check(Context);
+
+	const UPCGSortAttributesSettings* Settings = Context->GetInputSettings<UPCGSortAttributesSettings>();
+	check(Settings);
+	
+	TArray<FPCGTaggedData> Inputs = Context->InputData.GetInputsByPin(PCGPinConstants::DefaultInputLabel);
+	TArray<FPCGTaggedData>& Outputs = Context->OutputData.TaggedData;
+
+	for (int32 i = 0; i < Inputs.Num(); ++i)
+	{
+		TRACE_CPUPROFILER_EVENT_SCOPE(FPCGSortAttributesElement::InputLoop);
+
+		const FPCGTaggedData& Input = Inputs[i];
+
+		if (const UPCGData* InputData = Input.Data)
+		{
+			FPCGAttributePropertyInputSelector InputSource = Settings->InputSource.CopyAndFixLast(InputData);
+			TUniquePtr<const IPCGAttributeAccessor> Accessor = PCGAttributeAccessorHelpers::CreateConstAccessor(InputData, InputSource);
+			TUniquePtr<const IPCGAttributeAccessorKeys> Keys = PCGAttributeAccessorHelpers::CreateConstKeys(InputData, InputSource);
+
+			if (!Accessor.IsValid() || !Keys.IsValid())
+			{
+				PCGE_LOG(Error, GraphAndLog, FText::Format(LOCTEXT("InvalidAccessor", "Attribute '{0}' does not exist for input {1}."), InputSource.GetDisplayText(), FText::AsNumber(i)));
+				continue;
+			}
+
+			if (!PCGMetadataAttribute::CallbackWithRightType(Accessor->GetUnderlyingType(), [](auto Dummy) -> bool { return PCG::Private::MetadataTraits<decltype(Dummy)>::CanCompare; }))
+			{
+				PCGE_LOG(Error, GraphAndLog, FText::Format(LOCTEXT("InvalidAccessor", "Attribute '{0}' exists but is not of a comparable type ({1}) for input {2}."), InputSource.GetDisplayText(), PCG::Private::GetTypeNameText(Accessor->GetUnderlyingType()), FText::AsNumber(i)));
+				continue;
+			}
+
+			UPCGData* OutputData = nullptr;
+			if (const UPCGPointData* InputPointData = Cast<const UPCGPointData>(InputData))
+			{
+				UPCGPointData* OutputPointData = static_cast<UPCGPointData*>(InputData->DuplicateData());
+				PCGAttributeAccessorHelpers::SortByAttribute(*Accessor, *Keys, OutputPointData->GetMutablePoints(), Settings->SortMethod == EPCGSortMethod::Ascending);
+				OutputData = OutputPointData;
+			}
+			else if (const UPCGMetadata* InMetadata = InputData->ConstMetadata())
+			{
+				// Duplicate data without metadata
+				OutputData = InputData->DuplicateData(/*bInitializeMetdata=*/false);
+				UPCGMetadata* OutMetadata = OutputData->MutableMetadata();
+
+				if (!ensure(OutMetadata))
+				{
+					continue;
+				}
+
+				// Gather all the entries and sort them
+				TArray<PCGMetadataEntryKey> Entries;
+				const int64 Count = InMetadata->GetItemCountForChild();
+
+				Entries.Reserve(Count);
+
+				for (PCGMetadataEntryKey Entry = 0; Entry < Count; ++Entry)
+				{
+					Entries.Add(Entry);
+				}
+
+				PCGAttributeAccessorHelpers::SortByAttribute(*Accessor, *Keys, Entries, Settings->SortMethod == EPCGSortMethod::Ascending);
+
+				OutMetadata->InitializeAsCopy(InputData->ConstMetadata(), &Entries);
+			}
+
+			FPCGTaggedData& Output = Outputs.Add_GetRef(Input);
+			Output.Data = OutputData;
+		}
+	}
+
+	return true;
+}
+
+#undef LOCTEXT_NAMESPACE
diff --git a/Engine/Plugins/Experimental/PCG/Source/PCG/Private/Elements/PCGSortPoints.cpp b/Engine/Plugins/Experimental/PCG/Source/PCG/Private/Elements/PCGSortPoints.cpp
deleted file mode 100644
diff --git a/Engine/Plugins/Experimental/PCG/Source/PCG/Private/Metadata/PCGMetadata.cpp b/Engine/Plugins/Experimental/PCG/Source/PCG/Private/Metadata/PCGMetadata.cpp
--- a/Engine/Plugins/Experimental/PCG/Source/PCG/Private/Metadata/PCGMetadata.cpp
+++ b/Engine/Plugins/Experimental/PCG/Source/PCG/Private/Metadata/PCGMetadata.cpp
@@ -105,12 +105,12 @@ void UPCGMetadata::InitializeWithAttributeFilter(const UPCGMetadata* InParent, c
 	}
 }
 
-void UPCGMetadata::InitializeAsCopy(const UPCGMetadata* InMetadataToCopy)
+void UPCGMetadata::InitializeAsCopy(const UPCGMetadata* InMetadataToCopy, const TArray<PCGMetadataEntryKey>* EntriesToCopy)
 {
-	InitializeAsCopyWithAttributeFilter(InMetadataToCopy, TSet<FName>(), EPCGMetadataFilterMode::ExcludeAttributes);
+	InitializeAsCopyWithAttributeFilter(InMetadataToCopy, TSet<FName>(), EPCGMetadataFilterMode::ExcludeAttributes, EntriesToCopy);
 }
 
-void UPCGMetadata::InitializeAsCopyWithAttributeFilter(const UPCGMetadata* InMetadataToCopy, const TSet<FName>& InFilteredAttributes, EPCGMetadataFilterMode InFilterMode)
+void UPCGMetadata::InitializeAsCopyWithAttributeFilter(const UPCGMetadata* InMetadataToCopy, const TSet<FName>& InFilteredAttributes, EPCGMetadataFilterMode InFilterMode, const TArray<PCGMetadataEntryKey>* EntriesToCopy)
 {
 	if (!InMetadataToCopy)
 	{
@@ -131,6 +131,20 @@ void UPCGMetadata::InitializeAsCopyWithAttributeFilter(const UPCGMetadata* InMet
 
 	const bool bSkipAttributesInFilterList = (InFilterMode == EPCGMetadataFilterMode::ExcludeAttributes);
 
+	const bool bPartialCopy = EntriesToCopy && EntriesToCopy->Num() <= InMetadataToCopy->GetItemCountForChild();
+	TArray<PCGMetadataEntryKey> NewEntryKeys;
+	TArray<PCGMetadataValueKey> NewValueKeys;
+	if (bPartialCopy)
+	{
+		const int32 Count = EntriesToCopy->Num();
+		NewEntryKeys.SetNumUninitialized(Count);
+		NewValueKeys.SetNumUninitialized(Count);
+		for (int32 j = 0; j < Count; ++j)
+		{
+			NewEntryKeys[j] = PCGMetadataEntryKey(j);
+		}
+	}
+
 	// Copy attributes
 	for (const TPair<FName, FPCGMetadataAttributeBase*>& OtherAttribute : InMetadataToCopy->Attributes)
 	{
@@ -139,7 +153,14 @@ void UPCGMetadata::InitializeAsCopyWithAttributeFilter(const UPCGMetadata* InMet
 
 		if (!bSkipThisAttribute)
 		{
-			CopyAttribute(OtherAttribute.Value, OtherAttribute.Key, /*bKeepParent=*/false, /*bCopyEntries=*/true, /*bCopyValues=*/true);
+			// Don't copy entries if we have a partial copy, we will set them all after.
+			FPCGMetadataAttributeBase* Attribute = CopyAttribute(OtherAttribute.Value, OtherAttribute.Key, /*bKeepParent=*/false, /*bCopyEntries=*/!bPartialCopy, /*bCopyValues=*/true);
+
+			if (bPartialCopy && OtherAttribute.Value && Attribute)
+			{
+				OtherAttribute.Value->GetValueKeys(*EntriesToCopy, NewValueKeys);
+				Attribute->SetValuesFromValueKeys(NewEntryKeys, NewValueKeys);
+			}
 		}
 	}
 }
diff --git a/Engine/Plugins/Experimental/PCG/Source/PCG/Private/Metadata/PCGMetadataAttribute.cpp b/Engine/Plugins/Experimental/PCG/Source/PCG/Private/Metadata/PCGMetadataAttribute.cpp
--- a/Engine/Plugins/Experimental/PCG/Source/PCG/Private/Metadata/PCGMetadataAttribute.cpp
+++ b/Engine/Plugins/Experimental/PCG/Source/PCG/Private/Metadata/PCGMetadataAttribute.cpp
@@ -41,6 +41,22 @@ void FPCGMetadataAttributeBase::SetValueFromValueKey(PCGMetadataEntryKey EntryKe
 	EntryToValueKeyMap.FindOrAdd(EntryKey) = ValueKey;
 }
 
+void FPCGMetadataAttributeBase::SetValueFromValueKey_Unsafe(PCGMetadataEntryKey EntryKey, PCGMetadataValueKey ValueKey, bool bResetValueOnDefaultValueKey)
+{
+	check(EntryKey != PCGInvalidEntryKey);
+	if (EntryKey == PCGDefaultValueKey)
+	{
+		if (bResetValueOnDefaultValueKey)
+		{
+			EntryToValueKeyMap.Remove(EntryKey);
+		}
+	}
+	else
+	{
+		EntryToValueKeyMap.FindOrAdd(EntryKey) = ValueKey;
+	}
+}
+
 void FPCGMetadataAttributeBase::SetValuesFromValueKeys(const TArray<TTuple<PCGMetadataEntryKey, PCGMetadataValueKey>>& EntryValuePairs, bool bResetValueOnDefaultValueKey)
 {
 	if (EntryValuePairs.IsEmpty())
@@ -51,18 +67,21 @@ void FPCGMetadataAttributeBase::SetValuesFromValueKeys(const TArray<TTuple<PCGMe
 	FWriteScopeLock ScopeLock(EntryMapLock);
 	for (const TTuple<PCGMetadataEntryKey, PCGMetadataValueKey>& EntryValuePair : EntryValuePairs)
 	{
-		check(EntryValuePair.Key != PCGInvalidEntryKey);
-		if (EntryValuePair.Value == PCGDefaultValueKey)
-		{
-			if (bResetValueOnDefaultValueKey)
-			{
-				EntryToValueKeyMap.Remove(EntryValuePair.Key);
-			}
-		}
-		else
-		{
-			EntryToValueKeyMap.FindOrAdd(EntryValuePair.Key, EntryValuePair.Value);
-		}
+		SetValueFromValueKey_Unsafe(EntryValuePair.Key, EntryValuePair.Value, bResetValueOnDefaultValueKey);
+	}
+}
+
+void FPCGMetadataAttributeBase::SetValuesFromValueKeys(const TArray<PCGMetadataEntryKey>& EntryKeys, TArray<PCGMetadataValueKey>& ValueKeys, bool bResetValueOnDefaultValueKey)
+{
+	if (EntryKeys.IsEmpty() || EntryKeys.Num() != ValueKeys.Num())
+	{
+		return;
+	}
+
+	FWriteScopeLock ScopeLock(EntryMapLock);
+	for (int32 i = 0; i < EntryKeys.Num(); ++i)
+	{
+		SetValueFromValueKey_Unsafe(EntryKeys[i], ValueKeys[i], bResetValueOnDefaultValueKey);
 	}
 }
 
@@ -94,6 +113,72 @@ PCGMetadataValueKey FPCGMetadataAttributeBase::GetValueKey(PCGMetadataEntryKey E
 	}
 }
 
+void FPCGMetadataAttributeBase::GetValueKeys(const TArray<PCGMetadataEntryKey>& EntryKeys, TArray<PCGMetadataValueKey>& OutValueKeys) const
+{
+	if (EntryKeys.IsEmpty())
+	{
+		return;
+	}
+
+	OutValueKeys.SetNumUninitialized(EntryKeys.Num());
+	// Bitset with all unset values. If we have any unset value, we will ask the parent for those.
+	TBitArray<> UnsetValues(true, EntryKeys.Num());
+
+	GetValueKeys_Internal(EntryKeys, OutValueKeys, UnsetValues);
+}
+
+void FPCGMetadataAttributeBase::GetValueKeys_Internal(const TArray<PCGMetadataEntryKey>& EntryKeys, TArray<PCGMetadataValueKey>& OutValueKeys, TBitArray<>& UnsetValues) const
+{
+	check(EntryKeys.Num() == OutValueKeys.Num() && OutValueKeys.Num() == UnsetValues.Num());
+
+	bool bFoundAllKeys = true;
+	TConstSetBitIterator<> It(UnsetValues);
+	if (!It)
+	{
+		return;
+	}
+
+	EntryMapLock.ReadLock();
+
+	for (; It; ++It)
+	{
+		const int32 Index = It.GetIndex();
+		const PCGMetadataEntryKey EntryKey = EntryKeys[Index];
+
+		auto SetValueKey = [Index, &OutValueKeys, &UnsetValues](PCGMetadataValueKey ValueKey)
+		{
+			OutValueKeys[Index] = ValueKey;
+			UnsetValues[Index] = false;
+		};
+		
+		if (EntryKey == PCGInvalidEntryKey)
+		{
+			SetValueKey(PCGDefaultValueKey);
+		}
+		else if (const PCGMetadataValueKey* FoundLocalKey = EntryToValueKeyMap.Find(EntryKey))
+		{
+			SetValueKey(*FoundLocalKey);
+		}
+		else if (!Parent)
+		{
+			SetValueKey(PCGDefaultValueKey);
+		}
+		else
+		{
+			bFoundAllKeys = false;
+		}
+	}
+
+	EntryMapLock.ReadUnlock();
+
+	ensure(Parent || bFoundAllKeys);
+
+	if (Parent && !bFoundAllKeys)
+	{
+		Parent->GetValueKeys_Internal(EntryKeys, OutValueKeys, UnsetValues);
+	}
+}
+
 bool FPCGMetadataAttributeBase::HasNonDefaultValue(PCGMetadataEntryKey EntryKey) const
 {
 	return GetValueKey(EntryKey) != PCGDefaultValueKey;
diff --git a/Engine/Plugins/Experimental/PCG/Source/PCG/Private/PCGData.cpp b/Engine/Plugins/Experimental/PCG/Source/PCG/Private/PCGData.cpp
--- a/Engine/Plugins/Experimental/PCG/Source/PCG/Private/PCGData.cpp
+++ b/Engine/Plugins/Experimental/PCG/Source/PCG/Private/PCGData.cpp
@@ -146,7 +146,7 @@ void UPCGData::Flatten()
 	}
 }
 
-UPCGData* UPCGData::DuplicateData() const
+UPCGData* UPCGData::DuplicateData(bool bInitializeMetadata) const
 {
 	return Cast<UPCGData>(StaticDuplicateObject(this, GetTransientPackage()));
 }
diff --git a/Engine/Plugins/Experimental/PCG/Source/PCG/Private/PCGParamData.cpp b/Engine/Plugins/Experimental/PCG/Source/PCG/Private/PCGParamData.cpp
--- a/Engine/Plugins/Experimental/PCG/Source/PCG/Private/PCGParamData.cpp
+++ b/Engine/Plugins/Experimental/PCG/Source/PCG/Private/PCGParamData.cpp
@@ -163,10 +163,14 @@ void UPCGParamData::SetLastSelector(const FPCGAttributePropertySelector& InSelec
 	CachedLastSelector.ImportFromOtherSelector(InSelector);
 }
 
-UPCGParamData* UPCGParamData::DuplicateData() const
+UPCGParamData* UPCGParamData::DuplicateData(bool bInitializeMetadata) const
 {
 	UPCGParamData* NewParamData = NewObject<UPCGParamData>();
-	NewParamData->Metadata->InitializeAsCopy(Metadata);
+	if (bInitializeMetadata)
+	{
+		check(NewParamData && NewParamData->Metadata);
+		NewParamData->Metadata->InitializeAsCopy(Metadata);
+	}
 
 	return NewParamData;
 }
\ No newline at end of file
diff --git a/Engine/Plugins/Experimental/PCG/Source/PCG/Private/Tests/Elements/PCGSortAttributesTest.cpp b/Engine/Plugins/Experimental/PCG/Source/PCG/Private/Tests/Elements/PCGSortAttributesTest.cpp
new file mode 100644
--- /dev/null
+++ b/Engine/Plugins/Experimental/PCG/Source/PCG/Private/Tests/Elements/PCGSortAttributesTest.cpp
@@ -0,0 +1,249 @@
+// Copyright Epic Games, Inc. All Rights Reserved.
+
+#include "Tests/PCGTestsCommon.h"
+
+#include "PCGComponent.h"
+#include "PCGContext.h"
+#include "PCGParamData.h"
+#include "Data/PCGPointData.h"
+#include "Elements/PCGSortAttributes.h"
+#include "Metadata/PCGMetadata.h"
+#include "Metadata/PCGMetadataAttributeTpl.h"
+
+namespace SortCommonTestData
+{
+	static const FName AttributeName = TEXT("Attr");
+
+	UPCGParamData* GenerateParamData(int32 EntriesCount, int32 Seed, bool bRandom)
+	{
+		UPCGParamData* OutData = NewObject<UPCGParamData>();
+		check(OutData && OutData->Metadata);
+		FPCGMetadataAttribute<int32>* Attribute = OutData->Metadata->CreateAttribute<int32>(AttributeName, 0, false, false);
+		check(Attribute);
+		FRandomStream RandomStream(Seed);
+
+		for (int32 i = 0; i < EntriesCount; ++i)
+		{
+			Attribute->SetValueFromValueKey(OutData->Metadata->AddEntry(), bRandom ? RandomStream.RandRange(1, 9999) : 1);
+		}
+
+		return OutData;
+	}
+
+	TUniquePtr<FPCGContext> RunSortElementOnData(EPCGSortMethod Method, const UPCGData* InData, const FPCGAttributePropertyInputSelector& InputSource)
+	{
+		PCGTestsCommon::FTestData TestData;
+		UPCGSortAttributesSettings* Settings = PCGTestsCommon::GenerateSettings<UPCGSortAttributesSettings>(TestData);
+		check(Settings);
+		Settings->InputSource = InputSource;
+		Settings->SortMethod = Method;
+		FPCGElementPtr TestElement = TestData.Settings->GetElement();
+
+		FPCGTaggedData& SourcePin = TestData.InputData.TaggedData.Emplace_GetRef();
+		SourcePin.Pin = PCGPinConstants::DefaultInputLabel;
+		SourcePin.Data = InData;
+
+		TUniquePtr<FPCGContext> Context = TestData.InitializeTestContext();
+
+		while (!TestElement->Execute(Context.Get())) {}
+	
+		return Context;
+	}
+}
+
+IMPLEMENT_CUSTOM_SIMPLE_AUTOMATION_TEST(FPCGSortPointsTest_Ascending, FPCGTestBaseClass, "Plugins.PCG.SortAttributes.Points.Ascending", PCGTestsCommon::TestFlags)
+
+bool FPCGSortPointsTest_Ascending::RunTest(const FString& Parameters)
+{
+	TUniquePtr<FPCGContext> Context = SortCommonTestData::RunSortElementOnData(EPCGSortMethod::Ascending,
+		PCGTestsCommon::CreateRandomPointData(100, 42, /*bRandomDensity=*/ true), 
+		FPCGAttributePropertySelector::CreatePointPropertySelector<FPCGAttributePropertyInputSelector>(EPCGPointProperties::Density));
+
+	const TArray<FPCGTaggedData>& Outputs = Context->OutputData.TaggedData;
+
+	UTEST_EQUAL("Output count", Outputs.Num(), 1);
+
+	// compare point data
+	const FPCGTaggedData& Output = Outputs[0];
+	const UPCGPointData* OutPointData = Cast<UPCGPointData>(Output.Data);
+	UTEST_NOT_NULL("Output is a point data", OutPointData);
+	check(OutPointData);
+
+	const TArray<FPCGPoint>& OutPoints = OutPointData->GetPoints();
+
+	for (int i = 0; i < OutPoints.Num() - 1; ++i)
+	{
+		UTEST_TRUE(FString::Format(TEXT("{0} is less than/equal to {1}"), { OutPoints[i].Density, OutPoints[i + 1].Density }), OutPoints[i].Density <= OutPoints[i + 1].Density);
+	}
+
+	
+	return true;
+}
+
+IMPLEMENT_CUSTOM_SIMPLE_AUTOMATION_TEST(FPCGSortAttributesTest_Ascending, FPCGTestBaseClass, "Plugins.PCG.SortAttributes.AttributeSet.Ascending", PCGTestsCommon::TestFlags)
+
+bool FPCGSortAttributesTest_Ascending::RunTest(const FString& Parameters)
+{
+	TUniquePtr<FPCGContext> Context = SortCommonTestData::RunSortElementOnData(EPCGSortMethod::Ascending,
+		SortCommonTestData::GenerateParamData(100, 42, /*bRandom=*/true),
+		FPCGAttributePropertySelector::CreateAttributeSelector<FPCGAttributePropertyInputSelector>(SortCommonTestData::AttributeName));
+
+	const TArray<FPCGTaggedData>& Outputs = Context->OutputData.TaggedData;
+
+	UTEST_EQUAL("Output count", Outputs.Num(), 1);
+
+	const FPCGTaggedData& Output = Outputs[0];
+	const UPCGParamData* OutParamData = Cast<UPCGParamData>(Output.Data);
+	UTEST_NOT_NULL("Output is an attribute set data", OutParamData);
+	check(OutParamData && OutParamData->Metadata);
+
+	const FPCGMetadataAttribute<int32>* Attribute = OutParamData->Metadata->GetConstTypedAttribute<int32>(SortCommonTestData::AttributeName);
+	UTEST_NOT_NULL("Output has the sorted attribute", Attribute);
+	check(Attribute);
+
+	for (int i = 0; i < OutParamData->Metadata->GetLocalItemCount() - 1; ++i)
+	{
+		const int32 Value = Attribute->GetValue(i);
+		const int32 NextValue = Attribute->GetValue(i + 1);
+		UTEST_TRUE(FString::Format(TEXT("{0} is less than/equal to {1}"), { Value, NextValue }), Value <= NextValue);
+	}
+
+	return true;
+}
+
+IMPLEMENT_CUSTOM_SIMPLE_AUTOMATION_TEST(FPCGSortPointsTest_Descending, FPCGTestBaseClass, "Plugins.PCG.SortAttributes.Points.Descending", PCGTestsCommon::TestFlags)
+
+bool FPCGSortPointsTest_Descending::RunTest(const FString& Parameters)
+{
+	TUniquePtr<FPCGContext> Context = SortCommonTestData::RunSortElementOnData(EPCGSortMethod::Descending,
+		PCGTestsCommon::CreateRandomPointData(100, 42, /*bRandomDensity=*/ true),
+		FPCGAttributePropertySelector::CreatePointPropertySelector<FPCGAttributePropertyInputSelector>(EPCGPointProperties::Density));
+
+	const TArray<FPCGTaggedData>& Outputs = Context->OutputData.TaggedData;
+
+	UTEST_EQUAL("Output count", Outputs.Num(), 1);
+
+	// compare point data
+	const FPCGTaggedData& Output = Outputs[0];
+	const UPCGPointData* OutPointData = Cast<UPCGPointData>(Output.Data);
+
+	UTEST_NOT_NULL("Output is a point data", OutPointData);
+	check(OutPointData);
+
+	const TArray<FPCGPoint>& OutPoints = OutPointData->GetPoints();
+
+	for (int i = 0; i < OutPoints.Num() - 1; ++i)
+	{
+		UTEST_TRUE(FString::Format(TEXT("{0} is greater than/equal to {1}"), { OutPoints[i].Density, OutPoints[i + 1].Density }), OutPoints[i].Density >= OutPoints[i + 1].Density);
+	}
+	
+	return true;
+}
+
+IMPLEMENT_CUSTOM_SIMPLE_AUTOMATION_TEST(FPCGSortAttributesTest_Descending, FPCGTestBaseClass, "Plugins.PCG.SortAttributes.AttributeSet.Descending", PCGTestsCommon::TestFlags)
+
+bool FPCGSortAttributesTest_Descending::RunTest(const FString& Parameters)
+{
+	TUniquePtr<FPCGContext> Context = SortCommonTestData::RunSortElementOnData(EPCGSortMethod::Descending,
+		SortCommonTestData::GenerateParamData(100, 42, /*bRandom=*/true),
+		FPCGAttributePropertySelector::CreateAttributeSelector<FPCGAttributePropertyInputSelector>(SortCommonTestData::AttributeName));
+
+	const TArray<FPCGTaggedData>& Outputs = Context->OutputData.TaggedData;
+
+	UTEST_EQUAL("Output count", Outputs.Num(), 1);
+
+	const FPCGTaggedData& Output = Outputs[0];
+	const UPCGParamData* OutParamData = Cast<UPCGParamData>(Output.Data);
+	UTEST_NOT_NULL("Output is an attribute set data", OutParamData);
+	check(OutParamData && OutParamData->Metadata);
+
+	const FPCGMetadataAttribute<int32>* Attribute = OutParamData->Metadata->GetConstTypedAttribute<int32>(SortCommonTestData::AttributeName);
+	UTEST_NOT_NULL("Output has the sorted attribute", Attribute);
+	check(Attribute);
+
+	for (int i = 0; i < OutParamData->Metadata->GetLocalItemCount() - 1; ++i)
+	{
+		const int32 Value = Attribute->GetValue(i);
+		const int32 NextValue = Attribute->GetValue(i + 1);
+		UTEST_TRUE(FString::Format(TEXT("{0} is greater than/equal to {1}"), { Value, NextValue }), Value >= NextValue);
+	}
+
+	return true;
+}
+
+IMPLEMENT_CUSTOM_SIMPLE_AUTOMATION_TEST(FPCGSortPointsTest_SameValues, FPCGTestBaseClass, "Plugins.PCG.SortAttributes.SameValues", PCGTestsCommon::TestFlags)
+
+bool FPCGSortPointsTest_SameValues::RunTest(const FString& Parameters)
+{
+	TUniquePtr<FPCGContext> Context = SortCommonTestData::RunSortElementOnData(EPCGSortMethod::Ascending,
+		PCGTestsCommon::CreateRandomPointData(100, 42, /*bRandomDensity=*/ false),
+		FPCGAttributePropertySelector::CreatePointPropertySelector<FPCGAttributePropertyInputSelector>(EPCGPointProperties::Density));
+
+	const TArray<FPCGTaggedData>& Inputs = Context->InputData.GetInputsByPin(PCGPinConstants::DefaultInputLabel);
+	const TArray<FPCGTaggedData>& Outputs = Context->OutputData.TaggedData;
+
+	UTEST_EQUAL("Output count", Outputs.Num(), 1);
+
+	// compare point data
+	const FPCGTaggedData& Input = Inputs[0];
+	const FPCGTaggedData& Output = Outputs[0];
+
+	const UPCGPointData* InPointData = Cast<UPCGPointData>(Input.Data);
+	const UPCGPointData* OutPointData = Cast<UPCGPointData>(Output.Data);
+
+	UTEST_NOT_NULL("Output is a point data", OutPointData);
+	check(InPointData && OutPointData);
+
+	const TArray<FPCGPoint>& InPoints = InPointData->GetPoints();
+	const TArray<FPCGPoint>& OutPoints = OutPointData->GetPoints();
+
+	UTEST_EQUAL("Arrays have the same number of points:", InPoints.Num(), OutPoints.Num());
+
+	for (int i = 0; i < InPoints.Num(); ++i)
+	{
+		//if they're in the same spots after sorting, they should have exactly the same properties across the board
+		UTEST_EQUAL(FString::Format(TEXT("UnsortedArray[{0}].Seed is equal to SortedArray[{0}].Seed"), {i}), OutPoints[i].Seed, InPoints[i].Seed);
+		UTEST_EQUAL(FString::Format(TEXT("UnsortedArray[{0}].Density is equal to SortedArray[{0}].Density"), {i}), OutPoints[i].Density, InPoints[i].Density);
+		UTEST_EQUAL(FString::Format(TEXT("UnsortedArray[{0}].Transform is equal to SortedArray[{0}].Transform"), {i}), OutPoints[i].Transform, InPoints[i].Transform);
+	}
+
+	return true;
+}
+
+IMPLEMENT_CUSTOM_SIMPLE_AUTOMATION_TEST(FPCGSortAttributesTest_SameValues, FPCGTestBaseClass, "Plugins.PCG.SortAttributes.AttributeSet.SameValues", PCGTestsCommon::TestFlags)
+
+bool FPCGSortAttributesTest_SameValues::RunTest(const FString& Parameters)
+{
+	const UPCGParamData* InputParamData = SortCommonTestData::GenerateParamData(100, 42, /*bRandom=*/false);
+	check(InputParamData);
+
+	TUniquePtr<FPCGContext> Context = SortCommonTestData::RunSortElementOnData(EPCGSortMethod::Ascending,
+		InputParamData,
+		FPCGAttributePropertySelector::CreateAttributeSelector<FPCGAttributePropertyInputSelector>(SortCommonTestData::AttributeName));
+
+	const TArray<FPCGTaggedData>& Outputs = Context->OutputData.TaggedData;
+
+	UTEST_EQUAL("Output count", Outputs.Num(), 1);
+
+	const FPCGTaggedData& Output = Outputs[0];
+	const UPCGParamData* OutParamData = Cast<UPCGParamData>(Output.Data);
+	UTEST_NOT_NULL("Output is an attribute set data", OutParamData);
+	check(OutParamData && OutParamData->Metadata);
+
+
+	const FPCGMetadataAttribute<int32>* InAttribute = InputParamData->Metadata->GetConstTypedAttribute<int32>(SortCommonTestData::AttributeName);
+	check(InAttribute);
+
+	const FPCGMetadataAttribute<int32>* OutAttribute = OutParamData->Metadata->GetConstTypedAttribute<int32>(SortCommonTestData::AttributeName);
+	UTEST_NOT_NULL("Output has the sorted attribute", OutAttribute);
+	check(OutAttribute);
+
+	for (int i = 0; i < OutParamData->Metadata->GetLocalItemCount() - 1; ++i)
+	{
+		const PCGMetadataValueKey InValueKey = InAttribute->GetValueKey(i);
+		const PCGMetadataValueKey OutValueKey = OutAttribute->GetValueKey(i);
+		UTEST_TRUE(FString::Format(TEXT("InAttribute and OutAttribute have the same value key for the same entry key ({0} vs {1})"), { InValueKey, OutValueKey }), InValueKey == OutValueKey);
+	}
+
+	return true;
+}
\ No newline at end of file
diff --git a/Engine/Plugins/Experimental/PCG/Source/PCG/Private/Tests/Elements/PCGSortPointsTest.cpp b/Engine/Plugins/Experimental/PCG/Source/PCG/Private/Tests/Elements/PCGSortPointsTest.cpp
deleted file mode 100644
diff --git a/Engine/Plugins/Experimental/PCG/Source/PCG/Public/Data/PCGSpatialData.h b/Engine/Plugins/Experimental/PCG/Source/PCG/Public/Data/PCGSpatialData.h
--- a/Engine/Plugins/Experimental/PCG/Source/PCG/Public/Data/PCGSpatialData.h
+++ b/Engine/Plugins/Experimental/PCG/Source/PCG/Public/Data/PCGSpatialData.h
@@ -47,9 +47,6 @@ public:
 	virtual FPCGAttributePropertyInputSelector GetCachedLastSelector() const override;
 	virtual void SetLastSelector(const FPCGAttributePropertySelector& InSelector) override;
 
-	virtual UPCGSpatialData* DuplicateData() const override { return DuplicateData(/*bInitializeMetadata=*/ true); }
-	// ~End UPCGData interface
-
 	/** Virtual call to allocate a new spacial data object, duplicate this spatial data into
 	*   and parent the new metadata with this class metadata (if asked).
 	*   Should be way cheaper than DuplicateObject, since we avoid duplicating metadata.
@@ -58,7 +55,8 @@ public:
 	*   They are mainly cached values (and octree for points).
 	*   TODO: If we want to also copy those values (can be an optimization), we need to guard the copy.
 	*/
-	UPCGSpatialData* DuplicateData(const bool bInitializeMetadata) const;
+	virtual UPCGSpatialData* DuplicateData(bool bInitializeMetadata = true) const override;
+	// ~End UPCGData interface
 
 	/** Returns the dimension of the data type, which has nothing to do with the dimension of its points */
 	UFUNCTION(BlueprintCallable, Category = SpatialData)
diff --git a/Engine/Plugins/Experimental/PCG/Source/PCG/Public/Elements/PCGSortAttributes.h b/Engine/Plugins/Experimental/PCG/Source/PCG/Public/Elements/PCGSortAttributes.h
new file mode 100644
--- /dev/null
+++ b/Engine/Plugins/Experimental/PCG/Source/PCG/Public/Elements/PCGSortAttributes.h
@@ -0,0 +1,53 @@
+// Copyright Epic Games, Inc. All Rights Reserved.
+
+#pragma once
+
+#include "PCGSettings.h"
+
+#include "PCGSortAttributes.generated.h"
+
+/**
+ * Sorts points based on an attribute.
+ */
+
+UENUM()
+enum class EPCGSortMethod : uint8
+{
+	Ascending,
+	Descending
+};
+
+UCLASS(BlueprintType, ClassGroup = (Procedural))
+class UPCGSortAttributesSettings : public UPCGSettings
+{
+	GENERATED_BODY()
+
+public:
+	//~Begin UPCGSettings interface
+#if WITH_EDITOR
+	virtual FName GetDefaultNodeName() const override;
+	virtual FText GetDefaultNodeTitle() const override;
+	virtual FText GetNodeTooltipText() const override;
+	virtual TArray<FText> GetNodeTitleAliases() const override;
+	virtual EPCGSettingsType GetType() const override { return EPCGSettingsType::Generic; }
+	virtual bool HasDynamicPins() const override { return true; }
+#endif
+
+protected:
+	virtual TArray<FPCGPinProperties> OutputPinProperties() const override;
+	virtual FPCGElementPtr CreateElement() const override;
+	//~End UPCGSettings interface
+
+public:
+	UPROPERTY(BlueprintReadWrite, EditAnywhere, Category = Settings, meta = (PCG_Overridable))
+	FPCGAttributePropertyInputSelector InputSource;
+
+	UPROPERTY(BlueprintReadWrite, EditAnywhere, Category = Settings, meta = (PCG_Overridable))
+	EPCGSortMethod SortMethod = EPCGSortMethod::Ascending;
+};
+
+class FPCGSortAttributesElement : public IPCGElement
+{
+protected:
+	virtual bool ExecuteInternal(FPCGContext* Context) const override;
+};
\ No newline at end of file
diff --git a/Engine/Plugins/Experimental/PCG/Source/PCG/Public/Elements/PCGSortPoints.h b/Engine/Plugins/Experimental/PCG/Source/PCG/Public/Elements/PCGSortPoints.h
deleted file mode 100644
diff --git a/Engine/Plugins/Experimental/PCG/Source/PCG/Public/Metadata/PCGMetadata.h b/Engine/Plugins/Experimental/PCG/Source/PCG/Public/Metadata/PCGMetadata.h
--- a/Engine/Plugins/Experimental/PCG/Source/PCG/Public/Metadata/PCGMetadata.h
+++ b/Engine/Plugins/Experimental/PCG/Source/PCG/Public/Metadata/PCGMetadata.h
@@ -39,13 +39,19 @@ public:
 	UFUNCTION(BlueprintCallable, Category = "PCG|Metadata")
 	void InitializeWithAttributeFilter(const UPCGMetadata* InParent, const TSet<FName>& InFilteredAttributes, EPCGMetadataFilterMode InFilterMode = EPCGMetadataFilterMode::ExcludeAttributes);
 
-	/** Initializes the metadata from a parent metadata by copying all attributes to it. */
-	UFUNCTION(BlueprintCallable, Category = "PCG|Metadata")
-	void InitializeAsCopy(const UPCGMetadata* InMetadataToCopy);
+	/** Initializes the metadata from a parent metadata by copying all attributes to it.
+	* @param InMetadataToCopy Metadata to copy from
+	* @param InOptionalEntriesToCopy Optional array that contains the keys to copy over. This array order will be respected, so it can also be used to re-order entries. Can be null.
+	*/
+	void InitializeAsCopy(const UPCGMetadata* InMetadataToCopy, const TArray<PCGMetadataEntryKey>* InOptionalEntriesToCopy = nullptr);
 
-	/** Initializes the metadata from a parent metadata by copy filtered attributes only to it */
-	UFUNCTION(BlueprintCallable, Category = "PCG|Metadata")
-	void InitializeAsCopyWithAttributeFilter(const UPCGMetadata* InMetadataToCopy, const TSet<FName>& InFilteredAttributes, EPCGMetadataFilterMode InFilterMode = EPCGMetadataFilterMode::ExcludeAttributes);
+	/** Initializes the metadata from a parent metadata by copy filtered attributes only to it
+	* @param InMetadataToCopy Metadata to copy from
+	* @param InFilteredAttributes Attributes to keep/exclude, can be empty.
+	* @param InFilterMode Filter to know if we should keep or exclude InFilteredAttributes.
+	* @param InOptionalEntriesToCopy Optional array that contains the keys to copy over. This array order will be respected, so it can also be used to re-order entries. Can be null.
+	*/
+	void InitializeAsCopyWithAttributeFilter(const UPCGMetadata* InMetadataToCopy, const TSet<FName>& InFilteredAttributes, EPCGMetadataFilterMode InFilterMode = EPCGMetadataFilterMode::ExcludeAttributes, const TArray<PCGMetadataEntryKey>* InOptionalEntriesToCopy = nullptr);
 
 	/** Creates missing attributes from another metadata if they are not currently present - note that this does not copy values */
 	UFUNCTION(BlueprintCallable, Category = "PCG|Metadata")
@@ -302,6 +308,27 @@ public:
 	template<typename T>
 	FPCGMetadataAttribute<T>* FindOrCreateAttribute(FName AttributeName, const T& DefaultValue = T{}, bool bAllowsInterpolation = true, bool bOverrideParent = true, bool bOverwriteIfTypeMismatch = true);
 
+	// Need this gymnastic for those 2 functions, because blueprints doesn't support default arguments for Arrays, and we don't want to force c++ user to provide OptionalNewEntriesOrder.
+
+	/** Initializes the metadata from a parent metadata by copying all attributes to it.
+	* @param InMetadataToCopy Metadata to copy from
+	* @param InOptionalEntriesToCopy Optional array that contains the keys to copy over. This array order will be respected, so it can also be used to re-order entries. If empty, copy them all.
+	*/
+	UFUNCTION(BlueprintCallable, Category = "PCG|Metadata", meta = (DisplayName = "Initialize As Copy", AutoCreateRefTerm = "InOptionalEntriesToCopy"))
+	void K2_InitializeAsCopy(const UPCGMetadata* InMetadataToCopy, const TArray<int64>& InOptionalEntriesToCopy) { return InitializeAsCopy(InMetadataToCopy, !InOptionalEntriesToCopy.IsEmpty() ? &InOptionalEntriesToCopy : nullptr); }
+
+	/** Initializes the metadata from a parent metadata by copy filtered attributes only to it
+	* @param InMetadataToCopy Metadata to copy from
+	* @param InFilteredAttributes Attributes to keep/exclude, can be empty.
+	* @param InOptionalEntriesToCopy Optional array that contains the keys to copy over. This array order will be respected, so it can also be used to re-order entries. If empty, copy them all.
+	* @param InFilterMode Filter to know if we should keep or exclude InFilteredAttributes.
+	*/
+	UFUNCTION(BlueprintCallable, Category = "PCG|Metadata", meta = (DisplayName = "Initialize As Copy With Attribute Filter", AutoCreateRefTerm = "InFilteredAttributes,InOptionalEntriesToCopy"))
+	void K2_InitializeAsCopyWithAttributeFilter(const UPCGMetadata* InMetadataToCopy, const TSet<FName>& InFilteredAttributes, const TArray<int64>& InOptionalEntriesToCopy, EPCGMetadataFilterMode InFilterMode = EPCGMetadataFilterMode::ExcludeAttributes)
+	{
+		return InitializeAsCopyWithAttributeFilter(InMetadataToCopy, InFilteredAttributes, InFilterMode, !InOptionalEntriesToCopy.IsEmpty() ? &InOptionalEntriesToCopy : nullptr);
+	}
+
 protected:
 	FPCGMetadataAttributeBase* CopyAttribute(FName AttributeToCopy, FName NewAttributeName, bool bKeepParent, bool bCopyEntries, bool bCopyValues);
 
diff --git a/Engine/Plugins/Experimental/PCG/Source/PCG/Public/Metadata/PCGMetadataAttribute.h b/Engine/Plugins/Experimental/PCG/Source/PCG/Public/Metadata/PCGMetadataAttribute.h
--- a/Engine/Plugins/Experimental/PCG/Source/PCG/Public/Metadata/PCGMetadataAttribute.h
+++ b/Engine/Plugins/Experimental/PCG/Source/PCG/Public/Metadata/PCGMetadataAttribute.h
@@ -53,11 +53,19 @@ public:
 	virtual bool AreValuesEqual(PCGMetadataValueKey ValueKey1, PCGMetadataValueKey ValueKey2) const = 0;
 
 	void SetValueFromValueKey(PCGMetadataEntryKey EntryKey, PCGMetadataValueKey ValueKey);
-	void SetValuesFromValueKeys(const TArray<TTuple<PCGMetadataEntryKey, PCGMetadataValueKey>>& EntryValuePairs, bool bResetValueOnDefaultValueKey = true);
 	PCGMetadataValueKey GetValueKey(PCGMetadataEntryKey EntryKey) const;
 	bool HasNonDefaultValue(PCGMetadataEntryKey EntryKey) const;
 	void ClearEntries();
 
+	/** Bulk getter, to lock in read only once per parent. */
+	void GetValueKeys(const TArray<PCGMetadataEntryKey>& EntryKeys, TArray<PCGMetadataValueKey>& OutValueKeys) const;
+
+	/** Bulk setter to lock in write only once. */
+	void SetValuesFromValueKeys(const TArray<TTuple<PCGMetadataEntryKey, PCGMetadataValueKey>>& EntryValuePairs, bool bResetValueOnDefaultValueKey = true);
+
+	/** Two arrays version of bulk setter to lock in write only once. Both arrays must be the same size. */
+	void SetValuesFromValueKeys(const TArray<PCGMetadataEntryKey>& EntryKeys, TArray<PCGMetadataValueKey>& ValueKeys, bool bResetValueOnDefaultValueKey = true);
+
 	bool AllowsInterpolation() const { return bAllowsInterpolation; }
 
 	int32 GetNumberOfEntries() const { return EntryToValueKeyMap.Num(); }
@@ -71,6 +79,12 @@ public:
 	static bool IsValidName(const FString& Name);
 	static bool IsValidName(const FName& Name);
 
+private:
+	// Unsafe version, needs to be write lock protected.
+	void SetValueFromValueKey_Unsafe(PCGMetadataEntryKey EntryKey, PCGMetadataValueKey ValueKey, bool bResetValueOnDefaultValueKey);
+
+	void GetValueKeys_Internal(const TArray<PCGMetadataEntryKey>& EntryKeys, TArray<PCGMetadataValueKey>& OutValueKeys, TBitArray<>& UnsetValues) const;
+
 protected:
 	TMap<PCGMetadataEntryKey, PCGMetadataValueKey> EntryToValueKeyMap;
 	mutable FRWLock EntryMapLock;
diff --git a/Engine/Plugins/Experimental/PCG/Source/PCG/Public/PCGData.h b/Engine/Plugins/Experimental/PCG/Source/PCG/Public/PCGData.h
--- a/Engine/Plugins/Experimental/PCG/Source/PCG/Public/PCGData.h
+++ b/Engine/Plugins/Experimental/PCG/Source/PCG/Public/PCGData.h
@@ -51,7 +51,7 @@ public:
 	virtual void SetLastSelector(const FPCGAttributePropertySelector& InSelector) {};
 
 	/** Return a copy of the data, with Metadata inheritence for spatial data. */
-	virtual UPCGData* DuplicateData() const;
+	virtual UPCGData* DuplicateData(bool bInitializeMetadata = true) const;
 
 	// ~Begin UObject interface
 	virtual void PostDuplicate(bool bDuplicateForPIE) override { InitUID(); }
@@ -239,8 +239,8 @@ public:
 	UFUNCTION(BlueprintCallable, Category = Data, meta = (ScriptMethod, DeterminesOutputType = "InDataTypeClass"))
 	static TArray<UPCGData*> GetTypedInputsByTag(const FPCGDataCollection& InCollection, const FString& InTag, TArray<FPCGTaggedData>& OutTaggedData, TSubclassOf<UPCGData> InDataTypeClass = nullptr);
 
-	/** Adds a data object to a given collection, simpler usage than making a PCGTaggedData object */
-	UFUNCTION(BlueprintCallable, Category = Data, meta = (ScriptMethod))
+	/** Adds a data object to a given collection, simpler usage than making a PCGTaggedData object. InTags can be empty. */
+	UFUNCTION(BlueprintCallable, Category = Data, meta = (ScriptMethod, AutoCreateRefTerm = "InTags"))
 	static void AddToCollection(UPARAM(ref) FPCGDataCollection& InCollection, const UPCGData* InData, FName InPinLabel, TArray<FString> InTags);
 
 	// Blueprint methods to support interaction with FPCGDataCollection
diff --git a/Engine/Plugins/Experimental/PCG/Source/PCG/Public/PCGParamData.h b/Engine/Plugins/Experimental/PCG/Source/PCG/Public/PCGParamData.h
--- a/Engine/Plugins/Experimental/PCG/Source/PCG/Public/PCGParamData.h
+++ b/Engine/Plugins/Experimental/PCG/Source/PCG/Public/PCGParamData.h
@@ -26,7 +26,7 @@ public:
 	virtual FPCGAttributePropertyInputSelector GetCachedLastSelector() const override;
 	virtual void SetLastSelector(const FPCGAttributePropertySelector& InSelector) override;
 
-	virtual UPCGParamData* DuplicateData() const override;
+	virtual UPCGParamData* DuplicateData(bool bInitializeMetadata = true) const override;
 	// ~End UPCGData interface
 
 	UFUNCTION(BlueprintCallable, Category = Metadata)
