[PCG] Transfer Attribute overrides are no longer spatial

#rnx

diff --git a/Engine/Plugins/Experimental/PCG/Source/PCG/Private/Elements/PCGAttributeTransferElement.cpp b/Engine/Plugins/Experimental/PCG/Source/PCG/Private/Elements/PCGAttributeTransferElement.cpp
--- a/Engine/Plugins/Experimental/PCG/Source/PCG/Private/Elements/PCGAttributeTransferElement.cpp
+++ b/Engine/Plugins/Experimental/PCG/Source/PCG/Private/Elements/PCGAttributeTransferElement.cpp
@@ -71,6 +71,12 @@ void UPCGAttributeTransferSettings::ApplyDeprecation(UPCGNode* InOutNode)
 
 EPCGDataType UPCGAttributeTransferSettings::GetCurrentPinTypes(const UPCGPin* InPin) const
 {
+	check(InPin);
+	if (!InPin->IsOutputPin())
+	{
+		return Super::GetCurrentPinTypes(InPin);
+	}
+
 	// All pins narrow to same type, which is Point if any input is Point, otherwise Spatial
 	const bool bAnyArePoint = (GetTypeUnionOfIncidentEdges(PCGAttributeTransferConstants::TargetLabel) == EPCGDataType::Point)
 		|| (GetTypeUnionOfIncidentEdges(PCGAttributeTransferConstants::SourceLabel) == EPCGDataType::Point);
