[PCG] Fix for generated component failing to refresh after redoing creating the PCG actor

#rnx

diff --git a/Engine/Plugins/Experimental/PCG/Source/PCG/Private/PCGComponent.cpp b/Engine/Plugins/Experimental/PCG/Source/PCG/Private/PCGComponent.cpp
--- a/Engine/Plugins/Experimental/PCG/Source/PCG/Private/PCGComponent.cpp
+++ b/Engine/Plugins/Experimental/PCG/Source/PCG/Private/PCGComponent.cpp
@@ -1669,7 +1669,11 @@ void UPCGComponent::PostEditUndo()
 
 	if (bGenerated)
 	{
-		Refresh(/*bIsStructural=*/true);
+		// Cancel existing means a refresh will always be scheduled even if another refresh was pending. If an undo
+		// operation removes the component, a valid refresh task ID is set but the refresh task itself will fail
+		// and leave the valid task ID hanging on the component. Forcing here means if we later retrieve this state
+		// from the undo/redo buffer, the refresh will be forced which will reset the state.
+		Refresh(/*bIsStructural=*/true, /*bCancelExistingRefresh=*/true);
 	}
 
 	Super::PostEditUndo();
@@ -1882,7 +1886,7 @@ const FPCGDataCollection* UPCGComponent::GetInspectionData(const FPCGStack& InSt
 	return InspectionCache.Find(InStack);
 }
 
-void UPCGComponent::Refresh(bool bStructural)
+void UPCGComponent::Refresh(bool bStructural, bool bCancelExistingRefresh)
 {
 	// Disable auto-refreshing on preview actors until we have something more robust on the execution side.
 	if (GetOwner() && GetOwner()->bIsEditorPreviewActor)
@@ -1930,9 +1934,19 @@ void UPCGComponent::Refresh(bool bStructural)
 		// Cancel an already existing generation if either the change is structural in nature (which requires a recompilation, so a full-rescheduling)
 		// or if the generation is already started
 		const bool bGenerationWasInProgress = IsGenerationInProgress();
-		if (CurrentGenerationTask != InvalidPCGTaskId && (bStructural || bGenerationWasInProgress))
+		bool bNeedToCancelCurrentTasks = (CurrentGenerationTask != InvalidPCGTaskId && (bStructural || bGenerationWasInProgress));
+
+		// Cancel an already existing refresh if caller allows this
+		if (bCancelExistingRefresh && CurrentRefreshTask != InvalidPCGTaskId)
+		{
+			bNeedToCancelCurrentTasks = true;
+
+			CurrentRefreshTask = InvalidPCGTaskId;
+		}
+
+		if (bNeedToCancelCurrentTasks)
 		{
-			CancelGeneration();
+			Subsystem->CancelGeneration(this);
 		}
 
 		// Calling a new refresh here might not be sufficient; if the current component was generating but was not previously generated,
diff --git a/Engine/Plugins/Experimental/PCG/Source/PCG/Public/PCGComponent.h b/Engine/Plugins/Experimental/PCG/Source/PCG/Public/PCGComponent.h
--- a/Engine/Plugins/Experimental/PCG/Source/PCG/Public/PCGComponent.h
+++ b/Engine/Plugins/Experimental/PCG/Source/PCG/Public/PCGComponent.h
@@ -270,7 +270,9 @@ public:
 	FPCGTaskId GetGenerationTaskId() const { return CurrentGenerationTask; }
 
 #if WITH_EDITOR
-	void Refresh(bool bStructural = false);
+	/** Schedules refresh of the component. If bCancelExistingRefresh is true, any existing refresh is cancelled and a new one is scheduled. */
+	void Refresh(bool bStructural = false, bool bCancelExistingRefresh = false);
+
 	void OnRefresh(bool bForceRefresh);
 
 	void StartGenerationInProgress();
