[PCG] Fix never ending SelfPruning if very low time slicing

* There was an early out if after getting the octree, if it was rebuilt
* Only early out if it is the first time we get the octree, it won't need to be rebuilt after that.

#rnx

diff --git a/Engine/Plugins/Experimental/PCG/Source/PCG/Private/Elements/PCGSelfPruning.cpp b/Engine/Plugins/Experimental/PCG/Source/PCG/Private/Elements/PCGSelfPruning.cpp
--- a/Engine/Plugins/Experimental/PCG/Source/PCG/Private/Elements/PCGSelfPruning.cpp
+++ b/Engine/Plugins/Experimental/PCG/Source/PCG/Private/Elements/PCGSelfPruning.cpp
@@ -121,7 +121,7 @@ namespace PCGSelfPruningElement
 
 		while (IterationState.CurrentPointIndex < IterationState.SortedPoints.Num())
 		{
-			// Don't check too many times.
+			// Don't check too many times. Pre-increment will make sure we always at least process "TimeSliceFrequencyCheck" elements at each iteration.
 			if (++CheckTimeSlicingCount >= TimeSliceFrequencyCheck)
 			{
 				if (ShouldStop(InOptionalContext))
@@ -164,7 +164,7 @@ namespace PCGSelfPruningElement
 
 		while (IterationState.CurrentPointIndex < IterationState.SortedPoints.Num())
 		{
-			// Don't check too many times.
+			// Don't check too many times. Pre-increment will make sure we always at least process "TimeSliceFrequencyCheck" elements at each iteration.
 			if (++CheckTimeSlicingCount >= TimeSliceFrequencyCheck)
 			{
 				if (ShouldStop(InOptionalContext))
@@ -261,9 +261,10 @@ namespace PCGSelfPruningElement
 		const FVector::FReal SquaredRadiusEquality = FMath::Square(RadiusEquality);
 		const TArray<FPCGPoint>& Points = InState.InputPointData->GetPoints();
 
-		// Force octree computation, and check if we need to stop after that
+		// Force octree computation, and check if we need to stop after that if it was dirty in the first place.
+		const bool bOctreeWasDirty = InState.InputPointData->IsOctreeDirty();
 		const UPCGPointData::PointOctree& Octree = InState.InputPointData->GetOctree();
-		if (ShouldStop(InOptionalContext))
+		if (bOctreeWasDirty && ShouldStop(InOptionalContext))
 		{
 			return false;
 		}
diff --git a/Engine/Plugins/Experimental/PCG/Source/PCG/Public/Data/PCGPointData.h b/Engine/Plugins/Experimental/PCG/Source/PCG/Public/Data/PCGPointData.h
--- a/Engine/Plugins/Experimental/PCG/Source/PCG/Public/Data/PCGPointData.h
+++ b/Engine/Plugins/Experimental/PCG/Source/PCG/Public/Data/PCGPointData.h
@@ -122,6 +122,9 @@ public:
 
 	bool ProjectPoint(const FTransform& InTransform, const FBox& InBounds, const FPCGProjectionParams& InParams, FPCGPoint& OutPoint, UPCGMetadata* OutMetadata, bool bUseBounds) const;
 
+	/** Get the dirty status of the Octree. Note that the Octree can be rebuilt from another thread, so this info can be invalidated at anytime. */
+	bool IsOctreeDirty() const { return bOctreeIsDirty; }
+
 protected:
 	void RebuildOctree() const;
 	void RecomputeBounds() const;
