Fixed potential mismanaged resource that could get garbage collected

#rnx

diff --git a/Engine/Plugins/Experimental/PCG/Source/PCG/Private/Elements/PCGMatchAndSetAttributes.cpp b/Engine/Plugins/Experimental/PCG/Source/PCG/Private/Elements/PCGMatchAndSetAttributes.cpp
--- a/Engine/Plugins/Experimental/PCG/Source/PCG/Private/Elements/PCGMatchAndSetAttributes.cpp
+++ b/Engine/Plugins/Experimental/PCG/Source/PCG/Private/Elements/PCGMatchAndSetAttributes.cpp
@@ -563,6 +563,7 @@ bool FPCGMatchAndSetAttributesElement::PrepareDataInternal(FPCGContext* InContex
 	check(TimeSlicedContext);
 
 	TArray<FPCGTaggedData> Inputs = InContext->InputData.GetInputsByPin(PCGPinConstants::DefaultInputLabel);
+	TArray<FPCGTaggedData>& Outputs = InContext->OutputData.TaggedData;
 	TArray<FPCGTaggedData> ParamDataInputs = InContext->InputData.GetInputsByPin(PCGMatchAndSetAttributesConstants::MatchDataLabel);
 
 	EPCGTimeSliceInitResult InitResult = TimeSlicedContext->InitializePerExecutionState([Settings, &ParamDataInputs](FPCGMatchAndSetAttributesElement::ContextType* Context, FPCGMatchAndSetAttributesExecutionState& OutState) -> EPCGTimeSliceInitResult
@@ -592,13 +593,14 @@ bool FPCGMatchAndSetAttributesElement::PrepareDataInternal(FPCGContext* InContex
 		return true;
 	}
 
-	TimeSlicedContext->InitializePerIterationStates(Inputs.Num(), [&Inputs, InContext](FPCGMatchAndSetAttributesIterationState& OutState, const FPCGMatchAndSetAttributesExecutionState&, int32 Index) -> EPCGTimeSliceInitResult
+	TimeSlicedContext->InitializePerIterationStates(Inputs.Num(), [&Inputs, &Outputs, InContext](FPCGMatchAndSetAttributesIterationState& OutState, const FPCGMatchAndSetAttributesExecutionState&, const uint32 IterationIndex) -> EPCGTimeSliceInitResult
 	{
-		OutState.InputData = Inputs[Index];
-		OutState.InPointData = Cast<UPCGPointData>(OutState.InputData.Data);
+		FPCGTaggedData& Output = Outputs.Add_GetRef(Inputs[IterationIndex]);
+
+		OutState.InPointData = Cast<UPCGPointData>(Inputs[IterationIndex].Data);
 		if (!OutState.InPointData)
 		{
-			PCGE_LOG_C(Error, GraphAndLog, InContext, FText::Format(LOCTEXT("InvalidInputDataType", "Input {0}: Input data must be of type Point"), FText::AsNumber(Index)));
+			PCGE_LOG_C(Error, GraphAndLog, InContext, FText::Format(LOCTEXT("InvalidInputDataType", "Input {0}: Input data must be of type Point"), FText::AsNumber(IterationIndex)));
 			return EPCGTimeSliceInitResult::NoOperation;
 		}
 
@@ -606,6 +608,8 @@ bool FPCGMatchAndSetAttributesElement::PrepareDataInternal(FPCGContext* InContex
 		OutState.OutPointData->InitializeFromData(OutState.InPointData);
 		OutState.OutPointData->GetMutablePoints().Reserve(OutState.InPointData->GetPoints().Num());
 
+		Output.Data = OutState.OutPointData;
+
 		return EPCGTimeSliceInitResult::Success;
 	});
 
@@ -642,7 +646,6 @@ bool FPCGMatchAndSetAttributesElement::ExecuteInternal(FPCGContext* InContext) c
 		// This iteration resulted in an early out for no sampling operation. Early out with a passthrough.
 		if (InitResult == EPCGTimeSliceInitResult::NoOperation)
 		{
-			Context->OutputData.TaggedData.Add(Inputs[IterationIndex]);
 			return true;
 		}
 
@@ -650,14 +653,7 @@ bool FPCGMatchAndSetAttributesElement::ExecuteInternal(FPCGContext* InContext) c
 		check(InitResult == EPCGTimeSliceInitResult::Success);
 
 		// Run the execution until the time slice is finished
-		const bool bDone = ExecState.Partition->SelectPoints(*Context, IterState.InPointData, IterState.CurrentPointIndex, IterState.OutPointData);
-		if (bDone)
-		{
-			FPCGTaggedData& Output = Context->OutputData.TaggedData.Add_GetRef(Inputs[IterationIndex]);
-			Output.Data = IterState.OutPointData;
-		}
-		
-		return bDone;
+		return ExecState.Partition->SelectPoints(*Context, IterState.InPointData, IterState.CurrentPointIndex, IterState.OutPointData);
 	});
 }
 
diff --git a/Engine/Plugins/Experimental/PCG/Source/PCG/Public/Elements/PCGMatchAndSetAttributes.h b/Engine/Plugins/Experimental/PCG/Source/PCG/Public/Elements/PCGMatchAndSetAttributes.h
--- a/Engine/Plugins/Experimental/PCG/Source/PCG/Public/Elements/PCGMatchAndSetAttributes.h
+++ b/Engine/Plugins/Experimental/PCG/Source/PCG/Public/Elements/PCGMatchAndSetAttributes.h
@@ -22,7 +22,6 @@ struct FPCGMatchAndSetAttributesExecutionState
 struct FPCGMatchAndSetAttributesIterationState
 {
 	int CurrentPointIndex = 0;
-	FPCGTaggedData InputData;
 	const UPCGPointData* InPointData = nullptr;
 	UPCGPointData* OutPointData = nullptr;
 };
