[PCG] Refactor CopyPoints tests

#rnx

diff --git a/Engine/Plugins/Experimental/PCG/Source/PCG/Private/Tests/Elements/PCGCopyPointsTest.cpp b/Engine/Plugins/Experimental/PCG/Source/PCG/Private/Tests/Elements/PCGCopyPointsTest.cpp
--- a/Engine/Plugins/Experimental/PCG/Source/PCG/Private/Tests/Elements/PCGCopyPointsTest.cpp
+++ b/Engine/Plugins/Experimental/PCG/Source/PCG/Private/Tests/Elements/PCGCopyPointsTest.cpp
@@ -12,18 +12,45 @@
 
 #if WITH_EDITOR
 
-IMPLEMENT_CUSTOM_SIMPLE_AUTOMATION_TEST(FPCGCopyPointsTest, FPCGTestBaseClass, "Plugins.PCG.CopyPoints.Basic", PCGTestsCommon::TestFlags)
+class FPCGCopyPointsTestBase : public FPCGTestBaseClass
+{
+public:
+	using FPCGTestBaseClass::FPCGTestBaseClass;
+	virtual ~FPCGCopyPointsTestBase() = default;
+
+protected:
+	struct FParameters
+	{
+		EPCGCopyPointsInheritanceMode RotationInheritance = EPCGCopyPointsInheritanceMode::Relative;
+		EPCGCopyPointsInheritanceMode ScaleInheritance = EPCGCopyPointsInheritanceMode::Relative;
+		EPCGCopyPointsInheritanceMode ColorInheritance = EPCGCopyPointsInheritanceMode::Relative;
+		EPCGCopyPointsInheritanceMode SeedInheritance = EPCGCopyPointsInheritanceMode::Relative;
+		EPCGCopyPointsMetadataInheritanceMode AttributeInheritance = EPCGCopyPointsMetadataInheritanceMode::SourceFirst;
+		TObjectPtr<UPCGPointData> SourceData = nullptr;
+		TObjectPtr<UPCGPointData> TargetData = nullptr;
+		TArray<FPCGPoint> ExpectedPoints;
+	};
 
-bool FPCGCopyPointsTest::RunTest(const FString& Parameters)
+	bool RunTestInternal(FParameters& Parameters);
+};
+
+bool FPCGCopyPointsTestBase::RunTestInternal(FParameters& Parameters)
 {
+	UPCGPointData* SourceData = Parameters.SourceData;
+	UPCGPointData* TargetData = Parameters.TargetData;
+
+	check(SourceData && TargetData);
+
 	PCGTestsCommon::FTestData TestData;
 	PCGTestsCommon::GenerateSettings<UPCGCopyPointsSettings>(TestData);
 	UPCGCopyPointsSettings* Settings = CastChecked<UPCGCopyPointsSettings>(TestData.Settings);
 	FPCGElementPtr CopyPointsElement = TestData.Settings->GetElement();
 
-	TObjectPtr<UPCGPointData> EmptyData = PCGTestsCommon::CreateEmptyPointData();
-	TObjectPtr<UPCGPointData> SourceData = PCGTestsCommon::CreateRandomPointData(10, TestData.Seed);
-	TObjectPtr<UPCGPointData> TargetData = PCGTestsCommon::CreateRandomPointData(4, TestData.Seed+1);
+	Settings->RotationInheritance = Parameters.RotationInheritance;
+	Settings->ScaleInheritance = Parameters.ScaleInheritance;
+	Settings->ColorInheritance = Parameters.ColorInheritance;
+	Settings->SeedInheritance = Parameters.SeedInheritance;
+	Settings->AttributeInheritance = Parameters.AttributeInheritance;
 
 	const FName LengthName = TEXT("Length");
 	const FName WidthName = TEXT("Width");
@@ -54,289 +81,308 @@ bool FPCGCopyPointsTest::RunTest(const FString& Parameters)
 	}
 
 	FPCGTaggedData& SourceTaggedData = TestData.InputData.TaggedData.Emplace_GetRef(FPCGTaggedData());
-	SourceTaggedData.Data = SourceData;
+	SourceTaggedData.Data = Parameters.SourceData;
 	SourceTaggedData.Pin = PCGCopyPointsConstants::SourcePointsLabel;
 
 	FPCGTaggedData& TargetTaggedData = TestData.InputData.TaggedData.Emplace_GetRef(FPCGTaggedData());
-	TargetTaggedData.Data = TargetData;
+	TargetTaggedData.Data = Parameters.TargetData;
 	TargetTaggedData.Pin = PCGCopyPointsConstants::TargetPointsLabel;
 
 	// Test only supports float attributes for simplicity
-	auto ValidateCopyPointsPoints = [this, &TestData, CopyPointsElement, Settings]() -> bool
-	{
-		TUniquePtr<FPCGContext> Context = TestData.InitializeTestContext();
 
-		while (!CopyPointsElement->Execute(Context.Get()))
-		{}
+	TUniquePtr<FPCGContext> Context = TestData.InitializeTestContext();
 
-		const TArray<FPCGTaggedData> Sources = Context->InputData.GetInputsByPin(PCGCopyPointsConstants::SourcePointsLabel);
-		const TArray<FPCGTaggedData> Targets = Context->InputData.GetInputsByPin(PCGCopyPointsConstants::TargetPointsLabel);
-		const TArray<FPCGTaggedData>& Outputs = Context->OutputData.GetInputs();
+	while (!CopyPointsElement->Execute(Context.Get()))
+	{}
 
-		if (!TestTrue("Valid number of inputs", Sources.Num() == 1 && Targets.Num() == 1))
-		{
-			return false;
-		}
-
-		if (!TestTrue("Valid number of outputs", Sources.Num() == 1))
-		{
-			return false;
-		}
+	const TArray<FPCGTaggedData> Sources = Context->InputData.GetInputsByPin(PCGCopyPointsConstants::SourcePointsLabel);
+	const TArray<FPCGTaggedData> Targets = Context->InputData.GetInputsByPin(PCGCopyPointsConstants::TargetPointsLabel);
+	const TArray<FPCGTaggedData>& Outputs = Context->OutputData.GetInputs();
 
-		const FPCGTaggedData& Source = Sources[0];
-		const FPCGTaggedData& Target = Targets[0];
-		const FPCGTaggedData& Output = Outputs[0];
+	UTEST_TRUE("Valid number of inputs", Sources.Num() == 1 && Targets.Num() == 1);
+	UTEST_EQUAL("Valid number of outputs", Sources.Num(), 1);
 
-		check(Source.Data && Target.Data);
+	const UPCGPointData* OutPointData = Cast<const UPCGPointData>(Outputs[0].Data);
+	UTEST_NOT_NULL("Valid output point data", OutPointData);
+	check(OutPointData);
 
-		const UPCGSpatialData* SourceSpatialData = Cast<UPCGSpatialData>(Source.Data);
-		const UPCGSpatialData* TargetSpatialData = Cast<UPCGSpatialData>(Target.Data);
+	const TArray<FPCGPoint>& OutPoints = OutPointData->GetPoints();
 
-		check(SourceSpatialData && TargetSpatialData);
+	UTEST_EQUAL("Valid number of output points", SourcePoints.Num() * TargetPoints.Num(), OutPoints.Num());
+	check(Parameters.ExpectedPoints.IsEmpty() || Parameters.ExpectedPoints.Num() == OutPoints.Num());
 
-		const UPCGPointData* SourcePointData = SourceSpatialData->ToPointData(Context.Get());
-		const UPCGPointData* TargetPointData = TargetSpatialData->ToPointData(Context.Get());
+	const UPCGMetadata* RootMetadata = nullptr;
+	const UPCGMetadata* NonRootMetadata = nullptr;
+	EPCGCopyPointsMetadataInheritanceMode RootMode;
+	EPCGCopyPointsMetadataInheritanceMode NonRootMode;
 
-		check(SourcePointData && TargetPointData);
+	if (Settings->AttributeInheritance == EPCGCopyPointsMetadataInheritanceMode::SourceFirst)
+	{
+		RootMetadata = SourceData->Metadata;
+		NonRootMetadata = TargetData->Metadata;
+		RootMode = EPCGCopyPointsMetadataInheritanceMode::SourceFirst;
+		NonRootMode = EPCGCopyPointsMetadataInheritanceMode::TargetFirst;
+	}
+	else // if (Settings->AttributeInheritance == EPCGCopyPointsMetadataInheritanceMode::TargetFirst)
+	{
+		RootMetadata = TargetData->Metadata;
+		NonRootMetadata = SourceData->Metadata;
+		RootMode = EPCGCopyPointsMetadataInheritanceMode::TargetFirst;
+		NonRootMode = EPCGCopyPointsMetadataInheritanceMode::SourceFirst;
+	}
 
-		if (!TestTrue("Valid output data", Output.Data != nullptr))
-		{
-			return false;
-		}
+	UTEST_TRUE("Valid input metadata", RootMetadata && NonRootMetadata);
 
-		const UPCGSpatialData* OutSpatialData = Cast<UPCGSpatialData>(Output.Data);
+	check(RootMetadata && NonRootMetadata); // Static analyzer doesn't understand semantics of TestTrue.
+		
+	bool bTestPassed = true;
 
-		if (!TestNotNull("Valid output SpatialData", OutSpatialData))
-		{
-			return false;
-		}
+	TArray<const FPCGMetadataAttribute<float>*> InheritedAttributes;
+	TArray<EPCGCopyPointsMetadataInheritanceMode> InheritedAttributeModes;
+	TArray<const FPCGMetadataAttribute<float>*> OutAttributes;
 
-		const UPCGPointData* OutPointData = OutSpatialData->ToPointData(Context.Get());
+	TArray<FName> AttributeNames;
+	TArray<EPCGMetadataTypes> AttributeTypes;
+	RootMetadata->GetAttributes(AttributeNames, AttributeTypes);
 
-		if (!TestNotNull("Valid output PointData", OutPointData))
+	for (const FName& AttributeName : AttributeNames)
+	{
+		if (TestTrue("Contains attribute found in source", OutPointData->Metadata->HasAttribute(AttributeName)))
 		{
-			return false;
-		}
+			const FPCGMetadataAttribute<float>* InheritedAttribute = static_cast<const FPCGMetadataAttribute<float>*>(RootMetadata->GetConstAttribute(AttributeName));
+			const FPCGMetadataAttribute<float>* OutAttribute = static_cast<const FPCGMetadataAttribute<float>*>(OutPointData->Metadata->GetConstAttribute(AttributeName));
 
-		const TArray<FPCGPoint>& SourcePoints = SourcePointData->GetPoints();
-		const TArray<FPCGPoint>& TargetPoints = TargetPointData->GetPoints();
-		const TArray<FPCGPoint>& OutPoints = OutPointData->GetPoints();
-
-		if (!TestEqual("Valid number of output points", SourcePoints.Num() * TargetPoints.Num(), OutPoints.Num()))
-		{ 
-			return false;
-		}
-
-		const UPCGMetadata* RootMetadata = nullptr;
-		const UPCGMetadata* NonRootMetadata = nullptr;
-		EPCGCopyPointsMetadataInheritanceMode RootMode;
-		EPCGCopyPointsMetadataInheritanceMode NonRootMode;
-
-		if (Settings->AttributeInheritance == EPCGCopyPointsMetadataInheritanceMode::SourceFirst)
-		{
-			RootMetadata = SourcePointData->Metadata;
-			NonRootMetadata = TargetPointData->Metadata;
-			RootMode = EPCGCopyPointsMetadataInheritanceMode::SourceFirst;
-			NonRootMode = EPCGCopyPointsMetadataInheritanceMode::TargetFirst;
+			InheritedAttributes.Add(InheritedAttribute);
+			InheritedAttributeModes.Add(RootMode);
+			OutAttributes.Add(OutAttribute);
 		}
-		else // if (Settings->AttributeInheritance == EPCGCopyPointsMetadataInheritanceMode::TargetFirst)
-		{
-			RootMetadata = TargetPointData->Metadata;
-			NonRootMetadata = SourcePointData->Metadata;
-			RootMode = EPCGCopyPointsMetadataInheritanceMode::TargetFirst;
-			NonRootMode = EPCGCopyPointsMetadataInheritanceMode::SourceFirst;
-		}
-
-		if (!TestTrue("Valid input metadata", RootMetadata && NonRootMetadata))
+		else
 		{
 			return false;
 		}
+	}
 
-		CA_ASSUME(RootMetadata); // Static analyzer doesn't understand semantics of TestTrue.
-		CA_ASSUME(NonRootMetadata);
-		
-		bool bTestPassed = true;
-
-		TArray<const FPCGMetadataAttribute<float>*> InheritedAttributes;
-		TArray<EPCGCopyPointsMetadataInheritanceMode> InheritedAttributeModes;
-		TArray<const FPCGMetadataAttribute<float>*> OutAttributes;
-
-		TArray<FName> AttributeNames;
-		TArray<EPCGMetadataTypes> AttributeTypes;
-		RootMetadata->GetAttributes(AttributeNames, AttributeTypes);
+	NonRootMetadata->GetAttributes(AttributeNames, AttributeTypes);
 
-		for (const FName& AttributeName : AttributeNames)
+	for (const FName& AttributeName : AttributeNames)
+	{
+		if (!RootMetadata->HasAttribute(AttributeName))
 		{
-			if (TestTrue("Contains attribute found in source", OutPointData->Metadata->HasAttribute(AttributeName)))
+			if (TestTrue("Contains attribute found in target", OutPointData->Metadata->HasAttribute(AttributeName)))
 			{
-				const FPCGMetadataAttribute<float>* InheritedAttribute = static_cast<const FPCGMetadataAttribute<float>*>(RootMetadata->GetConstAttribute(AttributeName));
+				const FPCGMetadataAttribute<float>* InheritedAttribute = static_cast<const FPCGMetadataAttribute<float>*>(NonRootMetadata->GetConstAttribute(AttributeName));
 				const FPCGMetadataAttribute<float>* OutAttribute = static_cast<const FPCGMetadataAttribute<float>*>(OutPointData->Metadata->GetConstAttribute(AttributeName));
 
 				InheritedAttributes.Add(InheritedAttribute);
-				InheritedAttributeModes.Add(RootMode);
+				InheritedAttributeModes.Add(NonRootMode);
 				OutAttributes.Add(OutAttribute);
 			}
 			else
 			{
-				bTestPassed = false;
+				return false;
 			}
 		}
+	}
 
-		NonRootMetadata->GetAttributes(AttributeNames, AttributeTypes);
+	check(InheritedAttributeModes.Num() == InheritedAttributeModes.Num() && InheritedAttributeModes.Num() == OutAttributes.Num())
 
-		for (const FName& AttributeName : AttributeNames)
+	for (int PointIndex = 0; PointIndex < OutPoints.Num(); ++PointIndex)
+	{
+		const FPCGPoint& SourcePoint = SourcePoints[PointIndex / TargetPoints.Num()];
+		const FPCGPoint& TargetPoint = TargetPoints[PointIndex % TargetPoints.Num()];
+		const FPCGPoint& OutPoint = OutPoints[PointIndex];
+
+		auto FormatMessage = [PointIndex](const FString& Message)
 		{
-			if (!RootMetadata->HasAttribute(AttributeName))
-			{
-				if (TestTrue("Contains attribute found in target", OutPointData->Metadata->HasAttribute(AttributeName)))
-				{
-					const FPCGMetadataAttribute<float>* InheritedAttribute = static_cast<const FPCGMetadataAttribute<float>*>(NonRootMetadata->GetConstAttribute(AttributeName));
-					const FPCGMetadataAttribute<float>* OutAttribute = static_cast<const FPCGMetadataAttribute<float>*>(OutPointData->Metadata->GetConstAttribute(AttributeName));
-
-					InheritedAttributes.Add(InheritedAttribute);
-					InheritedAttributeModes.Add(NonRootMode);
-					OutAttributes.Add(OutAttribute);
-				}
-				else
-				{
-					bTestPassed = false;
-				}
-			}
+			return FString::Printf(TEXT("Index %d: %s"), PointIndex, *Message);
+		};
+
+		UTEST_EQUAL(*FormatMessage(TEXT("SourcePoint and OutPoint index %d have same Density")), SourcePoint.Density, OutPoint.Density);
+		UTEST_EQUAL(*FormatMessage(TEXT("SourcePoint and OutPoint index %d have same BoundsMin")), SourcePoint.BoundsMin, OutPoint.BoundsMin);
+		UTEST_EQUAL(*FormatMessage(TEXT("SourcePoint and OutPoint index %d have same BoundsMax")), SourcePoint.BoundsMax, OutPoint.BoundsMax);
+		UTEST_EQUAL(*FormatMessage(TEXT("SourcePoint and OutPoint index %d have same Steepness")), SourcePoint.Steepness, OutPoint.Steepness);
+
+		// Validate transform inheritance
+		const FTransform& SourceTransform = SourcePoint.Transform;
+		const FTransform& TargetTransform = TargetPoint.Transform;
+		const FTransform& OutTransform = OutPoint.Transform;
+
+		const FVector Location = TargetPoint.Transform.TransformPosition(SourcePoint.Transform.GetLocation());
+		FQuat Rotation;
+		FVector Scale;
+		FVector4 Color;
+		int32 Seed;
+
+		if (Settings->RotationInheritance == EPCGCopyPointsInheritanceMode::Relative)
+		{
+			Rotation = Parameters.ExpectedPoints[PointIndex].Transform.GetRotation();
+		}
+		else if (Settings->RotationInheritance == EPCGCopyPointsInheritanceMode::Source)
+		{
+			Rotation = SourceTransform.GetRotation();
+		}
+		else // if (Settings->RotationInheritance == EPCGCopyPointsInheritanceMode::Target)
+		{
+			Rotation = TargetTransform.GetRotation();
+		}
+
+		if (Settings->ScaleInheritance == EPCGCopyPointsInheritanceMode::Relative)
+		{
+			Scale = Parameters.ExpectedPoints[PointIndex].Transform.GetScale3D();
+		}
+		else if (Settings->ScaleInheritance == EPCGCopyPointsInheritanceMode::Source)
+		{
+			Scale = SourceTransform.GetScale3D();
+		}
+		else // if (Settings->ScaleInheritance == EPCGCopyPointsInheritanceMode::Target)
+		{
+			Scale = TargetTransform.GetScale3D();
 		}
 
-		check(InheritedAttributeModes.Num() == InheritedAttributeModes.Num() && InheritedAttributeModes.Num() == OutAttributes.Num())
+		if (Settings->ColorInheritance == EPCGCopyPointsInheritanceMode::Relative)
+		{
+			Color = Parameters.ExpectedPoints[PointIndex].Color;
+		}
+		else if (Settings->ColorInheritance == EPCGCopyPointsInheritanceMode::Source)
+		{
+			Color = SourcePoint.Color;
+		}
+		else // if (Settings->ColorInheritance == EPCGCopyPointsInheritanceMode::Target)
+		{
+			Color = TargetPoint.Color;
+		}
 
-		for (int PointIndex = 0; PointIndex < OutPoints.Num(); ++PointIndex)
+		if (Settings->SeedInheritance == EPCGCopyPointsInheritanceMode::Relative)
 		{
-			const FPCGPoint& SourcePoint = SourcePoints[PointIndex / TargetPoints.Num()];
-			const FPCGPoint& TargetPoint = TargetPoints[PointIndex % TargetPoints.Num()];
-			const FPCGPoint& OutPoint = OutPoints[PointIndex];
-
-			bTestPassed &= TestEqual("SourcePoint and OutPoint have same Density", SourcePoint.Density, OutPoint.Density);
-			bTestPassed &= TestEqual("SourcePoint and OutPoint have same BoundsMin", SourcePoint.BoundsMin, OutPoint.BoundsMin);
-			bTestPassed &= TestEqual("SourcePoint and OutPoint have same BoundsMax", SourcePoint.BoundsMax, OutPoint.BoundsMax);
-			bTestPassed &= TestEqual("SourcePoint and OutPoint have same Steepness", SourcePoint.Steepness, OutPoint.Steepness);
-
-			// Validate transform inheritance
-			const FTransform& SourceTransform = SourcePoint.Transform;
-			const FTransform& TargetTransform = TargetPoint.Transform;
-			const FTransform& OutTransform = OutPoint.Transform;
-
-			const FVector Location = TargetPoint.Transform.TransformPosition(SourcePoint.Transform.GetLocation());
-			FQuat Rotation;
-			FVector Scale;
-			FVector4 Color;
-			int32 Seed;
-
-			if (Settings->RotationInheritance == EPCGCopyPointsInheritanceMode::Relative)
-			{
-				Rotation = TargetTransform.GetRotation() * SourceTransform.GetRotation();
-			}
-			else if (Settings->RotationInheritance == EPCGCopyPointsInheritanceMode::Source)
-			{
-				Rotation = SourceTransform.GetRotation();
-			}
-			else // if (Settings->RotationInheritance == EPCGCopyPointsInheritanceMode::Target)
-			{
-				Rotation = TargetTransform.GetRotation();
-			}
+			Seed = Parameters.ExpectedPoints[PointIndex].Seed;
+		}
+		else if (Settings->SeedInheritance == EPCGCopyPointsInheritanceMode::Source)
+		{
+			Seed = SourcePoint.Seed;
+		}
+		else // if (Settings->SeedInheritance == EPCGCopyPointsInheritanceMode::Target)
+		{
+			Seed = TargetPoint.Seed;
+		}
 
-			if (Settings->ScaleInheritance == EPCGCopyPointsInheritanceMode::Relative)
-			{
-				Scale = SourceTransform.GetScale3D() * TargetTransform.GetScale3D();
-			}
-			else if (Settings->ScaleInheritance == EPCGCopyPointsInheritanceMode::Source)
-			{
-				Scale = SourceTransform.GetScale3D();
-			}
-			else // if (Settings->ScaleInheritance == EPCGCopyPointsInheritanceMode::Target)
-			{
-				Scale = TargetTransform.GetScale3D();
-			}
+		UTEST_TRUE(*FormatMessage(TEXT("Valid rotation")), Rotation.Equals(OutTransform.GetRotation()));
+		UTEST_EQUAL(*FormatMessage(TEXT("Valid scale")), Scale, OutTransform.GetScale3D());
+		UTEST_TRUE(*FormatMessage(TEXT("Valid color")), Color.Equals(OutPoint.Color));
+		UTEST_EQUAL(*FormatMessage(TEXT("Valid location")), Location, OutTransform.GetLocation());
+		UTEST_EQUAL(*FormatMessage(TEXT("Valid seed")), Seed, OutPoint.Seed);
 
-			if (Settings->ColorInheritance == EPCGCopyPointsInheritanceMode::Relative)
-			{
-				Color = SourcePoint.Color * TargetPoint.Color;
-			}
-			else if (Settings->ColorInheritance == EPCGCopyPointsInheritanceMode::Source)
-			{
-				Color = SourcePoint.Color;
-			}
-			else // if (Settings->ColorInheritance == EPCGCopyPointsInheritanceMode::Target)
-			{
-				Color = TargetPoint.Color;
-			}
+		// Validate point value keys
+		for (int AttributeIndex = 0; AttributeIndex < InheritedAttributes.Num(); ++AttributeIndex)
+		{
+			const FPCGMetadataAttribute<float>* InheritedAttribute = InheritedAttributes[AttributeIndex];
+			const FPCGMetadataAttribute<float>* OutAttribute = OutAttributes[AttributeIndex];
+			check(InheritedAttribute && OutAttribute);
 
-			if (Settings->SeedInheritance == EPCGCopyPointsInheritanceMode::Relative)
-			{
-				Seed = PCGHelpers::ComputeSeed(SourcePoint.Seed, TargetPoint.Seed);
-			}
-			else if (Settings->SeedInheritance == EPCGCopyPointsInheritanceMode::Source)
-			{
-				Seed = SourcePoint.Seed;
-			}
-			else // if (Settings->SeedInheritance == EPCGCopyPointsInheritanceMode::Target)
-			{
-				Seed = TargetPoint.Seed;
-			}
+			const PCGMetadataEntryKey EntryKey = (InheritedAttributeModes[AttributeIndex] == EPCGCopyPointsMetadataInheritanceMode::SourceFirst) ? SourcePoint.MetadataEntry : TargetPoint.MetadataEntry;
+			UTEST_EQUAL(*FormatMessage(TEXT("Valid metadata value")), InheritedAttribute->GetValueKey(EntryKey), OutAttribute->GetValueKey(OutPoint.MetadataEntry));
+		}
+	}
 
-			bTestPassed &= TestEqual("Valid rotation", Rotation, OutTransform.GetRotation());
-			bTestPassed &= TestEqual("Valid scale", Scale, OutTransform.GetScale3D());
-			bTestPassed &= TestEqual("Valid color", Color, OutPoint.Color);
-			bTestPassed &= TestEqual("Valid location", Location, OutTransform.GetLocation());
-			bTestPassed &= TestEqual("Valid seed", Seed, OutPoint.Seed);		
+	return true;
+}
 
-			// Validate point value keys
-			for (int AttributeIndex = 0; AttributeIndex < InheritedAttributes.Num(); ++AttributeIndex)
-			{
-				const FPCGMetadataAttribute<float>* InheritedAttribute = InheritedAttributes[AttributeIndex];
-				const FPCGMetadataAttribute<float>* OutAttribute = OutAttributes[AttributeIndex];
-				check(InheritedAttribute && OutAttribute);
 
-				const PCGMetadataEntryKey EntryKey = (InheritedAttributeModes[AttributeIndex] == EPCGCopyPointsMetadataInheritanceMode::SourceFirst) ? SourcePoint.MetadataEntry : TargetPoint.MetadataEntry;
-				bTestPassed &= TestEqual("Valid metadata value", InheritedAttribute->GetValueKey(EntryKey), OutAttribute->GetValueKey(OutPoint.MetadataEntry));
-			}
-		}
+IMPLEMENT_CUSTOM_SIMPLE_AUTOMATION_TEST(FPCGCopyPointsTest_Relative, FPCGCopyPointsTestBase, "Plugins.PCG.CopyPoints.Relative", PCGTestsCommon::TestFlags)
+IMPLEMENT_CUSTOM_SIMPLE_AUTOMATION_TEST(FPCGCopyPointsTest_Source, FPCGCopyPointsTestBase, "Plugins.PCG.CopyPoints.Source", PCGTestsCommon::TestFlags)
+IMPLEMENT_CUSTOM_SIMPLE_AUTOMATION_TEST(FPCGCopyPointsTest_Target, FPCGCopyPointsTestBase, "Plugins.PCG.CopyPoints.Target", PCGTestsCommon::TestFlags)
+IMPLEMENT_CUSTOM_SIMPLE_AUTOMATION_TEST(FPCGCopyPointsTest_EmptySource, FPCGCopyPointsTestBase, "Plugins.PCG.CopyPoints.EmptySource", PCGTestsCommon::TestFlags)
+IMPLEMENT_CUSTOM_SIMPLE_AUTOMATION_TEST(FPCGCopyPointsTest_EmptyTarget, FPCGCopyPointsTestBase, "Plugins.PCG.CopyPoints.EmptyTarget", PCGTestsCommon::TestFlags)
 
-		return bTestPassed;
-	};
+bool FPCGCopyPointsTest_Relative::RunTest(const FString&)
+{
+	FPCGCopyPointsTestBase::FParameters Parameters;
+	Parameters.RotationInheritance = EPCGCopyPointsInheritanceMode::Relative;
+	Parameters.ScaleInheritance = EPCGCopyPointsInheritanceMode::Relative;
+	Parameters.ColorInheritance = EPCGCopyPointsInheritanceMode::Relative;
+	Parameters.SeedInheritance = EPCGCopyPointsInheritanceMode::Relative;
+
+	// For relative, we hard code the results
+	UPCGPointData* SourceData = NewObject<UPCGPointData>();
+	UPCGPointData* TargetData = NewObject<UPCGPointData>();
+	TArray<FPCGPoint>& SourcePoints = SourceData->GetMutablePoints();
+	TArray<FPCGPoint>& TargetPoints = TargetData->GetMutablePoints();
+	
+	SourcePoints.Emplace(FTransform(FQuat::Identity, FVector(3.0, 4.0, 5.0), FVector(1.0, 1.0, 1.0)), 1.0f, 42);
+	SourcePoints.Emplace(FTransform(FQuat::MakeFromEuler(FVector(45.0, 45.0, 45.0)), FVector(2.0, 1.0, -1.0), FVector(1.0, 1.0, 1.0)), 1.0f, 43);
+	SourcePoints.Emplace(FTransform(FQuat::MakeFromEuler(FVector(46.0, 47.0, 48.0)), FVector(6.0, 2.0, 3.0), FVector(-1.0, 1.0, -1.0)), 1.0f, 44);
+	SourcePoints[0].Color = FVector4(0.1, 0.5, 0.2, 0.3);
+	SourcePoints[1].Color = FVector4(0.4, 0.1, 0.1, 0.6);
+	SourcePoints[2].Color = FVector4(0.7, 0.4, 0.9, 0.9);
+
+	TargetPoints.Emplace(FTransform(FQuat::MakeFromEuler(FVector(20.0, 21.0, 22.0)), FVector(-6.0, -4.0, -1.0), FVector(1.0, 1.0, 1.0)), 1.0f, 45);
+	TargetPoints.Emplace(FTransform(FQuat::MakeFromEuler(FVector(22.0, 23.0, 24.0)), FVector(-7.0, -5.0, -10.0), FVector(-1.0, 1.0, -1.0)), 1.0f, 46);
+	TargetPoints[0].Color = FVector4(0.4, 0.3, 1.0, 1.0);
+	TargetPoints[1].Color = FVector4(0.1, 0.2, 0.7, 0.4);
+
+	// Hard coded values for the values above.
+	Parameters.ExpectedPoints.Emplace(FTransform(FQuat(-0.133360, -0.208748, 0.153700, 0.956564), FVector(-6.558469, 1.672731, 3.184284), FVector(1.000000, 1.000000, 1.000000)), 1.0f, -21834409);
+	Parameters.ExpectedPoints.Last().Color = FVector4(0.040000, 0.150000, 0.200000, 0.300000);
+	Parameters.ExpectedPoints.Emplace(FTransform(FQuat(-0.142203, -0.230303, 0.162785, 0.948810), FVector(-8.079758, -3.471313, -16.818888), FVector(-1.000000, 1.000000, -1.000000)), 1.0f, -1038680728);
+	Parameters.ExpectedPoints.Last().Color = FVector4(0.010000, 0.100000, 0.140000, 0.120000);
+	Parameters.ExpectedPoints.Emplace(FTransform(FQuat(-0.264611, -0.622081, 0.334511, 0.656581), FVector(-4.066812, -2.574331, -1.479846), FVector(1.000000, 1.000000, 1.000000)), 1.0f, -226535144);
+	Parameters.ExpectedPoints.Last().Color = FVector4(0.160000, 0.030000, 0.100000, 0.600000);
+	Parameters.ExpectedPoints.Emplace(FTransform(FQuat(0.180702, -0.628874, 0.065701, 0.753357), FVector(-9.408575, -4.647380, -10.272812), FVector(-1.000000, 1.000000, -1.000000)), 1.0f, -824479449);
+	Parameters.ExpectedPoints.Last().Color = FVector4(0.040000, 0.020000, 0.070000, 0.240000);
+	Parameters.ExpectedPoints.Emplace(FTransform(FQuat(-0.249636, -0.634935, 0.347541, 0.643237), FVector(-2.604210, 0.505612, 3.143437), FVector(-1.000000, 1.000000, -1.000000)), 1.0f, -406069811);
+	Parameters.ExpectedPoints.Last().Color = FVector4(0.280000, 0.120000, 0.900000, 0.900000);
+	Parameters.ExpectedPoints.Emplace(FTransform(FQuat(0.174866, -0.648235, 0.054540, 0.739079), FVector(-11.082370, -6.017904, -15.594473), FVector(1.000000, 1.000000, 1.000000)), 1.0f, -613419534);
+	Parameters.ExpectedPoints.Last().Color = FVector4(0.070000, 0.080000, 0.630000, 0.360000);
+
+	Parameters.SourceData = SourceData;
+	Parameters.TargetData = TargetData;
+
+	return RunTestInternal(Parameters);
+}
 
-	bool bTestPassed = true;
+bool FPCGCopyPointsTest_Source::RunTest(const FString&)
+{
+	FPCGCopyPointsTestBase::FParameters Parameters;
+	Parameters.RotationInheritance = EPCGCopyPointsInheritanceMode::Source;
+	Parameters.ScaleInheritance = EPCGCopyPointsInheritanceMode::Source;
+	Parameters.ColorInheritance = EPCGCopyPointsInheritanceMode::Source;
+	Parameters.SeedInheritance = EPCGCopyPointsInheritanceMode::Source;
+	Parameters.AttributeInheritance = EPCGCopyPointsMetadataInheritanceMode::SourceFirst;
+	Parameters.SourceData = PCGTestsCommon::CreateRandomPointData(10, 42);
+	Parameters.TargetData = PCGTestsCommon::CreateRandomPointData(4, 43);
+
+	return RunTestInternal(Parameters);
+}
+
+bool FPCGCopyPointsTest_Target::RunTest(const FString&)
+{
+	FPCGCopyPointsTestBase::FParameters Parameters;
+	Parameters.RotationInheritance = EPCGCopyPointsInheritanceMode::Target;
+	Parameters.ScaleInheritance = EPCGCopyPointsInheritanceMode::Target;
+	Parameters.ColorInheritance = EPCGCopyPointsInheritanceMode::Target;
+	Parameters.SeedInheritance = EPCGCopyPointsInheritanceMode::Target;
+	Parameters.AttributeInheritance = EPCGCopyPointsMetadataInheritanceMode::TargetFirst;
+	Parameters.SourceData = PCGTestsCommon::CreateRandomPointData(10, 42);
+	Parameters.TargetData = PCGTestsCommon::CreateRandomPointData(4, 43);
+
+	return RunTestInternal(Parameters);
+}
+
+bool FPCGCopyPointsTest_EmptySource::RunTest(const FString&)
+{
+	FPCGCopyPointsTestBase::FParameters Parameters;
+	Parameters.SourceData = PCGTestsCommon::CreateEmptyPointData();
+	Parameters.TargetData = PCGTestsCommon::CreateRandomPointData(4, 43);
+
+	return RunTestInternal(Parameters);
+}
+
+bool FPCGCopyPointsTest_EmptyTarget::RunTest(const FString&)
+{
+	FPCGCopyPointsTestBase::FParameters Parameters;
+	Parameters.SourceData = PCGTestsCommon::CreateRandomPointData(10, 42);
+	Parameters.TargetData = PCGTestsCommon::CreateEmptyPointData();
 
-	// Test 1 - inherit from transform multiplication
-	Settings->RotationInheritance = EPCGCopyPointsInheritanceMode::Relative;
-	Settings->ScaleInheritance = EPCGCopyPointsInheritanceMode::Relative;
-	Settings->ColorInheritance = EPCGCopyPointsInheritanceMode::Relative;
-	Settings->SeedInheritance = EPCGCopyPointsInheritanceMode::Relative;
-	bTestPassed &= ValidateCopyPointsPoints();
-
-	// Test 2 - inherit from source points
-	Settings->RotationInheritance = EPCGCopyPointsInheritanceMode::Source;
-	Settings->ScaleInheritance = EPCGCopyPointsInheritanceMode::Source;
-	Settings->ColorInheritance = EPCGCopyPointsInheritanceMode::Source;
-	Settings->SeedInheritance = EPCGCopyPointsInheritanceMode::Source;
-	Settings->AttributeInheritance = EPCGCopyPointsMetadataInheritanceMode::SourceFirst;
-	bTestPassed &= ValidateCopyPointsPoints();
-
-	// Test 3 - inherit from target points
-	Settings->RotationInheritance = EPCGCopyPointsInheritanceMode::Target;
-	Settings->ScaleInheritance = EPCGCopyPointsInheritanceMode::Target;
-	Settings->ColorInheritance = EPCGCopyPointsInheritanceMode::Target;
-	Settings->SeedInheritance = EPCGCopyPointsInheritanceMode::Target;
-	Settings->AttributeInheritance = EPCGCopyPointsMetadataInheritanceMode::TargetFirst;
-	bTestPassed &= ValidateCopyPointsPoints();
-
-	// Test 4 - empty source data
-	SourceTaggedData.Data = EmptyData;
-	bTestPassed &= ValidateCopyPointsPoints();
-	SourceTaggedData.Data = SourceData;
-
-	// Test 5 - empty target data
-	TargetTaggedData.Data = EmptyData;
-	bTestPassed &= ValidateCopyPointsPoints();
-	TargetTaggedData.Data = TargetData;
-
-	return bTestPassed;
+	return RunTestInternal(Parameters);
 }
 
 #endif // WITH_EDITOR
