Replaced manual rooting by FCGObject usage and updated the way we manage resources in PCGContext

diff --git a/Engine/Plugins/Experimental/PCG/Source/PCG/Private/Elements/PCGExecuteBlueprint.cpp b/Engine/Plugins/Experimental/PCG/Source/PCG/Private/Elements/PCGExecuteBlueprint.cpp
--- a/Engine/Plugins/Experimental/PCG/Source/PCG/Private/Elements/PCGExecuteBlueprint.cpp
+++ b/Engine/Plugins/Experimental/PCG/Source/PCG/Private/Elements/PCGExecuteBlueprint.cpp
@@ -687,21 +687,12 @@ bool FPCGExecuteBlueprintElement::ExecuteInternal(FPCGContext* InContext) const
 			}
 		}
 
-		// Since the ExecuteWithContext method is a blueprint call, it will perform a copy of the context
-		// in a local variable. To prevent the overridden settings from being unrooted when the call finishes,
-		// we'll mark it temporarily off here. Note that for the same reason, the context is actually sliced
-		// so there should never be any members in the BP element context that are visible/accessible from blueprint
-		const bool bShouldUnrootSettingsOnDelete = Context->bShouldUnrootSettingsOnDelete;
-		Context->bShouldUnrootSettingsOnDelete = false;
-
+		// Note that the context is actually sliced so there should never be any members in the BP element context that are visible/accessible from blueprint
 		/** Finally, execute the actual blueprint */
 		Context->BlueprintElementInstance->SetCurrentContext(Context);
 		Context->BlueprintElementInstance->ExecuteWithContext(*Context, Context->InputData, Context->OutputData);
 		Context->BlueprintElementInstance->SetCurrentContext(nullptr);
 
-		// Put back the proper unroot flag
-		Context->bShouldUnrootSettingsOnDelete = bShouldUnrootSettingsOnDelete;
-
 		// Log info on outputs
 		for (int32 OutputIndex = 0; OutputIndex < Context->OutputData.TaggedData.Num(); ++OutputIndex)
 		{
@@ -884,12 +875,11 @@ void UPCGBlueprintElement::IterationLoop(FPCGContext& InContext, int64 NumIterat
 	});
 }
 
-FPCGBlueprintExecutionContext::~FPCGBlueprintExecutionContext()
+void FPCGBlueprintExecutionContext::AddExtraStructReferencedObjects(FReferenceCollector& Collector)
 {
 	if (BlueprintElementInstance)
 	{
-		BlueprintElementInstance->RemoveFromRoot();
-		BlueprintElementInstance = nullptr;
+		Collector.AddReferencedObject(BlueprintElementInstance);
 	}
 }
 
@@ -905,7 +895,7 @@ FPCGContext* FPCGExecuteBlueprintElement::Initialize(const FPCGDataCollection& I
 	if (Settings && Settings->BlueprintElementInstance)
 	{
 		Context->BlueprintElementInstance = CastChecked<UPCGBlueprintElement>(StaticDuplicateObject(Settings->BlueprintElementInstance, GetTransientPackage(), FName()));
-		Context->BlueprintElementInstance->AddToRoot();
+
 #if !WITH_EDITOR
 		if (SourceComponent.IsValid() && SourceComponent->GetOwner())
 		{
diff --git a/Engine/Plugins/Experimental/PCG/Source/PCG/Private/Elements/PCGIndirectionElement.cpp b/Engine/Plugins/Experimental/PCG/Source/PCG/Private/Elements/PCGIndirectionElement.cpp
--- a/Engine/Plugins/Experimental/PCG/Source/PCG/Private/Elements/PCGIndirectionElement.cpp
+++ b/Engine/Plugins/Experimental/PCG/Source/PCG/Private/Elements/PCGIndirectionElement.cpp
@@ -168,21 +168,18 @@ bool FPCGIndirectionElement::CanExecuteOnlyOnMainThread(FPCGContext* InContext)
 
 FPCGIndirectionContext::~FPCGIndirectionContext()
 {
-	if (bNeedsToUnrootInnerSettings)
-	{
-		const UPCGIndirectionSettings* Settings = GetInputSettings<UPCGIndirectionSettings>();
-		check(Settings);
-
-		if (UPCGSettings* InnerSettings = Settings->Settings.Get())
-		{
-			InnerSettings->RemoveFromRoot();
-		}
-	}
-
 	delete InnerContext;
 	InnerContext = nullptr;
 }
 
+void FPCGIndirectionContext::AddExtraStructReferencedObjects(FReferenceCollector& Collector)
+{
+	if (InnerSettings)
+	{
+		Collector.AddReferencedObject(InnerSettings);
+	}
+}
+
 FPCGContext* FPCGIndirectionElement::CreateContext()
 {
 	return new FPCGIndirectionContext();
@@ -196,7 +193,9 @@ bool FPCGIndirectionElement::PrepareDataInternal(FPCGContext* InContext) const
 	const UPCGIndirectionSettings* Settings = Context->GetInputSettings<UPCGIndirectionSettings>();
 	check(Settings);
 
-	if (UPCGSettings* InnerSettings = Settings->Settings.LoadSynchronous())
+	Context->InnerSettings = Settings->Settings.LoadSynchronous();
+
+	if(UPCGSettings* InnerSettings = Context->InnerSettings)
 	{
 		if (Settings->ProxyInterfaceMode == EPCGProxyInterfaceMode::ByNativeElement)
 		{
@@ -220,14 +219,6 @@ bool FPCGIndirectionElement::PrepareDataInternal(FPCGContext* InContext) const
 			}
 		} 
 
-		// TODO: while we can root it here, if this node or multiple indirection node were to execute in parallel,
-		// the lifetime of the inner settings would be a bit unclear - we need to improve this.
-		if (!InnerSettings->IsRooted())
-		{
-			InnerSettings->AddToRoot();
-			Context->bNeedsToUnrootInnerSettings = true;
-		}
-
 		Context->InnerElement = InnerSettings->GetElement();
 		check(Context->InnerElement);
 
diff --git a/Engine/Plugins/Experimental/PCG/Source/PCG/Private/Elements/PCGStaticMeshSpawner.cpp b/Engine/Plugins/Experimental/PCG/Source/PCG/Private/Elements/PCGStaticMeshSpawner.cpp
--- a/Engine/Plugins/Experimental/PCG/Source/PCG/Private/Elements/PCGStaticMeshSpawner.cpp
+++ b/Engine/Plugins/Experimental/PCG/Source/PCG/Private/Elements/PCGStaticMeshSpawner.cpp
@@ -236,16 +236,10 @@ bool FPCGStaticMeshSpawnerElement::PrepareDataInternal(FPCGContext* InContext) c
 
 			if (Settings->InstanceDataPackerParameters)
 			{
-				// This blueprint overridden implementation will create an implicit copy of the context
-				// and when the copy destructs, will cause the SettingsOverride to unroot and mark for GC
-				const bool bShouldUnrootSettingsOnDelete = Context->bShouldUnrootSettingsOnDelete;
-				Context->bShouldUnrootSettingsOnDelete = false;
 				for (int32 InstanceListIndex = 0; InstanceListIndex < MeshInstances.Num(); ++InstanceListIndex)
 				{
 					Settings->InstanceDataPackerParameters->PackInstances(*Context, Context->CurrentPointData, MeshInstances[InstanceListIndex], PackedCustomData[InstanceListIndex]);
 				}
-				// Put back the proper unroot flag
-				Context->bShouldUnrootSettingsOnDelete = bShouldUnrootSettingsOnDelete;
 			}
 		}
 
diff --git a/Engine/Plugins/Experimental/PCG/Source/PCG/Private/Elements/PCGSurfaceSampler.cpp b/Engine/Plugins/Experimental/PCG/Source/PCG/Private/Elements/PCGSurfaceSampler.cpp
--- a/Engine/Plugins/Experimental/PCG/Source/PCG/Private/Elements/PCGSurfaceSampler.cpp
+++ b/Engine/Plugins/Experimental/PCG/Source/PCG/Private/Elements/PCGSurfaceSampler.cpp
@@ -353,8 +353,7 @@ namespace PCGSurfaceSamplerHelpers
 			OutState.BoundingShape = FindBoundingShape(Context, bUnionWasCreated);
 			if (OutState.BoundingShape && bUnionWasCreated)
 			{
-				// Must cast away const to root the union object
-				Context->RootAndTrackObject(const_cast<UPCGSpatialData*>(OutState.BoundingShape));
+				Context->TrackObject(OutState.BoundingShape);
 			}
 
 			if (OutState.BoundingShape)
diff --git a/Engine/Plugins/Experimental/PCG/Source/PCG/Private/Elements/PCGVolumeSampler.cpp b/Engine/Plugins/Experimental/PCG/Source/PCG/Private/Elements/PCGVolumeSampler.cpp
--- a/Engine/Plugins/Experimental/PCG/Source/PCG/Private/Elements/PCGVolumeSampler.cpp
+++ b/Engine/Plugins/Experimental/PCG/Source/PCG/Private/Elements/PCGVolumeSampler.cpp
@@ -200,7 +200,7 @@ namespace PCGVolumeSamplerHelpers
 			OutState.BoundingShape = Context->InputData.GetSpatialUnionOfInputsByPin(PCGVolumeSamplerConstants::BoundingShapeLabel, bUnionWasCreated);
 			if (OutState.BoundingShape && bUnionWasCreated)
 			{
-				Context->RootAndTrackObject(const_cast<UPCGSpatialData*>(OutState.BoundingShape));
+				Context->TrackObject(OutState.BoundingShape);
 			}
 
 			if (!OutState.BoundingShape && Context->SourceComponent.IsValid())
diff --git a/Engine/Plugins/Experimental/PCG/Source/PCG/Private/Graph/PCGGraphCache.cpp b/Engine/Plugins/Experimental/PCG/Source/PCG/Private/Graph/PCGGraphCache.cpp
--- a/Engine/Plugins/Experimental/PCG/Source/PCG/Private/Graph/PCGGraphCache.cpp
+++ b/Engine/Plugins/Experimental/PCG/Source/PCG/Private/Graph/PCGGraphCache.cpp
@@ -27,12 +27,11 @@ static TAutoConsoleVariable<bool> CVarCacheMemoryBudgetEnabled(
 // Initial max number of entries graph cache
 static const int32 GPCGGraphCacheInitialCapacity = 65536;
 
-FPCGGraphCache::FPCGGraphCache(TWeakObjectPtr<UObject> InOwner, FPCGRootSet* InRootSet)
+FPCGGraphCache::FPCGGraphCache(TWeakObjectPtr<UObject> InOwner)
 	: CacheData(GPCGGraphCacheInitialCapacity)
 	, Owner(InOwner)
-	, RootSet(InRootSet)
 {
-	check(InOwner.Get() && InRootSet);
+	check(InOwner.Get());
 }
 
 FPCGGraphCache::~FPCGGraphCache()
@@ -103,8 +102,6 @@ void FPCGGraphCache::StoreInCache(const IPCGElement* InElement, const FPCGCrc& I
 		FPCGCacheEntryKey CacheKey(InElement, InDependenciesCrc);
 		CacheData.Add(CacheKey, InOutput);
 		
-		InOutput.AddToRootSet(*RootSet);
-
 		AddDataToAccountedMemory(InOutput);
 	}
 }
@@ -113,12 +110,6 @@ void FPCGGraphCache::ClearCache()
 {
 	FWriteScopeLock ScopedWriteLock(CacheLock);
 
-	// Unroot all previously rooted data
-	for (FPCGDataCollection CacheEntry : CacheData)
-	{
-		CacheEntry.RemoveFromRootSet(*RootSet);
-	}
-
 	MemoryRecords.Empty();
 	TotalMemoryUsed = 0;
 
@@ -147,7 +138,6 @@ bool FPCGGraphCache::EnforceMemoryBudget()
 		while (TotalMemoryUsed > MemoryBudget && CacheData.Num() > 0)
 		{
 			FPCGDataCollection RemovedData = CacheData.RemoveLeastRecent();
-			RemovedData.RemoveFromRootSet(*RootSet);
 			RemoveFromMemoryTotal(RemovedData);
 		}
 	}
@@ -184,8 +174,6 @@ void FPCGGraphCache::CleanFromCache(const IPCGElement* InElement, const UPCGSett
 			if (const FPCGDataCollection* Data = CacheData.Find(Key))
 			{
 				RemoveFromMemoryTotal(*Data);
-				Data->RemoveFromRootSet(*RootSet);
-				
 				CacheData.Remove(Key);
 			}
 		}
@@ -214,6 +202,17 @@ uint32 FPCGGraphCache::GetGraphCacheEntryCount(IPCGElement* InElement) const
 }
 #endif // WITH_EDITOR
 
+void FPCGGraphCache::AddReferencedObjects(FReferenceCollector& Collector)
+{
+	TRACE_CPUPROFILER_EVENT_SCOPE(FPCGGraphCache::AddReferencedObjects);
+	FReadScopeLock ScopedReadLock(CacheLock);
+
+	for (FPCGDataCollection& CacheEntry : CacheData)
+	{
+		CacheEntry.AddReferences(Collector);
+	}
+}
+
 void FPCGGraphCache::GrowCache_Unsafe()
 {
 	TRACE_CPUPROFILER_EVENT_SCOPE(FPCGGraphCache::GrowCache_Unsafe);
diff --git a/Engine/Plugins/Experimental/PCG/Source/PCG/Private/Graph/PCGGraphCache.h b/Engine/Plugins/Experimental/PCG/Source/PCG/Private/Graph/PCGGraphCache.h
--- a/Engine/Plugins/Experimental/PCG/Source/PCG/Private/Graph/PCGGraphCache.h
+++ b/Engine/Plugins/Experimental/PCG/Source/PCG/Private/Graph/PCGGraphCache.h
@@ -6,6 +6,7 @@
 #include "PCGData.h"
 
 #include "Containers/LruCache.h"
+#include "UObject/GCObject.h"
 
 class IPCGElement;
 class UPCGComponent;
@@ -60,9 +61,9 @@ private:
 * In cases where we have some subgraph reuse. Under that premise, we can then
 * instead store by element, as we will never recreate elements (except for arbitrary tasks)
 */
-struct FPCGGraphCache
+struct FPCGGraphCache : public FGCObject
 {
-	FPCGGraphCache(TWeakObjectPtr<UObject> InOwner, FPCGRootSet* InRootSet);
+	explicit FPCGGraphCache(TWeakObjectPtr<UObject> InOwner);
 	~FPCGGraphCache();
 
 	/** Returns true if data was found from the cache, in which case the outputs are written in OutOutput. InNode is optional and for logging only. */
@@ -88,6 +89,11 @@ struct FPCGGraphCache
 	uint32 GetGraphCacheEntryCount(IPCGElement* InElement) const;
 #endif
 
+	//~Begin FGCObject interface
+	virtual void AddReferencedObjects(FReferenceCollector& Collector) override;
+	virtual FString GetReferencerName() const override { return TEXT("FPCGGraphCache"); }
+	//~End FGCObject interface
+
 private:
 	// Grow cache entry capacity, preserving current entries.
 	void GrowCache_Unsafe();
@@ -103,9 +109,6 @@ private:
 	TLruCache<FPCGCacheEntryKey, FPCGDataCollection> CacheData;
 	TWeakObjectPtr<UObject> Owner = nullptr;
 
-	/** To prevent garbage collection on data in the cache, we'll need to root some data */
-	FPCGRootSet* RootSet = nullptr;
-
 	/** Map from data UIDs to records. Provides ref counting and caches memory size. */
 	TMap<uint64, FCachedMemoryRecord> MemoryRecords;
 
diff --git a/Engine/Plugins/Experimental/PCG/Source/PCG/Private/Graph/PCGGraphExecutor.cpp b/Engine/Plugins/Experimental/PCG/Source/PCG/Private/Graph/PCGGraphExecutor.cpp
--- a/Engine/Plugins/Experimental/PCG/Source/PCG/Private/Graph/PCGGraphExecutor.cpp
+++ b/Engine/Plugins/Experimental/PCG/Source/PCG/Private/Graph/PCGGraphExecutor.cpp
@@ -69,7 +69,7 @@ static TAutoConsoleVariable<bool> CVarStripEmptyPointData(
 
 FPCGGraphExecutor::FPCGGraphExecutor(UObject* InOwner)
 	: GraphCompiler(MakeUnique<FPCGGraphCompiler>())
-	, GraphCache(InOwner, &DataRootSet)
+	, GraphCache(InOwner)
 #if WITH_EDITOR
 	, GenerationProgressNotification(GetNotificationTextFormat())
 #endif
@@ -81,7 +81,6 @@ FPCGGraphExecutor::~FPCGGraphExecutor()
 	// We don't really need to do this here (it would be done in the destructor of these both)
 	// but this is to clarify/ensure the order in which this happens
 	GraphCache.ClearCache();
-	DataRootSet.Clear();
 
 #if WITH_EDITOR
 	// Cleanup + clear notification
@@ -720,8 +719,6 @@ void FPCGGraphExecutor::Execute()
 					{
 						ActiveTask.bIsBypassed = true;
 						ActiveTask.Context->OutputData = CachedOutput;
-						// Since we've copied the output data, we need to make sure to add count ref in the root set since it'll be removed once the task is executed
-						ActiveTask.Context->OutputData.AddToRootSet(DataRootSet);
 					}
 #endif
 
@@ -779,8 +776,6 @@ void FPCGGraphExecutor::Execute()
 					check(!ActiveTask.Element->CanExecuteOnlyOnMainThread(ActiveTask.Context.Get()));
 					ActiveTask.Context->AsyncState.EndTime = EndTime;
 					ActiveTask.Context->AsyncState.bIsRunningOnMainThread = false;
-					// Remove the precreated data so we properly count in the root set
-					ActiveTask.Context->OutputData.RemoveFromRootSet(DataRootSet);
 
 					Futures.Emplace(ExecutionIndex, Async(EAsyncExecution::ThreadPool, [&ActiveTask]()
 					{
@@ -857,13 +852,6 @@ void FPCGGraphExecutor::Execute()
 			QueueNextTasks(ActiveTask.NodeId);
 			bAnyTaskEnded = true;
 
-			// Un-root temporary input data
-			if (FPCGDataCollection* TemporaryInput = InputTemporaryData.Find(ActiveTask.NodeId))
-			{
-				TemporaryInput->RemoveFromRootSet(DataRootSet);
-				InputTemporaryData.Remove(ActiveTask.NodeId);
-			}
-
 			// Remove current active task from list
 			ActiveTasks.RemoveAtSwap(TaskIndex);
 		};
@@ -879,8 +867,6 @@ void FPCGGraphExecutor::Execute()
 				check(!MainThreadTask.Context->bIsPaused);
 				MainThreadTask.Context->AsyncState.EndTime = EndTime;
 				MainThreadTask.Context->AsyncState.bIsRunningOnMainThread = true;
-				// Remove the precreated data so we properly count in the root set
-				MainThreadTask.Context->OutputData.RemoveFromRootSet(DataRootSet);
 
 #if WITH_EDITOR
 				if(MainThreadTask.bIsBypassed || MainThreadTask.bWasCancelled || MainThreadTask.Element->Execute(MainThreadTask.Context.Get()))
@@ -910,11 +896,6 @@ void FPCGGraphExecutor::Execute()
 					{
 						PostTaskExecute(ExecutionIndex);
 					}
-					else
-					{
-						// Task isn't done, but we need to make sure the data isn't garbage collected
-						ActiveTasks[ExecutionIndex].Context->OutputData.AddToRootSet(DataRootSet);
-					}
 				}
 			}
 
@@ -922,11 +903,6 @@ void FPCGGraphExecutor::Execute()
 			{
 				PostTaskExecute(0);
 			}
-			else if (!ActiveTasks.IsEmpty())
-			{
-				// Task isn't done, make sure we don't garbage collected the precreated data
-				ActiveTasks[0].Context->OutputData.AddToRootSet(DataRootSet);
-			}
 		}
 
 		// Any paused tasks at that point should relinquish their resources
@@ -1238,17 +1214,12 @@ void FPCGGraphExecutor::CombineParams(FPCGTaskId InTaskId, FPCGDataCollection& I
 		}
 
 		// Add to the root set since we created a new object, that needs to be kept alive for the duration of the task.
-		DataRootSet.Add(CombinedParamData);
 		FPCGTaggedData CombineParams{};
 		CombineParams.Data = CombinedParamData;
 		CombineParams.Pin = PCGPinConstants::DefaultParamsLabel;
 		TempTaggedData.Add(CombineParams);
 
 		InTaskInput.TaggedData = std::move(TempTaggedData);
-
-		// Also store it into the TemporaryMap to track it down
-		FPCGDataCollection& TemporaryInput = InputTemporaryData.FindOrAdd(InTaskId);
-		TemporaryInput.TaggedData.Add(CombineParams);
 	}
 }
 
@@ -1258,9 +1229,6 @@ void FPCGGraphExecutor::StoreResults(FPCGTaskId InTaskId, const FPCGDataCollecti
 
 	// Store output in map
 	OutputData.Add(InTaskId, InTaskOutput);
-
-	// Root any non-rooted results, otherwise they'll get garbage-collected
-	InTaskOutput.AddToRootSet(DataRootSet);
 }
 
 void FPCGGraphExecutor::ClearResults()
@@ -1273,19 +1241,36 @@ void FPCGGraphExecutor::ClearResults()
 		NextTaskId = 0;
 	}
 
-	for (const TPair<FPCGTaskId, FPCGDataCollection>& TaskTemporaryInput : InputTemporaryData)
+	OutputData.Reset();
+
+	ScheduleLock.Unlock();
+}
+
+void FPCGGraphExecutor::AddReferencedObjects(FReferenceCollector& Collector)
+{
+	TRACE_CPUPROFILER_EVENT_SCOPE(FPCGGraphExecutor::AddReferencedObjects);
+
+	// Go through all data in the cached output map
+	for (auto& OutputDataEntry : OutputData)
 	{
-		TaskTemporaryInput.Value.RemoveFromRootSet(DataRootSet);
+		OutputDataEntry.Value.AddReferences(Collector);
 	}
-	InputTemporaryData.Reset();
 
-	for (const TPair<FPCGTaskId, FPCGDataCollection>& TaskOutput : OutputData)
+	// Go through ready tasks, active tasks and sleeping tasks contexts
+	auto AddReferences = [&Collector](auto& TaskContainer)
 	{
-		TaskOutput.Value.RemoveFromRootSet(DataRootSet);
-	}
-	OutputData.Reset();
+		for (auto& Task : TaskContainer)
+		{
+			if (Task.Context)
+			{
+				Task.Context->AddStructReferencedObjects(Collector);
+			}
+		}
+	};
 
-	ScheduleLock.Unlock();
+	AddReferences(ReadyTasks);
+	AddReferences(ActiveTasks);
+	AddReferences(SleepingTasks);
 }
 
 FPCGElementPtr FPCGGraphExecutor::GetFetchInputElement()
diff --git a/Engine/Plugins/Experimental/PCG/Source/PCG/Private/Graph/PCGGraphExecutor.h b/Engine/Plugins/Experimental/PCG/Source/PCG/Private/Graph/PCGGraphExecutor.h
--- a/Engine/Plugins/Experimental/PCG/Source/PCG/Private/Graph/PCGGraphExecutor.h
+++ b/Engine/Plugins/Experimental/PCG/Source/PCG/Private/Graph/PCGGraphExecutor.h
@@ -8,6 +8,8 @@
 #include "Graph/PCGGraphCache.h"
 #include "Graph/PCGStackContext.h"
 
+#include "UObject/GCObject.h"
+
 #if WITH_EDITOR
 #include "AsyncCompilationHelpers.h"
 #include "WorldPartition/WorldPartitionHandle.h" // Needed for FWorldPartitionReference
@@ -82,10 +84,10 @@ struct FPCGGraphActiveTask
 	TSharedPtr<const FPCGStackContext> StackContext;
 };
 
-class FPCGGraphExecutor
+class FPCGGraphExecutor : public FGCObject
 {
 public:
-	FPCGGraphExecutor(UObject* InOwner);
+	explicit FPCGGraphExecutor(UObject* InOwner);
 	~FPCGGraphExecutor();
 
 	/** Compile (and cache) a graph for later use. This call is threadsafe */
@@ -155,6 +157,11 @@ public:
 	/** True if graph cache debugging is enabled. */
 	bool IsGraphCacheDebuggingEnabled() const { return GraphCache.IsDebuggingEnabled(); }
 
+	//~Begin FGCObject interface
+	virtual void AddReferencedObjects(FReferenceCollector& Collector) override;
+	virtual FString GetReferencerName() const override { return TEXT("FPCGGraphExecutor"); }
+	//~End FGCObject interface
+
 private:
 	TSet<UPCGComponent*> Cancel(TFunctionRef<bool(TWeakObjectPtr<UPCGComponent>)> CancelFilter);
 	void ClearAllTasks();
@@ -180,9 +187,6 @@ private:
 	/** Graph compiler that turns a graph into tasks */
 	TUniquePtr<FPCGGraphCompiler> GraphCompiler;
 
-	/** Rootset to hold temporary results (for current computation) + graph cache */
-	FPCGRootSet DataRootSet;
-
 	/** Graph results cache */
 	FPCGGraphCache GraphCache;
 
@@ -200,8 +204,6 @@ private:
 	/** Map of node instances to their output, could be cleared once execution is done */
 	/** Note: this should at some point unload based on loaded/unloaded proxies, otherwise memory cost will be unbounded */
 	TMap<FPCGTaskId, FPCGDataCollection> OutputData;
-	/** Map of node instances to their temporary input, could be cleared once execution is done */
-	TMap<FPCGTaskId, FPCGDataCollection> InputTemporaryData;
 	/** Monotonically increasing id. Should be reset once all tasks are executed, should be protected by the ScheduleLock */
 	FPCGTaskId NextTaskId = 0;
 
diff --git a/Engine/Plugins/Experimental/PCG/Source/PCG/Private/PCGContext.cpp b/Engine/Plugins/Experimental/PCG/Source/PCG/Private/PCGContext.cpp
--- a/Engine/Plugins/Experimental/PCG/Source/PCG/Private/PCGContext.cpp
+++ b/Engine/Plugins/Experimental/PCG/Source/PCG/Private/PCGContext.cpp
@@ -22,16 +22,6 @@
 
 #define LOCTEXT_NAMESPACE "PCGContext"
 
-FPCGContext::~FPCGContext()
-{
-	if (SettingsWithOverride && bShouldUnrootSettingsOnDelete)
-	{
-		SettingsWithOverride->RemoveFromRoot();
-		SettingsWithOverride->MarkAsGarbage();
-		SettingsWithOverride = nullptr;
-	}
-}
-
 FString FPCGContext::GetTaskName() const
 {
 	if (Node)
@@ -132,7 +122,6 @@ void FPCGContext::InitializeSettings()
 			{
 				SettingsWithOverride = Cast<UPCGSettings>(StaticDuplicateObject(NodeSettings, GetTransientPackage()));
 				SettingsWithOverride->SetFlags(RF_Transient);
-				SettingsWithOverride->AddToRoot();
 
 				// Force seed copy to prevent issue due to delta serialization vs. Seed being initialized in the constructor only for new nodes
 				SettingsWithOverride->Seed = NodeSettings->Seed;
@@ -302,4 +291,17 @@ bool FPCGContext::HasVisualLogs() const
 }
 #endif // WITH_EDITOR
 
+void FPCGContext::AddStructReferencedObjects(FReferenceCollector& Collector)
+{
+	InputData.AddReferences(Collector);
+	OutputData.AddReferences(Collector);
+
+	if (SettingsWithOverride)
+	{
+		Collector.AddReferencedObject(SettingsWithOverride);
+	}
+
+	AddExtraStructReferencedObjects(Collector);
+}
+
 #undef LOCTEXT_NAMESPACE
\ No newline at end of file
diff --git a/Engine/Plugins/Experimental/PCG/Source/PCG/Private/PCGData.cpp b/Engine/Plugins/Experimental/PCG/Source/PCG/Private/PCGData.cpp
--- a/Engine/Plugins/Experimental/PCG/Source/PCG/Private/PCGData.cpp
+++ b/Engine/Plugins/Experimental/PCG/Source/PCG/Private/PCGData.cpp
@@ -21,65 +21,6 @@ static TAutoConsoleVariable<bool> CVarCachePropagateCrcThroughBooleanData(
 	false,
 	TEXT("Whether intersection, union, difference combine Crc values from operands. If false they fall back to using data UID."));
 
-void FPCGRootSet::Clear()
-{
-	for (TPair<UObject*, int32>& Entry : RootSet)
-	{
-		Entry.Key->RemoveFromRoot();
-	}
-
-	RootSet.Reset();
-}
-
-void FPCGRootSet::Add(UObject* InObject)
-{
-	check(InObject);
-	AddInternal(InObject);
-}
-
-void FPCGRootSet::AddInternal(UObject* InObject)
-{
-	check(InObject && !InObject->IsA<UPackage>());
-
-	if (int32* Found = RootSet.Find(InObject))
-	{
-		(*Found)++;
-	}
-	else if (!InObject->IsRooted() && InObject->GetPackage() == GetTransientPackage())
-	{
-		InObject->AddToRoot();
-		RootSet.Emplace(InObject, 1);
-	}
-}
-
-void FPCGRootSet::Remove(UObject* InObject)
-{
-	if (!InObject)
-	{
-		UE_LOG(LogPCG, Warning, TEXT("Trying to remove a null object from the rootset"));
-		return;
-	}
-
-	RemoveInternal(InObject);
-}
-
-void FPCGRootSet::RemoveInternal(UObject* InObject)
-{
-	check(InObject && !InObject->IsA<UPackage>());
-
-	if (int32* Found = RootSet.Find(InObject))
-	{
-		check(InObject->IsRooted());
-		(*Found)--;
-
-		if (*Found == 0)
-		{
-			InObject->RemoveFromRoot();
-			RootSet.Remove(InObject);
-		}
-	}
-}
-
 UPCGData::UPCGData(const FObjectInitializer& ObjectInitializer)
 	: Super(ObjectInitializer)
 {
@@ -440,24 +381,13 @@ bool FPCGDataCollection::operator!=(const FPCGDataCollection& Other) const
 	return !operator==(Other);
 }
 
-void FPCGDataCollection::AddToRootSet(FPCGRootSet& RootSet) const
+void FPCGDataCollection::AddReferences(FReferenceCollector& Collector)
 {
-	for (const FPCGTaggedData& Data : TaggedData)
-	{
-		if (Data.Data)
-		{
-			RootSet.Add(const_cast<UPCGData*>(Data.Data.Get()));
-		}
-	}
-}
-
-void FPCGDataCollection::RemoveFromRootSet(FPCGRootSet& RootSet) const
-{
-	for (const FPCGTaggedData& Data : TaggedData)
+	for (FPCGTaggedData& Data : TaggedData)
 	{
 		if (Data.Data)
 		{
-			RootSet.Remove(const_cast<UPCGData*>(Data.Data.Get()));
+			Collector.AddReferencedObject(Data.Data);
 		}
 	}
 }
diff --git a/Engine/Plugins/Experimental/PCG/Source/PCG/Private/PCGSubgraph.cpp b/Engine/Plugins/Experimental/PCG/Source/PCG/Private/PCGSubgraph.cpp
--- a/Engine/Plugins/Experimental/PCG/Source/PCG/Private/PCGSubgraph.cpp
+++ b/Engine/Plugins/Experimental/PCG/Source/PCG/Private/PCGSubgraph.cpp
@@ -642,7 +642,7 @@ bool FPCGSubgraphElement::ExecuteInternal(FPCGContext* InContext) const
 FPCGInputForwardingElement::FPCGInputForwardingElement(const FPCGDataCollection& InputToForward)
 	: Input(InputToForward)
 {
-	// Root any previously unrooted data, so we don't need to make sure the caller is still around until this is executed
+	// Root any previously unrooted data, needed here because the context does not exist yet and we need to ensure that the input is not garbage collected
 	for (const FPCGTaggedData& TaggedData : Input.TaggedData)
 	{
 		if (TaggedData.Data && !TaggedData.Data->IsRooted())
@@ -656,7 +656,7 @@ FPCGInputForwardingElement::FPCGInputForwardingElement(const FPCGDataCollection&
 
 bool FPCGInputForwardingElement::ExecuteInternal(FPCGContext* Context) const
 {
-	// Remove from rootset during the execution; data will be re-rooted by the normal process in the graph executor
+	// Remove from rootset during the execution if we had previously done so. After execution, the graph cache will keep track of these references
 	for (UPCGData* DataToUnroot : RootedData)
 	{
 		ensure(DataToUnroot->IsRooted());
diff --git a/Engine/Plugins/Experimental/PCG/Source/PCG/Public/Elements/PCGExecuteBlueprint.h b/Engine/Plugins/Experimental/PCG/Source/PCG/Public/Elements/PCGExecuteBlueprint.h
--- a/Engine/Plugins/Experimental/PCG/Source/PCG/Public/Elements/PCGExecuteBlueprint.h
+++ b/Engine/Plugins/Experimental/PCG/Source/PCG/Public/Elements/PCGExecuteBlueprint.h
@@ -382,11 +382,10 @@ protected:
 
 struct FPCGBlueprintExecutionContext : public FPCGContext
 {
-	virtual ~FPCGBlueprintExecutionContext();
-
-	UPCGBlueprintElement* BlueprintElementInstance = nullptr;
+	TObjectPtr<UPCGBlueprintElement> BlueprintElementInstance = nullptr;
 
 protected:
+	virtual void AddExtraStructReferencedObjects(FReferenceCollector& Collector) override;
 	virtual UObject* GetExternalContainerForOverridableParam(const FPCGSettingsOverridableParam& InParam) override { return BlueprintElementInstance; }
 };
 
diff --git a/Engine/Plugins/Experimental/PCG/Source/PCG/Public/Elements/PCGIndirectionElement.h b/Engine/Plugins/Experimental/PCG/Source/PCG/Public/Elements/PCGIndirectionElement.h
--- a/Engine/Plugins/Experimental/PCG/Source/PCG/Public/Elements/PCGIndirectionElement.h
+++ b/Engine/Plugins/Experimental/PCG/Source/PCG/Public/Elements/PCGIndirectionElement.h
@@ -62,12 +62,18 @@ public:
 
 struct FPCGIndirectionContext : public FPCGContext
 {
+public:
 	virtual ~FPCGIndirectionContext();
 
+protected:
+	virtual void AddExtraStructReferencedObjects(FReferenceCollector& Collector) override;
+
+public:
 	FPCGElementPtr InnerElement;
 	FPCGContext* InnerContext = nullptr;
-	bool bNeedsToUnrootInnerSettings = false;
-	bool bShouldActAsPassthrough = false;
+	bool bShouldActAsPassthrough = false;	
+
+	TObjectPtr<UPCGSettings> InnerSettings = nullptr;
 };
 
 class FPCGIndirectionElement : public IPCGElement
diff --git a/Engine/Plugins/Experimental/PCG/Source/PCG/Public/Elements/PCGTimeSlicedElementBase.h b/Engine/Plugins/Experimental/PCG/Source/PCG/Public/Elements/PCGTimeSlicedElementBase.h
--- a/Engine/Plugins/Experimental/PCG/Source/PCG/Public/Elements/PCGTimeSlicedElementBase.h
+++ b/Engine/Plugins/Experimental/PCG/Source/PCG/Public/Elements/PCGTimeSlicedElementBase.h
@@ -14,7 +14,7 @@
  * Override the element type with TTimeSlicedPCGElement with two template arguments, the first being the static struct,and the second being the "per-iteration"
  * Pass a function or lambda matching the correct signature to InitializePerExecutionState and initialize the static struct within
  * Do the same for InitializePerIterationStates, but also pass the number of iterations. This will iterate through a state initialization for each iteration
- * [Optional] If needed, mark UObjects that need to bypass GC with RootAndTrackObject or RootAndTrackObjectByName if you need to retrieve it later. The objects will be garbage collected at the end of the context lifetime
+ * [Optional] If needed, mark UObjects with TrackObject or TrackObjectByName if you need to retrieve it later and to prevent their garbage collection.
  * DataIsPrepared SHOULD BE used to verify initialization was successful, such as in the ExecuteInternal if data was previously initialized in PrepareDataInternal
  * Call ExecuteSlice with an execution function or lambda that returns a boolean that is true once the full execution is completed, or false otherwise
  *
@@ -46,8 +46,6 @@ namespace PCGTimeSlice
 template <typename PerExecutionStateT = PCGTimeSlice::FEmptyStruct, typename PerIterationStateT = PCGTimeSlice::FEmptyStruct>
 struct TPCGTimeSlicedContext : public FPCGContext
 {
-	virtual ~TPCGTimeSlicedContext() override;
-
 	virtual bool TimeSliceIsEnabled() const override final { return bTimeSliceIsEnabled; }
 	void SetTimeSliceIsEnabled(const bool bEnableTimeSlice = true) { bTimeSliceIsEnabled = bEnableTimeSlice; }
 
@@ -104,8 +102,11 @@ struct TPCGTimeSlicedContext : public FPCGContext
 		return true;
 	}
 
-	/** Fire and forget function to root a UObject for the duration of the Context, and then mark for garbage collection at context lifecycle end. */
-	void RootAndTrackObject(UObject* Object);
+	/** Fire and forget function to make sure a UObject will not be GC'ed for the duration of the Context. */
+	void TrackObject(const UObject* Object);
+
+protected:
+	virtual void AddExtraStructReferencedObjects(FReferenceCollector& Collector) override;
 
 private:
 	// Allow exposure to iteration index, etc
@@ -135,7 +136,7 @@ private:
 	TArray<PerIterationStateT> PerIterationStateArray;
 
 	/** Tracks rooted UObjects that need to avoid garbage collection during the lifetime of the context */
-	TArray<UObject*> RootedAndTrackedObjectArray;
+	TArray<TObjectPtr<const UObject>> TrackedObjectArray;
 };
 
 /**
@@ -163,22 +164,9 @@ public:
 };
 
 template <typename PerExecutionStateT, typename PerIterationStateT>
-TPCGTimeSlicedContext<PerExecutionStateT, PerIterationStateT>::~TPCGTimeSlicedContext()
+void TPCGTimeSlicedContext<PerExecutionStateT, PerIterationStateT>::AddExtraStructReferencedObjects(FReferenceCollector& Collector)
 {
-	for (UObject* Object : RootedAndTrackedObjectArray)
-	{
-		// Sanity check
-		check(Object);
-		if (!ensure(IsValid(Object) && Object->IsRooted()))
-		{
-			continue;
-		}
-
-		Object->RemoveFromRoot();
-		Object->MarkAsGarbage();
-	}
-
-	RootedAndTrackedObjectArray.Empty();
+	Collector.AddReferencedObjects(TrackedObjectArray);
 }
 
 template <typename PerExecutionStateT, typename PerIterationStateT>
@@ -232,11 +220,10 @@ const TArray<EPCGTimeSliceInitResult>& TPCGTimeSlicedContext<PerExecutionStateT,
 }
 
 template <typename PerExecutionStateT, typename PerIterationStateT>
-void TPCGTimeSlicedContext<PerExecutionStateT, PerIterationStateT>::RootAndTrackObject(UObject* Object)
+void TPCGTimeSlicedContext<PerExecutionStateT, PerIterationStateT>::TrackObject(const UObject* Object)
 {
-	check(Object && IsValid(Object) && !Object->IsRooted());
-	Object->AddToRoot();
-	RootedAndTrackedObjectArray.AddUnique(Object);
+	check(Object && IsValid(Object));
+	TrackedObjectArray.AddUnique(Object);
 }
 
 template <typename PerExecutionStateT, typename PerIterationStateT>
diff --git a/Engine/Plugins/Experimental/PCG/Source/PCG/Public/PCGContext.h b/Engine/Plugins/Experimental/PCG/Source/PCG/Public/PCGContext.h
--- a/Engine/Plugins/Experimental/PCG/Source/PCG/Public/PCGContext.h
+++ b/Engine/Plugins/Experimental/PCG/Source/PCG/Public/PCGContext.h
@@ -6,6 +6,8 @@
 #include "PCGNode.h" // IWYU pragma: keep
 #include "Helpers/PCGAsyncState.h"
 
+#include "UObject/GCObject.h"
+
 #include "PCGContext.generated.h"
 
 class UPCGComponent;
@@ -47,7 +49,7 @@ struct PCG_API FPCGContext
 {
 	GENERATED_BODY()
 
-	virtual ~FPCGContext();
+	virtual ~FPCGContext() {}
 
 	FPCGDataCollection InputData;
 	FPCGDataCollection OutputData;
@@ -62,9 +64,6 @@ struct PCG_API FPCGContext
 	FPCGTaskId CompiledTaskId = InvalidPCGTaskId;
 	bool bIsPaused = false;
 
-	// Used to prevent settings override being deleted, needs to be false when going through blueprint calls with a context
-	bool bShouldUnrootSettingsOnDelete = true;
-
 	EPCGExecutionPhase CurrentPhase = EPCGExecutionPhase::NotExecuted;
 	int32 BypassedOutputCount = 0;
 
@@ -112,9 +111,15 @@ struct PCG_API FPCGContext
 	bool HasVisualLogs() const;
 #endif // WITH_EDITOR
 
+	/** Gathers references to objects to prevent them from being garbage collected. Extend resources to be collected with AddExtractStructReferencedObjects. */
+	// Implementation note: this is NOT the same as a struct using WithAddStructReferencedObjects since we are not holding contexts
+	//  in properties in any case. This will be called from the graph executor when needed and is implemented to look like normal reference traversal.
+	void AddStructReferencedObjects(FReferenceCollector& Collector);
+
 protected:
 	virtual UObject* GetExternalContainerForOverridableParam(const FPCGSettingsOverridableParam& InParam) { return nullptr; }
 	virtual void* GetUnsafeExternalContainerForOverridableParam(const FPCGSettingsOverridableParam& InParam) { return nullptr; }
+	virtual void AddExtraStructReferencedObjects(FReferenceCollector& Collector) {}
 
 private:
 	template<typename SettingsType>
@@ -128,4 +133,4 @@ private:
 
 	// List of params that were in effect overriden
 	TArray<const FPCGSettingsOverridableParam*> OverriddenParams;
-};
+};
\ No newline at end of file
diff --git a/Engine/Plugins/Experimental/PCG/Source/PCG/Public/PCGData.h b/Engine/Plugins/Experimental/PCG/Source/PCG/Public/PCGData.h
--- a/Engine/Plugins/Experimental/PCG/Source/PCG/Public/PCGData.h
+++ b/Engine/Plugins/Experimental/PCG/Source/PCG/Public/PCGData.h
@@ -103,19 +103,6 @@ struct PCG_API FPCGTaggedData
 	bool operator!=(const FPCGTaggedData& Other) const;
 };
 
-struct FPCGRootSet
-{
-	~FPCGRootSet() { Clear(); }
-	void Clear();
-	void Add(UObject* InObject);
-	void Remove(UObject* InObject);
-
-	TMap<UObject*, int32> RootSet;
-private:
-	void AddInternal(UObject* InObject);
-	void RemoveInternal(UObject* InObject);
-};
-
 USTRUCT(BlueprintType)
 struct PCG_API FPCGDataCollection
 {
@@ -174,8 +161,7 @@ struct PCG_API FPCGDataCollection
 
 	bool operator==(const FPCGDataCollection& Other) const;
 	bool operator!=(const FPCGDataCollection& Other) const;
-	void AddToRootSet(FPCGRootSet& RootSet) const;
-	void RemoveFromRootSet(FPCGRootSet& RootSet) const;
+	void AddReferences(FReferenceCollector& Collector);
 
 	/** Computes Crc for this data. */
 	FPCGCrc ComputeCrc(bool bFullDataCrc);
diff --git a/Engine/Plugins/Experimental/PCGGeometryScriptInterop/Source/PCGGeometryScriptInterop/Private/Elements/PCGMeshSampler.cpp b/Engine/Plugins/Experimental/PCGGeometryScriptInterop/Source/PCGGeometryScriptInterop/Private/Elements/PCGMeshSampler.cpp
--- a/Engine/Plugins/Experimental/PCGGeometryScriptInterop/Source/PCGGeometryScriptInterop/Private/Elements/PCGMeshSampler.cpp
+++ b/Engine/Plugins/Experimental/PCGGeometryScriptInterop/Source/PCGGeometryScriptInterop/Private/Elements/PCGMeshSampler.cpp
@@ -86,12 +86,13 @@ FPCGMeshSamplerContext::~FPCGMeshSamplerContext()
 		StopSampling = true;
 		SamplingFuture.Wait();
 	}
+}
 
-	// DynamicMesh needs to be removed from root, as we added it to root when we created it to avoid getting GC between executions.
+void FPCGMeshSamplerContext::AddExtraStructReferencedObjects(FReferenceCollector& Collector)
+{
 	if (DynamicMesh)
 	{
-		DynamicMesh->RemoveFromRoot();
-		DynamicMesh->MarkAsGarbage();
+		Collector.AddReferencedObject(DynamicMesh);
 	}
 }
 
@@ -126,8 +127,6 @@ bool FPCGMeshSamplerElement::PrepareDataInternal(FPCGContext* InContext) const
 	}
 
 	Context->DynamicMesh = NewObject<UDynamicMesh>();
-	// Need to add to root to make sure we don't delete it. Will be removed from root in the context destructor.
-	Context->DynamicMesh->AddToRoot();
 
 	EGeometryScriptOutcomePins Outcome;
 
diff --git a/Engine/Plugins/Experimental/PCGGeometryScriptInterop/Source/PCGGeometryScriptInterop/Public/Elements/PCGMeshSampler.h b/Engine/Plugins/Experimental/PCGGeometryScriptInterop/Source/PCGGeometryScriptInterop/Public/Elements/PCGMeshSampler.h
--- a/Engine/Plugins/Experimental/PCGGeometryScriptInterop/Source/PCGGeometryScriptInterop/Public/Elements/PCGMeshSampler.h
+++ b/Engine/Plugins/Experimental/PCGGeometryScriptInterop/Source/PCGGeometryScriptInterop/Public/Elements/PCGMeshSampler.h
@@ -157,8 +157,13 @@ protected:
 */
 struct FPCGMeshSamplerContext : public FPCGContext
 {
+public:
 	~FPCGMeshSamplerContext();
+
+protected:
+	virtual void AddExtraStructReferencedObjects(FReferenceCollector& Collector) override;
 	
+public:
 	// Dynamic mesh. Will be added to root.
 	TObjectPtr<UDynamicMesh> DynamicMesh;
 
