[PCG] Fix localization

#rnx

diff --git a/Engine/Plugins/Experimental/PCG/Source/PCG/Private/Elements/PCGSortAttributes.cpp b/Engine/Plugins/Experimental/PCG/Source/PCG/Private/Elements/PCGSortAttributes.cpp
--- a/Engine/Plugins/Experimental/PCG/Source/PCG/Private/Elements/PCGSortAttributes.cpp
+++ b/Engine/Plugins/Experimental/PCG/Source/PCG/Private/Elements/PCGSortAttributes.cpp
@@ -78,7 +78,7 @@ bool FPCGSortAttributesElement::ExecuteInternal(FPCGContext* Context) const
 
 			if (!PCGMetadataAttribute::CallbackWithRightType(Accessor->GetUnderlyingType(), [](auto Dummy) -> bool { return PCG::Private::MetadataTraits<decltype(Dummy)>::CanCompare; }))
 			{
-				PCGE_LOG(Error, GraphAndLog, FText::Format(LOCTEXT("InvalidAccessor", "Attribute '{0}' exists but is not of a comparable type ({1}) for input {2}."), InputSource.GetDisplayText(), PCG::Private::GetTypeNameText(Accessor->GetUnderlyingType()), FText::AsNumber(i)));
+				PCGE_LOG(Error, GraphAndLog, FText::Format(LOCTEXT("InvalidTypeAccessor", "Attribute '{0}' exists but is not of a comparable type ({1}) for input {2}."), InputSource.GetDisplayText(), PCG::Private::GetTypeNameText(Accessor->GetUnderlyingType()), FText::AsNumber(i)));
 				continue;
 			}
 
