[PCG] Fix node titles not refreshing due to cleaned up change propagation

#rnx

diff --git a/Engine/Plugins/Experimental/PCG/Source/PCG/Private/Elements/PCGBranch.cpp b/Engine/Plugins/Experimental/PCG/Source/PCG/Private/Elements/PCGBranch.cpp
--- a/Engine/Plugins/Experimental/PCG/Source/PCG/Private/Elements/PCGBranch.cpp
+++ b/Engine/Plugins/Experimental/PCG/Source/PCG/Private/Elements/PCGBranch.cpp
@@ -27,6 +27,20 @@ FText UPCGBranchSettings::GetNodeTooltipText() const
 {
 	return LOCTEXT("NodeTooltip", "Control flow node that will route the input to either Output A or Output B, based on the 'Output To B' property - which can also be overridden.");
 }
+
+EPCGChangeType UPCGBranchSettings::GetChangeTypeForProperty(const FName& InPropertyName) const
+{
+	EPCGChangeType ChangeType = Super::GetChangeTypeForProperty(InPropertyName) | EPCGChangeType::Cosmetic;
+
+	// Grid sizes are processed during graph compilation and is part of the graph structure.
+	if (InPropertyName == GET_MEMBER_NAME_CHECKED(UPCGBranchSettings, bEnabled)
+		|| InPropertyName == GET_MEMBER_NAME_CHECKED(UPCGBranchSettings, bOutputToB))
+	{
+		ChangeType |= EPCGChangeType::Structural;
+	}
+
+	return ChangeType;
+}
 #endif // WITH_EDITOR
 
 TArray<FPCGPinProperties> UPCGBranchSettings::OutputPinProperties() const
@@ -51,16 +65,6 @@ FPCGElementPtr UPCGBranchSettings::CreateElement() const
 	return MakeShared<FPCGBranchElement>();
 }
 
-#if WITH_EDITOR
-bool UPCGBranchSettings::IsStructuralProperty(const FName& InPropertyName) const
-{
-	// Static branches are processed during graph compilation and are part of the graph structure.
-	return InPropertyName == GET_MEMBER_NAME_CHECKED(UPCGBranchSettings, bEnabled)
-		|| InPropertyName == GET_MEMBER_NAME_CHECKED(UPCGBranchSettings, bOutputToB)
-		|| Super::IsStructuralProperty(InPropertyName);
-}
-#endif
-
 bool UPCGBranchSettings::IsDynamicBranch() const
 {
 	UPCGNode* Node = Cast<UPCGNode>(GetOuter());
diff --git a/Engine/Plugins/Experimental/PCG/Source/PCG/Private/Elements/PCGHiGenGridSize.cpp b/Engine/Plugins/Experimental/PCG/Source/PCG/Private/Elements/PCGHiGenGridSize.cpp
--- a/Engine/Plugins/Experimental/PCG/Source/PCG/Private/Elements/PCGHiGenGridSize.cpp
+++ b/Engine/Plugins/Experimental/PCG/Source/PCG/Private/Elements/PCGHiGenGridSize.cpp
@@ -96,11 +96,18 @@ FName UPCGHiGenGridSizeSettings::AdditionalTaskName() const
 }
 
 #if WITH_EDITOR
-bool UPCGHiGenGridSizeSettings::IsStructuralProperty(const FName& InPropertyName) const
+EPCGChangeType UPCGHiGenGridSizeSettings::GetChangeTypeForProperty(const FName& InPropertyName) const
 {
+	EPCGChangeType ChangeType = Super::GetChangeTypeForProperty(InPropertyName) | EPCGChangeType::Cosmetic;
+
 	// Grid sizes are processed during graph compilation and is part of the graph structure.
-	return InPropertyName == GET_MEMBER_NAME_CHECKED(UPCGHiGenGridSizeSettings, bEnabled)
-		|| InPropertyName == GET_MEMBER_NAME_CHECKED(UPCGHiGenGridSizeSettings, HiGenGridSize);
+	if (InPropertyName == GET_MEMBER_NAME_CHECKED(UPCGHiGenGridSizeSettings, bEnabled)
+		|| InPropertyName == GET_MEMBER_NAME_CHECKED(UPCGHiGenGridSizeSettings, HiGenGridSize))
+	{
+		ChangeType |= EPCGChangeType::Structural;
+	}
+
+	return ChangeType;
 }
 #endif
 
diff --git a/Engine/Plugins/Experimental/PCG/Source/PCG/Private/Elements/PCGSpawnActor.cpp b/Engine/Plugins/Experimental/PCG/Source/PCG/Private/Elements/PCGSpawnActor.cpp
--- a/Engine/Plugins/Experimental/PCG/Source/PCG/Private/Elements/PCGSpawnActor.cpp
+++ b/Engine/Plugins/Experimental/PCG/Source/PCG/Private/Elements/PCGSpawnActor.cpp
@@ -238,12 +238,18 @@ void UPCGSpawnActorSettings::PostLoad()
 }
 
 #if WITH_EDITOR
-bool UPCGSpawnActorSettings::IsStructuralProperty(const FName& InPropertyName) const
+EPCGChangeType UPCGSpawnActorSettings::GetChangeTypeForProperty(const FName& InPropertyName) const
 {
-	return InPropertyName == GET_MEMBER_NAME_CHECKED(UPCGSpawnActorSettings, TemplateActorClass) || 
+	EPCGChangeType ChangeType = Super::GetChangeTypeForProperty(InPropertyName);
+
+	if (InPropertyName == GET_MEMBER_NAME_CHECKED(UPCGSpawnActorSettings, TemplateActorClass) ||
 		InPropertyName == GET_MEMBER_NAME_CHECKED(UPCGSpawnActorSettings, Option) ||
-		InPropertyName == GET_MEMBER_NAME_CHECKED(UPCGSpawnActorSettings, bSpawnByAttribute) ||
-		Super::IsStructuralProperty(InPropertyName);
+		InPropertyName == GET_MEMBER_NAME_CHECKED(UPCGSpawnActorSettings, bSpawnByAttribute))
+	{
+		ChangeType |= EPCGChangeType::Structural;
+	}
+	
+	return ChangeType;
 }
 #endif // WITH_EDITOR
 
diff --git a/Engine/Plugins/Experimental/PCG/Source/PCG/Private/PCGSettings.cpp b/Engine/Plugins/Experimental/PCG/Source/PCG/Private/PCGSettings.cpp
--- a/Engine/Plugins/Experimental/PCG/Source/PCG/Private/PCGSettings.cpp
+++ b/Engine/Plugins/Experimental/PCG/Source/PCG/Private/PCGSettings.cpp
@@ -154,8 +154,8 @@ void UPCGSettingsInterface::SetEnabled(bool bInEnabled)
 #if WITH_EDITOR
 		if (UPCGSettings* Settings = GetSettings())
 		{
-			const bool bIsStructuralChange = Settings->IsStructuralProperty(GET_MEMBER_NAME_CHECKED(UPCGSettingsInterface, bEnabled));
-			OnSettingsChangedDelegate.Broadcast(Settings, (bIsStructuralChange ? EPCGChangeType::Structural : EPCGChangeType::None) | EPCGChangeType::Settings);
+			const EPCGChangeType ChangeType = Settings->GetChangeTypeForProperty(GET_MEMBER_NAME_CHECKED(UPCGSettingsInterface, bEnabled));
+			OnSettingsChangedDelegate.Broadcast(Settings, ChangeType);
 		}
 #endif
 	}
@@ -542,12 +542,19 @@ void UPCGSettings::PostEditChangeProperty(FPropertyChangedEvent& PropertyChanged
 
 	if (PropertyChangedEvent.GetPropertyName() != GET_MEMBER_NAME_CHECKED(UPCGSettings, DeterminismSettings))
 	{
-		OnSettingsChangedDelegate.Broadcast(this, IsStructuralProperty(PropertyChangedEvent.GetPropertyName()) ? EPCGChangeType::Structural : EPCGChangeType::Settings);
+		OnSettingsChangedDelegate.Broadcast(this, GetChangeTypeForProperty(PropertyChangedEvent.GetPropertyName()));
 	}
 
 	CacheCrc();
 }
 
+EPCGChangeType UPCGSettings::GetChangeTypeForProperty(const FName& InPropertyName) const
+{
+PRAGMA_DISABLE_DEPRECATION_WARNINGS
+	return EPCGChangeType::Settings | (IsStructuralProperty(InPropertyName) ? EPCGChangeType::Structural : EPCGChangeType::None);
+PRAGMA_ENABLE_DEPRECATION_WARNINGS
+}
+
 void UPCGSettings::DirtyCache()
 {
 	if (GEditor)
diff --git a/Engine/Plugins/Experimental/PCG/Source/PCG/Private/PCGSubgraph.cpp b/Engine/Plugins/Experimental/PCG/Source/PCG/Private/PCGSubgraph.cpp
--- a/Engine/Plugins/Experimental/PCG/Source/PCG/Private/PCGSubgraph.cpp
+++ b/Engine/Plugins/Experimental/PCG/Source/PCG/Private/PCGSubgraph.cpp
@@ -137,7 +137,7 @@ void UPCGBaseSubgraphSettings::PostEditUndo()
 
 void UPCGBaseSubgraphSettings::PreEditChange(FProperty* PropertyAboutToChange)
 {
-	if (PropertyAboutToChange && IsStructuralProperty(PropertyAboutToChange->GetFName()))
+	if (PropertyAboutToChange && !!(GetChangeTypeForProperty(PropertyAboutToChange->GetFName()) & EPCGChangeType::Structural))
 	{
 		TeardownCallbacks();
 	}
@@ -149,7 +149,7 @@ void UPCGBaseSubgraphSettings::PostEditChangeProperty(struct FPropertyChangedEve
 {
 	Super::PostEditChangeProperty(PropertyChangedEvent);
 
-	if (PropertyChangedEvent.Property && IsStructuralProperty(PropertyChangedEvent.Property->GetFName()))
+	if (PropertyChangedEvent.Property && !!(GetChangeTypeForProperty(PropertyChangedEvent.Property->GetFName()) & EPCGChangeType::Structural))
 	{
 		SetupCallbacks();
 	}
@@ -163,9 +163,16 @@ void UPCGBaseSubgraphSettings::GetTrackedActorKeys(FPCGActorSelectionKeyToSettin
 	}
 }
 
-bool UPCGBaseSubgraphSettings::IsStructuralProperty(const FName& InPropertyName) const
+EPCGChangeType UPCGBaseSubgraphSettings::GetChangeTypeForProperty(const FName& InPropertyName) const
 {
-	return InPropertyName == GET_MEMBER_NAME_CHECKED(UPCGSettingsInterface, bEnabled) || Super::IsStructuralProperty(InPropertyName);
+	EPCGChangeType ChangeType = Super::GetChangeTypeForProperty(InPropertyName) | EPCGChangeType::Cosmetic;
+
+	if (InPropertyName == GET_MEMBER_NAME_CHECKED(UPCGSettingsInterface, bEnabled))
+	{
+		ChangeType |= EPCGChangeType::Structural;
+	}
+
+	return ChangeType;
 }
 
 void UPCGBaseSubgraphSettings::OnSubgraphChanged(UPCGGraphInterface* InGraph, EPCGChangeType ChangeType)
@@ -321,10 +328,17 @@ UObject* UPCGSubgraphSettings::GetJumpTargetForDoubleClick() const
 	return Cast<UObject>(GetSubgraph());
 }
 
-bool UPCGSubgraphSettings::IsStructuralProperty(const FName& InPropertyName) const
+EPCGChangeType UPCGSubgraphSettings::GetChangeTypeForProperty(const FName& InPropertyName) const
 {
+	EPCGChangeType ChangeType = Super::GetChangeTypeForProperty(InPropertyName);
+
 	// Force structural if name is none. We are probably in a undo/redo situation
-	return (InPropertyName == NAME_None) || (InPropertyName == GET_MEMBER_NAME_CHECKED(UPCGSubgraphSettings, SubgraphInstance)) || Super::IsStructuralProperty(InPropertyName);
+	if ((InPropertyName == NAME_None) || (InPropertyName == GET_MEMBER_NAME_CHECKED(UPCGSubgraphSettings, SubgraphInstance)))
+	{
+		ChangeType |= EPCGChangeType::Structural;
+	}
+
+	return ChangeType;
 }
 #endif
 
diff --git a/Engine/Plugins/Experimental/PCG/Source/PCG/Public/Elements/Metadata/PCGMetadataBitwiseOpElement.h b/Engine/Plugins/Experimental/PCG/Source/PCG/Public/Elements/Metadata/PCGMetadataBitwiseOpElement.h
--- a/Engine/Plugins/Experimental/PCG/Source/PCG/Public/Elements/Metadata/PCGMetadataBitwiseOpElement.h
+++ b/Engine/Plugins/Experimental/PCG/Source/PCG/Public/Elements/Metadata/PCGMetadataBitwiseOpElement.h
@@ -48,6 +48,9 @@ public:
 
 protected:
 	//~Begin UPCGSettings interface
+#if WITH_EDITOR
+	virtual EPCGChangeType GetChangeTypeForProperty(const FName& InPropertyName) const override { return Super::GetChangeTypeForProperty(InPropertyName) | EPCGChangeType::Cosmetic; }
+#endif
 	virtual FPCGElementPtr CreateElement() const override;
 	//~End UPCGSettings interface
 
diff --git a/Engine/Plugins/Experimental/PCG/Source/PCG/Public/Elements/Metadata/PCGMetadataBooleanOpElement.h b/Engine/Plugins/Experimental/PCG/Source/PCG/Public/Elements/Metadata/PCGMetadataBooleanOpElement.h
--- a/Engine/Plugins/Experimental/PCG/Source/PCG/Public/Elements/Metadata/PCGMetadataBooleanOpElement.h
+++ b/Engine/Plugins/Experimental/PCG/Source/PCG/Public/Elements/Metadata/PCGMetadataBooleanOpElement.h
@@ -48,6 +48,9 @@ public:
 
 protected:
 	//~Begin UPCGSettings interface
+#if WITH_EDITOR
+	virtual EPCGChangeType GetChangeTypeForProperty(const FName& InPropertyName) const override { return Super::GetChangeTypeForProperty(InPropertyName) | EPCGChangeType::Cosmetic; }
+#endif
 	virtual FPCGElementPtr CreateElement() const override;
 	//~End UPCGSettings interface
 
diff --git a/Engine/Plugins/Experimental/PCG/Source/PCG/Public/Elements/Metadata/PCGMetadataCompareOpElement.h b/Engine/Plugins/Experimental/PCG/Source/PCG/Public/Elements/Metadata/PCGMetadataCompareOpElement.h
--- a/Engine/Plugins/Experimental/PCG/Source/PCG/Public/Elements/Metadata/PCGMetadataCompareOpElement.h
+++ b/Engine/Plugins/Experimental/PCG/Source/PCG/Public/Elements/Metadata/PCGMetadataCompareOpElement.h
@@ -50,6 +50,9 @@ public:
 
 protected:
 	//~Begin UPCGSettings interface
+#if WITH_EDITOR
+	virtual EPCGChangeType GetChangeTypeForProperty(const FName& InPropertyName) const override { return Super::GetChangeTypeForProperty(InPropertyName) | EPCGChangeType::Cosmetic; }
+#endif
 	virtual FPCGElementPtr CreateElement() const override;
 	//~End UPCGSettings interface
 
diff --git a/Engine/Plugins/Experimental/PCG/Source/PCG/Public/Elements/Metadata/PCGMetadataMathsOpElement.h b/Engine/Plugins/Experimental/PCG/Source/PCG/Public/Elements/Metadata/PCGMetadataMathsOpElement.h
--- a/Engine/Plugins/Experimental/PCG/Source/PCG/Public/Elements/Metadata/PCGMetadataMathsOpElement.h
+++ b/Engine/Plugins/Experimental/PCG/Source/PCG/Public/Elements/Metadata/PCGMetadataMathsOpElement.h
@@ -75,6 +75,9 @@ public:
 
 protected:
 	//~Begin UPCGSettings interface
+#if WITH_EDITOR
+	virtual EPCGChangeType GetChangeTypeForProperty(const FName& InPropertyName) const override { return Super::GetChangeTypeForProperty(InPropertyName) | EPCGChangeType::Cosmetic; }
+#endif
 	virtual FPCGElementPtr CreateElement() const override;
 	//~End UPCGSettings interface
 
diff --git a/Engine/Plugins/Experimental/PCG/Source/PCG/Public/Elements/Metadata/PCGMetadataRotatorOpElement.h b/Engine/Plugins/Experimental/PCG/Source/PCG/Public/Elements/Metadata/PCGMetadataRotatorOpElement.h
--- a/Engine/Plugins/Experimental/PCG/Source/PCG/Public/Elements/Metadata/PCGMetadataRotatorOpElement.h
+++ b/Engine/Plugins/Experimental/PCG/Source/PCG/Public/Elements/Metadata/PCGMetadataRotatorOpElement.h
@@ -58,6 +58,9 @@ public:
 
 protected:
 	//~Begin UPCGSettings interface
+#if WITH_EDITOR
+	virtual EPCGChangeType GetChangeTypeForProperty(const FName& InPropertyName) const override { return Super::GetChangeTypeForProperty(InPropertyName) | EPCGChangeType::Cosmetic; }
+#endif
 	virtual FPCGElementPtr CreateElement() const override;
 	//~End UPCGSettings interface
 
diff --git a/Engine/Plugins/Experimental/PCG/Source/PCG/Public/Elements/Metadata/PCGMetadataStringOpElement.h b/Engine/Plugins/Experimental/PCG/Source/PCG/Public/Elements/Metadata/PCGMetadataStringOpElement.h
--- a/Engine/Plugins/Experimental/PCG/Source/PCG/Public/Elements/Metadata/PCGMetadataStringOpElement.h
+++ b/Engine/Plugins/Experimental/PCG/Source/PCG/Public/Elements/Metadata/PCGMetadataStringOpElement.h
@@ -26,6 +26,9 @@ public:
 	virtual FName AdditionalTaskName() const override;
 
 protected:
+#if WITH_EDITOR
+	virtual EPCGChangeType GetChangeTypeForProperty(const FName& InPropertyName) const override { return Super::GetChangeTypeForProperty(InPropertyName) | EPCGChangeType::Cosmetic; }
+#endif
 	virtual FPCGElementPtr CreateElement() const override;
 	//~End UPCGSettings interface
 
diff --git a/Engine/Plugins/Experimental/PCG/Source/PCG/Public/Elements/Metadata/PCGMetadataTransformOpElement.h b/Engine/Plugins/Experimental/PCG/Source/PCG/Public/Elements/Metadata/PCGMetadataTransformOpElement.h
--- a/Engine/Plugins/Experimental/PCG/Source/PCG/Public/Elements/Metadata/PCGMetadataTransformOpElement.h
+++ b/Engine/Plugins/Experimental/PCG/Source/PCG/Public/Elements/Metadata/PCGMetadataTransformOpElement.h
@@ -60,6 +60,9 @@ public:
 
 protected:
 	//~Begin UPCGSettings interface
+#if WITH_EDITOR
+	virtual EPCGChangeType GetChangeTypeForProperty(const FName& InPropertyName) const override { return Super::GetChangeTypeForProperty(InPropertyName) | EPCGChangeType::Cosmetic; }
+#endif
 	virtual FPCGElementPtr CreateElement() const override;
 	//~End UPCGSettings interface
 
diff --git a/Engine/Plugins/Experimental/PCG/Source/PCG/Public/Elements/Metadata/PCGMetadataTrigOpElement.h b/Engine/Plugins/Experimental/PCG/Source/PCG/Public/Elements/Metadata/PCGMetadataTrigOpElement.h
--- a/Engine/Plugins/Experimental/PCG/Source/PCG/Public/Elements/Metadata/PCGMetadataTrigOpElement.h
+++ b/Engine/Plugins/Experimental/PCG/Source/PCG/Public/Elements/Metadata/PCGMetadataTrigOpElement.h
@@ -53,6 +53,9 @@ public:
 
 protected:
 	//~Begin UPCGSettings interface
+#if WITH_EDITOR
+	virtual EPCGChangeType GetChangeTypeForProperty(const FName& InPropertyName) const override { return Super::GetChangeTypeForProperty(InPropertyName) | EPCGChangeType::Cosmetic; }
+#endif
 	virtual FPCGElementPtr CreateElement() const override;
 	//~End UPCGSettings interface
 
diff --git a/Engine/Plugins/Experimental/PCG/Source/PCG/Public/Elements/Metadata/PCGMetadataVectorOpElement.h b/Engine/Plugins/Experimental/PCG/Source/PCG/Public/Elements/Metadata/PCGMetadataVectorOpElement.h
--- a/Engine/Plugins/Experimental/PCG/Source/PCG/Public/Elements/Metadata/PCGMetadataVectorOpElement.h
+++ b/Engine/Plugins/Experimental/PCG/Source/PCG/Public/Elements/Metadata/PCGMetadataVectorOpElement.h
@@ -57,6 +57,9 @@ public:
 
 protected:
 	//~Begin UPCGSettings interface
+#if WITH_EDITOR
+	virtual EPCGChangeType GetChangeTypeForProperty(const FName& InPropertyName) const override { return Super::GetChangeTypeForProperty(InPropertyName) | EPCGChangeType::Cosmetic; }
+#endif
 	virtual FPCGElementPtr CreateElement() const override;
 	//~End UPCGSettings interface
 
diff --git a/Engine/Plugins/Experimental/PCG/Source/PCG/Public/Elements/PCGAttributeReduceElement.h b/Engine/Plugins/Experimental/PCG/Source/PCG/Public/Elements/PCGAttributeReduceElement.h
--- a/Engine/Plugins/Experimental/PCG/Source/PCG/Public/Elements/PCGAttributeReduceElement.h
+++ b/Engine/Plugins/Experimental/PCG/Source/PCG/Public/Elements/PCGAttributeReduceElement.h
@@ -47,10 +47,15 @@ public:
 
 	virtual FName AdditionalTaskName() const override;
 
+protected:
+#if WITH_EDITOR
+	virtual EPCGChangeType GetChangeTypeForProperty(const FName& InPropertyName) const override { return Super::GetChangeTypeForProperty(InPropertyName) | EPCGChangeType::Cosmetic; }
+#endif
 	virtual TArray<FPCGPinProperties> InputPinProperties() const override;
 	virtual TArray<FPCGPinProperties> OutputPinProperties() const override;
 	//~End UPCGSettings interface
 
+public:
 	UPROPERTY(BlueprintReadWrite, EditAnywhere, Category = Settings, meta = (PCG_Overridable))
 	FPCGAttributePropertyInputSelector InputSource;
 
diff --git a/Engine/Plugins/Experimental/PCG/Source/PCG/Public/Elements/PCGAttributeSelectElement.h b/Engine/Plugins/Experimental/PCG/Source/PCG/Public/Elements/PCGAttributeSelectElement.h
--- a/Engine/Plugins/Experimental/PCG/Source/PCG/Public/Elements/PCGAttributeSelectElement.h
+++ b/Engine/Plugins/Experimental/PCG/Source/PCG/Public/Elements/PCGAttributeSelectElement.h
@@ -62,13 +62,17 @@ public:
 	virtual EPCGSettingsType GetType() const override { return EPCGSettingsType::Metadata; }
 	virtual void ApplyDeprecation(UPCGNode* InOutNode) override;
 #endif
-
 	virtual FName AdditionalTaskName() const override;
 
+protected:
+#if WITH_EDITOR
+	virtual EPCGChangeType GetChangeTypeForProperty(const FName& InPropertyName) const override { return Super::GetChangeTypeForProperty(InPropertyName) | EPCGChangeType::Cosmetic; }
+#endif
 	virtual TArray<FPCGPinProperties> InputPinProperties() const override;
 	virtual TArray<FPCGPinProperties> OutputPinProperties() const override;
 	//~End UPCGSettings interface
 
+public:
 	UPROPERTY(BlueprintReadWrite, EditAnywhere, Category = Settings, meta = (PCG_Overridable))
 	FPCGAttributePropertyInputSelector InputSource;
 
diff --git a/Engine/Plugins/Experimental/PCG/Source/PCG/Public/Elements/PCGAttributeTransferElement.h b/Engine/Plugins/Experimental/PCG/Source/PCG/Public/Elements/PCGAttributeTransferElement.h
--- a/Engine/Plugins/Experimental/PCG/Source/PCG/Public/Elements/PCGAttributeTransferElement.h
+++ b/Engine/Plugins/Experimental/PCG/Source/PCG/Public/Elements/PCGAttributeTransferElement.h
@@ -48,6 +48,9 @@ public:
 	virtual FName AdditionalTaskName() const override;
 
 protected:
+#if WITH_EDITOR
+	virtual EPCGChangeType GetChangeTypeForProperty(const FName& InPropertyName) const override { return Super::GetChangeTypeForProperty(InPropertyName) | EPCGChangeType::Cosmetic; }
+#endif
 	virtual TArray<FPCGPinProperties> InputPinProperties() const override;
 	virtual TArray<FPCGPinProperties> OutputPinProperties() const override;
 	virtual FPCGElementPtr CreateElement() const override;
diff --git a/Engine/Plugins/Experimental/PCG/Source/PCG/Public/Elements/PCGBranch.h b/Engine/Plugins/Experimental/PCG/Source/PCG/Public/Elements/PCGBranch.h
--- a/Engine/Plugins/Experimental/PCG/Source/PCG/Public/Elements/PCGBranch.h
+++ b/Engine/Plugins/Experimental/PCG/Source/PCG/Public/Elements/PCGBranch.h
@@ -26,11 +26,11 @@ public:
 	virtual bool IsPinStaticallyActive(const FName& PinLabel) const override;
 
 protected:
-	virtual TArray<FPCGPinProperties> OutputPinProperties() const override;
-	virtual FPCGElementPtr CreateElement() const override;
 #if WITH_EDITOR
-	virtual bool IsStructuralProperty(const FName& InPropertyName) const override;
+	virtual EPCGChangeType GetChangeTypeForProperty(const FName& InPropertyName) const override;
 #endif
+	virtual TArray<FPCGPinProperties> OutputPinProperties() const override;
+	virtual FPCGElementPtr CreateElement() const override;
 	//~End UPCGSettings interface
 
 	/** Branch is dynamic if selection value is overridden / cannot be determined prior to execution. */
diff --git a/Engine/Plugins/Experimental/PCG/Source/PCG/Public/Elements/PCGCreateAttribute.h b/Engine/Plugins/Experimental/PCG/Source/PCG/Public/Elements/PCGCreateAttribute.h
--- a/Engine/Plugins/Experimental/PCG/Source/PCG/Public/Elements/PCGCreateAttribute.h
+++ b/Engine/Plugins/Experimental/PCG/Source/PCG/Public/Elements/PCGCreateAttribute.h
@@ -34,8 +34,14 @@ public:
 	virtual bool CanEditChange(const FProperty* InProperty) const override;
 #endif // WITH_EDITOR
 	virtual EPCGDataType GetCurrentPinTypes(const UPCGPin* InPin) const override;
+
+protected:
+#if WITH_EDITOR
+	virtual EPCGChangeType GetChangeTypeForProperty(const FName& InPropertyName) const override { return Super::GetChangeTypeForProperty(InPropertyName) | EPCGChangeType::Cosmetic; }
+#endif
 	//~End UPCGSettings interface
 
+public:
 	virtual FName GetOutputAttributeName(const FPCGAttributePropertyInputSelector* InSource, const UPCGData* InSourceData) const PURE_VIRTUAL(UPCGCreateAttributeBaseSettings::GetOutputAttributeName, return NAME_None;);
 
 	// This can be set false by inheriting nodes to hide the 'From Source Param' property.
@@ -125,10 +131,15 @@ public:
 	virtual bool HasDynamicPins() const override { return true; }
 	virtual void ApplyDeprecation(UPCGNode* InOutNode) override;
 #endif // WITH_EDITOR
-
 	virtual FName AdditionalTaskName() const override;
+
+protected:
+#if WITH_EDITOR
+	virtual EPCGChangeType GetChangeTypeForProperty(const FName& InPropertyName) const override { return Super::GetChangeTypeForProperty(InPropertyName) | EPCGChangeType::Cosmetic; }
+#endif
 	//~End UPCGSettings interface
 
+public:
 	//~Begin UPCGCreateAttributeBaseSettings interface
 	virtual FName GetOutputAttributeName(const FPCGAttributePropertyInputSelector* InSource, const UPCGData* InSourceData) const { return OutputTarget.CopyAndFixSource(InSource, InSourceData).GetName(); }
 	//~End UPCGCreateAttributeBaseSettings interface
@@ -166,8 +177,14 @@ public:
 	virtual bool HasDynamicPins() const override { return false; }
 #endif
 	virtual FName AdditionalTaskName() const override;
+
+protected:
+#if WITH_EDITOR
+	virtual EPCGChangeType GetChangeTypeForProperty(const FName& InPropertyName) const override { return Super::GetChangeTypeForProperty(InPropertyName) | EPCGChangeType::Cosmetic; }
+#endif
 	//~End UPCGSettings interface
 
+public:
 	UPROPERTY(BlueprintReadWrite, EditAnywhere, Category = Settings, meta = (PCG_DiscardPropertySelection, PCG_DiscardExtraSelection, PCG_Overridable))
 	FPCGAttributePropertyOutputNoSourceSelector OutputTarget;
 
diff --git a/Engine/Plugins/Experimental/PCG/Source/PCG/Public/Elements/PCGDataFromActor.h b/Engine/Plugins/Experimental/PCG/Source/PCG/Public/Elements/PCGDataFromActor.h
--- a/Engine/Plugins/Experimental/PCG/Source/PCG/Public/Elements/PCGDataFromActor.h
+++ b/Engine/Plugins/Experimental/PCG/Source/PCG/Public/Elements/PCGDataFromActor.h
@@ -39,6 +39,9 @@ public:
 	virtual FName AdditionalTaskName() const override;
 
 protected:
+#if WITH_EDITOR
+	virtual EPCGChangeType GetChangeTypeForProperty(const FName& InPropertyName) const override { return Super::GetChangeTypeForProperty(InPropertyName) | EPCGChangeType::Cosmetic; }
+#endif
 	virtual TArray<FPCGPinProperties> InputPinProperties() const override { return TArray<FPCGPinProperties>(); }
 	virtual TArray<FPCGPinProperties> OutputPinProperties() const override;
 
diff --git a/Engine/Plugins/Experimental/PCG/Source/PCG/Public/Elements/PCGDataTableRowToParamData.h b/Engine/Plugins/Experimental/PCG/Source/PCG/Public/Elements/PCGDataTableRowToParamData.h
--- a/Engine/Plugins/Experimental/PCG/Source/PCG/Public/Elements/PCGDataTableRowToParamData.h
+++ b/Engine/Plugins/Experimental/PCG/Source/PCG/Public/Elements/PCGDataTableRowToParamData.h
@@ -39,6 +39,9 @@ public:
 	TSoftObjectPtr<UDataTable> DataTable;
 
 protected:
+#if WITH_EDITOR
+	virtual EPCGChangeType GetChangeTypeForProperty(const FName& InPropertyName) const override { return Super::GetChangeTypeForProperty(InPropertyName) | EPCGChangeType::Cosmetic; }
+#endif
 	virtual FPCGElementPtr CreateElement() const override;
 	// ~End UPCGSettings interface
 };
diff --git a/Engine/Plugins/Experimental/PCG/Source/PCG/Public/Elements/PCGDeleteAttributesElement.h b/Engine/Plugins/Experimental/PCG/Source/PCG/Public/Elements/PCGDeleteAttributesElement.h
--- a/Engine/Plugins/Experimental/PCG/Source/PCG/Public/Elements/PCGDeleteAttributesElement.h
+++ b/Engine/Plugins/Experimental/PCG/Source/PCG/Public/Elements/PCGDeleteAttributesElement.h
@@ -42,6 +42,9 @@ public:
 	virtual FName AdditionalTaskName() const override;
 
 protected:
+#if WITH_EDITOR
+	virtual EPCGChangeType GetChangeTypeForProperty(const FName& InPropertyName) const override { return Super::GetChangeTypeForProperty(InPropertyName) | EPCGChangeType::Cosmetic; }
+#endif
 	virtual TArray<FPCGPinProperties> InputPinProperties() const override;
 	virtual TArray<FPCGPinProperties> OutputPinProperties() const override;
 	virtual FPCGElementPtr CreateElement() const override;
diff --git a/Engine/Plugins/Experimental/PCG/Source/PCG/Public/Elements/PCGDeleteTags.h b/Engine/Plugins/Experimental/PCG/Source/PCG/Public/Elements/PCGDeleteTags.h
--- a/Engine/Plugins/Experimental/PCG/Source/PCG/Public/Elements/PCGDeleteTags.h
+++ b/Engine/Plugins/Experimental/PCG/Source/PCG/Public/Elements/PCGDeleteTags.h
@@ -30,6 +30,9 @@ public:
 	virtual FName AdditionalTaskName() const override;
 
 protected:
+#if WITH_EDITOR
+	virtual EPCGChangeType GetChangeTypeForProperty(const FName& InPropertyName) const override { return Super::GetChangeTypeForProperty(InPropertyName) | EPCGChangeType::Cosmetic; }
+#endif
 	virtual TArray<FPCGPinProperties> OutputPinProperties() const override;
 	virtual FPCGElementPtr CreateElement() const override;
 	//~End UPCGSettings interface
diff --git a/Engine/Plugins/Experimental/PCG/Source/PCG/Public/Elements/PCGExecuteBlueprint.h b/Engine/Plugins/Experimental/PCG/Source/PCG/Public/Elements/PCGExecuteBlueprint.h
--- a/Engine/Plugins/Experimental/PCG/Source/PCG/Public/Elements/PCGExecuteBlueprint.h
+++ b/Engine/Plugins/Experimental/PCG/Source/PCG/Public/Elements/PCGExecuteBlueprint.h
@@ -318,6 +318,7 @@ public:
 protected:
 	virtual FPCGElementPtr CreateElement() const override;
 #if WITH_EDITOR
+	virtual EPCGChangeType GetChangeTypeForProperty(const FName& InPropertyName) const override { return Super::GetChangeTypeForProperty(InPropertyName) | EPCGChangeType::Cosmetic; }
 	virtual TArray<FPCGSettingsOverridableParam> GatherOverridableParams() const override;
 #endif // WITH_EDITOR
 	virtual void FixingOverridableParamPropertyClass(FPCGSettingsOverridableParam& Param) const override;
diff --git a/Engine/Plugins/Experimental/PCG/Source/PCG/Public/Elements/PCGFilterByAttribute.h b/Engine/Plugins/Experimental/PCG/Source/PCG/Public/Elements/PCGFilterByAttribute.h
--- a/Engine/Plugins/Experimental/PCG/Source/PCG/Public/Elements/PCGFilterByAttribute.h
+++ b/Engine/Plugins/Experimental/PCG/Source/PCG/Public/Elements/PCGFilterByAttribute.h
@@ -23,6 +23,9 @@ public:
 	virtual FName AdditionalTaskName() const override;
 
 protected:
+#if WITH_EDITOR
+	virtual EPCGChangeType GetChangeTypeForProperty(const FName& InPropertyName) const override { return Super::GetChangeTypeForProperty(InPropertyName) | EPCGChangeType::Cosmetic; }
+#endif
 	virtual FPCGElementPtr CreateElement() const override;
 	//~End UPCGSettings interface
 
diff --git a/Engine/Plugins/Experimental/PCG/Source/PCG/Public/Elements/PCGFilterByIndex.h b/Engine/Plugins/Experimental/PCG/Source/PCG/Public/Elements/PCGFilterByIndex.h
--- a/Engine/Plugins/Experimental/PCG/Source/PCG/Public/Elements/PCGFilterByIndex.h
+++ b/Engine/Plugins/Experimental/PCG/Source/PCG/Public/Elements/PCGFilterByIndex.h
@@ -22,6 +22,9 @@ public:
 	virtual FName AdditionalTaskName() const override;
 
 protected:
+#if WITH_EDITOR
+	virtual EPCGChangeType GetChangeTypeForProperty(const FName& InPropertyName) const override { return Super::GetChangeTypeForProperty(InPropertyName) | EPCGChangeType::Cosmetic; }
+#endif
 	virtual FPCGElementPtr CreateElement() const override;
 	//~End UPCGSettings interface
 
diff --git a/Engine/Plugins/Experimental/PCG/Source/PCG/Public/Elements/PCGFilterByTag.h b/Engine/Plugins/Experimental/PCG/Source/PCG/Public/Elements/PCGFilterByTag.h
--- a/Engine/Plugins/Experimental/PCG/Source/PCG/Public/Elements/PCGFilterByTag.h
+++ b/Engine/Plugins/Experimental/PCG/Source/PCG/Public/Elements/PCGFilterByTag.h
@@ -29,6 +29,9 @@ public:
 	virtual FName AdditionalTaskName() const override;
 
 protected:
+#if WITH_EDITOR
+	virtual EPCGChangeType GetChangeTypeForProperty(const FName& InPropertyName) const override { return Super::GetChangeTypeForProperty(InPropertyName) | EPCGChangeType::Cosmetic; }
+#endif
 	virtual FPCGElementPtr CreateElement() const override;
 	//~End UPCGSettings interface
 
diff --git a/Engine/Plugins/Experimental/PCG/Source/PCG/Public/Elements/PCGGetActorProperty.h b/Engine/Plugins/Experimental/PCG/Source/PCG/Public/Elements/PCGGetActorProperty.h
--- a/Engine/Plugins/Experimental/PCG/Source/PCG/Public/Elements/PCGGetActorProperty.h
+++ b/Engine/Plugins/Experimental/PCG/Source/PCG/Public/Elements/PCGGetActorProperty.h
@@ -39,6 +39,9 @@ public:
 	virtual TArray<FPCGPinProperties> InputPinProperties() const override { return TArray<FPCGPinProperties>(); }
 
 protected:
+#if WITH_EDITOR
+	virtual EPCGChangeType GetChangeTypeForProperty(const FName& InPropertyName) const override { return Super::GetChangeTypeForProperty(InPropertyName) | EPCGChangeType::Cosmetic; }
+#endif
 	virtual FPCGElementPtr CreateElement() const override;
 	//~End UPCGSettings interface
 
diff --git a/Engine/Plugins/Experimental/PCG/Source/PCG/Public/Elements/PCGHiGenGridSize.h b/Engine/Plugins/Experimental/PCG/Source/PCG/Public/Elements/PCGHiGenGridSize.h
--- a/Engine/Plugins/Experimental/PCG/Source/PCG/Public/Elements/PCGHiGenGridSize.h
+++ b/Engine/Plugins/Experimental/PCG/Source/PCG/Public/Elements/PCGHiGenGridSize.h
@@ -26,16 +26,16 @@ public:
 	virtual EPCGSettingsType GetType() const override { return EPCGSettingsType::HierarchicalGeneration; }
 	virtual bool HasDynamicPins() const override { return true; }
 #endif
+	virtual FName AdditionalTaskName() const override;
 	virtual EPCGDataType GetCurrentPinTypes(const UPCGPin* InPin) const override;
 
 protected:
+#if WITH_EDITOR
+	virtual EPCGChangeType GetChangeTypeForProperty(const FName& InPropertyName) const override;
+#endif
 	virtual TArray<FPCGPinProperties> InputPinProperties() const override;
 	virtual TArray<FPCGPinProperties> OutputPinProperties() const override;
 	virtual FPCGElementPtr CreateElement() const override;
-	virtual FName AdditionalTaskName() const override;
-#if WITH_EDITOR
-	virtual bool IsStructuralProperty(const FName& InPropertyName) const override;
-#endif
 	//~End UPCGSettings interface
 
 	UPROPERTY(BlueprintReadWrite, EditAnywhere, Category = Settings)
diff --git a/Engine/Plugins/Experimental/PCG/Source/PCG/Public/Elements/PCGIndirectionElement.h b/Engine/Plugins/Experimental/PCG/Source/PCG/Public/Elements/PCGIndirectionElement.h
--- a/Engine/Plugins/Experimental/PCG/Source/PCG/Public/Elements/PCGIndirectionElement.h
+++ b/Engine/Plugins/Experimental/PCG/Source/PCG/Public/Elements/PCGIndirectionElement.h
@@ -33,6 +33,9 @@ public:
 #endif // WITH_EDITOR
 
 protected:
+#if WITH_EDITOR
+	virtual EPCGChangeType GetChangeTypeForProperty(const FName& InPropertyName) const override { return Super::GetChangeTypeForProperty(InPropertyName) | EPCGChangeType::Cosmetic; }
+#endif // WITH_EDITOR
 	virtual TArray<FPCGPinProperties> InputPinProperties() const override;
 	virtual TArray<FPCGPinProperties> OutputPinProperties() const override;
 
diff --git a/Engine/Plugins/Experimental/PCG/Source/PCG/Public/Elements/PCGLoopElement.h b/Engine/Plugins/Experimental/PCG/Source/PCG/Public/Elements/PCGLoopElement.h
--- a/Engine/Plugins/Experimental/PCG/Source/PCG/Public/Elements/PCGLoopElement.h
+++ b/Engine/Plugins/Experimental/PCG/Source/PCG/Public/Elements/PCGLoopElement.h
@@ -30,6 +30,9 @@ public:
 	virtual FName AdditionalTaskName() const override;
 
 protected:
+#if WITH_EDITOR
+	virtual EPCGChangeType GetChangeTypeForProperty(const FName& InPropertyName) const override { return Super::GetChangeTypeForProperty(InPropertyName) | EPCGChangeType::Cosmetic; }
+#endif
 	virtual FPCGElementPtr CreateElement() const override;
 	// ~End UPCGSettings interface implementation
 
diff --git a/Engine/Plugins/Experimental/PCG/Source/PCG/Public/Elements/PCGSpawnActor.h b/Engine/Plugins/Experimental/PCG/Source/PCG/Public/Elements/PCGSpawnActor.h
--- a/Engine/Plugins/Experimental/PCG/Source/PCG/Public/Elements/PCGSpawnActor.h
+++ b/Engine/Plugins/Experimental/PCG/Source/PCG/Public/Elements/PCGSpawnActor.h
@@ -116,6 +116,9 @@ public:
 #endif
 
 protected:
+#if WITH_EDITOR	
+	virtual EPCGChangeType GetChangeTypeForProperty(const FName& InPropertyName) const override;
+#endif
 	virtual TArray<FPCGPinProperties> InputPinProperties() const override { return DefaultPointInputPinProperties(); }
 	virtual TArray<FPCGPinProperties> OutputPinProperties() const override { return DefaultPointOutputPinProperties(); }
 	virtual FPCGElementPtr CreateElement() const override;
@@ -131,7 +134,6 @@ public:
 
 protected:
 #if WITH_EDITOR
-	virtual bool IsStructuralProperty(const FName& InPropertyName) const override;
 	//~End UPCGBaseSubgraphSettings interface
 
 	void SetupBlueprintEvent();
diff --git a/Engine/Plugins/Experimental/PCG/Source/PCG/Public/PCGSettings.h b/Engine/Plugins/Experimental/PCG/Source/PCG/Public/PCGSettings.h
--- a/Engine/Plugins/Experimental/PCG/Source/PCG/Public/PCGSettings.h
+++ b/Engine/Plugins/Experimental/PCG/Source/PCG/Public/PCGSettings.h
@@ -293,7 +293,7 @@ public:
 
 	/* Perform post-operations when an editor node is copied */
 	virtual void PostPaste();
-#endif
+#endif // WITH_EDITOR
 
 	virtual void ApplyPreconfiguredSettings(const FPCGPreConfiguredSettingsInfo& PreconfigureInfo) {}
 
@@ -366,8 +366,13 @@ protected:
 
 #if WITH_EDITOR
 	virtual void PostEditChangeProperty(FPropertyChangedEvent& PropertyChangedEvent) override;
+
+	UE_DEPRECATED(5.4, "IsStructuralProperty is deprecated, return EPCGChangeType::Structural from GetChangeTypeForProperty instead.")
 	virtual bool IsStructuralProperty(const FName& InPropertyName) const { return false; }
 
+	/** Gets the change impact for a given property. Can be used to signal structural or cosmetic node changes for example. */
+	virtual EPCGChangeType GetChangeTypeForProperty(const FName& InPropertyName) const;
+
 	/** Method that can be called to dirty the cache data from this settings objects if the operator== does not allow to detect changes */
 	void DirtyCache();
 
diff --git a/Engine/Plugins/Experimental/PCG/Source/PCG/Public/PCGSubgraph.h b/Engine/Plugins/Experimental/PCG/Source/PCG/Public/PCGSubgraph.h
--- a/Engine/Plugins/Experimental/PCG/Source/PCG/Public/PCGSubgraph.h
+++ b/Engine/Plugins/Experimental/PCG/Source/PCG/Public/PCGSubgraph.h
@@ -46,10 +46,7 @@ protected:
 
 	//~Begin UPCGSettings interface
 	virtual void GetTrackedActorKeys(FPCGActorSelectionKeyToSettingsMap& OutKeysToSettings, TArray<TObjectPtr<const UPCGGraph>>& OutVisitedGraphs) const override;
-	virtual bool IsStructuralProperty(const FName& InPropertyName) const override;
-
-	void SetupCallbacks();
-	void TeardownCallbacks();
+	virtual EPCGChangeType GetChangeTypeForProperty(const FName& InPropertyName) const override;
 #endif
 
 	TArray<FPCGPinProperties> InputPinProperties() const override;
@@ -59,6 +56,9 @@ protected:
 	virtual void SetSubgraphInternal(UPCGGraphInterface* InGraph) {}
 
 #if WITH_EDITOR
+	void SetupCallbacks();
+	void TeardownCallbacks();
+
 	void OnSubgraphChanged(UPCGGraphInterface* InGraph, EPCGChangeType ChangeType);
 #endif
 
@@ -101,6 +101,9 @@ public:
 	virtual FName AdditionalTaskName() const override;
 
 protected:
+#if WITH_EDITOR
+	virtual EPCGChangeType GetChangeTypeForProperty(const FName& InPropertyName) const override;
+#endif
 	virtual FPCGElementPtr CreateElement() const override;
 	//~End UPCGSettings interface implementation
 
@@ -110,9 +113,6 @@ public:
 	virtual bool IsDynamicGraph() const override;
 protected:
 	virtual void SetSubgraphInternal(UPCGGraphInterface* InGraph) override;
-#if WITH_EDITOR
-	virtual bool IsStructuralProperty(const FName& InPropertyName) const override;
-#endif
 	//~End UPCGBaseSubgraphSettings interface
 
 public:
