[PCG] Mark debug objects which have data for the currently inspected node.

diff --git a/Engine/Plugins/Experimental/PCG/Source/PCG/Private/PCGComponent.cpp b/Engine/Plugins/Experimental/PCG/Source/PCG/Private/PCGComponent.cpp
--- a/Engine/Plugins/Experimental/PCG/Source/PCG/Private/PCGComponent.cpp
+++ b/Engine/Plugins/Experimental/PCG/Source/PCG/Private/PCGComponent.cpp
@@ -1455,7 +1455,7 @@ void UPCGComponent::RefreshAfterGraphChanged(UPCGGraphInterface* InGraph, bool b
 			CleanupLocal(/*bRemoveComponents=*/true, /*bSave=*/ false);
 		}
 
-		InspectionCache.Empty();
+		ClearInspectionData();
 		return;
 	}
 #endif
@@ -1837,55 +1837,99 @@ void UPCGComponent::DisableInspection()
 	
 	if (InspectionCounter == 0)
 	{
-		InspectionCache.Empty();
+		ClearInspectionData();
 	}
 };
 
 void UPCGComponent::StoreInspectionData(const FPCGStack* InStack, const UPCGNode* InNode, const FPCGDataCollection& InInputData, const FPCGDataCollection& InOutputData)
 {
-	if (!IsInspecting() || !InNode || !ensure(InStack))
+	if (!InNode || !ensure(InStack))
 	{
 		return;
 	}
 
-	auto StorePinInspectionData = [InStack, InNode](const TArray<TObjectPtr<UPCGPin>>& InPins, const FPCGDataCollection& InData, TMap<FPCGStack, FPCGDataCollection>& InOutInspectionCache)
+	if (!InOutputData.TaggedData.IsEmpty())
 	{
-		for (const UPCGPin* Pin : InPins)
-		{
-			FPCGStack Stack = *InStack;
+		FWriteScopeLock Lock(NodeToStacksThatProducedDataLock);
 
-			// Append the Node and Pin to the current Stack to uniquely identify each DataCollection
-			TArray<FPCGStackFrame>& StackFrames = Stack.GetStackFramesMutable();
-			StackFrames.Reserve(StackFrames.Num() + 2);
-			StackFrames.Emplace(InNode);
-			StackFrames.Emplace(Pin);
+		NodeToStacksThatProducedData.FindOrAdd(InNode).Add(*InStack);
+	}
+	else
+	{
+		FWriteScopeLock Lock(NodeToStacksThatProducedDataLock);
 
-			FPCGDataCollection PinDataCollection;
-			PinDataCollection.TaggedData = InData.GetInputsByPin(Pin->Properties.Label);
-			// The data collection for each pin is given the Crc from the data collection. This is to enable inspecting the normal node output Crc
-			// when cache debugging is enabled.
-			PinDataCollection.Crc = InData.Crc;
+		if (TSet<FPCGStack>* Stacks = NodeToStacksThatProducedData.Find(InNode))
+		{
+			Stacks->Remove(*InStack);
+		}
+	}
 
-			if (!PinDataCollection.TaggedData.IsEmpty())
-			{
-				InOutInspectionCache.Add(Stack, PinDataCollection);
-			}
-			else
+	if (IsInspecting())
+	{
+		auto StorePinInspectionData = [InStack, InNode](const TArray<TObjectPtr<UPCGPin>>& InPins, const FPCGDataCollection& InData, TMap<FPCGStack, FPCGDataCollection>& InOutInspectionCache)
+		{
+			for (const UPCGPin* Pin : InPins)
 			{
-				InOutInspectionCache.Remove(Stack);
+				FPCGStack Stack = *InStack;
+
+				// Append the Node and Pin to the current Stack to uniquely identify each DataCollection
+				TArray<FPCGStackFrame>& StackFrames = Stack.GetStackFramesMutable();
+				StackFrames.Reserve(StackFrames.Num() + 2);
+				StackFrames.Emplace(InNode);
+				StackFrames.Emplace(Pin);
+
+				FPCGDataCollection PinDataCollection;
+				PinDataCollection.TaggedData = InData.GetInputsByPin(Pin->Properties.Label);
+				// The data collection for each pin is given the Crc from the data collection. This is to enable inspecting the normal node output Crc
+				// when cache debugging is enabled.
+				PinDataCollection.Crc = InData.Crc;
+
+				if (!PinDataCollection.TaggedData.IsEmpty())
+				{
+					InOutInspectionCache.Add(Stack, PinDataCollection);
+				}
+				else
+				{
+					InOutInspectionCache.Remove(Stack);
+				}
 			}
-		}
-	};
+		};
+
+		FWriteScopeLock Lock(InspectionCacheLock);
 
-	StorePinInspectionData(InNode->GetInputPins(), InInputData, InspectionCache);
-	StorePinInspectionData(InNode->GetOutputPins(), InOutputData, InspectionCache);
+		StorePinInspectionData(InNode->GetInputPins(), InInputData, InspectionCache);
+		StorePinInspectionData(InNode->GetOutputPins(), InOutputData, InspectionCache);
+	}
 }
 
 const FPCGDataCollection* UPCGComponent::GetInspectionData(const FPCGStack& InStack) const
 {
+	FReadScopeLock Lock(InspectionCacheLock);
 	return InspectionCache.Find(InStack);
 }
 
+void UPCGComponent::ClearInspectionData()
+{
+	{
+		FWriteScopeLock Lock(InspectionCacheLock);
+		InspectionCache.Reset();
+	}
+
+	{
+		FWriteScopeLock Lock(NodeToStacksThatProducedDataLock);
+		NodeToStacksThatProducedData.Reset();
+	}
+}
+
+bool UPCGComponent::HasNodeProducedData(const UPCGNode* InNode, const FPCGStack& Stack) const
+{
+	FReadScopeLock Lock(NodeToStacksThatProducedDataLock);
+
+	const TSet<FPCGStack>* StacksThatProducedData = NodeToStacksThatProducedData.Find(InNode);
+
+	return StacksThatProducedData && StacksThatProducedData->Contains(Stack);
+}
+
 void UPCGComponent::Refresh(bool bStructural, bool bCancelExistingRefresh)
 {
 	// Disable auto-refreshing on preview actors until we have something more robust on the execution side.
diff --git a/Engine/Plugins/Experimental/PCG/Source/PCG/Private/PCGSubsystem.cpp b/Engine/Plugins/Experimental/PCG/Source/PCG/Private/PCGSubsystem.cpp
--- a/Engine/Plugins/Experimental/PCG/Source/PCG/Private/PCGSubsystem.cpp
+++ b/Engine/Plugins/Experimental/PCG/Source/PCG/Private/PCGSubsystem.cpp
@@ -426,6 +426,10 @@ void UPCGSubsystem::OnOriginalComponentUnregistered(UPCGComponent* InComponent)
 	{
 		RuntimeGenScheduler->OnOriginalComponentUnregistered(InComponent);
 	}
+
+#if WITH_EDITOR
+	ClearExecutionMetadata(InComponent);
+#endif
 }
 
 #if WITH_EDITOR
@@ -443,7 +447,7 @@ void UPCGSubsystem::OnScheduleGraph(const FPCGStackContext& StackContext)
 	{
 		if (BaseStack->IsCurrentFrameInRootGraph())
 		{
-			ClearExecutedStacks(*BaseStack);
+			ClearExecutionMetadata(*BaseStack);
 		}
 	}
 
@@ -1507,6 +1511,25 @@ uint32 UPCGSubsystem::GetGraphCacheEntryCount(IPCGElement* InElement) const
 	return GraphExecutor ? GraphExecutor->GetGraphCacheEntryCount(InElement) : 0;
 }
 
+void UPCGSubsystem::ClearExecutionMetadata(const FPCGStack& BaseStack)
+{
+	ClearExecutedStacks(BaseStack);
+
+	if (UPCGComponent* Component = const_cast<UPCGComponent*>(BaseStack.GetRootComponent()))
+	{
+		Component->ClearInspectionData();
+
+		NodeVisualLogs.ClearLogs(Component);
+	}
+}
+
+void UPCGSubsystem::ClearExecutionMetadata(UPCGComponent* InComponent)
+{
+	FPCGStack Stack;
+	Stack.PushFrame(InComponent);
+	ClearExecutionMetadata(Stack);
+}
+
 TArray<FPCGStack> UPCGSubsystem::GetExecutedStacks(const UPCGComponent* InComponent, const UPCGGraph* InSubgraph)
 {
 	FReadScopeLock Lock(ExecutedStacksLock);
diff --git a/Engine/Plugins/Experimental/PCG/Source/PCG/Private/Utils/PCGNodeVisualLogs.cpp b/Engine/Plugins/Experimental/PCG/Source/PCG/Private/Utils/PCGNodeVisualLogs.cpp
--- a/Engine/Plugins/Experimental/PCG/Source/PCG/Private/Utils/PCGNodeVisualLogs.cpp
+++ b/Engine/Plugins/Experimental/PCG/Source/PCG/Private/Utils/PCGNodeVisualLogs.cpp
@@ -219,6 +219,13 @@ void FPCGNodeVisualLogs::ClearLogs(const FPCGStack& InPCGStack)
 	}
 }
 
+void FPCGNodeVisualLogs::ClearLogs(const UPCGComponent* InComponent)
+{
+	FPCGStack Stack;
+	Stack.PushFrame(InComponent);
+	ClearLogs(Stack);
+}
+
 #endif // WITH_EDITOR
 
 #undef LOCTEXT_NAMESPACE
diff --git a/Engine/Plugins/Experimental/PCG/Source/PCG/Public/PCGComponent.h b/Engine/Plugins/Experimental/PCG/Source/PCG/Public/PCGComponent.h
--- a/Engine/Plugins/Experimental/PCG/Source/PCG/Public/PCGComponent.h
+++ b/Engine/Plugins/Experimental/PCG/Source/PCG/Public/PCGComponent.h
@@ -11,6 +11,7 @@
 #include "ComponentInstanceDataCache.h"
 #include "Components/ActorComponent.h"
 #include "Containers/Map.h"
+#include "UObject/ObjectKey.h"
 
 #include "PCGComponent.generated.h"
 
@@ -293,6 +294,10 @@ public:
 	void DisableInspection();
 	void StoreInspectionData(const FPCGStack* InStack, const UPCGNode* InNode, const FPCGDataCollection& InInputData, const FPCGDataCollection& InOutputData);
 	const FPCGDataCollection* GetInspectionData(const FPCGStack& InStack) const;
+	void ClearInspectionData();
+
+	/** Did the given node produce one or more data items in the given stack in a previous execution. */
+	bool HasNodeProducedData(const UPCGNode* InNode, const FPCGStack& Stack) const;
 
 	bool IsActorTracked(AActor* InActor, bool& bOutIsCulled) const;
 
@@ -520,6 +525,14 @@ private:
 	TMap<FPCGStack, FPCGDataCollection> InspectionCache;
 #endif
 
+#if WITH_EDITOR
+	mutable FRWLock InspectionCacheLock;
+
+	/** Map from nodes to all stacks for which the node produced at least one data item. */
+	TMap<TObjectKey<const UPCGNode>, TSet<FPCGStack>> NodeToStacksThatProducedData;
+	mutable FRWLock NodeToStacksThatProducedDataLock;
+#endif
+
 	mutable FCriticalSection GeneratedResourcesLock;
 
 	// Graph instance
diff --git a/Engine/Plugins/Experimental/PCG/Source/PCG/Public/PCGSubsystem.h b/Engine/Plugins/Experimental/PCG/Source/PCG/Public/PCGSubsystem.h
--- a/Engine/Plugins/Experimental/PCG/Source/PCG/Public/PCGSubsystem.h
+++ b/Engine/Plugins/Experimental/PCG/Source/PCG/Public/PCGSubsystem.h
@@ -238,6 +238,10 @@ public:
 	/** Notify that we exited the Landscape edit mode. */
 	void NotifyLandscapeEditModeExited() { ActorAndComponentMapping.NotifyLandscapeEditModeExited(); }
 
+	/** Clear any data collected during execution, normally called prior to generating the component. */
+	void ClearExecutionMetadata(const FPCGStack& BaseStack);
+	void ClearExecutionMetadata(UPCGComponent* InComponent);
+
 	/** Get a list of stacks that were executed during the last execution. */
 	TArray<FPCGStack> GetExecutedStacks(const UPCGComponent* InComponent, const UPCGGraph* InSubgraph);
 	void ClearExecutedStacks(FPCGStack BeginningWithStack);
diff --git a/Engine/Plugins/Experimental/PCG/Source/PCG/Public/Utils/PCGNodeVisualLogs.h b/Engine/Plugins/Experimental/PCG/Source/PCG/Public/Utils/PCGNodeVisualLogs.h
--- a/Engine/Plugins/Experimental/PCG/Source/PCG/Public/Utils/PCGNodeVisualLogs.h
+++ b/Engine/Plugins/Experimental/PCG/Source/PCG/Public/Utils/PCGNodeVisualLogs.h
@@ -56,6 +56,9 @@ public:
 	/** Clear all errors and warnings that occurred while executing stacks beginning with the given stack. */
 	void ClearLogs(const FPCGStack& InPCGStack);
 
+	/** Clear all errors and warnings corresponding to the given component. */
+	void ClearLogs(const UPCGComponent* InComponent);
+
 private:
 	const int MaxLogsInSummary = 8;
 
diff --git a/Engine/Plugins/Experimental/PCG/Source/PCGEditor/Private/PCGEditor.cpp b/Engine/Plugins/Experimental/PCG/Source/PCGEditor/Private/PCGEditor.cpp
--- a/Engine/Plugins/Experimental/PCG/Source/PCGEditor/Private/PCGEditor.cpp
+++ b/Engine/Plugins/Experimental/PCG/Source/PCGEditor/Private/PCGEditor.cpp
@@ -1239,6 +1239,8 @@ void FPCGEditor::OnToggleInspected()
 
 	OnInspectedNodeChangedDelegate.Broadcast(PCGGraphNodeBeingInspected);
 	GetTabManager()->TryInvokeTab(FPCGEditor_private::AttributesID);
+
+	DebugObjectTreeWidget->SetNodeBeingInspected(PCGGraphNodeBeingInspected ? PCGNode : nullptr);
 }
 
 bool FPCGEditor::CanToggleInspected() const
diff --git a/Engine/Plugins/Experimental/PCG/Source/PCGEditor/Private/PCGEditor.h b/Engine/Plugins/Experimental/PCG/Source/PCGEditor/Private/PCGEditor.h
--- a/Engine/Plugins/Experimental/PCG/Source/PCGEditor/Private/PCGEditor.h
+++ b/Engine/Plugins/Experimental/PCG/Source/PCGEditor/Private/PCGEditor.h
@@ -31,7 +31,6 @@ class UPCGEditorGraphNodeBase;
 class UPCGGraph;
 class UPCGNode;
 
-DECLARE_MULTICAST_DELEGATE_OneParam(FOnInspectedComponentChanged, UPCGComponent*);
 DECLARE_MULTICAST_DELEGATE_OneParam(FOnInspectedStackChanged, const FPCGStack&);
 DECLARE_MULTICAST_DELEGATE_OneParam(FOnInspectedNodeChanged, UPCGEditorGraphNodeBase*);
 
diff --git a/Engine/Plugins/Experimental/PCG/Source/PCGEditor/Private/SPCGEditorGraphDebugObjectTree.cpp b/Engine/Plugins/Experimental/PCG/Source/PCGEditor/Private/SPCGEditorGraphDebugObjectTree.cpp
--- a/Engine/Plugins/Experimental/PCG/Source/PCGEditor/Private/SPCGEditorGraphDebugObjectTree.cpp
+++ b/Engine/Plugins/Experimental/PCG/Source/PCGEditor/Private/SPCGEditorGraphDebugObjectTree.cpp
@@ -133,6 +133,8 @@ void SPCGEditorGraphDebugObjectItemRow::Construct(const FArguments& InArgs, cons
 		[
 			SNew(STextBlock)
 			.Text(FText::FromString(Item->GetLabel()))
+			// Highlight effect applied if user is inspecting a node
+			.ColorAndOpacity(Item->IsGrayedOut() ? FColor(75, 75, 75) : FColor::White)
 		]
 		+SHorizontalBox::Slot()
 		.AutoWidth()
@@ -360,6 +362,13 @@ void SPCGEditorGraphDebugObjectTree::SetDebugObjectSelection(const FPCGStack& Fu
 	bDisableDebugObjectChangeNotification = false;
 }
 
+void SPCGEditorGraphDebugObjectTree::SetNodeBeingInspected(const UPCGNode* InPCGNode)
+{
+	PCGNodeBeingInspected = InPCGNode;
+
+	RequestRefresh();
+}
+
 void SPCGEditorGraphDebugObjectTree::SelectedDebugObject_OnClicked() const
 {
 	if (UPCGComponent* PCGComponent = PCGEditor.Pin()->GetPCGComponentBeingInspected())
@@ -504,7 +513,8 @@ void SPCGEditorGraphDebugObjectTree::AddStacksToTree(const TArray<FPCGStack>& St
 	TMap<const FPCGStack, FPCGEditorGraphDebugObjectItemPtr>& InOutStackToItem)
 {
 	const UPCGGraph* GraphBeingEdited = GetPCGGraph();
-	if (!GraphBeingEdited)
+	UPCGSubsystem* Subsystem = PCGEditor.Pin()->GetSubsystem();
+	if (!GraphBeingEdited || !Subsystem)
 	{
 		return;
 	}
@@ -543,6 +553,9 @@ void SPCGEditorGraphDebugObjectTree::AddStacksToTree(const TArray<FPCGStack>& St
 			continue;
 		}
 
+		// If we're inspecting a node which has not logged inspection data in a previous execution, display grayed out.
+		const bool bDisplayGrayedOut = PCGNodeBeingInspected && !PCGComponent->HasNodeProducedData(PCGNodeBeingInspected, Stack);
+
 		AActor* Actor = PCGComponent->GetOwner();
 		if (!Actor)
 		{
@@ -557,7 +570,7 @@ void SPCGEditorGraphDebugObjectTree::AddStacksToTree(const TArray<FPCGStack>& St
 		}
 		else
 		{
-			ActorItem = InOutActorItems.Emplace(Actor, MakeShared<FPCGEditorGraphDebugObjectItem_Actor>(Actor));
+			ActorItem = InOutActorItems.Emplace(Actor, MakeShared<FPCGEditorGraphDebugObjectItem_Actor>(Actor, bDisplayGrayedOut));
 			AllGraphItems.Add(ActorItem);
 		}
 
@@ -589,7 +602,7 @@ void SPCGEditorGraphDebugObjectTree::AddStacksToTree(const TArray<FPCGStack>& St
 					{
 						FPCGEditorGraphDebugObjectItemPtr TopGraphItem = InOutStackToItem.Emplace(
 							GraphStack,
-							MakeShared<FPCGEditorGraphDebugObjectItem_PCGComponent>(PCGComponent, StackGraph, GraphStack, bIsDebuggable));
+							MakeShared<FPCGEditorGraphDebugObjectItem_PCGComponent>(PCGComponent, StackGraph, GraphStack, bIsDebuggable, bDisplayGrayedOut));
 
 						AllGraphItems.Add(TopGraphItem);
 
@@ -606,7 +619,7 @@ void SPCGEditorGraphDebugObjectTree::AddStacksToTree(const TArray<FPCGStack>& St
 					{
 						FPCGEditorGraphDebugObjectItemPtr GraphItem = InOutStackToItem.Emplace(
 							GraphStack,
-							MakeShared<FPCGEditorGraphDebugObjectItem_PCGSubgraph>(SubgraphNode, StackGraph, GraphStack, bIsDebuggable));
+							MakeShared<FPCGEditorGraphDebugObjectItem_PCGSubgraph>(SubgraphNode, StackGraph, GraphStack, bIsDebuggable, bDisplayGrayedOut));
 
 						AllGraphItems.Add(GraphItem);
 
@@ -637,7 +650,7 @@ void SPCGEditorGraphDebugObjectTree::AddStacksToTree(const TArray<FPCGStack>& St
 						{
 							FPCGEditorGraphDebugObjectItemPtr LoopGraphItem = InOutStackToItem.Emplace(
 								LoopGraphStack,
-								MakeShared<FPCGEditorGraphDebugObjectItem_PCGSubgraph>(LoopSubgraphNode, StackGraph, LoopGraphStack, /*bIsDebuggable=*/false));
+								MakeShared<FPCGEditorGraphDebugObjectItem_PCGSubgraph>(LoopSubgraphNode, StackGraph, LoopGraphStack, /*bIsDebuggable=*/false, false));
 
 							AllGraphItems.Add(LoopGraphItem);
 
@@ -662,7 +675,7 @@ void SPCGEditorGraphDebugObjectTree::AddStacksToTree(const TArray<FPCGStack>& St
 						{
 							FPCGEditorGraphDebugObjectItemPtr LoopIterationItem = InOutStackToItem.Emplace(
 								LoopIterationStack,
-								MakeShared<FPCGEditorGraphDebugObjectItem_PCGLoopIndex>(PreviousStackFrame.LoopIndex, StackGraph, LoopIterationStack, bIsDebuggable));
+								MakeShared<FPCGEditorGraphDebugObjectItem_PCGLoopIndex>(PreviousStackFrame.LoopIndex, StackGraph, LoopIterationStack, bIsDebuggable, bDisplayGrayedOut));
 
 							AllGraphItems.Add(LoopIterationItem);
 
@@ -881,12 +894,9 @@ void SPCGEditorGraphDebugObjectTree::OnSelectionChanged(FPCGEditorGraphDebugObje
 		return;
 	}
 
-	if (UPCGComponent* PCGComponent = InItem->GetPCGComponent())
+	if (const FPCGStack* PCGStack = InItem->GetPCGStack())
 	{
-		if (const FPCGStack* PCGStack = InItem->GetPCGStack())
-		{
-			PCGEditor.Pin()->SetStackBeingInspected(*PCGStack, FPCGDebugObjectSelectionMethod::DebugObjectTree);
-		}
+		PCGEditor.Pin()->SetStackBeingInspected(*PCGStack, FPCGDebugObjectSelectionMethod::DebugObjectTree);
 	}
 }
 
diff --git a/Engine/Plugins/Experimental/PCG/Source/PCGEditor/Private/SPCGEditorGraphDebugObjectTree.h b/Engine/Plugins/Experimental/PCG/Source/PCGEditor/Private/SPCGEditorGraphDebugObjectTree.h
--- a/Engine/Plugins/Experimental/PCG/Source/PCGEditor/Private/SPCGEditorGraphDebugObjectTree.h
+++ b/Engine/Plugins/Experimental/PCG/Source/PCGEditor/Private/SPCGEditorGraphDebugObjectTree.h
@@ -21,6 +21,12 @@ class FPCGEditorGraphDebugObjectItem : public TSharedFromThis<FPCGEditorGraphDeb
 {
 public:
 	FPCGEditorGraphDebugObjectItem() = default;
+
+	explicit FPCGEditorGraphDebugObjectItem(bool bInGrayedOut)
+		: bGrayedOut(bInGrayedOut)
+	{
+	}
+	
 	virtual ~FPCGEditorGraphDebugObjectItem() = default;
 
 	void AddChild(TSharedRef<FPCGEditorGraphDebugObjectItem> InChild);
@@ -31,6 +37,8 @@ public:
 	bool IsExpanded() const { return bIsExpanded; }
 	void SetExpanded(bool bInIsExpanded) { bIsExpanded = bInIsExpanded; }
 
+	bool IsGrayedOut() const { return bGrayedOut; }
+
 	/** Optional sort priority, if returns INDEX_NONE then sort will fall back to alphabetical. */
 	virtual int32 GetSortPriority() const { return INDEX_NONE; }
 
@@ -46,13 +54,15 @@ protected:
 	TWeakPtr<FPCGEditorGraphDebugObjectItem> Parent;
 	TSet<TSharedPtr<FPCGEditorGraphDebugObjectItem>> Children;
 	bool bIsExpanded = false;
+	bool bGrayedOut = false;
 };
 
 class FPCGEditorGraphDebugObjectItem_Actor : public FPCGEditorGraphDebugObjectItem
 {
 public:
-	FPCGEditorGraphDebugObjectItem_Actor(TWeakObjectPtr<AActor> InActor)
-		: Actor(InActor)
+	explicit FPCGEditorGraphDebugObjectItem_Actor(TWeakObjectPtr<AActor> InActor, bool bInHasInspectionData)
+		: FPCGEditorGraphDebugObjectItem(bInHasInspectionData)
+		, Actor(InActor)
 	{
 		PCGStack.PushFrame(InActor.Get());
 	}
@@ -71,12 +81,14 @@ protected:
 class FPCGEditorGraphDebugObjectItem_PCGComponent : public FPCGEditorGraphDebugObjectItem
 {
 public:
-	FPCGEditorGraphDebugObjectItem_PCGComponent(
+	explicit FPCGEditorGraphDebugObjectItem_PCGComponent(
 		TWeakObjectPtr<UPCGComponent> InPCGComponent,
 		TWeakObjectPtr<const UPCGGraph> InPCGGraph,
 		const FPCGStack& InPCGStack,
-		bool bInIsDebuggable)
-		: PCGComponent(InPCGComponent)
+		bool bInIsDebuggable,
+		bool bInHasInspectionData)
+		: FPCGEditorGraphDebugObjectItem(bInHasInspectionData)
+		, PCGComponent(InPCGComponent)
 		, PCGGraph(InPCGGraph)
 		, PCGStack(InPCGStack)
 		, bIsDebuggable(bInIsDebuggable)
@@ -99,12 +111,14 @@ protected:
 class FPCGEditorGraphDebugObjectItem_PCGSubgraph : public FPCGEditorGraphDebugObjectItem
 {
 public:
-	FPCGEditorGraphDebugObjectItem_PCGSubgraph(
+	explicit FPCGEditorGraphDebugObjectItem_PCGSubgraph(
 		TWeakObjectPtr<const UPCGNode> InPCGNode,
 		TWeakObjectPtr<const UPCGGraph> InPCGGraph,
 		const FPCGStack& InPCGStack,
-		bool bInIsDebuggable)
-		: PCGNode(InPCGNode)
+		bool bInIsDebuggable,
+		bool bInHasInspectionData)
+		: FPCGEditorGraphDebugObjectItem(bInHasInspectionData)
+		, PCGNode(InPCGNode)
 		, PCGGraph(InPCGGraph)
 		, PCGStack(InPCGStack)
 		, bIsDebuggable(bInIsDebuggable)
@@ -127,12 +141,14 @@ protected:
 class FPCGEditorGraphDebugObjectItem_PCGLoopIndex : public FPCGEditorGraphDebugObjectItem
 {
 public:
-	FPCGEditorGraphDebugObjectItem_PCGLoopIndex(
+	explicit FPCGEditorGraphDebugObjectItem_PCGLoopIndex(
 		int32 InLoopIndex,
 		TWeakObjectPtr<const UObject> InLoopedPCGGraph,
 		const FPCGStack& InPCGStack,
-		bool bInIsDebuggable)
-		: LoopIndex(InLoopIndex)
+		bool bInIsDebuggable,
+		bool bInHasInspectionData)
+		: FPCGEditorGraphDebugObjectItem(bInHasInspectionData)
+		, LoopIndex(InLoopIndex)
 		, LoopedPCGGraph(InLoopedPCGGraph)
 		, PCGStack(InPCGStack)
 		, bIsDebuggable(bInIsDebuggable)
@@ -193,7 +209,9 @@ public:
 
 	virtual void Tick(const FGeometry& AllottedGeometry, const double InCurrentTime, const float InDeltaTime) override;
 
-	void RequestRefresh() { bNeedsRefresh = true; }
+	void RequestRefresh() { bNeedsRefresh = true; };
+
+	void SetNodeBeingInspected(const UPCGNode* InPCGNode);
 
 	void SetDebugObjectSelection(const FPCGStack& FullStack);
 
@@ -243,4 +261,6 @@ private:
 
 	/** Used to retain item selection state across tree refreshes. */
 	FPCGStack SelectedStack;
+
+	const UPCGNode* PCGNodeBeingInspected = nullptr;
 };
