[PCG] Nativation of CombinePoints Node - Combines each point to share a singular bound extent.
--Includes Unit Tests

diff --git a/Engine/Plugins/Experimental/PCG/Source/PCG/Private/Elements/PCGCombinePoints.cpp b/Engine/Plugins/Experimental/PCG/Source/PCG/Private/Elements/PCGCombinePoints.cpp
new file mode 100644
--- /dev/null
+++ b/Engine/Plugins/Experimental/PCG/Source/PCG/Private/Elements/PCGCombinePoints.cpp
@@ -0,0 +1,81 @@
+// Copyright Epic Games, Inc. All Rights Reserved.
+
+#include "Elements/PCGCombinePoints.h"
+
+#include "PCGContext.h"
+#include "PCGPin.h"
+#include "Data/PCGPointData.h"
+#include "Helpers/PCGHelpers.h"
+
+#define LOCTEXT_NAMESPACE "PCGCombinePointsElement"
+
+FPCGElementPtr UPCGCombinePointsSettings::CreateElement() const
+{
+	return MakeShared<FPCGCombinePointsElement>();
+}
+
+bool FPCGCombinePointsElement::ExecuteInternal(FPCGContext* Context) const
+{
+	TRACE_CPUPROFILER_EVENT_SCOPE(FPCGCombinePointsElement::Execute);
+	check(Context);
+
+	const UPCGCombinePointsSettings* Settings = Context->GetInputSettings<UPCGCombinePointsSettings>();
+	check(Settings);
+
+	const TArray<FPCGTaggedData> Inputs = Context->InputData.GetInputsByPin(PCGPinConstants::DefaultInputLabel);
+	TArray<FPCGTaggedData>& Outputs = Context->OutputData.TaggedData;
+
+	for (int i = 0; i < Inputs.Num(); ++i)
+	{
+		const UPCGPointData* InputPointData = Cast<UPCGPointData>(Inputs[i].Data);
+		if (!InputPointData)
+		{
+			PCGE_LOG(Warning, GraphAndLog, LOCTEXT("InvalidInputPointData", "The input is not point data, skipped."));
+			continue;
+		}
+
+		const TArray<FPCGPoint>& InputPoints = InputPointData->GetPoints();
+		if (InputPoints.IsEmpty())
+		{
+			PCGE_LOG(Warning, GraphAndLog, LOCTEXT("NoPointsFound", "No points were found in the input data, skipped."));
+			continue;
+		}
+
+		FPCGPoint InputPoint = InputPoints[0];
+		InputPoint.ApplyScaleToBounds();
+
+		const FTransform& PointTransform = Settings->bUseFirstPointTransform ? InputPoint.Transform : Settings->PointTransform;
+		const FMatrix InversePointTransform = PointTransform.ToInverseMatrixWithScale();
+
+		FPCGTaggedData& Output = Outputs.Add_GetRef(Inputs[i]);
+		UPCGPointData* OutputPointData = NewObject<UPCGPointData>();
+		OutputPointData->InitializeFromData(InputPointData);
+		FPCGPoint& OutputPoint = OutputPointData->GetMutablePoints().Emplace_GetRef();
+		Output.Data = OutputPointData;
+
+		FPCGPoint OutPoint = FPCGPoint();
+		FBox OutBox(EForceInit::ForceInit);
+
+		for (int j = 0; j < InputPoints.Num(); ++j)
+		{
+			const FPCGPoint& InPoint = InputPoints[j];
+
+			OutBox += InPoint.GetLocalBounds().TransformBy(InPoint.Transform.ToMatrixWithScale() * InversePointTransform);
+		};
+
+		OutPoint.SetLocalBounds(OutBox);
+		OutPoint.Transform = PointTransform;
+		OutPoint.Seed = PCGHelpers::ComputeSeedFromPosition(OutPoint.Transform.GetLocation());
+
+		if (Settings->bCenterPivot)
+		{
+			OutPoint.ResetPointCenter(FVector(0.5));
+		}
+
+		OutputPoint = OutPoint;
+	}
+
+	return true;
+}
+
+#undef LOCTEXT_NAMESPACE
\ No newline at end of file
diff --git a/Engine/Plugins/Experimental/PCG/Source/PCG/Private/Tests/Elements/PCGCombinePointsTest.cpp b/Engine/Plugins/Experimental/PCG/Source/PCG/Private/Tests/Elements/PCGCombinePointsTest.cpp
new file mode 100644
--- /dev/null
+++ b/Engine/Plugins/Experimental/PCG/Source/PCG/Private/Tests/Elements/PCGCombinePointsTest.cpp
@@ -0,0 +1,212 @@
+// Copyright Epic Games, Inc. All Rights Reserved.
+
+#include "Tests/PCGTestsCommon.h"
+
+#include "PCGContext.h"
+#include "Data/PCGPointData.h"
+
+#include "Elements/PCGCombinePoints.h"
+
+namespace CombinePointsCommonTests
+{
+	TUniquePtr<FPCGContext> GenerateTestDataAndRun(bool CenterPivot, bool UseSourceTransform, FTransform PtTransform)
+	{
+		PCGTestsCommon::FTestData TestData;
+		PCGTestsCommon::GenerateSettings<UPCGCombinePointsSettings>(TestData);
+		UPCGCombinePointsSettings* Settings = CastChecked<UPCGCombinePointsSettings>(TestData.Settings);
+		Settings->bCenterPivot = CenterPivot;
+		Settings->bUseFirstPointTransform = UseSourceTransform;
+		Settings->PointTransform = PtTransform;
+
+		FPCGTaggedData& Inputs = TestData.InputData.TaggedData.Emplace_GetRef();
+		Inputs.Pin = PCGPinConstants::DefaultInputLabel;
+		UPCGPointData* InData = PCGTestsCommon::CreateRandomPointData(5, 42, false);
+		Inputs.Data = InData;
+
+		// Setting the transform explicitly
+		UPCGPointData* InputPointData = Cast<UPCGPointData>(InData);
+		TArray<FPCGPoint>& InPoints = InputPointData->GetMutablePoints();
+		for (int i = 0; i < InPoints.Num(); ++i)
+		{
+			InPoints[i].Transform.SetLocation(FVector(10.0, 10.0, 10.0) * i);
+			InPoints[i].Transform.SetRotation(FQuat::Identity);
+			InPoints[i].Transform.SetScale3D(FVector::One() * i + 1);
+		}
+
+		FPCGElementPtr TestElement = TestData.Settings->GetElement();
+
+		TUniquePtr<FPCGContext> Context = TestData.InitializeTestContext();
+
+		while (!TestElement->Execute(Context.Get())) {}
+
+		return Context;
+	}
+}
+
+IMPLEMENT_CUSTOM_SIMPLE_AUTOMATION_TEST(FPCGCombinePointsTest_Default, FPCGTestBaseClass, "Plugins.PCG.CombinePoints.Default", PCGTestsCommon::TestFlags)
+
+bool FPCGCombinePointsTest_Default::RunTest(const FString& Parameters)
+{
+	TUniquePtr<FPCGContext> Context = CombinePointsCommonTests::GenerateTestDataAndRun(true, true, FTransform());
+
+	// Test our point data
+	const TArray<FPCGTaggedData>& Outputs = Context->OutputData.TaggedData;
+
+	UTEST_EQUAL("Output count", Outputs.Num(), 1);
+
+	const UPCGPointData* OutPointData = Cast<UPCGPointData>(Outputs[0].Data);
+
+	UTEST_NOT_NULL("OutputA point data", OutPointData);
+
+	const TArray<FPCGPoint>& OutPoints = OutPointData->GetPoints();
+
+	UTEST_EQUAL("OutputA point count", OutPoints.Num(), 1);
+
+	/** Point[0] at(0)->bounds(-1, 1)
+	 * Point [1] at (10) -> bounds (8, 12)
+	 * Point [2] at (20) -> bounds (17, 23)
+	 * Point [3] at (30) -> bounds (26, 34)
+	 * Point [4] at (40) -> bounds (35, 45)
+	 * 45 + (-1) / 2 == 22 <- check center pivot
+	 * 45 + |-1| / 2 == 23 <- check extents */
+
+	UTEST_EQUAL("Output Location", OutPoints[0].Transform.GetLocation(), FVector(22.0));
+	UTEST_EQUAL("Output Bounds Min", OutPoints[0].BoundsMin, FVector(-23.0));
+	UTEST_EQUAL("Output Bounds Max", OutPoints[0].BoundsMax, FVector(23.0));
+	
+	return true;
+}
+
+IMPLEMENT_CUSTOM_SIMPLE_AUTOMATION_TEST(FPCGCombinePointsTest_SourceTransform, FPCGTestBaseClass, "Plugins.PCG.CombinePoints.SourceTransform", PCGTestsCommon::TestFlags)
+
+bool FPCGCombinePointsTest_SourceTransform::RunTest(const FString& Parameters)
+{
+	TUniquePtr<FPCGContext> Context = CombinePointsCommonTests::GenerateTestDataAndRun(false, true, FTransform());
+
+	// Test our point data
+	const TArray<FPCGTaggedData>& Outputs = Context->OutputData.TaggedData;
+
+	UTEST_EQUAL("Output count", Outputs.Num(), 1);
+
+	const UPCGPointData* OutPointData = Cast<UPCGPointData>(Outputs[0].Data);
+
+	UTEST_NOT_NULL("OutputA point data", OutPointData);
+
+	const TArray<FPCGPoint>& OutPoints = OutPointData->GetPoints();
+
+	UTEST_EQUAL("OutputA point count", OutPoints.Num(), 1);
+
+	UTEST_EQUAL("Output Location", OutPoints[0].Transform.GetLocation(), FVector(0.0));
+	UTEST_EQUAL("Output Bounds Min", OutPoints[0].BoundsMin, FVector(-1.0));
+	UTEST_EQUAL("Output Bounds Max", OutPoints[0].BoundsMax, FVector(45.0));
+
+	return true;
+}
+
+IMPLEMENT_CUSTOM_SIMPLE_AUTOMATION_TEST(FPCGCombinePointsTest_TransformLocation, FPCGTestBaseClass, "Plugins.PCG.CombinePoints.TransformLocation", PCGTestsCommon::TestFlags)
+
+bool FPCGCombinePointsTest_TransformLocation::RunTest(const FString& Parameters)
+{
+	TUniquePtr<FPCGContext> Context = CombinePointsCommonTests::GenerateTestDataAndRun(false, false, FTransform(FRotator::ZeroRotator, FVector(10.0), FVector::One()));
+
+	// Test our point data
+	const TArray<FPCGTaggedData>& Outputs = Context->OutputData.TaggedData;
+
+	UTEST_EQUAL("Output count", Outputs.Num(), 1);
+
+	const UPCGPointData* OutPointData = Cast<UPCGPointData>(Outputs[0].Data);
+
+	UTEST_NOT_NULL("OutputA point data", OutPointData);
+
+	const TArray<FPCGPoint>& OutPoints = OutPointData->GetPoints();
+
+	UTEST_EQUAL("OutputA point count", OutPoints.Num(), 1);
+
+	// Has min at -1, max at 45 at origin
+	// Bound extents shift by (10) 
+
+	UTEST_EQUAL("Output Location", OutPoints[0].Transform.GetLocation(), FVector(10.0));
+	UTEST_EQUAL("Output Bounds Min", OutPoints[0].BoundsMin, FVector(-11.0));
+	UTEST_EQUAL("Output Bounds Max", OutPoints[0].BoundsMax, FVector(35.0));
+
+	return true;
+}
+
+IMPLEMENT_CUSTOM_SIMPLE_AUTOMATION_TEST(FPCGCombinePointsTest_TransformRotation, FPCGTestBaseClass, "Plugins.PCG.CombinePoints.TransformRotation", PCGTestsCommon::TestFlags)
+
+bool FPCGCombinePointsTest_TransformRotation::RunTest(const FString& Parameters)
+{
+	TUniquePtr<FPCGContext> Context = CombinePointsCommonTests::GenerateTestDataAndRun(false, false, FTransform(FRotator(90.0, 0.0, 0.0), FVector::ZeroVector, FVector::One()));
+
+	// Test our point data
+	const TArray<FPCGTaggedData>& Outputs = Context->OutputData.TaggedData;
+
+	UTEST_EQUAL("Output count", Outputs.Num(), 1);
+
+	const UPCGPointData* OutPointData = Cast<UPCGPointData>(Outputs[0].Data);
+
+	UTEST_NOT_NULL("OutputA point data", OutPointData);
+
+	const TArray<FPCGPoint>& OutPoints = OutPointData->GetPoints();
+
+	UTEST_EQUAL("OutputA point count", OutPoints.Num(), 1);
+
+	// Rotate a cube around X axis
+	UTEST_EQUAL("Output Location", OutPoints[0].Transform.GetLocation(), FVector::ZeroVector);
+	UTEST_EQUAL("Output Bounds Min", OutPoints[0].BoundsMin, FVector(-1.0, -1.0, -45.0));
+	UTEST_EQUAL("Output Bounds Max", OutPoints[0].BoundsMax, FVector(45.0, 45.0, 1.0));
+
+	return true;
+}
+
+IMPLEMENT_CUSTOM_SIMPLE_AUTOMATION_TEST(FPCGCombinePointsTest_TransformScale, FPCGTestBaseClass, "Plugins.PCG.CombinePoints.TransformScale", PCGTestsCommon::TestFlags)
+
+bool FPCGCombinePointsTest_TransformScale::RunTest(const FString& Parameters)
+{
+	TUniquePtr<FPCGContext> Context = CombinePointsCommonTests::GenerateTestDataAndRun(false, false, FTransform(FRotator::ZeroRotator, FVector::ZeroVector, FVector(2.0)));
+
+	// Test our point data
+	const TArray<FPCGTaggedData>& Outputs = Context->OutputData.TaggedData;
+
+	UTEST_EQUAL("Output count", Outputs.Num(), 1);
+
+	const UPCGPointData* OutPointData = Cast<UPCGPointData>(Outputs[0].Data);
+
+	UTEST_NOT_NULL("OutputA point data", OutPointData);
+
+	const TArray<FPCGPoint>& OutPoints = OutPointData->GetPoints();
+
+	UTEST_EQUAL("OutputA point count", OutPoints.Num(), 1);
+
+	UTEST_EQUAL("Output Location", OutPoints[0].Transform.GetLocation(), FVector::ZeroVector);
+	UTEST_EQUAL("Output Bounds Min", OutPoints[0].BoundsMin, FVector(-0.5));
+	UTEST_EQUAL("Output Bounds Max", OutPoints[0].BoundsMax, FVector(22.5));
+
+	return true;
+}
+
+IMPLEMENT_CUSTOM_SIMPLE_AUTOMATION_TEST(FPCGCombinePointsTest_TransformNegScale, FPCGTestBaseClass, "Plugins.PCG.CombinePoints.TransformNegScale", PCGTestsCommon::TestFlags)
+
+bool FPCGCombinePointsTest_TransformNegScale::RunTest(const FString& Parameters)
+{
+	TUniquePtr<FPCGContext> Context = CombinePointsCommonTests::GenerateTestDataAndRun(false, false, FTransform(FRotator::ZeroRotator, FVector::ZeroVector, FVector(-2.0)));
+
+	// Test our point data
+	const TArray<FPCGTaggedData>& Outputs = Context->OutputData.TaggedData;
+
+	UTEST_EQUAL("Output count", Outputs.Num(), 1);
+
+	const UPCGPointData* OutPointData = Cast<UPCGPointData>(Outputs[0].Data);
+
+	UTEST_NOT_NULL("OutputA point data", OutPointData);
+
+	const TArray<FPCGPoint>& OutPoints = OutPointData->GetPoints();
+
+	UTEST_EQUAL("OutputA point count", OutPoints.Num(), 1);
+
+	UTEST_EQUAL("Output Location", OutPoints[0].Transform.GetLocation(), FVector::ZeroVector);
+	UTEST_EQUAL("Output Bounds Min", OutPoints[0].BoundsMin, FVector(-22.5));
+	UTEST_EQUAL("Output Bounds Max", OutPoints[0].BoundsMax, FVector(0.5));
+
+	return true;
+}
\ No newline at end of file
diff --git a/Engine/Plugins/Experimental/PCG/Source/PCG/Public/Elements/PCGCombinePoints.h b/Engine/Plugins/Experimental/PCG/Source/PCG/Public/Elements/PCGCombinePoints.h
new file mode 100644
--- /dev/null
+++ b/Engine/Plugins/Experimental/PCG/Source/PCG/Public/Elements/PCGCombinePoints.h
@@ -0,0 +1,51 @@
+// Copyright Epic Games, Inc. All Rights Reserved.
+
+#pragma once
+
+#include "PCGSettings.h"
+#include "Elements/PCGPointOperationElementBase.h"
+
+#include "PCGCombinePoints.generated.h"
+
+/**
+ * Combines each point to share a singular bound extent.
+ */
+UCLASS(BlueprintType, ClassGroup = (Procedural))
+class UPCGCombinePointsSettings : public UPCGSettings
+{
+	GENERATED_BODY()
+
+public:
+	//~Begin UPCGSettings interface
+#if WITH_EDITOR
+	virtual FName GetDefaultNodeName() const override { return FName(TEXT("CombinePoints")); }
+	virtual FText GetDefaultNodeTitle() const override { return NSLOCTEXT("PCGCombinePointsElement", "NodeTitle", "Combine Points"); }
+	virtual FText GetNodeTooltipText() const override { return NSLOCTEXT("PCGCombinePointsElement", "NodeTooltip", "Combines each point to share a singular bound extent."); }
+	virtual EPCGSettingsType GetType() const override { return EPCGSettingsType::Spatial; }
+#endif
+
+protected:
+	virtual TArray<FPCGPinProperties> InputPinProperties() const override { return Super::DefaultPointInputPinProperties(); }
+	virtual TArray<FPCGPinProperties> OutputPinProperties() const override { return Super::DefaultPointOutputPinProperties(); }
+	virtual FPCGElementPtr CreateElement() const override;
+	//~End UPCGSettings interface
+
+public:
+	/** Places the point at the center of the combined bounds. */
+	UPROPERTY(BlueprintReadWrite, EditAnywhere, Category = Settings, meta = (PCG_Overridable))
+	bool bCenterPivot = true;
+
+	/** Use the transform of the initial point. */
+	UPROPERTY(BlueprintReadWrite, EditAnywhere, Category = Settings, meta = (PCG_Overridable))
+	bool bUseFirstPointTransform = true;
+
+	/** Transform the point and adjust the bounds. */
+	UPROPERTY(BlueprintReadWrite, EditAnywhere, Category = Settings, meta = (PCG_Overridable, EditCondition = "!bUseFirstPointTransform"))
+	FTransform PointTransform = FTransform();
+};
+
+class FPCGCombinePointsElement : public IPCGElement
+{
+protected:
+	virtual bool ExecuteInternal(FPCGContext* Context) const override;
+};
\ No newline at end of file
