[PCG] Fixing Subgraph collapse issues

* Fix Unit test after input pin deprecation
* No longer need to update the touched nodes, as it was causing an ensure (because the editor graph was not yet reconstructed).

#rnx

diff --git a/Engine/Plugins/Experimental/PCG/Source/PCG/Private/Helpers/PCGSubgraphHelpers.cpp b/Engine/Plugins/Experimental/PCG/Source/PCG/Private/Helpers/PCGSubgraphHelpers.cpp
--- a/Engine/Plugins/Experimental/PCG/Source/PCG/Private/Helpers/PCGSubgraphHelpers.cpp
+++ b/Engine/Plugins/Experimental/PCG/Source/PCG/Private/Helpers/PCGSubgraphHelpers.cpp
@@ -495,10 +495,6 @@ UPCGGraph* FPCGSubgraphHelpers::CollapseIntoSubgraph(UPCGGraph* InOriginalGraph,
 #if WITH_EDITOR
 	InOriginalGraph->NotifyGraphChanged(EPCGChangeType::Structural);
 	NewPCGGraph->NotifyGraphChanged(EPCGChangeType::Structural);
-	for (UPCGNode* TouchedNode : TouchedNodes)
-	{
-		TouchedNode->OnNodeChangedDelegate.Broadcast(TouchedNode, EPCGChangeType::Node);
-	}
 #endif // WITH_EDITOR
 
 	return NewPCGGraph;
diff --git a/Engine/Plugins/Experimental/PCG/Source/PCG/Private/Tests/Subgraphs/PCGCollapseSubgraphTest.cpp b/Engine/Plugins/Experimental/PCG/Source/PCG/Private/Tests/Subgraphs/PCGCollapseSubgraphTest.cpp
--- a/Engine/Plugins/Experimental/PCG/Source/PCG/Private/Tests/Subgraphs/PCGCollapseSubgraphTest.cpp
+++ b/Engine/Plugins/Experimental/PCG/Source/PCG/Private/Tests/Subgraphs/PCGCollapseSubgraphTest.cpp
@@ -122,6 +122,10 @@ bool FPCGCollapseSubgraphSimple::RunTest(const FString& Parameters)
 	UTEST_NOT_NULL("InputNode is valid", InputNode);
 	UTEST_NOT_NULL("OutputNode is valid", OutputNode);
 
+	// Add a Landscape pin to the input node, and update the node
+	const FPCGPinProperties& InputLandscapePin = CastChecked<UPCGGraphInputOutputSettings>(InputNode->GetSettings())->AddPin(FPCGPinProperties(PCGInputOutputConstants::DefaultLandscapeLabel, EPCGDataType::Landscape));
+	InputNode->UpdateAfterSettingsChangeDuringCreation();
+
 	UPCGSurfaceSamplerSettings* SurfaceSamplerSettings = nullptr;
 	UPCGAttributeNoiseSettings* AttributeNoiseSettings = nullptr;
 	UPCGNode* SurfaceSamplerNode = MainGraph->AddNodeOfType<UPCGSurfaceSamplerSettings>(SurfaceSamplerSettings);
@@ -131,11 +135,11 @@ bool FPCGCollapseSubgraphSimple::RunTest(const FString& Parameters)
 	UTEST_NOT_NULL("AttributeNoise node is valid", AttributeNoiseNode);
 	UTEST_EQUAL("There are 2 nodes", MainGraph->GetNodes().Num(), 2);
 
-	MainGraph->AddLabeledEdge(InputNode, PCGInputOutputConstants::DefaultLandscapeLabel, SurfaceSamplerNode, TEXT("Surface"));
+	MainGraph->AddLabeledEdge(InputNode, InputLandscapePin.Label, SurfaceSamplerNode, TEXT("Surface"));
 	MainGraph->AddLabeledEdge(SurfaceSamplerNode, PCGPinConstants::DefaultOutputLabel, AttributeNoiseNode, PCGPinConstants::DefaultInputLabel);
 	MainGraph->AddLabeledEdge(AttributeNoiseNode, PCGPinConstants::DefaultOutputLabel, OutputNode, PCGPinConstants::DefaultOutputLabel);
 
-	VALIDATE_EDGE(UPCGGraphInputOutputSettings, UPCGSurfaceSamplerSettings, MainGraph, PCGInputOutputConstants::DefaultLandscapeLabel, TEXT("Surface"))
+	VALIDATE_EDGE(UPCGGraphInputOutputSettings, UPCGSurfaceSamplerSettings, MainGraph, InputLandscapePin.Label, TEXT("Surface"))
 	VALIDATE_EDGE(UPCGSurfaceSamplerSettings, UPCGAttributeNoiseSettings, MainGraph, PCGPinConstants::DefaultOutputLabel, PCGPinConstants::DefaultInputLabel)
 	VALIDATE_EDGE(UPCGAttributeNoiseSettings, UPCGGraphInputOutputSettings, MainGraph, PCGPinConstants::DefaultOutputLabel, PCGPinConstants::DefaultOutputLabel)
 
@@ -153,8 +157,10 @@ bool FPCGCollapseSubgraphSimple::RunTest(const FString& Parameters)
 	UTEST_NOT_NULL("Node is a subgraph node", SubgraphSettings);
 	UTEST_EQUAL("Subgraph graph is the right one", SubgraphSettings->SubgraphInstance->GetGraph(), Subgraph);
 
-	// Landscape pin from input is connected to landscape pin of subgraph node
-	VALIDATE_EDGE(UPCGGraphInputOutputSettings, UPCGSubgraphSettings, MainGraph, PCGInputOutputConstants::DefaultLandscapeLabel, PCGInputOutputConstants::DefaultLandscapeLabel)
+	const FName InputPinSubgraphLabel = TEXT("Surface Sampler Surface");
+
+	// Landscape pin from input is connected to created pin of subgraph node
+	VALIDATE_EDGE(UPCGGraphInputOutputSettings, UPCGSubgraphSettings, MainGraph, InputLandscapePin.Label, InputPinSubgraphLabel)
 
 	// Output node is connected to the subgraph node
 	VALIDATE_EDGE(UPCGSubgraphSettings, UPCGGraphInputOutputSettings, MainGraph, NAME_None, PCGPinConstants::DefaultOutputLabel)
@@ -162,7 +168,7 @@ bool FPCGCollapseSubgraphSimple::RunTest(const FString& Parameters)
 	// Same verification for subgraph than the initial main graph, expect for output pin
 	UTEST_EQUAL("There is 2 nodes in the subgraph", Subgraph->GetNodes().Num(), 2);
 
-	VALIDATE_EDGE(UPCGGraphInputOutputSettings, UPCGSurfaceSamplerSettings, Subgraph, PCGInputOutputConstants::DefaultLandscapeLabel, TEXT("Surface"))
+	VALIDATE_EDGE(UPCGGraphInputOutputSettings, UPCGSurfaceSamplerSettings, Subgraph, InputPinSubgraphLabel, TEXT("Surface"))
 	VALIDATE_EDGE(UPCGSurfaceSamplerSettings, UPCGAttributeNoiseSettings, Subgraph, PCGPinConstants::DefaultOutputLabel, PCGPinConstants::DefaultInputLabel)
 	VALIDATE_EDGE(UPCGAttributeNoiseSettings, UPCGGraphInputOutputSettings, Subgraph, PCGPinConstants::DefaultOutputLabel, NAME_None);
 
