[PCG] Fix unit test and uniformize object/class accessor to use their soft counterparts

* Filter out nullptr data in the SmokeTestAndValidResult, since it is not a good way of flowing data -> fix Collapse element test
* Rename forgotten test header
* Use SoftObjectPath/SoftObjectClass as expected types -> fix Get Actor Property tests
* Class ptr are now FSoftClassPath underneath.

#rnx

diff --git a/Engine/Plugins/Experimental/PCG/Source/PCG/Private/Metadata/Accessors/PCGAttributeAccessorHelpers.cpp b/Engine/Plugins/Experimental/PCG/Source/PCG/Private/Metadata/Accessors/PCGAttributeAccessorHelpers.cpp
--- a/Engine/Plugins/Experimental/PCG/Source/PCG/Private/Metadata/Accessors/PCGAttributeAccessorHelpers.cpp
+++ b/Engine/Plugins/Experimental/PCG/Source/PCG/Private/Metadata/Accessors/PCGAttributeAccessorHelpers.cpp
@@ -302,9 +302,13 @@ namespace PCGAttributeAccessorHelpers
 		{
 			return Functor(Signature<FPCGPropertySoftObjectPathAccessor>{}, SoftObjectProperty);
 		}
+		else if (const FClassProperty* ClassProperty = CastField<FClassProperty>(InProperty))
+		{
+			return Functor(Signature<FPCGPropertyObjectPtrAccessor<FClassProperty>>{}, ClassProperty);
+		}
 		else if (const FObjectProperty* ObjectProperty = CastField<FObjectProperty>(InProperty))
 		{
-			return Functor(Signature<FPCGPropertyObjectPtrAccessor>{}, ObjectProperty);
+			return Functor(Signature<FPCGPropertyObjectPtrAccessor<FObjectProperty>>{}, ObjectProperty);
 		}
 		else if (const FStructProperty* StructProperty = CastField<FStructProperty>(InProperty))
 		{
diff --git a/Engine/Plugins/Experimental/PCG/Source/PCG/Private/Tests/Elements/PCGGetActorPropertyTest.cpp b/Engine/Plugins/Experimental/PCG/Source/PCG/Private/Tests/Elements/PCGGetActorPropertyTest.cpp
--- a/Engine/Plugins/Experimental/PCG/Source/PCG/Private/Tests/Elements/PCGGetActorPropertyTest.cpp
+++ b/Engine/Plugins/Experimental/PCG/Source/PCG/Private/Tests/Elements/PCGGetActorPropertyTest.cpp
@@ -1,6 +1,6 @@
 // Copyright Epic Games, Inc. All Rights Reserved.
 
-#include "Tests/Elements/PCGPropertyToParamDataElementTest.h"
+#include "Tests/Elements/PCGGetActorPropertyTest.h"
 #include "PCGComponent.h"
 #include "PCGContext.h"
 #include "PCGGraph.h"
@@ -9,10 +9,6 @@
 #include "Elements/PCGGetActorProperty.h"
 #include "Tests/PCGTestsCommon.h"
 
-#include UE_INLINE_GENERATED_CPP_BY_NAME(PCGPropertyToParamDataElementTest)
-
-// To run: automation runtest pcg.tests.PropertyToParamData
-
 #if WITH_EDITOR
 
 IMPLEMENT_CUSTOM_SIMPLE_AUTOMATION_TEST(FPCGPropertyToParamDataPropertyTypeTest, FPCGTestBaseClass, "Plugins.PCG.PropertyToParamData.PropertyType", PCGTestsCommon::TestFlags)
@@ -275,8 +271,8 @@ bool FPCGPropertyToParamDataPropertyTypeTest::RunTest(const FString& Parameters)
 	bSuccess &= VerifyAttributeValueValid(this, TestData, GET_MEMBER_NAME_CHECKED(APCGUnitTestDummyActor, SoftClassPathProperty), SoftClassPathValue, ExtraTestWhat);
 
 	// Objects properties as String
-	bSuccess &= VerifyAttributeValueValid(this, TestData, GET_MEMBER_NAME_CHECKED(APCGUnitTestDummyActor, ClassProperty), UPCGDummyGetPropertyTest::StaticClass()->GetPathName(), ExtraTestWhat);
-	bSuccess &= VerifyAttributeValueValid(this, TestData, GET_MEMBER_NAME_CHECKED(APCGUnitTestDummyActor, ObjectProperty), ObjectValue->GetPathName(), ExtraTestWhat);
+	bSuccess &= VerifyAttributeValueValid(this, TestData, GET_MEMBER_NAME_CHECKED(APCGUnitTestDummyActor, ClassProperty), FSoftClassPath(UPCGDummyGetPropertyTest::StaticClass()), ExtraTestWhat);
+	bSuccess &= VerifyAttributeValueValid(this, TestData, GET_MEMBER_NAME_CHECKED(APCGUnitTestDummyActor, ObjectProperty), FSoftObjectPath(ObjectValue->GetPathName()), ExtraTestWhat);
 
 	// Struct Property Extracted - Colors
 	// Extracting int will always yield a int64 and extracting floats with yield doubles. Here color is u8, so cast all of them to int64
@@ -290,7 +286,7 @@ bool FPCGPropertyToParamDataPropertyTypeTest::RunTest(const FString& Parameters)
 	// Arrays of supported properties
 	bSuccess &= VerifyAttributeValuesValid(this, TestData, GET_MEMBER_NAME_CHECKED(APCGUnitTestDummyActor, ArrayOfIntsProperty), { GET_MEMBER_NAME_CHECKED(APCGUnitTestDummyActor, ArrayOfIntsProperty) }, ExtraTestWhat, 42ll, 43ll, 44ll);
 	bSuccess &= VerifyAttributeValuesValid(this, TestData, GET_MEMBER_NAME_CHECKED(APCGUnitTestDummyActor, ArrayOfVectorsProperty), { GET_MEMBER_NAME_CHECKED(APCGUnitTestDummyActor, ArrayOfVectorsProperty) }, ExtraTestWhat, VectorValue, SecondVectorValue);
-	bSuccess &= VerifyAttributeValuesValid(this, TestData, GET_MEMBER_NAME_CHECKED(APCGUnitTestDummyActor, ArrayOfObjectsProperty), { GET_MEMBER_NAME_CHECKED(APCGUnitTestDummyActor, ArrayOfObjectsProperty) }, ExtraTestWhat, ObjectValue->GetPathName(), SecondObjectValue->GetPathName());
+	bSuccess &= VerifyAttributeValuesValid(this, TestData, GET_MEMBER_NAME_CHECKED(APCGUnitTestDummyActor, ArrayOfObjectsProperty), { GET_MEMBER_NAME_CHECKED(APCGUnitTestDummyActor, ArrayOfObjectsProperty) }, ExtraTestWhat, FSoftObjectPath(ObjectValue), FSoftObjectPath(SecondObjectValue));
 
 	// Arrays of extracted properties
 	bSuccess &= VerifyAttributeValuesValid(this, TestData, GET_MEMBER_NAME_CHECKED(APCGUnitTestDummyActor, ArrayOfStructsProperty), ColorPropertyNames, ExtraTestWhat, 
diff --git a/Engine/Plugins/Experimental/PCG/Source/PCG/Private/Tests/Elements/PCGGetActorPropertyTest.h b/Engine/Plugins/Experimental/PCG/Source/PCG/Private/Tests/Elements/PCGGetActorPropertyTest.h
new file mode 100644
--- /dev/null
+++ b/Engine/Plugins/Experimental/PCG/Source/PCG/Private/Tests/Elements/PCGGetActorPropertyTest.h
@@ -0,0 +1,141 @@
+// Copyright Epic Games, Inc. All Rights Reserved.
+#pragma once
+
+#include "GameFramework/Actor.h"
+
+
+#include "PCGGetActorPropertyTest.generated.h"
+
+UENUM(meta = (Hidden))
+enum class EPCGUnitTestDummyEnum : int64
+{
+	One,
+	Two,
+	Three
+};
+
+USTRUCT(BlueprintType, meta = (Hidden))
+struct FPCGDummyGetPropertyLevel2Struct
+{
+	GENERATED_BODY()
+
+	UPROPERTY(BlueprintReadOnly, Category = "")
+	TArray<double> DoubleArrayProperty;
+};
+
+USTRUCT(BlueprintType, meta = (Hidden))
+struct FPCGDummyGetPropertyStruct
+{
+	GENERATED_BODY()
+
+	UPROPERTY(BlueprintReadOnly, Category = "")
+	TArray<int> IntArrayProperty;
+
+	UPROPERTY(BlueprintReadOnly, Category = "")
+	float FloatProperty = 0.0f;
+
+	UPROPERTY(BlueprintReadOnly, Category = "")
+	FPCGDummyGetPropertyLevel2Struct Level2Struct;
+};
+
+UCLASS(MinimalAPI, NotBlueprintable, NotPlaceable, NotBlueprintType, Transient, HideDropdown, meta = (Hidden))
+class UPCGDummyGetPropertyTest : public UObject
+{
+	GENERATED_BODY()
+
+public:
+	UPROPERTY(BlueprintReadOnly, Category = "")
+	int64 Int64Property = 0;
+
+	UPROPERTY(BlueprintReadOnly, Category = "")
+	double DoubleProperty = 0.0;
+};
+
+UCLASS(MinimalAPI, NotBlueprintable, NotPlaceable, NotBlueprintType, Transient, HideDropdown, meta = (Hidden))
+class APCGUnitTestDummyActor : public AActor
+{
+	GENERATED_BODY()
+
+public:
+	UPROPERTY(BlueprintReadOnly, Category="")
+	int IntProperty = 0;
+
+	UPROPERTY(BlueprintReadOnly, Category = "")
+	float FloatProperty = 0.0f;
+
+	UPROPERTY(BlueprintReadOnly, Category = "")
+	int64 Int64Property = 0;
+
+	UPROPERTY(BlueprintReadOnly, Category = "")
+	double DoubleProperty = 0.0;
+
+	UPROPERTY(BlueprintReadOnly, Category = "")
+	bool BoolProperty = false;
+
+	UPROPERTY(BlueprintReadOnly, Category = "")
+	FName NameProperty = NAME_None;
+
+	UPROPERTY(BlueprintReadOnly, Category = "")
+	FString StringProperty = "";
+
+	UPROPERTY(BlueprintReadOnly, Category = "")
+	EPCGUnitTestDummyEnum EnumProperty = EPCGUnitTestDummyEnum::One;
+
+	UPROPERTY(BlueprintReadOnly, Category = "")
+	FVector VectorProperty;
+
+	UPROPERTY(BlueprintReadOnly, Category = "")
+	FVector4 Vector4Property;
+
+	UPROPERTY(BlueprintReadOnly, Category = "")
+	FTransform TransformProperty;
+
+	UPROPERTY(BlueprintReadOnly, Category = "")
+	FRotator RotatorProperty;
+
+	UPROPERTY(BlueprintReadOnly, Category = "")
+	FQuat QuatProperty;
+
+	UPROPERTY(BlueprintReadOnly, Category = "")
+	FSoftObjectPath SoftObjectPathProperty;
+
+	UPROPERTY(BlueprintReadOnly, Category = "")
+	FSoftClassPath SoftClassPathProperty;
+
+	UPROPERTY(BlueprintReadOnly, Category = "")
+	TSubclassOf<AActor> ClassProperty = AActor::StaticClass();
+
+	UPROPERTY(BlueprintReadOnly, Category = "")
+	TObjectPtr<UPCGDummyGetPropertyTest> ObjectProperty = nullptr;
+
+	UPROPERTY(BlueprintReadOnly, Category = "")
+	FVector2D Vector2Property;
+
+	UPROPERTY(BlueprintReadOnly, Category = "")
+	FColor ColorProperty;
+
+	UPROPERTY(BlueprintReadOnly, Category = "")
+	TArray<int32> ArrayOfIntsProperty;
+
+	UPROPERTY(BlueprintReadOnly, Category = "")
+	TArray<FVector> ArrayOfVectorsProperty;
+
+	UPROPERTY(BlueprintReadOnly, Category = "")
+	TArray<FColor> ArrayOfStructsProperty;
+
+	UPROPERTY(BlueprintReadOnly, Category = "")
+	TArray<TObjectPtr<UPCGDummyGetPropertyTest>> ArrayOfObjectsProperty;
+
+	UPROPERTY(BlueprintReadOnly, Category = "")
+	FPCGDummyGetPropertyStruct DummyStruct;
+};
+
+UCLASS(MinimalAPI, NotBlueprintable, NotPlaceable, NotBlueprintType, Transient, HideDropdown, meta = (Hidden))
+class UPCGUnitTestDummyComponent : public UActorComponent
+{
+	GENERATED_BODY()
+
+public:
+	UPROPERTY(BlueprintReadOnly, Category = "")
+	int IntProperty = 0;
+};
diff --git a/Engine/Plugins/Experimental/PCG/Source/PCG/Private/Tests/Elements/PCGPropertyToParamDataElementTest.h b/Engine/Plugins/Experimental/PCG/Source/PCG/Private/Tests/Elements/PCGPropertyToParamDataElementTest.h
deleted file mode 100644
diff --git a/Engine/Plugins/Experimental/PCG/Source/PCG/Private/Tests/PCGTestsCommon.cpp b/Engine/Plugins/Experimental/PCG/Source/PCG/Private/Tests/PCGTestsCommon.cpp
--- a/Engine/Plugins/Experimental/PCG/Source/PCG/Private/Tests/PCGTestsCommon.cpp
+++ b/Engine/Plugins/Experimental/PCG/Source/PCG/Private/Tests/PCGTestsCommon.cpp
@@ -203,13 +203,15 @@ namespace PCGTestsCommon
 	{
 		TArray<FPCGDataCollection> Data;
 
-		TMap<EPCGDataType, TFunction<UPCGData*(void)>> TypeToDataFn;
-		TypeToDataFn.Add(EPCGDataType::Point, []() { return PCGTestsCommon::CreatePointData(); });
-		TypeToDataFn.Add(EPCGDataType::PolyLine, []() { return PCGTestsCommon::CreatePolyLineData(); });
-		TypeToDataFn.Add(EPCGDataType::Surface, []() { return PCGTestsCommon::CreateSurfaceData(); });
-		TypeToDataFn.Add(EPCGDataType::Volume, []() { return PCGTestsCommon::CreateVolumeData(); });
-		TypeToDataFn.Add(EPCGDataType::Primitive, []() { return PCGTestsCommon::CreatePrimitiveData(); });
-		TypeToDataFn.Add(EPCGDataType::Param, []() { return PCGTestsCommon::CreateEmptyParamData(); });
+		static const TMap<EPCGDataType, TFunction<UPCGData* (void)>> TypeToDataFn
+		{
+			{ EPCGDataType::Point, []() { return PCGTestsCommon::CreatePointData(); }},
+			{ EPCGDataType::PolyLine, []() { return PCGTestsCommon::CreatePolyLineData(); }},
+			{ EPCGDataType::Surface, []() { return PCGTestsCommon::CreateSurfaceData(); }},
+			{ EPCGDataType::Volume, []() { return PCGTestsCommon::CreateVolumeData(); }},
+			{ EPCGDataType::Primitive, []() { return PCGTestsCommon::CreatePrimitiveData(); }},
+			{ EPCGDataType::Param, []() { return PCGTestsCommon::CreateEmptyParamData();}}
+		};
 
 		// Create empty data
 		Data.Emplace();
@@ -222,9 +224,15 @@ namespace PCGTestsCommon
 				continue;
 			}
 
+			const UPCGData* SingleData = TypeToData.Value();
+			if (!SingleData)
+			{
+				continue;
+			}
+
 			FPCGDataCollection& SingleCollection = Data.Emplace_GetRef();
 			FPCGTaggedData& SingleTaggedData = SingleCollection.TaggedData.Emplace_GetRef();
-			SingleTaggedData.Data = TypeToData.Value();
+			SingleTaggedData.Data = SingleData;
 			SingleTaggedData.Pin = PinProperties.Label;
 
 			if (!PinProperties.AllowsMultipleConnections())
@@ -239,13 +247,19 @@ namespace PCGTestsCommon
 					continue;
 				}
 
+				const UPCGData* SecondaryData = SecondaryTypeToData.Value();
+				if (!SecondaryData)
+				{
+					continue;
+				}
+
 				FPCGDataCollection& MultiCollection = Data.Emplace_GetRef();
 				FPCGTaggedData& FirstTaggedData = MultiCollection.TaggedData.Emplace_GetRef();
-				FirstTaggedData.Data = TypeToData.Value();
+				FirstTaggedData.Data = SingleData;
 				FirstTaggedData.Pin = PinProperties.Label;
 
 				FPCGTaggedData& SecondTaggedData = MultiCollection.TaggedData.Emplace_GetRef();
-				SecondTaggedData.Data = SecondaryTypeToData.Value();
+				SecondTaggedData.Data = SecondaryData;
 				SecondTaggedData.Pin = PinProperties.Label;
 			}
 		}
diff --git a/Engine/Plugins/Experimental/PCG/Source/PCG/Public/Metadata/Accessors/PCGPropertyAccessor.h b/Engine/Plugins/Experimental/PCG/Source/PCG/Public/Metadata/Accessors/PCGPropertyAccessor.h
--- a/Engine/Plugins/Experimental/PCG/Source/PCG/Public/Metadata/Accessors/PCGPropertyAccessor.h
+++ b/Engine/Plugins/Experimental/PCG/Source/PCG/Public/Metadata/Accessors/PCGPropertyAccessor.h
@@ -392,20 +392,23 @@ private:
 /**
 * Templated accessor class for object/class ptr properties.
 * Do not instantiate it manually, use PCGAttributeAccessorHelpers::CreatePropertyAccessor.
-* Will always convert to FSoftObjectPath for PCG
+* Will always convert to FSoftObjectPath/FSoftClassPath for PCG
 * Key supported: Generic object
 */
-class FPCGPropertyObjectPtrAccessor : public IPCGAttributeAccessorT<FPCGPropertyObjectPtrAccessor>, IPCGPropertyChainAccessor
+template <typename PropertyType>
+class FPCGPropertyObjectPtrAccessor : public IPCGAttributeAccessorT<FPCGPropertyObjectPtrAccessor<PropertyType>>, IPCGPropertyChainAccessor
 {
 public:
-	using Type = FSoftObjectPath;
+	using Type = std::conditional_t<std::is_same_v<PropertyType, FClassProperty>, FSoftClassPath, FSoftObjectPath>;
 	using Super = IPCGAttributeAccessorT<FPCGPropertyObjectPtrAccessor>;
 
-	FPCGPropertyObjectPtrAccessor(const FObjectProperty* InProperty, TArray<const FProperty*>&& ExtraProperties = {})
+	FPCGPropertyObjectPtrAccessor(const PropertyType* InProperty, TArray<const FProperty*>&& ExtraProperties = {})
 		: Super(/*bInReadOnly=*/ false)
 		, IPCGPropertyChainAccessor(InProperty, std::forward<TArray<const FProperty*>>(ExtraProperties))
 		, Property(InProperty)
 	{
+		// Making sure it is the right properties.
+		static_assert(std::is_same_v<PropertyType, FClassProperty> || std::is_same_v<PropertyType, FObjectProperty>);
 		check(Property);
 	}
 
@@ -426,7 +429,7 @@ public:
 	}
 
 private:
-	const FObjectProperty* Property = nullptr;
+	const PropertyType* Property = nullptr;
 };
 
 /**
