Minor issue - the create target actor node didn't create an entry which would have caused potential issues in some instances. Also named the attribute for better user readability.

#rnx

diff --git a/Engine/Plugins/Experimental/PCG/Source/PCG/Private/Elements/PCGCreateTargetActor.cpp b/Engine/Plugins/Experimental/PCG/Source/PCG/Private/Elements/PCGCreateTargetActor.cpp
--- a/Engine/Plugins/Experimental/PCG/Source/PCG/Private/Elements/PCGCreateTargetActor.cpp
+++ b/Engine/Plugins/Experimental/PCG/Source/PCG/Private/Elements/PCGCreateTargetActor.cpp
@@ -8,6 +8,7 @@
 #include "PCGManagedResource.h"
 #include "PCGModule.h"
 #include "PCGParamData.h"
+#include "Data/PCGPointData.h"
 #include "Helpers/PCGActorHelpers.h"
 #include "Helpers/PCGHelpers.h"
 
@@ -307,9 +308,13 @@ bool FPCGCreateTargetActorElement::ExecuteInternal(FPCGContext* Context) const
 	}
 
 	// Create param data output with reference to actor
+	FSoftObjectPath GeneratedActorPath(GeneratedActor);
+
 	UPCGParamData* ParamData = NewObject<UPCGParamData>();
 	check(ParamData && ParamData->Metadata);
-	ParamData->Metadata->CreateAttribute<FSoftObjectPath>(NAME_None, FSoftObjectPath(GeneratedActor), /*bAllowsInterpolation=*/false, /*bOverrideParent=*/false);
+	FPCGMetadataAttribute<FSoftObjectPath>* ActorPathAttribute = ParamData->Metadata->CreateAttribute<FSoftObjectPath>(PCGPointDataConstants::ActorReferenceAttribute, GeneratedActorPath, /*bAllowsInterpolation=*/false, /*bOverrideParent=*/false);
+	check(ActorPathAttribute);
+	ActorPathAttribute->SetValue(ParamData->Metadata->AddEntry(), GeneratedActorPath);	
 
 	// Add param data to output and we're done
 	Context->OutputData.TaggedData.Emplace_GetRef().Data = ParamData;
