Added option to treat the match and set attributes as a no-op if there is no attribute set passed to it.

#rnx

diff --git a/Engine/Plugins/Experimental/PCG/Source/PCG/Private/Elements/PCGMatchAndSetAttributes.cpp b/Engine/Plugins/Experimental/PCG/Source/PCG/Private/Elements/PCGMatchAndSetAttributes.cpp
--- a/Engine/Plugins/Experimental/PCG/Source/PCG/Private/Elements/PCGMatchAndSetAttributes.cpp
+++ b/Engine/Plugins/Experimental/PCG/Source/PCG/Private/Elements/PCGMatchAndSetAttributes.cpp
@@ -577,19 +577,36 @@ bool FPCGMatchAndSetAttributesElement::PrepareDataInternal(FPCGContext* InContex
 
 		if (!ParamData)
 		{
-			PCGE_LOG_C(Warning, GraphAndLog, Context, LOCTEXT("NoMatchData", "Must have exactly one Attribute Set to match against"));
-			return EPCGTimeSliceInitResult::AbortExecution;
+			if (Settings->bWarnIfNoMatchData)
+			{
+				PCGE_LOG_C(Warning, GraphAndLog, Context, LOCTEXT("NoMatchData", "Must have exactly one Attribute Set to match against"));
+			}
+
+			return EPCGTimeSliceInitResult::NoOperation;
 		}
 
 		OutState.Partition = new FPCGMatchAndSetPartition(Context, Settings, Context->SourceComponent.Get(), ParamData);
 		check(OutState.Partition);
 
-		return OutState.Partition->Initialize() ? EPCGTimeSliceInitResult::Success : EPCGTimeSliceInitResult::AbortExecution;
+		if (OutState.Partition->Initialize())
+		{
+			return EPCGTimeSliceInitResult::Success;
+		}
+		else
+		{
+			PCGE_LOG_C(Warning, GraphAndLog, Context, LOCTEXT("CouldNotInitializeExecutionState", "Could not initialize per-execution timeslice state data"));
+			return EPCGTimeSliceInitResult::AbortExecution;
+		}
 	});
 
 	if (InitResult == EPCGTimeSliceInitResult::AbortExecution)
 	{
-		PCGE_LOG_C(Warning, GraphAndLog, InContext, LOCTEXT("CouldNotInitializeExecutionState", "Could not initialize per-execution timeslice state data"));
+		// Implementation note: the previous code paths already emit necessary warnings
+		return true;
+	}
+	else if (InitResult == EPCGTimeSliceInitResult::NoOperation)
+	{
+		Outputs = Inputs;
 		return true;
 	}
 
@@ -615,6 +632,7 @@ bool FPCGMatchAndSetAttributesElement::PrepareDataInternal(FPCGContext* InContex
 
 	if (!TimeSlicedContext->DataIsPreparedForExecution())
 	{
+		TimeSlicedContext->OutputData.TaggedData.Empty();
 		PCGE_LOG_C(Warning, GraphAndLog, InContext, LOCTEXT("CouldNotInitializeStateData", "Could not initialize timeslice state data"));
 		return true;
 	}
@@ -631,10 +649,9 @@ bool FPCGMatchAndSetAttributesElement::ExecuteInternal(FPCGContext* InContext) c
 
 	TArray<FPCGTaggedData> Inputs = InContext->InputData.GetInputsByPin(PCGPinConstants::DefaultInputLabel);
 
-	// Prepare data failed, no need to execute. Return an empty output
+	// Prepare data failed, no need to execute.
 	if (!TimeSlicedContext->DataIsPreparedForExecution())
 	{
-		TimeSlicedContext->OutputData.TaggedData.Empty();
 		return true;
 	}
 
diff --git a/Engine/Plugins/Experimental/PCG/Source/PCG/Public/Elements/PCGMatchAndSetAttributes.h b/Engine/Plugins/Experimental/PCG/Source/PCG/Public/Elements/PCGMatchAndSetAttributes.h
--- a/Engine/Plugins/Experimental/PCG/Source/PCG/Public/Elements/PCGMatchAndSetAttributes.h
+++ b/Engine/Plugins/Experimental/PCG/Source/PCG/Public/Elements/PCGMatchAndSetAttributes.h
@@ -84,6 +84,10 @@ public:
 	/** Attribute to weight more or less some entries from the attribute set */
 	UPROPERTY(BlueprintReadWrite, EditAnywhere, Category = Settings, meta = (DisplayName="Match Weight Attribute", EditCondition = "bUseWeightAttribute", PCG_Overridable))
 	FName WeightAttribute = NAME_None;
+
+	/** Controls whether we will emit a warning and return nothing if there is no provided attribute set */
+	UPROPERTY(BlueprintReadWrite, EditAnywhere, Category = Settings)
+	bool bWarnIfNoMatchData = true;
 };
 
 class FPCGMatchAndSetAttributesElement : public TPCGTimeSlicedElementBase<FPCGMatchAndSetAttributesExecutionState, FPCGMatchAndSetAttributesIterationState>
