[PCG] Fix bugs when renaming PCG nodes
* Leaving rename node mode when cache debugging is enabled will dirty and rename the node
* De-selecting the node does not submit the new title

diff --git a/Engine/Plugins/Experimental/PCG/Source/PCGEditor/Private/PCGEditor.cpp b/Engine/Plugins/Experimental/PCG/Source/PCGEditor/Private/PCGEditor.cpp
--- a/Engine/Plugins/Experimental/PCG/Source/PCGEditor/Private/PCGEditor.cpp
+++ b/Engine/Plugins/Experimental/PCG/Source/PCGEditor/Private/PCGEditor.cpp
@@ -21,6 +21,7 @@
 
 #include "PCGEditorCommands.h"
 #include "PCGEditorGraph.h"
+#include "PCGEditorGraphNode.h"
 #include "PCGEditorGraphNodeInput.h"
 #include "PCGEditorGraphNodeOutput.h"
 #include "PCGEditorGraphSchema.h"
@@ -2207,11 +2208,17 @@ void FPCGEditor::OnNodeTitleCommitted(const FText& NewText, ETextCommit::Type Co
 {
 	if (NodeBeingChanged)
 	{
-		if (CommitInfo == ETextCommit::OnEnter)
+		if (CommitInfo == ETextCommit::OnEnter || CommitInfo == ETextCommit::OnUserMovedFocus)
 		{
 			const FScopedTransaction Transaction(*FPCGEditorCommon::ContextIdentifier, LOCTEXT("PCGEditorRenameNode", "PCG Editor: Rename Node"), nullptr);
-			NodeBeingChanged->Modify();
 			NodeBeingChanged->OnRenameNode(NewText.ToString());
+
+			// Implementation detail: In UPCGEditorGraphNode we only set the title under certain conditions, so it calls Modify() itself.
+			// However, UEdGraphNode does not call Modify() on its own, so we should still call it in this case.
+			if (!NodeBeingChanged->IsA<UPCGEditorGraphNode>())
+			{
+				NodeBeingChanged->Modify();
+			}
 		}
 
 		if (UPCGEditorGraphNodeBase* PCGEditorNode = Cast<UPCGEditorGraphNodeBase>(NodeBeingChanged))
diff --git a/Engine/Plugins/Experimental/PCG/Source/PCGEditor/Private/PCGEditorGraphNode.cpp b/Engine/Plugins/Experimental/PCG/Source/PCGEditor/Private/PCGEditorGraphNode.cpp
--- a/Engine/Plugins/Experimental/PCG/Source/PCGEditor/Private/PCGEditorGraphNode.cpp
+++ b/Engine/Plugins/Experimental/PCG/Source/PCGEditor/Private/PCGEditorGraphNode.cpp
@@ -108,6 +108,7 @@ void UPCGEditorGraphNode::OnRenameNode(const FString& NewName)
 
 	if(PCGNode->GetNodeTitle().ToString() != NewName)
 	{
+		Modify();
 		PCGNode->Modify();
 		PCGNode->NodeTitle = FName(*NewName);
 	}
