Replaced text by icons in the compact node (opt-in)

#rnx

diff --git a/Engine/Plugins/Experimental/PCG/Source/PCG/Private/Elements/PCGCollapseElement.cpp b/Engine/Plugins/Experimental/PCG/Source/PCG/Private/Elements/PCGCollapseElement.cpp
--- a/Engine/Plugins/Experimental/PCG/Source/PCG/Private/Elements/PCGCollapseElement.cpp
+++ b/Engine/Plugins/Experimental/PCG/Source/PCG/Private/Elements/PCGCollapseElement.cpp
@@ -9,6 +9,14 @@
 
 #include UE_INLINE_GENERATED_CPP_BY_NAME(PCGCollapseElement)
 
+#if WITH_EDITOR
+bool UPCGCollapseSettings::GetCompactNodeIcon(FName& OutCompactNodeIcon) const
+{
+	OutCompactNodeIcon = PCGNodeConstants::Icons::CompactNodeConvert;
+	return true;
+}
+#endif
+
 TArray<FPCGPinProperties> UPCGCollapseSettings::InputPinProperties() const
 {
 	TArray<FPCGPinProperties> PinProperties;
diff --git a/Engine/Plugins/Experimental/PCG/Source/PCG/Private/Elements/PCGConvertToAttributeSet.cpp b/Engine/Plugins/Experimental/PCG/Source/PCG/Private/Elements/PCGConvertToAttributeSet.cpp
--- a/Engine/Plugins/Experimental/PCG/Source/PCG/Private/Elements/PCGConvertToAttributeSet.cpp
+++ b/Engine/Plugins/Experimental/PCG/Source/PCG/Private/Elements/PCGConvertToAttributeSet.cpp
@@ -8,6 +8,14 @@
 
 #include UE_INLINE_GENERATED_CPP_BY_NAME(PCGConvertToAttributeSet)
 
+#if WITH_EDITOR
+bool UPCGConvertToAttributeSetSettings::GetCompactNodeIcon(FName& OutCompactNodeIcon) const
+{
+	OutCompactNodeIcon = PCGNodeConstants::Icons::CompactNodeConvert;
+	return true;
+}
+#endif
+
 TArray<FPCGPinProperties> UPCGConvertToAttributeSetSettings::InputPinProperties() const
 {
 	return Super::DefaultPointInputPinProperties();
diff --git a/Engine/Plugins/Experimental/PCG/Source/PCG/Private/Elements/PCGFilterByType.cpp b/Engine/Plugins/Experimental/PCG/Source/PCG/Private/Elements/PCGFilterByType.cpp
--- a/Engine/Plugins/Experimental/PCG/Source/PCG/Private/Elements/PCGFilterByType.cpp
+++ b/Engine/Plugins/Experimental/PCG/Source/PCG/Private/Elements/PCGFilterByType.cpp
@@ -25,7 +25,13 @@ FText UPCGFilterByTypeSettings::GetNodeTooltipText() const
 {
 	return LOCTEXT("FilterByTypeNodeTooltip", "Filters data in the collection according to data type");
 }
-#endif
+
+bool UPCGFilterByTypeSettings::GetCompactNodeIcon(FName& OutCompactNodeIcon) const
+{
+	OutCompactNodeIcon = PCGNodeConstants::Icons::CompactNodeFilter;
+	return true;
+}
+#endif // WITH_EDITOR
 
 EPCGDataType UPCGFilterByTypeSettings::GetCurrentPinTypes(const UPCGPin* InPin) const
 {
diff --git a/Engine/Plugins/Experimental/PCG/Source/PCG/Public/Elements/PCGCollapseElement.h b/Engine/Plugins/Experimental/PCG/Source/PCG/Public/Elements/PCGCollapseElement.h
--- a/Engine/Plugins/Experimental/PCG/Source/PCG/Public/Elements/PCGCollapseElement.h
+++ b/Engine/Plugins/Experimental/PCG/Source/PCG/Public/Elements/PCGCollapseElement.h
@@ -19,6 +19,7 @@ public:
 	virtual FText GetDefaultNodeTitle() const override { return NSLOCTEXT("PCGCollapseElement", "NodeTitle", "To Point"); }
 	virtual EPCGSettingsType GetType() const override { return EPCGSettingsType::Spatial; }
 	virtual bool ShouldDrawNodeCompact() const override { return true; }
+	virtual bool GetCompactNodeIcon(FName& OutCompactNodeIcon) const override;
 #endif
 
 protected:
@@ -39,8 +40,6 @@ public:
 #if WITH_EDITOR
 	virtual FName GetDefaultNodeName() const override { return FName(TEXT("AttributeSetToPoint")); }
 	virtual FText GetDefaultNodeTitle() const override { return NSLOCTEXT("PCGConvertToPointElement", "NodeTitle", "Attribute Set To Point"); }
-	virtual EPCGSettingsType GetType() const override { return EPCGSettingsType::Spatial; }
-	virtual bool ShouldDrawNodeCompact() const override { return true; }
 #endif
 
 protected:
diff --git a/Engine/Plugins/Experimental/PCG/Source/PCG/Public/Elements/PCGConvertToAttributeSet.h b/Engine/Plugins/Experimental/PCG/Source/PCG/Public/Elements/PCGConvertToAttributeSet.h
--- a/Engine/Plugins/Experimental/PCG/Source/PCG/Public/Elements/PCGConvertToAttributeSet.h
+++ b/Engine/Plugins/Experimental/PCG/Source/PCG/Public/Elements/PCGConvertToAttributeSet.h
@@ -20,6 +20,7 @@ public:
 	virtual FText GetNodeTooltipText() const override { return NSLOCTEXT("PCGConvertToAttributeSetElement", "NodeTooltip", "Converts point data to an attribute set with one entry per point and the same attributes."); }
 	virtual EPCGSettingsType GetType() const override { return EPCGSettingsType::Param; }
 	virtual bool ShouldDrawNodeCompact() const override { return true; }
+	virtual bool GetCompactNodeIcon(FName& OutCompactNodeIcon) const override;
 #endif
 
 protected:
diff --git a/Engine/Plugins/Experimental/PCG/Source/PCG/Public/Elements/PCGFilterByType.h b/Engine/Plugins/Experimental/PCG/Source/PCG/Public/Elements/PCGFilterByType.h
--- a/Engine/Plugins/Experimental/PCG/Source/PCG/Public/Elements/PCGFilterByType.h
+++ b/Engine/Plugins/Experimental/PCG/Source/PCG/Public/Elements/PCGFilterByType.h
@@ -19,6 +19,7 @@ public:
 	virtual FText GetDefaultNodeTitle() const override;
 	virtual FText GetNodeTooltipText() const override;
 	virtual bool ShouldDrawNodeCompact() const override { return true; }
+	virtual bool GetCompactNodeIcon(FName& OutCompactNodeIcon) const override;
 #endif
 	virtual EPCGDataType GetCurrentPinTypes(const UPCGPin* InPin) const override;
 
diff --git a/Engine/Plugins/Experimental/PCG/Source/PCG/Public/PCGCommon.h b/Engine/Plugins/Experimental/PCG/Source/PCG/Public/PCGCommon.h
--- a/Engine/Plugins/Experimental/PCG/Source/PCG/Public/PCGCommon.h
+++ b/Engine/Plugins/Experimental/PCG/Source/PCG/Public/PCGCommon.h
@@ -108,6 +108,15 @@ namespace Icons
 }
 }
 
+namespace PCGNodeConstants
+{
+namespace Icons
+{
+	const FName CompactNodeConvert = TEXT("PCG.Node.Compact.Convert");
+	const FName CompactNodeFilter = TEXT("PCG.Node.Compact.Filter");
+}
+}
+
 // Metadata used by PCG
 namespace PCGObjectMetadata
 {
diff --git a/Engine/Plugins/Experimental/PCG/Source/PCG/Public/PCGSettings.h b/Engine/Plugins/Experimental/PCG/Source/PCG/Public/PCGSettings.h
--- a/Engine/Plugins/Experimental/PCG/Source/PCG/Public/PCGSettings.h
+++ b/Engine/Plugins/Experimental/PCG/Source/PCG/Public/PCGSettings.h
@@ -250,6 +250,8 @@ public:
 #if WITH_EDITOR
 	/** Puts node title on node body, reducing overall node size */
 	virtual bool ShouldDrawNodeCompact() const { return false; }
+	/** Returns the icon to use instead of text in compact node form */
+	virtual bool GetCompactNodeIcon(FName& OutCompactNodeIcon) const { return false; }
 
 	/** UpdatePins will kick off invalid edges, so this is useful for moving edges around in case of pin changes. */
 	virtual void ApplyDeprecationBeforeUpdatePins(UPCGNode* InOutNode, TArray<TObjectPtr<UPCGPin>>& InputPins, TArray<TObjectPtr<UPCGPin>>& OutputPins);
diff --git a/Engine/Plugins/Experimental/PCG/Source/PCGEditor/Private/PCGEditorGraphNodeBase.cpp b/Engine/Plugins/Experimental/PCG/Source/PCGEditor/Private/PCGEditorGraphNodeBase.cpp
--- a/Engine/Plugins/Experimental/PCG/Source/PCGEditor/Private/PCGEditorGraphNodeBase.cpp
+++ b/Engine/Plugins/Experimental/PCG/Source/PCGEditor/Private/PCGEditorGraphNodeBase.cpp
@@ -461,6 +461,12 @@ bool UPCGEditorGraphNodeBase::ShouldDrawCompact() const
 	return Settings && Settings->ShouldDrawNodeCompact();
 }
 
+bool UPCGEditorGraphNodeBase::GetCompactNodeIcon(FName& OutCompactNodeIcon) const
+{
+	UPCGSettings* Settings = PCGNode ? PCGNode->GetSettings() : nullptr;
+	return Settings && Settings->GetCompactNodeIcon(OutCompactNodeIcon);
+}
+
 void UPCGEditorGraphNodeBase::EnterRenamingMode()
 {
 	bCanRenameNode = true;
diff --git a/Engine/Plugins/Experimental/PCG/Source/PCGEditor/Private/PCGEditorGraphNodeBase.h b/Engine/Plugins/Experimental/PCG/Source/PCGEditor/Private/PCGEditorGraphNodeBase.h
--- a/Engine/Plugins/Experimental/PCG/Source/PCGEditor/Private/PCGEditorGraphNodeBase.h
+++ b/Engine/Plugins/Experimental/PCG/Source/PCGEditor/Private/PCGEditorGraphNodeBase.h
@@ -75,6 +75,9 @@ public:
 	/** Puts node title on node body, reducing overall node size */
 	bool ShouldDrawCompact() const;
 
+	/** Returns custom compact node icon if available */
+	bool GetCompactNodeIcon(FName& OutCompactNodeIcon) const;
+
 	// Highlighted nodes draw with a light tint. Currently used to indicate nodes from hi-gen grids that are available for use.
 	void SetIsHighlighted(bool bInIsHighlighted) { bIsHighlighted = bInIsHighlighted; }
 	bool IsHighlighted() const { return bIsHighlighted; }
diff --git a/Engine/Plugins/Experimental/PCG/Source/PCGEditor/Private/PCGEditorStyle.cpp b/Engine/Plugins/Experimental/PCG/Source/PCGEditor/Private/PCGEditorStyle.cpp
--- a/Engine/Plugins/Experimental/PCG/Source/PCGEditor/Private/PCGEditorStyle.cpp
+++ b/Engine/Plugins/Experimental/PCG/Source/PCGEditor/Private/PCGEditorStyle.cpp
@@ -2,6 +2,8 @@
 
 #include "PCGEditorStyle.h"
 
+#include "PCGCommon.h"
+
 #include "Brushes/SlateImageBrush.h"
 #include "Interfaces/IPluginManager.h"
 #include "Misc/Paths.h"
@@ -22,7 +24,7 @@ void FPCGEditorStyle::Unregister()
 
 FPCGEditorStyle::FPCGEditorStyle() : FSlateStyleSet("PCGEditorStyle")
 {
-    static const FVector2D Icon20x20(20.0f, 20.0f);
+	static const FVector2D Icon20x20(20.0f, 20.0f);
 	
 	SetCoreContentRoot(FPaths::EngineContentDir() / TEXT("Editor/Slate"));
 	SetContentRoot(IPluginManager::Get().FindPlugin("PCG")->GetBaseDir() / TEXT("Content"));
@@ -64,6 +66,9 @@ FPCGEditorStyle::FPCGEditorStyle() : FSlateStyleSet("PCGEditorStyle")
 
 	Set(PCGEditorStyleConstants::Node_Overlay_Inactive, new IMAGE_BRUSH_SVG("Style/PCG_Node_Overlay_Inactive", Icon20x20));
 
+	Set(PCGNodeConstants::Icons::CompactNodeFilter, new IMAGE_BRUSH_SVG("Style/PCG_Graph_Filter", FVector2D(28.0f, 28.0f)));
+	Set(PCGNodeConstants::Icons::CompactNodeConvert, new IMAGE_BRUSH_SVG("Style/PCG_Graph_To", PinSize));
+
 	FInlineEditableTextBlockStyle NodeTitleStyle = FInlineEditableTextBlockStyle(FAppStyle::Get().GetWidgetStyle<FInlineEditableTextBlockStyle>("Graph.Node.NodeTitleInlineEditableText"));
 	FTextBlockStyle GraphNodeItalicTitle = FTextBlockStyle(FAppStyle::Get().GetWidgetStyle<FTextBlockStyle>("Graph.Node.NodeTitle"))
 		.SetFont(DEFAULT_FONT("BoldCondensedItalic", 10));
diff --git a/Engine/Plugins/Experimental/PCG/Source/PCGEditor/Private/SPCGEditorGraphNode.cpp b/Engine/Plugins/Experimental/PCG/Source/PCGEditor/Private/SPCGEditorGraphNode.cpp
--- a/Engine/Plugins/Experimental/PCG/Source/PCGEditor/Private/SPCGEditorGraphNode.cpp
+++ b/Engine/Plugins/Experimental/PCG/Source/PCGEditor/Private/SPCGEditorGraphNode.cpp
@@ -10,6 +10,7 @@
 #include "PCGSettingsWithDynamicInputs.h"
 
 #include "GraphEditorSettings.h"
+#include "IDocumentation.h"
 #include "SCommentBubble.h"
 #include "SGraphPin.h"
 #include "SLevelOfDetailBranchNode.h"
@@ -601,6 +602,12 @@ void SPCGEditorGraphNode::UpdateCompactNode()
 	RightNodeBox.Reset();
 	LeftNodeBox.Reset();
 
+	if (!SWidget::GetToolTip().IsValid())
+	{
+		TSharedRef<SToolTip> DefaultToolTip = IDocumentation::Get()->CreateToolTip(TAttribute< FText >(this, &SGraphNode::GetNodeTooltip), NULL, GraphNode->GetDocumentationLink(), GraphNode->GetDocumentationExcerptName());
+		SetToolTip(DefaultToolTip);
+	}
+
 	// Setup a meta tag for this node
 	FGraphNodeMetaData TagMeta(TEXT("Graphnode"));
 	PopulateMetaTag(&TagMeta);
@@ -627,29 +634,46 @@ void SPCGEditorGraphNode::UpdateCompactNode()
 		];
 	}
 
-	NodeOverlay->AddSlot()
-		.HAlign(HAlign_Center)
-		.VAlign(VAlign_Center)
-		.Padding(45.f, 0.f, 45.f, 0.f)
-		[
-			// MIDDLE
-			SNew(SVerticalBox)
-			+SVerticalBox::Slot()
+	FName CompactBodyIcon = NAME_None;
+	check(PCGEditorGraphNode);
+	if (PCGEditorGraphNode->GetCompactNodeIcon(CompactBodyIcon))
+	{
+		NodeOverlay->AddSlot()
 			.HAlign(HAlign_Center)
-			.AutoHeight()
+			.VAlign(VAlign_Center)
+			.Padding(45.f, 0.f, 45.f, 0.f)
 			[
-				SNew(STextBlock)
-					.TextStyle(FAppStyle::Get(), "Graph.Node.NodeTitle" )
-					.Text( NodeTitle.Get(), &SNodeTitle::GetHeadTitle )
-					.WrapTextAt(128.0f)
+				SNew(SImage)
+					.Image(FPCGEditorStyle::Get().GetBrush(CompactBodyIcon))
 					.ColorAndOpacity(this, &SGraphNode::GetNodeTitleIconColor)
-			]
-			+SVerticalBox::Slot()
-			.AutoHeight()
+			];
+	}
+	else
+	{
+		NodeOverlay->AddSlot()
+			.HAlign(HAlign_Center)
+			.VAlign(VAlign_Center)
+			.Padding(45.f, 0.f, 45.f, 0.f)
 			[
-				NodeTitle.ToSharedRef()
-			]
-		];
+				// MIDDLE
+				SNew(SVerticalBox)
+				+SVerticalBox::Slot()
+				.HAlign(HAlign_Center)
+				.AutoHeight()
+				[
+					SNew(STextBlock)
+						.TextStyle(FAppStyle::Get(), "Graph.Node.NodeTitle" )
+						.Text( NodeTitle.Get(), &SNodeTitle::GetHeadTitle )
+						.WrapTextAt(128.0f)
+						.ColorAndOpacity(this, &SGraphNode::GetNodeTitleIconColor)
+				]
+				+SVerticalBox::Slot()
+				.AutoHeight()
+				[
+					NodeTitle.ToSharedRef()
+				]
+			];
+	}
 	
 	NodeOverlay->AddSlot()
 		.HAlign(HAlign_Left)
