[PCG] Fix Graph Parameters Unit Tests

* User Structs are needed to be valid for undo/redo to work, was done in ctor but not after a remove
* Modifying emulation callbacks for modifying a value didn't have the right member property
* Should also fix undo/redo after a remove.

#rnx

diff --git a/Engine/Plugins/Experimental/PCG/Source/PCG/Private/PCGGraph.cpp b/Engine/Plugins/Experimental/PCG/Source/PCG/Private/PCGGraph.cpp
--- a/Engine/Plugins/Experimental/PCG/Source/PCG/Private/PCGGraph.cpp
+++ b/Engine/Plugins/Experimental/PCG/Source/PCG/Private/PCGGraph.cpp
@@ -1315,18 +1315,6 @@ void UPCGGraph::OnGraphParametersChanged(EPCGGraphParameterEvent InChangeType, F
 		}
 	}
 
-	// Also if we have no more properties, just reset
-	if (UserParameters.GetNumPropertiesInBag() == 0)
-	{
-		if (!bWasModified)
-		{
-			bWasModified = true;
-			Modify();
-		}
-
-		UserParameters.Reset();
-	}
-
 #if WITH_EDITOR
 	NotifyGraphParametersChanged(InChangeType, InChangedPropertyName);
 #endif // WITH_EDITOR
diff --git a/Engine/Plugins/Experimental/PCG/Source/PCG/Private/Tests/PCGGraphParametersTest.cpp b/Engine/Plugins/Experimental/PCG/Source/PCG/Private/Tests/PCGGraphParametersTest.cpp
--- a/Engine/Plugins/Experimental/PCG/Source/PCG/Private/Tests/PCGGraphParametersTest.cpp
+++ b/Engine/Plugins/Experimental/PCG/Source/PCG/Private/Tests/PCGGraphParametersTest.cpp
@@ -73,14 +73,12 @@ namespace PCGTests
 		check(UserParameters && UserParameters->GetPropertyBagStruct());
 
 		FProperty* ValueProperty = UserParameters->GetPropertyBagStruct()->FindPropertyByName(PropertyName);
-		FProperty* UserParametersProperty = Graph->GetClass()->FindPropertyByName(UserParametersName);
-		check(ValueProperty && UserParametersProperty);
+		check(ValueProperty);
 		Graph->PreEditChange(ValueProperty);
 
 		Callback();
 
 		FPropertyChangedEvent PropertyChangedEvent{ ValueProperty };
-		PropertyChangedEvent.MemberProperty = UserParametersProperty;
 		Graph->PostEditChangeProperty(PropertyChangedEvent);
 	}
 }
@@ -90,7 +88,8 @@ bool FPCGGraphNoUserParameters::RunTest(const FString& Parameters)
 	PCGTests::TempObject<UPCGGraph> Graph{};
 	UTEST_NOT_NULL("Graph UserParameters is not null", Graph->GetUserParametersStruct());
 
-	UTEST_FALSE("Graph UserParameters property bag is valid", Graph->GetUserParametersStruct()->IsValid());
+	// Bags are now valid (but empty) for new graphs.
+	UTEST_TRUE("Graph UserParameters property bag is valid", Graph->GetUserParametersStruct()->IsValid());
 	UTEST_EQUAL("Graph UserParameters has no property", Graph->GetUserParametersStruct()->GetNumPropertiesInBag(), 0);
 
 	return true;
