Introduced a PCG coordinate space to unify specification of space referential.
Update both CreatePoints and CreatePointsGrid.
Fixed an issue where the culling wasn't working properly for local components.
Optimized CreatePointsGrid to pre-cull, which should be faster and allow to generate larger grids as needed.

diff --git a/Engine/Plugins/Experimental/PCG/Source/PCG/Private/Elements/PCGCreatePoints.cpp b/Engine/Plugins/Experimental/PCG/Source/PCG/Private/Elements/PCGCreatePoints.cpp
--- a/Engine/Plugins/Experimental/PCG/Source/PCG/Private/Elements/PCGCreatePoints.cpp
+++ b/Engine/Plugins/Experimental/PCG/Source/PCG/Private/Elements/PCGCreatePoints.cpp
@@ -6,6 +6,7 @@
 #include "PCGContext.h"
 #include "Helpers/PCGAsync.h"
 #include "Helpers/PCGBlueprintHelpers.h"
+#include "Helpers/PCGHelpers.h"
 #include "Helpers/PCGSettingsHelpers.h"
 
 #include "GameFramework/Actor.h"
@@ -18,6 +19,21 @@ UPCGCreatePointsSettings::UPCGCreatePointsSettings()
 	PointsToCreate.Add(FPCGPoint());
 }
 
+void UPCGCreatePointsSettings::PostLoad()
+{
+	Super::PostLoad();
+
+#if WITH_EDITOR
+PRAGMA_DISABLE_DEPRECATION_WARNINGS
+	if (GridPivot_DEPRECATED != EPCGLocalGridPivot::Global)
+	{
+		CoordinateSpace = static_cast<EPCGCoordinateSpace>(static_cast<uint8>(GridPivot_DEPRECATED));
+		GridPivot_DEPRECATED = EPCGLocalGridPivot::Global;
+	}
+PRAGMA_ENABLE_DEPRECATION_WARNINGS
+#endif
+}
+
 TArray<FPCGPinProperties> UPCGCreatePointsSettings::InputPinProperties() const
 {
 	return TArray<FPCGPinProperties>();
@@ -32,76 +48,91 @@ bool FPCGCreatePointsElement::ExecuteInternal(FPCGContext* Context) const
 {
 	TRACE_CPUPROFILER_EVENT_SCOPE(FPCGCreatePointsElement::Execute);
 	
-	check(Context);
+	check(Context && Context->SourceComponent.Get());
 
 	const UPCGCreatePointsSettings* Settings = Context->GetInputSettings<UPCGCreatePointsSettings>();
 	check(Settings);
 
-	UPCGComponent* PCGComponent = nullptr;
-	check(Context->SourceComponent.Get());
+	// Used for culling, regardless of generation coordinate space
+	const UPCGSpatialData* CullingShape = Settings->bCullPointsOutsideVolume ? Cast<UPCGSpatialData>(Context->SourceComponent->GetActorPCGData()) : nullptr;
 
-	if (Settings->GridPivot == EPCGLocalGridPivot::OriginalComponent)
+	// Early out if the culling shape isn't valid
+	if (!CullingShape && Settings->bCullPointsOutsideVolume)
 	{
-		PCGComponent = Context->SourceComponent->GetOriginalComponent();
-	}
-	else if (Settings->GridPivot == EPCGLocalGridPivot::LocalComponent)
-	{
-		PCGComponent = Context->SourceComponent.Get();
+		PCGE_LOG(Error, GraphAndLog, LOCTEXT("CannotCullWithoutAShape", "Unable to cull since the supporting actor has no data."));
+		return true;
 	}
 
-	FTransform OriginalComponentTransform = FTransform();
-	FTransform ComponentTransformScaleOne = FTransform();
-	const UPCGSpatialData* Target = nullptr;
+	FTransform LocalTransform = FTransform::Identity;
 
-	if (PCGComponent)
+	if(Settings->CoordinateSpace == EPCGCoordinateSpace::OriginalComponent)
 	{
-		check(PCGComponent->GetOwner());
-
-		OriginalComponentTransform = PCGComponent->GetOwner()->GetActorTransform();
-		ComponentTransformScaleOne = FTransform(OriginalComponentTransform.Rotator(), OriginalComponentTransform.GetLocation(), FVector::One());
-		Target = Settings->bCullPointsOutsideVolume ? Cast<UPCGSpatialData>(PCGComponent->GetActorPCGData()) : nullptr;
+		check(Context->SourceComponent->GetOriginalComponent() && Context->SourceComponent->GetOriginalComponent()->GetOwner());
+		LocalTransform = Context->SourceComponent->GetOriginalComponent()->GetOwner()->GetActorTransform();
 	}
+	else if (Settings->CoordinateSpace == EPCGCoordinateSpace::LocalComponent)
+	{
+		check(Context->SourceComponent->GetOwner());
+		LocalTransform = Context->SourceComponent->GetOwner()->GetActorTransform();
+	}
+
+	// Reset scale as we are not going to derive the points size from it
+	LocalTransform.SetScale3D(FVector::One());
 
-	TArray<FPCGPoint> PointsToLoopOn = Settings->PointsToCreate;
+	const TArray<FPCGPoint>& PointsToLoopOn = Settings->PointsToCreate;
 	
 	TArray<FPCGTaggedData>& Outputs = Context->OutputData.TaggedData;
 	FPCGTaggedData& Output = Outputs.Emplace_GetRef();
 
-	UPCGPointData* PtData = NewObject<UPCGPointData>();
-	check(PtData); 
+	UPCGPointData* PointData = NewObject<UPCGPointData>();
+	check(PointData); 
 
-	TArray<FPCGPoint>& OutputPoints = PtData->GetMutablePoints();
-	Output.Data = PtData;
+	TArray<FPCGPoint>& OutputPoints = PointData->GetMutablePoints();
+	Output.Data = PointData;
 
-	if (Settings->GridPivot == EPCGLocalGridPivot::Global && !Settings->bCullPointsOutsideVolume)
+	if (Settings->CoordinateSpace == EPCGCoordinateSpace::World)
 	{
-		for (auto& Points : PointsToLoopOn)
+		if (CullingShape)
 		{
-			if (Points.Seed == 0)
+			OutputPoints.Reserve(PointsToLoopOn.Num());
+
+			for (const FPCGPoint& Point : PointsToLoopOn)
+			{
+				if (CullingShape->GetDensityAtPosition(Point.Transform.GetLocation()) > 0)
+				{
+					OutputPoints.Add(Point);
+				}
+			}
+		}
+		else
+		{
+			OutputPoints = PointsToLoopOn;
+		}
+		
+		for (FPCGPoint& Point : OutputPoints)
+		{
+			if (Point.Seed == 0)
 			{
 				// If the seed is the default value, generate a new seed based on the its transform
-				Points.Seed = UPCGBlueprintHelpers::ComputeSeedFromPosition(Points.Transform.GetLocation());
+				Point.Seed = UPCGBlueprintHelpers::ComputeSeedFromPosition(Point.Transform.GetLocation());
 			}
 		}
-
-		PtData->SetPoints(PointsToLoopOn);
 	}
 	else
 	{
-		FPCGAsync::AsyncPointProcessing(Context, PointsToLoopOn.Num(), OutputPoints, [&PointsToLoopOn, Settings, &ComponentTransformScaleOne, Target](int32 Index, FPCGPoint& OutPoint)
+		check(Settings->CoordinateSpace == EPCGCoordinateSpace::LocalComponent || Settings->CoordinateSpace == EPCGCoordinateSpace::OriginalComponent);
+
+		FPCGAsync::AsyncPointProcessing(Context, PointsToLoopOn.Num(), OutputPoints, [&PointsToLoopOn, &LocalTransform, CullingShape](int32 Index, FPCGPoint& OutPoint)
 		{
 			const FPCGPoint& InPoint = PointsToLoopOn[Index];
 			OutPoint = InPoint;
+			OutPoint.Transform *= LocalTransform;
 
-			if (Settings->GridPivot == EPCGLocalGridPivot::LocalComponent || Settings->GridPivot == EPCGLocalGridPivot::OriginalComponent)
-			{
-				OutPoint.Transform *= ComponentTransformScaleOne;
-			}
-
-			OutPoint.Seed = UPCGBlueprintHelpers::ComputeSeedFromPosition(OutPoint.Transform.GetLocation());
+			const int SeedFromPosition = UPCGBlueprintHelpers::ComputeSeedFromPosition(OutPoint.Transform.GetLocation());
+			OutPoint.Seed = (InPoint.Seed == 0 ? SeedFromPosition : PCGHelpers::ComputeSeed(InPoint.Seed, SeedFromPosition));
 
 			// Discards all points that are outside the volume
-			return !Target || (Target->GetDensityAtPosition(OutPoint.Transform.GetLocation()) > 0.0f);
+			return !CullingShape || (CullingShape->GetDensityAtPosition(OutPoint.Transform.GetLocation()) > 0.0f);
 		});
 	}
 
@@ -112,7 +143,7 @@ bool FPCGCreatePointsElement::IsCacheable(const UPCGSettings* InSettings) const
 {
 	const UPCGCreatePointsSettings* Settings = Cast<const UPCGCreatePointsSettings>(InSettings);
 
-	return Settings && Settings->GridPivot == EPCGLocalGridPivot::Global;
+	return Settings && Settings->CoordinateSpace == EPCGCoordinateSpace::World;
 }
 
 #undef LOCTEXT_NAMESPACE
diff --git a/Engine/Plugins/Experimental/PCG/Source/PCG/Private/Elements/PCGCreatePointsGrid.cpp b/Engine/Plugins/Experimental/PCG/Source/PCG/Private/Elements/PCGCreatePointsGrid.cpp
--- a/Engine/Plugins/Experimental/PCG/Source/PCG/Private/Elements/PCGCreatePointsGrid.cpp
+++ b/Engine/Plugins/Experimental/PCG/Source/PCG/Private/Elements/PCGCreatePointsGrid.cpp
@@ -12,6 +12,21 @@
 
 #define LOCTEXT_NAMESPACE "PCGCreatePointsGridElement"
 
+void UPCGCreatePointsGridSettings::PostLoad()
+{
+	Super::PostLoad();
+
+#if WITH_EDITOR
+PRAGMA_DISABLE_DEPRECATION_WARNINGS
+	if (GridPivot_DEPRECATED != EPCGGridPivot::Global)
+	{
+		CoordinateSpace = static_cast<EPCGCoordinateSpace>(static_cast<int8>(GridPivot_DEPRECATED));
+		GridPivot_DEPRECATED = EPCGGridPivot::Global;
+	}
+PRAGMA_ENABLE_DEPRECATION_WARNINGS
+#endif
+}
+
 TArray<FPCGPinProperties> UPCGCreatePointsGridSettings::InputPinProperties() const
 {
 	return TArray<FPCGPinProperties>();
@@ -26,35 +41,36 @@ bool FPCGCreatePointsGridElement::ExecuteInternal(FPCGContext* Context) const
 {
 	TRACE_CPUPROFILER_EVENT_SCOPE(FPCGCreatePointsGridElement::Execute);
 	
-	check(Context);
+	check(Context && Context->SourceComponent.Get());
 
 	const UPCGCreatePointsGridSettings* Settings = Context->GetInputSettings<UPCGCreatePointsGridSettings>();
 	check(Settings);
 
-	UPCGComponent* PCGComponent = nullptr;
-	check(Context->SourceComponent.Get());
+	// Used for culling, regardless of generation coordinate space
+	const UPCGSpatialData* CullingShape = Settings->bCullPointsOutsideVolume ? Cast<UPCGSpatialData>(Context->SourceComponent->GetActorPCGData()) : nullptr;
 
-	if (Settings->GridPivot == EPCGGridPivot::OriginalComponent)
+	// Early out if the culling shape isn't valid
+	if (!CullingShape && Settings->bCullPointsOutsideVolume)
 	{
-		PCGComponent = Context->SourceComponent->GetOriginalComponent();
-	}
-	else if (Settings->GridPivot == EPCGGridPivot::LocalComponent)
-	{
-		PCGComponent = Context->SourceComponent.Get();
+		PCGE_LOG(Error, GraphAndLog, LOCTEXT("CannotCullWithoutAShape", "Unable to cull since the supporting actor has no data."));
+		return true;
 	}
 
-	FTransform OriginalComponentTransform = FTransform();
-	FTransform ComponentTransformScaleOne = FTransform();
-	const UPCGSpatialData* Target = nullptr;
+	FTransform LocalTransform = FTransform::Identity;
 
-	if (PCGComponent)
+	if (Settings->CoordinateSpace == EPCGCoordinateSpace::OriginalComponent)
 	{
-		check(PCGComponent->GetOwner());
-
-		OriginalComponentTransform = PCGComponent->GetOwner()->GetActorTransform();
-		ComponentTransformScaleOne = FTransform(OriginalComponentTransform.Rotator(), OriginalComponentTransform.GetLocation(), FVector::One());
-		Target = Settings->bCullPointsOutsideVolume ? Cast<UPCGSpatialData>(PCGComponent->GetActorPCGData()) : nullptr;
+		check(Context->SourceComponent->GetOriginalComponent() && Context->SourceComponent->GetOriginalComponent()->GetOwner());
+		LocalTransform = Context->SourceComponent->GetOriginalComponent()->GetOwner()->GetActorTransform();
 	}
+	else if (Settings->CoordinateSpace == EPCGCoordinateSpace::LocalComponent)
+	{
+		check(Context->SourceComponent->GetOwner());
+		LocalTransform = Context->SourceComponent->GetOwner()->GetActorTransform();
+	}
+
+	// Reset the scale on the local transform since we don't want to derive the points scale from the referential
+	LocalTransform.SetScale3D(FVector::One());
 
 	if (Settings->CellSize.X <= 0.0 || Settings->CellSize.Y <= 0.0 || Settings->CellSize.Z <= 0.0)
 	{
@@ -73,10 +89,10 @@ bool FPCGCreatePointsGridElement::ExecuteInternal(FPCGContext* Context) const
 
 	int32 PointCountX = FMath::TruncToInt((2 * GridExtents.X) / CellSize.X);
 	int32 PointCountY = FMath::TruncToInt((2 * GridExtents.Y) / CellSize.Y);
-	int32 PointCountZ = FMath::TruncToInt((2 * GridExtents.Z) / CellSize.Z);
+	int32 PointCountZ = FMath::TruncToInt((2 * GridExtents.Z) / CellSize.Z); 
 
 	if (Settings->PointPosition == EPCGPointPosition::CellCorners)
-	{ 
+	{
 		PointCountX++;
 		PointCountY++;
 		PointCountZ++;
@@ -120,7 +136,7 @@ bool FPCGCreatePointsGridElement::ExecuteInternal(FPCGContext* Context) const
 		GridExtents.Z = GridExtents.Z - FMath::Fmod(GridExtents.Z, CellSize.Z / 2);
 	}
 
-	const int64 NumIterations64 = static_cast<int64>(PointCountX) * static_cast<int64>(PointCountY) * static_cast<int64>(PointCountZ);
+	int64 NumIterations64 = static_cast<int64>(PointCountX) * static_cast<int64>(PointCountY) * static_cast<int64>(PointCountZ);
 
 	if (NumIterations64 <= 0)
 	{
@@ -128,6 +144,48 @@ bool FPCGCreatePointsGridElement::ExecuteInternal(FPCGContext* Context) const
 		return true;
 	}
 
+	// Adjust point count/offset with respect to the Culling Shape bounding box if any.
+	int32 PointOffsetX = 0;
+	int32 PointOffsetY = 0;
+	int32 PointOffsetZ = 0;
+
+	if (CullingShape)
+	{
+		// Get target bounds in world space
+		FBox TargetBounds = CullingShape->GetBounds();
+
+		// Move bounds to coordinate space
+		TargetBounds = TargetBounds.InverseTransformBy(LocalTransform);
+
+		// Discretize target bounds to the cell size
+		TargetBounds.Min /= CellSize;
+		TargetBounds.Max /= CellSize;
+		
+		FIntVector DiscreteMin(FMath::FloorToInt(TargetBounds.Min.X), FMath::FloorToInt(TargetBounds.Min.Y), FMath::FloorToInt(TargetBounds.Min.Z));
+		FIntVector DiscreteMax(FMath::CeilToInt(TargetBounds.Max.X), FMath::CeilToInt(TargetBounds.Max.Y), FMath::CeilToInt(TargetBounds.Max.Z));
+
+		// Grid is built from -GridExtents to +GridExtents, so intuitively we'll do this:
+		//PointOffset = (PointCount / 2) + FMath::Max(-PointCount / 2, DiscreteMin);
+		//PointCount = (PointCount / 2) + FMath::Min(+PointCount / 2, DiscreteMax) - PointOffset;
+		PointOffsetX = FMath::Max(0, PointCountX / 2 + DiscreteMin.X);
+		PointOffsetY = FMath::Max(0, PointCountY / 2 + DiscreteMin.Y);
+		PointOffsetZ = FMath::Max(0, PointCountZ / 2 + DiscreteMin.Z);
+
+		// Implementation note: the +1 here is to make sure we properly manage the case where PointCount isn't even
+		PointCountX = FMath::Min(PointCountX, (PointCountX+1) / 2 + DiscreteMax.X) - PointOffsetX;
+		PointCountY = FMath::Min(PointCountY, (PointCountY+1) / 2 + DiscreteMax.Y) - PointOffsetY;
+		PointCountZ = FMath::Min(PointCountZ, (PointCountZ+1) / 2 + DiscreteMax.Z) - PointOffsetZ;
+
+		// Update iteration count
+		NumIterations64 = static_cast<int64>(PointCountX) * static_cast<int64>(PointCountY) * static_cast<int64>(PointCountZ);
+
+		// If it is now equal to zero, return quietly as we might have culled completely
+		if (NumIterations64 <= 0)
+		{
+			return true;
+		}
+	}
+
 	if (NumIterations64 >= MAX_int32)
 	{
 		PCGE_LOG(Error, GraphAndLog, LOCTEXT("Overflow_int32", "The number of iterations produced is larger than what a 32-bit integer can hold."));
@@ -151,13 +209,13 @@ bool FPCGCreatePointsGridElement::ExecuteInternal(FPCGContext* Context) const
 	TArray<FPCGPoint>& OutputPoints = PtData->GetMutablePoints();
 	Output.Data = PtData;
 
-	FPCGAsync::AsyncPointProcessing(Context, NumIterations, OutputPoints, [Settings, Target, &PointCountX, &PointCountY, &CellSize, &GridExtents, &ComponentTransformScaleOne](int32 Index, FPCGPoint& OutPoint)
+	FPCGAsync::AsyncPointProcessing(Context, NumIterations, OutputPoints, [Settings, CullingShape, &PointCountX, &PointCountY, &PointOffsetX, &PointOffsetY, &PointOffsetZ, &CellSize, &GridExtents, &LocalTransform](int32 Index, FPCGPoint& OutPoint)
 	{
 		OutPoint = FPCGPoint();
 
-		double XCoordinate = Index % PointCountX;
-		double YCoordinate = (Index / PointCountX) % PointCountY;
-		double ZCoordinate = Index / (PointCountX * PointCountY);
+		double XCoordinate = (Index % PointCountX) + PointOffsetX;
+		double YCoordinate = ((Index / PointCountX) % PointCountY) + PointOffsetY;
+		double ZCoordinate = (Index / (PointCountX * PointCountY)) + PointOffsetZ;
 
 		if (Settings->PointPosition == EPCGPointPosition::CellCenter)
 		{
@@ -181,26 +239,24 @@ bool FPCGCreatePointsGridElement::ExecuteInternal(FPCGContext* Context) const
 			ZCoordinate = 0.0;
 		}
 
-		const FVector GridTransformLocation((CellSize.X * XCoordinate) - GridExtents.X, (CellSize.Y * YCoordinate) - GridExtents.Y, (CellSize.Z * ZCoordinate) - GridExtents.Z);
+		FTransform PointTransform = FTransform(FVector((CellSize.X * XCoordinate) - GridExtents.X, (CellSize.Y * YCoordinate) - GridExtents.Y, (CellSize.Z * ZCoordinate) - GridExtents.Z));
 
-		if (Settings->GridPivot == EPCGGridPivot::LocalComponent || Settings->GridPivot == EPCGGridPivot::OriginalComponent)
+		if (Settings->CoordinateSpace == EPCGCoordinateSpace::LocalComponent || Settings->CoordinateSpace == EPCGCoordinateSpace::OriginalComponent)
 		{
-			OutPoint.Transform = FTransform(FRotator::ZeroRotator, GridTransformLocation, FVector::One()) * ComponentTransformScaleOne;
-		}
-		else
-		{
-			OutPoint.Transform.SetLocation(GridTransformLocation);
+			PointTransform *= LocalTransform;
 		}
 
+		OutPoint.Transform = PointTransform;
+
 		if (Settings->bSetPointsBounds)
 		{
 			OutPoint.SetExtents(CellSize * 0.5);
 		}
 
-		OutPoint.Seed = PCGHelpers::ComputeSeed(OutPoint.Transform.GetLocation().X, OutPoint.Transform.GetLocation().Y, OutPoint.Transform.GetLocation().Z);
+		OutPoint.Seed = PCGHelpers::ComputeSeedFromPosition(PointTransform.GetLocation());
 
 		// Discards points outside of the volume
-		return !Target || (Target->GetDensityAtPosition(OutPoint.Transform.GetLocation()) > 0.0f);
+		return !CullingShape || (CullingShape->GetDensityAtPosition(PointTransform.GetLocation()) > 0.0f);
 	});
 
 	return true;
@@ -213,15 +269,15 @@ void FPCGCreatePointsGridElement::GetDependenciesCrc(const FPCGDataCollection& I
 	
 	if (const UPCGCreatePointsGridSettings* Settings = Cast<UPCGCreatePointsGridSettings>(InSettings))
 	{
-		int GridPivot = static_cast<int>(EPCGGridPivot::Global);
+		int CoordinateSpace = static_cast<int>(EPCGCoordinateSpace::World);
 		bool bCullPointsOutsideVolume = false;
-		PCGSettingsHelpers::GetOverrideValue(InInput, Settings, GET_MEMBER_NAME_CHECKED(UPCGCreatePointsGridSettings, GridPivot), static_cast<int>(Settings->GridPivot), GridPivot);
+		PCGSettingsHelpers::GetOverrideValue(InInput, Settings, GET_MEMBER_NAME_CHECKED(UPCGCreatePointsGridSettings, CoordinateSpace), static_cast<int>(Settings->CoordinateSpace), CoordinateSpace);
 		PCGSettingsHelpers::GetOverrideValue(InInput, Settings, GET_MEMBER_NAME_CHECKED(UPCGCreatePointsGridSettings, bCullPointsOutsideVolume), Settings->bCullPointsOutsideVolume, bCullPointsOutsideVolume);
 		
 		// We're using the bounds of the pcg volume, so we extract the actor data here
-		EPCGGridPivot EnumGridPivot = static_cast<EPCGGridPivot>(GridPivot);
+		EPCGCoordinateSpace EnumCoordinateSpace = static_cast<EPCGCoordinateSpace>(CoordinateSpace);
 
-		if ((EnumGridPivot == EPCGGridPivot::OriginalComponent || EnumGridPivot == EPCGGridPivot::LocalComponent || bCullPointsOutsideVolume) && InComponent)
+		if ((EnumCoordinateSpace == EPCGCoordinateSpace::OriginalComponent || EnumCoordinateSpace == EPCGCoordinateSpace::LocalComponent || bCullPointsOutsideVolume) && InComponent)
 		{
 			if (const UPCGData* Data = InComponent->GetActorPCGData())
 			{
diff --git a/Engine/Plugins/Experimental/PCG/Source/PCG/Private/Tests/Elements/PCGCreatePointsGridTest.cpp b/Engine/Plugins/Experimental/PCG/Source/PCG/Private/Tests/Elements/PCGCreatePointsGridTest.cpp
--- a/Engine/Plugins/Experimental/PCG/Source/PCG/Private/Tests/Elements/PCGCreatePointsGridTest.cpp
+++ b/Engine/Plugins/Experimental/PCG/Source/PCG/Private/Tests/Elements/PCGCreatePointsGridTest.cpp
@@ -17,7 +17,7 @@ namespace CreatePointsGridCommonTests
 		PCGTestsCommon::GenerateSettings<UPCGCreatePointsGridSettings>(TestData);
 		UPCGCreatePointsGridSettings* Settings = CastChecked<UPCGCreatePointsGridSettings>(TestData.Settings);
 		Settings->PointPosition = Position;
-		Settings->GridPivot = EPCGGridPivot::Global;
+		Settings->CoordinateSpace = EPCGCoordinateSpace::World;
 		Settings->GridExtents = GridSize;
 		Settings->CellSize = Cell;
 
@@ -30,13 +30,13 @@ namespace CreatePointsGridCommonTests
 		return Context;
 	}
 
-	PCGTestsCommon::FTestData GenerateTestData(EPCGPointPosition Position, EPCGGridPivot Pivot)
+	PCGTestsCommon::FTestData GenerateTestData(EPCGPointPosition Position, EPCGCoordinateSpace CoordinateSpace)
 	{
 		PCGTestsCommon::FTestData TestData(42, nullptr, APCGVolume::StaticClass());
 		PCGTestsCommon::GenerateSettings<UPCGCreatePointsGridSettings>(TestData);
 		UPCGCreatePointsGridSettings* Settings = CastChecked<UPCGCreatePointsGridSettings>(TestData.Settings);
 		Settings->PointPosition = Position;
-		Settings->GridPivot = Pivot;
+		Settings->CoordinateSpace = CoordinateSpace;
 		Settings->GridExtents = FVector(100.0, 100.0, 50.0);
 		Settings->CellSize = FVector(100.0);
 
@@ -322,7 +322,7 @@ IMPLEMENT_CUSTOM_SIMPLE_AUTOMATION_TEST(FPCGCreatePointsGridTest_LocalCenter, FP
 
 bool FPCGCreatePointsGridTest_LocalCenter::RunTest(const FString& Parameters)
 {
-	PCGTestsCommon::FTestData TestData = CreatePointsGridCommonTests::GenerateTestData(EPCGPointPosition::CellCenter, EPCGGridPivot::LocalComponent);
+	PCGTestsCommon::FTestData TestData = CreatePointsGridCommonTests::GenerateTestData(EPCGPointPosition::CellCenter, EPCGCoordinateSpace::LocalComponent);
 
 	// Creating the volume and proper bounds
 	AActor* TestActor = TestData.TestActor;
@@ -367,7 +367,7 @@ IMPLEMENT_CUSTOM_SIMPLE_AUTOMATION_TEST(FPCGCreatePointsGridTest_LocalCorner, FP
 
 bool FPCGCreatePointsGridTest_LocalCorner::RunTest(const FString& Parameters)
 {
-	PCGTestsCommon::FTestData TestData = CreatePointsGridCommonTests::GenerateTestData(EPCGPointPosition::CellCorners, EPCGGridPivot::LocalComponent);
+	PCGTestsCommon::FTestData TestData = CreatePointsGridCommonTests::GenerateTestData(EPCGPointPosition::CellCorners, EPCGCoordinateSpace::LocalComponent);
 
 	// Creating the volume and proper bounds
 	AActor* TestActor = TestData.TestActor;
diff --git a/Engine/Plugins/Experimental/PCG/Source/PCG/Public/Elements/PCGCreatePoints.h b/Engine/Plugins/Experimental/PCG/Source/PCG/Public/Elements/PCGCreatePoints.h
--- a/Engine/Plugins/Experimental/PCG/Source/PCG/Public/Elements/PCGCreatePoints.h
+++ b/Engine/Plugins/Experimental/PCG/Source/PCG/Public/Elements/PCGCreatePoints.h
@@ -3,11 +3,13 @@
 
 #include "PCGSettings.h"
 
+#include "PCGCommon.h"
 #include "Data/PCGPointData.h"
+
 #include "PCGCreatePoints.generated.h"
 
 UENUM()
-enum class EPCGLocalGridPivot : uint8
+enum class UE_DEPRECATED(5.4, "Not used anymore, replaced by EPCGCoordinateSpace.") EPCGLocalGridPivot : uint8
 {
 	Global,
 	OriginalComponent,
@@ -24,6 +26,11 @@ class UPCGCreatePointsSettings : public UPCGSettings
 
 public:
 	UPCGCreatePointsSettings();
+
+	//~Begin UObject interface
+	virtual void PostLoad() override;
+	//~End UObject interface
+
 	//~Begin UPCGSettings interface
 #if WITH_EDITOR
 	virtual FName GetDefaultNodeName() const override { return FName(TEXT("CreatePoints")); }
@@ -42,9 +49,17 @@ public:
 	UPROPERTY(BlueprintReadWrite, EditAnywhere, Category = Settings, meta = (PCG_Overridable))
 	TArray<FPCGPoint> PointsToCreate;
 
+#if WITH_EDITORONLY_DATA
+PRAGMA_DISABLE_DEPRECATION_WARNINGS
 	/** Sets the points transform to world or local space*/
-	UPROPERTY(BlueprintReadWrite, EditAnywhere, Category = Settings, meta = (PCG_Overridable))
-	EPCGLocalGridPivot GridPivot = EPCGLocalGridPivot::Global;
+	UPROPERTY()
+	EPCGLocalGridPivot GridPivot_DEPRECATED = EPCGLocalGridPivot::Global;
+PRAGMA_ENABLE_DEPRECATION_WARNINGS
+#endif
+
+	/** Sets the generation referential of the points */
+	UPROPERTY(BlueprintReadWrite, EditAnywhere, Category = Settings, meta = (PCG_Overridable, PCG_OverrideAliases="GridPivot"))
+	EPCGCoordinateSpace CoordinateSpace = EPCGCoordinateSpace::World;
 
 	/** If true, points are removed if they are outside of the volume */
 	UPROPERTY(BlueprintReadWrite, EditAnywhere, Category = Settings, meta = (PCG_Overridable))
diff --git a/Engine/Plugins/Experimental/PCG/Source/PCG/Public/Elements/PCGCreatePointsGrid.h b/Engine/Plugins/Experimental/PCG/Source/PCG/Public/Elements/PCGCreatePointsGrid.h
--- a/Engine/Plugins/Experimental/PCG/Source/PCG/Public/Elements/PCGCreatePointsGrid.h
+++ b/Engine/Plugins/Experimental/PCG/Source/PCG/Public/Elements/PCGCreatePointsGrid.h
@@ -14,7 +14,7 @@ enum class EPCGPointPosition : uint8
 };
 
 UENUM()
-enum class EPCGGridPivot : uint8
+enum class UE_DEPRECATED(5.4, "Not used anymore, replaced by EPCGCoordinateSpace.") EPCGGridPivot : uint8
 {
 	Global,
 	OriginalComponent,
@@ -30,6 +30,10 @@ class UPCGCreatePointsGridSettings : public UPCGSettings
 	GENERATED_BODY()
 
 public:
+	//~Begin UObject interface
+	virtual void PostLoad() override;
+	//~End UObject interface
+
 	//~Begin UPCGSettings interface
 #if WITH_EDITOR
 	virtual FName GetDefaultNodeName() const override { return FName(TEXT("CreatePointsGrid")); }
@@ -51,9 +55,17 @@ public:
 	UPROPERTY(BlueprintReadWrite, EditAnywhere, Category = Settings, meta = (PCG_Overridable, ClampMin = "0.0", UIMin = "0.0"))
 	FVector CellSize = FVector(100.0, 100.0, 100.0);
 
+#if WITH_EDITORONLY_DATA
+PRAGMA_DISABLE_DEPRECATION_WARNINGS
 	/** Sets the points transform to world or local space*/
-	UPROPERTY(BlueprintReadWrite, EditAnywhere, Category = Settings, meta = (PCG_Overridable))
-	EPCGGridPivot GridPivot = EPCGGridPivot::Global;
+	UPROPERTY()
+	EPCGGridPivot GridPivot_DEPRECATED = EPCGGridPivot::Global;
+PRAGMA_ENABLE_DEPRECATION_WARNINGS
+#endif
+
+	/** Sets the generation referential of the points */
+	UPROPERTY(BlueprintReadWrite, EditAnywhere, Category = Settings, meta = (PCG_Overridable, PCG_OverrideAliases="GridPivot"))
+	EPCGCoordinateSpace CoordinateSpace = EPCGCoordinateSpace::World;
 
 	/** If true, the bounds of the points are set to 50.0, if false, 1.0 */
 	UPROPERTY(BlueprintReadWrite, EditAnywhere, Category = Settings, meta = (PCG_Overridable))
diff --git a/Engine/Plugins/Experimental/PCG/Source/PCG/Public/PCGCommon.h b/Engine/Plugins/Experimental/PCG/Source/PCG/Public/PCGCommon.h
--- a/Engine/Plugins/Experimental/PCG/Source/PCG/Public/PCGCommon.h
+++ b/Engine/Plugins/Experimental/PCG/Source/PCG/Public/PCGCommon.h
@@ -151,6 +151,15 @@ namespace PCGFeatureSwitches
 	extern PCG_API TAutoConsoleVariable<bool> CVarCheckSamplerMemory;
 }
 
+/** Describes space referential for operations that create data */
+UENUM(BlueprintType)
+enum class EPCGCoordinateSpace : uint8
+{
+	World UMETA(DisplayName = "Global"),
+	OriginalComponent,
+	LocalComponent
+};
+
 /** Describes one or more target execution grids. */
 UENUM(meta = (Bitflags))
 enum class EPCGHiGenGrid : uint32
