[PCG] Graph instances editor improvements

* Make the asset a darker color to differentiate them from PCG graphs.
* Allow to create graph instances from the content browser.
* Allow to change the graph in a graph instance.

#rnx

diff --git a/Engine/Plugins/Experimental/PCG/Source/PCG/Private/PCGGraph.cpp b/Engine/Plugins/Experimental/PCG/Source/PCG/Private/PCGGraph.cpp
--- a/Engine/Plugins/Experimental/PCG/Source/PCG/Private/PCGGraph.cpp
+++ b/Engine/Plugins/Experimental/PCG/Source/PCG/Private/PCGGraph.cpp
@@ -1607,19 +1607,13 @@ void UPCGGraphInstance::OnGraphStructureChanged(UPCGGraphInterface* InGraph)
 
 bool UPCGGraphInstance::CanEditChange(const FProperty* InProperty) const
 {
-	// Graph can only be changed if it is in a PCGComponent (not local) or a PCGSubgraphSettings
+	// Graph can only be changed if it is not in a local PCGComponent
 	if (InProperty && InProperty->GetFName() == GET_MEMBER_NAME_CHECKED(UPCGGraphInstance, Graph))
 	{
-		UObject* Outer = this->GetOuter();
-
-		if (UPCGComponent* Component = Cast<UPCGComponent>(Outer))
+		if (UPCGComponent* Component = Cast<UPCGComponent>(GetOuter()))
 		{
 			return !Component->IsLocalComponent();
 		}
-		else
-		{
-			return Outer && Outer->IsA<UPCGSubgraphSettings>();
-		}
 	}
 
 	return true;
diff --git a/Engine/Plugins/Experimental/PCG/Source/PCGEditor/Private/AssetDefinitions/AssetDefinition_PCGGraphInstance.cpp b/Engine/Plugins/Experimental/PCG/Source/PCGEditor/Private/AssetDefinitions/AssetDefinition_PCGGraphInstance.cpp
--- a/Engine/Plugins/Experimental/PCG/Source/PCGEditor/Private/AssetDefinitions/AssetDefinition_PCGGraphInstance.cpp
+++ b/Engine/Plugins/Experimental/PCG/Source/PCGEditor/Private/AssetDefinitions/AssetDefinition_PCGGraphInstance.cpp
@@ -11,6 +11,15 @@ FText UAssetDefinition_PCGGraphInstance::GetAssetDisplayName() const
 	return LOCTEXT("DisplayName", "PCG Graph Instance");
 }
 
+FLinearColor UAssetDefinition_PCGGraphInstance::GetAssetColor() const
+{
+	FLinearColor Color = UAssetDefinition_PCGGraphInterface::GetAssetColor();
+	// Make it darker
+	Color *= 0.25f;
+	Color.A = 1.0f;
+	return Color;
+}
+
 TSoftClassPtr<UObject> UAssetDefinition_PCGGraphInstance::GetAssetClass() const
 {
 	return UPCGGraphInstance::StaticClass();
diff --git a/Engine/Plugins/Experimental/PCG/Source/PCGEditor/Private/AssetDefinitions/AssetDefinition_PCGGraphInstance.h b/Engine/Plugins/Experimental/PCG/Source/PCGEditor/Private/AssetDefinitions/AssetDefinition_PCGGraphInstance.h
--- a/Engine/Plugins/Experimental/PCG/Source/PCGEditor/Private/AssetDefinitions/AssetDefinition_PCGGraphInstance.h
+++ b/Engine/Plugins/Experimental/PCG/Source/PCGEditor/Private/AssetDefinitions/AssetDefinition_PCGGraphInstance.h
@@ -13,5 +13,6 @@ class UAssetDefinition_PCGGraphInstance : public UAssetDefinition_PCGGraphInterf
 
 public:
 	virtual FText GetAssetDisplayName() const override;
+	virtual FLinearColor GetAssetColor() const override;
 	virtual TSoftClassPtr<UObject> GetAssetClass() const override;
 };
diff --git a/Engine/Plugins/Experimental/PCG/Source/PCGEditor/Private/PCGGraphFactory.cpp b/Engine/Plugins/Experimental/PCG/Source/PCGEditor/Private/PCGGraphFactory.cpp
--- a/Engine/Plugins/Experimental/PCG/Source/PCGEditor/Private/PCGGraphFactory.cpp
+++ b/Engine/Plugins/Experimental/PCG/Source/PCGEditor/Private/PCGGraphFactory.cpp
@@ -51,5 +51,5 @@ UObject* UPCGGraphInstanceFactory::FactoryCreateNew(UClass* InClass, UObject* In
 
 bool UPCGGraphInstanceFactory::ShouldShowInNewMenu() const
 {
-	return false;
+	return true;
 }
\ No newline at end of file
