[PCG] Remove error in PointFromMesh element when no SM is provided
#rnx

diff --git a/Engine/Plugins/Experimental/PCG/Source/PCG/Private/Elements/PCGPointFromMeshElement.cpp b/Engine/Plugins/Experimental/PCG/Source/PCG/Private/Elements/PCGPointFromMeshElement.cpp
--- a/Engine/Plugins/Experimental/PCG/Source/PCG/Private/Elements/PCGPointFromMeshElement.cpp
+++ b/Engine/Plugins/Experimental/PCG/Source/PCG/Private/Elements/PCGPointFromMeshElement.cpp
@@ -40,6 +40,11 @@ bool FPCGPointFromMeshElement::ExecuteInternal(FPCGContext* Context) const
 	const UPCGPointFromMeshSettings* Settings = Context->GetInputSettings<UPCGPointFromMeshSettings>();
 	check(Settings);
 
+	if (Settings->StaticMesh.IsNull())
+	{
+		return true;
+	}
+
 	if (!Settings->StaticMesh.LoadSynchronous())
 	{
 		PCGE_LOG(Error, GraphAndLog, LOCTEXT("LoadStaticMeshFailed", "Failed to load StaticMesh"));
