Fixed an issue in the find value for compressed metadata attributes.
Cleaned up the merge element to prevent issue if it had overrides.
Updated tests accordingly.

#rnx

diff --git a/Engine/Plugins/Experimental/PCG/Source/PCG/Private/Elements/PCGMergeElement.cpp b/Engine/Plugins/Experimental/PCG/Source/PCG/Private/Elements/PCGMergeElement.cpp
--- a/Engine/Plugins/Experimental/PCG/Source/PCG/Private/Elements/PCGMergeElement.cpp
+++ b/Engine/Plugins/Experimental/PCG/Source/PCG/Private/Elements/PCGMergeElement.cpp
@@ -48,7 +48,7 @@ bool FPCGMergeElement::ExecuteInternal(FPCGContext* Context) const
 
 	const bool bMergeMetadata = Settings->bMergeMetadata;
 
-	TArray<FPCGTaggedData> Sources = Context->InputData.GetInputs();
+	TArray<FPCGTaggedData> Sources = Context->InputData.GetInputsByPin(PCGPinConstants::DefaultInputLabel);
 	TArray<FPCGTaggedData>& Outputs = Context->OutputData.TaggedData;
 
 	if (Sources.IsEmpty())
@@ -87,7 +87,7 @@ bool FPCGMergeElement::ExecuteInternal(FPCGContext* Context) const
 			}
 			
 			check(TargetTaggedData);
-			TargetTaggedData->Tags.Append(Source.Tags); // TODO: only unique? if yes, fix union too
+			TargetTaggedData->Tags.Append(Source.Tags);
 		}
 	}
 
diff --git a/Engine/Plugins/Experimental/PCG/Source/PCG/Private/MeshSelectors/PCGMeshSelectorWeightedByCategory.cpp b/Engine/Plugins/Experimental/PCG/Source/PCG/Private/MeshSelectors/PCGMeshSelectorWeightedByCategory.cpp
--- a/Engine/Plugins/Experimental/PCG/Source/PCG/Private/MeshSelectors/PCGMeshSelectorWeightedByCategory.cpp
+++ b/Engine/Plugins/Experimental/PCG/Source/PCG/Private/MeshSelectors/PCGMeshSelectorWeightedByCategory.cpp
@@ -82,10 +82,9 @@ bool UPCGMeshSelectorWeightedByCategory::SelectInstances(
 				continue;
 			}
 
-			PCGMetadataValueKey ValueKey = Attribute->FindValue<FString>(Entry.CategoryEntry);
+			PCGMetadataValueKey ValueKey = Attribute->FindValue(Entry.CategoryEntry);
 
-			// DefaultValueKey is only an invalid ValueKey when the CategoryEntry is not the AttributeDefaultValue
-			if (ValueKey == PCGDefaultValueKey && !Entry.CategoryEntry.Equals(AttributeDefaultValue))
+			if (ValueKey == PCGNotFoundValueKey)
 			{
 				PCGE_LOG_C(Verbose, LogOnly, &Context, FText::Format(LOCTEXT("UnusedCategory", "Unused category '{0}'. Not a valid value for attribute '{1}'."), 
 					FText::FromString(Entry.CategoryEntry), FText::FromName(CategoryAttribute)));
diff --git a/Engine/Plugins/Experimental/PCG/Source/PCG/Private/Tests/Elements/PCGStaticMeshSpawnerTest.cpp b/Engine/Plugins/Experimental/PCG/Source/PCG/Private/Tests/Elements/PCGStaticMeshSpawnerTest.cpp
--- a/Engine/Plugins/Experimental/PCG/Source/PCG/Private/Tests/Elements/PCGStaticMeshSpawnerTest.cpp
+++ b/Engine/Plugins/Experimental/PCG/Source/PCG/Private/Tests/Elements/PCGStaticMeshSpawnerTest.cpp
@@ -96,7 +96,7 @@ void ValidateOutPointData(
 	}
 
 	const FPCGMetadataAttribute<FString>* OutAttribute = static_cast<const FPCGMetadataAttribute<FString>*>(AttributeBase);
-	const FString& MeshPath = FSoftObjectPath(ISMC->GetStaticMesh()).ToString();
+	const FString MeshPath = ISMC->GetStaticMesh().GetPath();
 	const PCGMetadataValueKey MeshValueKey = OutAttribute->FindValue(MeshPath);
 
 	if (MeshValueKey == PCGDefaultValueKey)
diff --git a/Engine/Plugins/Experimental/PCG/Source/PCG/Public/Metadata/PCGMetadataAttributeTpl.h b/Engine/Plugins/Experimental/PCG/Source/PCG/Public/Metadata/PCGMetadataAttributeTpl.h
--- a/Engine/Plugins/Experimental/PCG/Source/PCG/Public/Metadata/PCGMetadataAttributeTpl.h
+++ b/Engine/Plugins/Experimental/PCG/Source/PCG/Public/Metadata/PCGMetadataAttributeTpl.h
@@ -443,16 +443,16 @@ public:
 	/** Adds the value, returns the value key for the given value */
 	PCGMetadataValueKey AddValue(const T& InValue)
 	{
-		PCGMetadataValueKey FoundValue = FindValue(InValue);
+		PCGMetadataValueKey FoundValueKey = FindValue(InValue);
 
-		if (FoundValue == PCGDefaultValueKey)
+		if (FoundValueKey == PCGNotFoundValueKey)
 		{
 			FWriteScopeLock ScopeLock(ValueLock);
 			return Values.Add(InValue) + ValueKeyOffset;
 		}
 		else
 		{
-			return FoundValue;
+			return FoundValueKey;
 		}
 	}
 
@@ -505,8 +505,13 @@ public:
 	template<typename IT = T, typename TEnableIf<PCG::Private::MetadataTraits<IT>::CompressData>::Type* = nullptr>
 	PCGMetadataValueKey FindValue(const T& InValue) const
 	{
-		PCGMetadataValueKey ParentValueKey = (GetParent() ? GetParent()->FindValue(InValue) : PCGDefaultValueKey);
-		if (ParentValueKey != PCGDefaultValueKey)
+		if (InValue == DefaultValue)
+		{
+			return PCGDefaultValueKey;
+		}
+
+		PCGMetadataValueKey ParentValueKey = (GetParent() ? GetParent()->FindValue(InValue) : PCGNotFoundValueKey);
+		if (ParentValueKey != PCGNotFoundValueKey)
 		{
 			return ParentValueKey;
 		}
@@ -522,7 +527,7 @@ public:
 			}
 			else
 			{
-				return ParentValueKey;
+				return PCGNotFoundValueKey;
 			}
 		}
 	}
@@ -530,7 +535,7 @@ public:
 	template<typename IT = T, typename TEnableIf<!PCG::Private::MetadataTraits<IT>::CompressData>::Type* = nullptr>
 	PCGMetadataValueKey FindValue(const T& InValue) const
 	{
-		return PCGDefaultValueKey;
+		return PCGNotFoundValueKey;
 	}
 
 	void SetDefaultValue(const T& Value)
diff --git a/Engine/Plugins/Experimental/PCG/Source/PCG/Public/Metadata/PCGMetadataCommon.h b/Engine/Plugins/Experimental/PCG/Source/PCG/Public/Metadata/PCGMetadataCommon.h
--- a/Engine/Plugins/Experimental/PCG/Source/PCG/Public/Metadata/PCGMetadataCommon.h
+++ b/Engine/Plugins/Experimental/PCG/Source/PCG/Public/Metadata/PCGMetadataCommon.h
@@ -11,6 +11,7 @@ typedef int32 PCGMetadataValueKey;
 
 const PCGMetadataEntryKey PCGInvalidEntryKey = -1;
 const PCGMetadataValueKey PCGDefaultValueKey = -1;
+const PCGMetadataValueKey PCGNotFoundValueKey = -2;
 
 UENUM()
 enum class EPCGMetadataOp : uint8
