[PCG] Update HiGenGridSize display name to be consistent

#rnx

diff --git a/Engine/Plugins/Experimental/PCG/Source/PCG/Public/Elements/PCGHiGenGridSize.h b/Engine/Plugins/Experimental/PCG/Source/PCG/Public/Elements/PCGHiGenGridSize.h
--- a/Engine/Plugins/Experimental/PCG/Source/PCG/Public/Elements/PCGHiGenGridSize.h
+++ b/Engine/Plugins/Experimental/PCG/Source/PCG/Public/Elements/PCGHiGenGridSize.h
@@ -38,7 +38,7 @@ protected:
 	virtual FPCGElementPtr CreateElement() const override;
 	//~End UPCGSettings interface
 
-	UPROPERTY(BlueprintReadWrite, EditAnywhere, Category = Settings)
+	UPROPERTY(BlueprintReadWrite, EditAnywhere, Category = Settings, meta = (DisplayName = "HiGen Grid Size"))
 	EPCGHiGenGrid HiGenGridSize = EPCGHiGenGrid::Grid256;
 };
 
