[PCG] Debug object tree - show tooltips with error summaries on error icons and remove slate property lambdas

diff --git a/Engine/Plugins/Experimental/PCG/Source/PCG/Private/Utils/PCGNodeVisualLogs.cpp b/Engine/Plugins/Experimental/PCG/Source/PCG/Private/Utils/PCGNodeVisualLogs.cpp
--- a/Engine/Plugins/Experimental/PCG/Source/PCG/Private/Utils/PCGNodeVisualLogs.cpp
+++ b/Engine/Plugins/Experimental/PCG/Source/PCG/Private/Utils/PCGNodeVisualLogs.cpp
@@ -142,11 +142,16 @@ FText FPCGNodeVisualLogs::GetLogsSummaryText(const UPCGNode* InNode, ELogVerbosi
 	return Summary;
 }
 
-FText FPCGNodeVisualLogs::GetLogsSummaryText(const FPCGStack& InBaseStack) const
+FText FPCGNodeVisualLogs::GetLogsSummaryText(const FPCGStack& InBaseStack, ELogVerbosity::Type* OutMinimumVerbosity) const
 {
 	FReadScopeLock ScopedReadLock(LogsLock);
 
-	FText ResultText = FText::GetEmpty();
+	FText Summary = FText::GetEmpty();
+
+	if (OutMinimumVerbosity)
+	{
+		*OutMinimumVerbosity = ELogVerbosity::All;
+	}
 
 	int32 LogCounter = 0;
 
@@ -159,27 +164,32 @@ FText FPCGNodeVisualLogs::GetLogsSummaryText(const FPCGStack& InBaseStack) const
 
 		const FPCGPerNodeVisualLogs& NodeLogs = Entry.Value;
 
-		for (int32 i = 0; i < NodeLogs.Num(); ++i)
+		for (int32 LogIndex = 0; LogIndex < NodeLogs.Num(); ++LogIndex)
 		{
 			++LogCounter;
 
 			if (LogCounter > MaxLogsInSummary)
 			{
-				ResultText = FText::Format(FText::FromString(TEXT("{0}\n...")), ResultText);
-				return ResultText;
+				Summary = FText::Format(FText::FromString(TEXT("{0}\n...")), Summary);
+				return Summary;
 			}
 
 			if (LogCounter > 1)
 			{
-				ResultText = FText::Format(FText::FromString(TEXT("{0}\n")), ResultText);
+				Summary = FText::Format(FText::FromString(TEXT("{0}\n")), Summary);
 			}
 
-			const FText MessageVerbosity = NodeLogs[i].Verbosity == ELogVerbosity::Warning ? FText::FromString(TEXT("Warning")) : FText::FromString(TEXT("Error"));
-			ResultText = FText::Format(LOCTEXT("NodeTooltipLog", "{0}{1}: {2}"), ResultText, /*i + 1, NodeLogs.Num(),*/ MessageVerbosity, NodeLogs[i].Message);
+			if (OutMinimumVerbosity)
+			{
+				*OutMinimumVerbosity = FMath::Min(*OutMinimumVerbosity, NodeLogs[LogIndex].Verbosity);
+			}
+
+			const FText MessageVerbosity = NodeLogs[LogIndex].Verbosity == ELogVerbosity::Warning ? FText::FromString(TEXT("Warning")) : FText::FromString(TEXT("Error"));
+			Summary = FText::Format(LOCTEXT("NodeTooltipLog", "{0}{1}: {2}"), Summary, MessageVerbosity, NodeLogs[LogIndex].Message);
 		}
 	}
 
-	return ResultText;
+	return Summary;
 }
 
 void FPCGNodeVisualLogs::ClearLogs(const FPCGStack& InPCGStack)
diff --git a/Engine/Plugins/Experimental/PCG/Source/PCG/Public/Utils/PCGNodeVisualLogs.h b/Engine/Plugins/Experimental/PCG/Source/PCG/Public/Utils/PCGNodeVisualLogs.h
--- a/Engine/Plugins/Experimental/PCG/Source/PCG/Public/Utils/PCGNodeVisualLogs.h
+++ b/Engine/Plugins/Experimental/PCG/Source/PCG/Public/Utils/PCGNodeVisualLogs.h
@@ -45,7 +45,7 @@ public:
 	bool HasLogsOfVerbosity(const FPCGStack& InPCGStack, ELogVerbosity::Type InVerbosity) const;
 
 	/** Summary text of all visual logs produced while executing the provided base stack, appropriate for display in graph editor tooltip. */
-	FText GetLogsSummaryText(const FPCGStack& InBaseStack) const;
+	FText GetLogsSummaryText(const FPCGStack& InBaseStack, ELogVerbosity::Type* OutMinimumVerbosity = nullptr) const;
 
 	/**
 	* Returns summary text of visual logs from recent execution, appropriate for display in graph editor tooltip. Writes the minimum encountered verbosity
diff --git a/Engine/Plugins/Experimental/PCG/Source/PCGEditor/Private/SPCGEditorGraphDebugObjectTree.cpp b/Engine/Plugins/Experimental/PCG/Source/PCGEditor/Private/SPCGEditorGraphDebugObjectTree.cpp
--- a/Engine/Plugins/Experimental/PCG/Source/PCGEditor/Private/SPCGEditorGraphDebugObjectTree.cpp
+++ b/Engine/Plugins/Experimental/PCG/Source/PCGEditor/Private/SPCGEditorGraphDebugObjectTree.cpp
@@ -5,11 +5,10 @@
 #include "PCGComponent.h"
 #include "PCGGraph.h"
 #include "PCGSubsystem.h"
-#include "Elements/PCGLoopElement.h"
+#include "Helpers/PCGHelpers.h"
 
 #include "PCGEditor.h"
 #include "PCGEditorGraph.h"
-#include "Helpers/PCGHelpers.h"
 
 #include "PropertyCustomizationHelpers.h"
 #include "Selection.h"
@@ -23,6 +22,66 @@
 
 #define LOCTEXT_NAMESPACE "SPCGEditorGraphDebugObjectTree"
 
+namespace PCGEditorGraphDebugObjectTree
+{
+	void GetRowIconState(FPCGEditorGraphDebugObjectItemPtr Item, const FSlateBrush*& OutBrush, FLinearColor& OutColorAndOpacity)
+	{
+		const UObject* ItemObject = Item->GetObject();
+
+		if (Item->IsDebuggable())
+		{
+			OutBrush = FAppStyle::Get().GetBrush("LevelEditor.Tabs.Debug");
+		}
+		else if (IsValid(ItemObject) && ItemObject->IsA<AActor>())
+		{
+			OutBrush = FSlateIconFinder::FindIconBrushForClass(ItemObject->GetClass());
+		}
+		else
+		{
+			OutBrush = Item->IsExpanded() ? FAppStyle::Get().GetBrush("Icons.FolderOpen") : FAppStyle::Get().GetBrush("Icons.FolderClosed");
+		}
+
+		OutColorAndOpacity = Item->IsDebuggable() ? FLinearColor::White : FLinearColor(1.0f, 1.0f, 1.0f, 0.5f);
+	}
+
+	void GetErrorIconState(FPCGEditorGraphDebugObjectItemPtr Item, EVisibility& OutIconVisibility, FText& OutIconTooltipText, FLinearColor& OutColorAndOpacity)
+	{
+		OutIconTooltipText = FText();
+		OutColorAndOpacity = FLinearColor::White;
+
+		ELogVerbosity::Type MinVerbosity = ELogVerbosity::All;
+
+		if (const UPCGSubsystem* Subsystem = UPCGSubsystem::GetActiveEditorInstance())
+		{
+			if (Item->GetPCGStack())
+			{
+				OutIconTooltipText = Subsystem->GetNodeVisualLogs().GetLogsSummaryText(*Item->GetPCGStack(), &MinVerbosity);
+			}
+			else
+			{
+				for (const FPCGEditorGraphDebugObjectItemPtr& Child : Item->GetChildren())
+				{
+					if (Child->GetPCGStack())
+					{
+						OutIconTooltipText = Subsystem->GetNodeVisualLogs().GetLogsSummaryText(*Child->GetPCGStack(), &MinVerbosity);
+						break;
+					}
+				}
+			}
+		}
+
+		OutIconVisibility = OutIconTooltipText.IsEmpty() ? EVisibility::Hidden : EVisibility::Visible;
+
+		if (MinVerbosity < ELogVerbosity::All)
+		{
+			constexpr FLinearColor WarningColor(1.0f, 0.75f, 0.0f, 0.9f);
+			constexpr FLinearColor ErrorColor(1.0f, 0.0f, 0.0f, 0.9f);
+
+			OutColorAndOpacity = MinVerbosity <= ELogVerbosity::Error ? ErrorColor : WarningColor;
+		}
+	}
+}
+
 void FPCGEditorGraphDebugObjectItem::AddChild(TSharedRef<FPCGEditorGraphDebugObjectItem> InChild)
 {
 	check(!Children.Contains(InChild));
@@ -99,9 +158,21 @@ FString FPCGEditorGraphDebugObjectItem_PCGLoopIndex::GetLabel() const
 void SPCGEditorGraphDebugObjectItemRow::Construct(const FArguments& InArgs, const TSharedRef<STableViewBase>& InOwnerTableView, FPCGEditorGraphDebugObjectItemPtr InItem)
 {
 	Item = InItem;
+
 	// This function should auto-expand the row and select the deepest entry as the debug object if it is unambiguous (the only entry at its level in the tree).
 	DoubleClickFunc = InArgs._OnDoubleClickFunc;
 
+	// Computed once during construction as the tree is refreshed on relevant events.
+	const FSlateBrush* RowIcon = nullptr;
+	FLinearColor RowIconColorAndOpacity;
+	PCGEditorGraphDebugObjectTree::GetRowIconState(Item, RowIcon, RowIconColorAndOpacity);
+
+	// Icon indicating warnings and errors. Tree refreshes after execution so computing once here is sufficient.
+	FText ErrorIconTooltipText;
+	EVisibility ErrorIconVisibility;
+	FLinearColor ErrorIconColorAndOpacity;
+	PCGEditorGraphDebugObjectTree::GetErrorIconState(InItem, ErrorIconVisibility, ErrorIconTooltipText, ErrorIconColorAndOpacity);
+	
 	ChildSlot
 	[
 		SNew(SHorizontalBox)
@@ -111,98 +182,26 @@ void SPCGEditorGraphDebugObjectItemRow::Construct(const FArguments& InArgs, cons
 		[
 			SNew(SImage)
 			.Visibility(EVisibility::HitTestInvisible)
-			.ColorAndOpacity_Lambda([this]() { return Item->IsDebuggable() ? FLinearColor::White : FLinearColor(1.0, 1.0f, 1.0f, 0.5f); })
-			.Image_Lambda([this]()
-			{
-				if (Item->IsDebuggable())
-				{
-					return FAppStyle::Get().GetBrush("LevelEditor.Tabs.Debug");
-				}
-
-				const UObject* ItemObject = Item->GetObject();
-				if (IsValid(ItemObject) && ItemObject->IsA<AActor>())
-				{
-					return FSlateIconFinder::FindIconBrushForClass(ItemObject->GetClass());
-				}
-				
-				return Item->IsExpanded() ? FAppStyle::Get().GetBrush("Icons.FolderOpen") : FAppStyle::Get().GetBrush("Icons.FolderClosed");
-			})
+			.ColorAndOpacity(RowIconColorAndOpacity)
+			.Image(RowIcon)
 		]
 		+SHorizontalBox::Slot()
 		.AutoWidth()
 		[
 			SNew(STextBlock)
 			.Text(FText::FromString(Item->GetLabel()))
-			// Highlight effect applied if user is inspecting a node
-			.ColorAndOpacity(Item->IsGrayedOut() ? FColor(75, 75, 75) : FColor::White)
+			// Highlight based on data available for currently inspected node. Computed dynamically in lambda to respond to inspection changes.
+			.ColorAndOpacity_Lambda([this] { return Item->IsGrayedOut() ? FColor(75, 75, 75) : FColor::White; })
 		]
 		+SHorizontalBox::Slot()
 		.AutoWidth()
 		.Padding(20.0f, 0.0f)
 		[
 			SNew(SImage)
-			.Visibility_Lambda([this]()
-			{
-				if (const UPCGSubsystem* Subsystem = UPCGSubsystem::GetActiveEditorInstance())
-				{
-					if (Item->GetPCGStack())
-					{
-						if (Subsystem->GetNodeVisualLogs().HasLogs(*Item->GetPCGStack()))
-						{
-							return EVisibility::Visible;
-						}
-					}
-					else
-					{
-						for (const FPCGEditorGraphDebugObjectItemPtr& Child : Item->GetChildren())
-						{
-							if (Child->GetPCGStack() && Subsystem->GetNodeVisualLogs().HasLogs(*Child->GetPCGStack()))
-							{
-								return EVisibility::Visible;
-							}
-						}
-					}
-				}
-				
-				return EVisibility::Hidden;
-			})
+			.Visibility(ErrorIconVisibility)
+			.ToolTipText(ErrorIconTooltipText)
 			.Image(FAppStyle::Get().GetBrush("Icons.Error"))
-			.ColorAndOpacity_Lambda([this]()
-			{
-				if (const UPCGSubsystem* Subsystem = UPCGSubsystem::GetActiveEditorInstance())
-				{
-					constexpr FLinearColor WarningColor(1.0f, 0.75f, 0.0f, 0.9f);
-					constexpr FLinearColor ErrorColor(1.0f, 0.0f, 0.0f, 0.9f);
-
-					if (Item->GetPCGStack())
-					{
-						ELogVerbosity::Type MinVerbosity;
-						if (Subsystem->GetNodeVisualLogs().HasLogs(*Item->GetPCGStack(), MinVerbosity))
-						{
-							return FSlateColor(MinVerbosity <= ELogVerbosity::Error ? ErrorColor : WarningColor);
-						}
-					}
-					else
-					{
-						ELogVerbosity::Type MinVerbosityOfAll = ELogVerbosity::All;
-						for (const FPCGEditorGraphDebugObjectItemPtr& Child : Item->GetChildren())
-						{
-							ELogVerbosity::Type MinVerbosity;
-							if (Child->GetPCGStack() && Subsystem->GetNodeVisualLogs().HasLogs(*Child->GetPCGStack(), MinVerbosity))
-							{
-								MinVerbosityOfAll = FMath::Min(MinVerbosityOfAll, MinVerbosity);
-							}
-						}
-
-						if (MinVerbosityOfAll < ELogVerbosity::All)
-						{
-							return FSlateColor(MinVerbosityOfAll <= ELogVerbosity::Error ? ErrorColor : WarningColor);
-						}
-					}
-				}
-
-				return FSlateColor(FLinearColor(1.0f, 1.0f, 1.0f, 0.0f));
-			})
+			.ColorAndOpacity(ErrorIconColorAndOpacity)
 		]
 	];
 }
@@ -264,8 +263,7 @@ void SPCGEditorGraphDebugObjectTree::Construct(const FArguments& InArgs, TShared
 	const TSharedRef<SWidget> BrowseButton = PropertyCustomizationHelpers::MakeBrowseButton(
 		FSimpleDelegate::CreateSP(this, &SPCGEditorGraphDebugObjectTree::SelectedDebugObject_OnClicked),
 		LOCTEXT("DebugSelectActor", "Select and frame the debug actor in the Level Editor."),
-		TAttribute<bool>::Create(TAttribute<bool>::FGetter::CreateSP(this, &SPCGEditorGraphDebugObjectTree::IsSelectDebugObjectButtonEnabled))
-);
+		TAttribute<bool>::Create(TAttribute<bool>::FGetter::CreateSP(this, &SPCGEditorGraphDebugObjectTree::IsSelectDebugObjectButtonEnabled)));
 
 	ChildSlot
 	[
