[PCG] Attribute partition was not outputing the partitions in the same order than entries

* We were partitioning on the value keys order, not entry keys order.
* Add unit tests for this case

#rnx

diff --git a/Engine/Plugins/Experimental/PCG/Source/PCG/Private/Metadata/PCGMetadataPartitionCommon.cpp b/Engine/Plugins/Experimental/PCG/Source/PCG/Private/Metadata/PCGMetadataPartitionCommon.cpp
--- a/Engine/Plugins/Experimental/PCG/Source/PCG/Private/Metadata/PCGMetadataPartitionCommon.cpp
+++ b/Engine/Plugins/Experimental/PCG/Source/PCG/Private/Metadata/PCGMetadataPartitionCommon.cpp
@@ -114,6 +114,24 @@ namespace PCGMetadataPartitionCommon
 			}
 		}
 
+		// Since we partition on the value array, it is not guarenteed that the values appears in the same order than the entries.
+		// So sort the final array using the first index as a sort criteria. Empty partitions will be at the beginning too.
+		PartitionedData.Sort([](const TArray<int32>& LHS, const TArray<int32>& RHS) -> bool
+		{ 
+			if (LHS.IsEmpty())
+			{
+				return true;
+			}
+			else if (RHS.IsEmpty())
+			{
+				return false;
+			}
+			else
+			{
+				return LHS[0] < RHS[0];
+			}
+		});
+
 		return PartitionedData;
 	}
 
diff --git a/Engine/Plugins/Experimental/PCG/Source/PCG/Private/Tests/Elements/PCGAttributePartitionTest.cpp b/Engine/Plugins/Experimental/PCG/Source/PCG/Private/Tests/Elements/PCGAttributePartitionTest.cpp
--- a/Engine/Plugins/Experimental/PCG/Source/PCG/Private/Tests/Elements/PCGAttributePartitionTest.cpp
+++ b/Engine/Plugins/Experimental/PCG/Source/PCG/Private/Tests/Elements/PCGAttributePartitionTest.cpp
@@ -11,6 +11,7 @@
 
 IMPLEMENT_CUSTOM_SIMPLE_AUTOMATION_TEST(FPCGAttributePartition_Points, FPCGTestBaseClass, "Plugins.PCG.AttributePartition.Points", PCGTestsCommon::TestFlags)
 IMPLEMENT_CUSTOM_SIMPLE_AUTOMATION_TEST(FPCGAttributePartition_AttributeSet, FPCGTestBaseClass, "Plugins.PCG.AttributePartition.AttributeSet", PCGTestsCommon::TestFlags)
+IMPLEMENT_CUSTOM_SIMPLE_AUTOMATION_TEST(FPCGAttributePartition_Order, FPCGTestBaseClass, "Plugins.PCG.AttributePartition.Order", PCGTestsCommon::TestFlags)
 
 bool FPCGAttributePartition_Points::RunTest(const FString& Parameters)
 {
@@ -71,6 +72,81 @@ bool FPCGAttributePartition_Points::RunTest(const FString& Parameters)
 	return true;
 }
 
+/**
+* Test that the partitions are in the right order.
+* Points in each partition should appear in the same order that they were in the original set.
+* Partitions should be in the same order than the partition value appear in the original set.
+* For example, if the original set was [(0, 4), (1, 4), (2, 2), (3, 2)], and we partition on the second value, the result should be
+* [(0, 4), (1, 4)] and [(2, 2), (3, 2)]
+*/
+bool FPCGAttributePartition_Order::RunTest(const FString& Parameters)
+{
+	PCGTestsCommon::FTestData TestData;
+	UPCGMetadataPartitionSettings* Settings = PCGTestsCommon::GenerateSettings<UPCGMetadataPartitionSettings>(TestData);
+	check(Settings);
+
+	const FName AttributeName = TEXT("Attr");
+	Settings->PartitionAttributeSource.SetAttributeName(AttributeName);
+
+	UPCGPointData* InputPointData = NewObject<UPCGPointData>();
+	FPCGMetadataAttribute<int>* Attribute = InputPointData->Metadata->CreateAttribute<int>(AttributeName, 0, false, false);
+
+	constexpr int NumPoints = 10;
+
+	TArray<FPCGPoint>& Points = InputPointData->GetMutablePoints();
+	Points.Reserve(NumPoints);
+	for (int32 i = 0; i < NumPoints; ++i)
+	{
+		FPCGPoint& Point = Points.Emplace_GetRef();
+		// Store the index in the density
+		Point.Density = i;
+		InputPointData->Metadata->InitializeOnSet(Point.MetadataEntry);
+		Attribute->SetValue(Point.MetadataEntry, (2 * i) / NumPoints);
+	}
+
+	// Reverse the points order
+	for (int32 i = 0; i < NumPoints / 2; ++i)
+	{
+		std::swap(Points[i], Points[NumPoints - i - 1]);
+	}
+
+	FPCGTaggedData& InputTaggedData = TestData.InputData.TaggedData.Emplace_GetRef();
+	InputTaggedData.Data = InputPointData;
+	InputTaggedData.Pin = PCGPinConstants::DefaultInputLabel;
+
+	FPCGElementPtr TestElement = TestData.Settings->GetElement();
+	TUniquePtr<FPCGContext> Context = TestData.InitializeTestContext();
+
+	while (!TestElement->Execute(Context.Get())) {}
+
+	UTEST_EQUAL("There are 2 outputs", Context->OutputData.TaggedData.Num(), 2);
+
+	for (int32 i = 0; i < Context->OutputData.TaggedData.Num(); ++i)
+	{
+		const UPCGPointData* OutputPointData = Cast<const UPCGPointData>(Context->OutputData.TaggedData[i].Data);
+		UTEST_NOT_NULL(*FString::Printf(TEXT("Output %d is a point data"), i), OutputPointData);
+		check(OutputPointData);
+
+		const TArray<FPCGPoint>& OutputPoints = OutputPointData->GetPoints();
+		const FPCGMetadataAttribute<int>* OutputAttribute = OutputPointData->Metadata->GetConstTypedAttribute<int>(AttributeName);
+
+		UTEST_NOT_NULL(*FString::Printf(TEXT("Output %d has the expected attribute"), i), OutputAttribute);
+		check(OutputAttribute);
+
+		UTEST_EQUAL(*FString::Printf(TEXT("Output %d has %d points"), i, NumPoints / 2), OutputPoints.Num(), NumPoints / 2);
+
+		for (int32 j = 0; j < OutputPoints.Num() - 1; ++j)
+		{
+			UTEST_TRUE(*FString::Printf(TEXT("Output %d: Point_%d is at the right place"), i, j), OutputPoints[j].Density > OutputPoints[j + 1].Density);
+			// First partition should have the value 1, and second partition should have the value 0
+			UTEST_EQUAL(*FString::Printf(TEXT("Output %d: Point_%d has the right attribute value"), i, j), OutputAttribute->GetValue(OutputPoints[j].MetadataEntry), (i + 1) % 2);
+		}
+	}
+
+	return true;
+}
+
+
 bool FPCGAttributePartition_AttributeSet::RunTest(const FString& Parameters)
 {
 	PCGTestsCommon::FTestData TestData;
