Added ability to copy all attributes

diff --git a/Engine/Plugins/Experimental/PCG/Source/PCG/Private/Elements/Metadata/PCGMetadataElement.cpp b/Engine/Plugins/Experimental/PCG/Source/PCG/Private/Elements/Metadata/PCGMetadataElement.cpp
--- a/Engine/Plugins/Experimental/PCG/Source/PCG/Private/Elements/Metadata/PCGMetadataElement.cpp
+++ b/Engine/Plugins/Experimental/PCG/Source/PCG/Private/Elements/Metadata/PCGMetadataElement.cpp
@@ -29,6 +29,32 @@ UPCGMetadataOperationSettings::UPCGMetadataOperationSettings()
 	OutputTarget.SetAttributeName(NAME_None);
 }
 
+FName UPCGMetadataOperationSettings::AdditionalTaskName() const
+{
+#if WITH_EDITOR
+	if (bCopyAllAttributes)
+	{
+		return FName(LOCTEXT("NoteTitleAllAttributes", "Copy All Attributes").ToString());
+	}
+#endif
+
+	return NAME_None;
+}
+
+#if WITH_EDITOR
+EPCGChangeType UPCGMetadataOperationSettings::GetChangeTypeForProperty(const FName& InPropertyName) const
+{
+	EPCGChangeType ChangeType = Super::GetChangeTypeForProperty(InPropertyName);
+
+	if (InPropertyName == GET_MEMBER_NAME_CHECKED(UPCGMetadataOperationSettings, bCopyAllAttributes))
+	{
+		ChangeType |= EPCGChangeType::Cosmetic;
+	}
+
+	return ChangeType;
+}
+#endif // WITH_EDITOR
+
 FPCGElementPtr UPCGMetadataOperationSettings::CreateElement() const
 {
 	return MakeShared<FPCGMetadataOperationElement>();
@@ -119,7 +145,7 @@ bool FPCGMetadataOperationElement::ExecuteInternal(FPCGContext* Context) const
 
 	const UPCGParamData* SourceAttributeSet = (!SourceAttributeInputs.IsEmpty() ? Cast<UPCGParamData>(SourceAttributeInputs[0].Data) : nullptr);
 
-	if (SourceAttributeSet && Settings->InputSource.GetSelection() != EPCGAttributePropertySelection::Attribute)
+	if (SourceAttributeSet && Settings->InputSource.GetSelection() != EPCGAttributePropertySelection::Attribute && !Settings->bCopyAllAttributes)
 	{
 		PCGE_LOG(Error, GraphAndLog, FText::Format(LOCTEXT("InvalidInputSource", "{0} pin is connected but Input Source is not set to get an attribute."), FText::FromName(PCGMetadataOperationSettings::AttributeLabel)));
 		return true;
@@ -156,138 +182,160 @@ bool FPCGMetadataOperationElement::ExecuteInternal(FPCGContext* Context) const
 			continue;
 		}
 
-		FPCGAttributePropertyInputSelector InputSource = Settings->InputSource.CopyAndFixLast(SourceData);
-		// Deprecation, old behavior was if the Target was None, we took LastCreated in the SourceData
-		FPCGAttributePropertyOutputSelector OutputTarget = Settings->OutputTarget.CopyAndFixSource(&InputSource, SourceData);
-
-		const FName LocalSourceAttribute = InputSource.GetName();
-		const FName LocalDestinationAttribute = OutputTarget.GetName();
-
-		if (InputSource.GetSelection() == EPCGAttributePropertySelection::Attribute && !SourceMetadata->HasAttribute(LocalSourceAttribute))
+		// Early out when trying to copy all attributes from self, since nothing will happen
+		if (!SourceAttributeSet && Settings->bCopyAllAttributes)
 		{
-			PCGE_LOG(Warning, GraphAndLog, FText::Format(LOCTEXT("InputMissingAttribute", "Input does not have the '{0}' attribute"), FText::FromName(LocalSourceAttribute)));
+			PCGE_LOG(Verbose, LogOnly, LOCTEXT("TrivialCopy", "Copying all attributes on itself is a trivial operation."));
 			continue;
 		}
 
-		const TArray<FPCGPoint>& Points = OriginalData->GetPoints();
-		const int OriginalPointCount = Points.Num();
+		// Prepare attribute accessors
+		TArray<TPair<FPCGAttributePropertyInputSelector, FPCGAttributePropertyOutputSelector>> AttributeSelectors;
 
-		UPCGPointData* SampledData = NewObject<UPCGPointData>();
-		SampledData->InitializeFromData(OriginalData);
-		TArray<FPCGPoint>& SampledPoints = SampledData->GetMutablePoints();
-
-		Output.Data = SampledData;
-
-		// Copy points and then apply the operation
-		SampledPoints = Points;
-
-		// If it is attribute to attribute, just copy the attributes, if they exist and are valid
-		// Only do that if it is really attribute to attribute, without any extra accessor. Any extra accessor will behave as a property.
-		const bool bInputHasAnyExtra = !InputSource.GetExtraNames().IsEmpty();
-		const bool bOutputHasAnyExtra = !OutputTarget.GetExtraNames().IsEmpty();
-		if (!bInputHasAnyExtra && !bOutputHasAnyExtra && Settings->InputSource.GetSelection() == EPCGAttributePropertySelection::Attribute && OutputTarget.GetSelection() == EPCGAttributePropertySelection::Attribute)
+		if (Settings->bCopyAllAttributes)
 		{
-			if (!SourceAttributeSet && LocalSourceAttribute == LocalDestinationAttribute)
-			{
-				// Nothing to do if we try to copy an attribute into itself in the original point data.
-				continue;
-			}
-			
-			const FPCGMetadataAttributeBase* OriginalAttribute = SourceMetadata->GetConstAttribute(LocalSourceAttribute);
-			check(OriginalAttribute);
-
-			// We only copy entries/values if we copy from the input spatial metadata
-			// If it is from the source attribute set, we don't copy (and all points will have the same default value, value from the attribute set)
-			const bool bCopyEntriesAndValues = (SourceAttributeSet == nullptr);
-			if (!SampledData->Metadata->CopyAttribute(OriginalAttribute, LocalDestinationAttribute, /*bKeepParent=*/ false, /*bCopyEntries=*/ bCopyEntriesAndValues, /*bCopyValues=*/ bCopyEntriesAndValues))
+			TArray<FName> AttributeNames;
+			TArray<EPCGMetadataTypes> AttributeTypes;
+			SourceMetadata->GetAttributes(AttributeNames, AttributeTypes);
+
+			for (const FName& AttributeName : AttributeNames)
 			{
-				PCGE_LOG(Warning, GraphAndLog, FText::Format(LOCTEXT("FailedCopyToNewAttribute", "Failed to copy to new attribute '{0}'"), FText::FromName(LocalDestinationAttribute)));
+				TPair<FPCGAttributePropertyInputSelector, FPCGAttributePropertyOutputSelector>& Selectors = AttributeSelectors.Emplace_GetRef();
+				Selectors.Key.SetAttributeName(AttributeName);
+				Selectors.Value.SetAttributeName(AttributeName);
 			}
+		}
+		else
+		{
+			FPCGAttributePropertyInputSelector InputSource = Settings->InputSource.CopyAndFixLast(SourceData);
+			// Deprecation, old behavior was if the Target was None, we took LastCreated in the SourceData
+			FPCGAttributePropertyOutputSelector OutputTarget = Settings->OutputTarget.CopyAndFixSource(&InputSource, SourceData);
 
-			continue;
+			AttributeSelectors.Emplace(MoveTemp(InputSource), MoveTemp(OutputTarget));
 		}
 
-		TUniquePtr<const IPCGAttributeAccessor> InputAccessor = PCGAttributeAccessorHelpers::CreateConstAccessor(SourceData, InputSource);
-		TUniquePtr<const IPCGAttributeAccessorKeys> InputKeys = PCGAttributeAccessorHelpers::CreateConstKeys(SourceData, InputSource);
+		UPCGPointData* SampledData = nullptr;
 
-		if (!InputAccessor.IsValid() || !InputKeys.IsValid())
+		for (const auto& SelectorPair : AttributeSelectors)
 		{
-			PCGE_LOG(Warning, GraphAndLog, LOCTEXT("FailedToCreateInputAccessor", "Failed to create input accessor or iterator"));
-			continue;
-		}
+			const FPCGAttributePropertyInputSelector& InputSource = SelectorPair.Key;
+			const FPCGAttributePropertyOutputSelector& OutputTarget = SelectorPair.Value;
 
-		// If the target is an attribute, only create a new one if the attribute doesn't already exist or we have any extra.
-		// If it exist or have any extra, it will try to write to it.
-		if (!bOutputHasAnyExtra && OutputTarget.GetSelection() == EPCGAttributePropertySelection::Attribute && !SampledData->Metadata->HasAttribute(LocalDestinationAttribute))
-		{
-			auto CreateAttribute = [SampledData, LocalDestinationAttribute](auto Dummy)
-			{
-				using AttributeType = decltype(Dummy);
-				return PCGMetadataElementCommon::ClearOrCreateAttribute<AttributeType>(SampledData->Metadata, LocalDestinationAttribute) != nullptr;
-			};
-			
-			if (!PCGMetadataAttribute::CallbackWithRightType(InputAccessor->GetUnderlyingType(), CreateAttribute))
+			const FName LocalSourceAttribute = InputSource.GetName();
+			const FName LocalDestinationAttribute = OutputTarget.GetName();
+
+			if (InputSource.GetSelection() == EPCGAttributePropertySelection::Attribute && !SourceMetadata->HasAttribute(LocalSourceAttribute))
 			{
-				PCGE_LOG(Warning, GraphAndLog, FText::Format(LOCTEXT("FailedToCreateNewAttribute", "Failed to create new attribute '{0}'"), FText::FromName(LocalDestinationAttribute)));
+				PCGE_LOG(Warning, GraphAndLog, FText::Format(LOCTEXT("InputMissingAttribute", "Input does not have the '{0}' attribute"), FText::FromName(LocalSourceAttribute)));
 				continue;
 			}
-		}
 
-		TUniquePtr<IPCGAttributeAccessor> OutputAccessor = PCGAttributeAccessorHelpers::CreateAccessor(SampledData, OutputTarget);
-		TUniquePtr<IPCGAttributeAccessorKeys> OutputKeys = PCGAttributeAccessorHelpers::CreateKeys(SampledData, OutputTarget);
+			const TArray<FPCGPoint>& Points = OriginalData->GetPoints();
+			const int OriginalPointCount = Points.Num();
 
-		if (!OutputAccessor.IsValid() || !OutputKeys.IsValid())
-		{
-			PCGE_LOG(Warning, GraphAndLog, LOCTEXT("FailedToCreateOutputAccessor", "Failed to create output accessor or iterator"));
-			continue;
-		}
+			if (!SampledData)
+			{
+				SampledData = NewObject<UPCGPointData>();
+				SampledData->InitializeFromData(OriginalData);
 
-		if (OutputAccessor->IsReadOnly())
-		{
-			PCGE_LOG(Warning, GraphAndLog, FText::Format(LOCTEXT("OutputAccessorIsReadOnly", "Attribute/Property '{0}' is read only."), OutputTarget.GetDisplayText()));
-			continue;
-		}
+				SampledData->GetMutablePoints() = Points;
+			}
 
-		// Final verification, if we can put the value of input into output
-		if (!PCG::Private::IsBroadcastable(InputAccessor->GetUnderlyingType(), OutputAccessor->GetUnderlyingType()))
-		{
-			PCGE_LOG(Error, GraphAndLog, LOCTEXT("CannotBroadcastTypes", "Cannot broadcast input type into output type"));
-			continue;
-		}
+			TArray<FPCGPoint>& SampledPoints = SampledData->GetMutablePoints();
 
-		// At this point, we are ready.
-		auto Operation = [&InputAccessor, &InputKeys, &OutputAccessor, &OutputKeys](auto Dummy)
-		{
-			using OutputType = decltype(Dummy);
-			OutputType Value{};
+			// If it is attribute to attribute, just copy the attributes, if they exist and are valid
+			// Only do that if it is really attribute to attribute, without any extra accessor. Any extra accessor will behave as a property.
+			const bool bInputHasAnyExtra = !InputSource.GetExtraNames().IsEmpty();
+			const bool bOutputHasAnyExtra = !OutputTarget.GetExtraNames().IsEmpty();
+			if (!bInputHasAnyExtra && !bOutputHasAnyExtra && Settings->InputSource.GetSelection() == EPCGAttributePropertySelection::Attribute && OutputTarget.GetSelection() == EPCGAttributePropertySelection::Attribute)
+			{
+				if (!SourceAttributeSet && LocalSourceAttribute == LocalDestinationAttribute)
+				{
+					// Nothing to do if we try to copy an attribute into itself in the original point data.
+					continue;
+				}
+
+				const FPCGMetadataAttributeBase* OriginalAttribute = SourceMetadata->GetConstAttribute(LocalSourceAttribute);
+				check(OriginalAttribute);
+
+				// We only copy entries/values if we copy from the input spatial metadata
+				// If it is from the source attribute set, we don't copy (and all points will have the same default value, value from the attribute set)
+				const bool bCopyEntriesAndValues = (SourceAttributeSet == nullptr);
+				if (!SampledData->Metadata->CopyAttribute(OriginalAttribute, LocalDestinationAttribute, /*bKeepParent=*/ false, /*bCopyEntries=*/ bCopyEntriesAndValues, /*bCopyValues=*/ bCopyEntriesAndValues))
+				{
+					PCGE_LOG(Warning, GraphAndLog, FText::Format(LOCTEXT("FailedCopyToNewAttribute", "Failed to copy to new attribute '{0}'"), FText::FromName(LocalDestinationAttribute)));
+				}
+
+				continue;
+			}
 
-			EPCGAttributeAccessorFlags Flags = EPCGAttributeAccessorFlags::AllowBroadcast;
+			TUniquePtr<const IPCGAttributeAccessor> InputAccessor = PCGAttributeAccessorHelpers::CreateConstAccessor(SourceData, InputSource);
+			TUniquePtr<const IPCGAttributeAccessorKeys> InputKeys = PCGAttributeAccessorHelpers::CreateConstKeys(SourceData, InputSource);
+
+			if (!InputAccessor.IsValid() || !InputKeys.IsValid())
+			{
+				PCGE_LOG(Warning, GraphAndLog, LOCTEXT("FailedToCreateInputAccessor", "Failed to create input accessor or iterator"));
+				continue;
+			}
 
-			const int32 NumberOfElements = OutputKeys->GetNum();
-			constexpr int32 ChunkSize = 256;
+			// If the target is an attribute, only create a new one if the attribute doesn't already exist or we have any extra.
+			// If it exist or have any extra, it will try to write to it.
+			if (!bOutputHasAnyExtra && OutputTarget.GetSelection() == EPCGAttributePropertySelection::Attribute && !SampledData->Metadata->HasAttribute(LocalDestinationAttribute))
+			{
+				auto CreateAttribute = [SampledData, LocalDestinationAttribute](auto Dummy)
+				{
+					using AttributeType = decltype(Dummy);
+					return PCGMetadataElementCommon::ClearOrCreateAttribute<AttributeType>(SampledData->Metadata, LocalDestinationAttribute) != nullptr;
+				};
+
+				if (!PCGMetadataAttribute::CallbackWithRightType(InputAccessor->GetUnderlyingType(), CreateAttribute))
+				{
+					PCGE_LOG(Warning, GraphAndLog, FText::Format(LOCTEXT("FailedToCreateNewAttribute", "Failed to create new attribute '{0}'"), FText::FromName(LocalDestinationAttribute)));
+					continue;
+				}
+			}
 
-			TArray<OutputType, TInlineAllocator<ChunkSize>> TempValues;
-			TempValues.SetNum(ChunkSize);
+			TUniquePtr<IPCGAttributeAccessor> OutputAccessor = PCGAttributeAccessorHelpers::CreateAccessor(SampledData, OutputTarget);
+			TUniquePtr<IPCGAttributeAccessorKeys> OutputKeys = PCGAttributeAccessorHelpers::CreateKeys(SampledData, OutputTarget);
 
-			const int32 NumberOfIterations = (NumberOfElements + ChunkSize - 1) / ChunkSize;
+			if (!OutputAccessor.IsValid() || !OutputKeys.IsValid())
+			{
+				PCGE_LOG(Warning, GraphAndLog, LOCTEXT("FailedToCreateOutputAccessor", "Failed to create output accessor or iterator"));
+				continue;
+			}
 
-			for (int32 i = 0; i < NumberOfIterations; ++i)
+			if (OutputAccessor->IsReadOnly())
 			{
-				const int32 StartIndex = i * ChunkSize;
-				const int32 Range = FMath::Min(NumberOfElements - StartIndex, ChunkSize);
-				TArrayView<OutputType> View(TempValues.GetData(), Range);
+				PCGE_LOG(Warning, GraphAndLog, FText::Format(LOCTEXT("OutputAccessorIsReadOnly", "Attribute/Property '{0}' is read only."), OutputTarget.GetDisplayText()));
+				continue;
+			}
 
-				InputAccessor->GetRange<OutputType>(View, StartIndex, *InputKeys, Flags);
-				OutputAccessor->SetRange<OutputType>(View, StartIndex, *OutputKeys, Flags);
+			// Final verification, if we can put the value of input into output
+			if (!PCG::Private::IsBroadcastable(InputAccessor->GetUnderlyingType(), OutputAccessor->GetUnderlyingType()))
+			{
+				PCGE_LOG(Error, GraphAndLog, LOCTEXT("CannotBroadcastTypes", "Cannot broadcast input type into output type"));
+				continue;
 			}
 
-			return true;
-		};
+			// At this point, we are ready.
+			PCGMetadataElementCommon::FCopyFromAccessorToAccessorParams Params;
+			Params.InKeys = InputKeys.Get();
+			Params.InAccessor = InputAccessor.Get();
+			Params.OutKeys = OutputKeys.Get();
+			Params.OutAccessor = OutputAccessor.Get();
+			Params.IterationCount = PCGMetadataElementCommon::FCopyFromAccessorToAccessorParams::Out;
+			Params.Flags = EPCGAttributeAccessorFlags::AllowBroadcast | EPCGAttributeAccessorFlags::AllowConstructible;
+
+			if (!PCGMetadataElementCommon::CopyFromAccessorToAccessor(Params))
+			{
+				PCGE_LOG(Warning, GraphAndLog, LOCTEXT("ErrorGettingSettingValues", "Error while getting/setting values"));
+				continue;
+			}
+		}
 
-		if (!PCGMetadataAttribute::CallbackWithRightType(OutputAccessor->GetUnderlyingType(), Operation))
+		if (SampledData)
 		{
-			PCGE_LOG(Warning, GraphAndLog, LOCTEXT("ErrorGettingSettingValues", "Error while getting/setting values"));
-			continue;
+			Output.Data = SampledData;
 		}
 	}
 
diff --git a/Engine/Plugins/Experimental/PCG/Source/PCG/Private/Metadata/PCGAttributePropertySelector.cpp b/Engine/Plugins/Experimental/PCG/Source/PCG/Private/Metadata/PCGAttributePropertySelector.cpp
--- a/Engine/Plugins/Experimental/PCG/Source/PCG/Private/Metadata/PCGAttributePropertySelector.cpp
+++ b/Engine/Plugins/Experimental/PCG/Source/PCG/Private/Metadata/PCGAttributePropertySelector.cpp
@@ -200,7 +200,7 @@ bool FPCGAttributePropertySelector::IsValid() const
 		FPCGMetadataAttributeBase::IsValidName(ThisAttributeName);
 }
 
-bool FPCGAttributePropertySelector::Update(FString NewValue)
+bool FPCGAttributePropertySelector::Update(const FString& NewValue)
 {
 	TArray<FString> NewValues;
 	if (NewValue.IsEmpty())
@@ -212,7 +212,7 @@ bool FPCGAttributePropertySelector::Update(FString NewValue)
 		NewValue.ParseIntoArray(NewValues, PCGAttributePropertySelectorConstants::ExtraSeparator, /*InCullEmpty=*/ false);
 	}
 
-	const FString NewName = NewValues[0];
+	const FString& NewName = NewValues[0];
 	TArray<FString> ExtraNamesTemp;
 	if (NewValues.Num() > 1)
 	{
@@ -268,6 +268,30 @@ void FPCGAttributePropertySelector::AddToCrc(FArchiveCrc32& Ar) const
 	}
 }
 
+uint32 GetTypeHash(const FPCGAttributePropertySelector& Selector)
+{
+	uint32 Hash = GetTypeHash(Selector.Selection);
+	if (Selector.Selection == EPCGAttributePropertySelection::Attribute)
+	{
+		Hash = HashCombine(Hash, GetTypeHash(Selector.AttributeName));
+	}
+	else if (Selector.Selection == EPCGAttributePropertySelection::PointProperty)
+	{
+		Hash = HashCombine(Hash, GetTypeHash(Selector.PointProperty));
+	}
+	else
+	{
+		Hash = HashCombine(Hash, GetTypeHash(Selector.ExtraProperty));
+	}
+
+	for (const FString& ExtraName : Selector.ExtraNames)
+	{
+		Hash = HashCombine(Hash, GetTypeHash(ExtraName));
+	}
+
+	return Hash;
+}
+
 bool FPCGAttributePropertySelector::IsBasicAttribute() const
 {
 	return Selection == EPCGAttributePropertySelection::Attribute && ExtraNames.IsEmpty();
diff --git a/Engine/Plugins/Experimental/PCG/Source/PCG/Public/Elements/Metadata/PCGMetadataElement.h b/Engine/Plugins/Experimental/PCG/Source/PCG/Public/Elements/Metadata/PCGMetadataElement.h
--- a/Engine/Plugins/Experimental/PCG/Source/PCG/Public/Elements/Metadata/PCGMetadataElement.h
+++ b/Engine/Plugins/Experimental/PCG/Source/PCG/Public/Elements/Metadata/PCGMetadataElement.h
@@ -36,20 +36,27 @@ public:
 	virtual bool IsPinUsedByNodeExecution(const UPCGPin* InPin) const override;
 	virtual void ApplyDeprecation(UPCGNode* InOutNode) override;
 #endif
+	virtual FName AdditionalTaskName() const override;
 
 protected:
+#if WITH_EDITOR
+	virtual EPCGChangeType GetChangeTypeForProperty(const FName& InPropertyName) const override;
+#endif
 	virtual TArray<FPCGPinProperties> InputPinProperties() const override;
 	virtual TArray<FPCGPinProperties> OutputPinProperties() const override { return Super::DefaultPointOutputPinProperties(); }
 	virtual FPCGElementPtr CreateElement() const override;
 	//~End UPCGSettings interface
 
 public:
-	UPROPERTY(BlueprintReadWrite, EditAnywhere, Category = Settings, meta = (PCG_Overridable))
+	UPROPERTY(BlueprintReadWrite, EditAnywhere, Category = Settings, meta = (EditCondition="!bCopyAllAttributes", PCG_Overridable))
 	FPCGAttributePropertyInputSelector InputSource;
 
-	UPROPERTY(BlueprintReadWrite, EditAnywhere, Category = Settings, meta = (PCG_Overridable))
+	UPROPERTY(BlueprintReadWrite, EditAnywhere, Category = Settings, meta = (EditCondition = "!bCopyAllAttributes", PCG_Overridable))
 	FPCGAttributePropertyOutputSelector OutputTarget;
 
+	UPROPERTY(BlueprintReadWrite, EditAnywhere, Category = Settings, meta = (PCG_Overridable))
+	bool bCopyAllAttributes = false;
+
 #if WITH_EDITORONLY_DATA
 	UPROPERTY()
 	FName SourceAttribute_DEPRECATED = NAME_None;
diff --git a/Engine/Plugins/Experimental/PCG/Source/PCG/Public/Metadata/PCGAttributePropertySelector.h b/Engine/Plugins/Experimental/PCG/Source/PCG/Public/Metadata/PCGAttributePropertySelector.h
--- a/Engine/Plugins/Experimental/PCG/Source/PCG/Public/Metadata/PCGAttributePropertySelector.h
+++ b/Engine/Plugins/Experimental/PCG/Source/PCG/Public/Metadata/PCGAttributePropertySelector.h
@@ -72,7 +72,7 @@ public:
 	bool IsValid() const;
 
 	// Update the selector with an incoming string.
-	bool Update(FString NewValue);
+	bool Update(const FString& NewValue);
 
 	template <typename T>
 	static T CreateFromOtherSelector(const FPCGAttributePropertySelector& InOther)
@@ -124,6 +124,8 @@ public:
 
 	virtual void AddToCrc(FArchiveCrc32& Ar) const;
 
+	friend uint32 GetTypeHash(const FPCGAttributePropertySelector& Selector);
+
 protected:
 	UPROPERTY()
 	EPCGAttributePropertySelection Selection = EPCGAttributePropertySelection::Attribute;
