[PCG] Added a new node type category for Point Operation Nodes

diff --git a/Engine/Plugins/Experimental/PCG/Source/PCG/Public/Elements/PCGApplyScaleToBounds.h b/Engine/Plugins/Experimental/PCG/Source/PCG/Public/Elements/PCGApplyScaleToBounds.h
--- a/Engine/Plugins/Experimental/PCG/Source/PCG/Public/Elements/PCGApplyScaleToBounds.h
+++ b/Engine/Plugins/Experimental/PCG/Source/PCG/Public/Elements/PCGApplyScaleToBounds.h
@@ -20,7 +20,7 @@ public:
 	virtual FName GetDefaultNodeName() const override { return FName(TEXT("ApplyScaleToBounds")); }
 	virtual FText GetDefaultNodeTitle() const override { return NSLOCTEXT("PCGApplyScaleToBoundsElement", "NodeTitle", "Apply Scale To Bounds"); }
 	virtual FText GetNodeTooltipText() const override { return NSLOCTEXT("PCGApplyScaleToBoundsElement", "NodeTooltip", "Applies the scale of each point to its bounds and resets the scale."); }
-	virtual EPCGSettingsType GetType() const override { return EPCGSettingsType::Spatial; }
+	virtual EPCGSettingsType GetType() const override { return EPCGSettingsType::PointOps; }
 #endif
 
 protected:
diff --git a/Engine/Plugins/Experimental/PCG/Source/PCG/Public/Elements/PCGBoundsModifier.h b/Engine/Plugins/Experimental/PCG/Source/PCG/Public/Elements/PCGBoundsModifier.h
--- a/Engine/Plugins/Experimental/PCG/Source/PCG/Public/Elements/PCGBoundsModifier.h
+++ b/Engine/Plugins/Experimental/PCG/Source/PCG/Public/Elements/PCGBoundsModifier.h
@@ -32,7 +32,7 @@ public:
 	virtual FName GetDefaultNodeName() const override { return FName(TEXT("BoundsModifier")); }
 	virtual FText GetDefaultNodeTitle() const override { return NSLOCTEXT("PCGBoundsModifierSettings", "NodeTitle", "Bounds Modifier"); }
 	virtual FText GetNodeTooltipText() const override;
-	virtual EPCGSettingsType GetType() const override { return EPCGSettingsType::Spatial; }
+	virtual EPCGSettingsType GetType() const override { return EPCGSettingsType::PointOps; }
 #endif
 
 protected:
diff --git a/Engine/Plugins/Experimental/PCG/Source/PCG/Public/Elements/PCGCombinePoints.h b/Engine/Plugins/Experimental/PCG/Source/PCG/Public/Elements/PCGCombinePoints.h
--- a/Engine/Plugins/Experimental/PCG/Source/PCG/Public/Elements/PCGCombinePoints.h
+++ b/Engine/Plugins/Experimental/PCG/Source/PCG/Public/Elements/PCGCombinePoints.h
@@ -21,7 +21,7 @@ public:
 	virtual FName GetDefaultNodeName() const override { return FName(TEXT("CombinePoints")); }
 	virtual FText GetDefaultNodeTitle() const override { return NSLOCTEXT("PCGCombinePointsElement", "NodeTitle", "Combine Points"); }
 	virtual FText GetNodeTooltipText() const override { return NSLOCTEXT("PCGCombinePointsElement", "NodeTooltip", "Combines each point to share a singular bound extent."); }
-	virtual EPCGSettingsType GetType() const override { return EPCGSettingsType::Spatial; }
+	virtual EPCGSettingsType GetType() const override { return EPCGSettingsType::PointOps; }
 #endif
 
 protected:
diff --git a/Engine/Plugins/Experimental/PCG/Source/PCG/Public/Elements/PCGDuplicatePoint.h b/Engine/Plugins/Experimental/PCG/Source/PCG/Public/Elements/PCGDuplicatePoint.h
--- a/Engine/Plugins/Experimental/PCG/Source/PCG/Public/Elements/PCGDuplicatePoint.h
+++ b/Engine/Plugins/Experimental/PCG/Source/PCG/Public/Elements/PCGDuplicatePoint.h
@@ -20,7 +20,7 @@ public:
 	virtual FName GetDefaultNodeName() const override { return FName(TEXT("DuplicatePoint")); }
 	virtual FText GetDefaultNodeTitle() const override { return NSLOCTEXT("PCGDuplicatePointElement", "NodeTitle", "Duplicate Point"); }
 	virtual FText GetNodeTooltipText() const override { return NSLOCTEXT("PCGDuplicatePointElement", "NodeTooltip", "Creates duplicates of each point with optional transform offsets."); }
-	virtual EPCGSettingsType GetType() const override { return EPCGSettingsType::Spatial; }
+	virtual EPCGSettingsType GetType() const override { return EPCGSettingsType::PointOps; }
 #endif
 
 protected:
diff --git a/Engine/Plugins/Experimental/PCG/Source/PCG/Public/Elements/PCGPointExtentsModifier.h b/Engine/Plugins/Experimental/PCG/Source/PCG/Public/Elements/PCGPointExtentsModifier.h
--- a/Engine/Plugins/Experimental/PCG/Source/PCG/Public/Elements/PCGPointExtentsModifier.h
+++ b/Engine/Plugins/Experimental/PCG/Source/PCG/Public/Elements/PCGPointExtentsModifier.h
@@ -29,7 +29,7 @@ public:
 #if WITH_EDITOR
 	virtual FName GetDefaultNodeName() const override { return FName(TEXT("ExtentsModifier")); }
 	virtual FText GetDefaultNodeTitle() const override { return NSLOCTEXT("PCGPointExtentsModifierSettings", "NodeTitle", "Extents Modifier"); }
-	virtual EPCGSettingsType GetType() const override { return EPCGSettingsType::Spatial; }
+	virtual EPCGSettingsType GetType() const override { return EPCGSettingsType::PointOps; }
 #endif
 
 protected:
diff --git a/Engine/Plugins/Experimental/PCG/Source/PCG/Public/Elements/PCGResetPointCenter.h b/Engine/Plugins/Experimental/PCG/Source/PCG/Public/Elements/PCGResetPointCenter.h
--- a/Engine/Plugins/Experimental/PCG/Source/PCG/Public/Elements/PCGResetPointCenter.h
+++ b/Engine/Plugins/Experimental/PCG/Source/PCG/Public/Elements/PCGResetPointCenter.h
@@ -21,7 +21,7 @@ public:
 	virtual FName GetDefaultNodeName() const override { return FName(TEXT("ResetPointCenter")); }
 	virtual FText GetDefaultNodeTitle() const override { return NSLOCTEXT("PCGResetPointCenterElement", "NodeTitle", "Reset Point Center"); }
 	virtual FText GetNodeTooltipText() const override { return NSLOCTEXT("PCGResetPointCenterElement", "NodeTooltip", "Modify the position of a point within its bounds, while keeping its bounds the same."); }
-	virtual EPCGSettingsType GetType() const override { return EPCGSettingsType::Spatial; }
+	virtual EPCGSettingsType GetType() const override { return EPCGSettingsType::PointOps; }
 #endif
 
 protected:
diff --git a/Engine/Plugins/Experimental/PCG/Source/PCG/Public/Elements/PCGSplitPoints.h b/Engine/Plugins/Experimental/PCG/Source/PCG/Public/Elements/PCGSplitPoints.h
--- a/Engine/Plugins/Experimental/PCG/Source/PCG/Public/Elements/PCGSplitPoints.h
+++ b/Engine/Plugins/Experimental/PCG/Source/PCG/Public/Elements/PCGSplitPoints.h
@@ -28,7 +28,7 @@ public:
 	virtual FName GetDefaultNodeName() const override { return FName(TEXT("SplitPoints")); }
 	virtual FText GetDefaultNodeTitle() const override { return NSLOCTEXT("PCGSplitPointsElement", "NodeTitle", "Split Points"); }
 	virtual FText GetNodeTooltipText() const override { return NSLOCTEXT("PCGSplitPointsElement", "NodeTooltip", "Splits each input point into two separate points and sets bounds based on the position and axis of the cut."); }
-	virtual EPCGSettingsType GetType() const override { return EPCGSettingsType::Spatial; }
+	virtual EPCGSettingsType GetType() const override { return EPCGSettingsType::PointOps; }
 #endif
 
 protected:
diff --git a/Engine/Plugins/Experimental/PCG/Source/PCG/Public/Elements/PCGTransformPoints.h b/Engine/Plugins/Experimental/PCG/Source/PCG/Public/Elements/PCGTransformPoints.h
--- a/Engine/Plugins/Experimental/PCG/Source/PCG/Public/Elements/PCGTransformPoints.h
+++ b/Engine/Plugins/Experimental/PCG/Source/PCG/Public/Elements/PCGTransformPoints.h
@@ -18,7 +18,7 @@ public:
 #if WITH_EDITOR
 	virtual FName GetDefaultNodeName() const override { return FName(TEXT("TransformPoints")); }
 	virtual FText GetDefaultNodeTitle() const override { return NSLOCTEXT("PCGTransformPointsSettings", "NodeTitle", "Transform Points"); }
-	virtual EPCGSettingsType GetType() const override { return EPCGSettingsType::Spatial; }
+	virtual EPCGSettingsType GetType() const override { return EPCGSettingsType::PointOps; }
 #endif
 
 protected:
diff --git a/Engine/Plugins/Experimental/PCG/Source/PCG/Public/PCGSettings.h b/Engine/Plugins/Experimental/PCG/Source/PCG/Public/PCGSettings.h
--- a/Engine/Plugins/Experimental/PCG/Source/PCG/Public/PCGSettings.h
+++ b/Engine/Plugins/Experimental/PCG/Source/PCG/Public/PCGSettings.h
@@ -48,7 +48,8 @@ enum class EPCGSettingsType : uint8
 	Generic,
 	Param,
 	HierarchicalGeneration,
-	ControlFlow
+	ControlFlow,
+	PointOps
 };
 
 #if WITH_EDITOR
diff --git a/Engine/Plugins/Experimental/PCG/Source/PCGEditor/Private/PCGEditorSettings.cpp b/Engine/Plugins/Experimental/PCG/Source/PCGEditor/Private/PCGEditorSettings.cpp
--- a/Engine/Plugins/Experimental/PCG/Source/PCGEditor/Private/PCGEditorSettings.cpp
+++ b/Engine/Plugins/Experimental/PCG/Source/PCGEditor/Private/PCGEditorSettings.cpp
@@ -24,6 +24,7 @@ UPCGEditorSettings::UPCGEditorSettings(const FObjectInitializer& ObjectInitializ
 	ParamDataNodeColor = FLinearColor(1.0f, 0.38f, 0.02f);
 	DebugNodeColor = FLinearColor(1.0f, 0.0f, 1.0f);
 	ControlFlowNodeColor = FLinearColor(0.0f, 1.0f, 0.0f);
+	PointOpsNodeColor = FLinearColor(0.0f, 0.04f, 0.23f);
 	HierarchicalGenerationNodeColor = FLinearColor(1.0f, 0.132868f, 0.0f);
 
 	DefaultPinColor = FLinearColor(0.29f, 0.29f, 0.29f);
@@ -86,6 +87,8 @@ FLinearColor UPCGEditorSettings::GetColor(UPCGSettings* Settings) const
 			return HierarchicalGenerationNodeColor;
 		case EPCGSettingsType::ControlFlow:
 			return ControlFlowNodeColor;
+		case EPCGSettingsType::PointOps:
+			return PointOpsNodeColor;
 		case EPCGSettingsType::Generic: // falls through
 		default:
 			// Finally, we couldn't find any match, so return the default value
diff --git a/Engine/Plugins/Experimental/PCG/Source/PCGEditor/Public/PCGEditorSettings.h b/Engine/Plugins/Experimental/PCG/Source/PCGEditor/Public/PCGEditorSettings.h
--- a/Engine/Plugins/Experimental/PCG/Source/PCGEditor/Public/PCGEditorSettings.h
+++ b/Engine/Plugins/Experimental/PCG/Source/PCGEditor/Public/PCGEditorSettings.h
@@ -80,6 +80,10 @@ public:
 	/** Color used for control flow operations */
 	UPROPERTY(EditAnywhere, config, Category = Node, meta = (HideAlphaChannel))
 	FLinearColor ControlFlowNodeColor;
+	
+	/** Color used for point operations */
+	UPROPERTY(EditAnywhere, config, Category = Node, meta = (HideAlphaChannel))
+	FLinearColor PointOpsNodeColor;
 
 	/** Color used for hierarchical generation operations */
 	UPROPERTY(EditAnywhere, config, Category = Node, meta = (HideAlphaChannel))
