[PCG] Rename EPCGMedadata to EPCGMetadata enums

#rnx

diff --git a/Engine/Plugins/Experimental/PCG/Config/BasePCG.ini b/Engine/Plugins/Experimental/PCG/Config/BasePCG.ini
--- a/Engine/Plugins/Experimental/PCG/Config/BasePCG.ini
+++ b/Engine/Plugins/Experimental/PCG/Config/BasePCG.ini
@@ -76,3 +76,11 @@
 
 +EnumRedirects=(OldName="EPCGDensityNoiseMode",NewName="EPCGAttributeNoiseMode")
 +EnumRedirects=(OldName="EPCGPointFilterOperator",NewName="EPCGAttributeFilterOperator")
++EnumRedirects=(OldName="EPCGMedadataTransformOperation",NewName="EPCGMetadataTransformOperation")
++EnumRedirects=(OldName="EPCGMedadataBooleanOperation",NewName="EPCGMetadataBooleanOperation")
++EnumRedirects=(OldName="EPCGMedadataMathsOperation",NewName="EPCGMetadataMathsOperation")
++EnumRedirects=(OldName="EPCGMedadataTrigOperation",NewName="EPCGMetadataTrigOperation")
++EnumRedirects=(OldName="EPCGMedadataCompareOperation",NewName="EPCGMetadataCompareOperation")
++EnumRedirects=(OldName="EPCGMedadataVectorOperation",NewName="EPCGMetadataVectorOperation")
++EnumRedirects=(OldName="EPCGMedadataBitwiseOperation",NewName="EPCGMetadataBitwiseOperation")
++EnumRedirects=(OldName="EPCGMedadataRotatorOperation",NewName="EPCGMetadataRotatorOperation")
\ No newline at end of file
diff --git a/Engine/Plugins/Experimental/PCG/Source/PCG/Private/Elements/Metadata/PCGMetadataBitwiseOpElement.cpp b/Engine/Plugins/Experimental/PCG/Source/PCG/Private/Elements/Metadata/PCGMetadataBitwiseOpElement.cpp
--- a/Engine/Plugins/Experimental/PCG/Source/PCG/Private/Elements/Metadata/PCGMetadataBitwiseOpElement.cpp
+++ b/Engine/Plugins/Experimental/PCG/Source/PCG/Private/Elements/Metadata/PCGMetadataBitwiseOpElement.cpp
@@ -10,19 +10,19 @@ namespace PCGMetadataBitwiseSettings
 {
 	inline int64 UnaryOp(const int64& Value)
 	{
-		// EPCGMedadataBitwiseOperation::Not
+		// EPCGMetadataBitwiseOperation::Not
 		return ~Value;
 	}
 
-	inline int64 BinaryOp(const int64& Value1, const int64& Value2, EPCGMedadataBitwiseOperation Operation)
+	inline int64 BinaryOp(const int64& Value1, const int64& Value2, EPCGMetadataBitwiseOperation Operation)
 	{
 		switch (Operation)
 		{
-		case EPCGMedadataBitwiseOperation::And:
+		case EPCGMetadataBitwiseOperation::And:
 			return (Value1 & Value2);
-		case EPCGMedadataBitwiseOperation::Or:
+		case EPCGMetadataBitwiseOperation::Or:
 			return (Value1 | Value2);
-		case EPCGMedadataBitwiseOperation::Xor:
+		case EPCGMetadataBitwiseOperation::Xor:
 			return (Value1 ^ Value2);
 		default:
 			return 0;
@@ -54,7 +54,7 @@ FName UPCGMetadataBitwiseSettings::GetInputPinLabel(uint32 Index) const
 	switch (Index)
 	{
 	case 0:
-		return (Operation == EPCGMedadataBitwiseOperation::Not) ? PCGPinConstants::DefaultInputLabel : PCGMetadataSettingsBaseConstants::DoubleInputFirstLabel;
+		return (Operation == EPCGMetadataBitwiseOperation::Not) ? PCGPinConstants::DefaultInputLabel : PCGMetadataSettingsBaseConstants::DoubleInputFirstLabel;
 	case 1:
 		return PCGMetadataSettingsBaseConstants::DoubleInputSecondLabel;
 	default:
@@ -64,7 +64,7 @@ FName UPCGMetadataBitwiseSettings::GetInputPinLabel(uint32 Index) const
 
 uint32 UPCGMetadataBitwiseSettings::GetOperandNum() const
 {
-	return (Operation == EPCGMedadataBitwiseOperation::Not) ? 1 : 2;
+	return (Operation == EPCGMetadataBitwiseOperation::Not) ? 1 : 2;
 }
 
 bool UPCGMetadataBitwiseSettings::IsSupportedInputType(uint16 TypeId, uint32 InputIndex, bool& bHasSpecialRequirement) const
@@ -88,7 +88,7 @@ FPCGAttributePropertyInputSelector UPCGMetadataBitwiseSettings::GetInputSource(u
 
 FName UPCGMetadataBitwiseSettings::AdditionalTaskName() const
 {
-	if (const UEnum* EnumPtr = FindObject<UEnum>(nullptr, TEXT("/Script/PCG.EPCGMedadataBitwiseOperation"), true))
+	if (const UEnum* EnumPtr = StaticEnum<EPCGMetadataBitwiseOperation>())
 	{
 		return FName(FString("Bitwise: ") + EnumPtr->GetNameStringByValue(static_cast<int>(Operation)));
 	}
@@ -111,17 +111,17 @@ FText UPCGMetadataBitwiseSettings::GetDefaultNodeTitle() const
 
 TArray<FPCGPreConfiguredSettingsInfo> UPCGMetadataBitwiseSettings::GetPreconfiguredInfo() const
 {
-	return PCGMetadataElementCommon::FillPreconfiguredSettingsInfoFromEnum<EPCGMedadataBitwiseOperation>();
+	return PCGMetadataElementCommon::FillPreconfiguredSettingsInfoFromEnum<EPCGMetadataBitwiseOperation>();
 }
 #endif // WITH_EDITOR
 
 void UPCGMetadataBitwiseSettings::ApplyPreconfiguredSettings(const FPCGPreConfiguredSettingsInfo& PreconfiguredInfo)
 {
-	if (const UEnum* EnumPtr = StaticEnum<EPCGMedadataBitwiseOperation>())
+	if (const UEnum* EnumPtr = StaticEnum<EPCGMetadataBitwiseOperation>())
 	{
 		if (EnumPtr->IsValidEnumValue(PreconfiguredInfo.PreconfiguredIndex))
 		{
-			Operation = EPCGMedadataBitwiseOperation(PreconfiguredInfo.PreconfiguredIndex);
+			Operation = EPCGMetadataBitwiseOperation(PreconfiguredInfo.PreconfiguredIndex);
 		}
 	}
 }
@@ -142,7 +142,7 @@ bool FPCGMetadataBitwiseElement::DoOperation(PCGMetadataOps::FOperationData& Ope
 
 	const UPCGMetadataBitwiseSettings* Settings = CastChecked<UPCGMetadataBitwiseSettings>(OperationData.Settings);
 
-	if (Settings->Operation == EPCGMedadataBitwiseOperation::Not)
+	if (Settings->Operation == EPCGMetadataBitwiseOperation::Not)
 	{
 		DoUnaryOp<int64>(OperationData, [](const int64& Value) -> int64 { return PCGMetadataBitwiseSettings::UnaryOp(Value); });
 	}
diff --git a/Engine/Plugins/Experimental/PCG/Source/PCG/Private/Elements/Metadata/PCGMetadataBooleanOpElement.cpp b/Engine/Plugins/Experimental/PCG/Source/PCG/Private/Elements/Metadata/PCGMetadataBooleanOpElement.cpp
--- a/Engine/Plugins/Experimental/PCG/Source/PCG/Private/Elements/Metadata/PCGMetadataBooleanOpElement.cpp
+++ b/Engine/Plugins/Experimental/PCG/Source/PCG/Private/Elements/Metadata/PCGMetadataBooleanOpElement.cpp
@@ -10,19 +10,19 @@ namespace PCGMetadataBooleanSettings
 {
 	inline bool UnaryOp(const bool& Value)
 	{
-		// EPCGMedadataBooleanOperation::Not
+		// EPCGMetadataBooleanOperation::Not
 		return !Value;
 	}
 
-	inline bool BinaryOp(const bool& Value1, const bool& Value2, EPCGMedadataBooleanOperation Operation)
+	inline bool BinaryOp(const bool& Value1, const bool& Value2, EPCGMetadataBooleanOperation Operation)
 	{
 		switch (Operation)
 		{
-		case EPCGMedadataBooleanOperation::And:
+		case EPCGMetadataBooleanOperation::And:
 			return (Value1 && Value2);
-		case EPCGMedadataBooleanOperation::Or:
+		case EPCGMetadataBooleanOperation::Or:
 			return (Value1 || Value2);
-		case EPCGMedadataBooleanOperation::Xor:
+		case EPCGMetadataBooleanOperation::Xor:
 			return (Value1 != Value2);
 		default:
 			return false;
@@ -54,7 +54,7 @@ FName UPCGMetadataBooleanSettings::GetInputPinLabel(uint32 Index) const
 	switch (Index)
 	{
 	case 0:
-		return (Operation == EPCGMedadataBooleanOperation::Not) ? PCGPinConstants::DefaultInputLabel : PCGMetadataSettingsBaseConstants::DoubleInputFirstLabel;
+		return (Operation == EPCGMetadataBooleanOperation::Not) ? PCGPinConstants::DefaultInputLabel : PCGMetadataSettingsBaseConstants::DoubleInputFirstLabel;
 	case 1:
 		return PCGMetadataSettingsBaseConstants::DoubleInputSecondLabel;
 	default:
@@ -64,7 +64,7 @@ FName UPCGMetadataBooleanSettings::GetInputPinLabel(uint32 Index) const
 
 uint32 UPCGMetadataBooleanSettings::GetOperandNum() const
 {
-	return (Operation == EPCGMedadataBooleanOperation::Not) ? 1 : 2;
+	return (Operation == EPCGMetadataBooleanOperation::Not) ? 1 : 2;
 }
 
 bool UPCGMetadataBooleanSettings::IsSupportedInputType(uint16 TypeId, uint32 InputIndex, bool& bHasSpecialRequirement) const
@@ -93,7 +93,7 @@ uint16 UPCGMetadataBooleanSettings::GetOutputType(uint16 InputTypeId) const
 
 FName UPCGMetadataBooleanSettings::AdditionalTaskName() const
 {
-	if (const UEnum* EnumPtr = StaticEnum<EPCGMedadataBooleanOperation>())
+	if (const UEnum* EnumPtr = StaticEnum<EPCGMetadataBooleanOperation>())
 	{
 		return FName(FString("Boolean: ") + EnumPtr->GetNameStringByValue(static_cast<int>(Operation)));
 	}
@@ -116,17 +116,17 @@ FText UPCGMetadataBooleanSettings::GetDefaultNodeTitle() const
 
 TArray<FPCGPreConfiguredSettingsInfo> UPCGMetadataBooleanSettings::GetPreconfiguredInfo() const
 {
-	return PCGMetadataElementCommon::FillPreconfiguredSettingsInfoFromEnum<EPCGMedadataBooleanOperation>();
+	return PCGMetadataElementCommon::FillPreconfiguredSettingsInfoFromEnum<EPCGMetadataBooleanOperation>();
 }
 #endif // WITH_EDITOR
 
 void UPCGMetadataBooleanSettings::ApplyPreconfiguredSettings(const FPCGPreConfiguredSettingsInfo& PreconfiguredInfo)
 {
-	if (const UEnum* EnumPtr = StaticEnum<EPCGMedadataBooleanOperation>())
+	if (const UEnum* EnumPtr = StaticEnum<EPCGMetadataBooleanOperation>())
 	{
 		if (EnumPtr->IsValidEnumValue(PreconfiguredInfo.PreconfiguredIndex))
 		{
-			Operation = EPCGMedadataBooleanOperation(PreconfiguredInfo.PreconfiguredIndex);
+			Operation = EPCGMetadataBooleanOperation(PreconfiguredInfo.PreconfiguredIndex);
 		}
 	}
 }
@@ -142,7 +142,7 @@ bool FPCGMetadataBooleanElement::DoOperation(PCGMetadataOps::FOperationData& Ope
 
 	const UPCGMetadataBooleanSettings* Settings = CastChecked<UPCGMetadataBooleanSettings>(OperationData.Settings);
 
-	if (Settings->Operation == EPCGMedadataBooleanOperation::Not)
+	if (Settings->Operation == EPCGMetadataBooleanOperation::Not)
 	{
 		DoUnaryOp<bool>(OperationData, [](const bool& Value) -> bool { return PCGMetadataBooleanSettings::UnaryOp(Value); });
 	}
diff --git a/Engine/Plugins/Experimental/PCG/Source/PCG/Private/Elements/Metadata/PCGMetadataCompareOpElement.cpp b/Engine/Plugins/Experimental/PCG/Source/PCG/Private/Elements/Metadata/PCGMetadataCompareOpElement.cpp
--- a/Engine/Plugins/Experimental/PCG/Source/PCG/Private/Elements/Metadata/PCGMetadataCompareOpElement.cpp
+++ b/Engine/Plugins/Experimental/PCG/Source/PCG/Private/Elements/Metadata/PCGMetadataCompareOpElement.cpp
@@ -10,13 +10,13 @@
 namespace PCGMetadataCompareSettings
 {
 	template <typename T>
-	bool ApplyCompare(const T& Input1, const T& Input2, EPCGMedadataCompareOperation Operation, double Tolerance)
+	bool ApplyCompare(const T& Input1, const T& Input2, EPCGMetadataCompareOperation Operation, double Tolerance)
 	{
-		if (Operation == EPCGMedadataCompareOperation::Equal)
+		if (Operation == EPCGMetadataCompareOperation::Equal)
 		{
 			return PCG::Private::MetadataTraits<T>::Equal(Input1, Input2);
 		}
-		else if (Operation == EPCGMedadataCompareOperation::NotEqual)
+		else if (Operation == EPCGMetadataCompareOperation::NotEqual)
 		{
 			return !PCG::Private::MetadataTraits<T>::Equal(Input1, Input2);
 		}
@@ -25,13 +25,13 @@ namespace PCGMetadataCompareSettings
 		{
 			switch (Operation)
 			{
-			case EPCGMedadataCompareOperation::Greater:
+			case EPCGMetadataCompareOperation::Greater:
 				return PCG::Private::MetadataTraits<T>::Greater(Input1, Input2);
-			case EPCGMedadataCompareOperation::GreaterOrEqual:
+			case EPCGMetadataCompareOperation::GreaterOrEqual:
 				return PCG::Private::MetadataTraits<T>::GreaterOrEqual(Input1, Input2);
-			case EPCGMedadataCompareOperation::Less:
+			case EPCGMetadataCompareOperation::Less:
 				return PCG::Private::MetadataTraits<T>::Less(Input1, Input2);
-			case EPCGMedadataCompareOperation::LessOrEqual:
+			case EPCGMetadataCompareOperation::LessOrEqual:
 				return PCG::Private::MetadataTraits<T>::LessOrEqual(Input1, Input2);
 			default:
 				return false;
@@ -84,7 +84,7 @@ uint32 UPCGMetadataCompareSettings::GetOperandNum() const
 bool UPCGMetadataCompareSettings::IsSupportedInputType(uint16 TypeId, uint32 InputIndex, bool& bHasSpecialRequirement) const
 {
 	bHasSpecialRequirement = false;
-	if (Operation == EPCGMedadataCompareOperation::Equal || Operation == EPCGMedadataCompareOperation::NotEqual)
+	if (Operation == EPCGMetadataCompareOperation::Equal || Operation == EPCGMetadataCompareOperation::NotEqual)
 	{
 		return PCG::Private::IsPCGType(TypeId);
 	}
@@ -119,7 +119,7 @@ uint16 UPCGMetadataCompareSettings::GetOutputType(uint16 InputTypeId) const
 
 FName UPCGMetadataCompareSettings::AdditionalTaskName() const
 {
-	if (const UEnum* EnumPtr = StaticEnum<EPCGMedadataCompareOperation>())
+	if (const UEnum* EnumPtr = StaticEnum<EPCGMetadataCompareOperation>())
 	{
 		return FName(FString("Compare: ") + EnumPtr->GetNameStringByValue(static_cast<int>(Operation)));
 	}
@@ -142,17 +142,17 @@ FText UPCGMetadataCompareSettings::GetDefaultNodeTitle() const
 
 TArray<FPCGPreConfiguredSettingsInfo> UPCGMetadataCompareSettings::GetPreconfiguredInfo() const
 {
-	return PCGMetadataElementCommon::FillPreconfiguredSettingsInfoFromEnum<EPCGMedadataCompareOperation>();
+	return PCGMetadataElementCommon::FillPreconfiguredSettingsInfoFromEnum<EPCGMetadataCompareOperation>();
 }
 #endif // WITH_EDITOR
 
 void UPCGMetadataCompareSettings::ApplyPreconfiguredSettings(const FPCGPreConfiguredSettingsInfo& PreconfiguredInfo)
 {
-	if (const UEnum* EnumPtr = StaticEnum<EPCGMedadataCompareOperation>())
+	if (const UEnum* EnumPtr = StaticEnum<EPCGMetadataCompareOperation>())
 	{
 		if (EnumPtr->IsValidEnumValue(PreconfiguredInfo.PreconfiguredIndex))
 		{
-			Operation = EPCGMedadataCompareOperation(PreconfiguredInfo.PreconfiguredIndex);
+			Operation = EPCGMetadataCompareOperation(PreconfiguredInfo.PreconfiguredIndex);
 		}
 	}
 }
diff --git a/Engine/Plugins/Experimental/PCG/Source/PCG/Private/Elements/Metadata/PCGMetadataMathsOpElement.cpp b/Engine/Plugins/Experimental/PCG/Source/PCG/Private/Elements/Metadata/PCGMetadataMathsOpElement.cpp
--- a/Engine/Plugins/Experimental/PCG/Source/PCG/Private/Elements/Metadata/PCGMetadataMathsOpElement.cpp
+++ b/Engine/Plugins/Experimental/PCG/Source/PCG/Private/Elements/Metadata/PCGMetadataMathsOpElement.cpp
@@ -11,33 +11,33 @@
 
 namespace PCGMetadataMathsSettings
 {
-	inline constexpr bool IsUnaryOp(EPCGMedadataMathsOperation Operation)
+	inline constexpr bool IsUnaryOp(EPCGMetadataMathsOperation Operation)
 	{
-		return !!(Operation & EPCGMedadataMathsOperation::UnaryOp);
+		return !!(Operation & EPCGMetadataMathsOperation::UnaryOp);
 	}
 
-	inline constexpr bool IsBinaryOp(EPCGMedadataMathsOperation Operation)
+	inline constexpr bool IsBinaryOp(EPCGMetadataMathsOperation Operation)
 	{
-		return !!(Operation & EPCGMedadataMathsOperation::BinaryOp);
+		return !!(Operation & EPCGMetadataMathsOperation::BinaryOp);
 	}
 
-	inline constexpr bool IsTernaryOp(EPCGMedadataMathsOperation Operation)
+	inline constexpr bool IsTernaryOp(EPCGMetadataMathsOperation Operation)
 	{
-		return !!(Operation & EPCGMedadataMathsOperation::TernaryOp);
+		return !!(Operation & EPCGMetadataMathsOperation::TernaryOp);
 	}
 
-	inline FName GetFirstPinLabel(EPCGMedadataMathsOperation Operation)
+	inline FName GetFirstPinLabel(EPCGMetadataMathsOperation Operation)
 	{
 		if (PCGMetadataMathsSettings::IsUnaryOp(Operation)
-			|| Operation == EPCGMedadataMathsOperation::Clamp
-			|| Operation == EPCGMedadataMathsOperation::ClampMin
-			|| Operation == EPCGMedadataMathsOperation::ClampMax)
+			|| Operation == EPCGMetadataMathsOperation::Clamp
+			|| Operation == EPCGMetadataMathsOperation::ClampMin
+			|| Operation == EPCGMetadataMathsOperation::ClampMax)
 		{
 			return PCGPinConstants::DefaultInputLabel;
 		}
 
 		if (PCGMetadataMathsSettings::IsBinaryOp(Operation)
-			|| Operation == EPCGMedadataMathsOperation::Lerp)
+			|| Operation == EPCGMetadataMathsOperation::Lerp)
 		{
 			return PCGMetadataSettingsBaseConstants::DoubleInputFirstLabel;
 		}
@@ -45,14 +45,14 @@ namespace PCGMetadataMathsSettings
 		return NAME_None;
 	}
 
-	inline FName GetSecondPinLabel(EPCGMedadataMathsOperation Operation)
+	inline FName GetSecondPinLabel(EPCGMetadataMathsOperation Operation)
 	{
-		if (Operation == EPCGMedadataMathsOperation::ClampMin || Operation == EPCGMedadataMathsOperation::Clamp)
+		if (Operation == EPCGMetadataMathsOperation::ClampMin || Operation == EPCGMetadataMathsOperation::Clamp)
 		{
 			return PCGMetadataSettingsBaseConstants::ClampMinLabel;
 		}
 
-		if (Operation == EPCGMedadataMathsOperation::ClampMax)
+		if (Operation == EPCGMetadataMathsOperation::ClampMax)
 		{
 			return PCGMetadataSettingsBaseConstants::ClampMaxLabel;
 		}
@@ -65,14 +65,14 @@ namespace PCGMetadataMathsSettings
 		return NAME_None;
 	}
 
-	inline FName GetThirdPinLabel(EPCGMedadataMathsOperation Operation)
+	inline FName GetThirdPinLabel(EPCGMetadataMathsOperation Operation)
 	{
-		if (Operation == EPCGMedadataMathsOperation::Clamp)
+		if (Operation == EPCGMetadataMathsOperation::Clamp)
 		{
 			return PCGMetadataSettingsBaseConstants::ClampMaxLabel;
 		}
 
-		if (Operation == EPCGMedadataMathsOperation::Lerp)
+		if (Operation == EPCGMetadataMathsOperation::Lerp)
 		{
 			return PCGMetadataSettingsBaseConstants::LerpRatioLabel;
 		}
@@ -81,27 +81,27 @@ namespace PCGMetadataMathsSettings
 	}
 
 	template <typename T>
-	T UnaryOp(const T& Value, EPCGMedadataMathsOperation Op)
+	T UnaryOp(const T& Value, EPCGMetadataMathsOperation Op)
 	{
 		switch (Op)
 		{
-		case EPCGMedadataMathsOperation::Sign:
+		case EPCGMetadataMathsOperation::Sign:
 			return PCGMetadataMaths::Sign(Value);
-		case EPCGMedadataMathsOperation::Frac:
+		case EPCGMetadataMathsOperation::Frac:
 			return PCGMetadataMaths::Frac(Value);
-		case EPCGMedadataMathsOperation::Truncate:
+		case EPCGMetadataMathsOperation::Truncate:
 			return PCGMetadataMaths::Truncate(Value);
-		case EPCGMedadataMathsOperation::Round:
+		case EPCGMetadataMathsOperation::Round:
 			return PCGMetadataMaths::Round(Value);
-		case EPCGMedadataMathsOperation::Sqrt:
+		case EPCGMetadataMathsOperation::Sqrt:
 			return PCGMetadataMaths::Sqrt(Value);
-		case EPCGMedadataMathsOperation::Abs:
+		case EPCGMetadataMathsOperation::Abs:
 			return PCGMetadataMaths::Abs(Value);
-		case EPCGMedadataMathsOperation::Floor:
+		case EPCGMetadataMathsOperation::Floor:
 			return PCGMetadataMaths::Floor(Value);
-		case EPCGMedadataMathsOperation::Ceil:
+		case EPCGMetadataMathsOperation::Ceil:
 			return PCGMetadataMaths::Ceil(Value);
-		case EPCGMedadataMathsOperation::OneMinus:
+		case EPCGMetadataMathsOperation::OneMinus:
 			return PCGMetadataMaths::OneMinus(Value);
 		default:
 			return T{};
@@ -109,31 +109,31 @@ namespace PCGMetadataMathsSettings
 	}
 
 	template <typename T>
-	T BinaryOp(const T& Value1, const T& Value2, EPCGMedadataMathsOperation Op)
+	T BinaryOp(const T& Value1, const T& Value2, EPCGMetadataMathsOperation Op)
 	{
 		switch (Op)
 		{
-		case EPCGMedadataMathsOperation::Add:
+		case EPCGMetadataMathsOperation::Add:
 			return Value1 + Value2;
-		case EPCGMedadataMathsOperation::Subtract:
+		case EPCGMetadataMathsOperation::Subtract:
 			return Value1 - Value2;
-		case EPCGMedadataMathsOperation::Multiply:
+		case EPCGMetadataMathsOperation::Multiply:
 			return Value1 * Value2;
-		case EPCGMedadataMathsOperation::Divide:
+		case EPCGMetadataMathsOperation::Divide:
 			return (Value2 != T{0}) ? (Value1 / Value2) : T{0}; // To mirror FMath
-		case EPCGMedadataMathsOperation::Max:
+		case EPCGMetadataMathsOperation::Max:
 			return PCGMetadataMaths::Max(Value1, Value2);
-		case EPCGMedadataMathsOperation::Min:
+		case EPCGMetadataMathsOperation::Min:
 			return PCGMetadataMaths::Min(Value1, Value2);
-		case EPCGMedadataMathsOperation::ClampMin:
+		case EPCGMetadataMathsOperation::ClampMin:
 			return PCGMetadataMaths::Clamp(Value1, Value2, Value1);
-		case EPCGMedadataMathsOperation::ClampMax:
+		case EPCGMetadataMathsOperation::ClampMax:
 			return PCGMetadataMaths::Clamp(Value1, Value1, Value2);
-		case EPCGMedadataMathsOperation::Pow:
+		case EPCGMetadataMathsOperation::Pow:
 			return PCGMetadataMaths::Pow(Value1, Value2);
-		case EPCGMedadataMathsOperation::Modulo:
+		case EPCGMetadataMathsOperation::Modulo:
 			return PCGMetadataMaths::Modulo(Value1, Value2);
-		case EPCGMedadataMathsOperation::Set:
+		case EPCGMetadataMathsOperation::Set:
 			return Value2;
 		default:
 			return T{};
@@ -141,13 +141,13 @@ namespace PCGMetadataMathsSettings
 	}
 
 	template <typename T>
-	T TernaryOp(const T& Value1, const T& Value2, const T& Value3, EPCGMedadataMathsOperation Op)
+	T TernaryOp(const T& Value1, const T& Value2, const T& Value3, EPCGMetadataMathsOperation Op)
 	{
 		switch (Op)
 		{
-		case EPCGMedadataMathsOperation::Clamp:
+		case EPCGMetadataMathsOperation::Clamp:
 			return PCGMetadataMaths::Clamp(Value1, Value2, Value3);
-		case EPCGMedadataMathsOperation::Lerp:
+		case EPCGMetadataMathsOperation::Lerp:
 			return PCGMetadataMaths::Lerp(Value1, Value2, Value3);
 		default:
 			return T{};
@@ -225,19 +225,19 @@ bool UPCGMetadataMathsSettings::IsSupportedInputType(uint16 TypeId, uint32 Input
 bool UPCGMetadataMathsSettings::ShouldForceOutputToInt(uint16 InputTypeId) const
 {
 	return PCG::Private::IsOfTypes<float, double>(InputTypeId) && bForceRoundingOpToInt &&
-		(Operation == EPCGMedadataMathsOperation::Round ||
-		Operation == EPCGMedadataMathsOperation::Truncate ||
-		Operation == EPCGMedadataMathsOperation::Floor ||
-		Operation == EPCGMedadataMathsOperation::Ceil);
+		(Operation == EPCGMetadataMathsOperation::Round ||
+		Operation == EPCGMetadataMathsOperation::Truncate ||
+		Operation == EPCGMetadataMathsOperation::Floor ||
+		Operation == EPCGMetadataMathsOperation::Ceil);
 }
 
 bool UPCGMetadataMathsSettings::ShouldForceOutputToDouble(uint16 InputTypeId) const
 {
 	return PCG::Private::IsOfTypes<int32, int64>(InputTypeId) && bForceOpToDouble &&
-		(Operation == EPCGMedadataMathsOperation::Divide ||
-		Operation == EPCGMedadataMathsOperation::Sqrt ||
-		Operation == EPCGMedadataMathsOperation::Pow ||
-		Operation == EPCGMedadataMathsOperation::Lerp);
+		(Operation == EPCGMetadataMathsOperation::Divide ||
+		Operation == EPCGMetadataMathsOperation::Sqrt ||
+		Operation == EPCGMetadataMathsOperation::Pow ||
+		Operation == EPCGMetadataMathsOperation::Lerp);
 }
 
 uint16 UPCGMetadataMathsSettings::GetOutputType(uint16 InputTypeId) const
@@ -275,7 +275,7 @@ FPCGAttributePropertyInputSelector UPCGMetadataMathsSettings::GetInputSource(uin
 
 FName UPCGMetadataMathsSettings::AdditionalTaskName() const
 {
-	if (const UEnum* EnumPtr = StaticEnum<EPCGMedadataMathsOperation>())
+	if (const UEnum* EnumPtr = StaticEnum<EPCGMetadataMathsOperation>())
 	{
 		return FName(FString("Maths: ") + EnumPtr->GetNameStringByValue(static_cast<int>(Operation)));
 	}
@@ -298,17 +298,17 @@ FText UPCGMetadataMathsSettings::GetDefaultNodeTitle() const
 
 TArray<FPCGPreConfiguredSettingsInfo> UPCGMetadataMathsSettings::GetPreconfiguredInfo() const
 {
-	return PCGMetadataElementCommon::FillPreconfiguredSettingsInfoFromEnum<EPCGMedadataMathsOperation>({ EPCGMedadataMathsOperation::UnaryOp, EPCGMedadataMathsOperation::BinaryOp, EPCGMedadataMathsOperation::TernaryOp });
+	return PCGMetadataElementCommon::FillPreconfiguredSettingsInfoFromEnum<EPCGMetadataMathsOperation>({ EPCGMetadataMathsOperation::UnaryOp, EPCGMetadataMathsOperation::BinaryOp, EPCGMetadataMathsOperation::TernaryOp });
 }
 #endif // WITH_EDITOR
 
 void UPCGMetadataMathsSettings::ApplyPreconfiguredSettings(const FPCGPreConfiguredSettingsInfo& PreconfiguredInfo)
 {
-	if (const UEnum* EnumPtr = StaticEnum<EPCGMedadataMathsOperation>())
+	if (const UEnum* EnumPtr = StaticEnum<EPCGMetadataMathsOperation>())
 	{
 		if (EnumPtr->IsValidEnumValue(PreconfiguredInfo.PreconfiguredIndex))
 		{
-			Operation = EPCGMedadataMathsOperation(PreconfiguredInfo.PreconfiguredIndex);
+			Operation = EPCGMetadataMathsOperation(PreconfiguredInfo.PreconfiguredIndex);
 		}
 	}
 }
diff --git a/Engine/Plugins/Experimental/PCG/Source/PCG/Private/Elements/Metadata/PCGMetadataRotatorOpElement.cpp b/Engine/Plugins/Experimental/PCG/Source/PCG/Private/Elements/Metadata/PCGMetadataRotatorOpElement.cpp
--- a/Engine/Plugins/Experimental/PCG/Source/PCG/Private/Elements/Metadata/PCGMetadataRotatorOpElement.cpp
+++ b/Engine/Plugins/Experimental/PCG/Source/PCG/Private/Elements/Metadata/PCGMetadataRotatorOpElement.cpp
@@ -27,63 +27,63 @@ FRotator PCGMetadataRotatorHelpers::RLerp(const FRotator& A, const FRotator& B,
 
 namespace PCGMetadataRotatorSettings
 {
-	inline constexpr bool IsUnaryOp(EPCGMedadataRotatorOperation Operation)
+	inline constexpr bool IsUnaryOp(EPCGMetadataRotatorOperation Operation)
 	{
-		return Operation == EPCGMedadataRotatorOperation::Invert ||
-			Operation == EPCGMedadataRotatorOperation::Normalize;
+		return Operation == EPCGMetadataRotatorOperation::Invert ||
+			Operation == EPCGMetadataRotatorOperation::Normalize;
 	}
 
-	inline constexpr bool IsTernaryOp(EPCGMedadataRotatorOperation Operation)
+	inline constexpr bool IsTernaryOp(EPCGMetadataRotatorOperation Operation)
 	{
-		return Operation == EPCGMedadataRotatorOperation::Lerp;
+		return Operation == EPCGMetadataRotatorOperation::Lerp;
 	}
 
-	inline constexpr bool IsTransfromOp(EPCGMedadataRotatorOperation Operation)
+	inline constexpr bool IsTransfromOp(EPCGMetadataRotatorOperation Operation)
 	{
-		return (uint16)Operation >= (uint16)EPCGMedadataRotatorOperation::TransformOp;
+		return (uint16)Operation >= (uint16)EPCGMetadataRotatorOperation::TransformOp;
 	}
 
-	inline FRotator ApplyRotatorOperation(const FRotator& Input1, const FRotator& Input2, double Ratio, EPCGMedadataRotatorOperation Operation)
+	inline FRotator ApplyRotatorOperation(const FRotator& Input1, const FRotator& Input2, double Ratio, EPCGMetadataRotatorOperation Operation)
 	{
 		switch (Operation)
 		{
-		case EPCGMedadataRotatorOperation::Combine:
+		case EPCGMetadataRotatorOperation::Combine:
 			return FRotator(Input2.Quaternion() * Input1.Quaternion());
-		case EPCGMedadataRotatorOperation::Lerp:
+		case EPCGMetadataRotatorOperation::Lerp:
 			return PCGMetadataRotatorHelpers::RLerp(Input1, Input2, Ratio, false);
-		case EPCGMedadataRotatorOperation::Invert:
+		case EPCGMetadataRotatorOperation::Invert:
 			return Input1.GetInverse();
-		case EPCGMedadataRotatorOperation::Normalize:
+		case EPCGMetadataRotatorOperation::Normalize:
 			return Input1.GetNormalized();
 		default:
 			return FRotator{};
 		}
 	}
 
-	inline FQuat ApplyRotatorOperation(const FQuat& Input1, const FQuat& Input2, double Ratio, EPCGMedadataRotatorOperation Operation)
+	inline FQuat ApplyRotatorOperation(const FQuat& Input1, const FQuat& Input2, double Ratio, EPCGMetadataRotatorOperation Operation)
 	{
 		switch (Operation)
 		{
-		case EPCGMedadataRotatorOperation::Combine:
+		case EPCGMetadataRotatorOperation::Combine:
 			return Input2 * Input1;
-		case EPCGMedadataRotatorOperation::Lerp:
+		case EPCGMetadataRotatorOperation::Lerp:
 			return FQuat::Slerp(Input1, Input2, Ratio);
-		case EPCGMedadataRotatorOperation::Invert:
+		case EPCGMetadataRotatorOperation::Invert:
 			return Input1.Inverse();
-		case EPCGMedadataRotatorOperation::Normalize:
+		case EPCGMetadataRotatorOperation::Normalize:
 			return Input1.GetNormalized();
 		default:
 			return FQuat{};
 		}
 	}
 
-	inline FQuat ApplyTransformOperation(const FQuat& Input, const FTransform& Transform, EPCGMedadataRotatorOperation Operation)
+	inline FQuat ApplyTransformOperation(const FQuat& Input, const FTransform& Transform, EPCGMetadataRotatorOperation Operation)
 	{
 		switch (Operation)
 		{
-		case EPCGMedadataRotatorOperation::TransformRotation:
+		case EPCGMetadataRotatorOperation::TransformRotation:
 			return Transform.TransformRotation(Input);
-		case EPCGMedadataRotatorOperation::InverseTransformRotation:
+		case EPCGMetadataRotatorOperation::InverseTransformRotation:
 			return Transform.InverseTransformRotation(Input);
 		default:
 			return FQuat{};
@@ -91,7 +91,7 @@ namespace PCGMetadataRotatorSettings
 	}
 
 	// In Kismet Math Library, they transform Rotators in Quaternions.
-	inline FRotator ApplyTransformOperation(const FRotator& Input, const FTransform& Transform, EPCGMedadataRotatorOperation Operation)
+	inline FRotator ApplyTransformOperation(const FRotator& Input, const FTransform& Transform, EPCGMetadataRotatorOperation Operation)
 	{
 		return FRotator(ApplyTransformOperation(Input.Quaternion(), Transform, Operation));
 	}
@@ -196,7 +196,7 @@ FPCGAttributePropertyInputSelector UPCGMetadataRotatorSettings::GetInputSource(u
 
 FName UPCGMetadataRotatorSettings::AdditionalTaskName() const
 {
-	if (const UEnum* EnumPtr = StaticEnum<EPCGMedadataRotatorOperation>())
+	if (const UEnum* EnumPtr = StaticEnum<EPCGMetadataRotatorOperation>())
 	{
 		return FName(FString("Rotator: ") + EnumPtr->GetNameStringByValue(static_cast<int>(Operation)));
 	}
@@ -219,17 +219,17 @@ FText UPCGMetadataRotatorSettings::GetDefaultNodeTitle() const
 
 TArray<FPCGPreConfiguredSettingsInfo> UPCGMetadataRotatorSettings::GetPreconfiguredInfo() const
 {
-	return PCGMetadataElementCommon::FillPreconfiguredSettingsInfoFromEnum<EPCGMedadataRotatorOperation>({ EPCGMedadataRotatorOperation::RotatorOp, EPCGMedadataRotatorOperation::TransformOp });
+	return PCGMetadataElementCommon::FillPreconfiguredSettingsInfoFromEnum<EPCGMetadataRotatorOperation>({ EPCGMetadataRotatorOperation::RotatorOp, EPCGMetadataRotatorOperation::TransformOp });
 }
 #endif // WITH_EDITOR
 
 void UPCGMetadataRotatorSettings::ApplyPreconfiguredSettings(const FPCGPreConfiguredSettingsInfo& PreconfiguredInfo)
 {
-	if (const UEnum* EnumPtr = StaticEnum<EPCGMedadataRotatorOperation>())
+	if (const UEnum* EnumPtr = StaticEnum<EPCGMetadataRotatorOperation>())
 	{
 		if (EnumPtr->IsValidEnumValue(PreconfiguredInfo.PreconfiguredIndex))
 		{
-			Operation = EPCGMedadataRotatorOperation(PreconfiguredInfo.PreconfiguredIndex);
+			Operation = EPCGMetadataRotatorOperation(PreconfiguredInfo.PreconfiguredIndex);
 		}
 	}
 }
diff --git a/Engine/Plugins/Experimental/PCG/Source/PCG/Private/Elements/Metadata/PCGMetadataTransformOpElement.cpp b/Engine/Plugins/Experimental/PCG/Source/PCG/Private/Elements/Metadata/PCGMetadataTransformOpElement.cpp
--- a/Engine/Plugins/Experimental/PCG/Source/PCG/Private/Elements/Metadata/PCGMetadataTransformOpElement.cpp
+++ b/Engine/Plugins/Experimental/PCG/Source/PCG/Private/Elements/Metadata/PCGMetadataTransformOpElement.cpp
@@ -11,14 +11,14 @@
 
 namespace PCGMetadataTransfromSettings
 {
-	inline bool IsUnaryOp(EPCGMedadataTransformOperation Operation)
+	inline bool IsUnaryOp(EPCGMetadataTransformOperation Operation)
 	{
-		return Operation == EPCGMedadataTransformOperation::Invert;
+		return Operation == EPCGMetadataTransformOperation::Invert;
 	}
 
-	inline bool IsTernaryOp(EPCGMedadataTransformOperation Operation)
+	inline bool IsTernaryOp(EPCGMetadataTransformOperation Operation)
 	{
-		return Operation == EPCGMedadataTransformOperation::Lerp;
+		return Operation == EPCGMetadataTransformOperation::Lerp;
 	}
 
 	// Taken from Kismet Math Library
@@ -56,15 +56,15 @@ namespace PCGMetadataTransfromSettings
 		}
 	}
 
-	inline FTransform ApplyTransformOperation(const FTransform& Value1, const FTransform& Value2, double Ratio, EPCGMedadataTransformOperation Operation, EPCGTransformLerpMode Mode)
+	inline FTransform ApplyTransformOperation(const FTransform& Value1, const FTransform& Value2, double Ratio, EPCGMetadataTransformOperation Operation, EPCGTransformLerpMode Mode)
 	{
 		switch (Operation)
 		{
-		case EPCGMedadataTransformOperation::Invert:
+		case EPCGMetadataTransformOperation::Invert:
 			return Value1.Inverse();
-		case EPCGMedadataTransformOperation::Compose:
+		case EPCGMetadataTransformOperation::Compose:
 			return Value1 * Value2;
-		case EPCGMedadataTransformOperation::Lerp:
+		case EPCGMetadataTransformOperation::Lerp:
 			return LerpTransform(Value1, Value2, Ratio, Mode);
 		default:
 			return FTransform{};
@@ -102,7 +102,7 @@ FName UPCGMetadataTransformSettings::GetInputPinLabel(uint32 Index) const
 	switch (Index)
 	{
 	case 0:
-		return (Operation == EPCGMedadataTransformOperation::Invert) ? PCGPinConstants::DefaultInputLabel : PCGMetadataSettingsBaseConstants::DoubleInputFirstLabel;
+		return (Operation == EPCGMetadataTransformOperation::Invert) ? PCGPinConstants::DefaultInputLabel : PCGMetadataSettingsBaseConstants::DoubleInputFirstLabel;
 	case 1:
 		return PCGMetadataSettingsBaseConstants::DoubleInputSecondLabel;
 	case 2:
@@ -114,11 +114,11 @@ FName UPCGMetadataTransformSettings::GetInputPinLabel(uint32 Index) const
 
 uint32 UPCGMetadataTransformSettings::GetOperandNum() const
 {
-	if (Operation == EPCGMedadataTransformOperation::Invert)
+	if (Operation == EPCGMetadataTransformOperation::Invert)
 	{
 		return 1;
 	}
-	else if (Operation == EPCGMedadataTransformOperation::Lerp)
+	else if (Operation == EPCGMetadataTransformOperation::Lerp)
 	{
 		return 3;
 	}
@@ -159,7 +159,7 @@ FPCGAttributePropertyInputSelector UPCGMetadataTransformSettings::GetInputSource
 
 FName UPCGMetadataTransformSettings::AdditionalTaskName() const
 {
-	if (const UEnum* EnumPtr = StaticEnum<EPCGMedadataTransformOperation>())
+	if (const UEnum* EnumPtr = StaticEnum<EPCGMetadataTransformOperation>())
 	{
 		return FName(FString("Transform: ") + EnumPtr->GetNameStringByValue(static_cast<int>(Operation)));
 	}
@@ -182,17 +182,17 @@ FText UPCGMetadataTransformSettings::GetDefaultNodeTitle() const
 
 TArray<FPCGPreConfiguredSettingsInfo> UPCGMetadataTransformSettings::GetPreconfiguredInfo() const
 {
-	return PCGMetadataElementCommon::FillPreconfiguredSettingsInfoFromEnum<EPCGMedadataTransformOperation>();
+	return PCGMetadataElementCommon::FillPreconfiguredSettingsInfoFromEnum<EPCGMetadataTransformOperation>();
 }
 #endif // WITH_EDITOR
 
 void UPCGMetadataTransformSettings::ApplyPreconfiguredSettings(const FPCGPreConfiguredSettingsInfo& PreconfiguredInfo)
 {
-	if (const UEnum* EnumPtr = StaticEnum<EPCGMedadataTransformOperation>())
+	if (const UEnum* EnumPtr = StaticEnum<EPCGMetadataTransformOperation>())
 	{
 		if (EnumPtr->IsValidEnumValue(PreconfiguredInfo.PreconfiguredIndex))
 		{
-			Operation = EPCGMedadataTransformOperation(PreconfiguredInfo.PreconfiguredIndex);
+			Operation = EPCGMetadataTransformOperation(PreconfiguredInfo.PreconfiguredIndex);
 		}
 	}
 }
diff --git a/Engine/Plugins/Experimental/PCG/Source/PCG/Private/Elements/Metadata/PCGMetadataTrigOpElement.cpp b/Engine/Plugins/Experimental/PCG/Source/PCG/Private/Elements/Metadata/PCGMetadataTrigOpElement.cpp
--- a/Engine/Plugins/Experimental/PCG/Source/PCG/Private/Elements/Metadata/PCGMetadataTrigOpElement.cpp
+++ b/Engine/Plugins/Experimental/PCG/Source/PCG/Private/Elements/Metadata/PCGMetadataTrigOpElement.cpp
@@ -9,25 +9,25 @@
 namespace PCGMetadataTrigSettings
 {
 	template <typename OutType>
-	OutType UnaryOp(const OutType& Input1, EPCGMedadataTrigOperation Operation)
+	OutType UnaryOp(const OutType& Input1, EPCGMetadataTrigOperation Operation)
 	{
 		switch (Operation)
 		{
-		case EPCGMedadataTrigOperation::Acos:
+		case EPCGMetadataTrigOperation::Acos:
 			return FMath::Acos(Input1);
-		case EPCGMedadataTrigOperation::Asin:
+		case EPCGMetadataTrigOperation::Asin:
 			return FMath::Asin(Input1);
-		case EPCGMedadataTrigOperation::Atan:
+		case EPCGMetadataTrigOperation::Atan:
 			return FMath::Atan(Input1);
-		case EPCGMedadataTrigOperation::Cos:
+		case EPCGMetadataTrigOperation::Cos:
 			return FMath::Cos(Input1);
-		case EPCGMedadataTrigOperation::Sin:
+		case EPCGMetadataTrigOperation::Sin:
 			return FMath::Sin(Input1);
-		case EPCGMedadataTrigOperation::Tan:
+		case EPCGMetadataTrigOperation::Tan:
 			return FMath::Tan(Input1);
-		case EPCGMedadataTrigOperation::DegToRad:
+		case EPCGMetadataTrigOperation::DegToRad:
 			return FMath::DegreesToRadians(Input1);
-		case EPCGMedadataTrigOperation::RadToDeg:
+		case EPCGMetadataTrigOperation::RadToDeg:
 			return FMath::RadiansToDegrees(Input1);
 		default:
 			return OutType{};
@@ -37,7 +37,7 @@ namespace PCGMetadataTrigSettings
 	template <typename OutType>
 	OutType BinaryOp(const OutType& Input1, const OutType& Input2)
 	{
-		// EPCGMedadataTrigOperation::Atan2:
+		// EPCGMetadataTrigOperation::Atan2:
 		return FMath::Atan2(Input1, Input2);
 	}
 }
@@ -66,7 +66,7 @@ FName UPCGMetadataTrigSettings::GetInputPinLabel(uint32 Index) const
 	switch (Index)
 	{
 	case 0:
-		return (Operation != EPCGMedadataTrigOperation::Atan2) ? PCGPinConstants::DefaultInputLabel : PCGMetadataSettingsBaseConstants::DoubleInputFirstLabel;
+		return (Operation != EPCGMetadataTrigOperation::Atan2) ? PCGPinConstants::DefaultInputLabel : PCGMetadataSettingsBaseConstants::DoubleInputFirstLabel;
 	case 1:
 		return PCGMetadataSettingsBaseConstants::DoubleInputSecondLabel;
 	default:
@@ -76,7 +76,7 @@ FName UPCGMetadataTrigSettings::GetInputPinLabel(uint32 Index) const
 
 uint32 UPCGMetadataTrigSettings::GetOperandNum() const
 {
-	return (Operation != EPCGMedadataTrigOperation::Atan2) ? 1 : 2;
+	return (Operation != EPCGMetadataTrigOperation::Atan2) ? 1 : 2;
 }
 
 bool UPCGMetadataTrigSettings::IsSupportedInputType(uint16 TypeId, uint32 InputIndex, bool& bHasSpecialRequirement) const
@@ -105,7 +105,7 @@ uint16 UPCGMetadataTrigSettings::GetOutputType(uint16 InputTypeId) const
 
 FName UPCGMetadataTrigSettings::AdditionalTaskName() const
 {
-	if (const UEnum* EnumPtr = StaticEnum<EPCGMedadataTrigOperation>())
+	if (const UEnum* EnumPtr = StaticEnum<EPCGMetadataTrigOperation>())
 	{
 		return FName(FString("Trig: ") + EnumPtr->GetNameStringByValue(static_cast<int>(Operation)));
 	}
@@ -128,17 +128,17 @@ FText UPCGMetadataTrigSettings::GetDefaultNodeTitle() const
 
 TArray<FPCGPreConfiguredSettingsInfo> UPCGMetadataTrigSettings::GetPreconfiguredInfo() const
 {
-	return PCGMetadataElementCommon::FillPreconfiguredSettingsInfoFromEnum<EPCGMedadataTrigOperation>();
+	return PCGMetadataElementCommon::FillPreconfiguredSettingsInfoFromEnum<EPCGMetadataTrigOperation>();
 }
 #endif // WITH_EDITOR
 
 void UPCGMetadataTrigSettings::ApplyPreconfiguredSettings(const FPCGPreConfiguredSettingsInfo& PreconfiguredInfo)
 {
-	if (const UEnum* EnumPtr = StaticEnum<EPCGMedadataTrigOperation>())
+	if (const UEnum* EnumPtr = StaticEnum<EPCGMetadataTrigOperation>())
 	{
 		if (EnumPtr->IsValidEnumValue(PreconfiguredInfo.PreconfiguredIndex))
 		{
-			Operation = EPCGMedadataTrigOperation(PreconfiguredInfo.PreconfiguredIndex);
+			Operation = EPCGMetadataTrigOperation(PreconfiguredInfo.PreconfiguredIndex);
 		}
 	}
 }
@@ -154,7 +154,7 @@ bool FPCGMetadataTrigElement::DoOperation(PCGMetadataOps::FOperationData& Operat
 
 	const UPCGMetadataTrigSettings* Settings = CastChecked<UPCGMetadataTrigSettings>(OperationData.Settings);
 
-	if (Settings->Operation == EPCGMedadataTrigOperation::Atan2)
+	if (Settings->Operation == EPCGMetadataTrigOperation::Atan2)
 	{
 		DoBinaryOp<double, double>(OperationData, [](const double& Value1, const double& Value2) -> double { return PCGMetadataTrigSettings::BinaryOp(Value1, Value2); });
 	}
diff --git a/Engine/Plugins/Experimental/PCG/Source/PCG/Private/Elements/Metadata/PCGMetadataVectorOpElement.cpp b/Engine/Plugins/Experimental/PCG/Source/PCG/Private/Elements/Metadata/PCGMetadataVectorOpElement.cpp
--- a/Engine/Plugins/Experimental/PCG/Source/PCG/Private/Elements/Metadata/PCGMetadataVectorOpElement.cpp
+++ b/Engine/Plugins/Experimental/PCG/Source/PCG/Private/Elements/Metadata/PCGMetadataVectorOpElement.cpp
@@ -15,29 +15,29 @@ namespace PCGMetadataVectorConstants
 
 namespace PCGMetadataVectorSettings
 {
-	inline constexpr bool IsUnaryOp(EPCGMedadataVectorOperation Operation)
+	inline constexpr bool IsUnaryOp(EPCGMetadataVectorOperation Operation)
 	{
-		return Operation == EPCGMedadataVectorOperation::Normalize || 
-			Operation == EPCGMedadataVectorOperation::Length;
+		return Operation == EPCGMetadataVectorOperation::Normalize || 
+			Operation == EPCGMetadataVectorOperation::Length;
 	}
 
-	inline constexpr bool IsTernaryOp(EPCGMedadataVectorOperation Operation)
+	inline constexpr bool IsTernaryOp(EPCGMetadataVectorOperation Operation)
 	{
-		return Operation == EPCGMedadataVectorOperation::RotateAroundAxis;
+		return Operation == EPCGMetadataVectorOperation::RotateAroundAxis;
 	}
 
-	inline constexpr bool IsTransformOp(EPCGMedadataVectorOperation Operation)
+	inline constexpr bool IsTransformOp(EPCGMetadataVectorOperation Operation)
 	{
-		return (uint16)Operation >= (uint16)EPCGMedadataVectorOperation::TransformOp;
+		return (uint16)Operation >= (uint16)EPCGMetadataVectorOperation::TransformOp;
 	}
 
 	// Mimic KismetMathLibrary
 	template <typename InType>
-	inline InType ApplyTransformOperation(const InType& Value, const FTransform& Transform, EPCGMedadataVectorOperation Operation)
+	inline InType ApplyTransformOperation(const InType& Value, const FTransform& Transform, EPCGMetadataVectorOperation Operation)
 	{
 		switch (Operation)
 		{
-		case EPCGMedadataVectorOperation::TransformDirection:
+		case EPCGMetadataVectorOperation::TransformDirection:
 		{
 			if constexpr (std::is_same_v<InType, FVector2D>)
 			{
@@ -54,7 +54,7 @@ namespace PCGMetadataVectorSettings
 				return Transform.TransformVectorNoScale(Value);
 			}
 		}
-		case EPCGMedadataVectorOperation::TransformLocation:
+		case EPCGMetadataVectorOperation::TransformLocation:
 		{
 			if constexpr (std::is_same_v<InType, FVector2D>)
 			{
@@ -74,7 +74,7 @@ namespace PCGMetadataVectorSettings
 			}
 			break;
 		}
-		case EPCGMedadataVectorOperation::InverseTransformDirection:
+		case EPCGMetadataVectorOperation::InverseTransformDirection:
 		{
 			if constexpr (std::is_same_v<InType, FVector2D>)
 			{
@@ -94,7 +94,7 @@ namespace PCGMetadataVectorSettings
 			}
 			break;
 		}
-		case EPCGMedadataVectorOperation::InverseTransformLocation:
+		case EPCGMetadataVectorOperation::InverseTransformLocation:
 		{
 			if constexpr (std::is_same_v<InType, FVector2D>)
 			{
@@ -122,11 +122,11 @@ namespace PCGMetadataVectorSettings
 	}
 
 	template <typename InType>
-	inline InType ApplyVectorOperation(const InType& Value1, const InType& Value2, double& DoubleValue, EPCGMedadataVectorOperation Operation)
+	inline InType ApplyVectorOperation(const InType& Value1, const InType& Value2, double& DoubleValue, EPCGMetadataVectorOperation Operation)
 	{
 		switch (Operation)
 		{
-		case EPCGMedadataVectorOperation::Cross:
+		case EPCGMetadataVectorOperation::Cross:
 		{
 			if constexpr (std::is_same_v<InType, FVector2D>)
 			{
@@ -138,7 +138,7 @@ namespace PCGMetadataVectorSettings
 			}
 			break;
 		}
-		case EPCGMedadataVectorOperation::Dot:
+		case EPCGMetadataVectorOperation::Dot:
 		{
 			if constexpr (std::is_same_v<FVector4, InType>)
 			{
@@ -150,10 +150,10 @@ namespace PCGMetadataVectorSettings
 			}
 			break;
 		}
-		case EPCGMedadataVectorOperation::Distance:
+		case EPCGMetadataVectorOperation::Distance:
 			DoubleValue = (Value1 - Value2).Size();
 			break;
-		case EPCGMedadataVectorOperation::RotateAroundAxis:
+		case EPCGMetadataVectorOperation::RotateAroundAxis:
 		{
 			if constexpr (std::is_same_v<FVector4, InType>)
 			{
@@ -173,7 +173,7 @@ namespace PCGMetadataVectorSettings
 			}
 			break;
 		}
-		case EPCGMedadataVectorOperation::Normalize:
+		case EPCGMetadataVectorOperation::Normalize:
 			if constexpr (std::is_same_v<FVector4, InType>)
 			{
 				double Length = Value1.Size();
@@ -193,7 +193,7 @@ namespace PCGMetadataVectorSettings
 				return Res;
 			}
 			break;
-		case EPCGMedadataVectorOperation::Length:
+		case EPCGMetadataVectorOperation::Length:
 			DoubleValue = Value1.Size();
 			break;
 		default:
@@ -237,7 +237,7 @@ FName UPCGMetadataVectorSettings::GetInputPinLabel(uint32 Index) const
 	{
 		if (PCGMetadataVectorSettings::IsUnaryOp(Operation) || 
 			PCGMetadataVectorSettings::IsTransformOp(Operation) ||
-			Operation == EPCGMedadataVectorOperation::RotateAroundAxis)
+			Operation == EPCGMetadataVectorOperation::RotateAroundAxis)
 		{
 			return PCGPinConstants::DefaultInputLabel;
 		}
@@ -248,7 +248,7 @@ FName UPCGMetadataVectorSettings::GetInputPinLabel(uint32 Index) const
 	}
 	case 1:
 	{
-		if (Operation == EPCGMedadataVectorOperation::RotateAroundAxis)
+		if (Operation == EPCGMetadataVectorOperation::RotateAroundAxis)
 		{
 			return PCGMetadataVectorConstants::AxisLabel;
 		}
@@ -288,7 +288,7 @@ bool UPCGMetadataVectorSettings::IsSupportedInputType(uint16 TypeId, uint32 Inpu
 {
 	bHasSpecialRequirement = false;
 
-	if (Operation == EPCGMedadataVectorOperation::RotateAroundAxis && InputIndex == 2)
+	if (Operation == EPCGMetadataVectorOperation::RotateAroundAxis && InputIndex == 2)
 	{
 		bHasSpecialRequirement = true;
 		return PCG::Private::IsOfTypes<double, float>(TypeId);
@@ -321,7 +321,7 @@ FPCGAttributePropertyInputSelector UPCGMetadataVectorSettings::GetInputSource(ui
 
 FName UPCGMetadataVectorSettings::AdditionalTaskName() const
 {
-	if (const UEnum* EnumPtr = StaticEnum<EPCGMedadataVectorOperation>())
+	if (const UEnum* EnumPtr = StaticEnum<EPCGMetadataVectorOperation>())
 	{
 		return FName(FString("Vector: ") + EnumPtr->GetNameStringByValue(static_cast<int>(Operation)));
 	}
@@ -344,17 +344,17 @@ FText UPCGMetadataVectorSettings::GetDefaultNodeTitle() const
 
 TArray<FPCGPreConfiguredSettingsInfo> UPCGMetadataVectorSettings::GetPreconfiguredInfo() const
 {
-	return PCGMetadataElementCommon::FillPreconfiguredSettingsInfoFromEnum<EPCGMedadataVectorOperation>({ EPCGMedadataVectorOperation::VectorOp, EPCGMedadataVectorOperation::TransformOp });
+	return PCGMetadataElementCommon::FillPreconfiguredSettingsInfoFromEnum<EPCGMetadataVectorOperation>({ EPCGMetadataVectorOperation::VectorOp, EPCGMetadataVectorOperation::TransformOp });
 }
 #endif // WITH_EDITOR
 
 void UPCGMetadataVectorSettings::ApplyPreconfiguredSettings(const FPCGPreConfiguredSettingsInfo& PreconfiguredInfo)
 {
-	if (const UEnum* EnumPtr = StaticEnum<EPCGMedadataVectorOperation>())
+	if (const UEnum* EnumPtr = StaticEnum<EPCGMetadataVectorOperation>())
 	{
 		if (EnumPtr->IsValidEnumValue(PreconfiguredInfo.PreconfiguredIndex))
 		{
-			Operation = EPCGMedadataVectorOperation(PreconfiguredInfo.PreconfiguredIndex);
+			Operation = EPCGMetadataVectorOperation(PreconfiguredInfo.PreconfiguredIndex);
 		}
 	}
 }
@@ -367,10 +367,10 @@ FPCGElementPtr UPCGMetadataVectorSettings::CreateElement() const
 uint16 UPCGMetadataVectorSettings::GetOutputType(uint16 InputTypeId) const
 {
 	// Dot, Length and Cross product with Vec2 output Double values
-	if (Operation == EPCGMedadataVectorOperation::Dot ||
-		Operation == EPCGMedadataVectorOperation::Length ||
-		Operation == EPCGMedadataVectorOperation::Distance ||
-		(Operation == EPCGMedadataVectorOperation::Cross && InputTypeId == (uint16)EPCGMetadataTypes::Vector2))
+	if (Operation == EPCGMetadataVectorOperation::Dot ||
+		Operation == EPCGMetadataVectorOperation::Length ||
+		Operation == EPCGMetadataVectorOperation::Distance ||
+		(Operation == EPCGMetadataVectorOperation::Cross && InputTypeId == (uint16)EPCGMetadataTypes::Vector2))
 	{
 		return (uint16)EPCGMetadataTypes::Double;
 	}
diff --git a/Engine/Plugins/Experimental/PCG/Source/PCG/Private/Tests/Elements/Metadata/PCGMetadataBooleanOpTest.cpp b/Engine/Plugins/Experimental/PCG/Source/PCG/Private/Tests/Elements/Metadata/PCGMetadataBooleanOpTest.cpp
--- a/Engine/Plugins/Experimental/PCG/Source/PCG/Private/Tests/Elements/Metadata/PCGMetadataBooleanOpTest.cpp
+++ b/Engine/Plugins/Experimental/PCG/Source/PCG/Private/Tests/Elements/Metadata/PCGMetadataBooleanOpTest.cpp
@@ -116,7 +116,7 @@ bool FPCGMetadataBooleanOpTest::RunTest(const FString& Parameters)
 	/////////////////////////////////////////
 	ParamTaggedData1.Pin = PCGPinConstants::DefaultInputLabel;
 
-	Settings->Operation = EPCGMedadataBooleanOperation::Not;
+	Settings->Operation = EPCGMetadataBooleanOperation::Not;
 
 	{
 		Settings->InputSource1.SetAttributeName(TrueAttribute);
@@ -135,7 +135,7 @@ bool FPCGMetadataBooleanOpTest::RunTest(const FString& Parameters)
 	// Validate invalid op when there is type mismatch.
 	ParamTaggedData1.Pin = PCGMetadataSettingsBaseConstants::DoubleInputFirstLabel;
 	ParamTaggedData2.Pin = PCGMetadataSettingsBaseConstants::DoubleInputSecondLabel;
-	Settings->Operation = EPCGMedadataBooleanOperation::And;
+	Settings->Operation = EPCGMetadataBooleanOperation::And;
 
 	AddExpectedError(TEXT("Attribute/Property 'Invalid' from pin"), EAutomationExpectedErrorFlags::Contains, 3);
 
@@ -186,7 +186,7 @@ bool FPCGMetadataBooleanOpTest::RunTest(const FString& Parameters)
 	/////////////////////////////////////////
 	// Or
 	/////////////////////////////////////////
-	Settings->Operation = EPCGMedadataBooleanOperation::Or;
+	Settings->Operation = EPCGMetadataBooleanOperation::Or;
 	{
 		Settings->InputSource1.SetAttributeName(TrueAttribute);
 		Settings->InputSource2.SetAttributeName(TrueAttribute);
@@ -214,7 +214,7 @@ bool FPCGMetadataBooleanOpTest::RunTest(const FString& Parameters)
 	/////////////////////////////////////////
 	// Xor
 	/////////////////////////////////////////
-	Settings->Operation = EPCGMedadataBooleanOperation::Xor;
+	Settings->Operation = EPCGMetadataBooleanOperation::Xor;
 	{
 		Settings->InputSource1.SetAttributeName(TrueAttribute);
 		Settings->InputSource2.SetAttributeName(TrueAttribute);
diff --git a/Engine/Plugins/Experimental/PCG/Source/PCG/Private/Tests/Elements/Metadata/PCGMetadataMathsOpTest.cpp b/Engine/Plugins/Experimental/PCG/Source/PCG/Private/Tests/Elements/Metadata/PCGMetadataMathsOpTest.cpp
--- a/Engine/Plugins/Experimental/PCG/Source/PCG/Private/Tests/Elements/Metadata/PCGMetadataMathsOpTest.cpp
+++ b/Engine/Plugins/Experimental/PCG/Source/PCG/Private/Tests/Elements/Metadata/PCGMetadataMathsOpTest.cpp
@@ -23,7 +23,7 @@ protected:
 	template <typename T>
 	struct TestParams
 	{
-		explicit TestParams(const EPCGMedadataMathsOperation Operation, TArray<TArray<FPCGTaggedData>>&& InputDataByPin, TArray<T>&& ExpectedDefaultValues, TArray<T>&& ExpectedResultValues, const int32 PinToForward = 0) :
+		explicit TestParams(const EPCGMetadataMathsOperation Operation, TArray<TArray<FPCGTaggedData>>&& InputDataByPin, TArray<T>&& ExpectedDefaultValues, TArray<T>&& ExpectedResultValues, const int32 PinToForward = 0) :
 			Operation(Operation),
 			InputDataByPin(std::move(InputDataByPin)),
 			ExpectedDefaultValues(ExpectedDefaultValues),
@@ -32,7 +32,7 @@ protected:
 		{
 		}
 
-		EPCGMedadataMathsOperation Operation;
+		EPCGMetadataMathsOperation Operation;
 		TArray<TArray<FPCGTaggedData>> InputDataByPin;
 		TArray<T> ExpectedDefaultValues;
 		TArray<T> ExpectedResultValues;
@@ -43,7 +43,7 @@ protected:
 	template <typename T>
 	bool ExecuteTest(TestParams<T>&& Params, const FPCGAttributePropertyInputSelector InputSelector = DefaultSelector)
 	{
-		const EPCGMedadataMathsOperation& Operation = Params.Operation;
+		const EPCGMetadataMathsOperation& Operation = Params.Operation;
 		const TArray<TArray<FPCGTaggedData>>& InputDataByPin = Params.InputDataByPin;
 		const TArray<T>& ExpectedDefaults = Params.ExpectedDefaultValues;
 		const TArray<T>& ExpectedValues = Params.ExpectedResultValues;
@@ -51,7 +51,7 @@ protected:
 
 		check(!InputDataByPin.IsEmpty() && !ExpectedDefaults.IsEmpty() && !ExpectedValues.IsEmpty());
 
-		const UEnum* EnumClass = StaticEnum<EPCGMedadataMathsOperation>();
+		const UEnum* EnumClass = StaticEnum<EPCGMetadataMathsOperation>();
 		auto FormatText = [Operation, PinToForward, EnumClass](const FString& Text) -> FString
 		{
 			return FString::Printf(TEXT("[Pin %d] %s (%s): %s."), PinToForward, *EnumClass->GetDisplayNameTextByValue(static_cast<int64>(Operation)).ToString(), *PCG::Private::GetTypeName<T>(), *Text);
@@ -190,7 +190,7 @@ bool FPCGMetadataMathsOp_OneMinus::RunTest(const FString& Parameters)
 {
 	bool bSuccess = true;
 
-	using EPCGMedadataMathsOperation::OneMinus;
+	using EPCGMetadataMathsOperation::OneMinus;
 
 	// Param Data
 	bSuccess &= ExecuteTest(TestParams(OneMinus, GenerateParamData<int32>({{5}}), TArray({1}), TArray({-4})));
@@ -219,7 +219,7 @@ bool FPCGMetadataMathsOp_Sign::RunTest(const FString& Parameters)
 {
 	bool bSuccess = true;
 
-	using EPCGMedadataMathsOperation::Sign;
+	using EPCGMetadataMathsOperation::Sign;
 
 	// Param Data
 	bSuccess &= ExecuteTest(TestParams(Sign, GenerateParamData<int32>({{0}}), TArray({0}), TArray({0})));
@@ -250,7 +250,7 @@ bool FPCGMetadataMathsOp_Frac::RunTest(const FString& Parameters)
 {
 	bool bSuccess = true;
 
-	using EPCGMedadataMathsOperation::Frac;
+	using EPCGMetadataMathsOperation::Frac;
 
 	// Param Data
 	bSuccess &= ExecuteTest(TestParams(Frac, GenerateParamData<int32>({{0}}), TArray({0}), TArray({0})));
@@ -279,7 +279,7 @@ bool FPCGMetadataMathsOp_Truncate::RunTest(const FString& Parameters)
 {
 	bool bSuccess = true;
 
-	using EPCGMedadataMathsOperation::Truncate;
+	using EPCGMetadataMathsOperation::Truncate;
 
 	// Param Data
 	bSuccess &= ExecuteTest(TestParams(Truncate, GenerateParamData<int32>({{0}}), TArray({0}), TArray({0})));
@@ -308,7 +308,7 @@ bool FPCGMetadataMathsOp_Add::RunTest(const FString& Parameters)
 {
 	bool bSuccess = true;
 
-	using EPCGMedadataMathsOperation::Add;
+	using EPCGMetadataMathsOperation::Add;
 
 	// Param Data
 	bSuccess &= ExecuteTest(TestParams(Add, GenerateParamData<int32>({{5}, {1}}), TArray({0}), TArray({6})));
@@ -337,7 +337,7 @@ bool FPCGMetadataMathsOp_Subtract::RunTest(const FString& Parameters)
 {
 	bool bSuccess = true;
 
-	using EPCGMedadataMathsOperation::Subtract;
+	using EPCGMetadataMathsOperation::Subtract;
 
 	// Param Data
 	bSuccess &= ExecuteTest(TestParams(Subtract, GenerateParamData<int32>({{5}, {1}}), TArray({0}), TArray({4})));
@@ -366,7 +366,7 @@ bool FPCGMetadataMathsOp_Multiply::RunTest(const FString& Parameters)
 {
 	bool bSuccess = true;
 
-	using EPCGMedadataMathsOperation::Multiply;
+	using EPCGMetadataMathsOperation::Multiply;
 
 	// Param Data
 	bSuccess &= ExecuteTest(TestParams(Multiply, GenerateParamData<int32>({{5}, {1}}), TArray({0}), TArray({5})));
@@ -395,7 +395,7 @@ bool FPCGMetadataMathsOp_Divide::RunTest(const FString& Parameters)
 {
 	bool bSuccess = true;
 
-	using EPCGMedadataMathsOperation::Divide;
+	using EPCGMetadataMathsOperation::Divide;
 
 	// Param Data
 	bSuccess &= ExecuteTest(TestParams(Divide, GenerateParamData<int32>({{32}, {8}}), TArray({0}), TArray({4})));
@@ -424,7 +424,7 @@ bool FPCGMetadataMathsOp_MultipleInput::RunTest(const FString& Parameters)
 {
 	bool bSuccess = true;
 
-	using EPCGMedadataMathsOperation::Add;
+	using EPCGMetadataMathsOperation::Add;
 
 	// 2 x 2 - N:N
 	bSuccess &= ExecuteTest(TestParams(Add, GenerateParamData<int32>({{4, 5}, {6, 7}}), TArray({0, 0}), TArray({10, 12})));
diff --git a/Engine/Plugins/Experimental/PCG/Source/PCG/Public/Elements/Metadata/PCGMetadataBitwiseOpElement.h b/Engine/Plugins/Experimental/PCG/Source/PCG/Public/Elements/Metadata/PCGMetadataBitwiseOpElement.h
--- a/Engine/Plugins/Experimental/PCG/Source/PCG/Public/Elements/Metadata/PCGMetadataBitwiseOpElement.h
+++ b/Engine/Plugins/Experimental/PCG/Source/PCG/Public/Elements/Metadata/PCGMetadataBitwiseOpElement.h
@@ -7,7 +7,7 @@
 #include "PCGMetadataBitwiseOpElement.generated.h"
 
 UENUM()
-enum class EPCGMedadataBitwiseOperation : uint16
+enum class EPCGMetadataBitwiseOperation : uint16
 {
 	And,
 	Not,
@@ -56,12 +56,12 @@ protected:
 
 public:
 	UPROPERTY(BlueprintReadWrite, EditAnywhere, Category = Settings)
-	EPCGMedadataBitwiseOperation Operation = EPCGMedadataBitwiseOperation::And;
+	EPCGMetadataBitwiseOperation Operation = EPCGMetadataBitwiseOperation::And;
 
 	UPROPERTY(BlueprintReadWrite, EditAnywhere, Category = Input, meta = (PCG_Overridable))
 	FPCGAttributePropertyInputSelector InputSource1;
 
-	UPROPERTY(BlueprintReadWrite, EditAnywhere, Category = Input, meta = (EditCondition = "Operation != EPCGMedadataBitwiseOperation::Not", EditConditionHides, PCG_Overridable))
+	UPROPERTY(BlueprintReadWrite, EditAnywhere, Category = Input, meta = (EditCondition = "Operation != EPCGMetadataBitwiseOperation::Not", EditConditionHides, PCG_Overridable))
 	FPCGAttributePropertyInputSelector InputSource2;
 	
 #if WITH_EDITORONLY_DATA
diff --git a/Engine/Plugins/Experimental/PCG/Source/PCG/Public/Elements/Metadata/PCGMetadataBooleanOpElement.h b/Engine/Plugins/Experimental/PCG/Source/PCG/Public/Elements/Metadata/PCGMetadataBooleanOpElement.h
--- a/Engine/Plugins/Experimental/PCG/Source/PCG/Public/Elements/Metadata/PCGMetadataBooleanOpElement.h
+++ b/Engine/Plugins/Experimental/PCG/Source/PCG/Public/Elements/Metadata/PCGMetadataBooleanOpElement.h
@@ -7,7 +7,7 @@
 #include "PCGMetadataBooleanOpElement.generated.h"
 
 UENUM()
-enum class EPCGMedadataBooleanOperation : uint16
+enum class EPCGMetadataBooleanOperation : uint16
 {
 	And,
 	Not,
@@ -56,12 +56,12 @@ protected:
 
 public:
 	UPROPERTY(BlueprintReadWrite, EditAnywhere, Category = Settings)
-	EPCGMedadataBooleanOperation Operation = EPCGMedadataBooleanOperation::And;
+	EPCGMetadataBooleanOperation Operation = EPCGMetadataBooleanOperation::And;
 
 	UPROPERTY(BlueprintReadWrite, EditAnywhere, Category = Input, meta = (PCG_Overridable))
 	FPCGAttributePropertyInputSelector InputSource1;
 
-	UPROPERTY(BlueprintReadWrite, EditAnywhere, Category = Input, meta = (EditCondition = "Operation != EPCGMedadataBooleanOperation::Not", EditConditionHides, PCG_Overridable))
+	UPROPERTY(BlueprintReadWrite, EditAnywhere, Category = Input, meta = (EditCondition = "Operation != EPCGMetadataBooleanOperation::Not", EditConditionHides, PCG_Overridable))
 	FPCGAttributePropertyInputSelector InputSource2;
 
 #if WITH_EDITORONLY_DATA
diff --git a/Engine/Plugins/Experimental/PCG/Source/PCG/Public/Elements/Metadata/PCGMetadataCompareOpElement.h b/Engine/Plugins/Experimental/PCG/Source/PCG/Public/Elements/Metadata/PCGMetadataCompareOpElement.h
--- a/Engine/Plugins/Experimental/PCG/Source/PCG/Public/Elements/Metadata/PCGMetadataCompareOpElement.h
+++ b/Engine/Plugins/Experimental/PCG/Source/PCG/Public/Elements/Metadata/PCGMetadataCompareOpElement.h
@@ -7,7 +7,7 @@
 #include "PCGMetadataCompareOpElement.generated.h"
 
 UENUM()
-enum class EPCGMedadataCompareOperation : uint16
+enum class EPCGMetadataCompareOperation : uint16
 {
 	Equal,
 	NotEqual,
@@ -58,7 +58,7 @@ protected:
 
 public:
 	UPROPERTY(BlueprintReadWrite, EditAnywhere, Category = Settings)
-	EPCGMedadataCompareOperation Operation = EPCGMedadataCompareOperation::Equal;
+	EPCGMetadataCompareOperation Operation = EPCGMetadataCompareOperation::Equal;
 
 	UPROPERTY(BlueprintReadWrite, EditAnywhere, Category = Input, meta = (PCG_Overridable))
 	FPCGAttributePropertyInputSelector InputSource1;
@@ -66,7 +66,7 @@ public:
 	UPROPERTY(BlueprintReadWrite, EditAnywhere, Category = Input, meta = (PCG_Overridable))
 	FPCGAttributePropertyInputSelector InputSource2;
 
-	UPROPERTY(BlueprintReadWrite, EditAnywhere, Category = Settings, meta = (EditCondition = "Operation == EPCGMedadataCompareOperation::Equal || Operation == EPCGMedadataCompareOperation::NotEqual", EditConditionHides))
+	UPROPERTY(BlueprintReadWrite, EditAnywhere, Category = Settings, meta = (EditCondition = "Operation == EPCGMetadataCompareOperation::Equal || Operation == EPCGMetadataCompareOperation::NotEqual", EditConditionHides))
 	double Tolerance = UE_DOUBLE_SMALL_NUMBER;
 
 #if WITH_EDITORONLY_DATA
diff --git a/Engine/Plugins/Experimental/PCG/Source/PCG/Public/Elements/Metadata/PCGMetadataMathsOpElement.h b/Engine/Plugins/Experimental/PCG/Source/PCG/Public/Elements/Metadata/PCGMetadataMathsOpElement.h
--- a/Engine/Plugins/Experimental/PCG/Source/PCG/Public/Elements/Metadata/PCGMetadataMathsOpElement.h
+++ b/Engine/Plugins/Experimental/PCG/Source/PCG/Public/Elements/Metadata/PCGMetadataMathsOpElement.h
@@ -7,7 +7,7 @@
 #include "PCGMetadataMathsOpElement.generated.h"
 
 UENUM(Meta=(Bitflags))
-enum class EPCGMedadataMathsOperation : uint16
+enum class EPCGMetadataMathsOperation : uint16
 {
 	// Unary op
 	UnaryOp = 1 << 10 UMETA(Hidden),
@@ -40,7 +40,7 @@ enum class EPCGMedadataMathsOperation : uint16
 	Clamp,
 	Lerp,
 };
-ENUM_CLASS_FLAGS(EPCGMedadataMathsOperation);
+ENUM_CLASS_FLAGS(EPCGMetadataMathsOperation);
 
 UCLASS(BlueprintType, ClassGroup = (Procedural))
 class PCG_API UPCGMetadataMathsSettings : public UPCGMetadataSettingsBase
@@ -83,23 +83,23 @@ protected:
 
 public:
 	UPROPERTY(BlueprintReadWrite, EditAnywhere, Category = Settings)
-	EPCGMedadataMathsOperation Operation = EPCGMedadataMathsOperation::Add;
+	EPCGMetadataMathsOperation Operation = EPCGMetadataMathsOperation::Add;
 
 	/** For rounding operation, if the input type is float or double, use this option to force the output attribute to be int64. */
-	UPROPERTY(BlueprintReadWrite, EditAnywhere, Category = Settings, meta = (EditConditionHides, EditCondition = "Operation == EPCGMedadataMathsOperation::Round || Operation == EPCGMedadataMathsOperation::Truncate || Operation == EPCGMedadataMathsOperation::Floor || Operation == EPCGMedadataMathsOperation::Ceil"))
+	UPROPERTY(BlueprintReadWrite, EditAnywhere, Category = Settings, meta = (EditConditionHides, EditCondition = "Operation == EPCGMetadataMathsOperation::Round || Operation == EPCGMetadataMathsOperation::Truncate || Operation == EPCGMetadataMathsOperation::Floor || Operation == EPCGMetadataMathsOperation::Ceil"))
 	bool bForceRoundingOpToInt = false;
 
 	/** For operations that can yield floating point values, if the input type are ints, use this option to force the output attribute to be double. */
-	UPROPERTY(BlueprintReadWrite, EditAnywhere, Category = Settings, meta = (EditConditionHides, EditCondition = "Operation == EPCGMedadataMathsOperation::Divide || Operation == EPCGMedadataMathsOperation::Sqrt || Operation == EPCGMedadataMathsOperation::Pow || Operation == EPCGMedadataMathsOperation::Lerp"))
+	UPROPERTY(BlueprintReadWrite, EditAnywhere, Category = Settings, meta = (EditConditionHides, EditCondition = "Operation == EPCGMetadataMathsOperation::Divide || Operation == EPCGMetadataMathsOperation::Sqrt || Operation == EPCGMetadataMathsOperation::Pow || Operation == EPCGMetadataMathsOperation::Lerp"))
 	bool bForceOpToDouble = false;
 
 	UPROPERTY(BlueprintReadWrite, EditAnywhere, Category = Input, meta = (PCG_Overridable))
 	FPCGAttributePropertyInputSelector InputSource1;
 
-	UPROPERTY(BlueprintReadWrite, EditAnywhere, Category = Input, meta = (EditCondition = "(Operation & '/Script/PCG.EPCGMedadataMathsOperation::BinaryOp') || (Operation & '/Script/PCG.EPCGMedadataMathsOperation::TernaryOp')", EditConditionHides, PCG_Overridable))
+	UPROPERTY(BlueprintReadWrite, EditAnywhere, Category = Input, meta = (EditCondition = "(Operation & '/Script/PCG.EPCGMetadataMathsOperation::BinaryOp') || (Operation & '/Script/PCG.EPCGMetadataMathsOperation::TernaryOp')", EditConditionHides, PCG_Overridable))
 	FPCGAttributePropertyInputSelector InputSource2;
 
-	UPROPERTY(BlueprintReadWrite, EditAnywhere, Category = Input, meta = (EditCondition = "Operation & '/Script/PCG.EPCGMedadataMathsOperation::TernaryOp'", EditConditionHides, PCG_Overridable))
+	UPROPERTY(BlueprintReadWrite, EditAnywhere, Category = Input, meta = (EditCondition = "Operation & '/Script/PCG.EPCGMetadataMathsOperation::TernaryOp'", EditConditionHides, PCG_Overridable))
 	FPCGAttributePropertyInputSelector InputSource3;
 
 #if WITH_EDITORONLY_DATA
diff --git a/Engine/Plugins/Experimental/PCG/Source/PCG/Public/Elements/Metadata/PCGMetadataRotatorOpElement.h b/Engine/Plugins/Experimental/PCG/Source/PCG/Public/Elements/Metadata/PCGMetadataRotatorOpElement.h
--- a/Engine/Plugins/Experimental/PCG/Source/PCG/Public/Elements/Metadata/PCGMetadataRotatorOpElement.h
+++ b/Engine/Plugins/Experimental/PCG/Source/PCG/Public/Elements/Metadata/PCGMetadataRotatorOpElement.h
@@ -7,7 +7,7 @@
 #include "PCGMetadataRotatorOpElement.generated.h"
 
 UENUM()
-enum class EPCGMedadataRotatorOperation : uint16
+enum class EPCGMetadataRotatorOperation : uint16
 {
 	RotatorOp = 0 UMETA(Hidden),
 	Combine,
@@ -66,15 +66,15 @@ protected:
 
 public:
 	UPROPERTY(BlueprintReadWrite, EditAnywhere, Category = Settings)
-	EPCGMedadataRotatorOperation Operation = EPCGMedadataRotatorOperation::Combine;
+	EPCGMetadataRotatorOperation Operation = EPCGMetadataRotatorOperation::Combine;
 
 	UPROPERTY(BlueprintReadWrite, EditAnywhere, Category = Input, meta = (PCG_Overridable))
 	FPCGAttributePropertyInputSelector InputSource1;
 
-	UPROPERTY(BlueprintReadWrite, EditAnywhere, Category = Input, meta = (EditCondition = "Operation != EPCGMedadataRotatorOperation::Invert", EditConditionHides, PCG_Overridable))
+	UPROPERTY(BlueprintReadWrite, EditAnywhere, Category = Input, meta = (EditCondition = "Operation != EPCGMetadataRotatorOperation::Invert", EditConditionHides, PCG_Overridable))
 	FPCGAttributePropertyInputSelector InputSource2;
 
-	UPROPERTY(BlueprintReadWrite, EditAnywhere, Category = Input, meta = (EditCondition = "Operation == EPCGMedadataRotatorOperation::Lerp", EditConditionHides, PCG_Overridable))
+	UPROPERTY(BlueprintReadWrite, EditAnywhere, Category = Input, meta = (EditCondition = "Operation == EPCGMetadataRotatorOperation::Lerp", EditConditionHides, PCG_Overridable))
 	FPCGAttributePropertyInputSelector InputSource3;
 
 #if WITH_EDITORONLY_DATA
diff --git a/Engine/Plugins/Experimental/PCG/Source/PCG/Public/Elements/Metadata/PCGMetadataTransformOpElement.h b/Engine/Plugins/Experimental/PCG/Source/PCG/Public/Elements/Metadata/PCGMetadataTransformOpElement.h
--- a/Engine/Plugins/Experimental/PCG/Source/PCG/Public/Elements/Metadata/PCGMetadataTransformOpElement.h
+++ b/Engine/Plugins/Experimental/PCG/Source/PCG/Public/Elements/Metadata/PCGMetadataTransformOpElement.h
@@ -7,7 +7,7 @@
 #include "PCGMetadataTransformOpElement.generated.h"
 
 UENUM()
-enum class EPCGMedadataTransformOperation : uint16
+enum class EPCGMetadataTransformOperation : uint16
 {
 	Compose,
 	Invert,
@@ -68,18 +68,18 @@ protected:
 
 public:
 	UPROPERTY(BlueprintReadWrite, EditAnywhere, Category = Settings)
-	EPCGMedadataTransformOperation Operation = EPCGMedadataTransformOperation::Compose;
+	EPCGMetadataTransformOperation Operation = EPCGMetadataTransformOperation::Compose;
 
-	UPROPERTY(BlueprintReadWrite, EditAnywhere, Category = Settings, meta = (EditCondition = "Operation == EPCGMedadataTransformOperation::Lerp", EditConditionHides))
+	UPROPERTY(BlueprintReadWrite, EditAnywhere, Category = Settings, meta = (EditCondition = "Operation == EPCGMetadataTransformOperation::Lerp", EditConditionHides))
 	EPCGTransformLerpMode TransformLerpMode = EPCGTransformLerpMode::QuatInterp;
 
 	UPROPERTY(BlueprintReadWrite, EditAnywhere, Category = Input, meta = (PCG_Overridable))
 	FPCGAttributePropertyInputSelector InputSource1;
 
-	UPROPERTY(BlueprintReadWrite, EditAnywhere, Category = Input, meta = (EditCondition = "Operation != EPCGMedadataTransformOperation::Invert", EditConditionHides, PCG_Overridable))
+	UPROPERTY(BlueprintReadWrite, EditAnywhere, Category = Input, meta = (EditCondition = "Operation != EPCGMetadataTransformOperation::Invert", EditConditionHides, PCG_Overridable))
 	FPCGAttributePropertyInputSelector InputSource2;
 
-	UPROPERTY(BlueprintReadWrite, EditAnywhere, Category = Input, meta = (EditCondition = "Operation == EPCGMedadataTransformOperation::Lerp", EditConditionHides, PCG_Overridable))
+	UPROPERTY(BlueprintReadWrite, EditAnywhere, Category = Input, meta = (EditCondition = "Operation == EPCGMetadataTransformOperation::Lerp", EditConditionHides, PCG_Overridable))
 	FPCGAttributePropertyInputSelector InputSource3;
 
 #if WITH_EDITORONLY_DATA
diff --git a/Engine/Plugins/Experimental/PCG/Source/PCG/Public/Elements/Metadata/PCGMetadataTrigOpElement.h b/Engine/Plugins/Experimental/PCG/Source/PCG/Public/Elements/Metadata/PCGMetadataTrigOpElement.h
--- a/Engine/Plugins/Experimental/PCG/Source/PCG/Public/Elements/Metadata/PCGMetadataTrigOpElement.h
+++ b/Engine/Plugins/Experimental/PCG/Source/PCG/Public/Elements/Metadata/PCGMetadataTrigOpElement.h
@@ -7,7 +7,7 @@
 #include "PCGMetadataTrigOpElement.generated.h"
 
 UENUM()
-enum class EPCGMedadataTrigOperation : uint16
+enum class EPCGMetadataTrigOperation : uint16
 {
 	Acos,
 	Asin,
@@ -61,12 +61,12 @@ protected:
 
 public:
 	UPROPERTY(BlueprintReadWrite, EditAnywhere, Category = Settings)
-	EPCGMedadataTrigOperation Operation = EPCGMedadataTrigOperation::Acos;
+	EPCGMetadataTrigOperation Operation = EPCGMetadataTrigOperation::Acos;
 
 	UPROPERTY(BlueprintReadWrite, EditAnywhere, Category = Input, meta = (PCG_Overridable))
 	FPCGAttributePropertyInputSelector InputSource1;
 
-	UPROPERTY(BlueprintReadWrite, EditAnywhere, Category = Input, meta = (EditCondition = "Operation == EPCGMedadataTrigOperation::Atan2", EditConditionHides, PCG_Overridable))
+	UPROPERTY(BlueprintReadWrite, EditAnywhere, Category = Input, meta = (EditCondition = "Operation == EPCGMetadataTrigOperation::Atan2", EditConditionHides, PCG_Overridable))
 	FPCGAttributePropertyInputSelector InputSource2;
 
 #if WITH_EDITORONLY_DATA
diff --git a/Engine/Plugins/Experimental/PCG/Source/PCG/Public/Elements/Metadata/PCGMetadataVectorOpElement.h b/Engine/Plugins/Experimental/PCG/Source/PCG/Public/Elements/Metadata/PCGMetadataVectorOpElement.h
--- a/Engine/Plugins/Experimental/PCG/Source/PCG/Public/Elements/Metadata/PCGMetadataVectorOpElement.h
+++ b/Engine/Plugins/Experimental/PCG/Source/PCG/Public/Elements/Metadata/PCGMetadataVectorOpElement.h
@@ -7,7 +7,7 @@
 #include "PCGMetadataVectorOpElement.generated.h"
 
 UENUM()
-enum class EPCGMedadataVectorOperation : uint16
+enum class EPCGMetadataVectorOperation : uint16
 {
 	VectorOp = 0 UMETA(Hidden),
 	Cross,
@@ -65,15 +65,15 @@ protected:
 
 public:
 	UPROPERTY(BlueprintReadWrite, EditAnywhere, Category = Settings)
-	EPCGMedadataVectorOperation Operation = EPCGMedadataVectorOperation::Cross;
+	EPCGMetadataVectorOperation Operation = EPCGMetadataVectorOperation::Cross;
 
 	UPROPERTY(BlueprintReadWrite, EditAnywhere, Category = Input, meta = (PCG_Overridable))
 	FPCGAttributePropertyInputSelector InputSource1;
 
-	UPROPERTY(BlueprintReadWrite, EditAnywhere, Category = Input, meta = (EditCondition = "Operation != EPCGMedadataVectorOperation::Normalize && Operation != EPCGMedadataVectorOperation::Length", EditConditionHides, PCG_Overridable))
+	UPROPERTY(BlueprintReadWrite, EditAnywhere, Category = Input, meta = (EditCondition = "Operation != EPCGMetadataVectorOperation::Normalize && Operation != EPCGMetadataVectorOperation::Length", EditConditionHides, PCG_Overridable))
 	FPCGAttributePropertyInputSelector InputSource2;
 
-	UPROPERTY(BlueprintReadWrite, EditAnywhere, Category = Input, meta = (EditCondition = "Operation == EPCGMedadataVectorOperation::RotateAroundAxis", EditConditionHides, PCG_Overridable))
+	UPROPERTY(BlueprintReadWrite, EditAnywhere, Category = Input, meta = (EditCondition = "Operation == EPCGMetadataVectorOperation::RotateAroundAxis", EditConditionHides, PCG_Overridable))
 	FPCGAttributePropertyInputSelector InputSource3;
 
 #if WITH_EDITORONLY_DATA
