Added missing profiling tags and overrides

#rnx

diff --git a/Engine/Plugins/Experimental/PCG/Source/PCG/Private/Elements/PCGCullPointsOutsideActorBounds.cpp b/Engine/Plugins/Experimental/PCG/Source/PCG/Private/Elements/PCGCullPointsOutsideActorBounds.cpp
--- a/Engine/Plugins/Experimental/PCG/Source/PCG/Private/Elements/PCGCullPointsOutsideActorBounds.cpp
+++ b/Engine/Plugins/Experimental/PCG/Source/PCG/Private/Elements/PCGCullPointsOutsideActorBounds.cpp
@@ -40,6 +40,8 @@ FPCGElementPtr UPCGCullPointsOutsideActorBoundsSettings::CreateElement() const
 
 bool FPCGCullPointsOutsideActorBoundsElement::ExecuteInternal(FPCGContext* Context) const
 {
+	TRACE_CPUPROFILER_EVENT_SCOPE(FPCGCullPointsOutsideActorBoundsElement::Execute);
+
 	if (!Context->SourceComponent.IsValid())
 	{
 		return true;
diff --git a/Engine/Plugins/Experimental/PCG/Source/PCG/Private/Elements/PCGGetBounds.cpp b/Engine/Plugins/Experimental/PCG/Source/PCG/Private/Elements/PCGGetBounds.cpp
--- a/Engine/Plugins/Experimental/PCG/Source/PCG/Private/Elements/PCGGetBounds.cpp
+++ b/Engine/Plugins/Experimental/PCG/Source/PCG/Private/Elements/PCGGetBounds.cpp
@@ -45,6 +45,8 @@ FPCGElementPtr UPCGGetBoundsSettings::CreateElement() const
 
 bool FPCGGetBoundsElement::ExecuteInternal(FPCGContext* Context) const
 {
+	TRACE_CPUPROFILER_EVENT_SCOPE(FPCGGetBoundsElement::Execute);
+
 	const TArray<FPCGTaggedData>& Inputs = Context->InputData.TaggedData;
 	bool bEmittedUnboundedDataWarning = false;
 
diff --git a/Engine/Plugins/Experimental/PCG/Source/PCG/Private/Elements/PCGHiGenGridSize.cpp b/Engine/Plugins/Experimental/PCG/Source/PCG/Private/Elements/PCGHiGenGridSize.cpp
--- a/Engine/Plugins/Experimental/PCG/Source/PCG/Private/Elements/PCGHiGenGridSize.cpp
+++ b/Engine/Plugins/Experimental/PCG/Source/PCG/Private/Elements/PCGHiGenGridSize.cpp
@@ -113,6 +113,8 @@ EPCGChangeType UPCGHiGenGridSizeSettings::GetChangeTypeForProperty(const FName&
 
 bool FPCGHiGenGridSizeElement::ExecuteInternal(FPCGContext* Context) const
 {
+	TRACE_CPUPROFILER_EVENT_SCOPE(FPCGHiGenGridSizeElement::Execute);
+
 	// Validation
 	check(Context);
 	if (Context->Node && Context->Node->GetGraph() && !Context->Node->GetGraph()->IsHierarchicalGenerationEnabled())
diff --git a/Engine/Plugins/Experimental/PCG/Source/PCG/Private/Elements/PCGUserParameterGet.cpp b/Engine/Plugins/Experimental/PCG/Source/PCG/Private/Elements/PCGUserParameterGet.cpp
--- a/Engine/Plugins/Experimental/PCG/Source/PCG/Private/Elements/PCGUserParameterGet.cpp
+++ b/Engine/Plugins/Experimental/PCG/Source/PCG/Private/Elements/PCGUserParameterGet.cpp
@@ -141,6 +141,8 @@ FPCGElementPtr UPCGUserParameterGetSettings::CreateElement() const
 
 bool FPCGUserParameterGetElement::ExecuteInternal(FPCGContext* Context) const
 {
+	TRACE_CPUPROFILER_EVENT_SCOPE(FPCGUserParameterGetElement::Execute);
+
 	check(Context);
 
 	const UPCGUserParameterGetSettings* Settings = Context->GetInputSettings<UPCGUserParameterGetSettings>();
diff --git a/Engine/Plugins/Experimental/PCG/Source/PCG/Public/Elements/PCGCollapseElement.h b/Engine/Plugins/Experimental/PCG/Source/PCG/Public/Elements/PCGCollapseElement.h
--- a/Engine/Plugins/Experimental/PCG/Source/PCG/Public/Elements/PCGCollapseElement.h
+++ b/Engine/Plugins/Experimental/PCG/Source/PCG/Public/Elements/PCGCollapseElement.h
@@ -54,5 +54,5 @@ public:
 	virtual bool ShouldComputeFullOutputDataCrc(FPCGContext* Context) const override { return true; }
 
 protected:
-	virtual bool ExecuteInternal(FPCGContext* Context) const;
+	virtual bool ExecuteInternal(FPCGContext* Context) const override;
 };
diff --git a/Engine/Plugins/Experimental/PCG/Source/PCG/Public/Elements/PCGConvertToAttributeSet.h b/Engine/Plugins/Experimental/PCG/Source/PCG/Public/Elements/PCGConvertToAttributeSet.h
--- a/Engine/Plugins/Experimental/PCG/Source/PCG/Public/Elements/PCGConvertToAttributeSet.h
+++ b/Engine/Plugins/Experimental/PCG/Source/PCG/Public/Elements/PCGConvertToAttributeSet.h
@@ -37,5 +37,5 @@ public:
 	virtual bool ShouldComputeFullOutputDataCrc(FPCGContext* Context) const override { return true; }
 
 protected:
-	virtual bool ExecuteInternal(FPCGContext* Context) const;
+	virtual bool ExecuteInternal(FPCGContext* Context) const override;
 };
\ No newline at end of file
diff --git a/Engine/Plugins/Experimental/PCG/Source/PCG/Public/Elements/PCGDensityRemap.h b/Engine/Plugins/Experimental/PCG/Source/PCG/Public/Elements/PCGDensityRemap.h
--- a/Engine/Plugins/Experimental/PCG/Source/PCG/Public/Elements/PCGDensityRemap.h
+++ b/Engine/Plugins/Experimental/PCG/Source/PCG/Public/Elements/PCGDensityRemap.h
@@ -41,7 +41,7 @@ public:
 class FPCGLinearDensityRemapElement : public IPCGElement
 {
 protected:
-	virtual bool ExecuteInternal(FPCGContext* Context) const;
+	virtual bool ExecuteInternal(FPCGContext* Context) const override;
 };
 
 #if UE_ENABLE_INCLUDE_ORDER_DEPRECATED_IN_5_2
diff --git a/Engine/Plugins/Experimental/PCG/Source/PCG/Public/Elements/PCGDensityRemapElement.h b/Engine/Plugins/Experimental/PCG/Source/PCG/Public/Elements/PCGDensityRemapElement.h
--- a/Engine/Plugins/Experimental/PCG/Source/PCG/Public/Elements/PCGDensityRemapElement.h
+++ b/Engine/Plugins/Experimental/PCG/Source/PCG/Public/Elements/PCGDensityRemapElement.h
@@ -50,7 +50,7 @@ public:
 class FPCGDensityRemapElement : public FPCGPointProcessingElementBase
 {
 protected:
-	virtual bool ExecuteInternal(FPCGContext* Context) const;
+	virtual bool ExecuteInternal(FPCGContext* Context) const override;
 };
 
 #if UE_ENABLE_INCLUDE_ORDER_DEPRECATED_IN_5_2
diff --git a/Engine/Plugins/Experimental/PCG/Source/PCG/Public/Elements/PCGPointExtentsModifier.h b/Engine/Plugins/Experimental/PCG/Source/PCG/Public/Elements/PCGPointExtentsModifier.h
--- a/Engine/Plugins/Experimental/PCG/Source/PCG/Public/Elements/PCGPointExtentsModifier.h
+++ b/Engine/Plugins/Experimental/PCG/Source/PCG/Public/Elements/PCGPointExtentsModifier.h
@@ -49,7 +49,7 @@ public:
 class FPCGPointExtentsModifier : public FPCGPointProcessingElementBase
 {
 protected:
-	virtual bool ExecuteInternal(FPCGContext* Context) const;
+	virtual bool ExecuteInternal(FPCGContext* Context) const override;
 };
 
 #if UE_ENABLE_INCLUDE_ORDER_DEPRECATED_IN_5_2
