Keep a copy of the volume body outside of the physics scene when static to sidetrack locking issues

#rnx

diff --git a/Engine/Plugins/Experimental/PCG/Source/PCG/Private/Data/PCGVolumeData.cpp b/Engine/Plugins/Experimental/PCG/Source/PCG/Private/Data/PCGVolumeData.cpp
--- a/Engine/Plugins/Experimental/PCG/Source/PCG/Private/Data/PCGVolumeData.cpp
+++ b/Engine/Plugins/Experimental/PCG/Source/PCG/Private/Data/PCGVolumeData.cpp
@@ -10,8 +10,30 @@
 #include "Engine/CollisionProfile.h"
 #include "GameFramework/Volume.h"
 
+#include "Chaos/ChaosEngineInterface.h"
+#include "Physics/PhysicsInterfaceDeclares.h"
+
 #include UE_INLINE_GENERATED_CPP_BY_NAME(PCGVolumeData)
 
+UPCGVolumeData::~UPCGVolumeData()
+{
+	ReleaseInternalBodyInstance();
+}
+
+void UPCGVolumeData::ReleaseInternalBodyInstance()
+{
+	if (VolumeBodyInstance)
+	{
+		if (VolumeBodyInstance->IsValidBodyInstance())
+		{
+			VolumeBodyInstance->TermBody();
+		}
+
+		delete VolumeBodyInstance;
+		VolumeBodyInstance = nullptr;
+	}
+}
+
 void UPCGVolumeData::Initialize(AVolume* InVolume)
 {
 	check(InVolume);
@@ -27,10 +49,26 @@ void UPCGVolumeData::Initialize(AVolume* InVolume)
 	}
 
 	FBoxSphereBounds BoxSphereBounds = Volume->GetBounds();
-	Bounds = FBox::BuildAABB(BoxSphereBounds.Origin, BoxSphereBounds.BoxExtent);
-
 	// TODO: Compute the strict bounds, we must find a FBox inscribed into the oriented box.
 	// Currently, we'll leave the strict bounds empty and fall back to checking against the local box
+	Bounds = FBox::BuildAABB(BoxSphereBounds.Origin, BoxSphereBounds.BoxExtent);
+	
+	// Keep a "sceneless" equivalent body so we can do queries against it without locking constraints
+	if (UBrushComponent* BrushComponent = Volume->GetBrushComponent())
+	{
+		FBodyInstance* BodyInstance = BrushComponent->GetBodyInstance();
+		UBodySetup* BodySetup = BrushComponent->GetBodySetup();
+
+		if (BodyInstance && BodySetup && !BodyInstance->IsDynamic())
+		{
+			ReleaseInternalBodyInstance();
+
+			VolumeBodyInstance = new FBodyInstance();
+			VolumeBodyInstance->bAutoWeld = false;
+			VolumeBodyInstance->bSimulatePhysics = false;
+			VolumeBodyInstance->InitBody(BodySetup, BrushComponent->GetComponentTransform(), nullptr, nullptr);
+		}
+	}
 }
 
 void UPCGVolumeData::Initialize(const FBox& InBounds)
@@ -92,6 +130,14 @@ bool UPCGVolumeData::SamplePoint(const FTransform& InTransform, const FBox& InBo
 		{
 			PointDensity = 1.0f;
 		}
+		else if (VolumeBodyInstance)
+		{
+			float OutDistanceSquared = -1.0f;
+			if (FPhysicsInterface::GetSquaredDistanceToBody(VolumeBodyInstance, InPosition, OutDistanceSquared))
+			{
+				PointDensity = (OutDistanceSquared == 0.0f ? 1.0f : 0.0f);
+			}
+		}
 		else
 		{
 			PointDensity = Volume->EncompassesPoint(InPosition) ? 1.0f : 0.0f;
diff --git a/Engine/Plugins/Experimental/PCG/Source/PCG/Public/Data/PCGVolumeData.h b/Engine/Plugins/Experimental/PCG/Source/PCG/Public/Data/PCGVolumeData.h
--- a/Engine/Plugins/Experimental/PCG/Source/PCG/Public/Data/PCGVolumeData.h
+++ b/Engine/Plugins/Experimental/PCG/Source/PCG/Public/Data/PCGVolumeData.h
@@ -7,6 +7,7 @@
 #include "PCGVolumeData.generated.h"
 
 class AVolume;
+struct FBodyInstance;
 
 UCLASS(BlueprintType, ClassGroup = (Procedural))
 class PCG_API UPCGVolumeData : public UPCGSpatialDataWithPointCache
@@ -14,6 +15,7 @@ class PCG_API UPCGVolumeData : public UPCGSpatialDataWithPointCache
 	GENERATED_BODY()
 
 public:
+	~UPCGVolumeData();
 	void Initialize(AVolume* InVolume);
 	void Initialize(const FBox& InBounds);
 
@@ -43,6 +45,7 @@ public:
 
 protected:
 	void CopyBaseVolumeData(UPCGVolumeData* NewVolumeData) const;
+	void ReleaseInternalBodyInstance();
 
 	UPROPERTY(BlueprintReadWrite, EditAnywhere, Category = SourceData)
 	TWeakObjectPtr<AVolume> Volume = nullptr;
@@ -52,6 +55,9 @@ protected:
 
 	UPROPERTY()
 	FBox StrictBounds = FBox(EForceInit::ForceInit);
+
+	// Internal body instance to perform queries faster, used in static cases only
+	FBodyInstance* VolumeBodyInstance = nullptr;
 };
 
 #if UE_ENABLE_INCLUDE_ORDER_DEPRECATED_IN_5_2
