Fixed issue where at runtime, the world keys in the landscape cache wouldn't be setup properly.
Also adjusted getting the bounds of a PCG partition actor when used through a world ray hit

#rnx

diff --git a/Engine/Plugins/Experimental/PCG/Source/PCG/Private/Helpers/PCGHelpers.cpp b/Engine/Plugins/Experimental/PCG/Source/PCG/Private/Helpers/PCGHelpers.cpp
--- a/Engine/Plugins/Experimental/PCG/Source/PCG/Private/Helpers/PCGHelpers.cpp
+++ b/Engine/Plugins/Experimental/PCG/Source/PCG/Private/Helpers/PCGHelpers.cpp
@@ -140,23 +140,32 @@ namespace PCGHelpers
 		// This is to ensure stable bounds and no timing issues (cleared ISMs, etc.)
 		FBox Box(EForceInit::ForceInit);
 
-		const bool bNonColliding = true;
-		const bool bIncludeFromChildActors = true;
-
 		if (InActor)
 		{
-			const FTransform& ActorToWorld = InActor->GetTransform();
-			const FTransform WorldToActor = ActorToWorld.Inverse();
-
-			InActor->ForEachComponent<UPrimitiveComponent>(bIncludeFromChildActors, [bNonColliding, bIgnorePCGCreatedComponents, &WorldToActor, &Box](const UPrimitiveComponent* InPrimComp)
+			if (const APCGPartitionActor* PartitionActor = Cast<const APCGPartitionActor>(InActor))
 			{
-				if ((bNonColliding || InPrimComp->IsCollisionEnabled()) &&
-					(!bIgnorePCGCreatedComponents || !InPrimComp->ComponentTags.Contains(DefaultPCGTag)))
+				// Skip per-component check, return fixed bounds only replaced on origin
+				Box = PartitionActor->GetFixedBounds();
+				Box = Box.MoveTo(FVector::ZeroVector);
+			}
+			else
+			{
+				const bool bNonColliding = true;
+				const bool bIncludeFromChildActors = true;
+
+				const FTransform& ActorToWorld = InActor->GetTransform();
+				const FTransform WorldToActor = ActorToWorld.Inverse();
+
+				InActor->ForEachComponent<UPrimitiveComponent>(bIncludeFromChildActors, [bNonColliding, bIgnorePCGCreatedComponents, &WorldToActor, &Box](const UPrimitiveComponent* InPrimComp)
 				{
-					const FTransform ComponentToActor = InPrimComp->GetComponentTransform() * WorldToActor;
-					Box += InPrimComp->CalcBounds(ComponentToActor).GetBox();
-				}
-			});
+					if ((bNonColliding || InPrimComp->IsCollisionEnabled()) &&
+						(!bIgnorePCGCreatedComponents || !InPrimComp->ComponentTags.Contains(DefaultPCGTag)))
+					{
+						const FTransform ComponentToActor = InPrimComp->GetComponentTransform() * WorldToActor;
+						Box += InPrimComp->CalcBounds(ComponentToActor).GetBox();
+					}
+				});
+			}
 		}
 		else
 		{
