Improved some parts of metadata to allow easier merging of point data & attribute sets by taking the simplest type automatically.
Small optimizations around settings multiple values, and more to come.

diff --git a/Engine/Plugins/Experimental/PCG/Source/PCG/Private/Elements/Metadata/PCGMetadataElementCommon.cpp b/Engine/Plugins/Experimental/PCG/Source/PCG/Private/Elements/Metadata/PCGMetadataElementCommon.cpp
--- a/Engine/Plugins/Experimental/PCG/Source/PCG/Private/Elements/Metadata/PCGMetadataElementCommon.cpp
+++ b/Engine/Plugins/Experimental/PCG/Source/PCG/Private/Elements/Metadata/PCGMetadataElementCommon.cpp
@@ -64,7 +64,6 @@ namespace PCGMetadataElementCommon
 		auto Operation = [&Params, Count](auto Dummy)
 		{
 			using OutputType = decltype(Dummy);
-			OutputType Value{};
 
 			auto SetToAccessor = [&Params](const TArrayView<OutputType>& View, const int32 Start, const int32 Range)
 			{
diff --git a/Engine/Plugins/Experimental/PCG/Source/PCG/Private/Metadata/Accessors/PCGAttributeAccessorHelpers.cpp b/Engine/Plugins/Experimental/PCG/Source/PCG/Private/Metadata/Accessors/PCGAttributeAccessorHelpers.cpp
--- a/Engine/Plugins/Experimental/PCG/Source/PCG/Private/Metadata/Accessors/PCGAttributeAccessorHelpers.cpp
+++ b/Engine/Plugins/Experimental/PCG/Source/PCG/Private/Metadata/Accessors/PCGAttributeAccessorHelpers.cpp
@@ -39,6 +39,37 @@ namespace PCGAttributeAccessorHelpers
 		}
 	}
 
+	template <
+		typename AccessorType,
+		typename = typename std::enable_if_t<
+		std::is_same_v<IPCGAttributeAccessor, std::remove_const_t<AccessorType>>
+		>
+	>
+	TUniquePtr<AccessorType> CreateAttributeAccessorImpl(FPCGMetadataAttributeBase* Attribute, UPCGMetadata* Metadata, bool bForceReadOnly = false)
+	{
+		auto CreateTypedAccessor = [Attribute, Metadata, bForceReadOnly](auto Dummy) -> TUniquePtr<IPCGAttributeAccessor>
+		{
+			using AttributeType = decltype(Dummy);
+			if constexpr (std::is_const_v<AccessorType>)
+			{
+				return MakeUnique<FPCGAttributeAccessor<AttributeType>>(static_cast<const FPCGMetadataAttribute<AttributeType>*>(Attribute), Metadata, bForceReadOnly);
+			}
+			else
+			{
+				return MakeUnique<FPCGAttributeAccessor<AttributeType>>(static_cast<FPCGMetadataAttribute<AttributeType>*>(Attribute), Metadata, bForceReadOnly);
+			}
+		};
+
+		if (Attribute && Metadata)
+		{
+			return PCGMetadataAttribute::CallbackWithRightType(Attribute->GetTypeId(), CreateTypedAccessor);
+		}
+		else
+		{
+			return TUniquePtr<AccessorType>();
+		}
+	}
+
 	// Don't be afraid of this enormous template!
 	// We want to be able to create a TUniquePtr<const IPCGAttributeAccessor> for const accessors and TUniquePtr<IPCGAttributeAccessor> for mutable accessors.
 	// This template just says that:
@@ -101,27 +132,8 @@ namespace PCGAttributeAccessorHelpers
 			// It is OK to const_cast here, since we will create a const accessor if the input is const.
 			ExtractMetadataAtribute(const_cast<UPCGData*>(InData), Name, Metadata, Attribute);
 
-			auto CreateTypedAccessor = [Attribute, Metadata](auto Dummy) -> TUniquePtr<IPCGAttributeAccessor>
-			{
-				using AttributeType = decltype(Dummy);
-				if constexpr (std::is_const_v<AccessorType>)
-				{
-					return MakeUnique<FPCGAttributeAccessor<AttributeType>>(static_cast<const FPCGMetadataAttribute<AttributeType>*>(Attribute), Metadata);
-				}
-				else
-				{
-					return MakeUnique<FPCGAttributeAccessor<AttributeType>>(static_cast<FPCGMetadataAttribute<AttributeType>*>(Attribute), Metadata);
-				}
-			};
-
-			if (Attribute && Metadata)
-			{
-				Accessor = PCGMetadataAttribute::CallbackWithRightType(Attribute->GetTypeId(), CreateTypedAccessor);
-			}
-			else
-			{
-				return TUniquePtr<AccessorType>();
-			}
+			// To simplify the code here & below we'll get a non-const accessor but force readonly if it should be
+			Accessor = CreateAttributeAccessorImpl<std::remove_const_t<AccessorType>>(Attribute, Metadata, /*bForceReadOnly=*/std::is_const_v<AccessorType>);
 		}
 
 		if (!Accessor.IsValid())
@@ -588,11 +600,21 @@ TUniquePtr<const IPCGAttributeAccessor> PCGAttributeAccessorHelpers::CreateConst
 	return CreateAccessorImpl<const IPCGAttributeAccessor, const UPCGData>(InData, InSelector);
 }
 
+TUniquePtr<const IPCGAttributeAccessor> PCGAttributeAccessorHelpers::CreateConstAccessor(const FPCGMetadataAttributeBase* InAttribute, const UPCGMetadata* InMetadata)
+{
+	return CreateAttributeAccessorImpl<const IPCGAttributeAccessor>(const_cast<FPCGMetadataAttributeBase*>(InAttribute), const_cast<UPCGMetadata*>(InMetadata));
+}
+
 TUniquePtr<IPCGAttributeAccessor> PCGAttributeAccessorHelpers::CreateAccessor(UPCGData* InData, const FPCGAttributePropertySelector& InSelector)
 {
 	return CreateAccessorImpl<IPCGAttributeAccessor, UPCGData>(InData, InSelector);
 }
 
+TUniquePtr<IPCGAttributeAccessor> PCGAttributeAccessorHelpers::CreateAccessor(FPCGMetadataAttributeBase* InAttribute, UPCGMetadata* InMetadata)
+{
+	return CreateAttributeAccessorImpl<IPCGAttributeAccessor>(InAttribute, InMetadata);
+}
+
 TUniquePtr<const IPCGAttributeAccessorKeys> PCGAttributeAccessorHelpers::CreateConstKeys(const UPCGData* InData, const FPCGAttributePropertySelector& InSelector)
 {
 	return CreateKeysImpl<const IPCGAttributeAccessorKeys, const UPCGData>(InData, InSelector);
diff --git a/Engine/Plugins/Experimental/PCG/Source/PCG/Private/Metadata/Accessors/PCGAttributeAccessorKeys.cpp b/Engine/Plugins/Experimental/PCG/Source/PCG/Private/Metadata/Accessors/PCGAttributeAccessorKeys.cpp
--- a/Engine/Plugins/Experimental/PCG/Source/PCG/Private/Metadata/Accessors/PCGAttributeAccessorKeys.cpp
+++ b/Engine/Plugins/Experimental/PCG/Source/PCG/Private/Metadata/Accessors/PCGAttributeAccessorKeys.cpp
@@ -16,33 +16,51 @@ FPCGAttributeAccessorKeysEntries::FPCGAttributeAccessorKeysEntries(const FPCGMet
 	{
 		TArray<PCGMetadataEntryKey> Temp;
 		Current->GetEntryToValueKeyMap_NotThreadSafe().GenerateKeyArray(Temp);
-		Entries.Append(Temp);
+		ExtractedEntries.Append(Temp);
 		Current = Current->GetParent();
 	}
 
 	// If the attribute doesn't have any entry, re-try with metadata entries.
-	if (Attribute && Entries.IsEmpty())
+	if (Attribute && ExtractedEntries.IsEmpty())
 	{
 		InitializeFromMetadata(Attribute->GetMetadata());
 	}
 
 	// If the attribute still doesn't have any entry, we will always take the default value.
-	if (Entries.IsEmpty())
+	if (ExtractedEntries.IsEmpty())
 	{
-		Entries.Add(PCGInvalidEntryKey);
+		ExtractedEntries.Add(PCGInvalidEntryKey);
 	}
+
+	Entries = TArrayView<PCGMetadataEntryKey>(ExtractedEntries);
 }
 
 FPCGAttributeAccessorKeysEntries::FPCGAttributeAccessorKeysEntries(PCGMetadataEntryKey EntryKey)
 	: IPCGAttributeAccessorKeys(/*bInReadOnly=*/ false)
 {
-	Entries.Add(EntryKey);
+	ExtractedEntries.Add(EntryKey);
+	Entries = TArrayView<PCGMetadataEntryKey>(ExtractedEntries);
 }
 
 FPCGAttributeAccessorKeysEntries::FPCGAttributeAccessorKeysEntries(const UPCGMetadata* Metadata)
 	: IPCGAttributeAccessorKeys(/*bInReadOnly=*/ true)
 {
 	InitializeFromMetadata(Metadata);
+	Entries = TArrayView<PCGMetadataEntryKey>(ExtractedEntries);
+}
+
+FPCGAttributeAccessorKeysEntries::FPCGAttributeAccessorKeysEntries(const TArrayView<PCGMetadataEntryKey>& InEntries)
+	: IPCGAttributeAccessorKeys(/*bInReadOnly=*/false)
+	, Entries(InEntries)
+{
+
+}
+
+FPCGAttributeAccessorKeysEntries::FPCGAttributeAccessorKeysEntries(const TArrayView<const PCGMetadataEntryKey>& InEntries)
+	: IPCGAttributeAccessorKeys(/*bInReadOnly=*/true)
+	, Entries(const_cast<PCGMetadataEntryKey*>(InEntries.GetData()), InEntries.Num())
+{
+
 }
 
 void FPCGAttributeAccessorKeysEntries::InitializeFromMetadata(const UPCGMetadata* Metadata)
@@ -52,7 +70,7 @@ void FPCGAttributeAccessorKeysEntries::InitializeFromMetadata(const UPCGMetadata
 		return;
 	}
 
-	check(Entries.IsEmpty());
+	check(ExtractedEntries.IsEmpty());
 
 	const PCGMetadataEntryKey ItemKeyLowerBound = Metadata->GetItemKeyCountForParent();
 	const PCGMetadataEntryKey ItemKeyUpperBound = Metadata->GetItemCountForChild();
@@ -60,11 +78,11 @@ void FPCGAttributeAccessorKeysEntries::InitializeFromMetadata(const UPCGMetadata
 
 	if (Count > 0)
 	{
-		Entries.Reserve(Count);
+		ExtractedEntries.Reserve(Count);
 
 		for (PCGMetadataEntryKey Entry = ItemKeyLowerBound; Entry < ItemKeyUpperBound; ++Entry)
 		{
-			Entries.Add(Entry);
+			ExtractedEntries.Add(Entry);
 		}
 	}
 }
diff --git a/Engine/Plugins/Experimental/PCG/Source/PCG/Private/Metadata/PCGMetadata.cpp b/Engine/Plugins/Experimental/PCG/Source/PCG/Private/Metadata/PCGMetadata.cpp
--- a/Engine/Plugins/Experimental/PCG/Source/PCG/Private/Metadata/PCGMetadata.cpp
+++ b/Engine/Plugins/Experimental/PCG/Source/PCG/Private/Metadata/PCGMetadata.cpp
@@ -4,9 +4,12 @@
 
 #include "PCGData.h"
 #include "PCGPoint.h"
+#include "Elements/Metadata/PCGMetadataElementCommon.h"
 #include "Helpers/PCGPropertyHelpers.h"
 #include "Metadata/PCGAttributePropertySelector.h"
 
+#include "Algo/Transform.h"
+
 #include UE_INLINE_GENERATED_CPP_BY_NAME(PCGMetadata)
 
 void UPCGMetadata::Serialize(FArchive& InArchive)
@@ -181,17 +184,29 @@ void UPCGMetadata::AddAttributesFiltered(const UPCGMetadata* InOther, const TSet
 		const bool bSkipAttributesInFilterList = InFilterMode == EPCGMetadataFilterMode::ExcludeAttributes;
 		const bool bSkipThisAttribute = bSkipAttributesInFilterList == bAttributeInFilterList;
 
-		if (bSkipThisAttribute || HasAttribute(OtherAttribute.Key))
+		if(bSkipThisAttribute || !OtherAttribute.Value)
 		{
 			continue;
 		}
-		else
+		else if (HasAttribute(OtherAttribute.Key))
 		{
-			if (CopyAttribute(OtherAttribute.Value, OtherAttribute.Key, /*bKeepParent=*/InOther == Parent, /*bCopyEntries=*/false, /*bCopyValues=*/false))
+			// If both the current attribute and the other attribute have the same type - nothing to do
+			// If the current attribute can be broadcasted to the other but not the other way around - change the type
+			// If none of this is true - do nothing
+			const FPCGMetadataAttributeBase* Attribute = GetConstAttribute(OtherAttribute.Key);
+			check(Attribute);
+
+			if(Attribute->GetTypeId() != OtherAttribute.Value->GetTypeId() && 
+				!PCG::Private::IsBroadcastable(OtherAttribute.Value->GetTypeId(), Attribute->GetTypeId()) &&
+				PCG::Private::IsBroadcastable(Attribute->GetTypeId(), OtherAttribute.Value->GetTypeId()))
 			{
-				bAttributeAdded = true;
+				ChangeAttributeType(OtherAttribute.Key, OtherAttribute.Value->GetTypeId());
 			}
 		}
+		else if (CopyAttribute(OtherAttribute.Value, OtherAttribute.Key, /*bKeepParent=*/InOther == Parent, /*bCopyEntries=*/false, /*bCopyValues=*/false))
+		{
+			bAttributeAdded = true;
+		}
 	}
 
 	if (InOther != Parent && bAttributeAdded)
@@ -834,6 +849,38 @@ void UPCGMetadata::DeleteAttribute(FName AttributeToDelete)
 	}
 }
 
+bool UPCGMetadata::ChangeAttributeType(FName AttributeName, int16 AttributeNewType)
+{
+	FPCGMetadataAttributeBase* Attribute = GetMutableAttribute(AttributeName);
+
+	if (!Attribute)
+	{
+		UE_LOG(LogPCG, Error, TEXT("Attribute '%s' does not exist and therefore cannot change its type"), *AttributeName.ToString());
+		return false;
+	}
+
+	if (Attribute->GetTypeId() == AttributeNewType)
+	{
+		// Nothing to do, attribute is already the type we want
+		return true;
+	}
+
+	if (FPCGMetadataAttributeBase* NewAttribute = Attribute->CopyToAnotherType(AttributeNewType))
+	{
+		NewAttribute->AttributeId = Attribute->AttributeId;
+
+		AttributeLock.WriteLock();
+		RemoveAttributeInternal(AttributeName);
+		AddAttributeInternal(AttributeName, NewAttribute);
+		AttributeLock.WriteUnlock();
+
+		delete Attribute;
+		Attribute = nullptr;
+	}
+
+	return true;
+}
+
 int64 UPCGMetadata::GetItemCountForChild() const
 {
 	FReadScopeLock ScopeLock(ItemLock);
@@ -1133,49 +1180,32 @@ void UPCGMetadata::SetAttributes(PCGMetadataEntryKey InKey, const UPCGMetadata*
 
 void UPCGMetadata::SetPointAttributes(const TArrayView<const FPCGPoint>& InPoints, const UPCGMetadata* InMetadata, const TArrayView<FPCGPoint>& OutPoints)
 {
-	if (!InMetadata || InMetadata->GetAttributeCount() == 0)
+	if (!InMetadata || InMetadata->GetAttributeCount() == 0 || GetAttributeCount() == 0)
 	{
 		return;
 	}
 
 	check(InPoints.Num() == OutPoints.Num());
 
-	for (int32 PointIndex = 0; PointIndex < InPoints.Num(); ++PointIndex)
-	{
-		const FPCGPoint& InPoint = InPoints[PointIndex];
-		FPCGPoint& OutPoint = OutPoints[PointIndex];
-		InitializeOnSet(OutPoint.MetadataEntry, InPoint.MetadataEntry, InMetadata);
-	}
-
-	AttributeLock.ReadLock();
-	for(const TPair<FName, FPCGMetadataAttributeBase*>& AttributePair : Attributes)
-	{
-		const FName& AttributeName = AttributePair.Key;
-		FPCGMetadataAttributeBase* Attribute = AttributePair.Value;
+	// Extract the metadata entry keys from the in & out points
+	TArray<PCGMetadataEntryKey> InKeys;
+	TArray<PCGMetadataEntryKey> OutKeys;
 
-		if (const FPCGMetadataAttributeBase* OtherAttribute = InMetadata->GetConstAttribute(AttributeName))
-		{
-			if (OtherAttribute->GetTypeId() != Attribute->GetTypeId())
-			{
-				UE_LOG(LogPCG, Error, TEXT("Metadata type mismatch with attribute %s"), *AttributeName.ToString());
-				continue;
-			}
+	Algo::Transform(InPoints, InKeys, [](const FPCGPoint& Point) { return Point.MetadataEntry; });
+	Algo::Transform(OutPoints, OutKeys, [](const FPCGPoint& Point) { return Point.MetadataEntry; });
 
-			for (int32 PointIndex = 0; PointIndex < InPoints.Num(); ++PointIndex)
-			{
-				const FPCGPoint& InPoint = InPoints[PointIndex];
-				FPCGPoint& OutPoint = OutPoints[PointIndex];
+	SetAttributes(InKeys, InMetadata, OutKeys);
 
-				Attribute->SetValue(OutPoint.MetadataEntry, OtherAttribute, InPoint.MetadataEntry);
-			}
-		}
+	// Write back the keys on the points
+	for (int KeyIndex = 0; KeyIndex < OutKeys.Num(); ++KeyIndex)
+	{
+		OutPoints[KeyIndex].MetadataEntry = OutKeys[KeyIndex];
 	}
-	AttributeLock.ReadUnlock();
 }
 
 void UPCGMetadata::SetAttributes(const TArrayView<PCGMetadataEntryKey>& InKeys, const UPCGMetadata* InMetadata, const TArrayView<PCGMetadataEntryKey>& OutKeys)
 {
-	if (!InMetadata)
+	if (!InMetadata || InMetadata->GetAttributeCount() == 0 || GetAttributeCount() == 0)
 	{
 		return;
 	}
@@ -1187,6 +1217,9 @@ void UPCGMetadata::SetAttributes(const TArrayView<PCGMetadataEntryKey>& InKeys,
 		InitializeOnSet(OutKeys[KeyIndex], InKeys[KeyIndex], InMetadata);
 	}
 
+	// Rarely-used convenience array that might be used multiple times, will be kept out of the loop to limit the number of allocations if needed
+	TArray<PCGMetadataValueKey> ValueKeys;
+
 	AttributeLock.ReadLock();
 	for(const TPair<FName, FPCGMetadataAttributeBase*>& AttributePair : Attributes)
 	{
@@ -1195,15 +1228,43 @@ void UPCGMetadata::SetAttributes(const TArrayView<PCGMetadataEntryKey>& InKeys,
 
 		if (const FPCGMetadataAttributeBase* OtherAttribute = InMetadata->GetConstAttribute(AttributeName))
 		{
-			if (OtherAttribute->GetTypeId() != Attribute->GetTypeId())
+			if (!PCG::Private::IsBroadcastableOrConstructible(OtherAttribute->GetTypeId(), Attribute->GetTypeId()))
 			{
-				UE_LOG(LogPCG, Error, TEXT("Metadata type mismatch with attribute %s"), *AttributeName.ToString());
+				UE_LOG(LogPCG, Error, TEXT("Metadata type mismatch with attribute '%s'"), *AttributeName.ToString());
 				continue;
 			}
 
-			for (int32 KeyIndex = 0; KeyIndex < InKeys.Num(); ++KeyIndex)
+			if (Attribute == OtherAttribute)
 			{
-				Attribute->SetValue(OutKeys[KeyIndex], OtherAttribute, InKeys[KeyIndex]);
+				ValueKeys.Reset();
+				Attribute->GetValueKeys(InKeys, ValueKeys);
+				Attribute->SetValuesFromValueKeys(OutKeys, ValueKeys);
+			}
+			else
+			{
+				// Create accessor for the other attribute
+				TUniquePtr<const IPCGAttributeAccessor> OtherAttributeAccessor = PCGAttributeAccessorHelpers::CreateConstAccessor(OtherAttribute, InMetadata);
+				FPCGAttributeAccessorKeysEntries OtherAttributeKeys(InKeys);
+
+				if (!OtherAttributeAccessor)
+				{
+					continue;
+				}
+
+				auto GetAndSetValues = [Attribute, &OutKeys, &OtherAttributeAccessor, &OtherAttributeKeys](auto Dummy) -> bool
+				{
+					using Type = decltype(Dummy);
+
+					auto SetValues = [Attribute, &OutKeys](const TArrayView<Type>& View, const int32 Start, const int32 Range)
+					{
+						TArrayView<PCGMetadataEntryKey> Keys(OutKeys.GetData() + Start, Range);
+						static_cast<FPCGMetadataAttribute<Type>*>(Attribute)->SetValues(Keys, View);
+					};
+
+					return PCGMetadataElementCommon::ApplyOnAccessorRange<Type>(OtherAttributeKeys, *OtherAttributeAccessor, SetValues, EPCGAttributeAccessorFlags::AllowBroadcast | EPCGAttributeAccessorFlags::AllowConstructible);
+				};
+
+				PCGMetadataAttribute::CallbackWithRightType(Attribute->GetTypeId(), GetAndSetValues);
 			}
 		}
 	}
diff --git a/Engine/Plugins/Experimental/PCG/Source/PCG/Private/Metadata/PCGMetadataAttribute.cpp b/Engine/Plugins/Experimental/PCG/Source/PCG/Private/Metadata/PCGMetadataAttribute.cpp
--- a/Engine/Plugins/Experimental/PCG/Source/PCG/Private/Metadata/PCGMetadataAttribute.cpp
+++ b/Engine/Plugins/Experimental/PCG/Source/PCG/Private/Metadata/PCGMetadataAttribute.cpp
@@ -34,22 +34,23 @@ void FPCGMetadataAttributeBase::Serialize(UPCGMetadata* InMetadata, FArchive& In
 	InArchive << AttributeId;
 }
 
-void FPCGMetadataAttributeBase::SetValueFromValueKey(PCGMetadataEntryKey EntryKey, PCGMetadataValueKey ValueKey)
+void FPCGMetadataAttributeBase::SetValueFromValueKey(PCGMetadataEntryKey EntryKey, PCGMetadataValueKey ValueKey, bool bResetValueOnDefaultValueKey)
 {
-	check(EntryKey != PCGInvalidEntryKey);
 	FWriteScopeLock ScopeLock(EntryMapLock);
-	EntryToValueKeyMap.FindOrAdd(EntryKey) = ValueKey;
+	SetValueFromValueKey_Unsafe(EntryKey, ValueKey, bResetValueOnDefaultValueKey);
 }
 
-void FPCGMetadataAttributeBase::SetValueFromValueKey_Unsafe(PCGMetadataEntryKey EntryKey, PCGMetadataValueKey ValueKey, bool bResetValueOnDefaultValueKey)
+void FPCGMetadataAttributeBase::SetValueFromValueKey_Unsafe(PCGMetadataEntryKey EntryKey, PCGMetadataValueKey ValueKey, bool bResetValueOnDefaultValueKey, bool bAllowInvalidEntries)
 {
-	check(EntryKey != PCGInvalidEntryKey);
-	if (EntryKey == PCGDefaultValueKey)
+	if (EntryKey == PCGInvalidEntryKey)
 	{
-		if (bResetValueOnDefaultValueKey)
-		{
-			EntryToValueKeyMap.Remove(EntryKey);
-		}
+		check(bAllowInvalidEntries);
+		return;
+	}
+
+	if (ValueKey == PCGDefaultValueKey && bResetValueOnDefaultValueKey)
+	{
+		EntryToValueKeyMap.Remove(EntryKey);
 	}
 	else
 	{
@@ -57,7 +58,7 @@ void FPCGMetadataAttributeBase::SetValueFromValueKey_Unsafe(PCGMetadataEntryKey
 	}
 }
 
-void FPCGMetadataAttributeBase::SetValuesFromValueKeys(const TArray<TTuple<PCGMetadataEntryKey, PCGMetadataValueKey>>& EntryValuePairs, bool bResetValueOnDefaultValueKey)
+void FPCGMetadataAttributeBase::SetValuesFromValueKeys(const TArrayView<const TTuple<PCGMetadataEntryKey, PCGMetadataValueKey>>& EntryValuePairs, bool bResetValueOnDefaultValueKey)
 {
 	if (EntryValuePairs.IsEmpty())
 	{
@@ -67,11 +68,25 @@ void FPCGMetadataAttributeBase::SetValuesFromValueKeys(const TArray<TTuple<PCGMe
 	FWriteScopeLock ScopeLock(EntryMapLock);
 	for (const TTuple<PCGMetadataEntryKey, PCGMetadataValueKey>& EntryValuePair : EntryValuePairs)
 	{
-		SetValueFromValueKey_Unsafe(EntryValuePair.Key, EntryValuePair.Value, bResetValueOnDefaultValueKey);
+		SetValueFromValueKey_Unsafe(EntryValuePair.Key, EntryValuePair.Value, bResetValueOnDefaultValueKey, /*bAllowInvalidEntries=*/true);
+	}
+}
+
+void FPCGMetadataAttributeBase::SetValuesFromValueKeys(const TArrayView<const PCGMetadataEntryKey>& EntryKeys, const TArrayView<const PCGMetadataValueKey>& ValueKeys, bool bResetValueOnDefaultValueKey)
+{
+	if (EntryKeys.IsEmpty() || EntryKeys.Num() != ValueKeys.Num())
+	{
+		return;
+	}
+
+	FWriteScopeLock ScopeLock(EntryMapLock);
+	for (int32 i = 0; i < EntryKeys.Num(); ++i)
+	{
+		SetValueFromValueKey_Unsafe(EntryKeys[i], ValueKeys[i], bResetValueOnDefaultValueKey, /*bAllowInvalidEntries=*/true);
 	}
 }
 
-void FPCGMetadataAttributeBase::SetValuesFromValueKeys(const TArray<PCGMetadataEntryKey>& EntryKeys, TArray<PCGMetadataValueKey>& ValueKeys, bool bResetValueOnDefaultValueKey)
+void FPCGMetadataAttributeBase::SetValuesFromValueKeys(const TArrayView<const PCGMetadataEntryKey * const>& EntryKeys, const TArrayView<const PCGMetadataValueKey>& ValueKeys, bool bResetValueOnDefaultValueKey)
 {
 	if (EntryKeys.IsEmpty() || EntryKeys.Num() != ValueKeys.Num())
 	{
@@ -81,7 +96,7 @@ void FPCGMetadataAttributeBase::SetValuesFromValueKeys(const TArray<PCGMetadataE
 	FWriteScopeLock ScopeLock(EntryMapLock);
 	for (int32 i = 0; i < EntryKeys.Num(); ++i)
 	{
-		SetValueFromValueKey_Unsafe(EntryKeys[i], ValueKeys[i], bResetValueOnDefaultValueKey);
+		SetValueFromValueKey_Unsafe(*EntryKeys[i], ValueKeys[i], bResetValueOnDefaultValueKey, /*bAllowInvalidEntries=*/true);
 	}
 }
 
@@ -113,7 +128,7 @@ PCGMetadataValueKey FPCGMetadataAttributeBase::GetValueKey(PCGMetadataEntryKey E
 	}
 }
 
-void FPCGMetadataAttributeBase::GetValueKeys(const TArray<PCGMetadataEntryKey>& EntryKeys, TArray<PCGMetadataValueKey>& OutValueKeys) const
+void FPCGMetadataAttributeBase::GetValueKeys(const TArrayView<const PCGMetadataEntryKey>& EntryKeys, TArray<PCGMetadataValueKey>& OutValueKeys) const
 {
 	if (EntryKeys.IsEmpty())
 	{
@@ -127,7 +142,7 @@ void FPCGMetadataAttributeBase::GetValueKeys(const TArray<PCGMetadataEntryKey>&
 	GetValueKeys_Internal(EntryKeys, OutValueKeys, UnsetValues);
 }
 
-void FPCGMetadataAttributeBase::GetValueKeys_Internal(const TArray<PCGMetadataEntryKey>& EntryKeys, TArray<PCGMetadataValueKey>& OutValueKeys, TBitArray<>& UnsetValues) const
+void FPCGMetadataAttributeBase::GetValueKeys_Internal(const TArrayView<const PCGMetadataEntryKey>& EntryKeys, TArrayView<PCGMetadataValueKey> OutValueKeys, TBitArray<>& UnsetValues) const
 {
 	check(EntryKeys.Num() == OutValueKeys.Num() && OutValueKeys.Num() == UnsetValues.Num());
 
diff --git a/Engine/Plugins/Experimental/PCG/Source/PCG/Public/Elements/Metadata/PCGMetadataElementCommon.h b/Engine/Plugins/Experimental/PCG/Source/PCG/Public/Elements/Metadata/PCGMetadataElementCommon.h
--- a/Engine/Plugins/Experimental/PCG/Source/PCG/Public/Elements/Metadata/PCGMetadataElementCommon.h
+++ b/Engine/Plugins/Experimental/PCG/Source/PCG/Public/Elements/Metadata/PCGMetadataElementCommon.h
@@ -62,7 +62,14 @@ namespace PCGMetadataElementCommon
 		}
 
 		TArray<T, TInlineAllocator<DefaultChunkSize>> TempValues;
-		TempValues.SetNum(ChunkSize);
+		if constexpr (std::is_trivially_copyable_v<T>)
+		{
+			TempValues.SetNumUninitialized(ChunkSize);
+		}
+		else
+		{
+			TempValues.SetNum(ChunkSize);
+		}
 
 		const int32 NumberOfIterations = (NumberOfEntries + ChunkSize - 1) / ChunkSize;
 
diff --git a/Engine/Plugins/Experimental/PCG/Source/PCG/Public/Metadata/Accessors/PCGAttributeAccessor.h b/Engine/Plugins/Experimental/PCG/Source/PCG/Public/Metadata/Accessors/PCGAttributeAccessor.h
--- a/Engine/Plugins/Experimental/PCG/Source/PCG/Public/Metadata/Accessors/PCGAttributeAccessor.h
+++ b/Engine/Plugins/Experimental/PCG/Source/PCG/Public/Metadata/Accessors/PCGAttributeAccessor.h
@@ -21,15 +21,15 @@ public:
 	using Super = IPCGAttributeAccessorT<FPCGAttributeAccessor<T>>;
 
 	// Can't write if metadata is null
-	FPCGAttributeAccessor(FPCGMetadataAttribute<T>* InAttribute, UPCGMetadata* InMetadata)
-		: Super(/*bInReadOnly=*/ InMetadata == nullptr)
+	FPCGAttributeAccessor(FPCGMetadataAttribute<T>* InAttribute, UPCGMetadata* InMetadata, bool bForceReadOnly = false)
+		: Super(/*bInReadOnly=*/ InMetadata == nullptr || bForceReadOnly)
 		, Attribute(InAttribute)
 		, Metadata(InMetadata)
 	{
 		check(InAttribute);
 	}
 
-	FPCGAttributeAccessor(const FPCGMetadataAttribute<T>* InAttribute, const UPCGMetadata* InMetadata)
+	FPCGAttributeAccessor(const FPCGMetadataAttribute<T>* InAttribute, const UPCGMetadata* InMetadata, bool bForceReadOnly = false)
 		: Super(/*bInReadOnly=*/ true)
 		, Attribute(const_cast<FPCGMetadataAttribute<T>*>(InAttribute))
 		, Metadata(const_cast<UPCGMetadata*>(InMetadata))
@@ -39,8 +39,9 @@ public:
 
 	bool GetRangeImpl(TArrayView<T> OutValues, int32 Index, const IPCGAttributeAccessorKeys& Keys) const
 	{
-		TArray<const PCGMetadataEntryKey*> EntryKeys;
-		EntryKeys.SetNum(OutValues.Num());
+		TArray<const PCGMetadataEntryKey*, TInlineAllocator<256>> EntryKeys;
+		EntryKeys.SetNumUninitialized(OutValues.Num());
+
 		TArrayView<const PCGMetadataEntryKey*> EntryKeysView(EntryKeys);
 		if (!Keys.GetKeys<PCGMetadataEntryKey>(Index, EntryKeysView))
 		{
@@ -58,35 +59,39 @@ public:
 
 	bool SetRangeImpl(TArrayView<const T> InValues, int32 Index, IPCGAttributeAccessorKeys& Keys, EPCGAttributeAccessorFlags Flags)
 	{
-		TArray<PCGMetadataEntryKey*> EntryKeys;
-		EntryKeys.SetNum(InValues.Num());
+		TArray<PCGMetadataEntryKey*, TInlineAllocator<256>> EntryKeys;
+		EntryKeys.SetNumUninitialized(InValues.Num());
 		TArrayView<PCGMetadataEntryKey*> EntryKeysView(EntryKeys);
 		if (!Keys.GetKeys<PCGMetadataEntryKey>(Index, EntryKeysView))
 		{
 			return false;
 		}
 
-		// TODO: Same than above (avoid locking too many times), but perhaps will be a bit more complex, because of
-		// the added logic with PCGInvalidEntryKey, AddEntry and SetDefaultValue.
-		for (int32 i = 0; i < InValues.Num(); ++i)
+		int LastDefaultKeyIndex = INDEX_NONE;
+
+		for(int EntryIndex = 0; EntryIndex < EntryKeys.Num(); ++EntryIndex)
 		{
-			PCGMetadataEntryKey& EntryKey = *EntryKeys[i];
-			if (EntryKey != PCGInvalidEntryKey || !(Flags & EPCGAttributeAccessorFlags::AllowSetDefaultValue))
+			PCGMetadataEntryKey& EntryKey = *EntryKeys[EntryIndex];
+			if (EntryKey == PCGInvalidEntryKey)
 			{
-				// TODO: This part seems costly for a lot of entries. Maybe there are some optimizations to do.
-				if (EntryKey == PCGInvalidEntryKey)
+				if (!(Flags & EPCGAttributeAccessorFlags::AllowSetDefaultValue))
 				{
-					EntryKey = Metadata->AddEntry();
+					EntryKey = Metadata->AddEntry(); // TODO - replace by AddEntryPlaceholder ?
+				}
+				else
+				{
+					LastDefaultKeyIndex = EntryIndex;
 				}
-
-				Attribute->SetValue(EntryKey, InValues[i]);
-			}
-			else
-			{
-				Attribute->SetDefaultValue(InValues[i]);
 			}
 		}
 
+		Attribute->SetValues(EntryKeys, InValues);
+
+		if (LastDefaultKeyIndex != INDEX_NONE)
+		{
+			Attribute->SetDefaultValue(InValues[LastDefaultKeyIndex]);
+		}
+
 		return true;
 	}
 
diff --git a/Engine/Plugins/Experimental/PCG/Source/PCG/Public/Metadata/Accessors/PCGAttributeAccessorHelpers.h b/Engine/Plugins/Experimental/PCG/Source/PCG/Public/Metadata/Accessors/PCGAttributeAccessorHelpers.h
--- a/Engine/Plugins/Experimental/PCG/Source/PCG/Public/Metadata/Accessors/PCGAttributeAccessorHelpers.h
+++ b/Engine/Plugins/Experimental/PCG/Source/PCG/Public/Metadata/Accessors/PCGAttributeAccessorHelpers.h
@@ -64,6 +64,11 @@ namespace PCGAttributeAccessorHelpers
 	*/
 	PCG_API TUniquePtr<const IPCGAttributeAccessor> CreateConstAccessor(const UPCGData* InData, const FPCGAttributePropertySelector& InSelector);
 
+	/** 
+	* Creates a const accessor to an attribute without requiring a selector.
+	*/
+	PCG_API TUniquePtr<const IPCGAttributeAccessor> CreateConstAccessor(const FPCGMetadataAttributeBase* InAttribute, const UPCGMetadata* InMetadata);
+
 	/**
 	* Creates a accessor to the property or attribute pointed at by the InSelector.
 	* Note that InData must not be null if the selector points to an attribute,
@@ -72,6 +77,11 @@ namespace PCGAttributeAccessorHelpers
 	*/
 	PCG_API TUniquePtr<IPCGAttributeAccessor> CreateAccessor(UPCGData* InData, const FPCGAttributePropertySelector& InSelector);
 
+	/**
+	* Creates an accessor to an attribute without requiring a selector.
+	*/
+	PCG_API TUniquePtr<IPCGAttributeAccessor> CreateAccessor(FPCGMetadataAttributeBase* InAttribute, UPCGMetadata* InMetadata);
+
 	PCG_API TUniquePtr<const IPCGAttributeAccessorKeys> CreateConstKeys(const UPCGData* InData, const FPCGAttributePropertySelector& InSelector);
 	PCG_API TUniquePtr<IPCGAttributeAccessorKeys> CreateKeys(UPCGData* InData, const FPCGAttributePropertySelector& InSelector);
 
diff --git a/Engine/Plugins/Experimental/PCG/Source/PCG/Public/Metadata/Accessors/PCGAttributeAccessorKeys.h b/Engine/Plugins/Experimental/PCG/Source/PCG/Public/Metadata/Accessors/PCGAttributeAccessorKeys.h
--- a/Engine/Plugins/Experimental/PCG/Source/PCG/Public/Metadata/Accessors/PCGAttributeAccessorKeys.h
+++ b/Engine/Plugins/Experimental/PCG/Source/PCG/Public/Metadata/Accessors/PCGAttributeAccessorKeys.h
@@ -140,6 +140,8 @@ class PCG_API FPCGAttributeAccessorKeysEntries : public IPCGAttributeAccessorKey
 public:
 	explicit FPCGAttributeAccessorKeysEntries(const FPCGMetadataAttributeBase* Attribute);
 	explicit FPCGAttributeAccessorKeysEntries(PCGMetadataEntryKey EntryKey);
+	explicit FPCGAttributeAccessorKeysEntries(const TArrayView<PCGMetadataEntryKey>& InEntries);
+	explicit FPCGAttributeAccessorKeysEntries(const TArrayView<const PCGMetadataEntryKey>& InEntries);
 
 	// Read-only case where we will just iterate on all the entries in the metadata.
 	explicit FPCGAttributeAccessorKeysEntries(const UPCGMetadata* Metadata);
@@ -152,7 +154,8 @@ protected:
 
 	void InitializeFromMetadata(const UPCGMetadata* Metadata);
 
-	TArray<PCGMetadataEntryKey> Entries;
+	TArrayView<PCGMetadataEntryKey> Entries;
+	TArray<PCGMetadataEntryKey> ExtractedEntries;
 };
 
 ///////////////////////////////////////////////////////////////////////
diff --git a/Engine/Plugins/Experimental/PCG/Source/PCG/Public/Metadata/PCGMetadata.h b/Engine/Plugins/Experimental/PCG/Source/PCG/Public/Metadata/PCGMetadata.h
--- a/Engine/Plugins/Experimental/PCG/Source/PCG/Public/Metadata/PCGMetadata.h
+++ b/Engine/Plugins/Experimental/PCG/Source/PCG/Public/Metadata/PCGMetadata.h
@@ -220,6 +220,9 @@ public:
 	UFUNCTION(BlueprintCallable, Category = "PCG|Metadata")
 	void ClearAttribute(FName AttributeToClear);
 
+	/** Change type of an attribute */
+	bool ChangeAttributeType(FName AttributeName, int16 AttributeNewType);
+
 	/** Adds a unique entry key to the metadata */
 	UFUNCTION(BlueprintCallable, Category = "PCG|Metadata")
 	int64 AddEntry(int64 ParentEntryKey = -1);
diff --git a/Engine/Plugins/Experimental/PCG/Source/PCG/Public/Metadata/PCGMetadataAttribute.h b/Engine/Plugins/Experimental/PCG/Source/PCG/Public/Metadata/PCGMetadataAttribute.h
--- a/Engine/Plugins/Experimental/PCG/Source/PCG/Public/Metadata/PCGMetadataAttribute.h
+++ b/Engine/Plugins/Experimental/PCG/Source/PCG/Public/Metadata/PCGMetadataAttribute.h
@@ -37,12 +37,13 @@ public:
 	int16 GetTypeId() const { return TypeId; }
 
 	virtual FPCGMetadataAttributeBase* Copy(FName NewName, UPCGMetadata* InMetadata, bool bKeepParent, bool bCopyEntries = true, bool bCopyValues = true) const = 0;
+	virtual FPCGMetadataAttributeBase* CopyToAnotherType(int16 Type) const = 0;
 
 	virtual PCGMetadataValueKey GetValueKeyOffsetForChild() const = 0;
 	virtual void SetValue(PCGMetadataEntryKey ItemKey, const FPCGMetadataAttributeBase* InAttribute, PCGMetadataEntryKey InEntryKey) = 0;
 	virtual void SetZeroValue(PCGMetadataEntryKey ItemKey) = 0;
 	virtual void AccumulateValue(PCGMetadataEntryKey ItemKey, const FPCGMetadataAttributeBase* InAttributeA, PCGMetadataEntryKey InEntryKeyA, float Weight) = 0;
-	virtual void SetWeightedValue(PCGMetadataEntryKey ItemKey, const FPCGMetadataAttributeBase* InAttribute, const TArrayView<TPair<PCGMetadataEntryKey, float>> InWeightedKeys) = 0;
+	virtual void SetWeightedValue(PCGMetadataEntryKey ItemKey, const FPCGMetadataAttributeBase* InAttribute, const TArrayView<const TPair<PCGMetadataEntryKey, float>>& InWeightedKeys) = 0;
 	virtual void SetValue(PCGMetadataEntryKey ItemKey, const FPCGMetadataAttributeBase* InAttributeA, PCGMetadataEntryKey InEntryKeyA, const FPCGMetadataAttributeBase* InAttributeB, PCGMetadataEntryKey InEntryKeyB, EPCGMetadataOp Op) = 0;
 	virtual bool IsEqualToDefaultValue(PCGMetadataValueKey ValueKey) const = 0;
 	/** In the case of multi entry attribute and after some operations, we might have a single entry attribute with a default value that is different than the first entry. Use this function to fix that. Only valid if there is one and only one value. */
@@ -52,19 +53,20 @@ public:
 	virtual bool AreValuesEqualForEntryKeys(PCGMetadataEntryKey EntryKey1, PCGMetadataEntryKey EntryKey2) const = 0;
 	virtual bool AreValuesEqual(PCGMetadataValueKey ValueKey1, PCGMetadataValueKey ValueKey2) const = 0;
 
-	void SetValueFromValueKey(PCGMetadataEntryKey EntryKey, PCGMetadataValueKey ValueKey);
+	void SetValueFromValueKey(PCGMetadataEntryKey EntryKey, PCGMetadataValueKey ValueKey, bool bResetValueOnDefaultValueKey = false);
 	PCGMetadataValueKey GetValueKey(PCGMetadataEntryKey EntryKey) const;
 	bool HasNonDefaultValue(PCGMetadataEntryKey EntryKey) const;
 	void ClearEntries();
 
 	/** Bulk getter, to lock in read only once per parent. */
-	void GetValueKeys(const TArray<PCGMetadataEntryKey>& EntryKeys, TArray<PCGMetadataValueKey>& OutValueKeys) const;
+	void GetValueKeys(const TArrayView<const PCGMetadataEntryKey>& EntryKeys, TArray<PCGMetadataValueKey>& OutValueKeys) const;
 
 	/** Bulk setter to lock in write only once. */
-	void SetValuesFromValueKeys(const TArray<TTuple<PCGMetadataEntryKey, PCGMetadataValueKey>>& EntryValuePairs, bool bResetValueOnDefaultValueKey = true);
+	void SetValuesFromValueKeys(const TArrayView<const TTuple<PCGMetadataEntryKey, PCGMetadataValueKey>>& EntryValuePairs, bool bResetValueOnDefaultValueKey = true);
 
 	/** Two arrays version of bulk setter to lock in write only once. Both arrays must be the same size. */
-	void SetValuesFromValueKeys(const TArray<PCGMetadataEntryKey>& EntryKeys, TArray<PCGMetadataValueKey>& ValueKeys, bool bResetValueOnDefaultValueKey = true);
+	void SetValuesFromValueKeys(const TArrayView<const PCGMetadataEntryKey>& EntryKeys, const TArrayView<const PCGMetadataValueKey>& ValueKeys, bool bResetValueOnDefaultValueKey = true);
+	void SetValuesFromValueKeys(const TArrayView<const PCGMetadataEntryKey* const>& EntryKeys, const TArrayView<const PCGMetadataValueKey>& ValueKeys, bool bResetValueOnDefaultValueKey = true);
 
 	bool AllowsInterpolation() const { return bAllowsInterpolation; }
 
@@ -81,9 +83,9 @@ public:
 
 private:
 	// Unsafe version, needs to be write lock protected.
-	void SetValueFromValueKey_Unsafe(PCGMetadataEntryKey EntryKey, PCGMetadataValueKey ValueKey, bool bResetValueOnDefaultValueKey);
+	void SetValueFromValueKey_Unsafe(PCGMetadataEntryKey EntryKey, PCGMetadataValueKey ValueKey, bool bResetValueOnDefaultValueKey, bool bAllowInvalidEntries = false);
 
-	void GetValueKeys_Internal(const TArray<PCGMetadataEntryKey>& EntryKeys, TArray<PCGMetadataValueKey>& OutValueKeys, TBitArray<>& UnsetValues) const;
+	void GetValueKeys_Internal(const TArrayView<const PCGMetadataEntryKey>& EntryKeys, TArrayView<PCGMetadataValueKey> OutValueKeys, TBitArray<>& UnsetValues) const;
 
 protected:
 	TMap<PCGMetadataEntryKey, PCGMetadataValueKey> EntryToValueKeyMap;
diff --git a/Engine/Plugins/Experimental/PCG/Source/PCG/Public/Metadata/PCGMetadataAttributeTpl.h b/Engine/Plugins/Experimental/PCG/Source/PCG/Public/Metadata/PCGMetadataAttributeTpl.h
--- a/Engine/Plugins/Experimental/PCG/Source/PCG/Public/Metadata/PCGMetadataAttributeTpl.h
+++ b/Engine/Plugins/Experimental/PCG/Source/PCG/Public/Metadata/PCGMetadataAttributeTpl.h
@@ -16,6 +16,9 @@ class UPCGMetadata;
 template<typename T>
 class FPCGMetadataAttribute : public FPCGMetadataAttributeBase
 {
+	template<typename U>
+	friend class FPCGMetadataAttribute;
+
 public:
 	FPCGMetadataAttribute(UPCGMetadata* InMetadata, FName InName, const FPCGMetadataAttributeBase* InParent, const T& InDefaultValue, bool bInAllowsInterpolation)
 		: FPCGMetadataAttributeBase(InMetadata, InName, InParent, bInAllowsInterpolation)
@@ -190,6 +193,15 @@ public:
 	}
 
 	virtual FPCGMetadataAttributeBase* Copy(FName NewName, UPCGMetadata* InMetadata, bool bKeepParent, bool bCopyEntries = true, bool bCopyValues = true) const override
+	{
+		return CopyInternal<T>(NewName, InMetadata, bKeepParent, bCopyEntries, bCopyValues);
+	}
+
+	virtual FPCGMetadataAttributeBase* CopyToAnotherType(int16 TargetType) const override;
+
+	// TODO: add enable if only on compatible types, but this has some repercussion on using metadata on types that aren't normally supported.
+	template<typename U>
+	FPCGMetadataAttributeBase* CopyInternal(FName NewName, UPCGMetadata* InMetadata, bool bKeepParent, bool bCopyEntries, bool bCopyValues) const
 	{
 		// If we copy an attribute where we don't want to keep the parent, while copying entries and/or values, we'll lose data.
 		// In that case, we will copy all the data from this attribute and all its ancestors.
@@ -203,9 +215,12 @@ public:
 			UE_LOG(LogPCG, Error, TEXT("Try to create a new attribute with an invalid name: %s"), *NewName.ToString());
 			return nullptr;
 		}
+
+		U NewDefaultValue{};
+		PCG::Private::GetValueWithBroadcastAndConstructible(DefaultValue, NewDefaultValue);
 		
 		// This copies to a new attribute.
-		FPCGMetadataAttribute<T>* AttributeCopy = new FPCGMetadataAttribute<T>(InMetadata, NewName, bKeepParent ? this : nullptr, DefaultValue, bAllowsInterpolation);
+		FPCGMetadataAttribute<U>* AttributeCopy = new FPCGMetadataAttribute<U>(InMetadata, NewName, bKeepParent ? this : nullptr, NewDefaultValue, bAllowsInterpolation);
 
 		// Gather the chain of parents if we don't keep the parent and we want to copy entries/values.
 		// We always have at least one item, "this".
@@ -249,7 +264,22 @@ public:
 				const FPCGMetadataAttribute<T>* Current = Parents[i];
 				
 				Current->ValueLock.ReadLock();
-				AttributeCopy->Values.Append(Current->Values);
+
+				if constexpr (std::is_same_v<T, U>)
+				{
+					AttributeCopy->Values.Append(Current->Values);
+				}
+				else
+				{
+					U NewValue{};
+
+					for (const T& Value : Current->Values)
+					{
+						PCG::Private::GetValueWithBroadcastAndConstructible(Value, NewValue);
+						AttributeCopy->Values.Add(NewValue);
+					}
+				}
+
 				// The expected value key offset is the one for this attribute (i == 0), and only if we
 				// keep the parent. Otherwise we don't have any parent, so offset should be kept at 0.
 				if (i == 0 && bKeepParent)
@@ -258,8 +288,6 @@ public:
 				}
 				Current->ValueLock.ReadUnlock();
 			}
-			
-			
 		}
 
 		return AttributeCopy;
@@ -290,7 +318,7 @@ public:
 		Accumulate(ItemKey, InAttribute, InEntryKey, Weight);
 	}
 
-	virtual void SetWeightedValue(PCGMetadataEntryKey ItemKey, const FPCGMetadataAttributeBase* InAttribute, const TArrayView<TPair<PCGMetadataEntryKey, float>> InWeightedKeys) override
+	virtual void SetWeightedValue(PCGMetadataEntryKey ItemKey, const FPCGMetadataAttributeBase* InAttribute, const TArrayView<const TPair<PCGMetadataEntryKey, float>>& InWeightedKeys) override
 	{
 		check(ItemKey != PCGInvalidEntryKey);
 		Accumulate(ItemKey, InAttribute, InWeightedKeys);
@@ -456,12 +484,91 @@ public:
 		}
 	}
 
+	template<typename IT = T, typename TEnableIf<PCG::Private::MetadataTraits<IT>::CompressData>::Type* = nullptr>
+	TArray<PCGMetadataValueKey> AddValues(const TArrayView<const T>& InValues)
+	{
+		// Since we're getting raw values here, we might have duplicates
+		// so we should aim to remove duplicates here so we preserve our 'compress data' idea, otherwise it will break other foundational blocks (e.g. partition)
+		TArray<T, TInlineAllocator<256>> UniqueValues;
+
+		// Initially, fill with mapping to unique values so we can remap them at the end if needed
+		TArray<PCGMetadataValueKey> FoundValueKeys;
+		FoundValueKeys.Reserve(InValues.Num());
+
+		for (int ValueIndex = 0; ValueIndex < InValues.Num(); ++ValueIndex)
+		{
+			FoundValueKeys.Emplace(UniqueValues.AddUnique(InValues[ValueIndex]));
+		}
+
+		const bool bHasDuplicateValues = (UniqueValues.Num() != InValues.Num());
+
+		TArray<PCGMetadataValueKey> FoundUniqueValueKeys;
+		TArray<PCGMetadataValueKey>& FoundKeys = (bHasDuplicateValues ? FoundUniqueValueKeys : FoundValueKeys);
+
+		// Implementation note: when we don't have any duplicate values, the previously-set values in FoundValueKeys will be wiped out - this is intended
+		const bool bAtLeastOneValueNotFound = !FindValues(UniqueValues, FoundKeys);
+
+		if (bAtLeastOneValueNotFound)
+		{
+			FWriteScopeLock ScopeLock(ValueLock);
+			for (int ValueIndex = 0; ValueIndex < UniqueValues.Num(); ++ValueIndex)
+			{
+				if (FoundKeys[ValueIndex] == PCGNotFoundValueKey)
+				{
+					FoundKeys[ValueIndex] = Values.Add(UniqueValues[ValueIndex]) + ValueKeyOffset;
+				}
+			}
+		}
+
+		// Remap to full array if needed
+		if (bHasDuplicateValues)
+		{
+			for (PCGMetadataValueKey& ValueToRemap : FoundValueKeys)
+			{
+				ValueToRemap = FoundUniqueValueKeys[ValueToRemap];
+			}
+		}
+
+		return FoundValueKeys;
+	}
+
+	template<typename IT = T, typename TEnableIf<!PCG::Private::MetadataTraits<IT>::CompressData>::Type* = nullptr>
+	TArray<PCGMetadataValueKey> AddValues(const TArrayView<const T>& InValues)
+	{
+		TArray<PCGMetadataValueKey> ValueKeys;
+
+		int FirstValueIndex = 0;
+
+		ValueLock.WriteLock();
+		FirstValueIndex = Values.Num() + ValueKeyOffset;
+		Values.Append(InValues);
+		ValueLock.WriteUnlock();
+
+		ValueKeys.SetNum(InValues.Num());
+		for (int ValueIndex = 0; ValueIndex < InValues.Num(); ++ValueIndex)
+		{
+			ValueKeys[ValueIndex] = FirstValueIndex + ValueIndex;
+		}
+
+		return ValueKeys;
+	}
+
 	void SetValue(PCGMetadataEntryKey ItemKey, const T& InValue)
 	{
 		check(ItemKey != PCGInvalidEntryKey);
 		SetValueFromValueKey(ItemKey, AddValue(InValue));
 	}
 
+	void SetValues(const TArrayView<const PCGMetadataEntryKey>& ItemKeys, const TArrayView<const T>& InValues)
+	{
+		SetValuesFromValueKeys(ItemKeys, AddValues(InValues));
+	}
+
+	void SetValues(const TArrayView<const PCGMetadataEntryKey * const>& ItemKeys, const TArrayView<const T>& InValues)
+	{
+		SetValuesFromValueKeys(ItemKeys, AddValues(InValues));
+	}
+
 	template<typename U>
 	void SetValue(PCGMetadataEntryKey ItemKey, const U& InValue)
 	{
@@ -532,12 +639,73 @@ public:
 		}
 	}
 
+	template<typename IT = T, typename TEnableIf<PCG::Private::MetadataTraits<IT>::CompressData>::Type* = nullptr>
+	bool FindValues(const TArrayView<const T>& InValues, TArray<PCGMetadataValueKey>& OutValueKeys) const
+	{
+		OutValueKeys.Init(PCGNotFoundValueKey, InValues.Num());
+
+		int ValueKeysSet = 0;
+		FindValuesInternal(InValues, OutValueKeys, ValueKeysSet, /*bIsRoot=*/true);
+
+		return (ValueKeysSet == InValues.Num());
+	}
+
+	template<typename IT = T, typename TEnableIf<PCG::Private::MetadataTraits<IT>::CompressData>::Type* = nullptr>
+	void FindValuesInternal(const TArrayView<const T>& InValues, TArray<PCGMetadataValueKey>& ValueKeys, int& ValueKeysSet, bool bIsRoot) const
+	{
+		check(InValues.Num() == ValueKeys.Num());
+
+		if (bIsRoot)
+		{
+			for (int ValueIndex = 0; ValueIndex < InValues.Num(); ++ValueIndex)
+			{
+				if (InValues[ValueIndex] == DefaultValue)
+				{
+					ValueKeys[ValueIndex] = PCGDefaultValueKey;
+					++ValueKeysSet;
+				}
+			}
+		}
+
+		if (ValueKeysSet != InValues.Num() && GetParent())
+		{
+			GetParent()->FindValuesInternal(InValues, ValueKeys, ValueKeysSet, /*bRoot=*/false);
+		}
+
+		if (ValueKeysSet != InValues.Num())
+		{
+			ValueLock.ReadLock();
+			for (int ValueIndex = 0; ValueIndex < InValues.Num(); ++ValueIndex)
+			{
+				if (ValueKeys[ValueIndex] != PCGNotFoundValueKey)
+				{
+					continue;
+				}
+
+				const int32 FoundValueIndex = Values.FindLast(InValues[ValueIndex]);
+				if (FoundValueIndex != INDEX_NONE)
+				{
+					ValueKeys[ValueIndex] = FoundValueIndex + ValueKeyOffset;
+					++ValueKeysSet;
+				}
+			}
+			ValueLock.ReadUnlock();
+		}
+	}
+
 	template<typename IT = T, typename TEnableIf<!PCG::Private::MetadataTraits<IT>::CompressData>::Type* = nullptr>
 	PCGMetadataValueKey FindValue(const T& InValue) const
 	{
 		return PCGNotFoundValueKey;
 	}
 
+	template<typename IT = T, typename TEnableIf<!PCG::Private::MetadataTraits<IT>::CompressData>::Type* = nullptr>
+	bool FindValues(const TArrayView<const T>& InValues, TArray<PCGMetadataValueKey>& OutValueKeys) const
+	{
+		OutValueKeys.Init(PCGNotFoundValueKey, InValues.Num());
+		return false;
+	}
+
 	void SetDefaultValue(const T& Value)
 	{
 		DefaultValue = Value;
@@ -677,7 +845,7 @@ protected:
 	}
 
 	template<typename IT = T, typename TEnableIf<PCG::Private::MetadataTraits<IT>::CanInterpolate>::Type* = nullptr>
-	void Accumulate(PCGMetadataEntryKey ItemKey, const FPCGMetadataAttributeBase* InAttribute, const TArrayView<TPair<PCGMetadataEntryKey, float>>& InWeightedKeys)
+	void Accumulate(PCGMetadataEntryKey ItemKey, const FPCGMetadataAttributeBase* InAttribute, const TArrayView<const TPair<PCGMetadataEntryKey, float>>& InWeightedKeys)
 	{
 		IT Value = PCG::Private::MetadataTraits<IT>::ZeroValue();
 		for (const TPair<PCGMetadataEntryKey, float>& WeightedEntry : InWeightedKeys)
@@ -692,7 +860,7 @@ protected:
 	}
 
 	template<typename IT = T, typename TEnableIf<!PCG::Private::MetadataTraits<IT>::CanInterpolate>::Type* = nullptr>
-	void Accumulate(PCGMetadataEntryKey ItemKey, const FPCGMetadataAttributeBase* InAttribute, const TArrayView<TPair<PCGMetadataEntryKey, float>>& InWeightedKeys)
+	void Accumulate(PCGMetadataEntryKey ItemKey, const FPCGMetadataAttributeBase* InAttribute, const TArrayView<const TPair<PCGMetadataEntryKey, float>>& InWeightedKeys)
 	{
 		// Empty on purpose
 	}
@@ -751,6 +919,25 @@ namespace PCGMetadataAttribute
 	}
 }
 
+template<typename T>
+FPCGMetadataAttributeBase* FPCGMetadataAttribute<T>::CopyToAnotherType(int16 TargetType) const
+{
+	return PCGMetadataAttribute::CallbackWithRightType(TargetType, [this](auto Dummy) -> FPCGMetadataAttributeBase*
+	{
+		using U = decltype(Dummy);
+
+		if constexpr (PCG::Private::IsBroadcastableOrConstructible(PCG::Private::MetadataTypes<T>::Id, PCG::Private::MetadataTypes<U>::Id))
+		{
+			return CopyInternal<U>(Name, Metadata, /*bKeepParent=*/false, /*bCopyEntries=*/true, /*bCopyValues=*/true);
+		}
+		else
+		{
+			UE_LOG(LogPCG, Error, TEXT("Metadata attribute '%s' cannot change its type - delete and create instead"), *Name.ToString());
+			return nullptr;
+		}
+	});
+}
+
 #if UE_ENABLE_INCLUDE_ORDER_DEPRECATED_IN_5_2
 #include "CoreMinimal.h"
 #endif
diff --git a/Engine/Plugins/Experimental/PCG/Source/PCG/Public/Metadata/PCGMetadataAttributeTraits.h b/Engine/Plugins/Experimental/PCG/Source/PCG/Public/Metadata/PCGMetadataAttributeTraits.h
--- a/Engine/Plugins/Experimental/PCG/Source/PCG/Public/Metadata/PCGMetadataAttributeTraits.h
+++ b/Engine/Plugins/Experimental/PCG/Source/PCG/Public/Metadata/PCGMetadataAttributeTraits.h
@@ -1053,5 +1053,34 @@ namespace PCG
 				}
 			}
 		}
+
+		template<typename InType, typename OutType>
+		inline bool GetValueWithBroadcastAndConstructible(const InType& InValue, OutType& OutValue)
+		{
+			if constexpr (std::is_same_v<OutType, InType>)
+			{
+				OutValue = InValue;
+				return true;
+			}
+			else
+			{
+				// Special case'd because FSoftObjectPath currently has a deprecated constructor from FName which generates compile warnings.
+				constexpr bool bNameToSoftObjectPath = std::is_same_v<OutType, FSoftObjectPath>&& std::is_same_v<InType, FName>;
+
+				if constexpr (PCG::Private::IsBroadcastable<InType, OutType>())
+				{
+					return GetValueWithBroadcast<InType, OutType>(InValue, OutValue);
+				}
+				else if constexpr (std::is_constructible_v<OutType, InType> && !bNameToSoftObjectPath)
+				{
+					OutValue = OutType(InValue);
+					return true;
+				}
+				else
+				{
+					return false;
+				}
+			}
+		}
 	}
 }
