[PCG] Multi-line node titles

Nodes have an authored title line (primary) and a generated title line (secondary). Primary line tends to have most important/relevant information and can depend on whether user has authored a title or not.

Users requested some nodes have flipped order so that the generated part is displayed first, helps for e.g. BP nodes where the key bit of info is the executed blueprint.

Text clipping/overflow flipped for generated title lines so that the end of the line is always visible which usually has the key info.

diff --git a/Engine/Plugins/Experimental/PCG/Source/PCG/Private/Elements/Metadata/PCGMetadataBitwiseOpElement.cpp b/Engine/Plugins/Experimental/PCG/Source/PCG/Private/Elements/Metadata/PCGMetadataBitwiseOpElement.cpp
--- a/Engine/Plugins/Experimental/PCG/Source/PCG/Private/Elements/Metadata/PCGMetadataBitwiseOpElement.cpp
+++ b/Engine/Plugins/Experimental/PCG/Source/PCG/Private/Elements/Metadata/PCGMetadataBitwiseOpElement.cpp
@@ -86,15 +86,15 @@ FPCGAttributePropertyInputSelector UPCGMetadataBitwiseSettings::GetInputSource(u
 	}
 }
 
-FName UPCGMetadataBitwiseSettings::AdditionalTaskName() const
+FString UPCGMetadataBitwiseSettings::GetAdditionalTitleInformation() const
 {
 	if (const UEnum* EnumPtr = StaticEnum<EPCGMetadataBitwiseOperation>())
 	{
-		return FName(FString("Bitwise: ") + EnumPtr->GetNameStringByValue(static_cast<int>(Operation)));
+		return FString("Bitwise: ") + EnumPtr->GetNameStringByValue(static_cast<int>(Operation));
 	}
 	else
 	{
-		return NAME_None;
+		return FString();
 	}
 }
 
diff --git a/Engine/Plugins/Experimental/PCG/Source/PCG/Private/Elements/Metadata/PCGMetadataBooleanOpElement.cpp b/Engine/Plugins/Experimental/PCG/Source/PCG/Private/Elements/Metadata/PCGMetadataBooleanOpElement.cpp
--- a/Engine/Plugins/Experimental/PCG/Source/PCG/Private/Elements/Metadata/PCGMetadataBooleanOpElement.cpp
+++ b/Engine/Plugins/Experimental/PCG/Source/PCG/Private/Elements/Metadata/PCGMetadataBooleanOpElement.cpp
@@ -91,15 +91,15 @@ uint16 UPCGMetadataBooleanSettings::GetOutputType(uint16 InputTypeId) const
 	return (uint16)EPCGMetadataTypes::Boolean;
 }
 
-FName UPCGMetadataBooleanSettings::AdditionalTaskName() const
+FString UPCGMetadataBooleanSettings::GetAdditionalTitleInformation() const
 {
 	if (const UEnum* EnumPtr = StaticEnum<EPCGMetadataBooleanOperation>())
 	{
-		return FName(FString("Boolean: ") + EnumPtr->GetNameStringByValue(static_cast<int>(Operation)));
+		return FString("Boolean: ") + EnumPtr->GetNameStringByValue(static_cast<int>(Operation));
 	}
 	else
 	{
-		return NAME_None;
+		return FString();
 	}
 }
 
diff --git a/Engine/Plugins/Experimental/PCG/Source/PCG/Private/Elements/Metadata/PCGMetadataCompareOpElement.cpp b/Engine/Plugins/Experimental/PCG/Source/PCG/Private/Elements/Metadata/PCGMetadataCompareOpElement.cpp
--- a/Engine/Plugins/Experimental/PCG/Source/PCG/Private/Elements/Metadata/PCGMetadataCompareOpElement.cpp
+++ b/Engine/Plugins/Experimental/PCG/Source/PCG/Private/Elements/Metadata/PCGMetadataCompareOpElement.cpp
@@ -117,15 +117,15 @@ uint16 UPCGMetadataCompareSettings::GetOutputType(uint16 InputTypeId) const
 	return (uint16)EPCGMetadataTypes::Boolean;
 }
 
-FName UPCGMetadataCompareSettings::AdditionalTaskName() const
+FString UPCGMetadataCompareSettings::GetAdditionalTitleInformation() const
 {
 	if (const UEnum* EnumPtr = StaticEnum<EPCGMetadataCompareOperation>())
 	{
-		return FName(FString("Compare: ") + EnumPtr->GetNameStringByValue(static_cast<int>(Operation)));
+		return FString("Compare: ") + EnumPtr->GetNameStringByValue(static_cast<int>(Operation));
 	}
 	else
 	{
-		return NAME_None;
+		return FString();
 	}
 }
 
diff --git a/Engine/Plugins/Experimental/PCG/Source/PCG/Private/Elements/Metadata/PCGMetadataElement.cpp b/Engine/Plugins/Experimental/PCG/Source/PCG/Private/Elements/Metadata/PCGMetadataElement.cpp
--- a/Engine/Plugins/Experimental/PCG/Source/PCG/Private/Elements/Metadata/PCGMetadataElement.cpp
+++ b/Engine/Plugins/Experimental/PCG/Source/PCG/Private/Elements/Metadata/PCGMetadataElement.cpp
@@ -29,16 +29,23 @@ UPCGMetadataOperationSettings::UPCGMetadataOperationSettings()
 	OutputTarget.SetAttributeName(NAME_None);
 }
 
-FName UPCGMetadataOperationSettings::AdditionalTaskName() const
+FString UPCGMetadataOperationSettings::GetAdditionalTitleInformation() const
 {
 #if WITH_EDITOR
 	if (bCopyAllAttributes)
 	{
-		return FName(LOCTEXT("NoteTitleAllAttributes", "Copy All Attributes").ToString());
+		return LOCTEXT("NoteTitleAllAttributes", "All Attributes").ToString();
+	}
+
+	const FProperty* FromAttributeProperty = GetClass() ? FindFProperty<FProperty>(GetClass(), GET_MEMBER_NAME_CHECKED(UPCGMetadataOperationSettings, InputSource)) : nullptr;
+	const FProperty* ToAttributeProperty = GetClass() ? FindFProperty<FProperty>(GetClass(), GET_MEMBER_NAME_CHECKED(UPCGMetadataOperationSettings, OutputTarget)) : nullptr;
+	if ((FromAttributeProperty && IsPropertyOverriddenByPin(FromAttributeProperty)) || (ToAttributeProperty && IsPropertyOverriddenByPin(ToAttributeProperty)))
+	{
+		return FString();
 	}
 #endif
 
-	return NAME_None;
+	return FString::Printf(TEXT("%s -> %s"), *InputSource.GetDisplayText().ToString(), *OutputTarget.GetDisplayText().ToString());
 }
 
 #if WITH_EDITOR
diff --git a/Engine/Plugins/Experimental/PCG/Source/PCG/Private/Elements/Metadata/PCGMetadataMathsOpElement.cpp b/Engine/Plugins/Experimental/PCG/Source/PCG/Private/Elements/Metadata/PCGMetadataMathsOpElement.cpp
--- a/Engine/Plugins/Experimental/PCG/Source/PCG/Private/Elements/Metadata/PCGMetadataMathsOpElement.cpp
+++ b/Engine/Plugins/Experimental/PCG/Source/PCG/Private/Elements/Metadata/PCGMetadataMathsOpElement.cpp
@@ -273,15 +273,15 @@ FPCGAttributePropertyInputSelector UPCGMetadataMathsSettings::GetInputSource(uin
 	}
 }
 
-FName UPCGMetadataMathsSettings::AdditionalTaskName() const
+FString UPCGMetadataMathsSettings::GetAdditionalTitleInformation() const
 {
 	if (const UEnum* EnumPtr = StaticEnum<EPCGMetadataMathsOperation>())
 	{
-		return FName(FString("Maths: ") + EnumPtr->GetNameStringByValue(static_cast<int>(Operation)));
+		return EnumPtr->GetNameStringByValue(static_cast<int>(Operation));
 	}
 	else
 	{
-		return NAME_None;
+		return FString();
 	}
 }
 
diff --git a/Engine/Plugins/Experimental/PCG/Source/PCG/Private/Elements/Metadata/PCGMetadataPartition.cpp b/Engine/Plugins/Experimental/PCG/Source/PCG/Private/Elements/Metadata/PCGMetadataPartition.cpp
--- a/Engine/Plugins/Experimental/PCG/Source/PCG/Private/Elements/Metadata/PCGMetadataPartition.cpp
+++ b/Engine/Plugins/Experimental/PCG/Source/PCG/Private/Elements/Metadata/PCGMetadataPartition.cpp
@@ -43,6 +43,21 @@ void UPCGMetadataPartitionSettings::PostLoad()
 #endif // WITH_EDITOR
 }
 
+FString UPCGMetadataPartitionSettings::GetAdditionalTitleInformation() const
+{
+#if WITH_EDITOR
+	FProperty* AttributeProperty = GetClass() ? FindFProperty<FProperty>(GetClass(), GET_MEMBER_NAME_CHECKED(UPCGMetadataPartitionSettings, PartitionAttributeSource)) : nullptr;
+	if (AttributeProperty && IsPropertyOverriddenByPin(AttributeProperty))
+	{
+		return FString();
+	}
+	else
+#endif
+	{
+		return PartitionAttributeSource.GetDisplayText().ToString();
+	}
+}
+
 bool FPCGMetadataPartitionElement::ExecuteInternal(FPCGContext* Context) const
 {
 	TRACE_CPUPROFILER_EVENT_SCOPE(FPCGMetadataPartitionElement::Execute);
diff --git a/Engine/Plugins/Experimental/PCG/Source/PCG/Private/Elements/Metadata/PCGMetadataRenameElement.cpp b/Engine/Plugins/Experimental/PCG/Source/PCG/Private/Elements/Metadata/PCGMetadataRenameElement.cpp
--- a/Engine/Plugins/Experimental/PCG/Source/PCG/Private/Elements/Metadata/PCGMetadataRenameElement.cpp
+++ b/Engine/Plugins/Experimental/PCG/Source/PCG/Private/Elements/Metadata/PCGMetadataRenameElement.cpp
@@ -12,6 +12,21 @@
 
 #define LOCTEXT_NAMESPACE "PCGMetadataRenameElement"
 
+FString UPCGMetadataRenameSettings::GetAdditionalTitleInformation() const
+{
+#if WITH_EDITOR
+	const FProperty* FromAttributeProperty = GetClass() ? FindFProperty<FProperty>(GetClass(), GET_MEMBER_NAME_CHECKED(UPCGMetadataRenameSettings, AttributeToRename)) : nullptr;
+	const FProperty* ToAttributeProperty = GetClass() ? FindFProperty<FProperty>(GetClass(), GET_MEMBER_NAME_CHECKED(UPCGMetadataRenameSettings, NewAttributeName)) : nullptr;
+	if ((FromAttributeProperty && IsPropertyOverriddenByPin(FromAttributeProperty)) || (ToAttributeProperty && IsPropertyOverriddenByPin(ToAttributeProperty)))
+	{
+		return FString();
+	}
+	else
+#endif
+	{
+		return FString::Printf(TEXT("%s -> %s"), *AttributeToRename.ToString(), *NewAttributeName.ToString());
+	}
+}
 
 TArray<FPCGPinProperties> UPCGMetadataRenameSettings::InputPinProperties() const
 {
diff --git a/Engine/Plugins/Experimental/PCG/Source/PCG/Private/Elements/Metadata/PCGMetadataRotatorOpElement.cpp b/Engine/Plugins/Experimental/PCG/Source/PCG/Private/Elements/Metadata/PCGMetadataRotatorOpElement.cpp
--- a/Engine/Plugins/Experimental/PCG/Source/PCG/Private/Elements/Metadata/PCGMetadataRotatorOpElement.cpp
+++ b/Engine/Plugins/Experimental/PCG/Source/PCG/Private/Elements/Metadata/PCGMetadataRotatorOpElement.cpp
@@ -194,15 +194,15 @@ FPCGAttributePropertyInputSelector UPCGMetadataRotatorSettings::GetInputSource(u
 	}
 }
 
-FName UPCGMetadataRotatorSettings::AdditionalTaskName() const
+FString UPCGMetadataRotatorSettings::GetAdditionalTitleInformation() const
 {
 	if (const UEnum* EnumPtr = StaticEnum<EPCGMetadataRotatorOperation>())
 	{
-		return FName(FString("Rotator: ") + EnumPtr->GetNameStringByValue(static_cast<int>(Operation)));
+		return FString("Rotator: ") + EnumPtr->GetNameStringByValue(static_cast<int>(Operation));
 	}
 	else
 	{
-		return NAME_None;
+		return FString();
 	}
 }
 
diff --git a/Engine/Plugins/Experimental/PCG/Source/PCG/Private/Elements/Metadata/PCGMetadataStringOpElement.cpp b/Engine/Plugins/Experimental/PCG/Source/PCG/Private/Elements/Metadata/PCGMetadataStringOpElement.cpp
--- a/Engine/Plugins/Experimental/PCG/Source/PCG/Private/Elements/Metadata/PCGMetadataStringOpElement.cpp
+++ b/Engine/Plugins/Experimental/PCG/Source/PCG/Private/Elements/Metadata/PCGMetadataStringOpElement.cpp
@@ -37,15 +37,15 @@ FPCGAttributePropertyInputSelector UPCGMetadataStringOpSettings::GetInputSource(
 	}
 }
 
-FName UPCGMetadataStringOpSettings::AdditionalTaskName() const
+FString UPCGMetadataStringOpSettings::GetAdditionalTitleInformation() const
 {
 	if (const UEnum* EnumPtr = StaticEnum<EPCGMetadataStringOperation>())
 	{
-		return FName(FString("String: ") + EnumPtr->GetNameStringByValue(static_cast<int>(Operation)));
+		return FString("String: ") + EnumPtr->GetNameStringByValue(static_cast<int>(Operation));
 	}
 	else
 	{
-		return NAME_None;
+		return FString();
 	}
 }
  
diff --git a/Engine/Plugins/Experimental/PCG/Source/PCG/Private/Elements/Metadata/PCGMetadataTransformOpElement.cpp b/Engine/Plugins/Experimental/PCG/Source/PCG/Private/Elements/Metadata/PCGMetadataTransformOpElement.cpp
--- a/Engine/Plugins/Experimental/PCG/Source/PCG/Private/Elements/Metadata/PCGMetadataTransformOpElement.cpp
+++ b/Engine/Plugins/Experimental/PCG/Source/PCG/Private/Elements/Metadata/PCGMetadataTransformOpElement.cpp
@@ -157,15 +157,15 @@ FPCGAttributePropertyInputSelector UPCGMetadataTransformSettings::GetInputSource
 	}
 }
 
-FName UPCGMetadataTransformSettings::AdditionalTaskName() const
+FString UPCGMetadataTransformSettings::GetAdditionalTitleInformation() const
 {
 	if (const UEnum* EnumPtr = StaticEnum<EPCGMetadataTransformOperation>())
 	{
-		return FName(FString("Transform: ") + EnumPtr->GetNameStringByValue(static_cast<int>(Operation)));
+		return FString("Transform: ") + EnumPtr->GetNameStringByValue(static_cast<int>(Operation));
 	}
 	else
 	{
-		return NAME_None;
+		return FString();
 	}
 }
 
diff --git a/Engine/Plugins/Experimental/PCG/Source/PCG/Private/Elements/Metadata/PCGMetadataTrigOpElement.cpp b/Engine/Plugins/Experimental/PCG/Source/PCG/Private/Elements/Metadata/PCGMetadataTrigOpElement.cpp
--- a/Engine/Plugins/Experimental/PCG/Source/PCG/Private/Elements/Metadata/PCGMetadataTrigOpElement.cpp
+++ b/Engine/Plugins/Experimental/PCG/Source/PCG/Private/Elements/Metadata/PCGMetadataTrigOpElement.cpp
@@ -103,15 +103,15 @@ uint16 UPCGMetadataTrigSettings::GetOutputType(uint16 InputTypeId) const
 	return (uint16)EPCGMetadataTypes::Double;
 }
 
-FName UPCGMetadataTrigSettings::AdditionalTaskName() const
+FString UPCGMetadataTrigSettings::GetAdditionalTitleInformation() const
 {
 	if (const UEnum* EnumPtr = StaticEnum<EPCGMetadataTrigOperation>())
 	{
-		return FName(FString("Trig: ") + EnumPtr->GetNameStringByValue(static_cast<int>(Operation)));
+		return FString("Trig: ") + EnumPtr->GetNameStringByValue(static_cast<int>(Operation));
 	}
 	else
 	{
-		return NAME_None;
+		return FString();
 	}
 }
 
diff --git a/Engine/Plugins/Experimental/PCG/Source/PCG/Private/Elements/Metadata/PCGMetadataVectorOpElement.cpp b/Engine/Plugins/Experimental/PCG/Source/PCG/Private/Elements/Metadata/PCGMetadataVectorOpElement.cpp
--- a/Engine/Plugins/Experimental/PCG/Source/PCG/Private/Elements/Metadata/PCGMetadataVectorOpElement.cpp
+++ b/Engine/Plugins/Experimental/PCG/Source/PCG/Private/Elements/Metadata/PCGMetadataVectorOpElement.cpp
@@ -319,15 +319,15 @@ FPCGAttributePropertyInputSelector UPCGMetadataVectorSettings::GetInputSource(ui
 	}
 }
 
-FName UPCGMetadataVectorSettings::AdditionalTaskName() const
+FString UPCGMetadataVectorSettings::GetAdditionalTitleInformation() const
 {
 	if (const UEnum* EnumPtr = StaticEnum<EPCGMetadataVectorOperation>())
 	{
-		return FName(FString("Vector: ") + EnumPtr->GetNameStringByValue(static_cast<int>(Operation)));
+		return FString("Vector: ") + EnumPtr->GetNameStringByValue(static_cast<int>(Operation));
 	}
 	else
 	{
-		return NAME_None;
+		return FString();
 	}
 }
 
diff --git a/Engine/Plugins/Experimental/PCG/Source/PCG/Private/Elements/PCGActorSelector.cpp b/Engine/Plugins/Experimental/PCG/Source/PCG/Private/Elements/PCGActorSelector.cpp
--- a/Engine/Plugins/Experimental/PCG/Source/PCG/Private/Elements/PCGActorSelector.cpp
+++ b/Engine/Plugins/Experimental/PCG/Source/PCG/Private/Elements/PCGActorSelector.cpp
@@ -301,11 +301,11 @@ FText FPCGActorSelectorSettings::GetTaskNameSuffix() const
 	{
 		if (ActorSelection == EPCGActorSelection::ByClass)
 		{
-			return (ActorSelectionClass.Get() ? ActorSelectionClass->GetDisplayNameText() : FText::FromName(NAME_None));
+			return FText::Format(FText::FromString(TEXT("Class: {0}")), (ActorSelectionClass.Get() ? ActorSelectionClass->GetDisplayNameText() : FText::FromName(NAME_None)));
 		}
 		else if (ActorSelection == EPCGActorSelection::ByTag)
 		{
-			return FText::FromName(ActorSelectionTag);
+			return FText::Format(FText::FromString(TEXT("Tag: {0}")), FText::FromName(ActorSelectionTag));
 		}
 	}
 	else if(const UEnum* EnumPtr = StaticEnum<EPCGActorFilter>())
diff --git a/Engine/Plugins/Experimental/PCG/Source/PCG/Private/Elements/PCGAttributeFilter.cpp b/Engine/Plugins/Experimental/PCG/Source/PCG/Private/Elements/PCGAttributeFilter.cpp
--- a/Engine/Plugins/Experimental/PCG/Source/PCG/Private/Elements/PCGAttributeFilter.cpp
+++ b/Engine/Plugins/Experimental/PCG/Source/PCG/Private/Elements/PCGAttributeFilter.cpp
@@ -301,6 +301,21 @@ TArray<FPCGPinProperties> UPCGAttributeFilteringSettings::OutputPinProperties()
 	return PinProperties;
 }
 
+FString UPCGAttributeFilteringSettings::GetAdditionalTitleInformation() const
+{
+#if WITH_EDITOR
+	const FProperty* AttributeProperty = GetClass() ? FindFProperty<FProperty>(GetClass(), GET_MEMBER_NAME_CHECKED(UPCGAttributeFilteringSettings, TargetAttribute)) : nullptr;
+	if (AttributeProperty && IsPropertyOverriddenByPin(AttributeProperty))
+	{
+		return FString();
+	}
+	else
+#endif
+	{
+		return TargetAttribute.GetDisplayText().ToString();
+	}
+}
+
 FPCGElementPtr UPCGAttributeFilteringSettings::CreateElement() const
 {
 	return MakeShared<FPCGAttributeFilterElement>();
@@ -334,6 +349,21 @@ void UPCGAttributeFilteringRangeSettings::PostLoad()
 }
 #endif
 
+FString UPCGAttributeFilteringRangeSettings::GetAdditionalTitleInformation() const
+{
+#if WITH_EDITOR
+	const FProperty* AttributeProperty = GetClass() ? FindFProperty<FProperty>(GetClass(), GET_MEMBER_NAME_CHECKED(UPCGAttributeFilteringRangeSettings, TargetAttribute)) : nullptr;
+	if (AttributeProperty && IsPropertyOverriddenByPin(AttributeProperty))
+	{
+		return FString();
+	}
+	else
+#endif
+	{
+		return TargetAttribute.GetDisplayText().ToString();
+	}
+}
+
 FPCGElementPtr UPCGAttributeFilteringRangeSettings::CreateElement() const
 {
 	return MakeShared<FPCGAttributeFilterRangeElement>();
diff --git a/Engine/Plugins/Experimental/PCG/Source/PCG/Private/Elements/PCGAttributeReduceElement.cpp b/Engine/Plugins/Experimental/PCG/Source/PCG/Private/Elements/PCGAttributeReduceElement.cpp
--- a/Engine/Plugins/Experimental/PCG/Source/PCG/Private/Elements/PCGAttributeReduceElement.cpp
+++ b/Engine/Plugins/Experimental/PCG/Source/PCG/Private/Elements/PCGAttributeReduceElement.cpp
@@ -163,7 +163,7 @@ void UPCGAttributeReduceSettings::PostLoad()
 #endif
 }
 
-FName UPCGAttributeReduceSettings::AdditionalTaskName() const
+FString UPCGAttributeReduceSettings::GetAdditionalTitleInformation() const
 {
 	if (const UEnum* EnumPtr = StaticEnum<EPCGAttributeReduceOperation>())
 	{
@@ -177,16 +177,16 @@ FName UPCGAttributeReduceSettings::AdditionalTaskName() const
 
 		if (InputAttributeName != OutputAttributeName && OutputAttributeName != NAME_None)
 		{
-			return FName(FString::Printf(TEXT("Reduce %s to %s: %s"), *InputAttributeName.ToString(), *OutputAttributeName.ToString(), *OperationName));
+			return FString::Printf(TEXT("Reduce %s to %s: %s"), *InputAttributeName.ToString(), *OutputAttributeName.ToString(), *OperationName);
 		}
 		else
 		{
-			return FName(FString::Printf(TEXT("Reduce %s: %s"), *InputAttributeName.ToString(), *OperationName));
+			return FString::Printf(TEXT("Reduce %s: %s"), *InputAttributeName.ToString(), *OperationName);
 		}
 	}
 	else
 	{
-		return NAME_None;
+		return FString();
 	}
 }
 
diff --git a/Engine/Plugins/Experimental/PCG/Source/PCG/Private/Elements/PCGAttributeSelectElement.cpp b/Engine/Plugins/Experimental/PCG/Source/PCG/Private/Elements/PCGAttributeSelectElement.cpp
--- a/Engine/Plugins/Experimental/PCG/Source/PCG/Private/Elements/PCGAttributeSelectElement.cpp
+++ b/Engine/Plugins/Experimental/PCG/Source/PCG/Private/Elements/PCGAttributeSelectElement.cpp
@@ -207,7 +207,7 @@ void UPCGAttributeSelectSettings::ApplyDeprecation(UPCGNode* InOutNode)
 }
 #endif
 
-FName UPCGAttributeSelectSettings::AdditionalTaskName() const
+FString UPCGAttributeSelectSettings::GetAdditionalTitleInformation() const
 {
 	if (const UEnum* EnumOpPtr = StaticEnum<EPCGAttributeSelectOperation>())
 	{
@@ -232,16 +232,16 @@ FName UPCGAttributeSelectSettings::AdditionalTaskName() const
 
 			if (InputAttributeName != OutputAttributeName && OutputAttributeName != NAME_None)
 			{
-				return FName(FString::Printf(TEXT("Select %s to %s: %s on %s"), *InputAttributeName.ToString(), *OutputAttributeName.ToString(), *OperationName, *AxisName));
+				return FString::Printf(TEXT("Select %s to %s: %s on %s"), *InputAttributeName.ToString(), *OutputAttributeName.ToString(), *OperationName, *AxisName);
 			}
 			else
 			{
-				return FName(FString::Printf(TEXT("Select %s: %s on %s"), *InputAttributeName.ToString(), *OperationName, *AxisName));
+				return FString::Printf(TEXT("Select %s: %s on %s"), *InputAttributeName.ToString(), *OperationName, *AxisName);
 			}
 		}
 	}
 
-	return NAME_None;
+	return FString();
 }
 
 TArray<FPCGPinProperties> UPCGAttributeSelectSettings::InputPinProperties() const
diff --git a/Engine/Plugins/Experimental/PCG/Source/PCG/Private/Elements/PCGAttributeTransferElement.cpp b/Engine/Plugins/Experimental/PCG/Source/PCG/Private/Elements/PCGAttributeTransferElement.cpp
--- a/Engine/Plugins/Experimental/PCG/Source/PCG/Private/Elements/PCGAttributeTransferElement.cpp
+++ b/Engine/Plugins/Experimental/PCG/Source/PCG/Private/Elements/PCGAttributeTransferElement.cpp
@@ -84,18 +84,18 @@ EPCGDataType UPCGAttributeTransferSettings::GetCurrentPinTypes(const UPCGPin* In
 	return bAnyArePoint ? EPCGDataType::Point : EPCGDataType::Spatial;
 }
 
-FName UPCGAttributeTransferSettings::AdditionalTaskName() const
+FString UPCGAttributeTransferSettings::GetAdditionalTitleInformation() const
 {
 	const FName SourceAttributeName = SourceAttributeProperty.GetName();
 	const FName TargetAttributeName = TargetAttributeProperty.GetName();
 
 	if (SourceAttributeName != TargetAttributeName && TargetAttributeName != NAME_None)
 	{
-		return FName(FString::Printf(TEXT("%s %s to %s"), *PCGAttributeTransferConstants::NodeName.ToString(), *SourceAttributeName.ToString(), *TargetAttributeName.ToString()));
+		return FString::Printf(TEXT("%s %s to %s"), *PCGAttributeTransferConstants::NodeName.ToString(), *SourceAttributeName.ToString(), *TargetAttributeName.ToString());
 	}
 	else
 	{
-		return FName(FString::Printf(TEXT("%s %s"), *PCGAttributeTransferConstants::NodeName.ToString(), *SourceAttributeName.ToString()));
+		return FString::Printf(TEXT("%s %s"), *PCGAttributeTransferConstants::NodeName.ToString(), *SourceAttributeName.ToString());
 	}
 }
 
diff --git a/Engine/Plugins/Experimental/PCG/Source/PCG/Private/Elements/PCGCreateAttribute.cpp b/Engine/Plugins/Experimental/PCG/Source/PCG/Private/Elements/PCGCreateAttribute.cpp
--- a/Engine/Plugins/Experimental/PCG/Source/PCG/Private/Elements/PCGCreateAttribute.cpp
+++ b/Engine/Plugins/Experimental/PCG/Source/PCG/Private/Elements/PCGCreateAttribute.cpp
@@ -107,34 +107,33 @@ FPCGElementPtr UPCGCreateAttributeBaseSettings::CreateElement() const
 	return MakeShared<FPCGCreateAttributeElement>();
 }
 
-FName UPCGCreateAttributeBaseSettings::AdditionalTaskNameInternal(FName NodeName) const
+FString UPCGCreateAttributeBaseSettings::GetAdditionalTitleInformationInternal(FName NodeName) const
 {
 	if (HasAnyFlags(RF_ClassDefaultObject))
 	{
-		return NodeName;
+		return NodeName.ToString();
 	}
 
 	const UPCGNode* Node = Cast<UPCGNode>(GetOuter());
-	const bool AttributesPinIsConnected = Node ? Node->IsInputPinConnected(PCGCreateAttributeConstants::AttributesLabel) : false;
+	const bool bAttributesPinIsConnected = Node ? Node->IsInputPinConnected(PCGCreateAttributeConstants::AttributesLabel) : false;
 
-	const FName OutputAttributeName = GetOutputAttributeName(AttributesPinIsConnected ? &InputSource : nullptr, nullptr);
+	const FName OutputAttributeName = GetOutputAttributeName(bAttributesPinIsConnected ? &InputSource : nullptr, nullptr);
 	const FName SourceParamAttributeName = InputSource.GetName();
 
-	if (ShouldAddAttributesPin() && AttributesPinIsConnected)
+	if (ShouldAddAttributesPin() && bAttributesPinIsConnected)
 	{
 		if ((OutputAttributeName == NAME_None) && (SourceParamAttributeName == NAME_None))
 		{
-			return NodeName;
+			return NodeName.ToString();
 		}
 		else
 		{
-			const FString AttributeName = ((OutputAttributeName == NAME_None) ? SourceParamAttributeName : OutputAttributeName).ToString();
-			return FName(FString::Printf(TEXT("%s %s"), *NodeName.ToString(), *AttributeName));
+			return OutputAttributeName == NAME_None ? SourceParamAttributeName.ToString() : OutputAttributeName.ToString();
 		}
 	}
 	else
 	{
-		return FName(FString::Printf(TEXT("%s: %s"), *OutputAttributeName.ToString(), *AttributeTypes.ToString()));
+		return FString::Printf(TEXT("%s: %s"), *OutputAttributeName.ToString(), *AttributeTypes.ToString());
 	}
 }
 
@@ -143,9 +142,9 @@ UPCGAddAttributeSettings::UPCGAddAttributeSettings()
 	OutputTarget.SetAttributeName(NAME_None);
 }
 
-FName UPCGAddAttributeSettings::AdditionalTaskName() const
+FString UPCGAddAttributeSettings::GetAdditionalTitleInformation() const
 {
-	return AdditionalTaskNameInternal(PCGCreateAttributeConstants::NodeNameAddAttribute);
+	return GetAdditionalTitleInformationInternal(PCGCreateAttributeConstants::NodeNameAddAttribute);
 }
 
 void UPCGAddAttributeSettings::PostLoad()
@@ -252,9 +251,9 @@ TArray<FPCGPinProperties> UPCGCreateAttributeSetSettings::OutputPinProperties()
 	return PinProperties;
 }
 
-FName UPCGCreateAttributeSetSettings::AdditionalTaskName() const
+FString UPCGCreateAttributeSetSettings::GetAdditionalTitleInformation() const
 {
-	return AdditionalTaskNameInternal(PCGCreateAttributeConstants::NodeNameCreateAttribute);
+	return GetAdditionalTitleInformationInternal(PCGCreateAttributeConstants::NodeNameCreateAttribute);
 }
 
 bool FPCGCreateAttributeElement::ExecuteInternal(FPCGContext* Context) const
diff --git a/Engine/Plugins/Experimental/PCG/Source/PCG/Private/Elements/PCGDataFromActor.cpp b/Engine/Plugins/Experimental/PCG/Source/PCG/Private/Elements/PCGDataFromActor.cpp
--- a/Engine/Plugins/Experimental/PCG/Source/PCG/Private/Elements/PCGDataFromActor.cpp
+++ b/Engine/Plugins/Experimental/PCG/Source/PCG/Private/Elements/PCGDataFromActor.cpp
@@ -75,12 +75,12 @@ void UPCGDataFromActorSettings::PostLoad()
 	}
 }
 
-FName UPCGDataFromActorSettings::AdditionalTaskName() const
+FString UPCGDataFromActorSettings::GetAdditionalTitleInformation() const
 {
 #if WITH_EDITOR
-	return ActorSelector.GetTaskName(GetDefaultNodeTitle());
+	return ActorSelector.GetTaskNameSuffix().ToString();
 #else
-	return Super::AdditionalTaskName();
+	return Super::GetAdditionalTitleInformation();
 #endif
 }
 
diff --git a/Engine/Plugins/Experimental/PCG/Source/PCG/Private/Elements/PCGDataTableRowToParamData.cpp b/Engine/Plugins/Experimental/PCG/Source/PCG/Private/Elements/PCGDataTableRowToParamData.cpp
--- a/Engine/Plugins/Experimental/PCG/Source/PCG/Private/Elements/PCGDataTableRowToParamData.cpp
+++ b/Engine/Plugins/Experimental/PCG/Source/PCG/Private/Elements/PCGDataTableRowToParamData.cpp
@@ -32,9 +32,9 @@ TArray<FPCGPinProperties> UPCGDataTableRowToParamDataSettings::OutputPinProperti
 	return PinProperties;
 }
 
-FName UPCGDataTableRowToParamDataSettings::AdditionalTaskName() const
+FString UPCGDataTableRowToParamDataSettings::GetAdditionalTitleInformation() const
 {
-	return FName(FString::Printf(TEXT("%s[ %s ]"), DataTable ? *DataTable->GetFName().ToString() : TEXT("None"), *RowName.ToString()));
+	return FString::Printf(TEXT("%s[ %s ]"), DataTable ? *DataTable->GetFName().ToString() : TEXT("None"), *RowName.ToString());
 }
 
 bool FPCGDataTableRowToParamData::ExecuteInternal(FPCGContext* Context) const
diff --git a/Engine/Plugins/Experimental/PCG/Source/PCG/Private/Elements/PCGDeleteAttributesElement.cpp b/Engine/Plugins/Experimental/PCG/Source/PCG/Private/Elements/PCGDeleteAttributesElement.cpp
--- a/Engine/Plugins/Experimental/PCG/Source/PCG/Private/Elements/PCGDeleteAttributesElement.cpp
+++ b/Engine/Plugins/Experimental/PCG/Source/PCG/Private/Elements/PCGDeleteAttributesElement.cpp
@@ -67,25 +67,35 @@ TArray<FText> UPCGDeleteAttributesSettings::GetNodeTitleAliases() const
 }
 #endif
 
-FName UPCGDeleteAttributesSettings::AdditionalTaskName() const
+FString UPCGDeleteAttributesSettings::GetAdditionalTitleInformation() const
 {
-	if (const UEnum* SelectionEnum = StaticEnum<EPCGAttributeFilterOperation>())
+	// The display name for the operation is way too long when put in a node title, so abbreviate it here.
+	FString OperationString;
+	if (Operation == EPCGAttributeFilterOperation::KeepSelectedAttributes)
 	{
-		FText OperationText = SelectionEnum->GetDisplayNameTextByValue(static_cast<int64>(Operation));
-		TArray<FString> AttributesToKeep = PCGHelpers::GetStringArrayFromCommaSeparatedString(SelectedAttributes);
+		OperationString = LOCTEXT("OperationKeep", "Keep").ToString();
+	}
+	else if (Operation == EPCGAttributeFilterOperation::DeleteSelectedAttributes)
+	{
+		OperationString = LOCTEXT("OperationKeep", "Delete").ToString();
+	}
+	else
+	{
+		ensureMsgf(false, TEXT("Unrecognized operation"));
+	}
 
-		if (AttributesToKeep.Num() == 1)
-		{
-			return FName(FString::Printf(TEXT("%s (%s)"), *OperationText.ToString(), *AttributesToKeep[0]));
-		}
-		else
-		{
-			return FName(OperationText.ToString());
-		}
+	TArray<FString> AttributesToKeep = PCGHelpers::GetStringArrayFromCommaSeparatedString(SelectedAttributes);
+	if (AttributesToKeep.Num() == 1)
+	{
+		return FString::Printf(TEXT("%s (%s)"), *OperationString, *AttributesToKeep[0]);
+	}
+	else if (AttributesToKeep.IsEmpty())
+	{
+		return FString::Printf(TEXT("%s (%s)"), *OperationString, *LOCTEXT("NoAttributes", "none").ToString());
 	}
 	else
 	{
-		return NAME_None;
+		return FString::Printf(TEXT("%s (%s)"), *OperationString, *LOCTEXT("KeepMultipleAttributes", "multiple").ToString());
 	}
 }
 
diff --git a/Engine/Plugins/Experimental/PCG/Source/PCG/Private/Elements/PCGDeleteTags.cpp b/Engine/Plugins/Experimental/PCG/Source/PCG/Private/Elements/PCGDeleteTags.cpp
--- a/Engine/Plugins/Experimental/PCG/Source/PCG/Private/Elements/PCGDeleteTags.cpp
+++ b/Engine/Plugins/Experimental/PCG/Source/PCG/Private/Elements/PCGDeleteTags.cpp
@@ -20,25 +20,25 @@ FText UPCGDeleteTagsSettings::GetDefaultNodeTitle() const
 }
 #endif // WITH_EDITOR
 
-FName UPCGDeleteTagsSettings::AdditionalTaskName() const
+FString UPCGDeleteTagsSettings::GetAdditionalTitleInformation() const
 {
-	if (const UEnum* SelectionEnum = StaticEnum< EPCGTagFilterOperation>())
+	if (const UEnum* SelectionEnum = StaticEnum<EPCGTagFilterOperation>())
 	{
 		FText OperationText = SelectionEnum->GetDisplayNameTextByValue(static_cast<int64>(Operation));
 		TArray<FString> TagsToProcess = PCGHelpers::GetStringArrayFromCommaSeparatedString(SelectedTags);
 
 		if (TagsToProcess.Num() == 1)
 		{
-			return FName(FString::Printf(TEXT("%s (%s)"), *OperationText.ToString(), *TagsToProcess[0]));
+			return FString::Printf(TEXT("%s (%s)"), *OperationText.ToString(), *TagsToProcess[0]);
 		}
 		else
 		{
-			return FName(OperationText.ToString());
+			return OperationText.ToString();
 		}
 	}
 	else
 	{
-		return NAME_None;
+		return FString();
 	}
 }
 
diff --git a/Engine/Plugins/Experimental/PCG/Source/PCG/Private/Elements/PCGDifferenceElement.cpp b/Engine/Plugins/Experimental/PCG/Source/PCG/Private/Elements/PCGDifferenceElement.cpp
--- a/Engine/Plugins/Experimental/PCG/Source/PCG/Private/Elements/PCGDifferenceElement.cpp
+++ b/Engine/Plugins/Experimental/PCG/Source/PCG/Private/Elements/PCGDifferenceElement.cpp
@@ -84,6 +84,28 @@ void UPCGDifferenceSettings::ApplyStructuralDeprecation(UPCGNode* InOutNode)
 }
 #endif // WITH_EDITOR
 
+FString UPCGDifferenceSettings::GetAdditionalTitleInformation() const
+{
+#if WITH_EDITOR
+	const FProperty* DensityFunctionProperty = GetClass() ? FindFProperty<FProperty>(GetClass(), GET_MEMBER_NAME_CHECKED(UPCGDifferenceSettings, DensityFunction)) : nullptr;
+	if (DensityFunctionProperty && IsPropertyOverriddenByPin(DensityFunctionProperty))
+	{
+		return FString();
+	}
+	else
+#endif
+	{
+		if (const UEnum* EnumPtr = StaticEnum<EPCGDifferenceDensityFunction>())
+		{
+			return EnumPtr->GetNameStringByValue(static_cast<int>(DensityFunction));
+		}
+		else
+		{
+			return FString();
+		}
+	}
+}
+
 TArray<FPCGPinProperties> UPCGDifferenceSettings::InputPinProperties() const
 {
 	TArray<FPCGPinProperties> PinProperties;
diff --git a/Engine/Plugins/Experimental/PCG/Source/PCG/Private/Elements/PCGExecuteBlueprint.cpp b/Engine/Plugins/Experimental/PCG/Source/PCG/Private/Elements/PCGExecuteBlueprint.cpp
--- a/Engine/Plugins/Experimental/PCG/Source/PCG/Private/Elements/PCGExecuteBlueprint.cpp
+++ b/Engine/Plugins/Experimental/PCG/Source/PCG/Private/Elements/PCGExecuteBlueprint.cpp
@@ -553,23 +553,31 @@ void UPCGBlueprintSettings::ApplyPreconfiguredSettings(const FPCGPreConfiguredSe
 	}
 }
 
-FName UPCGBlueprintSettings::AdditionalTaskName() const
+FString UPCGBlueprintSettings::GetAdditionalTitleInformation() const
 {
 	if (BlueprintElementInstance && BlueprintElementInstance->NodeTitleOverride() != NAME_None)
 	{
-		return BlueprintElementInstance->NodeTitleOverride();
+		return BlueprintElementInstance->NodeTitleOverride().ToString();
 	}
 	else
 	{
-		FName ElementName = NAME_None;
+		FString ElementName;
 
 #if WITH_EDITOR
-		ElementName = (BlueprintElementType && BlueprintElementType->ClassGeneratedBy) ? BlueprintElementType->ClassGeneratedBy->GetFName() : Super::AdditionalTaskName();
+		ElementName = (BlueprintElementType && BlueprintElementType->ClassGeneratedBy) ? BlueprintElementType->ClassGeneratedBy->GetName() : Super::GetAdditionalTitleInformation();
 #else
-		ElementName = BlueprintElementType ? BlueprintElementType->GetFName() : Super::AdditionalTaskName();
+		ElementName = BlueprintElementType ? BlueprintElementType->GetName() : Super::GetAdditionalTitleInformation();
 #endif
+
 		// Normalize node name only if not explicitly set in the NodeTitleOverride call
-		return FName(FName::NameToDisplayString(ElementName.ToString(), false));
+		if (ElementName.IsEmpty())
+		{
+			return LOCTEXT("MissingBlueprint", "Missing Blueprint").ToString();
+		}
+		else
+		{
+			return FName::NameToDisplayString(ElementName, false);
+		}
 	}
 }
 
diff --git a/Engine/Plugins/Experimental/PCG/Source/PCG/Private/Elements/PCGFilterByAttribute.cpp b/Engine/Plugins/Experimental/PCG/Source/PCG/Private/Elements/PCGFilterByAttribute.cpp
--- a/Engine/Plugins/Experimental/PCG/Source/PCG/Private/Elements/PCGFilterByAttribute.cpp
+++ b/Engine/Plugins/Experimental/PCG/Source/PCG/Private/Elements/PCGFilterByAttribute.cpp
@@ -8,30 +8,25 @@
 
 #define LOCTEXT_NAMESPACE "PCGFilterByAttributeElement"
 
-namespace PCGFilterByAttributeConstants
-{
-	const FText NodeTitle = LOCTEXT("NodeTitle", "Filter Data By Attribute");
-}
-
 #if WITH_EDITOR
 FText UPCGFilterByAttributeSettings::GetDefaultNodeTitle() const
 {
-	return PCGFilterByAttributeConstants::NodeTitle;
+	return LOCTEXT("NodeTitle", "Filter Data By Attribute");
 }
 #endif // WITH_EDITOR
 
-FName UPCGFilterByAttributeSettings::AdditionalTaskName() const
+FString UPCGFilterByAttributeSettings::GetAdditionalTitleInformation() const
 {
 #if WITH_EDITOR
-	FProperty* AttributeProperty = GetClass() ? FindFProperty<FProperty>(GetClass(), GET_MEMBER_NAME_CHECKED(UPCGFilterByAttributeSettings, Attribute)) : nullptr;
+	const FProperty* AttributeProperty = GetClass() ? FindFProperty<FProperty>(GetClass(), GET_MEMBER_NAME_CHECKED(UPCGFilterByAttributeSettings, Attribute)) : nullptr;
 	if (AttributeProperty && IsPropertyOverriddenByPin(AttributeProperty))
 	{
-		return NAME_None;
+		return FString();
 	}
 	else
 #endif
 	{
-		return FName(FString::Printf(TEXT("%s (%s)"), *PCGFilterByAttributeConstants::NodeTitle.ToString(), *Attribute.ToString()));
+		return Attribute.ToString();
 	}
 }
 
diff --git a/Engine/Plugins/Experimental/PCG/Source/PCG/Private/Elements/PCGFilterByIndex.cpp b/Engine/Plugins/Experimental/PCG/Source/PCG/Private/Elements/PCGFilterByIndex.cpp
--- a/Engine/Plugins/Experimental/PCG/Source/PCG/Private/Elements/PCGFilterByIndex.cpp
+++ b/Engine/Plugins/Experimental/PCG/Source/PCG/Private/Elements/PCGFilterByIndex.cpp
@@ -13,9 +13,6 @@
 
 namespace PCGFilterByIndexConstants
 {
-	const FName NodeName = FName(TEXT("FilterDataByIndex"));
-	const FText NodeTitle = LOCTEXT("NodeTitle", "Filter Data By Index");
-
 	// After so many characters on the node, it will truncate with an ellipsis ...
 	static constexpr int32 IndexExpressionTruncation = 12;
 }
@@ -23,12 +20,12 @@ namespace PCGFilterByIndexConstants
 #if WITH_EDITOR
 FName UPCGFilterByIndexSettings::GetDefaultNodeName() const
 {
-	return PCGFilterByIndexConstants::NodeName;
+	return FName(TEXT("FilterDataByIndex"));
 }
 
 FText UPCGFilterByIndexSettings::GetDefaultNodeTitle() const
 {
-	return PCGFilterByIndexConstants::NodeTitle;
+	return LOCTEXT("NodeTitle", "Filter Data By Index");
 }
 
 FText UPCGFilterByIndexSettings::GetNodeTooltipText() const
@@ -37,12 +34,15 @@ FText UPCGFilterByIndexSettings::GetNodeTooltipText() const
 }
 #endif
 
-FName UPCGFilterByIndexSettings::AdditionalTaskName() const
+FString UPCGFilterByIndexSettings::GetAdditionalTitleInformation() const
 {
-	FString NodeName = PCGFilterByIndexConstants::NodeTitle.ToString();
-	NodeName += TEXT(" : ");
-	NodeName += SelectedIndices.Len() <= PCGFilterByIndexConstants::IndexExpressionTruncation ? SelectedIndices : SelectedIndices.Left(PCGFilterByIndexConstants::IndexExpressionTruncation - 3) + TEXT("...");
-	return FName(NodeName);
+	FString NodeName = TEXT("Index: ");
+
+	NodeName += SelectedIndices.Len() <= PCGFilterByIndexConstants::IndexExpressionTruncation
+		? SelectedIndices
+		: SelectedIndices.Left(PCGFilterByIndexConstants::IndexExpressionTruncation - 3) + TEXT("...");
+
+	return NodeName;
 }
 
 FPCGElementPtr UPCGFilterByIndexSettings::CreateElement() const
diff --git a/Engine/Plugins/Experimental/PCG/Source/PCG/Private/Elements/PCGFilterByTag.cpp b/Engine/Plugins/Experimental/PCG/Source/PCG/Private/Elements/PCGFilterByTag.cpp
--- a/Engine/Plugins/Experimental/PCG/Source/PCG/Private/Elements/PCGFilterByTag.cpp
+++ b/Engine/Plugins/Experimental/PCG/Source/PCG/Private/Elements/PCGFilterByTag.cpp
@@ -9,21 +9,15 @@
 
 #define LOCTEXT_NAMESPACE "PCGFilterByTag"
 
-namespace PCGFilterByTagConstants
-{
-	const FName NodeName = FName(TEXT("FilterDataByTag"));
-	const FText NodeTitle = LOCTEXT("NodeTitle", "Filter Data By Tag");
-}
-
 #if WITH_EDITOR
 FName UPCGFilterByTagSettings::GetDefaultNodeName() const
 {
-	return PCGFilterByTagConstants::NodeName;
+	return FName(TEXT("FilterDataByTag"));
 }
 
 FText UPCGFilterByTagSettings::GetDefaultNodeTitle() const
 {
-	return PCGFilterByTagConstants::NodeTitle;
+	return LOCTEXT("NodeTitle", "Filter Data By Tag");
 }
 
 FText UPCGFilterByTagSettings::GetNodeTooltipText() const
@@ -32,20 +26,23 @@ FText UPCGFilterByTagSettings::GetNodeTooltipText() const
 }
 #endif
 
-FName UPCGFilterByTagSettings::AdditionalTaskName() const
+FString UPCGFilterByTagSettings::GetAdditionalTitleInformation() const
 {
 	const TArray<FString> Tags = PCGHelpers::GetStringArrayFromCommaSeparatedString(SelectedTags);
 
-	FString NodeName = PCGFilterByTagConstants::NodeTitle.ToString();
-	NodeName += (Operation == EPCGFilterByTagOperation::KeepTagged ? TEXT(" (Keep)") : TEXT(" (Remove)"));
+	const FString Prefix = (Operation == EPCGFilterByTagOperation::KeepTagged ? TEXT("Tag (Keep):") : TEXT("Tag (Remove):"));
 
-	if (Tags.Num() == 1)
+	if (Tags.IsEmpty())
+	{
+		return Prefix;
+	}
+	else if (Tags.Num() == 1)
 	{
-		return FName(FString::Printf(TEXT("%s: %s"), *NodeName, *Tags[0]));
+		return FString::Printf(TEXT("%s %s"), *Prefix, *Tags[0]);
 	}
 	else
 	{
-		return FName(NodeName);
+		return FString::Printf(TEXT("%s (multiple)"), *Prefix);
 	}
 }
 
diff --git a/Engine/Plugins/Experimental/PCG/Source/PCG/Private/Elements/PCGGetActorProperty.cpp b/Engine/Plugins/Experimental/PCG/Source/PCG/Private/Elements/PCGGetActorProperty.cpp
--- a/Engine/Plugins/Experimental/PCG/Source/PCG/Private/Elements/PCGGetActorProperty.cpp
+++ b/Engine/Plugins/Experimental/PCG/Source/PCG/Private/Elements/PCGGetActorProperty.cpp
@@ -69,16 +69,15 @@ void UPCGGetActorPropertySettings::PostLoad()
 	}
 }
 
-FName UPCGGetActorPropertySettings::AdditionalTaskName() const
+FString UPCGGetActorPropertySettings::GetAdditionalTitleInformation() const
 {
 #if WITH_EDITOR
-	return ActorSelector.GetTaskName(GetDefaultNodeTitle());
+	return FString::Printf(TEXT("%s, %s"), *ActorSelector.GetTaskNameSuffix().ToString(), *PropertyName.ToString());
 #else
-	return Super::AdditionalTaskName();
+	return Super::GetAdditionalTitleInformation();
 #endif
 }
 
-
 TArray<FPCGPinProperties> UPCGGetActorPropertySettings::OutputPinProperties() const
 {
 	TArray<FPCGPinProperties> PinProperties;
diff --git a/Engine/Plugins/Experimental/PCG/Source/PCG/Private/Elements/PCGGetPropertyFromObjectPath.cpp b/Engine/Plugins/Experimental/PCG/Source/PCG/Private/Elements/PCGGetPropertyFromObjectPath.cpp
--- a/Engine/Plugins/Experimental/PCG/Source/PCG/Private/Elements/PCGGetPropertyFromObjectPath.cpp
+++ b/Engine/Plugins/Experimental/PCG/Source/PCG/Private/Elements/PCGGetPropertyFromObjectPath.cpp
@@ -51,6 +51,36 @@ bool UPCGGetPropertyFromObjectPathSettings::CanEditChange(const FProperty* InPro
 }
 #endif // WITH_EDITOR
 
+FString UPCGGetPropertyFromObjectPathSettings::GetAdditionalTitleInformation() const
+{
+	FString Path;
+
+	const UPCGNode* Node = Cast<UPCGNode>(GetOuter());
+	const bool bInPinIsConnected = Node ? Node->IsInputPinConnected(PCGPinConstants::DefaultInputLabel) : false;
+
+	// If the input pin is connected, don't display anything from the path.
+	if (!bInPinIsConnected)
+	{
+		if (ObjectPathsToExtract.IsEmpty())
+		{
+			Path = LOCTEXT("MissingPath", "Missing Path").ToString();
+		}
+		else if (ObjectPathsToExtract.Num() == 1)
+		{
+			Path = ObjectPathsToExtract[0].ToString();
+			Path = !Path.IsEmpty() ? Path : TEXT("None");
+		}
+		else
+		{
+			Path = LOCTEXT("MultiplePaths", "Multiple Paths").ToString();
+		}
+
+		Path += TEXT(", ");
+	}
+
+	return FString::Printf(TEXT("%s%s"), *Path, *PropertyName.ToString());
+}
+
 FPCGElementPtr UPCGGetPropertyFromObjectPathSettings::CreateElement() const
 {
 	return MakeShared<FPCGGetPropertyFromObjectPathElement>();
diff --git a/Engine/Plugins/Experimental/PCG/Source/PCG/Private/Elements/PCGHiGenGridSize.cpp b/Engine/Plugins/Experimental/PCG/Source/PCG/Private/Elements/PCGHiGenGridSize.cpp
--- a/Engine/Plugins/Experimental/PCG/Source/PCG/Private/Elements/PCGHiGenGridSize.cpp
+++ b/Engine/Plugins/Experimental/PCG/Source/PCG/Private/Elements/PCGHiGenGridSize.cpp
@@ -76,22 +76,21 @@ FPCGElementPtr UPCGHiGenGridSizeSettings::CreateElement() const
 	return MakeShared<FPCGHiGenGridSizeElement>();
 }
 
-FName UPCGHiGenGridSizeSettings::AdditionalTaskName() const
+FString UPCGHiGenGridSizeSettings::GetAdditionalTitleInformation() const
 {
 	if (HasAnyFlags(RF_ClassDefaultObject))
 	{
-		return PCGHiGenGridSizeConstants::NodeName;
+		return FString();
 	}
 
 	const UEnum* EnumPtr = StaticEnum<EPCGHiGenGrid>();
 	if (ensure(EnumPtr))
 	{
-		FText GridSizeDisplayName = EnumPtr->GetDisplayNameTextByValue(static_cast<int64>(HiGenGridSize));
-		return FName(FString::Printf(TEXT("%s: %s"), *PCGHiGenGridSizeConstants::NodeTitle.ToString(), *GridSizeDisplayName.ToString()));
+		return EnumPtr->GetDisplayNameTextByValue(static_cast<int64>(HiGenGridSize)).ToString();
 	}
 	else
 	{
-		return FName(FString::Printf(TEXT("%s: %d"), *PCGHiGenGridSizeConstants::NodeTitle.ToString(), static_cast<int32>(HiGenGridSize)));
+		return FString::Printf(TEXT("%d"), static_cast<int32>(HiGenGridSize));
 	}
 }
 
diff --git a/Engine/Plugins/Experimental/PCG/Source/PCG/Private/Elements/PCGIndirectionElement.cpp b/Engine/Plugins/Experimental/PCG/Source/PCG/Private/Elements/PCGIndirectionElement.cpp
--- a/Engine/Plugins/Experimental/PCG/Source/PCG/Private/Elements/PCGIndirectionElement.cpp
+++ b/Engine/Plugins/Experimental/PCG/Source/PCG/Private/Elements/PCGIndirectionElement.cpp
@@ -11,6 +11,11 @@
 
 #define LOCTEXT_NAMESPACE "PCGIndirectionElement"
 
+namespace PCGIndirectionSettings
+{
+	const FText DefaultNodeTitle = LOCTEXT("NodeTitle", "Proxy");
+}
+
 #if WITH_EDITOR
 
 FName UPCGIndirectionSettings::GetDefaultNodeName() const
@@ -25,40 +30,48 @@ FText UPCGIndirectionSettings::GetDefaultNodeTitle() const
 
 FText UPCGIndirectionSettings::GetNodeTooltipText() const
 {
-	return LOCTEXT("NodeTooltip", "Executes another settings object, which can be overriden.");
+	return LOCTEXT("NodeTooltip", "Executes another settings object, which can be overridden.");
 }
+#endif // WITH_EDITOR
 
-FName UPCGIndirectionSettings::AdditionalTaskName() const
+FString UPCGIndirectionSettings::GetAdditionalTitleInformation() const
 {
 	switch (ProxyInterfaceMode)
 	{
-		case EPCGProxyInterfaceMode::ByNativeElement:
-			if (IsValid(SettingsClass))
-			{
-				const UPCGSettings* SettingsDefaultObject = CastChecked<UPCGSettings>(SettingsClass->GetDefaultObject());
-				return FName(FText::Format(LOCTEXT("NodeTitleNative", "Proxy ({0})"), SettingsDefaultObject->GetDefaultNodeTitle()).ToString());
-			}
-			break;
-		case EPCGProxyInterfaceMode::ByBlueprintElement:
-			if (IsValid(BlueprintElementClass))
-			{
-				return FName(FText::Format(LOCTEXT("NodeTitleBlueprint", "Proxy ({0})"), BlueprintElementClass->GetDisplayNameText()).ToString());
-			}
-			break;
-		case EPCGProxyInterfaceMode::BySettings:
-			if (const UPCGSettings* SettingsPtr = Settings.LoadSynchronous())
-			{
-				return FName(FText::Format(LOCTEXT("NodeTitleNative", "Proxy ({0})"), FText::FromString(SettingsPtr->GetName())).ToString());
-			}
-			break;
+	case EPCGProxyInterfaceMode::ByNativeElement:
+		if (IsValid(SettingsClass))
+		{
+#if WITH_EDITOR
+			const UPCGSettings* SettingsDefaultObject = CastChecked<UPCGSettings>(SettingsClass->GetDefaultObject());
+			return SettingsDefaultObject->GetDefaultNodeTitle().ToString();
+#else
+			return SettingsClass->GetName();
+#endif
+		}
+		break;
+	case EPCGProxyInterfaceMode::ByBlueprintElement:
+		if (IsValid(BlueprintElementClass))
+		{
+#if WITH_EDITOR
+			return BlueprintElementClass->GetDisplayNameText().ToString();
+#else
+			return BlueprintElementClass->GetName();
+#endif
+		}
+		break;
+	case EPCGProxyInterfaceMode::BySettings:
+		if (const UPCGSettings* SettingsPtr = Settings.LoadSynchronous())
+		{
+			return SettingsPtr->GetName();
+		}
+		break;
 
-		default:
-			checkNoEntry();
+	default:
+		checkNoEntry();
 	}
 
-	return GetDefaultNodeName();
+	return LOCTEXT("MissingAsset", "Missing Asset").ToString();
 }
-#endif // WITH_EDITOR
 
 TArray<FPCGPinProperties> UPCGIndirectionSettings::InputPinProperties() const
 {
diff --git a/Engine/Plugins/Experimental/PCG/Source/PCG/Private/Elements/PCGLoopElement.cpp b/Engine/Plugins/Experimental/PCG/Source/PCG/Private/Elements/PCGLoopElement.cpp
--- a/Engine/Plugins/Experimental/PCG/Source/PCG/Private/Elements/PCGLoopElement.cpp
+++ b/Engine/Plugins/Experimental/PCG/Source/PCG/Private/Elements/PCGLoopElement.cpp
@@ -163,18 +163,6 @@ FPCGElementPtr UPCGLoopSettings::CreateElement() const
 	return MakeShared<FPCGLoopElement>();
 }
 
-FName UPCGLoopSettings::AdditionalTaskName() const
-{
-	if (UPCGGraph* TargetSubgraph = GetSubgraph())
-	{
-		return FName(FText::Format(LOCTEXT("NodeTitleExtended", "Loop Subgraph - {0}"), FText::FromName(TargetSubgraph->GetFName())).ToString());
-	}
-	else
-	{
-		return FName(LOCTEXT("NodeTitleExtendedInvalidSubgraph", "Loop Subgraph - Empty Subgraph").ToString());
-	}
-}
-
 bool FPCGLoopElement::ExecuteInternal(FPCGContext* InContext) const
 {
 	TRACE_CPUPROFILER_EVENT_SCOPE(FPCGLoopElement::Execute);
diff --git a/Engine/Plugins/Experimental/PCG/Source/PCG/Private/Elements/PCGTypedGetter.cpp b/Engine/Plugins/Experimental/PCG/Source/PCG/Private/Elements/PCGTypedGetter.cpp
--- a/Engine/Plugins/Experimental/PCG/Source/PCG/Private/Elements/PCGTypedGetter.cpp
+++ b/Engine/Plugins/Experimental/PCG/Source/PCG/Private/Elements/PCGTypedGetter.cpp
@@ -66,10 +66,10 @@ TArray<FPCGPinProperties> UPCGGetLandscapeSettings::OutputPinProperties() const
 	return PinProperties;
 }
 
-FName UPCGGetLandscapeSettings::AdditionalTaskName() const
+FString UPCGGetLandscapeSettings::GetAdditionalTitleInformation() const
 {
 	// Do not use the version from data from actor otherwise we'll show the selected actor class, which serves no purpose
-	return UPCGSettings::AdditionalTaskName();
+	return UPCGSettings::GetAdditionalTitleInformation();
 }
 
 #if WITH_EDITOR
diff --git a/Engine/Plugins/Experimental/PCG/Source/PCG/Private/Graph/PCGGraphCache.cpp b/Engine/Plugins/Experimental/PCG/Source/PCG/Private/Graph/PCGGraphCache.cpp
--- a/Engine/Plugins/Experimental/PCG/Source/PCG/Private/Graph/PCGGraphCache.cpp
+++ b/Engine/Plugins/Experimental/PCG/Source/PCG/Private/Graph/PCGGraphCache.cpp
@@ -64,7 +64,7 @@ bool FPCGGraphCache::GetFromCache(const UPCGNode* InNode, const IPCGElement* InE
 			if (bDebuggingEnabled)
 			{
 				// Leading spaces to align log content with warnings below - helps readability a lot.
-				UE_LOG(LogPCG, Log, TEXT("         [%s] %s\t\tCACHE HIT %u"), *InComponent->GetOwner()->GetName(), *InNode->GetNodeTitle().ToString(), InDependenciesCrc.GetValue());
+				UE_LOG(LogPCG, Log, TEXT("         [%s] %s\t\tCACHE HIT %u"), *InComponent->GetOwner()->GetName(), *InNode->GetNodeTitle(EPCGNodeTitleType::ListView).ToString(), InDependenciesCrc.GetValue());
 			}
 
 			OutOutput = *Value;
@@ -75,7 +75,7 @@ bool FPCGGraphCache::GetFromCache(const UPCGNode* InNode, const IPCGElement* InE
 		{
 			if (bDebuggingEnabled)
 			{
-				UE_LOG(LogPCG, Warning, TEXT("[%s] %s\t\tCACHE MISS %u"), *InComponent->GetOwner()->GetName(), *InNode->GetNodeTitle().ToString(), InDependenciesCrc.GetValue());
+				UE_LOG(LogPCG, Warning, TEXT("[%s] %s\t\tCACHE MISS %u"), *InComponent->GetOwner()->GetName(), *InNode->GetNodeTitle(EPCGNodeTitleType::ListView).ToString(), InDependenciesCrc.GetValue());
 			}
 
 			return false;
diff --git a/Engine/Plugins/Experimental/PCG/Source/PCG/Private/Graph/PCGGraphExecutor.cpp b/Engine/Plugins/Experimental/PCG/Source/PCG/Private/Graph/PCGGraphExecutor.cpp
--- a/Engine/Plugins/Experimental/PCG/Source/PCG/Private/Graph/PCGGraphExecutor.cpp
+++ b/Engine/Plugins/Experimental/PCG/Source/PCG/Private/Graph/PCGGraphExecutor.cpp
@@ -653,7 +653,7 @@ void FPCGGraphExecutor::Execute()
 
 				if (bGraphCacheDebuggingEnabled && !bCacheable && Task.SourceComponent.Get() && Task.Node)
 				{
-					UE_LOG(LogPCG, Warning, TEXT("[%s] %s\t\tCACHING DISABLED"), *Task.SourceComponent->GetOwner()->GetName(), *Task.Node->GetNodeTitle().ToString());
+					UE_LOG(LogPCG, Warning, TEXT("[%s] %s\t\tCACHING DISABLED"), *Task.SourceComponent->GetOwner()->GetName(), *Task.Node->GetNodeTitle(EPCGNodeTitleType::ListView).ToString());
 				}
 
 				FPCGDataCollection CachedOutput;
@@ -797,7 +797,7 @@ void FPCGGraphExecutor::Execute()
 			{
 				if (bGraphCacheDebuggingEnabled && ActiveTask.Context->SourceComponent.Get() && ActiveTask.Context->Node)
 				{
-					UE_LOG(LogPCG, Log, TEXT("         [%s] %s\t\tOUTPUT CRC %u"), *ActiveTask.Context->SourceComponent->GetOwner()->GetName(), *ActiveTask.Context->Node->GetNodeTitle().ToString(), ActiveTask.Context->OutputData.Crc.GetValue());
+					UE_LOG(LogPCG, Log, TEXT("         [%s] %s\t\tOUTPUT CRC %u"), *ActiveTask.Context->SourceComponent->GetOwner()->GetName(), *ActiveTask.Context->Node->GetNodeTitle(EPCGNodeTitleType::ListView).ToString(), ActiveTask.Context->OutputData.Crc.GetValue());
 				}
 
 				// Store result in cache as needed - done here because it needs to be done on the main thread
@@ -841,7 +841,7 @@ void FPCGGraphExecutor::Execute()
 				const int32 NumRemoved = ActiveTask.Context->OutputData.StripEmptyPointData();
 				if (NumRemoved > 0)
 				{
-					UE_LOG(LogPCG, Log, TEXT("%d empty point data stripped from node: %s"), NumRemoved, *ActiveTask.Context->Node->GetNodeTitle().ToString());
+					UE_LOG(LogPCG, Log, TEXT("%d empty point data stripped from node: %s"), NumRemoved, *ActiveTask.Context->Node->GetNodeTitle(EPCGNodeTitleType::ListView).ToString());
 				}
 			}
 
@@ -1096,7 +1096,7 @@ void FPCGGraphExecutor::BuildTaskInput(const FPCGGraphTask& Task, FPCGDataCollec
 		// Log only - currently context has not yet been allocated when this is called
 		UE_LOG(LogPCG, Warning, TEXT("[%s] %s - BuildTaskInput - too many data items arriving on single data pin '%s', only first data item will be used"),
 			(Task.SourceComponent.Get() && Task.SourceComponent->GetOwner()) ? *Task.SourceComponent->GetOwner()->GetName() : TEXT("MissingComponent"),
-			Task.Node ? *Task.Node->GetNodeTitle().ToString() : TEXT("MissingNode"),
+			Task.Node ? *Task.Node->GetNodeTitle(EPCGNodeTitleType::ListView).ToString() : TEXT("MissingNode"),
 			InPin ? *InPin->Properties.Label.ToString() : TEXT("MissingPin"));
 	};
 
diff --git a/Engine/Plugins/Experimental/PCG/Source/PCG/Private/Helpers/PCGSubgraphHelpers.cpp b/Engine/Plugins/Experimental/PCG/Source/PCG/Private/Helpers/PCGSubgraphHelpers.cpp
--- a/Engine/Plugins/Experimental/PCG/Source/PCG/Private/Helpers/PCGSubgraphHelpers.cpp
+++ b/Engine/Plugins/Experimental/PCG/Source/PCG/Private/Helpers/PCGSubgraphHelpers.cpp
@@ -26,7 +26,7 @@ namespace PCGSubgraphHelpersExtra
 	{
 		UPCGNode* NewInputOutputNode = bIsInput ? InGraph->GetInputNode() : InGraph->GetOutputNode();
 		UPCGGraphInputOutputSettings* NewInputOutputSettings = CastChecked<UPCGGraphInputOutputSettings>(NewInputOutputNode->GetSettings());
-		FString NewName = InPinToClone->Node->GetNodeTitle().ToString() + " " + InPinToClone->Properties.Label.ToString();
+		FString NewName = InPinToClone->Node->GetNodeTitle(EPCGNodeTitleType::ListView).ToString() + " " + InPinToClone->Properties.Label.ToString();
 		if (InOutNameCollisionMapping.Contains(NewName))
 		{
 			if (!bIsInput)
@@ -257,7 +257,7 @@ UPCGGraph* FPCGSubgraphHelpers::CollapseIntoSubgraph(UPCGGraph* InOriginalGraph,
 				{
 					// It is only problematic if the pin was connected
 					UE_LOG(LogPCG, Error, TEXT("[CollapseInSubgraph - %s] %s pin %s does not exist anymore. Edges will be broken."),
-						*PCGNode->GetNodeTitle().ToString(), (OriginalPin->IsOutputPin() ? TEXT("Output") : TEXT("Input")), *PinLabel.ToString());
+						*PCGNode->GetNodeTitle(EPCGNodeTitleType::ListView).ToString(), (OriginalPin->IsOutputPin() ? TEXT("Output") : TEXT("Input")), *PinLabel.ToString());
 				}
 			}
 		};
diff --git a/Engine/Plugins/Experimental/PCG/Source/PCG/Private/PCGContext.cpp b/Engine/Plugins/Experimental/PCG/Source/PCG/Private/PCGContext.cpp
--- a/Engine/Plugins/Experimental/PCG/Source/PCG/Private/PCGContext.cpp
+++ b/Engine/Plugins/Experimental/PCG/Source/PCG/Private/PCGContext.cpp
@@ -29,15 +29,15 @@ FString FPCGContext::GetTaskName() const
 		const FName NodeName = ((Node->NodeTitle != NAME_None) ? Node->NodeTitle : Node->GetFName());
 
 		const UPCGSettings* Settings = GetInputSettings<UPCGSettings>();
-		const FName NodeAdditionalName = Settings ? Settings->AdditionalTaskName() : NAME_None;
+		const FString NodeAdditionalInformation = Settings ? Settings->GetAdditionalTitleInformation() : FString();
 
-		if (NodeAdditionalName == NAME_None || NodeAdditionalName == NodeName)
+		if (NodeAdditionalInformation.IsEmpty() || NodeAdditionalInformation == NodeName)
 		{
 			return NodeName.ToString();
 		}
 		else
 		{
-			return FString::Printf(TEXT("%s (%s)"), *NodeName.ToString(), *NodeAdditionalName.ToString());
+			return FString::Printf(TEXT("%s (%s)"), *NodeName.ToString(), *NodeAdditionalInformation);
 		}
 	}
 	else
diff --git a/Engine/Plugins/Experimental/PCG/Source/PCG/Private/PCGData.cpp b/Engine/Plugins/Experimental/PCG/Source/PCG/Private/PCGData.cpp
--- a/Engine/Plugins/Experimental/PCG/Source/PCG/Private/PCGData.cpp
+++ b/Engine/Plugins/Experimental/PCG/Source/PCG/Private/PCGData.cpp
@@ -266,7 +266,7 @@ UPCGParamData* FPCGDataCollection::GetParamsWithDeprecation(const UPCGNode* Node
 			{
 				if (const UPCGParamData* Params = Cast<UPCGParamData>(TaggedDatum.Data))
 				{
-					UE_LOG(LogPCG, Warning, TEXT("[%s] Found an Attribute Set data on an input pin that should not accept attributes. Make sure to re-wire it to the Overrides pin if it is used for overrides."), *Node->GetNodeTitle().ToString());
+					UE_LOG(LogPCG, Warning, TEXT("[%s] Found an Attribute Set data on an input pin that should not accept attributes. Make sure to re-wire it to the Overrides pin if it is used for overrides."), *Node->GetNodeTitle(EPCGNodeTitleType::ListView).ToString());
 					return const_cast<UPCGParamData*>(Params);
 				}
 			}
diff --git a/Engine/Plugins/Experimental/PCG/Source/PCG/Private/PCGNode.cpp b/Engine/Plugins/Experimental/PCG/Source/PCG/Private/PCGNode.cpp
--- a/Engine/Plugins/Experimental/PCG/Source/PCG/Private/PCGNode.cpp
+++ b/Engine/Plugins/Experimental/PCG/Source/PCG/Private/PCGNode.cpp
@@ -7,12 +7,36 @@
 #include "PCGGraph.h"
 #include "PCGModule.h"
 #include "PCGPin.h"
+#include "PCGSubsystem.h"
 
 #include "Algo/Find.h"
 #include "UObject/Package.h"
 
+#if WITH_EDITOR
+#include "Editor.h"
+#endif
+
 #include UE_INLINE_GENERATED_CPP_BY_NAME(PCGNode)
 
+namespace PCGNodeHelpers
+{
+#if WITH_EDITOR
+	// Info to aid element cache analysis / debugging
+	void GetGraphCacheDebugInfo(const UPCGNode* InNode, bool& bOutDebuggingEnabled, uint32& OutNumCacheEntries)
+	{
+		UWorld* World = GEditor ? (GEditor->PlayWorld ? GEditor->PlayWorld.Get() : GEditor->GetEditorWorldContext().World()) : nullptr;
+		UPCGSubsystem* Subsystem = UPCGSubsystem::GetInstance(World);
+		bOutDebuggingEnabled = Subsystem && Subsystem->IsGraphCacheDebuggingEnabled();
+
+		if (bOutDebuggingEnabled)
+		{
+			IPCGElement* Element = (InNode && InNode->GetSettings()) ? InNode->GetSettings()->GetElement().Get() : nullptr;
+			OutNumCacheEntries = Element ? Subsystem->GetGraphCacheEntryCount(Element) : 0;
+		}
+	}
+#endif
+}
+
 UPCGNode::UPCGNode(const FObjectInitializer& ObjectInitializer)
 	: Super(ObjectInitializer)
 {
@@ -245,27 +269,147 @@ bool UPCGNode::RemoveEdgeTo(FName FromPinLabel, UPCGNode* To, FName ToPinLabel)
 	}
 }
 
-FText UPCGNode::GetNodeTitle() const
+FText UPCGNode::GetNodeTitle(EPCGNodeTitleType TitleType) const
 {
-	if (NodeTitle != NAME_None)
+	// Title length that looks reasonable.
+	constexpr int NodeTitleMaxLen = 45;
+
+	// Clip string at right hand side (standard overflow).
+	auto ClipRightSide = [NodeTitleMaxLen](FString& InOutTitle, int MaxLen)
 	{
-		return FText::FromString(FName::NameToDisplayString(NodeTitle.ToString(), false));
+		if (InOutTitle.Len() > MaxLen)
+		{
+			InOutTitle = InOutTitle.Left(MaxLen - 3) + TEXT("...");
+		}
+	};
+
+	// Clip left hand side to maximize displayed information at end of string.
+	auto ClipLeftSide = [NodeTitleMaxLen](FString& InOutTitle, int MaxLen)
+	{
+		if (InOutTitle.Len() > MaxLen)
+		{
+			InOutTitle = TEXT("...") + InOutTitle.Right(MaxLen - 3);
+		}
+	};
+
+	const bool bFlipTitleLines = HasFlippedTitleLines();
+	const bool bIsTitleAuthored = HasAuthoredTitle();
+
+	FString GeneratedTitle = GetGeneratedTitleLine().ToString();
+	const bool bHasMultipleTitleLines = !GeneratedTitle.IsEmpty();
+
+	FString PrimaryTitleLine;
+
+	// The normal title line is used if there is no generated title, or if the node has not requested flipped
+	// title lines. But if the user has edited the title, always use the users title as primary. 
+	if (!bHasMultipleTitleLines || !bFlipTitleLines || bIsTitleAuthored)
+	{
+		PrimaryTitleLine = GetAuthoredTitleLine().ToString();
+		ClipRightSide(PrimaryTitleLine, NodeTitleMaxLen);
+	}
+	else
+	{
+		PrimaryTitleLine = GeneratedTitle;
+
+		// Generated title clipped on left side because the end of the string often has the interesting part.
+		ClipLeftSide(PrimaryTitleLine, NodeTitleMaxLen);
 	}
-	else if (UPCGSettings* Settings = GetSettings())
+
+	// Only add a second title line if the full title is being requested (and if the has multiple title lines).
+	if (TitleType == EPCGNodeTitleType::FullTitle)
 	{
-		if (Settings->AdditionalTaskName() != NAME_None)
+		FString SecondaryTitleLine;
+
+		if (bHasMultipleTitleLines)
 		{
-			return FText::FromName(Settings->AdditionalTaskName());
+			// Secondary title is the node-generated title line if titles aren't flipped.
+			if (!bFlipTitleLines)
+			{
+				SecondaryTitleLine = GeneratedTitle;
+
+				// Generated title clipped on left side because the end of the string often has the interesting part.
+				ClipLeftSide(SecondaryTitleLine, NodeTitleMaxLen);
+			}
+			// If titles are flipped and user has not authored something, then just show the standard node name.
+			else if (!bIsTitleAuthored)
+			{
+				SecondaryTitleLine = GetDefaultTitle().ToString();
+				ClipRightSide(SecondaryTitleLine, NodeTitleMaxLen);
+			}
+			// If title has been authored and title lines are flipped, display "<DefaultTitle> - <GeneratedTitle>"
+			else
+			{
+				const FString DefaultTitle = GetDefaultTitle().ToString();
+
+				// Generated title clipped on left side because the end of the string often has the interesting part.
+				ClipLeftSide(GeneratedTitle, NodeTitleMaxLen - (DefaultTitle.Len() + 3));
+				SecondaryTitleLine = DefaultTitle + TEXT(" - ") + GeneratedTitle;
+			}
 		}
+
+		// Debug info - append how many copies of this element are currently in the cache to the node title.
 #if WITH_EDITOR
-		else
+		bool bDebuggingEnabled = false;
+		uint32 NumCacheEntries = 0;
+		PCGNodeHelpers::GetGraphCacheDebugInfo(this, bDebuggingEnabled, NumCacheEntries);
+
+		if (bDebuggingEnabled)
 		{
-			return Settings->GetDefaultNodeTitle();
+			SecondaryTitleLine += FString::Format(TEXT(" [{1}]"), { NumCacheEntries });
 		}
 #endif
+
+		if (!SecondaryTitleLine.IsEmpty())
+		{
+			return FText::Format(FText::FromString("{0}\r\n{1}"), FText::FromString(PrimaryTitleLine), FText::FromString(SecondaryTitleLine));
+		}
+	}
+
+	return FText::FromString(PrimaryTitleLine);
+}
+
+FText UPCGNode::GetDefaultTitle() const
+{
+#if WITH_EDITOR
+	if (UPCGSettings* Settings = GetSettings())
+	{
+		return Settings->GetDefaultNodeTitle();
+	}
+#endif
+
+	return NSLOCTEXT("PCGNode", "NodeTitle", "Unnamed Node");
+}
+
+FText UPCGNode::GetAuthoredTitleLine() const
+{
+	if (NodeTitle != NAME_None)
+	{
+		return FText::FromString(FName::NameToDisplayString(NodeTitle.ToString(), false));
+	}
+	else
+	{
+		return GetDefaultTitle();
+	}
+}
+
+bool UPCGNode::HasFlippedTitleLines() const
+{
+	const UPCGSettings* Settings = GetSettings();
+	return Settings && Settings->HasFlippedTitleLines();
+}
+
+FText UPCGNode::GetGeneratedTitleLine() const
+{
+	if (UPCGSettings* Settings = GetSettings())
+	{
+		const FString AdditionalInformation = Settings->GetAdditionalTitleInformation();
+		if (!AdditionalInformation.IsEmpty())
+		{
+			return FText::FromString(AdditionalInformation);
+		}
 	}
 
-	return NSLOCTEXT("PCGNode", "NodeTitle", "Unnamed node");
+	return FText::GetEmpty();
 }
 
 #if WITH_EDITOR
diff --git a/Engine/Plugins/Experimental/PCG/Source/PCG/Private/PCGSubgraph.cpp b/Engine/Plugins/Experimental/PCG/Source/PCG/Private/PCGSubgraph.cpp
--- a/Engine/Plugins/Experimental/PCG/Source/PCG/Private/PCGSubgraph.cpp
+++ b/Engine/Plugins/Experimental/PCG/Source/PCG/Private/PCGSubgraph.cpp
@@ -277,15 +277,22 @@ UPCGNode* UPCGSubgraphSettings::CreateNode() const
 	return NewObject<UPCGSubgraphNode>();
 }
 
-FName UPCGSubgraphSettings::AdditionalTaskName() const
+FString UPCGSubgraphSettings::GetAdditionalTitleInformation() const
 {
+	const UPCGNode* Node = Cast<UPCGNode>(GetOuter());
+	if (Node && Node->IsInputPinConnected(PCGPinConstants::DefaultInputLabel))
+	{
+		// If subgraph is dynamic, don't emit any extra info here.
+		return FString();
+	}
+
 	if (UPCGGraph* TargetSubgraph = GetSubgraph())
 	{
-		return TargetSubgraph->GetFName();
+		return TargetSubgraph->GetName();
 	}
 	else
 	{
-		return TEXT("Empty subgraph");
+		return LOCTEXT("NodeTitleExtendedInvalidSubgraph", "Missing Subgraph").ToString();
 	}
 }
 
diff --git a/Engine/Plugins/Experimental/PCG/Source/PCG/Private/Tests/Determinism/PCGDeterminismTestBlueprintBase.cpp b/Engine/Plugins/Experimental/PCG/Source/PCG/Private/Tests/Determinism/PCGDeterminismTestBlueprintBase.cpp
--- a/Engine/Plugins/Experimental/PCG/Source/PCG/Private/Tests/Determinism/PCGDeterminismTestBlueprintBase.cpp
+++ b/Engine/Plugins/Experimental/PCG/Source/PCG/Private/Tests/Determinism/PCGDeterminismTestBlueprintBase.cpp
@@ -9,7 +9,7 @@
 void UPCGDeterminismTestBlueprintBase::ExecuteTest_Implementation(const UPCGNode* InPCGNode, FDeterminismTestResult& InOutTestResult)
 {
 	// The default case if the user tries to run the blueprint test unimplemented
-	InOutTestResult.TestResultTitle = FName(*InPCGNode->GetNodeTitle().ToString());
+	InOutTestResult.TestResultTitle = FName(*InPCGNode->GetNodeTitle(EPCGNodeTitleType::ListView).ToString());
 	InOutTestResult.TestResultName = InPCGNode->GetName();
 	InOutTestResult.Seed = InPCGNode->GetSettings()->GetSeed();
 	PCGDeterminismTests::RunDeterminismTest(InPCGNode, InOutTestResult, PCGDeterminismTests::Defaults::DeterminismBasicTestInfo);
diff --git a/Engine/Plugins/Experimental/PCG/Source/PCG/Private/Tests/Subgraphs/PCGCollapseSubgraphTest.cpp b/Engine/Plugins/Experimental/PCG/Source/PCG/Private/Tests/Subgraphs/PCGCollapseSubgraphTest.cpp
--- a/Engine/Plugins/Experimental/PCG/Source/PCG/Private/Tests/Subgraphs/PCGCollapseSubgraphTest.cpp
+++ b/Engine/Plugins/Experimental/PCG/Source/PCG/Private/Tests/Subgraphs/PCGCollapseSubgraphTest.cpp
@@ -59,8 +59,18 @@ namespace PCGCollapseSubgraphTests
 		UPCGPin* InPin = InNode->GetOutputPin(InLabel);
 		UPCGPin* OutPin = OutNode->GetInputPin(OutLabel);
 
-		if (!InTestClass->TestTrue(*FString::Printf(TEXT("InLabel (%s) exists on node %s"), *InLabel.ToString(), *InNode->GetNodeTitle().ToString()), (InLabel == NAME_None) || InPin)) { return false; }
-		if (!InTestClass->TestTrue(*FString::Printf(TEXT("OutLabel (%s) exists on node %s"), *OutLabel.ToString(), *OutNode->GetNodeTitle().ToString()), (OutLabel == NAME_None) || OutPin)) { return false; }
+		if (!InTestClass->TestTrue(
+			*FString::Printf(TEXT("InLabel (%s) exists on node %s"), *InLabel.ToString(), *InNode->GetNodeTitle(EPCGNodeTitleType::ListView).ToString()),
+			(InLabel == NAME_None) || InPin))
+		{
+			return false;
+		}
+		if (!InTestClass->TestTrue(
+			*FString::Printf(TEXT("OutLabel (%s) exists on node %s"), *OutLabel.ToString(), *OutNode->GetNodeTitle(EPCGNodeTitleType::ListView).ToString()),
+			(OutLabel == NAME_None) || OutPin))
+		{
+			return false;
+		}
 
 		check(InPin || OutPin);
 
@@ -82,7 +92,13 @@ namespace PCGCollapseSubgraphTests
 			}
 		};
 
-		return InTestClass->TestTrue(*FString::Printf(TEXT("Edge between %s(%s) and %s(%s) exists"), *InNode->GetNodeTitle().ToString(), *InLabel.ToString(), *OutNode->GetNodeTitle().ToString(), *OutLabel.ToString()), Edges.ContainsByPredicate(Predicate));
+		return InTestClass->TestTrue(
+			*FString::Printf(TEXT("Edge between %s(%s) and %s(%s) exists"),
+				*InNode->GetNodeTitle(EPCGNodeTitleType::ListView).ToString(),
+				*InLabel.ToString(),
+				*OutNode->GetNodeTitle(EPCGNodeTitleType::ListView).ToString(),
+				*OutLabel.ToString()),
+			Edges.ContainsByPredicate(Predicate));
 	}
 
 	UPCGNode* AddGetGraphParameterNode(UPCGGraph* InGraph, const FName InParamName)
diff --git a/Engine/Plugins/Experimental/PCG/Source/PCG/Public/Elements/Metadata/PCGMetadataBitwiseOpElement.h b/Engine/Plugins/Experimental/PCG/Source/PCG/Public/Elements/Metadata/PCGMetadataBitwiseOpElement.h
--- a/Engine/Plugins/Experimental/PCG/Source/PCG/Public/Elements/Metadata/PCGMetadataBitwiseOpElement.h
+++ b/Engine/Plugins/Experimental/PCG/Source/PCG/Public/Elements/Metadata/PCGMetadataBitwiseOpElement.h
@@ -32,7 +32,7 @@ public:
 	virtual TArray<FPCGPreConfiguredSettingsInfo> GetPreconfiguredInfo() const override;
 	virtual bool OnlyExposePreconfiguredSettings() const override { return true; }
 #endif
-	virtual FName AdditionalTaskName() const override;
+	virtual FString GetAdditionalTitleInformation() const override;
 	virtual void ApplyPreconfiguredSettings(const FPCGPreConfiguredSettingsInfo& PreconfigureInfo) override;
 	//~End UPCGSettings interface
 
diff --git a/Engine/Plugins/Experimental/PCG/Source/PCG/Public/Elements/Metadata/PCGMetadataBooleanOpElement.h b/Engine/Plugins/Experimental/PCG/Source/PCG/Public/Elements/Metadata/PCGMetadataBooleanOpElement.h
--- a/Engine/Plugins/Experimental/PCG/Source/PCG/Public/Elements/Metadata/PCGMetadataBooleanOpElement.h
+++ b/Engine/Plugins/Experimental/PCG/Source/PCG/Public/Elements/Metadata/PCGMetadataBooleanOpElement.h
@@ -32,7 +32,7 @@ public:
 	virtual TArray<FPCGPreConfiguredSettingsInfo> GetPreconfiguredInfo() const override;
 	virtual bool OnlyExposePreconfiguredSettings() const override { return true; }
 #endif
-	virtual FName AdditionalTaskName() const override;
+	virtual FString GetAdditionalTitleInformation() const override;
 	virtual void ApplyPreconfiguredSettings(const FPCGPreConfiguredSettingsInfo& PreconfigureInfo) override;
 	//~End UPCGSettings interface
 
diff --git a/Engine/Plugins/Experimental/PCG/Source/PCG/Public/Elements/Metadata/PCGMetadataCompareOpElement.h b/Engine/Plugins/Experimental/PCG/Source/PCG/Public/Elements/Metadata/PCGMetadataCompareOpElement.h
--- a/Engine/Plugins/Experimental/PCG/Source/PCG/Public/Elements/Metadata/PCGMetadataCompareOpElement.h
+++ b/Engine/Plugins/Experimental/PCG/Source/PCG/Public/Elements/Metadata/PCGMetadataCompareOpElement.h
@@ -34,7 +34,7 @@ public:
 	virtual TArray<FPCGPreConfiguredSettingsInfo> GetPreconfiguredInfo() const override;
 	virtual bool OnlyExposePreconfiguredSettings() const override { return true; }
 #endif
-	virtual FName AdditionalTaskName() const override;
+	virtual FString GetAdditionalTitleInformation() const override;
 	virtual void ApplyPreconfiguredSettings(const FPCGPreConfiguredSettingsInfo& PreconfigureInfo) override;
 	//~End UPCGSettings interface
 
diff --git a/Engine/Plugins/Experimental/PCG/Source/PCG/Public/Elements/Metadata/PCGMetadataElement.h b/Engine/Plugins/Experimental/PCG/Source/PCG/Public/Elements/Metadata/PCGMetadataElement.h
--- a/Engine/Plugins/Experimental/PCG/Source/PCG/Public/Elements/Metadata/PCGMetadataElement.h
+++ b/Engine/Plugins/Experimental/PCG/Source/PCG/Public/Elements/Metadata/PCGMetadataElement.h
@@ -36,7 +36,7 @@ public:
 	virtual bool IsPinUsedByNodeExecution(const UPCGPin* InPin) const override;
 	virtual void ApplyDeprecation(UPCGNode* InOutNode) override;
 #endif
-	virtual FName AdditionalTaskName() const override;
+	virtual FString GetAdditionalTitleInformation() const override;
 
 protected:
 #if WITH_EDITOR
diff --git a/Engine/Plugins/Experimental/PCG/Source/PCG/Public/Elements/Metadata/PCGMetadataMathsOpElement.h b/Engine/Plugins/Experimental/PCG/Source/PCG/Public/Elements/Metadata/PCGMetadataMathsOpElement.h
--- a/Engine/Plugins/Experimental/PCG/Source/PCG/Public/Elements/Metadata/PCGMetadataMathsOpElement.h
+++ b/Engine/Plugins/Experimental/PCG/Source/PCG/Public/Elements/Metadata/PCGMetadataMathsOpElement.h
@@ -59,7 +59,7 @@ public:
 	virtual TArray<FPCGPreConfiguredSettingsInfo> GetPreconfiguredInfo() const override;
 	virtual bool OnlyExposePreconfiguredSettings() const override { return true; }
 #endif
-	virtual FName AdditionalTaskName() const override;
+	virtual FString GetAdditionalTitleInformation() const override;
 	virtual void ApplyPreconfiguredSettings(const FPCGPreConfiguredSettingsInfo& PreconfigureInfo) override;
 	//~End UPCGSettings interface
 
diff --git a/Engine/Plugins/Experimental/PCG/Source/PCG/Public/Elements/Metadata/PCGMetadataOpElementBase.h b/Engine/Plugins/Experimental/PCG/Source/PCG/Public/Elements/Metadata/PCGMetadataOpElementBase.h
--- a/Engine/Plugins/Experimental/PCG/Source/PCG/Public/Elements/Metadata/PCGMetadataOpElementBase.h
+++ b/Engine/Plugins/Experimental/PCG/Source/PCG/Public/Elements/Metadata/PCGMetadataOpElementBase.h
@@ -92,6 +92,7 @@ public:
 	virtual FText GetNodeTooltipText() const override;
 	virtual void ApplyDeprecation(UPCGNode* InOutNode) override;
 #endif
+	virtual bool HasFlippedTitleLines() const override { return true; }
 	virtual EPCGDataType GetCurrentPinTypes(const UPCGPin* InPin) const override;
 
 protected:
diff --git a/Engine/Plugins/Experimental/PCG/Source/PCG/Public/Elements/Metadata/PCGMetadataPartition.h b/Engine/Plugins/Experimental/PCG/Source/PCG/Public/Elements/Metadata/PCGMetadataPartition.h
--- a/Engine/Plugins/Experimental/PCG/Source/PCG/Public/Elements/Metadata/PCGMetadataPartition.h
+++ b/Engine/Plugins/Experimental/PCG/Source/PCG/Public/Elements/Metadata/PCGMetadataPartition.h
@@ -24,6 +24,7 @@ public:
 	virtual EPCGSettingsType GetType() const override { return EPCGSettingsType::Metadata; }
 	virtual bool HasDynamicPins() const override { return true; }
 #endif
+	virtual FString GetAdditionalTitleInformation() const override;
 
 protected:
 	virtual TArray<FPCGPinProperties> InputPinProperties() const override;
diff --git a/Engine/Plugins/Experimental/PCG/Source/PCG/Public/Elements/Metadata/PCGMetadataRenameElement.h b/Engine/Plugins/Experimental/PCG/Source/PCG/Public/Elements/Metadata/PCGMetadataRenameElement.h
--- a/Engine/Plugins/Experimental/PCG/Source/PCG/Public/Elements/Metadata/PCGMetadataRenameElement.h
+++ b/Engine/Plugins/Experimental/PCG/Source/PCG/Public/Elements/Metadata/PCGMetadataRenameElement.h
@@ -19,6 +19,7 @@ public:
 	virtual EPCGSettingsType GetType() const override { return EPCGSettingsType::Metadata; }
 	virtual bool HasDynamicPins() const override { return true; }
 #endif
+	virtual FString GetAdditionalTitleInformation() const override;
 
 protected:
 	virtual TArray<FPCGPinProperties> InputPinProperties() const override;
diff --git a/Engine/Plugins/Experimental/PCG/Source/PCG/Public/Elements/Metadata/PCGMetadataRotatorOpElement.h b/Engine/Plugins/Experimental/PCG/Source/PCG/Public/Elements/Metadata/PCGMetadataRotatorOpElement.h
--- a/Engine/Plugins/Experimental/PCG/Source/PCG/Public/Elements/Metadata/PCGMetadataRotatorOpElement.h
+++ b/Engine/Plugins/Experimental/PCG/Source/PCG/Public/Elements/Metadata/PCGMetadataRotatorOpElement.h
@@ -43,7 +43,7 @@ public:
 	virtual TArray<FPCGPreConfiguredSettingsInfo> GetPreconfiguredInfo() const override;
 	virtual bool OnlyExposePreconfiguredSettings() const override { return true; }
 #endif
-	virtual FName AdditionalTaskName() const override;
+	virtual FString GetAdditionalTitleInformation() const override;
 	virtual void ApplyPreconfiguredSettings(const FPCGPreConfiguredSettingsInfo& PreconfigureInfo) override;
 	//~End UPCGSettings interface
 
diff --git a/Engine/Plugins/Experimental/PCG/Source/PCG/Public/Elements/Metadata/PCGMetadataStringOpElement.h b/Engine/Plugins/Experimental/PCG/Source/PCG/Public/Elements/Metadata/PCGMetadataStringOpElement.h
--- a/Engine/Plugins/Experimental/PCG/Source/PCG/Public/Elements/Metadata/PCGMetadataStringOpElement.h
+++ b/Engine/Plugins/Experimental/PCG/Source/PCG/Public/Elements/Metadata/PCGMetadataStringOpElement.h
@@ -23,7 +23,7 @@ public:
 	virtual FName GetDefaultNodeName() const override;
 	virtual FText GetDefaultNodeTitle() const override;
 #endif
-	virtual FName AdditionalTaskName() const override;
+	virtual FString GetAdditionalTitleInformation() const override;
 
 protected:
 #if WITH_EDITOR
diff --git a/Engine/Plugins/Experimental/PCG/Source/PCG/Public/Elements/Metadata/PCGMetadataTransformOpElement.h b/Engine/Plugins/Experimental/PCG/Source/PCG/Public/Elements/Metadata/PCGMetadataTransformOpElement.h
--- a/Engine/Plugins/Experimental/PCG/Source/PCG/Public/Elements/Metadata/PCGMetadataTransformOpElement.h
+++ b/Engine/Plugins/Experimental/PCG/Source/PCG/Public/Elements/Metadata/PCGMetadataTransformOpElement.h
@@ -45,7 +45,7 @@ public:
 	virtual TArray<FPCGPreConfiguredSettingsInfo> GetPreconfiguredInfo() const override;
 	virtual bool OnlyExposePreconfiguredSettings() const override { return true; }
 #endif
-	virtual FName AdditionalTaskName() const override;
+	virtual FString GetAdditionalTitleInformation() const override;
 	virtual void ApplyPreconfiguredSettings(const FPCGPreConfiguredSettingsInfo& PreconfigureInfo) override;
 	//~End UPCGSettings interface
 
diff --git a/Engine/Plugins/Experimental/PCG/Source/PCG/Public/Elements/Metadata/PCGMetadataTrigOpElement.h b/Engine/Plugins/Experimental/PCG/Source/PCG/Public/Elements/Metadata/PCGMetadataTrigOpElement.h
--- a/Engine/Plugins/Experimental/PCG/Source/PCG/Public/Elements/Metadata/PCGMetadataTrigOpElement.h
+++ b/Engine/Plugins/Experimental/PCG/Source/PCG/Public/Elements/Metadata/PCGMetadataTrigOpElement.h
@@ -37,7 +37,7 @@ public:
 	virtual TArray<FPCGPreConfiguredSettingsInfo> GetPreconfiguredInfo() const override;
 	virtual bool OnlyExposePreconfiguredSettings() const override { return true; }
 #endif
-	virtual FName AdditionalTaskName() const override;
+	virtual FString GetAdditionalTitleInformation() const override;
 	virtual void ApplyPreconfiguredSettings(const FPCGPreConfiguredSettingsInfo& PreconfigureInfo) override;
 	//~End UPCGSettings interface
 
diff --git a/Engine/Plugins/Experimental/PCG/Source/PCG/Public/Elements/Metadata/PCGMetadataVectorOpElement.h b/Engine/Plugins/Experimental/PCG/Source/PCG/Public/Elements/Metadata/PCGMetadataVectorOpElement.h
--- a/Engine/Plugins/Experimental/PCG/Source/PCG/Public/Elements/Metadata/PCGMetadataVectorOpElement.h
+++ b/Engine/Plugins/Experimental/PCG/Source/PCG/Public/Elements/Metadata/PCGMetadataVectorOpElement.h
@@ -41,7 +41,7 @@ public:
 	virtual TArray<FPCGPreConfiguredSettingsInfo> GetPreconfiguredInfo() const override;
 	virtual bool OnlyExposePreconfiguredSettings() const override { return true; }
 #endif
-	virtual FName AdditionalTaskName() const override;
+	virtual FString GetAdditionalTitleInformation() const override;
 	virtual void ApplyPreconfiguredSettings(const FPCGPreConfiguredSettingsInfo& PreconfigureInfo) override;
 	//~End UPCGSettings interface
 
diff --git a/Engine/Plugins/Experimental/PCG/Source/PCG/Public/Elements/PCGAttributeFilter.h b/Engine/Plugins/Experimental/PCG/Source/PCG/Public/Elements/PCGAttributeFilter.h
--- a/Engine/Plugins/Experimental/PCG/Source/PCG/Public/Elements/PCGAttributeFilter.h
+++ b/Engine/Plugins/Experimental/PCG/Source/PCG/Public/Elements/PCGAttributeFilter.h
@@ -81,6 +81,7 @@ public:
 	virtual EPCGSettingsType GetType() const override { return EPCGSettingsType::Filter; }
 	virtual bool HasDynamicPins() const override { return true; }
 #endif
+	virtual FString GetAdditionalTitleInformation() const override;
 	virtual EPCGDataType GetCurrentPinTypes(const UPCGPin* InPin) const override;
 
 protected:
@@ -147,6 +148,7 @@ public:
 	virtual EPCGSettingsType GetType() const override { return EPCGSettingsType::Filter; }
 	virtual bool HasDynamicPins() const override { return true; }
 #endif
+	virtual FString GetAdditionalTitleInformation() const override;
 	virtual EPCGDataType GetCurrentPinTypes(const UPCGPin* InPin) const override;
 
 protected:
diff --git a/Engine/Plugins/Experimental/PCG/Source/PCG/Public/Elements/PCGAttributeReduceElement.h b/Engine/Plugins/Experimental/PCG/Source/PCG/Public/Elements/PCGAttributeReduceElement.h
--- a/Engine/Plugins/Experimental/PCG/Source/PCG/Public/Elements/PCGAttributeReduceElement.h
+++ b/Engine/Plugins/Experimental/PCG/Source/PCG/Public/Elements/PCGAttributeReduceElement.h
@@ -45,7 +45,7 @@ public:
 
 	virtual void ApplyPreconfiguredSettings(const FPCGPreConfiguredSettingsInfo& PreconfiguredInfo) override;
 
-	virtual FName AdditionalTaskName() const override;
+	virtual FString GetAdditionalTitleInformation() const override;
 
 protected:
 #if WITH_EDITOR
diff --git a/Engine/Plugins/Experimental/PCG/Source/PCG/Public/Elements/PCGAttributeSelectElement.h b/Engine/Plugins/Experimental/PCG/Source/PCG/Public/Elements/PCGAttributeSelectElement.h
--- a/Engine/Plugins/Experimental/PCG/Source/PCG/Public/Elements/PCGAttributeSelectElement.h
+++ b/Engine/Plugins/Experimental/PCG/Source/PCG/Public/Elements/PCGAttributeSelectElement.h
@@ -62,7 +62,7 @@ public:
 	virtual EPCGSettingsType GetType() const override { return EPCGSettingsType::Metadata; }
 	virtual void ApplyDeprecation(UPCGNode* InOutNode) override;
 #endif
-	virtual FName AdditionalTaskName() const override;
+	virtual FString GetAdditionalTitleInformation() const override;
 
 protected:
 #if WITH_EDITOR
diff --git a/Engine/Plugins/Experimental/PCG/Source/PCG/Public/Elements/PCGAttributeTransferElement.h b/Engine/Plugins/Experimental/PCG/Source/PCG/Public/Elements/PCGAttributeTransferElement.h
--- a/Engine/Plugins/Experimental/PCG/Source/PCG/Public/Elements/PCGAttributeTransferElement.h
+++ b/Engine/Plugins/Experimental/PCG/Source/PCG/Public/Elements/PCGAttributeTransferElement.h
@@ -45,7 +45,7 @@ public:
 	virtual void ApplyDeprecation(UPCGNode* InOutNode) override;
 #endif
 	virtual EPCGDataType GetCurrentPinTypes(const UPCGPin* InPin) const override;
-	virtual FName AdditionalTaskName() const override;
+	virtual FString GetAdditionalTitleInformation() const override;
 
 protected:
 #if WITH_EDITOR
diff --git a/Engine/Plugins/Experimental/PCG/Source/PCG/Public/Elements/PCGCreateAttribute.h b/Engine/Plugins/Experimental/PCG/Source/PCG/Public/Elements/PCGCreateAttribute.h
--- a/Engine/Plugins/Experimental/PCG/Source/PCG/Public/Elements/PCGCreateAttribute.h
+++ b/Engine/Plugins/Experimental/PCG/Source/PCG/Public/Elements/PCGCreateAttribute.h
@@ -105,7 +105,7 @@ public:
 #endif // WITH_EDITORONLY_DATA
 
 protected:
-	FName AdditionalTaskNameInternal(FName NodeName) const;
+	FString GetAdditionalTitleInformationInternal(FName NodeName) const;
 
 	virtual FPCGElementPtr CreateElement() const override;
 	virtual bool ShouldAddAttributesPin() const { return true; }
@@ -131,7 +131,7 @@ public:
 	virtual bool HasDynamicPins() const override { return true; }
 	virtual void ApplyDeprecation(UPCGNode* InOutNode) override;
 #endif // WITH_EDITOR
-	virtual FName AdditionalTaskName() const override;
+	virtual FString GetAdditionalTitleInformation() const override;
 
 protected:
 #if WITH_EDITOR
@@ -176,7 +176,7 @@ public:
 	virtual FText GetDefaultNodeTitle() const override;
 	virtual bool HasDynamicPins() const override { return false; }
 #endif
-	virtual FName AdditionalTaskName() const override;
+	virtual FString GetAdditionalTitleInformation() const override;
 
 protected:
 #if WITH_EDITOR
diff --git a/Engine/Plugins/Experimental/PCG/Source/PCG/Public/Elements/PCGDataFromActor.h b/Engine/Plugins/Experimental/PCG/Source/PCG/Public/Elements/PCGDataFromActor.h
--- a/Engine/Plugins/Experimental/PCG/Source/PCG/Public/Elements/PCGDataFromActor.h
+++ b/Engine/Plugins/Experimental/PCG/Source/PCG/Public/Elements/PCGDataFromActor.h
@@ -36,7 +36,7 @@ public:
 #endif
 	virtual EPCGDataType GetCurrentPinTypes(const UPCGPin* InPin) const override;
 
-	virtual FName AdditionalTaskName() const override;
+	virtual FString GetAdditionalTitleInformation() const override;
 
 protected:
 #if WITH_EDITOR
diff --git a/Engine/Plugins/Experimental/PCG/Source/PCG/Public/Elements/PCGDataTableRowToParamData.h b/Engine/Plugins/Experimental/PCG/Source/PCG/Public/Elements/PCGDataTableRowToParamData.h
--- a/Engine/Plugins/Experimental/PCG/Source/PCG/Public/Elements/PCGDataTableRowToParamData.h
+++ b/Engine/Plugins/Experimental/PCG/Source/PCG/Public/Elements/PCGDataTableRowToParamData.h
@@ -21,7 +21,7 @@ public:
 	virtual EPCGSettingsType GetType() const override { return EPCGSettingsType::Metadata; }
 #endif
 
-	virtual FName AdditionalTaskName() const override;
+	virtual FString GetAdditionalTitleInformation() const override;
 
 	virtual TArray<FPCGPinProperties> OutputPinProperties() const override;
 	virtual TArray<FPCGPinProperties> InputPinProperties() const override;
diff --git a/Engine/Plugins/Experimental/PCG/Source/PCG/Public/Elements/PCGDeleteAttributesElement.h b/Engine/Plugins/Experimental/PCG/Source/PCG/Public/Elements/PCGDeleteAttributesElement.h
--- a/Engine/Plugins/Experimental/PCG/Source/PCG/Public/Elements/PCGDeleteAttributesElement.h
+++ b/Engine/Plugins/Experimental/PCG/Source/PCG/Public/Elements/PCGDeleteAttributesElement.h
@@ -39,7 +39,7 @@ public:
 	virtual EPCGSettingsType GetType() const override { return EPCGSettingsType::Metadata; }
 	virtual bool HasDynamicPins() const override { return true; }
 #endif
-	virtual FName AdditionalTaskName() const override;
+	virtual FString GetAdditionalTitleInformation() const override;
 
 protected:
 #if WITH_EDITOR
diff --git a/Engine/Plugins/Experimental/PCG/Source/PCG/Public/Elements/PCGDeleteTags.h b/Engine/Plugins/Experimental/PCG/Source/PCG/Public/Elements/PCGDeleteTags.h
--- a/Engine/Plugins/Experimental/PCG/Source/PCG/Public/Elements/PCGDeleteTags.h
+++ b/Engine/Plugins/Experimental/PCG/Source/PCG/Public/Elements/PCGDeleteTags.h
@@ -27,7 +27,7 @@ public:
 	virtual EPCGSettingsType GetType() const override { return EPCGSettingsType::Generic; }
 	virtual bool HasDynamicPins() const override { return true; }
 #endif
-	virtual FName AdditionalTaskName() const override;
+	virtual FString GetAdditionalTitleInformation() const override;
 
 protected:
 #if WITH_EDITOR
diff --git a/Engine/Plugins/Experimental/PCG/Source/PCG/Public/Elements/PCGDifferenceElement.h b/Engine/Plugins/Experimental/PCG/Source/PCG/Public/Elements/PCGDifferenceElement.h
--- a/Engine/Plugins/Experimental/PCG/Source/PCG/Public/Elements/PCGDifferenceElement.h
+++ b/Engine/Plugins/Experimental/PCG/Source/PCG/Public/Elements/PCGDifferenceElement.h
@@ -25,14 +25,13 @@ public:
 	virtual FText GetDefaultNodeTitle() const override { return NSLOCTEXT("PCGDifferenceSettings", "NodeTitle", "Difference"); }
 	virtual FText GetNodeTooltipText() const override;
 	virtual EPCGSettingsType GetType() const override { return EPCGSettingsType::Spatial; }
-
 	virtual void ApplyStructuralDeprecation(UPCGNode* InOutNode) override;
 #endif
+	virtual FString GetAdditionalTitleInformation() const override;
 
+protected:
 	virtual TArray<FPCGPinProperties> InputPinProperties() const override;
 	virtual TArray<FPCGPinProperties> OutputPinProperties() const override;
-
-protected:
 	virtual FPCGElementPtr CreateElement() const override;
 	//~End UPCGSettings interface
 
diff --git a/Engine/Plugins/Experimental/PCG/Source/PCG/Public/Elements/PCGExecuteBlueprint.h b/Engine/Plugins/Experimental/PCG/Source/PCG/Public/Elements/PCGExecuteBlueprint.h
--- a/Engine/Plugins/Experimental/PCG/Source/PCG/Public/Elements/PCGExecuteBlueprint.h
+++ b/Engine/Plugins/Experimental/PCG/Source/PCG/Public/Elements/PCGExecuteBlueprint.h
@@ -307,9 +307,9 @@ public:
 	virtual TArray<FPCGPreConfiguredSettingsInfo> GetPreconfiguredInfo() const override;
 	virtual bool OnlyExposePreconfiguredSettings() const override;
 #endif
-
 	virtual void ApplyPreconfiguredSettings(const FPCGPreConfiguredSettingsInfo& InPreconfiguredsInfo) override;
-	virtual FName AdditionalTaskName() const override;
+	virtual FString GetAdditionalTitleInformation() const override;
+	virtual bool HasFlippedTitleLines() const override { return true; }
 	virtual TArray<FPCGPinProperties> InputPinProperties() const override;
 	virtual TArray<FPCGPinProperties> OutputPinProperties() const override;
 
diff --git a/Engine/Plugins/Experimental/PCG/Source/PCG/Public/Elements/PCGFilterByAttribute.h b/Engine/Plugins/Experimental/PCG/Source/PCG/Public/Elements/PCGFilterByAttribute.h
--- a/Engine/Plugins/Experimental/PCG/Source/PCG/Public/Elements/PCGFilterByAttribute.h
+++ b/Engine/Plugins/Experimental/PCG/Source/PCG/Public/Elements/PCGFilterByAttribute.h
@@ -20,7 +20,7 @@ public:
 	virtual FText GetNodeTooltipText() const override { return NSLOCTEXT("PCGFilterByAttributeElement", "NodeTooltip", "Separates input data by whether they have the specified attribute or not."); }
 #endif
 
-	virtual FName AdditionalTaskName() const override;
+	virtual FString GetAdditionalTitleInformation() const override;
 
 protected:
 #if WITH_EDITOR
diff --git a/Engine/Plugins/Experimental/PCG/Source/PCG/Public/Elements/PCGFilterByIndex.h b/Engine/Plugins/Experimental/PCG/Source/PCG/Public/Elements/PCGFilterByIndex.h
--- a/Engine/Plugins/Experimental/PCG/Source/PCG/Public/Elements/PCGFilterByIndex.h
+++ b/Engine/Plugins/Experimental/PCG/Source/PCG/Public/Elements/PCGFilterByIndex.h
@@ -19,7 +19,7 @@ public:
 	virtual FText GetDefaultNodeTitle() const override;
 	virtual FText GetNodeTooltipText() const override;
 #endif
-	virtual FName AdditionalTaskName() const override;
+	virtual FString GetAdditionalTitleInformation() const override;
 
 protected:
 #if WITH_EDITOR
diff --git a/Engine/Plugins/Experimental/PCG/Source/PCG/Public/Elements/PCGFilterByTag.h b/Engine/Plugins/Experimental/PCG/Source/PCG/Public/Elements/PCGFilterByTag.h
--- a/Engine/Plugins/Experimental/PCG/Source/PCG/Public/Elements/PCGFilterByTag.h
+++ b/Engine/Plugins/Experimental/PCG/Source/PCG/Public/Elements/PCGFilterByTag.h
@@ -26,7 +26,7 @@ public:
 	virtual FText GetDefaultNodeTitle() const override;
 	virtual FText GetNodeTooltipText() const override;
 #endif
-	virtual FName AdditionalTaskName() const override;
+	virtual FString GetAdditionalTitleInformation() const override;
 
 protected:
 #if WITH_EDITOR
diff --git a/Engine/Plugins/Experimental/PCG/Source/PCG/Public/Elements/PCGGetActorProperty.h b/Engine/Plugins/Experimental/PCG/Source/PCG/Public/Elements/PCGGetActorProperty.h
--- a/Engine/Plugins/Experimental/PCG/Source/PCG/Public/Elements/PCGGetActorProperty.h
+++ b/Engine/Plugins/Experimental/PCG/Source/PCG/Public/Elements/PCGGetActorProperty.h
@@ -34,7 +34,7 @@ public:
 	virtual void GetTrackedActorKeys(FPCGActorSelectionKeyToSettingsMap& OutKeysToSettings, TArray<TObjectPtr<const UPCGGraph>>& OutVisitedGraphs) const override;
 #endif
 
-	virtual FName AdditionalTaskName() const override;
+	virtual FString GetAdditionalTitleInformation() const override;
 	virtual TArray<FPCGPinProperties> OutputPinProperties() const override;
 	virtual TArray<FPCGPinProperties> InputPinProperties() const override { return TArray<FPCGPinProperties>(); }
 
diff --git a/Engine/Plugins/Experimental/PCG/Source/PCG/Public/Elements/PCGGetPropertyFromObjectPath.h b/Engine/Plugins/Experimental/PCG/Source/PCG/Public/Elements/PCGGetPropertyFromObjectPath.h
--- a/Engine/Plugins/Experimental/PCG/Source/PCG/Public/Elements/PCGGetPropertyFromObjectPath.h
+++ b/Engine/Plugins/Experimental/PCG/Source/PCG/Public/Elements/PCGGetPropertyFromObjectPath.h
@@ -26,8 +26,12 @@ public:
 	virtual bool CanEditChange(const FProperty* InProperty) const override;
 	virtual bool IsPinUsedByNodeExecution(const UPCGPin* InPin) const override;
 #endif
+	virtual FString GetAdditionalTitleInformation() const override;
 
 protected:
+#if WITH_EDITOR
+	virtual EPCGChangeType GetChangeTypeForProperty(const FName& InPropertyName) const override { return Super::GetChangeTypeForProperty(InPropertyName) | EPCGChangeType::Cosmetic; }
+#endif
 	virtual TArray<FPCGPinProperties> InputPinProperties() const override;
 	virtual TArray<FPCGPinProperties> OutputPinProperties() const override;
 	virtual FPCGElementPtr CreateElement() const override;
diff --git a/Engine/Plugins/Experimental/PCG/Source/PCG/Public/Elements/PCGHiGenGridSize.h b/Engine/Plugins/Experimental/PCG/Source/PCG/Public/Elements/PCGHiGenGridSize.h
--- a/Engine/Plugins/Experimental/PCG/Source/PCG/Public/Elements/PCGHiGenGridSize.h
+++ b/Engine/Plugins/Experimental/PCG/Source/PCG/Public/Elements/PCGHiGenGridSize.h
@@ -26,7 +26,7 @@ public:
 	virtual EPCGSettingsType GetType() const override { return EPCGSettingsType::HierarchicalGeneration; }
 	virtual bool HasDynamicPins() const override { return true; }
 #endif
-	virtual FName AdditionalTaskName() const override;
+	virtual FString GetAdditionalTitleInformation() const override;
 	virtual EPCGDataType GetCurrentPinTypes(const UPCGPin* InPin) const override;
 
 protected:
diff --git a/Engine/Plugins/Experimental/PCG/Source/PCG/Public/Elements/PCGIndirectionElement.h b/Engine/Plugins/Experimental/PCG/Source/PCG/Public/Elements/PCGIndirectionElement.h
--- a/Engine/Plugins/Experimental/PCG/Source/PCG/Public/Elements/PCGIndirectionElement.h
+++ b/Engine/Plugins/Experimental/PCG/Source/PCG/Public/Elements/PCGIndirectionElement.h
@@ -29,8 +29,9 @@ public:
 	virtual FText GetDefaultNodeTitle() const override;
 	virtual FText GetNodeTooltipText() const override;
 	virtual EPCGSettingsType GetType() const override { return EPCGSettingsType::Generic; }
-	virtual FName AdditionalTaskName() const override;
 #endif // WITH_EDITOR
+	virtual FString GetAdditionalTitleInformation() const override;
+	virtual bool HasFlippedTitleLines() const override { return true; }
 
 protected:
 #if WITH_EDITOR
diff --git a/Engine/Plugins/Experimental/PCG/Source/PCG/Public/Elements/PCGLoopElement.h b/Engine/Plugins/Experimental/PCG/Source/PCG/Public/Elements/PCGLoopElement.h
--- a/Engine/Plugins/Experimental/PCG/Source/PCG/Public/Elements/PCGLoopElement.h
+++ b/Engine/Plugins/Experimental/PCG/Source/PCG/Public/Elements/PCGLoopElement.h
@@ -27,8 +27,6 @@ public:
 	virtual bool GetPinExtraIcon(const UPCGPin* InPin, FName& OutExtraIcon, FText& OutTooltip) const override;
 #endif
 
-	virtual FName AdditionalTaskName() const override;
-
 protected:
 #if WITH_EDITOR
 	virtual EPCGChangeType GetChangeTypeForProperty(const FName& InPropertyName) const override { return Super::GetChangeTypeForProperty(InPropertyName) | EPCGChangeType::Cosmetic; }
diff --git a/Engine/Plugins/Experimental/PCG/Source/PCG/Public/Elements/PCGTypedGetter.h b/Engine/Plugins/Experimental/PCG/Source/PCG/Public/Elements/PCGTypedGetter.h
--- a/Engine/Plugins/Experimental/PCG/Source/PCG/Public/Elements/PCGTypedGetter.h
+++ b/Engine/Plugins/Experimental/PCG/Source/PCG/Public/Elements/PCGTypedGetter.h
@@ -27,7 +27,7 @@ public:
 	virtual FText GetNodeTooltipText() const override;
 #endif
 
-	virtual FName AdditionalTaskName() const override;
+	virtual FString GetAdditionalTitleInformation() const override;
 
 protected:
 	virtual TArray<FPCGPinProperties> OutputPinProperties() const override;
diff --git a/Engine/Plugins/Experimental/PCG/Source/PCG/Public/PCGCommon.h b/Engine/Plugins/Experimental/PCG/Source/PCG/Public/PCGCommon.h
--- a/Engine/Plugins/Experimental/PCG/Source/PCG/Public/PCGCommon.h
+++ b/Engine/Plugins/Experimental/PCG/Source/PCG/Public/PCGCommon.h
@@ -288,3 +288,11 @@ namespace PCGDelegates
 #endif
 }
 
+UENUM()
+enum class EPCGNodeTitleType : uint8
+{
+	/** The full title, may be multiple lines. */
+	FullTitle,
+	/** More concise, single line title. */
+	ListView,
+};
diff --git a/Engine/Plugins/Experimental/PCG/Source/PCG/Public/PCGNode.h b/Engine/Plugins/Experimental/PCG/Source/PCG/Public/PCGNode.h
--- a/Engine/Plugins/Experimental/PCG/Source/PCG/Public/PCGNode.h
+++ b/Engine/Plugins/Experimental/PCG/Source/PCG/Public/PCGNode.h
@@ -7,6 +7,7 @@
 class UPCGNode;
 class UPCGPin;
 enum class EPCGChangeType : uint8;
+enum class EPCGNodeTitleType : uint8;
 struct FPCGPinProperties;
 
 class UPCGSettings;
@@ -63,8 +64,23 @@ public:
 	UFUNCTION(BlueprintCallable, Category = Node)
 	bool RemoveEdgeTo(FName FromPinLable, UPCGNode* To, FName ToPinLabel);
 
-	/** Returns the node title, based either on the current node label, or defaulted to its settings */
-	FText GetNodeTitle() const;
+	/** Get title for node of specified type. */
+	FText GetNodeTitle(EPCGNodeTitleType TitleType) const;
+
+	/** Whether user has renamed the node. */
+	bool HasAuthoredTitle() const { return NodeTitle != NAME_None; }
+
+	/** Title to use if no title is authored. */
+	FText GetDefaultTitle() const;
+
+	/** Authored part of node title (like "Create Attribute 1"). */
+	FText GetAuthoredTitleLine() const;
+
+	/** Whether to flip the order of the title lines - display generated title first and authored second. */
+	bool HasFlippedTitleLines() const;
+
+	/** Generated part of node title, not user editable (like "MyValue = 5.0"). */
+	FText GetGeneratedTitleLine() const;
 
 #if WITH_EDITOR
 	/** Tooltip that describes node functionality and other information. */
diff --git a/Engine/Plugins/Experimental/PCG/Source/PCG/Public/PCGSettings.h b/Engine/Plugins/Experimental/PCG/Source/PCG/Public/PCGSettings.h
--- a/Engine/Plugins/Experimental/PCG/Source/PCG/Public/PCGSettings.h
+++ b/Engine/Plugins/Experimental/PCG/Source/PCG/Public/PCGSettings.h
@@ -298,10 +298,16 @@ public:
 	virtual void PostPaste();
 #endif // WITH_EDITOR
 
+	/** Derived classes can implement this to expose additional information or context, such as an asset in use by the node. */
+	virtual FString GetAdditionalTitleInformation() const { return FString(); }
+
+	/** Display generated title line as primary title (example: BP nodes display the blueprint name as the primary title). */
+	virtual bool HasFlippedTitleLines() const { return false; }
+
 	virtual void ApplyPreconfiguredSettings(const FPCGPreConfiguredSettingsInfo& PreconfigureInfo) {}
 
-	/** Derived classes can implement this to expose additional name information in the logs */
-	virtual FName AdditionalTaskName() const { return NAME_None; }
+	UE_DEPRECATED(5.4, "AdditionalTaskName is deprecated and is replaced with GetAdditionalTitleInformation.")
+	virtual FName AdditionalTaskName() const { return FName(GetAdditionalTitleInformation()); }
 
 	/** Returns true if InPin is in use by node (assuming node enabled). Can be used to communicate when a pin is not in use to user. */
 	virtual bool IsPinUsedByNodeExecution(const UPCGPin* InPin) const { return true; }
diff --git a/Engine/Plugins/Experimental/PCG/Source/PCG/Public/PCGSubgraph.h b/Engine/Plugins/Experimental/PCG/Source/PCG/Public/PCGSubgraph.h
--- a/Engine/Plugins/Experimental/PCG/Source/PCG/Public/PCGSubgraph.h
+++ b/Engine/Plugins/Experimental/PCG/Source/PCG/Public/PCGSubgraph.h
@@ -42,9 +42,15 @@ protected:
 	virtual void PostEditUndo() override;
 	virtual void PreEditChange(FProperty* PropertyAboutToChange) override;
 	virtual void PostEditChangeProperty(struct FPropertyChangedEvent& PropertyChangedEvent) override;
+#endif
 	//~End UObject interface implementation
 
+public:
 	//~Begin UPCGSettings interface
+	virtual bool HasFlippedTitleLines() const override { return true; }
+
+protected:
+#if WITH_EDITOR
 	virtual void GetTrackedActorKeys(FPCGActorSelectionKeyToSettingsMap& OutKeysToSettings, TArray<TObjectPtr<const UPCGGraph>>& OutVisitedGraphs) const override;
 	virtual EPCGChangeType GetChangeTypeForProperty(const FName& InPropertyName) const override;
 #endif
@@ -98,7 +104,7 @@ public:
 	virtual bool GetPinExtraIcon(const UPCGPin* InPin, FName& OutExtraIcon, FText& OutTooltip) const override { return false; }
 #endif
 
-	virtual FName AdditionalTaskName() const override;
+	virtual FString GetAdditionalTitleInformation() const override;
 
 protected:
 #if WITH_EDITOR
diff --git a/Engine/Plugins/Experimental/PCG/Source/PCGEditor/Private/PCGEditor.cpp b/Engine/Plugins/Experimental/PCG/Source/PCGEditor/Private/PCGEditor.cpp
--- a/Engine/Plugins/Experimental/PCG/Source/PCGEditor/Private/PCGEditor.cpp
+++ b/Engine/Plugins/Experimental/PCG/Source/PCGEditor/Private/PCGEditor.cpp
@@ -764,7 +764,7 @@ void FPCGEditor::OnDeterminismNodeTest()
 				check(PCGNode && PCGNode->GetSettings());
 
 				TSharedPtr<FDeterminismTestResult> NodeResult = MakeShared<FDeterminismTestResult>();
-				NodeResult->TestResultTitle = FName(*PCGNode->GetNodeTitle().ToString());
+				NodeResult->TestResultTitle = FName(*PCGNode->GetNodeTitle(EPCGNodeTitleType::ListView).ToString());
 				NodeResult->TestResultName = PCGNode->GetName();
 				NodeResult->Seed = PCGNode->GetSettings()->GetSeed();
 
diff --git a/Engine/Plugins/Experimental/PCG/Source/PCGEditor/Private/PCGEditorGraphNode.cpp b/Engine/Plugins/Experimental/PCG/Source/PCGEditor/Private/PCGEditorGraphNode.cpp
--- a/Engine/Plugins/Experimental/PCG/Source/PCGEditor/Private/PCGEditorGraphNode.cpp
+++ b/Engine/Plugins/Experimental/PCG/Source/PCGEditor/Private/PCGEditorGraphNode.cpp
@@ -2,35 +2,19 @@
 
 #include "PCGEditorGraphNode.h"
 
-#include "PCGEditorModule.h"
 #include "PCGNode.h"
+#include "PCGPin.h"
 #include "PCGSettings.h"
 #include "PCGSubsystem.h"
 
-#include "Framework/Commands/GenericCommands.h"
-#include "PCGPin.h"
+#include "PCGEditorModule.h"
+
 #include "ToolMenu.h"
 #include "ToolMenuSection.h"
+#include "Framework/Commands/GenericCommands.h"
 
 #define LOCTEXT_NAMESPACE "PCGEditorGraphNode"
 
-namespace UPCGEditorGraphNodeHelpers
-{
-	// Info to aid element cache analysis / debugging
-	void GetGraphCacheDebugInfo(const UPCGNode* InNode, bool& bOutDebuggingEnabled, uint32& OutNumCacheEntries)
-	{
-		UWorld* World = GEditor ? (GEditor->PlayWorld ? GEditor->PlayWorld.Get() : GEditor->GetEditorWorldContext().World()) : nullptr;
-		UPCGSubsystem* Subsystem = UPCGSubsystem::GetInstance(World);
-		bOutDebuggingEnabled = Subsystem && Subsystem->IsGraphCacheDebuggingEnabled();
-
-		if (bOutDebuggingEnabled)
-		{
-			IPCGElement* Element = (InNode && InNode->GetSettings()) ? InNode->GetSettings()->GetElement().Get() : nullptr;
-			OutNumCacheEntries = Element ? Subsystem->GetGraphCacheEntryCount(Element) : 0;
-		}
-	}
-}
-
 UPCGEditorGraphNode::UPCGEditorGraphNode(const FObjectInitializer& ObjectInitializer)
 	: Super(ObjectInitializer)
 {
@@ -39,29 +23,19 @@ UPCGEditorGraphNode::UPCGEditorGraphNode(const FObjectInitializer& ObjectInitial
 
 FText UPCGEditorGraphNode::GetNodeTitle(ENodeTitleType::Type TitleType) const
 {
-	constexpr int32 NodeTitleMaxSize = 70;
-
-	FString Result;
-	if (PCGNode)
+	if (!PCGNode)
 	{
-		Result = PCGNode->GetNodeTitle().ToString();
-		Result = (Result.Len() > NodeTitleMaxSize) ? Result.Left(NodeTitleMaxSize) : Result;
+		return NSLOCTEXT("PCGEditorGraphNode", "UnnamedNodeTitle", "Unnamed Node");
 	}
-	else
+
+	if (TitleType == ENodeTitleType::FullTitle)
 	{
-		Result = FString(TEXT("Unnamed node"));
+		return PCGNode->GetNodeTitle(EPCGNodeTitleType::FullTitle);
 	}
-
-	// Debug info - append how many copies of this element are currently in the cache to the node title
-	bool bDebuggingEnabled;
-	uint32 NumCacheEntries;
-	UPCGEditorGraphNodeHelpers::GetGraphCacheDebugInfo(PCGNode.Get(), bDebuggingEnabled, NumCacheEntries);
-	if (bDebuggingEnabled)
+	else
 	{
-		Result = FString::Format(TEXT("{0} [{1}]"), { Result, NumCacheEntries });
+		return PCGNode->GetNodeTitle(EPCGNodeTitleType::ListView);
 	}
-
-	return FText::FromString(Result);
 }
 
 void UPCGEditorGraphNode::GetNodeContextMenuActions(UToolMenu* Menu, class UGraphNodeContextMenuContext* Context) const
@@ -106,7 +80,7 @@ void UPCGEditorGraphNode::OnRenameNode(const FString& NewName)
 		return;
 	}
 
-	if(PCGNode->GetNodeTitle().ToString() != NewName)
+	if (PCGNode->GetAuthoredTitleLine().ToString() != NewName)
 	{
 		Modify();
 		PCGNode->Modify();
diff --git a/Engine/Plugins/Experimental/PCG/Source/PCGEditor/Private/PCGEditorGraphNodeBase.cpp b/Engine/Plugins/Experimental/PCG/Source/PCGEditor/Private/PCGEditorGraphNodeBase.cpp
--- a/Engine/Plugins/Experimental/PCG/Source/PCGEditor/Private/PCGEditorGraphNodeBase.cpp
+++ b/Engine/Plugins/Experimental/PCG/Source/PCGEditor/Private/PCGEditorGraphNodeBase.cpp
@@ -3,22 +3,23 @@
 #include "PCGEditorGraphNodeBase.h"
 
 #include "PCGComponent.h"
+#include "PCGGraph.h"
+#include "PCGPin.h"
+#include "PCGSettingsWithDynamicInputs.h"
+#include "PCGSubsystem.h"
+#include "PCGWorldActor.h"
+
 #include "PCGEditor.h"
 #include "PCGEditorCommands.h"
 #include "PCGEditorCommon.h"
 #include "PCGEditorGraph.h"
 #include "PCGEditorGraphSchema.h"
 #include "PCGEditorSettings.h"
-#include "PCGGraph.h"
-#include "PCGPin.h"
-#include "PCGSettingsWithDynamicInputs.h"
-#include "PCGSubsystem.h"
-#include "PCGWorldActor.h"
 
 #include "GraphEditorActions.h"
+#include "ScopedTransaction.h"
 #include "ToolMenu.h"
 #include "ToolMenuSection.h"
-#include "ScopedTransaction.h"
 #include "Logging/TokenizedMessage.h"
 #include "Misc/TransactionObjectEvent.h"
 #include "Widgets/Colors/SColorPicker.h"
@@ -467,6 +468,21 @@ bool UPCGEditorGraphNodeBase::GetCompactNodeIcon(FName& OutCompactNodeIcon) cons
 	return Settings && Settings->GetCompactNodeIcon(OutCompactNodeIcon);
 }
 
+bool UPCGEditorGraphNodeBase::HasFlippedTitleLines() const
+{
+	return PCGNode ? PCGNode->HasFlippedTitleLines() : false;
+}
+
+FText UPCGEditorGraphNodeBase::GetAuthoredTitleLine() const
+{
+	return PCGNode ? PCGNode->GetAuthoredTitleLine() : FText();
+}
+
+FText UPCGEditorGraphNodeBase::GetGeneratedTitleLine() const
+{
+	return PCGNode ? PCGNode->GetGeneratedTitleLine() : FText();
+}
+
 void UPCGEditorGraphNodeBase::EnterRenamingMode()
 {
 	bCanRenameNode = true;
diff --git a/Engine/Plugins/Experimental/PCG/Source/PCGEditor/Private/PCGEditorGraphNodeBase.h b/Engine/Plugins/Experimental/PCG/Source/PCGEditor/Private/PCGEditorGraphNodeBase.h
--- a/Engine/Plugins/Experimental/PCG/Source/PCGEditor/Private/PCGEditorGraphNodeBase.h
+++ b/Engine/Plugins/Experimental/PCG/Source/PCGEditor/Private/PCGEditorGraphNodeBase.h
@@ -84,6 +84,15 @@ public:
 
 	bool CanUserAddRemoveDynamicInputPins() const { return bCanUserAddRemoveSourcePins; }
 
+	/** Whether to flip the order of the title lines - display generated title first and authored second. */
+	bool HasFlippedTitleLines() const;
+
+	/** Authored part of node title (like "Create Attribute X"). */
+	FText GetAuthoredTitleLine() const;
+
+	/** Generated part of node title, not user editable (like "X = 5.0"). */
+	FText GetGeneratedTitleLine() const;
+
 	/** Marks the node as re-nameable and provokes the node to update, placing the user in rename mode. */
 	void EnterRenamingMode();
 
diff --git a/Engine/Plugins/Experimental/PCG/Source/PCGEditor/Private/SPCGEditorGraphFind.cpp b/Engine/Plugins/Experimental/PCG/Source/PCGEditor/Private/SPCGEditorGraphFind.cpp
--- a/Engine/Plugins/Experimental/PCG/Source/PCGEditor/Private/SPCGEditorGraphFind.cpp
+++ b/Engine/Plugins/Experimental/PCG/Source/PCGEditor/Private/SPCGEditorGraphFind.cpp
@@ -304,29 +304,27 @@ void SPCGEditorGraphFind::MatchTokensInternal(const TArray<FString>& InTokens, U
 	for (UEdGraphNode* Node : PCGEditorGraph->Nodes)
 	{
 		check(Node);
-		const FString NodeName = Node->GetNodeTitle(ENodeTitleType::FullTitle).ToString();
-		const FString NodeType = Node->GetNodeTitle(ENodeTitleType::ListView).ToString();
 
-		FString NodeSearchString = NodeName + NodeType + Node->NodeComment;
+		const FString NodeString = Node->GetNodeTitle(ENodeTitleType::ListView).ToString();
 
-		const UPCGEditorGraphNodeBase* PCGEditorGraphNodeBase = Cast<UPCGEditorGraphNodeBase>(Node);
+		// Search string has full title (both lines).
+		FString NodeSearchString = Node->GetNodeTitle(ENodeTitleType::FullTitle).ToString() + Node->NodeComment;
 
-		if (PCGEditorGraphNodeBase)
+		// Add internal object name which will still display here and there.
+		const UPCGEditorGraphNodeBase* PCGEditorGraphNodeBase = Cast<UPCGEditorGraphNodeBase>(Node);
+		if (const UPCGNode* PCGNode = PCGEditorGraphNodeBase ? PCGEditorGraphNodeBase->GetPCGNode() : nullptr)
 		{
-			if (const UPCGNode* PCGNode = PCGEditorGraphNodeBase->GetPCGNode())
-			{
-				NodeSearchString.Append(PCGNode->GetName());
-			}
+			NodeSearchString.Append(PCGNode->GetName());
 		}
 
 		NodeSearchString = NodeSearchString.Replace(TEXT(" "), TEXT(""));
 
 		FPCGEditorGraphFindResultPtr NodeResult;
-		auto GetOrCreateNodeResult = [this, &NodeResult, &NodeName, &NodeType, Node, PCGEditorGraph, &GetParentFunc]() -> FPCGEditorGraphFindResultPtr&
+		auto GetOrCreateNodeResult = [this, &NodeResult, &NodeString, Node, PCGEditorGraph, &GetParentFunc]() -> FPCGEditorGraphFindResultPtr&
 		{
 			if (!NodeResult.IsValid())
 			{
-				NodeResult = MakeShared<FPCGEditorGraphFindResult>((NodeName == NodeType) ? NodeName : NodeName + " - " + NodeType, RootFindResult, Node);
+				NodeResult = MakeShared<FPCGEditorGraphFindResult>(NodeString, RootFindResult, Node);
 
 				FPCGEditorGraphFindResultPtr Parent = GetParentFunc();
 				if(Parent)
diff --git a/Engine/Plugins/Experimental/PCG/Source/PCGEditor/Private/SPCGEditorGraphLogView.cpp b/Engine/Plugins/Experimental/PCG/Source/PCGEditor/Private/SPCGEditorGraphLogView.cpp
--- a/Engine/Plugins/Experimental/PCG/Source/PCGEditor/Private/SPCGEditorGraphLogView.cpp
+++ b/Engine/Plugins/Experimental/PCG/Source/PCGEditor/Private/SPCGEditorGraphLogView.cpp
@@ -308,7 +308,7 @@ FReply SPCGEditorGraphLogView::Refresh()
 		const UPCGNode* PCGNode = PCGEditorNode ? PCGEditorNode->GetPCGNode() : nullptr;
 		if (PCGNode)
 		{
-			CreateAndAddItem(PCGEditorNode, PCGNode, *PCGNode->GetNodeTitle().ToString());
+			CreateAndAddItem(PCGEditorNode, PCGNode, *PCGNode->GetNodeTitle(EPCGNodeTitleType::ListView).ToString());
 		}
 	}
 
@@ -382,7 +382,7 @@ void SPCGEditorGraphLogView::CreateAndAddItem(const UPCGEditorGraphNode* InPCGEd
 	{
 		if (const UPCGGraph* Graph = CastChecked<const UPCGBaseSubgraphSettings>(Settings)->GetSubgraph())
 		{
-			FString Prefix = InName.ToString() + "/";
+			const FString Prefix = InName.ToString() + "/";
 			for (const UPCGNode* ChildNode : Graph->GetNodes())
 			{
 				if (!ChildNode)
@@ -390,8 +390,9 @@ void SPCGEditorGraphLogView::CreateAndAddItem(const UPCGEditorGraphNode* InPCGEd
 					continue;
 				}
 
-				FName ChildName = FName(Prefix + ChildNode->GetNodeTitle().ToString());
-				CreateAndAddItem(InPCGEditorNode, ChildNode, ChildName);
+				const FText ChildTitle = ChildNode->GetNodeTitle(EPCGNodeTitleType::ListView);
+
+				CreateAndAddItem(InPCGEditorNode, ChildNode, FName(Prefix + ChildTitle.ToString()));
 			}
 		}
 	}
diff --git a/Engine/Plugins/Experimental/PCG/Source/PCGEditor/Private/SPCGEditorGraphNode.cpp b/Engine/Plugins/Experimental/PCG/Source/PCGEditor/Private/SPCGEditorGraphNode.cpp
--- a/Engine/Plugins/Experimental/PCG/Source/PCGEditor/Private/SPCGEditorGraphNode.cpp
+++ b/Engine/Plugins/Experimental/PCG/Source/PCGEditor/Private/SPCGEditorGraphNode.cpp
@@ -580,6 +580,13 @@ void SPCGEditorGraphNode::OnNodeRenameInitiated()
 {
 	if (InlineEditableText.IsValid())
 	{
+		// If titles are flipped, the primary title widget will contain the generated title. Inject
+		// the other title here so the edit action will change the authored title.
+		if (PCGEditorGraphNode->HasFlippedTitleLines())
+		{
+			InlineEditableText->SetText(PCGEditorGraphNode->GetAuthoredTitleLine());
+		}
+
 		InlineEditableText->EnterEditingMode();
 	}
 }
diff --git a/Engine/Plugins/Experimental/PCG/Source/PCGEditor/Private/SPCGEditorGraphProfilingView.cpp b/Engine/Plugins/Experimental/PCG/Source/PCGEditor/Private/SPCGEditorGraphProfilingView.cpp
--- a/Engine/Plugins/Experimental/PCG/Source/PCGEditor/Private/SPCGEditorGraphProfilingView.cpp
+++ b/Engine/Plugins/Experimental/PCG/Source/PCGEditor/Private/SPCGEditorGraphProfilingView.cpp
@@ -381,7 +381,7 @@ namespace PCGEditorGraphProfilingView
 			FString Fullname = FolderName;
 			if (Info.Node)
 			{
-				Fullname += Info.Node->GetNodeTitle().ToString();
+				Fullname += Info.Node->GetNodeTitle(EPCGNodeTitleType::ListView).ToString();
 			}
 
 			// don't show inclusive times when bExpandSubgraph is on or the total sum of everything will mismatch reality
