[PCG] Support multi input on Get Attribute From Point Index

#rnx

diff --git a/Engine/Plugins/Experimental/PCG/Source/PCG/Private/Elements/PCGAttributeGetFromPointIndexElement.cpp b/Engine/Plugins/Experimental/PCG/Source/PCG/Private/Elements/PCGAttributeGetFromPointIndexElement.cpp
--- a/Engine/Plugins/Experimental/PCG/Source/PCG/Private/Elements/PCGAttributeGetFromPointIndexElement.cpp
+++ b/Engine/Plugins/Experimental/PCG/Source/PCG/Private/Elements/PCGAttributeGetFromPointIndexElement.cpp
@@ -56,7 +56,7 @@ void UPCGAttributeGetFromPointIndexSettings::PostLoad()
 TArray<FPCGPinProperties> UPCGAttributeGetFromPointIndexSettings::InputPinProperties() const
 {
 	TArray<FPCGPinProperties> PinProperties;
-	PinProperties.Emplace(PCGPinConstants::DefaultInputLabel, EPCGDataType::Point, /*bInAllowMultipleConnections=*/ false);
+	PinProperties.Emplace(PCGPinConstants::DefaultInputLabel, EPCGDataType::Point);
 
 	return PinProperties;
 }
@@ -86,93 +86,93 @@ bool FPCGAttributeGetFromPointIndexElement::ExecuteInternal(FPCGContext* Context
 
 	TArray<FPCGTaggedData> Inputs = Context->InputData.GetInputsByPin(PCGPinConstants::DefaultInputLabel);
 
-	if (Inputs.Num() != 1)
+	for (int32 InputIndex = 0; InputIndex < Inputs.Num(); ++InputIndex)
 	{
-		PCGE_LOG(Warning, LogOnly, FText::Format(LOCTEXT("WrongNumberOfInputs", "Input pin expected to have one input data element, encountered {0}"), Inputs.Num()));
-		return true;
-	}
+		const FPCGTaggedData& Input = Inputs[InputIndex];
+		const UPCGPointData* PointData = Cast<UPCGPointData>(Input.Data);
 
-	const UPCGPointData* PointData = Cast<UPCGPointData>(Inputs[0].Data);
+		if (!PointData)
+		{
+			PCGE_LOG(Error, GraphAndLog, FText::Format(LOCTEXT("InputNotPointData", "Input {0} is not a point data"), InputIndex));
+			continue;
+		}
 
-	if (!PointData)
-	{
-		PCGE_LOG(Error, GraphAndLog, LOCTEXT("InputNotPointData", "Input is not a point data"));
-		return true;
-	}
+		const int32 Index = Settings->Index;
 
-	const int32 Index = Settings->Index;
+		if (Index < 0 || Index >= PointData->GetPoints().Num())
+		{
+			PCGE_LOG(Error, GraphAndLog, FText::Format(LOCTEXT("IndexOutOfBounds", "Index for input {0} is out of bounds. Index: {1}; Number of Points: {2}"), InputIndex, Index, PointData->GetPoints().Num()));
+			continue;
+		}
 
-	if (Index < 0 || Index >= PointData->GetPoints().Num())
-	{
-		PCGE_LOG(Error, GraphAndLog, FText::Format(LOCTEXT("IndexOutOfBounds", "Index is out of bounds. Index: {0}; Number of Points: {1}"), Index, PointData->GetPoints().Num()));
-		return true;
-	}
+		TArray<FPCGTaggedData>& Outputs = Context->OutputData.TaggedData;
+		const FPCGPoint& Point = PointData->GetPoints()[Index];
 
-	TArray<FPCGTaggedData>& Outputs = Context->OutputData.TaggedData;
-	const FPCGPoint& Point = PointData->GetPoints()[Index];
+		FPCGAttributePropertyInputSelector InputSource = Settings->InputSource.CopyAndFixLast(PointData);
 
-#if !WITH_EDITOR
-	// Eschew output creation only in non-editor builds
-	if(Context->Node && Context->Node->IsOutputPinConnected(PCGAttributeGetFromPointIndexConstants::OutputPointLabel))
-#endif
-	{
-		UPCGPointData* OutputPointData = NewObject<UPCGPointData>();
-		OutputPointData->InitializeFromData(PointData);
-		OutputPointData->GetMutablePoints().Add(Point);
+		const FName OutputAttributeName = (Settings->OutputAttributeName == PCGMetadataAttributeConstants::SourceNameAttributeName) ? InputSource.GetName() : Settings->OutputAttributeName;
 
-		FPCGTaggedData& Output = Outputs.Add_GetRef(Inputs[0]);
-		Output.Data = OutputPointData;
-		Output.Pin = PCGAttributeGetFromPointIndexConstants::OutputPointLabel;
-	}
+		TUniquePtr<const IPCGAttributeAccessor> Accessor = PCGAttributeAccessorHelpers::CreateConstAccessor(PointData, InputSource);
+		FPCGAttributeAccessorKeysPoints PointKey(Point);
 
-	FPCGAttributePropertyInputSelector InputSource = Settings->InputSource.CopyAndFixLast(PointData);
+		if (Accessor.IsValid())
+		{
+			UPCGParamData* OutputParamData = NewObject<UPCGParamData>();
 
-	const FName OutputAttributeName = (Settings->OutputAttributeName == PCGMetadataAttributeConstants::SourceNameAttributeName) ? InputSource.GetName() : Settings->OutputAttributeName;
+			auto ExtractAttribute = [this, Context, OutputAttributeName, &OutputParamData, &Accessor, &PointKey, InputIndex](auto DummyValue) -> bool
+			{
+				using AttributeType = decltype(DummyValue);
 
-	TUniquePtr<const IPCGAttributeAccessor> Accessor = PCGAttributeAccessorHelpers::CreateConstAccessor(PointData, InputSource);
-	FPCGAttributeAccessorKeysPoints PointKey(Point);
+				AttributeType Value{};
 
-	if (Accessor.IsValid())
-	{
-		UPCGParamData* OutputParamData = NewObject<UPCGParamData>();
+				// Should never fail, as OutputType == Accessor->UnderlyingType
+				if (!ensure(Accessor->Get<AttributeType>(Value, PointKey)))
+				{
+					return false;
+				}
 
-		auto ExtractAttribute = [this, Context, OutputAttributeName, &OutputParamData, &Accessor, &PointKey](auto DummyValue) -> bool
-		{
-			using AttributeType = decltype(DummyValue);
+				FPCGMetadataAttribute<AttributeType>* NewAttribute = static_cast<FPCGMetadataAttribute<AttributeType>*>(
+					OutputParamData->Metadata->CreateAttribute<AttributeType>(OutputAttributeName, Value, /*bAllowInterpolation=*/true, /*bOverrideParent=*/false));
 
-			AttributeType Value{};
+				if (!NewAttribute)
+				{
+					PCGE_LOG(Error, GraphAndLog, FText::Format(LOCTEXT("ErrorCreatingTargetAttribute", "Error while creating target attribute '{0}' for output {1}"), FText::FromName(OutputAttributeName), InputIndex));
+					return false;
+				}
 
-			// Should never fail, as OutputType == Accessor->UnderlyingType
-			if (!ensure(Accessor->Get<AttributeType>(Value, PointKey)))
-			{
-				return false;
-			}
+				NewAttribute->SetValue(OutputParamData->Metadata->AddEntry(), Value);
 
-			FPCGMetadataAttribute<AttributeType>* NewAttribute = static_cast<FPCGMetadataAttribute<AttributeType>*>(
-				OutputParamData->Metadata->CreateAttribute<AttributeType>(OutputAttributeName, Value, /*bAllowInterpolation=*/true, /*bOverrideParent=*/false));
+				return true;
+			};
 
-			if (!NewAttribute)
+			if (PCGMetadataAttribute::CallbackWithRightType(Accessor->GetUnderlyingType(), ExtractAttribute))
 			{
-				PCGE_LOG(Error, GraphAndLog, FText::Format(LOCTEXT("ErrorCreatingTargetAttribute", "Error while creating target attribute '{0}'"), FText::FromName(OutputAttributeName)));
-				return false;
+#if !WITH_EDITOR
+				// Add the point
+				// Eschew output creation only in non-editor builds
+				if (Context->Node && Context->Node->IsOutputPinConnected(PCGAttributeGetFromPointIndexConstants::OutputPointLabel))
+#endif
+				{
+					UPCGPointData* OutputPointData = NewObject<UPCGPointData>();
+					OutputPointData->InitializeFromData(PointData);
+					OutputPointData->GetMutablePoints().Add(Point);
+
+					FPCGTaggedData& OutputPoint = Outputs.Add_GetRef(Input);
+					OutputPoint.Data = OutputPointData;
+					OutputPoint.Pin = PCGAttributeGetFromPointIndexConstants::OutputPointLabel;
+				}
+
+				// And the attribute
+				FPCGTaggedData& OutputAttribute = Outputs.Add_GetRef(Input);
+				OutputAttribute.Data = OutputParamData;
+				OutputAttribute.Pin = PCGAttributeGetFromPointIndexConstants::OutputAttributeLabel;
 			}
-
-			NewAttribute->SetValue(OutputParamData->Metadata->AddEntry(), Value);
-
-			return true;
-		};
-
-		if (PCGMetadataAttribute::CallbackWithRightType(Accessor->GetUnderlyingType(), ExtractAttribute))
+		}
+		else
 		{
-			FPCGTaggedData& Output = Outputs.Add_GetRef(Inputs[0]);
-			Output.Data = OutputParamData;
-			Output.Pin = PCGAttributeGetFromPointIndexConstants::OutputAttributeLabel;
+			PCGE_LOG(Warning, GraphAndLog, FText::Format(LOCTEXT("AttributeNotFound", "Cannot find attribute/property '{0}' in input {1}"), FText::FromName(InputSource.GetName()), InputIndex));
 		}
 	}
-	else
-	{
-		PCGE_LOG(Warning, GraphAndLog, FText::Format(LOCTEXT("AttributeNotFound", "Cannot find attribute/property '{0}' in input"), FText::FromName(InputSource.GetName())));
-	}
 
 	return true;
 }
