Fixed warning on copy attribute when there are no points, as this shouldn't be an error per se.

#rnx

diff --git a/Engine/Plugins/Experimental/PCG/Source/PCG/Private/Elements/Metadata/PCGMetadataElementCommon.cpp b/Engine/Plugins/Experimental/PCG/Source/PCG/Private/Elements/Metadata/PCGMetadataElementCommon.cpp
--- a/Engine/Plugins/Experimental/PCG/Source/PCG/Private/Elements/Metadata/PCGMetadataElementCommon.cpp
+++ b/Engine/Plugins/Experimental/PCG/Source/PCG/Private/Elements/Metadata/PCGMetadataElementCommon.cpp
@@ -61,6 +61,12 @@ namespace PCGMetadataElementCommon
 			return false;
 		}
 
+		// Early out - contrary to ApplyOnAccessorRange, having a zero count here is not an error
+		if (Count == 0)
+		{
+			return true;
+		}
+
 		auto Operation = [&Params, Count](auto Dummy)
 		{
 			using OutputType = decltype(Dummy);
