Fixed PVS error.

#rnx

diff --git a/Engine/Plugins/Experimental/PCG/Source/PCG/Private/Elements/PCGCreateAttribute.cpp b/Engine/Plugins/Experimental/PCG/Source/PCG/Private/Elements/PCGCreateAttribute.cpp
--- a/Engine/Plugins/Experimental/PCG/Source/PCG/Private/Elements/PCGCreateAttribute.cpp
+++ b/Engine/Plugins/Experimental/PCG/Source/PCG/Private/Elements/PCGCreateAttribute.cpp
@@ -128,7 +128,7 @@ FString UPCGCreateAttributeBaseSettings::GetAdditionalTitleInformationInternal(F
 		}
 		else
 		{
-			return OutputAttributeName == NAME_None ? SourceParamAttributeName.ToString() : OutputAttributeName.ToString();
+			return ((OutputAttributeName == NAME_None) ? SourceParamAttributeName : OutputAttributeName).ToString();
 		}
 	}
 	else
