[PCG] Fix failed save of newly created levels with partitioned components

diff --git a/Engine/Plugins/Experimental/PCG/Source/PCG/Private/Grid/PCGPartitionActor.cpp b/Engine/Plugins/Experimental/PCG/Source/PCG/Private/Grid/PCGPartitionActor.cpp
--- a/Engine/Plugins/Experimental/PCG/Source/PCG/Private/Grid/PCGPartitionActor.cpp
+++ b/Engine/Plugins/Experimental/PCG/Source/PCG/Private/Grid/PCGPartitionActor.cpp
@@ -89,9 +89,13 @@ void APCGPartitionActor::Serialize(FArchive& Ar)
 {
 #if WITH_EDITOR
 	TMap<TObjectPtr<UPCGComponent>, TSoftObjectPtr<UPCGComponent>> LocalToOriginalCopy;
+	TMap<TObjectPtr<UPCGComponent>, TSoftObjectPtr<UPCGComponent>> Transients;
+	TSet<TObjectPtr<UPCGComponent>> AddedLoadedPreviewKeys;
 
 	if (Ar.IsSaving())
 	{
+		// Split out the map into persistent and transient components. Keep the persistent ones in LocalToOriginal
+		// for serialization purposes (saving, linker fixup of SOPs).
 		LocalToOriginalCopy = MoveTemp(LocalToOriginal);
 
 		LocalToOriginal.Reset();
@@ -101,6 +105,10 @@ void APCGPartitionActor::Serialize(FArchive& Ar)
 			{
 				LocalToOriginal.Add(It);
 			}
+			else
+			{
+				Transients.Add(It);
+			}
 		}
 
 		// For components that were cleared but marked as load-as-preview
@@ -112,7 +120,11 @@ void APCGPartitionActor::Serialize(FArchive& Ar)
 			{
 				if (OriginalComponent->GetSerializedEditingMode() == EPCGEditorDirtyMode::LoadAsPreview && OriginalComponent->GetEditingMode() == EPCGEditorDirtyMode::Preview)
 				{
-					LocalToOriginal.Add(It);
+					if (ensure(!LocalToOriginal.Contains(It.Key)))
+					{
+						LocalToOriginal.Add(It);
+						AddedLoadedPreviewKeys.Add(It.Key);
+					}
 				}
 			}
 		}
@@ -124,7 +136,17 @@ void APCGPartitionActor::Serialize(FArchive& Ar)
 #if WITH_EDITOR
 	if (Ar.IsSaving())
 	{
-		LocalToOriginal = MoveTemp(LocalToOriginalCopy);
+		// Remove added loaded preview components.
+		for (const TObjectPtr<UPCGComponent>& Key : AddedLoadedPreviewKeys)
+		{
+			LocalToOriginal.Remove(Key);
+		}
+
+		// Restore removed transient components.
+		for (const auto& It : Transients)
+		{
+			LocalToOriginal.Add(It);
+		}
 	}
 #endif
 }
