[PCG] New Select (multi) Node

+ New Unit Tests for Select (boolean) node
+ New Unit Tests for Select (multi) node

diff --git a/Engine/Plugins/Experimental/PCG/Source/PCG/Private/Elements/Metadata/PCGMetadataElement.cpp b/Engine/Plugins/Experimental/PCG/Source/PCG/Private/Elements/Metadata/PCGMetadataElement.cpp
--- a/Engine/Plugins/Experimental/PCG/Source/PCG/Private/Elements/Metadata/PCGMetadataElement.cpp
+++ b/Engine/Plugins/Experimental/PCG/Source/PCG/Private/Elements/Metadata/PCGMetadataElement.cpp
@@ -37,9 +37,7 @@ FString UPCGMetadataOperationSettings::GetAdditionalTitleInformation() const
 		return LOCTEXT("NoteTitleAllAttributes", "All Attributes").ToString();
 	}
 
-	const FProperty* FromAttributeProperty = GetClass() ? FindFProperty<FProperty>(GetClass(), GET_MEMBER_NAME_CHECKED(UPCGMetadataOperationSettings, InputSource)) : nullptr;
-	const FProperty* ToAttributeProperty = GetClass() ? FindFProperty<FProperty>(GetClass(), GET_MEMBER_NAME_CHECKED(UPCGMetadataOperationSettings, OutputTarget)) : nullptr;
-	if ((FromAttributeProperty && IsPropertyOverriddenByPin(FromAttributeProperty)) || (ToAttributeProperty && IsPropertyOverriddenByPin(ToAttributeProperty)))
+	if (IsPropertyOverriddenByPin(GET_MEMBER_NAME_CHECKED(UPCGMetadataOperationSettings, InputSource)) || IsPropertyOverriddenByPin(GET_MEMBER_NAME_CHECKED(UPCGMetadataOperationSettings, OutputTarget)))
 	{
 		return FString();
 	}
diff --git a/Engine/Plugins/Experimental/PCG/Source/PCG/Private/Elements/Metadata/PCGMetadataPartition.cpp b/Engine/Plugins/Experimental/PCG/Source/PCG/Private/Elements/Metadata/PCGMetadataPartition.cpp
--- a/Engine/Plugins/Experimental/PCG/Source/PCG/Private/Elements/Metadata/PCGMetadataPartition.cpp
+++ b/Engine/Plugins/Experimental/PCG/Source/PCG/Private/Elements/Metadata/PCGMetadataPartition.cpp
@@ -46,8 +46,7 @@ void UPCGMetadataPartitionSettings::PostLoad()
 FString UPCGMetadataPartitionSettings::GetAdditionalTitleInformation() const
 {
 #if WITH_EDITOR
-	FProperty* AttributeProperty = GetClass() ? FindFProperty<FProperty>(GetClass(), GET_MEMBER_NAME_CHECKED(UPCGMetadataPartitionSettings, PartitionAttributeSource)) : nullptr;
-	if (AttributeProperty && IsPropertyOverriddenByPin(AttributeProperty))
+	if (IsPropertyOverriddenByPin(GET_MEMBER_NAME_CHECKED(UPCGMetadataPartitionSettings, PartitionAttributeSource)))
 	{
 		return FString();
 	}
diff --git a/Engine/Plugins/Experimental/PCG/Source/PCG/Private/Elements/Metadata/PCGMetadataRenameElement.cpp b/Engine/Plugins/Experimental/PCG/Source/PCG/Private/Elements/Metadata/PCGMetadataRenameElement.cpp
--- a/Engine/Plugins/Experimental/PCG/Source/PCG/Private/Elements/Metadata/PCGMetadataRenameElement.cpp
+++ b/Engine/Plugins/Experimental/PCG/Source/PCG/Private/Elements/Metadata/PCGMetadataRenameElement.cpp
@@ -15,9 +15,7 @@
 FString UPCGMetadataRenameSettings::GetAdditionalTitleInformation() const
 {
 #if WITH_EDITOR
-	const FProperty* FromAttributeProperty = GetClass() ? FindFProperty<FProperty>(GetClass(), GET_MEMBER_NAME_CHECKED(UPCGMetadataRenameSettings, AttributeToRename)) : nullptr;
-	const FProperty* ToAttributeProperty = GetClass() ? FindFProperty<FProperty>(GetClass(), GET_MEMBER_NAME_CHECKED(UPCGMetadataRenameSettings, NewAttributeName)) : nullptr;
-	if ((FromAttributeProperty && IsPropertyOverriddenByPin(FromAttributeProperty)) || (ToAttributeProperty && IsPropertyOverriddenByPin(ToAttributeProperty)))
+	if (IsPropertyOverriddenByPin(GET_MEMBER_NAME_CHECKED(UPCGMetadataRenameSettings, AttributeToRename)) || IsPropertyOverriddenByPin(GET_MEMBER_NAME_CHECKED(UPCGMetadataRenameSettings, NewAttributeName)))
 	{
 		return FString();
 	}
diff --git a/Engine/Plugins/Experimental/PCG/Source/PCG/Private/Elements/PCGAttributeFilter.cpp b/Engine/Plugins/Experimental/PCG/Source/PCG/Private/Elements/PCGAttributeFilter.cpp
--- a/Engine/Plugins/Experimental/PCG/Source/PCG/Private/Elements/PCGAttributeFilter.cpp
+++ b/Engine/Plugins/Experimental/PCG/Source/PCG/Private/Elements/PCGAttributeFilter.cpp
@@ -304,8 +304,7 @@ TArray<FPCGPinProperties> UPCGAttributeFilteringSettings::OutputPinProperties()
 FString UPCGAttributeFilteringSettings::GetAdditionalTitleInformation() const
 {
 #if WITH_EDITOR
-	const FProperty* AttributeProperty = GetClass() ? FindFProperty<FProperty>(GetClass(), GET_MEMBER_NAME_CHECKED(UPCGAttributeFilteringSettings, TargetAttribute)) : nullptr;
-	if (AttributeProperty && IsPropertyOverriddenByPin(AttributeProperty))
+	if (IsPropertyOverriddenByPin(GET_MEMBER_NAME_CHECKED(UPCGAttributeFilteringSettings, TargetAttribute)))
 	{
 		return FString();
 	}
@@ -352,8 +351,7 @@ void UPCGAttributeFilteringRangeSettings::PostLoad()
 FString UPCGAttributeFilteringRangeSettings::GetAdditionalTitleInformation() const
 {
 #if WITH_EDITOR
-	const FProperty* AttributeProperty = GetClass() ? FindFProperty<FProperty>(GetClass(), GET_MEMBER_NAME_CHECKED(UPCGAttributeFilteringRangeSettings, TargetAttribute)) : nullptr;
-	if (AttributeProperty && IsPropertyOverriddenByPin(AttributeProperty))
+	if (IsPropertyOverriddenByPin(GET_MEMBER_NAME_CHECKED(UPCGAttributeFilteringRangeSettings, TargetAttribute)))
 	{
 		return FString();
 	}
diff --git a/Engine/Plugins/Experimental/PCG/Source/PCG/Private/Elements/PCGDifferenceElement.cpp b/Engine/Plugins/Experimental/PCG/Source/PCG/Private/Elements/PCGDifferenceElement.cpp
--- a/Engine/Plugins/Experimental/PCG/Source/PCG/Private/Elements/PCGDifferenceElement.cpp
+++ b/Engine/Plugins/Experimental/PCG/Source/PCG/Private/Elements/PCGDifferenceElement.cpp
@@ -87,8 +87,7 @@ void UPCGDifferenceSettings::ApplyStructuralDeprecation(UPCGNode* InOutNode)
 FString UPCGDifferenceSettings::GetAdditionalTitleInformation() const
 {
 #if WITH_EDITOR
-	const FProperty* DensityFunctionProperty = GetClass() ? FindFProperty<FProperty>(GetClass(), GET_MEMBER_NAME_CHECKED(UPCGDifferenceSettings, DensityFunction)) : nullptr;
-	if (DensityFunctionProperty && IsPropertyOverriddenByPin(DensityFunctionProperty))
+	if (IsPropertyOverriddenByPin(GET_MEMBER_NAME_CHECKED(UPCGDifferenceSettings, DensityFunction)))
 	{
 		return FString();
 	}
diff --git a/Engine/Plugins/Experimental/PCG/Source/PCG/Private/Elements/PCGFilterByAttribute.cpp b/Engine/Plugins/Experimental/PCG/Source/PCG/Private/Elements/PCGFilterByAttribute.cpp
--- a/Engine/Plugins/Experimental/PCG/Source/PCG/Private/Elements/PCGFilterByAttribute.cpp
+++ b/Engine/Plugins/Experimental/PCG/Source/PCG/Private/Elements/PCGFilterByAttribute.cpp
@@ -18,8 +18,7 @@ FText UPCGFilterByAttributeSettings::GetDefaultNodeTitle() const
 FString UPCGFilterByAttributeSettings::GetAdditionalTitleInformation() const
 {
 #if WITH_EDITOR
-	const FProperty* AttributeProperty = GetClass() ? FindFProperty<FProperty>(GetClass(), GET_MEMBER_NAME_CHECKED(UPCGFilterByAttributeSettings, Attribute)) : nullptr;
-	if (AttributeProperty && IsPropertyOverriddenByPin(AttributeProperty))
+	if (IsPropertyOverriddenByPin(GET_MEMBER_NAME_CHECKED(UPCGFilterByAttributeSettings, Attribute)))
 	{
 		return FString();
 	}
diff --git a/Engine/Plugins/Experimental/PCG/Source/PCG/Private/Elements/PCGMultiSelect.cpp b/Engine/Plugins/Experimental/PCG/Source/PCG/Private/Elements/PCGMultiSelect.cpp
new file mode 100644
--- /dev/null
+++ b/Engine/Plugins/Experimental/PCG/Source/PCG/Private/Elements/PCGMultiSelect.cpp
@@ -0,0 +1,278 @@
+// Copyright Epic Games, Inc. All Rights Reserved.
+
+#include "Elements/PCGMultiSelect.h"
+
+#include "PCGContext.h"
+#include "PCGModule.h"
+#include "Elements/PCGGather.h"
+
+#define LOCTEXT_NAMESPACE "FPCGMultiSelectElement"
+
+namespace PCGMultiSelectConstants
+{
+	const FText NodeTitleBase = LOCTEXT("NodeTitleBase", "Select (Multi)");
+	const FText SubtitleInt = LOCTEXT("SubtitleInt", "Integer Selection");
+	const FText SubtitleEnum = LOCTEXT("SubtitleEnum", "Enum Selection");
+	const FText SubtitleString = LOCTEXT("SubtitleString", "String Selection");
+	const FName DefaultPathPinLabel = TEXT("Default");
+}
+
+void UPCGMultiSelectSettings::PostLoad()
+{
+	Super::PostLoad();
+	CachePinLabels();
+}
+
+#if WITH_EDITOR
+void UPCGMultiSelectSettings::PostEditChangeProperty(FPropertyChangedEvent& PropertyChangedEvent)
+{
+	Super::PostEditChangeProperty(PropertyChangedEvent);
+
+	// Only need to change the pin labels if the options have changed
+	if (PropertyChangedEvent.GetMemberPropertyName() == GET_MEMBER_NAME_CHECKED(UPCGMultiSelectSettings, IntOptions) ||
+		PropertyChangedEvent.GetMemberPropertyName() == GET_MEMBER_NAME_CHECKED(UPCGMultiSelectSettings, StringOptions) ||
+		PropertyChangedEvent.GetMemberPropertyName() == GET_MEMBER_NAME_CHECKED(UPCGMultiSelectSettings, EnumSelection))
+	{
+		CachePinLabels();
+	}
+}
+
+FText UPCGMultiSelectSettings::GetDefaultNodeTitle() const
+{
+	return PCGMultiSelectConstants::NodeTitleBase;
+}
+
+FText UPCGMultiSelectSettings::GetNodeTooltipText() const
+{
+	return LOCTEXT("NodeTooltip", "Control flow node that will select all input data on a single input pin that matches a given selection mode and corresponding 'selection' property - which can also be overridden.");
+}
+#endif // WITH_EDITOR
+
+EPCGDataType UPCGMultiSelectSettings::GetCurrentPinTypes(const UPCGPin* Pin) const
+{
+	// Output pin depends on the union of all non-advanced input pins. Advanced pins are dynamic user overrides and
+	// should not affect the output pin type.
+	if (Pin->IsOutputPin())
+	{
+		EPCGDataType InputTypeUnion = EPCGDataType::None;
+
+		const TArray<TObjectPtr<UPCGPin>>& InputPins = Pin->Node->GetInputPins();
+
+		for (int Index = 0; Index < InputPins.Num(); ++Index)
+		{
+			if (!InputPins[Index]->Properties.bAdvancedPin)
+			{
+				InputTypeUnion |= GetTypeUnionOfIncidentEdges(InputPins[Index]->Properties.Label);
+			}
+		}
+
+		return InputTypeUnion != EPCGDataType::None ? InputTypeUnion : EPCGDataType::Any;
+	}
+
+	return Super::GetCurrentPinTypes(Pin);
+}
+
+FString UPCGMultiSelectSettings::GetAdditionalTitleInformation() const
+{
+	switch (SelectionMode)
+	{
+		case EPCGControlFlowSelectionMode::Integer:
+		{
+			FString Subtitle = PCGMultiSelectConstants::SubtitleInt.ToString();
+			if (!IsPropertyOverriddenByPin(GET_MEMBER_NAME_CHECKED(UPCGMultiSelectSettings, IntegerSelection)))
+			{
+				Subtitle += FString::Format(TEXT(": {0}"), {IntegerSelection});
+			}
+
+			return Subtitle;
+		}
+
+		case EPCGControlFlowSelectionMode::Enum:
+			if (EnumSelection.Class)
+			{
+				FString Subtitle = PCGMultiSelectConstants::SubtitleEnum.ToString();
+				if (!IsPropertyOverriddenByPin({GET_MEMBER_NAME_CHECKED(UPCGMultiSelectSettings, EnumSelection), GET_MEMBER_NAME_CHECKED(FEnumSelector, Value)}))
+				{
+					Subtitle += EnumSelection.Class->GetNameStringByValue(EnumSelection.Value);
+				}
+
+				return Subtitle;
+			}
+			break;
+
+		case EPCGControlFlowSelectionMode::String:
+			{
+				FString Subtitle = PCGMultiSelectConstants::SubtitleString.ToString();
+				if (!IsPropertyOverriddenByPin(GET_MEMBER_NAME_CHECKED(UPCGMultiSelectSettings, StringSelection)))
+				{
+					Subtitle += FString::Format(TEXT(": {0}"), {StringSelection});
+				}
+
+				return Subtitle;
+			}
+
+		default:
+			checkNoEntry();
+			break;
+	}
+
+	return PCGMultiSelectConstants::NodeTitleBase.ToString();
+}
+
+TArray<FPCGPinProperties> UPCGMultiSelectSettings::InputPinProperties() const
+{
+	TArray<FPCGPinProperties> PinProperties;
+
+	switch (SelectionMode)
+	{
+		case EPCGControlFlowSelectionMode::Integer:
+			for (const int32 Value : IntOptions)
+			{
+				PinProperties.Emplace(FName(FString::FromInt(Value)));
+			}
+			break;
+		case EPCGControlFlowSelectionMode::String:
+			for (const FString& Value : StringOptions)
+			{
+				PinProperties.Emplace(FName(Value));
+			}
+			break;
+		case EPCGControlFlowSelectionMode::Enum:
+			// -1 to bypass the MAX value
+			for (int32 Index = 0; EnumSelection.Class && Index < EnumSelection.Class->NumEnums() - 1; ++Index)
+			{
+#if WITH_EDITOR
+				bool const bHidden = EnumSelection.Class->HasMetaData(TEXT("Hidden"), Index) || EnumSelection.Class->HasMetaData(TEXT("Spacer"), Index);
+				if (!bHidden)
+				{
+					PinProperties.Emplace(FName(EnumSelection.Class->GetDisplayNameTextByIndex(Index).ToString()));
+				}
+#else // WITH_EDITOR
+				// HasMetaData is editor only, so there will be extra pins at runtime, but that should be okay
+				PinProperties.Emplace(FName(EnumSelection.Class->GetDisplayNameTextByIndex(Index).ToString()));
+#endif // WITH_EDITOR
+			}
+			break;
+		default:
+			break;
+	}
+
+	PinProperties.Emplace(PCGMultiSelectConstants::DefaultPathPinLabel);
+
+	return PinProperties;
+}
+
+TArray<FPCGPinProperties> UPCGMultiSelectSettings::OutputPinProperties() const
+{
+	TArray<FPCGPinProperties> PinProperties;
+	PinProperties.Emplace(PCGPinConstants::DefaultOutputLabel,
+		EPCGDataType::Any,
+		/*bInAllowMultipleConnections=*/true,
+		/*bAllowMultipleData=*/true,
+		LOCTEXT("OutputPinTooltip", "All input on the selected input pin will be forwarded directly to the output."));
+
+	return PinProperties;
+}
+
+FPCGElementPtr UPCGMultiSelectSettings::CreateElement() const
+{
+	return MakeShared<FPCGMultiSelectElement>();
+}
+
+bool UPCGMultiSelectSettings::IsValuePresent(const int32 Value) const
+{
+	return IntOptions.Contains(Value);
+}
+
+bool UPCGMultiSelectSettings::IsValuePresent(const FString& Value) const
+{
+	return StringOptions.Contains(Value);
+}
+
+bool UPCGMultiSelectSettings::IsValuePresent(const int64 Value) const
+{
+	if (!EnumSelection.Class)
+	{
+		return false;
+	}
+
+	const int64 Index = EnumSelection.Class->GetIndexByValue(Value);
+	return Index != INDEX_NONE && Index < EnumSelection.Class->NumEnums() - 1;
+}
+
+bool UPCGMultiSelectSettings::GetSelectedPinLabel(FName& OutSelectedPinLabel) const
+{
+	if (CachedPinLabels.IsEmpty())
+	{
+		return false;
+	}
+
+	int32 Index = INDEX_NONE;
+	if (SelectionMode == EPCGControlFlowSelectionMode::Integer && IsValuePresent(IntegerSelection))
+	{
+		Index = IntOptions.IndexOfByKey(IntegerSelection);
+	}
+	else if (SelectionMode == EPCGControlFlowSelectionMode::String && IsValuePresent(StringSelection))
+	{
+		Index = StringOptions.IndexOfByKey(StringSelection);
+	}
+	else if (SelectionMode == EPCGControlFlowSelectionMode::Enum && IsValuePresent(EnumSelection.Value))
+	{
+		// To account for hidden enums missing from the pin properties, find the pin by label instead of index
+		const FName PinLabel(EnumSelection.Class->GetDisplayNameTextByValue(EnumSelection.Value).ToString());
+		for (int i = 0; i < CachedPinLabels.Num(); ++i)
+		{
+			if (CachedPinLabels[i] == PinLabel)
+			{
+				Index = i;
+				break;
+			}
+		}
+	}
+	else
+	{
+		OutSelectedPinLabel = PCGMultiSelectConstants::DefaultPathPinLabel;
+
+		return true;
+	}
+
+	if (Index < 0 || Index >= CachedPinLabels.Num())
+	{
+		return false;
+	}
+
+	OutSelectedPinLabel = CachedPinLabels[Index];
+
+	return true;
+}
+
+void UPCGMultiSelectSettings::CachePinLabels()
+{
+	CachedPinLabels.Empty();
+	Algo::Transform(InputPinProperties(), CachedPinLabels, [](const FPCGPinProperties& Property)
+	{
+		return Property.Label;
+	});
+}
+
+bool FPCGMultiSelectElement::ExecuteInternal(FPCGContext* Context) const
+{
+	TRACE_CPUPROFILER_EVENT_SCOPE(FPCGMultiSelectElement::ExecuteInternal);
+
+	const UPCGMultiSelectSettings* Settings = Context->GetInputSettings<UPCGMultiSelectSettings>();
+	check(Settings);
+
+	FName SelectedPinLabel;
+	if (!Settings->GetSelectedPinLabel(SelectedPinLabel))
+	{
+		PCGE_LOG_C(Error, GraphAndLog, Context, LOCTEXT("ValueDoesNotExist", "Selected value is not a valid option."));
+		return true;
+	}
+
+	// Reuse the functionality of the Gather node
+	Context->OutputData = PCGGather::GatherDataForPin(Context->InputData, SelectedPinLabel);
+
+	return true;
+}
+
+#undef LOCTEXT_NAMESPACE
\ No newline at end of file
diff --git a/Engine/Plugins/Experimental/PCG/Source/PCG/Private/PCGSettings.cpp b/Engine/Plugins/Experimental/PCG/Source/PCG/Private/PCGSettings.cpp
--- a/Engine/Plugins/Experimental/PCG/Source/PCG/Private/PCGSettings.cpp
+++ b/Engine/Plugins/Experimental/PCG/Source/PCG/Private/PCGSettings.cpp
@@ -603,8 +603,36 @@ bool UPCGSettings::CanEditChange(const FProperty* InProperty) const
 	return Super::CanEditChange(InProperty);
 }
 
+void UPCGSettings::PostPaste()
+{
+	// Params are not properly built when importing from pasted text (because param
+	// Properties are transient/unserialized), so we should reinitialize in post
+	InitializeCachedOverridableParams(/*bReset=*/true);
+}
+
+void UPCGSettings::ApplyDeprecationBeforeUpdatePins(UPCGNode* InOutNode, TArray<TObjectPtr<UPCGPin>>& InputPins, TArray<TObjectPtr<UPCGPin>>& OutputPins)
+{
+	check(InOutNode);
+
+	if (DataVersion < FPCGCustomVersion::AddParamPinToOverridableNodes)
+	{
+		PCGSettingsHelpers::DeprecationBreakOutParamsToNewPin(InOutNode, InputPins, OutputPins);
+	}
+
+	if (DataVersion < FPCGCustomVersion::RenameDefaultParamsToOverride)
+	{
+		InOutNode->RenameInputPin(PCGPinConstants::Private::OldDefaultParamsLabel, PCGPinConstants::DefaultParamsLabel);
+	}
+}
+#endif // WITH_EDITOR
+
 bool UPCGSettings::IsPropertyOverriddenByPin(const FProperty* InProperty) const
 {
+	if (!InProperty)
+	{
+		return false;
+	}
+
 	if (const UPCGNode* Node = Cast<UPCGNode>(GetOuter()))
 	{
 		const FPCGSettingsOverridableParam* Param = OverridableParams().FindByPredicate([InProperty](const FPCGSettingsOverridableParam& ParamToCheck)
@@ -626,28 +654,35 @@ bool UPCGSettings::IsPropertyOverriddenByPin(const FProperty* InProperty) const
 	return false;
 }
 
-void UPCGSettings::PostPaste()
+bool UPCGSettings::IsPropertyOverriddenByPin(const FName PropertyName) const
 {
-	// Params are not properly built when importing from pasted text (because param
-	// Properties are transient/unserialized), so we should reinitialize in post
-	InitializeCachedOverridableParams(/*bReset=*/true);
+	return IsPropertyOverriddenByPin(TArrayView<const FName>(&PropertyName, 1));
 }
 
-void UPCGSettings::ApplyDeprecationBeforeUpdatePins(UPCGNode* InOutNode, TArray<TObjectPtr<UPCGPin>>& InputPins, TArray<TObjectPtr<UPCGPin>>& OutputPins)
+bool UPCGSettings::IsPropertyOverriddenByPin(const TArrayView<const FName>& PropertyNameChain) const
 {
-	check(InOutNode);
-
-	if (DataVersion < FPCGCustomVersion::AddParamPinToOverridableNodes)
+	if (PropertyNameChain.IsEmpty())
 	{
-		PCGSettingsHelpers::DeprecationBreakOutParamsToNewPin(InOutNode, InputPins, OutputPins);
+		return false;
 	}
 
-	if (DataVersion < FPCGCustomVersion::RenameDefaultParamsToOverride)
+	if (const UPCGNode* Node = Cast<UPCGNode>(GetOuter()))
 	{
-		InOutNode->RenameInputPin(PCGPinConstants::Private::OldDefaultParamsLabel, PCGPinConstants::DefaultParamsLabel);
+		const FPCGSettingsOverridableParam* Param = OverridableParams().FindByPredicate([&PropertyNameChain](const FPCGSettingsOverridableParam& ParamToCheck)
+		{
+			return ParamToCheck.PropertiesNames == PropertyNameChain;
+		});
+
+		if (Param)
+		{
+			if (const UPCGPin* Pin = Node->GetInputPin(Param->Label))
+			{
+				return Pin->IsConnected();
+			}
+		}
 	}
+	return false;
 }
-#endif // WITH_EDITOR
 
 EPCGDataType UPCGSettings::GetCurrentPinTypes(const UPCGPin* InPin) const
 {
diff --git a/Engine/Plugins/Experimental/PCG/Source/PCG/Private/Tests/Elements/PCGBooleanSelectTest.cpp b/Engine/Plugins/Experimental/PCG/Source/PCG/Private/Tests/Elements/PCGBooleanSelectTest.cpp
new file mode 100644
--- /dev/null
+++ b/Engine/Plugins/Experimental/PCG/Source/PCG/Private/Tests/Elements/PCGBooleanSelectTest.cpp
@@ -0,0 +1,72 @@
+// Copyright Epic Games, Inc. All Rights Reserved.
+
+#include "PCGContext.h"
+#include "Data/PCGPointData.h"
+#include "Elements/PCGBooleanSelect.h"
+#include "Tests/PCGTestsCommon.h"
+
+#if WITH_EDITOR
+
+class FPCGBooleanSelectTest : public FPCGTestBaseClass
+{
+	using FPCGTestBaseClass::FPCGTestBaseClass;
+
+protected:
+	inline static const FString FirstDataTag = TEXT("FirstDataTag");
+	inline static const FString SecondDataTag = TEXT("SecondDataTag");
+	inline static const FName FirstInputPinLabel = TEXT("Input A");
+	inline static const FName SecondInputPinLabel = TEXT("Input B");
+
+	bool ExecuteTest(const bool bUseInputB, TArray<FPCGTaggedData>& OutputData)
+	{
+		PCGTestsCommon::FTestData TestData;
+		UPCGBooleanSelectSettings* TestSettings = PCGTestsCommon::GenerateSettings<UPCGBooleanSelectSettings>(TestData);
+		TestSettings->bUseInputB = bUseInputB;
+
+		const FPCGElementPtr TestElement = TestSettings->GetElement();
+
+		UTEST_TRUE("Test element created and valid", TestElement.IsValid());
+
+		TestData.InputData.TaggedData.Emplace(PCGTestsCommon::CreateEmptyPointData(), TSet({FirstDataTag}), FirstInputPinLabel);
+		TestData.InputData.TaggedData.Emplace(PCGTestsCommon::CreateEmptyPointData(), TSet({SecondDataTag}), SecondInputPinLabel);
+
+		const TUniquePtr<FPCGContext> TestContext = TestData.InitializeTestContext();
+
+		UTEST_TRUE("Test context created and valid", TestContext.IsValid());
+
+		while (!TestElement->Execute(TestContext.Get()))
+		{
+		}
+
+		OutputData = TestContext->OutputData.GetInputsByPin(PCGPinConstants::DefaultOutputLabel);
+
+		UTEST_TRUE("Output was generated", !OutputData.IsEmpty());
+
+		UTEST_EQUAL("Single output was generated", OutputData.Num(), 1);
+
+		UTEST_EQUAL("Only one tag output", OutputData.Last().Tags.Num(), 1);
+
+		return true;
+	}
+};
+
+IMPLEMENT_CUSTOM_SIMPLE_AUTOMATION_TEST(FPCGBooleanSelectTest_Basic, FPCGBooleanSelectTest, "Plugins.PCG.BooleanSelect.Basic", PCGTestsCommon::TestFlags)
+
+bool FPCGBooleanSelectTest_Basic::RunTest(const FString& Parameters)
+{
+	TArray<FPCGTaggedData> OutputData;
+
+	UTEST_TRUE("'Input A' test succeeded", ExecuteTest(false, OutputData));
+
+	UTEST_TRUE("Output contains correct data", OutputData[0].Tags.Contains(FirstDataTag));
+	UTEST_TRUE("Output does not contain incorrect data", !OutputData[0].Tags.Contains(SecondDataTag));
+
+	UTEST_TRUE("'Input B' test succeeded", ExecuteTest(true, OutputData));
+
+	UTEST_TRUE("Output contains correct data", OutputData[0].Tags.Contains(SecondDataTag));
+	UTEST_TRUE("Output does not contain incorrect data", !OutputData[0].Tags.Contains(FirstDataTag));
+
+	return true;
+}
+
+#endif // WITH_EDITOR
\ No newline at end of file
diff --git a/Engine/Plugins/Experimental/PCG/Source/PCG/Private/Tests/Elements/PCGMultiSelectTest.cpp b/Engine/Plugins/Experimental/PCG/Source/PCG/Private/Tests/Elements/PCGMultiSelectTest.cpp
new file mode 100644
--- /dev/null
+++ b/Engine/Plugins/Experimental/PCG/Source/PCG/Private/Tests/Elements/PCGMultiSelectTest.cpp
@@ -0,0 +1,251 @@
+// Copyright Epic Games, Inc. All Rights Reserved.
+
+#include "PCGContext.h"
+#include "Data/PCGPointData.h"
+#include "Elements/PCGMultiSelect.h"
+#include "Tests/PCGTestsCommon.h"
+
+#if WITH_EDITOR
+
+class FPCGMultiSelectTest : public FPCGTestBaseClass
+{
+	using FPCGTestBaseClass::FPCGTestBaseClass;
+
+protected:
+	inline static const FString DefaultPathPinDataTag = TEXT("DataTagDefault");
+	inline static const FName DefaultPathPinLabel = TEXT("Default");
+
+	static FString CreateDataTag(const int32 Index)
+	{
+		return FString("DataTag") + FString::FromInt(Index);
+	}
+
+	static void UpdateSettings(const UPCGMultiSelectSettings* InSettings, UPCGMultiSelectSettings* SelectSettings)
+	{
+		SelectSettings->SelectionMode = InSettings->SelectionMode;
+		SelectSettings->IntOptions = InSettings->IntOptions;
+		SelectSettings->IntegerSelection = InSettings->IntegerSelection;
+		SelectSettings->StringOptions = InSettings->StringOptions;
+		SelectSettings->StringSelection = InSettings->StringSelection;
+		SelectSettings->EnumSelection = InSettings->EnumSelection;
+		SelectSettings->CachePinLabels();
+	}
+
+	bool ExecuteTest(const UPCGMultiSelectSettings* InSettings, const int32 InputNum, TArray<FPCGTaggedData>& OutputData)
+	{
+		PCGTestsCommon::FTestData TestData;
+		UPCGMultiSelectSettings* TestSettings = PCGTestsCommon::GenerateSettings<UPCGMultiSelectSettings>(TestData);
+		UpdateSettings(InSettings, TestSettings);
+
+		const FPCGElementPtr TestElement = TestSettings->GetElement();
+
+		UTEST_TRUE("Test element created and valid", TestElement.IsValid());
+
+		// Add a tagged input to each of the valid pin paths
+		TArray<FPCGPinProperties> SelectElementPinProperties = TestSettings->AllInputPinProperties();
+		for (int i = 0; i < InputNum; ++i)
+		{
+			if (!SelectElementPinProperties[i].bAdvancedPin && SelectElementPinProperties[i].Label != DefaultPathPinLabel)
+			{
+				TestData.InputData.TaggedData.Emplace(PCGTestsCommon::CreateEmptyPointData(), TSet{CreateDataTag(i)}, SelectElementPinProperties[i].Label);
+			}
+		}
+
+		// Add an input to the default path
+		TestData.InputData.TaggedData.Emplace(PCGTestsCommon::CreateEmptyPointData(), TSet{DefaultPathPinDataTag}, DefaultPathPinLabel);
+
+		const TUniquePtr<FPCGContext> TestContext = TestData.InitializeTestContext();
+
+		UTEST_TRUE("Test context created and valid", TestContext.IsValid());
+
+		while (!TestElement->Execute(TestContext.Get()))
+		{
+		}
+
+		OutputData = TestContext->OutputData.GetInputsByPin(PCGPinConstants::DefaultOutputLabel);
+
+		return true;
+	}
+
+	bool ValidateOutput(const TArray<FPCGTaggedData>& Outputs)
+	{
+		UTEST_TRUE("Output was generated", !Outputs.IsEmpty());
+
+		UTEST_EQUAL("Single output was generated", Outputs.Num(), 1);
+
+		UTEST_EQUAL("Only one tag output", Outputs.Last().Tags.Num(), 1);
+
+		return true;
+	}
+
+	bool ValidateOutputInt(const TArray<FPCGTaggedData>& OutputData, const int32 ValidKey)
+	{
+		UTEST_TRUE("Validate output and output count", ValidateOutput(OutputData));
+
+		UTEST_TRUE("Output contains correct data", OutputData.Last().Tags.Contains(CreateDataTag(ValidKey)));
+
+		return true;
+	}
+
+	bool ValidateOutputString(const TArray<FString>& Options, const TArray<FPCGTaggedData>& OutputData, const FString& ValidKey)
+	{
+		UTEST_TRUE("Validate output and output count", ValidateOutput(OutputData));
+
+		for (int i = 0; i < Options.Num(); ++i)
+		{
+			if (Options[i] == ValidKey)
+			{
+				UTEST_TRUE("Output contains correct data", OutputData.Last().Tags.Contains(CreateDataTag(i)));
+			}
+		}
+
+		return true;
+	}
+
+	bool ValidateOutputEnum(const UEnum* Enum, const TArray<FPCGTaggedData>& OutputData, const int64 ValidKey)
+	{
+		UTEST_NOT_NULL("Enum is valid", Enum);
+
+		UTEST_TRUE("Validate output and output count", ValidateOutput(OutputData));
+
+		UTEST_TRUE("Output contains correct data", OutputData.Last().Tags.Contains(CreateDataTag(ValidKey)));
+
+		return true;
+	}
+
+	bool ValidateDefaultPathOutput(const TArray<FPCGTaggedData>& OutputData)
+	{
+		UTEST_TRUE("Validate output and output count", ValidateOutput(OutputData));
+
+		UTEST_TRUE("Output contains correct data", OutputData.Last().Tags.Contains(DefaultPathPinDataTag));
+
+		return true;
+	}
+};
+
+IMPLEMENT_CUSTOM_SIMPLE_AUTOMATION_TEST(FPCGMultiSelectTest_Integer, FPCGMultiSelectTest, "Plugins.PCG.MultiSelect.Integer", PCGTestsCommon::TestFlags)
+
+bool FPCGMultiSelectTest_Integer::RunTest(const FString& Parameters)
+{
+	TArray<FPCGTaggedData> OutputData;
+	UPCGMultiSelectSettings* TestSettings = NewObject<UPCGMultiSelectSettings>();
+	TestSettings->SelectionMode = EPCGControlFlowSelectionMode::Integer;
+
+	TArray<int32>& Options = TestSettings->IntOptions;
+	Options = {-1, 0, 2, 3, 5};
+
+	// Positive Validation
+	for (int i = 0; i < Options.Num(); ++i)
+	{
+		TestSettings->IntegerSelection = Options[i];
+
+		UTEST_TRUE("Integer test succeeded", ExecuteTest(TestSettings, Options.Num(), OutputData));
+
+		UTEST_TRUE("Integer selection is valid", ValidateOutputInt(OutputData, i));
+	}
+
+	const int32 Min = FMath::Min(Options);
+	const int32 Max = FMath::Max(Options);
+	for (int i = Min - 1; i <= Max + 1; ++i)
+	{
+		if (Options.Contains(i))
+		{
+			continue;
+		}
+
+		TestSettings->IntegerSelection = i;
+
+		UTEST_TRUE("Integer test succeeded", ExecuteTest(TestSettings, Options.Num(), OutputData));
+
+		UTEST_TRUE("Invalid int selection defaults output", ValidateDefaultPathOutput(OutputData));
+	}
+
+	return true;
+}
+
+IMPLEMENT_CUSTOM_SIMPLE_AUTOMATION_TEST(FPCGMultiSelectTest_String, FPCGMultiSelectTest, "Plugins.PCG.MultiSelect.String", PCGTestsCommon::TestFlags)
+
+bool FPCGMultiSelectTest_String::RunTest(const FString& Parameters)
+{
+	TArray<FPCGTaggedData> OutputData;
+	UPCGMultiSelectSettings* TestSettings = NewObject<UPCGMultiSelectSettings>();
+	TestSettings->SelectionMode = EPCGControlFlowSelectionMode::String;
+
+	TArray<FString>& Options = TestSettings->StringOptions;
+	Options = {"Red", "Green", "Blue", "Fuschia", "Orange"};
+
+	FString& Selection = TestSettings->StringSelection;
+
+	// Positive Validation
+	for (int i = 0; i < Options.Num(); ++i)
+	{
+		Selection = Options[i];
+
+		UTEST_TRUE("String test succeeded", ExecuteTest(TestSettings, Options.Num(), OutputData));
+
+		UTEST_TRUE("String selection is valid", ValidateOutputString(Options, OutputData, Selection));
+	}
+
+	// Negative Validation
+	Selection = "Square";
+
+	UTEST_TRUE("String test succeeded", ExecuteTest(TestSettings, Options.Num(), OutputData));
+
+	UTEST_TRUE("Invalid string selection defaults output", ValidateDefaultPathOutput(OutputData));
+
+	return true;
+}
+
+IMPLEMENT_CUSTOM_SIMPLE_AUTOMATION_TEST(FPCGMultiSelectTest_Enum, FPCGMultiSelectTest, "Plugins.PCG.MultiSelect.Enum", PCGTestsCommon::TestFlags)
+
+bool FPCGMultiSelectTest_Enum::RunTest(const FString& Parameters)
+{
+	TArray<FPCGTaggedData> OutputData;
+	UPCGMultiSelectSettings* TestSettings = NewObject<UPCGMultiSelectSettings>();
+	TestSettings->SelectionMode = EPCGControlFlowSelectionMode::Enum;
+
+	UEnum*& EnumClass = TestSettings->EnumSelection.Class;
+	EnumClass = StaticEnum<EPCGMetadataTypes>();
+
+	const int32 NumEnums = EnumClass->NumEnums();
+
+	// Two for the hidden enum values
+	AddExpectedError("Selected value is not a valid option.", EAutomationExpectedMessageFlags::Contains, 2);
+
+	// Positive Validation (also accounting for hidden metadata)
+	for (uint8 i = 0; i < NumEnums - 1; ++i)
+	{
+		TestSettings->EnumSelection.Value = EnumClass->GetValueByIndex(i);
+
+		UTEST_TRUE("Selection value is valid", TestSettings->EnumSelection.Value != INDEX_NONE);
+
+		UTEST_TRUE("Enum test succeeded", ExecuteTest(TestSettings, NumEnums - 1, OutputData));
+
+		if (TestSettings->EnumSelection.Class->HasMetaData(TEXT("Hidden"), i) || TestSettings->EnumSelection.Class->HasMetaData(TEXT("Spacer"), i))
+		{
+			UTEST_TRUE("Hidden enum selection output is empty", OutputData.IsEmpty());
+		}
+		else
+		{
+			UTEST_TRUE("Enum selection is valid", ValidateOutputEnum(EnumClass, OutputData, i));
+		}
+	}
+
+	// Negative Validation - Invalid selection value
+	TestSettings->EnumSelection.Value = NumEnums;
+
+	UTEST_TRUE("Enum test succeeded", ExecuteTest(TestSettings, NumEnums, OutputData));
+
+	UTEST_TRUE("Invalid enum value defaults output", ValidateDefaultPathOutput(OutputData));
+
+	// Negative Validation - Invalid enum
+	EnumClass = nullptr;
+
+	UTEST_TRUE("Enum test succeeded", ExecuteTest(TestSettings, 1, OutputData));
+
+	UTEST_TRUE("Invalid enum class defaults output", ValidateDefaultPathOutput(OutputData));
+
+	return true;
+}
+
+#endif // WITH_EDITOR
\ No newline at end of file
diff --git a/Engine/Plugins/Experimental/PCG/Source/PCG/Public/Elements/PCGMultiSelect.h b/Engine/Plugins/Experimental/PCG/Source/PCG/Public/Elements/PCGMultiSelect.h
new file mode 100644
--- /dev/null
+++ b/Engine/Plugins/Experimental/PCG/Source/PCG/Public/Elements/PCGMultiSelect.h
@@ -0,0 +1,100 @@
+// Copyright Epic Games, Inc. All Rights Reserved.
+
+#pragma once
+
+#include "PCGSettings.h"
+#include "Helpers/PCGPropertyHelpers.h"
+
+#include "PCGMultiSelect.generated.h"
+
+UENUM()
+enum class EPCGControlFlowSelectionMode : uint8
+{
+	Integer,
+	Enum,
+	String
+};
+
+/**
+ * Selects data from any number of input pins, based on a static selection criteria (Int/String/Enum) 
+ */
+UCLASS(BlueprintType, ClassGroup = (Procedural), meta=(Keywords = "if multi switch enum"))
+class UPCGMultiSelectSettings : public UPCGSettings
+{
+	GENERATED_BODY()
+
+public:
+	//~Begin UObject interface
+	virtual void PostLoad() override;
+
+#if WITH_EDITOR
+	virtual void PostEditChangeProperty(FPropertyChangedEvent& PropertyChangedEvent) override;
+	//~End UObject interface
+
+	//~Begin UPCGSettings interface
+	virtual FName GetDefaultNodeName() const override { return FName(TEXT("Select (Multi)")); }
+	virtual FText GetDefaultNodeTitle() const override;
+	virtual FText GetNodeTooltipText() const override;
+	virtual EPCGSettingsType GetType() const override { return EPCGSettingsType::ControlFlow; }
+	virtual bool HasDynamicPins() const override { return true; }
+#endif // WITH_EDITOR
+
+	virtual EPCGDataType GetCurrentPinTypes(const UPCGPin* Pin) const override;
+	virtual FString GetAdditionalTitleInformation() const override;
+	virtual bool HasFlippedTitleLines() const override { return true; }
+
+protected:
+	virtual TArray<FPCGPinProperties> InputPinProperties() const override;
+	virtual TArray<FPCGPinProperties> OutputPinProperties() const override;
+	virtual FPCGElementPtr CreateElement() const override;
+	//~End UPCGSettings interface
+
+public:
+	/** Determines the type of value to be used to select an input. */
+	UPROPERTY(BlueprintReadWrite, EditAnywhere, Category=Settings)
+	EPCGControlFlowSelectionMode SelectionMode = EPCGControlFlowSelectionMode::Integer;
+
+	/** Determines which input will be selected if the selection mode is Integer. */
+	UPROPERTY(BlueprintReadWrite, EditAnywhere, Category=Settings, meta=(PCG_Overridable, EditConditionHides, EditCondition="SelectionMode == EPCGControlFlowSelectionMode::Integer"))
+	int32 IntegerSelection = 0;
+
+	/** Determines the available input pin selection options. */
+	UPROPERTY(BlueprintReadWrite, EditAnywhere, Category=Settings, meta=(EditConditionHides, EditCondition="SelectionMode == EPCGControlFlowSelectionMode::Integer"))
+	TArray<int32> IntOptions = {0};
+
+	/** Determines which input will be selected if the selection mode is String. */
+	UPROPERTY(BlueprintReadWrite, EditAnywhere, Category=Settings, meta=(PCG_Overridable, EditConditionHides, NoResetToDefault, EditCondition="SelectionMode == EPCGControlFlowSelectionMode::String"))
+	FString StringSelection;
+
+	/** Determines the available input pin selection options. */
+	UPROPERTY(BlueprintReadWrite, EditAnywhere, Category=Settings, meta=(EditConditionHides, EditCondition="SelectionMode == EPCGControlFlowSelectionMode::String"))
+	TArray<FString> StringOptions;
+
+	/** Determines which input pin will be selected if the selection mode is Enum. */
+	UPROPERTY(BlueprintReadWrite, EditAnywhere, Category=Settings, meta=(PCG_Overridable, EditConditionHides, EditCondition="SelectionMode == EPCGControlFlowSelectionMode::Enum"))
+	FEnumSelector EnumSelection;
+
+	/** Cached pin labels for use during the selection process. */
+	UPROPERTY(Transient)
+	TArray<FName> CachedPinLabels;
+
+	/** Returns true if the integer value exists in the user defined array. */
+	bool IsValuePresent(int32 Value) const;
+
+	/** Returns true if the string value exists in the user defined array. */
+	bool IsValuePresent(const FString& Value) const;
+
+	/** Returns true if the enum value exists within the enum class. */
+	bool IsValuePresent(int64 Value) const;
+
+	/** Helper function to use the appropriate selection value to determine the current selection. Returns true if it succeeds or false if the value is invalid. */
+	bool GetSelectedPinLabel(FName& OutSelectedPinLabel) const;
+
+	void CachePinLabels();
+};
+
+class FPCGMultiSelectElement : public IPCGElement
+{
+protected:
+	virtual bool ExecuteInternal(FPCGContext* Context) const override;
+};
diff --git a/Engine/Plugins/Experimental/PCG/Source/PCG/Public/Helpers/PCGPropertyHelpers.h b/Engine/Plugins/Experimental/PCG/Source/PCG/Public/Helpers/PCGPropertyHelpers.h
--- a/Engine/Plugins/Experimental/PCG/Source/PCG/Public/Helpers/PCGPropertyHelpers.h
+++ b/Engine/Plugins/Experimental/PCG/Source/PCG/Public/Helpers/PCGPropertyHelpers.h
@@ -11,12 +11,26 @@
 
 #include "UObject/UnrealType.h" // IWYU pragma: keep
 
+#include "PCGPropertyHelpers.generated.h"
+
 enum class EPCGMetadataTypes : uint8;
 
 struct FPCGContext;
 class UPCGData;
 class UPCGParamData;
 
+USTRUCT(BlueprintType, meta = (Hidden))
+struct PCG_API FEnumSelector
+{
+	GENERATED_BODY()
+
+	UPROPERTY(DisplayName="Enum Class", meta=(PCG_NotOverridable))
+	UEnum* Class = nullptr;
+
+	UPROPERTY(DisplayName="Enum Value")
+	int64 Value = 0;
+};
+
 namespace PCGPropertyHelpers
 {
 	/**
diff --git a/Engine/Plugins/Experimental/PCG/Source/PCG/Public/PCGSettings.h b/Engine/Plugins/Experimental/PCG/Source/PCG/Public/PCGSettings.h
--- a/Engine/Plugins/Experimental/PCG/Source/PCG/Public/PCGSettings.h
+++ b/Engine/Plugins/Experimental/PCG/Source/PCG/Public/PCGSettings.h
@@ -286,15 +286,14 @@ public:
 	 *  By default, it will return the underlying class, to try to jump to its header in code
      */
 	virtual UObject* GetJumpTargetForDoubleClick() const;
-	virtual bool IsPropertyOverriddenByPin(const FProperty* InProperty) const;
 
-	/* Return preconfigured info that will be filled in the editor palette action, allowing to create pre-configured settings */
+	/** Return preconfigured info that will be filled in the editor palette action, allowing to create pre-configured settings */
 	virtual TArray<FPCGPreConfiguredSettingsInfo> GetPreconfiguredInfo() const { return {}; }
 
-	/* If there are preconfigured info, we can skip the default settings and only expose pre-configured actions in the editor palette */
+	/** If there are preconfigured info, we can skip the default settings and only expose pre-configured actions in the editor palette */
 	virtual bool OnlyExposePreconfiguredSettings() const { return false; }
 
-	/* Perform post-operations when an editor node is copied */
+	/** Perform post-operations when an editor node is copied */
 	virtual void PostPaste();
 #endif // WITH_EDITOR
 
@@ -304,6 +303,15 @@ public:
 	/** Display generated title line as primary title (example: BP nodes display the blueprint name as the primary title). */
 	virtual bool HasFlippedTitleLines() const { return false; }
 
+	/** Returns true if the given property is overridden by graph parameters */
+	virtual bool IsPropertyOverriddenByPin(const FProperty* InProperty) const;
+
+	/** Returns true if the base property, given by name, is overridden by graph parameters */
+	virtual bool IsPropertyOverriddenByPin(FName PropertyName) const;
+
+	/** Returns true if the property, given by chain of property names, is overridden by graph parameters */
+	virtual bool IsPropertyOverriddenByPin(const TArrayView<const FName>& PropertyNameChain) const;
+
 	virtual void ApplyPreconfiguredSettings(const FPCGPreConfiguredSettingsInfo& PreconfigureInfo) {}
 
 	UE_DEPRECATED(5.4, "AdditionalTaskName is deprecated and is replaced with GetAdditionalTitleInformation.")
diff --git a/Engine/Plugins/Experimental/PCG/Source/PCGEditor/PCGEditor.Build.cs b/Engine/Plugins/Experimental/PCG/Source/PCGEditor/PCGEditor.Build.cs
--- a/Engine/Plugins/Experimental/PCG/Source/PCGEditor/PCGEditor.Build.cs
+++ b/Engine/Plugins/Experimental/PCG/Source/PCGEditor/PCGEditor.Build.cs
@@ -63,6 +63,7 @@ namespace UnrealBuildTool.Rules
 					"StructUtils",
 					"StructUtilsEditor",
 					"ToolMenus",
+					"ToolWidgets",
 					"TypedElementFramework",
 					"TypedElementRuntime",
 					"UnrealEd",
diff --git a/Engine/Plugins/Experimental/PCG/Source/PCGEditor/Private/Details/EnumSelectorDetails.cpp b/Engine/Plugins/Experimental/PCG/Source/PCGEditor/Private/Details/EnumSelectorDetails.cpp
new file mode 100644
--- /dev/null
+++ b/Engine/Plugins/Experimental/PCG/Source/PCGEditor/Private/Details/EnumSelectorDetails.cpp
@@ -0,0 +1,328 @@
+// Copyright Epic Games, Inc. All Rights Reserved.
+
+#include "Details/EnumSelectorDetails.h"
+
+#include "Helpers/PCGPropertyHelpers.h"
+
+#include "DetailLayoutBuilder.h"
+#include "DetailWidgetRow.h"
+#include "ScopedTransaction.h"
+#include "SSearchableComboBox.h"
+#include "UObject/UObjectIterator.h"
+
+#define LOCTEXT_NAMESPACE "EnumSelectorDetails"
+
+namespace FEnumSelectorDetailsConstants
+{
+	inline static const FText SetSelectorClassTransaction = LOCTEXT("EditEnumSelectorClass", "Edit Enum Selector Class");
+	inline static const FText SetEnumSelectorValueTransaction = LOCTEXT("EditEnumSelectorValue", "Edit Enum Selector Value");
+}
+
+void FEnumSelectorDetails::CustomizeHeader(const TSharedRef<IPropertyHandle> InPropertyHandle, FDetailWidgetRow& HeaderRow, IPropertyTypeCustomizationUtils& CustomizationUtils)
+{
+	PropertyHandle = InPropertyHandle;
+	check(PropertyHandle.IsValid());
+
+	const FEnumSelector* Selector = GetEnumSelector();
+	if (!Selector)
+	{
+		return;
+	}
+
+	if (EnumClassSourceOptions.IsEmpty())
+	{
+		InitializeClassComboBox();
+	}
+
+	if (!ClassSelection.IsValid())
+	{
+		InitializeClassComboBoxSelection();
+	}
+
+	HeaderRow
+		.NameContent()
+		[
+			PropertyHandle->CreatePropertyNameWidget()
+		]
+		.ValueContent()
+		[
+			SNew(SVerticalBox)
+			// Enum Class Selection
+			+ SVerticalBox::Slot()
+			.AutoHeight()
+			[
+				SNew(SHorizontalBox)
+				// Enum Class Selection Title
+				+ SHorizontalBox::Slot()
+				.AutoWidth()
+				.VAlign(VAlign_Center)
+				[
+					SNew(STextBlock)
+					.Font(IDetailLayoutBuilder::GetDetailFont())
+					.Text(LOCTEXT("EnumClassSelection", "Class"))
+				]
+				// Enum Class Selection Options
+				+ SHorizontalBox::Slot()
+				.HAlign(HAlign_Left)
+				.FillWidth(1.f)
+				.Padding(12.f, 1.f, 3.f, 1.f)
+				[
+					// This is currently required to be TSharedPtr<FString>, which is not optimal
+					// TODO: UE-201874
+					SNew(SSearchableComboBox)
+					.OptionsSource(&EnumClassSourceOptions)
+					.OnGenerateWidget_Lambda([this](const TSharedPtr<FString> InItemString)
+					{
+						return SNew(STextBlock).Text(FText::FromString(*InItemString));
+					})
+					.OnSelectionChanged(this, &FEnumSelectorDetails::OnClassComboSelectionChanged)
+					.InitiallySelectedItem(ClassSelection)
+					.Content()
+					[
+						SNew(STextBlock)
+						.Font(IDetailLayoutBuilder::GetDetailFont())
+						.Text(this, &FEnumSelectorDetails::GetSelectedEnumClassText)
+					]
+				]
+			]
+			// Enum Value Selection
+			+ SVerticalBox::Slot()
+			.AutoHeight()
+			[
+				SNew(SHorizontalBox)
+				.Visibility_Raw(this, &FEnumSelectorDetails::IsValueSelectionVisible)
+				// Enum Value Selection Title
+				+ SHorizontalBox::Slot()
+				.AutoWidth()
+				.VAlign(VAlign_Center)
+				[
+					SNew(STextBlock)
+					.Font(IDetailLayoutBuilder::GetDetailFont())
+					.Text(LOCTEXT("EnumValueSelection", "Value"))
+				]
+				// Enum Value Selection Options
+				+ SHorizontalBox::Slot()
+				.HAlign(HAlign_Left)
+				.FillWidth(1.f)
+				.Padding(12.f, 1.f, 3.f, 1.f)
+				[
+					SAssignNew(EnumValueComboBox, SComboBox<FName>)
+					.OptionsSource(&EnumValueSourceOptions)
+					.OnSelectionChanged_Raw(this, &FEnumSelectorDetails::OnValueComboSelectionChanged)
+					.InitiallySelectedItem(ValueSelection)
+					.OnGenerateWidget_Lambda([this](const FName InItemName)
+					{
+						return SNew(STextBlock)
+						.Font(IDetailLayoutBuilder::GetDetailFont())
+						.Text(StripEnumClassName(InItemName));
+					})
+					.Content()
+					[
+						SNew(STextBlock)
+						.Font(IDetailLayoutBuilder::GetDetailFont())
+						.Text(this, &FEnumSelectorDetails::GetSelectedEnumValueText)
+					]
+				]
+			]
+		];
+}
+
+FText FEnumSelectorDetails::StripEnumClassName(const FName EnumValueName)
+{
+	FString EnumValueString = EnumValueName.ToString();
+
+	int32 DelimiterIndex;
+	if (EnumValueString.FindLastChar(':', DelimiterIndex))
+	{
+		EnumValueString.MidInline(DelimiterIndex + 1);
+	}
+
+	return FText::FromString(EnumValueString);
+}
+
+FEnumSelector* FEnumSelectorDetails::GetEnumSelector()
+{
+	return const_cast<FEnumSelector*>(std::as_const(*this).GetEnumSelector());
+}
+
+const FEnumSelector* FEnumSelectorDetails::GetEnumSelector() const
+{
+	void* Data = nullptr;
+	return (PropertyHandle->GetValueData(Data) == FPropertyAccess::Success) ? static_cast<const FEnumSelector*>(Data) : nullptr;
+}
+
+UEnum* FEnumSelectorDetails::GetSelectedEnumClass() const
+{
+	const FEnumSelector* Selector = GetEnumSelector();
+	return Selector ? Selector->Class : nullptr;
+}
+
+int64 FEnumSelectorDetails::GetSelectedEnumValue() const
+{
+	const FEnumSelector* Selector = GetEnumSelector();
+	return Selector ? Selector->Value : 0;
+}
+
+void FEnumSelectorDetails::SetSelectorEnumClass(UEnum* Class)
+{
+	if (FEnumSelector* Selector = GetEnumSelector())
+	{
+		FScopedTransaction Transaction(FEnumSelectorDetailsConstants::SetSelectorClassTransaction);
+		PropertyHandle->NotifyPreChange();
+		Selector->Class = Class;
+		PropertyHandle->NotifyPostChange(EPropertyChangeType::ValueSet);
+		PropertyHandle->NotifyFinishedChangingProperties();
+	}
+}
+
+void FEnumSelectorDetails::SetSelectorEnumValue(const int64 Value)
+{
+	if (FEnumSelector* Selector = GetEnumSelector())
+	{
+		FScopedTransaction Transaction(FEnumSelectorDetailsConstants::SetEnumSelectorValueTransaction);
+		PropertyHandle->NotifyPreChange();
+		Selector->Value = Value;
+		PropertyHandle->NotifyPostChange(EPropertyChangeType::ValueSet);
+		PropertyHandle->NotifyFinishedChangingProperties();
+	}
+}
+
+void FEnumSelectorDetails::UpdateEnumClass(const TSharedPtr<FString>& EnumString)
+{
+	if (EnumString.IsValid())
+	{
+		if (!NameStringToEnumClassMap.Contains(EnumString))
+		{
+			return;
+		}
+
+		SetSelectorEnumClass(NameStringToEnumClassMap[EnumString]);
+
+		ReloadValueOptions();
+
+		if (EnumValueComboBox)
+		{
+			EnumValueComboBox->RefreshOptions();
+			EnumValueComboBox->SetSelectedItem(EnumValueSourceOptions[0]);
+		}
+	}
+}
+
+void FEnumSelectorDetails::UpdateEnumValue(const FName ValueName)
+{
+	const FEnumSelector* Selector = GetEnumSelector();
+
+	if (Selector && Selector->Class)
+	{
+		SetSelectorEnumValue(Selector->Class->GetValueByName(ValueName));
+	}
+}
+
+void FEnumSelectorDetails::InitializeClassComboBox()
+{
+	// Find all the enums
+	for (TObjectIterator<UEnum> It; It; ++It)
+	{
+		if (!It->HasMetaData(TEXT("Hidden")))
+		{
+			EnumClassSourceOptions.Emplace(MakeShared<FString>(It->GetName()));
+			NameStringToEnumClassMap.Emplace(EnumClassSourceOptions.Last(), *It);
+		}
+	}
+
+	// Sort the options so they're alphabetical
+	EnumClassSourceOptions.Sort([](const TSharedPtr<FString>& Lhs, const TSharedPtr<FString>& Rhs)
+	{
+		return *Lhs < *Rhs;
+	});
+}
+
+void FEnumSelectorDetails::InitializeClassComboBoxSelection()
+{
+	bool bEnumFound = false;
+	// The selector is already defined, so use the struct data. Note: this is costly, but is needed to get the pointer to the correct widget option
+	if (UEnum* SelectedEnum = GetSelectedEnumClass())
+	{
+		const TSharedPtr<FString>* SelectedOption = EnumClassSourceOptions.FindByPredicate([SelectedEnum](const TSharedPtr<FString>& InString)
+		{
+			return SelectedEnum->GetName() == *InString;
+		});
+
+		if (SelectedOption)
+		{
+			ClassSelection = *SelectedOption;
+
+			ReloadValueOptions();
+			ValueSelection = SelectedEnum->GetNameByValue(GetSelectedEnumValue());
+			UpdateEnumValue(ValueSelection);
+			bEnumFound = true;
+		}
+	}
+
+	// The selector is null or the enum is invalid, so default to the first option
+	if (!bEnumFound)
+	{
+		// Note: already verified this wasn't empty earlier
+		UpdateEnumClass(EnumClassSourceOptions[0]);
+		ClassSelection = EnumClassSourceOptions[0];
+
+		// In the off chance that the enum has no options or only a MAX value
+		if (!EnumValueSourceOptions.IsEmpty())
+		{
+			UpdateEnumValue(EnumValueSourceOptions[0]);
+			ValueSelection = EnumValueSourceOptions[0];
+		}
+	}
+}
+
+void FEnumSelectorDetails::OnClassComboSelectionChanged(TSharedPtr<FString> InSelectedItem, ESelectInfo::Type SelectInfo)
+{
+	UpdateEnumClass(InSelectedItem);
+}
+
+void FEnumSelectorDetails::OnValueComboSelectionChanged(const FName InSelectedItem, ESelectInfo::Type SelectInfo)
+{
+	UpdateEnumValue(InSelectedItem);
+}
+
+void FEnumSelectorDetails::ReloadValueOptions()
+{
+	const FEnumSelector* Selector = GetEnumSelector();
+
+	// -1 to ignore the MAX
+	const int Count = Selector->Class->NumEnums() - 1;
+
+	EnumValueSourceOptions.Reset(Count);
+
+	// Collect all the valid values in the enum
+	for (int64 i = 0; i < Count; ++i)
+	{
+		bool const bHidden = Selector->Class->HasMetaData(TEXT("Hidden"), i) || Selector->Class->HasMetaData(TEXT("Spacer"), i);
+		if (!bHidden)
+		{
+			EnumValueSourceOptions.Emplace(Selector->Class->GetNameByIndex(i));
+		}
+	}
+
+	EnumValueSourceOptions.Shrink();
+}
+
+EVisibility FEnumSelectorDetails::IsValueSelectionVisible() const
+{
+	const FEnumSelector* Selector = GetEnumSelector();
+	return (Selector && Selector->Class != nullptr) ? EVisibility::Visible : EVisibility::Hidden;
+}
+
+FText FEnumSelectorDetails::GetSelectedEnumClassText() const
+{
+	const FEnumSelector* Selector = GetEnumSelector();
+	return Selector->Class ? Selector->Class->GetDisplayNameText() : FText::FromString(*EnumClassSourceOptions[0]);
+}
+
+FText FEnumSelectorDetails::GetSelectedEnumValueText() const
+{
+	const FEnumSelector* Selector = GetEnumSelector();
+	return Selector->Class ? StripEnumClassName(Selector->Class->GetNameByValue(Selector->Value)) : FText();
+}
+
+#undef LOCTEXT_NAMESPACE
diff --git a/Engine/Plugins/Experimental/PCG/Source/PCGEditor/Private/PCGEditorModule.cpp b/Engine/Plugins/Experimental/PCG/Source/PCGEditor/Private/PCGEditorModule.cpp
--- a/Engine/Plugins/Experimental/PCG/Source/PCGEditor/Private/PCGEditorModule.cpp
+++ b/Engine/Plugins/Experimental/PCG/Source/PCGEditor/Private/PCGEditorModule.cpp
@@ -16,6 +16,7 @@
 #include "ISettingsModule.h"
 #include "PropertyEditorModule.h"
 #include "ToolMenus.h"
+#include "Details/EnumSelectorDetails.h"
 #include "Details/PCGAttributePropertySelectorDetails.h"
 #include "Details/PCGBlueprintSettingsDetails.h"
 #include "Details/PCGGraphDetails.h"
@@ -114,6 +115,7 @@ void FPCGEditorModule::RegisterDetailsCustomizations()
 	PropertyEditor.RegisterCustomPropertyTypeLayout("PCGAttributePropertyOutputSelector", FOnGetPropertyTypeCustomizationInstance::CreateStatic(&FPCGAttributePropertySelectorDetails::MakeInstance));
 	PropertyEditor.RegisterCustomPropertyTypeLayout("PCGAttributePropertyOutputNoSourceSelector", FOnGetPropertyTypeCustomizationInstance::CreateStatic(&FPCGAttributePropertySelectorDetails::MakeInstance));
 	PropertyEditor.RegisterCustomPropertyTypeLayout("PCGOverrideInstancedPropertyBag", FOnGetPropertyTypeCustomizationInstance::CreateStatic(&FPCGOverrideInstancedPropertyBagDetails::MakeInstance));
+	PropertyEditor.RegisterCustomPropertyTypeLayout("EnumSelector", FOnGetPropertyTypeCustomizationInstance::CreateStatic(&FEnumSelectorDetails::MakeInstance));
 
 	PropertyEditor.NotifyCustomizationModuleChanged();
 }
diff --git a/Engine/Plugins/Experimental/PCG/Source/PCGEditor/Public/Details/EnumSelectorDetails.h b/Engine/Plugins/Experimental/PCG/Source/PCGEditor/Public/Details/EnumSelectorDetails.h
new file mode 100644
--- /dev/null
+++ b/Engine/Plugins/Experimental/PCG/Source/PCGEditor/Public/Details/EnumSelectorDetails.h
@@ -0,0 +1,87 @@
+// Copyright Epic Games, Inc. All Rights Reserved.
+
+#pragma once
+
+#include "IPropertyTypeCustomization.h"
+#include "Layout/Visibility.h"
+#include "Widgets/Input/SComboBox.h"
+
+struct FEnumSelector;
+
+/** Property type customization for an enum selector allowing the user to easily select from a list of available
+ * enum classes and then also select a value from that class's options.
+ */
+class FEnumSelectorDetails : public IPropertyTypeCustomization
+{
+public:
+	static TSharedRef<IPropertyTypeCustomization> MakeInstance()
+	{
+		return MakeShareable(new FEnumSelectorDetails);
+	}
+
+	/** ~Begin IPropertyTypeCustomization interface */
+	virtual void CustomizeHeader(TSharedRef<IPropertyHandle> InPropertyHandle, FDetailWidgetRow& HeaderRow, IPropertyTypeCustomizationUtils& CustomizationUtils) override;
+
+	virtual void CustomizeChildren(TSharedRef<IPropertyHandle> InPropertyHandle, IDetailChildrenBuilder& ChildBuilder, IPropertyTypeCustomizationUtils& CustomizationUtils) override
+	{
+		// Must override pure virtual function
+	}
+
+	/** ~End IPropertyTypeCustomization interface */
+
+	/** Strips the enum class from the enum value name: EnumClass::EnumValue -> EnumValue. */
+	static FText StripEnumClassName(FName EnumValueName);
+
+protected:
+	FEnumSelector* GetEnumSelector();
+
+	const FEnumSelector* GetEnumSelector() const;
+
+	/** Helper to convert the drop down selection into an enum class. */
+	UEnum* GetSelectedEnumClass() const;
+
+	/** Helper to convert the drop down selection into an enum value. */
+	int64 GetSelectedEnumValue() const;
+
+	/** Sets the selector enum class and calls the property update functions. */
+	void SetSelectorEnumClass(UEnum* Class);
+
+	/** Sets the selector enum value and calls the property update functions. */
+	void SetSelectorEnumValue(const int64 Value);
+
+	/** Updates the enum class from the SSearchableComboBox item type. */
+	void UpdateEnumClass(const TSharedPtr<FString>& EnumString);
+
+	/** Updates the enum value from the SComboBox item type. */
+	void UpdateEnumValue(FName ValueName);
+
+	void InitializeClassComboBox();
+	void InitializeClassComboBoxSelection();
+
+	void OnClassComboSelectionChanged(TSharedPtr<FString> InSelectedItem, ESelectInfo::Type SelectInfo);
+	void OnValueComboSelectionChanged(FName InSelectedItem, ESelectInfo::Type SelectInfo);
+
+	/** Reloads the enum values associated with an enum class, once a different enum class has been selected. */
+	void ReloadValueOptions();
+
+	/** Returns true only if the selector and enum class are both valid, false otherwise. */
+	EVisibility IsValueSelectionVisible() const;
+
+	/** Gets the display text for the currently selected enum class. */
+	FText GetSelectedEnumClassText() const;
+
+	/** Gets the display text for the currently selected enum value. */
+	FText GetSelectedEnumValueText() const;
+
+	TSharedPtr<IPropertyHandle> PropertyHandle;
+
+	TArray<TSharedPtr<FString>> EnumClassSourceOptions;
+	TArray<FName> EnumValueSourceOptions;
+
+	TMap<TSharedPtr<FString>, UEnum*> NameStringToEnumClassMap;
+
+	TSharedPtr<SComboBox<FName>> EnumValueComboBox;
+
+	TSharedPtr<FString> ClassSelection;
+	FName ValueSelection;
+};
\ No newline at end of file
