[PCG] Fix bug for PCGBooleanSelectTest and PCGMultiSelectTest compiling on Mac

#rnx

diff --git a/Engine/Plugins/Experimental/PCG/Source/PCG/Private/Tests/Elements/PCGBooleanSelectTest.cpp b/Engine/Plugins/Experimental/PCG/Source/PCG/Private/Tests/Elements/PCGBooleanSelectTest.cpp
--- a/Engine/Plugins/Experimental/PCG/Source/PCG/Private/Tests/Elements/PCGBooleanSelectTest.cpp
+++ b/Engine/Plugins/Experimental/PCG/Source/PCG/Private/Tests/Elements/PCGBooleanSelectTest.cpp
@@ -27,8 +27,17 @@ protected:
 
 		UTEST_TRUE("Test element created and valid", TestElement.IsValid());
 
-		TestData.InputData.TaggedData.Emplace(PCGTestsCommon::CreateEmptyPointData(), TSet({FirstDataTag}), FirstInputPinLabel);
-		TestData.InputData.TaggedData.Emplace(PCGTestsCommon::CreateEmptyPointData(), TSet({SecondDataTag}), SecondInputPinLabel);
+		UPCGPointData* FirstPointData = PCGTestsCommon::CreateEmptyPointData();
+		FPCGTaggedData& FirstTaggedData = TestData.InputData.TaggedData.Emplace_GetRef();
+		FirstTaggedData.Data = FirstPointData;
+		FirstTaggedData.Tags.Emplace(FirstDataTag);
+		FirstTaggedData.Pin = FirstInputPinLabel;
+
+		UPCGPointData* SecondPointData = PCGTestsCommon::CreateEmptyPointData();
+		FPCGTaggedData& SecondTaggedData = TestData.InputData.TaggedData.Emplace_GetRef();
+		SecondTaggedData.Data = SecondPointData;
+		SecondTaggedData.Tags.Emplace(SecondDataTag);
+		SecondTaggedData.Pin = SecondInputPinLabel;
 
 		const TUniquePtr<FPCGContext> TestContext = TestData.InitializeTestContext();
 
diff --git a/Engine/Plugins/Experimental/PCG/Source/PCG/Private/Tests/Elements/PCGMultiSelectTest.cpp b/Engine/Plugins/Experimental/PCG/Source/PCG/Private/Tests/Elements/PCGMultiSelectTest.cpp
--- a/Engine/Plugins/Experimental/PCG/Source/PCG/Private/Tests/Elements/PCGMultiSelectTest.cpp
+++ b/Engine/Plugins/Experimental/PCG/Source/PCG/Private/Tests/Elements/PCGMultiSelectTest.cpp
@@ -47,12 +47,20 @@ protected:
 		{
 			if (!SelectElementPinProperties[i].bAdvancedPin && SelectElementPinProperties[i].Label != DefaultPathPinLabel)
 			{
-				TestData.InputData.TaggedData.Emplace(PCGTestsCommon::CreateEmptyPointData(), TSet{CreateDataTag(i)}, SelectElementPinProperties[i].Label);
+				UPCGPointData* PointData = PCGTestsCommon::CreateEmptyPointData();
+				FPCGTaggedData& TaggedData = TestData.InputData.TaggedData.Emplace_GetRef();
+				TaggedData.Data = PointData;
+				TaggedData.Tags.Emplace(CreateDataTag(i));
+				TaggedData.Pin = SelectElementPinProperties[i].Label;
 			}
 		}
 
 		// Add an input to the default path
-		TestData.InputData.TaggedData.Emplace(PCGTestsCommon::CreateEmptyPointData(), TSet{DefaultPathPinDataTag}, DefaultPathPinLabel);
+		UPCGPointData* PointData = PCGTestsCommon::CreateEmptyPointData();
+		FPCGTaggedData& TaggedData = TestData.InputData.TaggedData.Emplace_GetRef();
+		TaggedData.Data = PointData;
+		TaggedData.Tags.Emplace(DefaultPathPinDataTag);
+		TaggedData.Pin = DefaultPathPinLabel;
 
 		const TUniquePtr<FPCGContext> TestContext = TestData.InitializeTestContext();
 
