[PCG] Add tracking support for Objects

* Only work on static dependencies for now.

diff --git a/Engine/Plugins/Experimental/PCG/Source/PCG/Private/Elements/IO/PCGDataTableElement.cpp b/Engine/Plugins/Experimental/PCG/Source/PCG/Private/Elements/IO/PCGDataTableElement.cpp
--- a/Engine/Plugins/Experimental/PCG/Source/PCG/Private/Elements/IO/PCGDataTableElement.cpp
+++ b/Engine/Plugins/Experimental/PCG/Source/PCG/Private/Elements/IO/PCGDataTableElement.cpp
@@ -13,6 +13,17 @@
 #define LOCTEXT_NAMESPACE "PCGLoadDataTable"
 
 #if WITH_EDITOR
+void UPCGLoadDataTableSettings::GetTrackedActorKeys(FPCGActorSelectionKeyToSettingsMap& OutKeysToSettings, TArray<TObjectPtr<const UPCGGraph>>& OutVisitedGraphs) const
+{
+	if (DataTable.IsNull())
+	{
+		return;
+	}
+
+	FPCGActorSelectionKey Key = FPCGActorSelectionKey::CreateFromPath(DataTable.ToSoftObjectPath());
+
+	OutKeysToSettings.FindOrAdd(Key).Emplace(this, /*bCulling=*/false);
+}
 
 FText UPCGLoadDataTableSettings::GetDefaultNodeTitle() const
 {
diff --git a/Engine/Plugins/Experimental/PCG/Source/PCG/Private/Elements/PCGActorSelector.cpp b/Engine/Plugins/Experimental/PCG/Source/PCG/Private/Elements/PCGActorSelector.cpp
--- a/Engine/Plugins/Experimental/PCG/Source/PCG/Private/Elements/PCGActorSelector.cpp
+++ b/Engine/Plugins/Experimental/PCG/Source/PCG/Private/Elements/PCGActorSelector.cpp
@@ -224,22 +224,41 @@ FPCGActorSelectionKey::FPCGActorSelectionKey(TSubclassOf<AActor> InSelectionClas
 	ActorFilter = EPCGActorFilter::AllWorldActors;
 }
 
+FPCGActorSelectionKey FPCGActorSelectionKey::CreateFromPath(const FSoftObjectPath& InObjectPath)
+{
+	FPCGActorSelectionKey Res{};
+	Res.Selection = EPCGActorSelection::ByPath;
+	Res.ObjectPath = InObjectPath;
+	Res.ActorFilter = EPCGActorFilter::AllWorldActors;
+
+	return Res;
+}
+
 void FPCGActorSelectionKey::SetExtraDependency(const UClass* InExtraDependency)
 {
 	OptionalExtraDependency = InExtraDependency;
 }
 
-bool FPCGActorSelectionKey::IsMatching(const AActor* InActor, const UPCGComponent* InComponent) const
+bool FPCGActorSelectionKey::IsMatching(const TSoftObjectPtr<UObject>& InObjectPtr, const UPCGComponent* InComponent) const
 {
-	if (!InActor)
+	if (InObjectPtr.IsNull())
 	{
 		return false;
 	}
+
+	const UObject* InObject = InObjectPtr.Get();
 	
 	// If we filter something else than all world actors, matching depends on the component.
-	// Re-use the same mecanism than Get Actor Data, which should be cheap since we don't look for all actors in the world.
+	// Re-use the same mechanism than Get Actor Data, which should be cheap since we don't look for all actors in the world.
 	if (ActorFilter != EPCGActorFilter::AllWorldActors)
 	{
+		const AActor* InActor = Cast<const AActor>(InObject);
+
+		if (!InActor)
+		{
+			return false;
+		}
+
 		// InKey provide the info for selecting a given actor.
 		// We reconstruct the selector settings from this key, and we also force it to SelectMultiple, since
 		// we want to gather all the actors that matches this given key, to find if ours matches.
@@ -252,9 +271,14 @@ bool FPCGActorSelectionKey::IsMatching(const AActor* InActor, const UPCGComponen
 	switch (Selection)
 	{
 	case EPCGActorSelection::ByTag:
-		return InActor->ActorHasTag(Tag);
+	{
+		const AActor* InActor = Cast<const AActor>(InObject);
+		return InActor && InActor->ActorHasTag(Tag);
+	}
 	case EPCGActorSelection::ByClass:
-		return InActor->IsA(ActorSelectionClass);
+		return InObject && InObject->GetClass()->IsChildOf(ActorSelectionClass);
+	case EPCGActorSelection::ByPath:
+		return InObjectPtr.ToSoftObjectPath() == ObjectPath;
 	default:
 		return false;
 	}
@@ -273,6 +297,8 @@ bool FPCGActorSelectionKey::operator==(const FPCGActorSelectionKey& InOther) con
 		return Tag == InOther.Tag;
 	case EPCGActorSelection::ByClass:
 		return ActorSelectionClass == InOther.ActorSelectionClass;
+	case EPCGActorSelection::ByPath:
+		return ObjectPath == InOther.ObjectPath;
 	case EPCGActorSelection::Unknown: // Fall-through
 	case EPCGActorSelection::ByName:
 		return true;
@@ -290,6 +316,7 @@ uint32 GetTypeHash(const FPCGActorSelectionKey& In)
 	HashResult = HashCombine(HashResult, GetTypeHash(In.Tag));
 	HashResult = HashCombine(HashResult, GetTypeHash(In.ActorSelectionClass));
 	HashResult = HashCombine(HashResult, GetTypeHash(In.OptionalExtraDependency));
+	HashResult = HashCombine(HashResult, GetTypeHash(In.ObjectPath));
 
 	return HashResult;
 }
diff --git a/Engine/Plugins/Experimental/PCG/Source/PCG/Private/Elements/PCGGetPropertyFromObjectPath.cpp b/Engine/Plugins/Experimental/PCG/Source/PCG/Private/Elements/PCGGetPropertyFromObjectPath.cpp
--- a/Engine/Plugins/Experimental/PCG/Source/PCG/Private/Elements/PCGGetPropertyFromObjectPath.cpp
+++ b/Engine/Plugins/Experimental/PCG/Source/PCG/Private/Elements/PCGGetPropertyFromObjectPath.cpp
@@ -13,6 +13,21 @@
 #define LOCTEXT_NAMESPACE "PCGGetPropertyFromObjectPathElement"
 
 #if WITH_EDITOR
+void UPCGGetPropertyFromObjectPathSettings::GetTrackedActorKeys(FPCGActorSelectionKeyToSettingsMap& OutKeysToSettings, TArray<TObjectPtr<const UPCGGraph>>& OutVisitedGraphs) const
+{
+	for (const FSoftObjectPath& ObjectPath : ObjectPathsToExtract)
+	{
+		if (ObjectPath.IsNull())
+		{
+			continue;
+		}
+
+		FPCGActorSelectionKey Key = FPCGActorSelectionKey::CreateFromPath(ObjectPath);
+
+		OutKeysToSettings.FindOrAdd(Key).Emplace(this, /*bCulling=*/false);
+	}
+}
+
 FName UPCGGetPropertyFromObjectPathSettings::GetDefaultNodeName() const
 {
 	return FName(TEXT("GetPropertyFromObjectPath"));
diff --git a/Engine/Plugins/Experimental/PCG/Source/PCG/Private/Elements/PCGIndirectionElement.cpp b/Engine/Plugins/Experimental/PCG/Source/PCG/Private/Elements/PCGIndirectionElement.cpp
--- a/Engine/Plugins/Experimental/PCG/Source/PCG/Private/Elements/PCGIndirectionElement.cpp
+++ b/Engine/Plugins/Experimental/PCG/Source/PCG/Private/Elements/PCGIndirectionElement.cpp
@@ -17,6 +17,17 @@ namespace PCGIndirectionSettings
 }
 
 #if WITH_EDITOR
+void UPCGIndirectionSettings::GetTrackedActorKeys(FPCGActorSelectionKeyToSettingsMap& OutKeysToSettings, TArray<TObjectPtr<const UPCGGraph>>& OutVisitedGraphs) const
+{
+	if (Settings.IsNull())
+	{
+		return;
+	}
+
+	FPCGActorSelectionKey Key = FPCGActorSelectionKey::CreateFromPath(Settings.ToSoftObjectPath());
+
+	OutKeysToSettings.FindOrAdd(Key).Emplace(this, /*bCulling=*/false);
+}
 
 FName UPCGIndirectionSettings::GetDefaultNodeName() const
 {
diff --git a/Engine/Plugins/Experimental/PCG/Source/PCG/Private/PCGActorAndComponentMapping.cpp b/Engine/Plugins/Experimental/PCG/Source/PCG/Private/PCGActorAndComponentMapping.cpp
--- a/Engine/Plugins/Experimental/PCG/Source/PCG/Private/PCGActorAndComponentMapping.cpp
+++ b/Engine/Plugins/Experimental/PCG/Source/PCG/Private/PCGActorAndComponentMapping.cpp
@@ -16,6 +16,7 @@
 #include "StaticMeshCompiler.h"
 #include "TextureCompiler.h"
 #include "Algo/AnyOf.h"
+#include "Engine/DataTable.h"
 #include "Engine/Engine.h"
 #include "Engine/Level.h"
 #include "Engine/StaticMesh.h"
@@ -826,9 +827,6 @@ void FPCGActorAndComponentMapping::RegisterTracking(UPCGComponent* InComponent)
 	RegisterActor(ComponentOwner);
 	AlwaysTrackedActorsToComponentsMap.FindOrAdd(ComponentOwner).Add(InComponent);
 
-	FPCGActorSelectionKey OwnerKey = FPCGActorSelectionKey(EPCGActorFilter::Self);
-	KeysToComponentsMap.FindOrAdd(OwnerKey).Add(InComponent);
-
 	UpdateTracking(InComponent, /*bInShouldDirtyActors=*/ false);
 
 	// Add tracking for when the graph was generated/cleaned, only once
@@ -872,68 +870,91 @@ void FPCGActorAndComponentMapping::UpdateTracking(UPCGComponent* InComponent, bo
 
 	check(ChangedKeys);
 
-	// And we also need to find all actors that should be tracked
-	if (UPCGGraph* PCGGraph = InComponent->GetGraph())
+	if (ChangedKeys->IsEmpty())
 	{
-		TSet<AActor*> CandidatesForTracking;
-		TSet<TObjectKey<AActor>> CandidatesForUntracking;
-		TSet<FPCGActorSelectionKey> CandidateKeysToRemove;
+		// Nothing to do
+		return;
+	}
 
-		auto GatherActors = [this, InComponent, &CandidatesForTracking, &CandidatesForUntracking, &CandidateKeysToRemove](const FPCGActorSelectionKey& InKey, bool bInShouldUntrack)
+	// And we also need to find all actors that should be tracked
+	TSet<TSoftObjectPtr<UObject>> CandidatesForTracking;
+	TSet<TSoftObjectPtr<UObject>> CandidatesForUntracking;
+
+	auto GatherObjects = [this, InComponent, &CandidatesForTracking, &CandidatesForUntracking](const FPCGActorSelectionKey& InKey, bool bInShouldUntrack)
+	{
+		if (InKey.Selection == EPCGActorSelection::ByPath)
 		{
-			// InKey provide the info for selecting a given actor.
-			// We reconstruct the selector settings from this key, and we also force it to SelectMultiple, since
-			// we want to gather all the actors that matches this given key.
-			FPCGActorSelectorSettings SelectorSettings = FPCGActorSelectorSettings::ReconstructFromKey(InKey);
-			SelectorSettings.bSelectMultiple = true;
-			TArray<AActor*> AllActors = PCGActorSelector::FindActors(SelectorSettings, InComponent, [](const AActor*) { return true; }, [](const AActor*) { return true; });
-
-			for (AActor* Actor : AllActors)
+			// Don't track null objects
+			if (InKey.ObjectPath.IsNull())
 			{
-				if (!Actor)
-				{
-					continue;
-				}
-
-				if (bInShouldUntrack && !CandidatesForTracking.Contains(Actor))
-				{
-					CandidatesForUntracking.Add(Actor);
-					CandidateKeysToRemove.Add(InKey);
-				}
-				else
-				{
-					CandidatesForTracking.Add(Actor);
-					CandidatesForUntracking.Remove(Actor);
-					CandidateKeysToRemove.Remove(InKey);
-				}
+				return;
 			}
-		};
 
-		for (const FPCGActorSelectionKey& Key : *ChangedKeys)
-		{
-			const bool bShouldUntrack = !InComponent->CachedTrackedKeysToSettings.Contains(Key);
-			GatherActors(Key, bShouldUntrack);
-			if (!bShouldUntrack)
+			TSoftObjectPtr<UObject> Object(InKey.ObjectPath);
+			if (bInShouldUntrack && !CandidatesForTracking.Contains(Object))
+			{
+				CandidatesForUntracking.Add(Object);
+			}
+			else
 			{
-				KeysToComponentsMap.FindOrAdd(Key).Add(InComponent);
+				CandidatesForTracking.Add(Object);
+				CandidatesForUntracking.Remove(Object);
 			}
+
+			return;
 		}
 
-		const bool bDisableDelayedActorRegistering = PCGActorAndComponentMapping::CVarDisableDelayedActorRegistering.GetValueOnAnyThread();
-		for (AActor* Actor : CandidatesForTracking)
+		// InKey provide the info for selecting a given actor.
+		// We reconstruct the selector settings from this key, and we also force it to SelectMultiple, since
+		// we want to gather all the actors that matches this given key.
+		FPCGActorSelectorSettings SelectorSettings = FPCGActorSelectorSettings::ReconstructFromKey(InKey);
+		SelectorSettings.bSelectMultiple = true;
+		TArray<AActor*> AllActors = PCGActorSelector::FindActors(SelectorSettings, InComponent, [](const AActor*) { return true; }, [](const AActor*) { return true; });
+
+		for (AActor* Actor : AllActors)
 		{
-			if (!Actor->HasActorRegisteredAllComponents() && !bDisableDelayedActorRegistering)
+			if (!Actor)
 			{
-				DelayedAddedActors.Emplace(Actor, { false, 0 });
 				continue;
 			}
 
-			bool bShouldCull = false;
-			if (!InComponent->IsActorTracked(Actor, bShouldCull))
+			if (bInShouldUntrack && !CandidatesForTracking.Contains(Actor))
 			{
-				continue;
+				CandidatesForUntracking.Add(Actor);
 			}
+			else
+			{
+				CandidatesForTracking.Add(Actor);
+				CandidatesForUntracking.Remove(Actor);
+			}
+		}
+	};
 
+	for (const FPCGActorSelectionKey& Key : *ChangedKeys)
+	{
+		const bool bShouldUntrack = !InComponent->CachedTrackedKeysToSettings.Contains(Key);
+		GatherObjects(Key, bShouldUntrack);
+	}
+
+	const bool bDisableDelayedActorRegistering = PCGActorAndComponentMapping::CVarDisableDelayedActorRegistering.GetValueOnAnyThread();
+	for (const TSoftObjectPtr<UObject>& ObjectPtr : CandidatesForTracking)
+	{
+		AActor* Actor = Cast<AActor>(ObjectPtr.Get());
+
+		if (Actor && !Actor->HasActorRegisteredAllComponents() && !bDisableDelayedActorRegistering)
+		{
+			DelayedAddedActors.Emplace(Actor, { false, 0 });
+			continue;
+		}
+
+		bool bShouldCull = false;
+		if (!InComponent->IsObjectTracked(ObjectPtr, bShouldCull))
+		{
+			continue;
+		}
+
+		if (Actor)
+		{
 			// Making sure that we only have the component in one map.
 			auto RemoveFromMap = [InComponent, Actor](TMap<TObjectKey<AActor>, TSet<UPCGComponent*>>& InMap)
 			{
@@ -960,10 +981,14 @@ void FPCGActorAndComponentMapping::UpdateTracking(UPCGComponent* InComponent, bo
 
 			RegisterActor(Actor);
 		}
-
-		// Also unregister all keys that are not tracked anymore.
-		UnregisterTracking(InComponent, &CandidatesForUntracking, &CandidateKeysToRemove);
+		else
+		{
+			TrackedObjectsToComponentsMap.FindOrAdd(ObjectPtr).Add(InComponent);
+		}
 	}
+
+	// Also unregister all keys that are not tracked anymore.
+	UnregisterTracking(InComponent, &CandidatesForUntracking);
 }
 
 void FPCGActorAndComponentMapping::RemapTracking(const UPCGComponent* InOldComponent, UPCGComponent* InNewComponent)
@@ -981,7 +1006,7 @@ void FPCGActorAndComponentMapping::RemapTracking(const UPCGComponent* InOldCompo
 
 	ReplaceInMap(CulledTrackedActorsToComponentsMap);
 	ReplaceInMap(AlwaysTrackedActorsToComponentsMap);
-	ReplaceInMap(KeysToComponentsMap);
+	ReplaceInMap(TrackedObjectsToComponentsMap);
 
 	// Old component will probably die, but we'll force removing the delegates even if it is const.
 	if (UPCGComponent* MutableOldComponent = const_cast<UPCGComponent*>(InOldComponent))
@@ -998,80 +1023,78 @@ void FPCGActorAndComponentMapping::RemapTracking(const UPCGComponent* InOldCompo
 	}
 }
 
-void FPCGActorAndComponentMapping::UnregisterTracking(UPCGComponent* InComponent, const TSet<TObjectKey<AActor>>* OptionalActorsToUntrack, const TSet<FPCGActorSelectionKey>* OptionalKeysToUntrack)
+void FPCGActorAndComponentMapping::UnregisterTracking(UPCGComponent* InComponent, const TSet<TSoftObjectPtr<UObject>>* OptionalObjectsToUntrack)
 {
-	if (!InComponent)
+	if (!InComponent || (OptionalObjectsToUntrack && OptionalObjectsToUntrack->IsEmpty()))
 	{
 		return;
 	}
 
-	TSet<TObjectKey<AActor>> CandidatesForUntrack;
-	TSet<FPCGActorSelectionKey> KeysToRemove;
+	TSet<AActor*> ActorCandidatesForUntrack;
 
-	auto RemoveAllFromMap = [InComponent]<typename Key>(TMap<Key, TSet<UPCGComponent*>>& InMap, TSet<Key>& InCandidateToRemove)
+	if (OptionalObjectsToUntrack)
 	{
-		for (TPair<Key, TSet<UPCGComponent*>>& It : InMap)
+		auto RemoveFromMap = [InComponent](auto& InMap, auto& Key, auto OnEmpty)
 		{
-			It.Value.Remove(InComponent);
-			if (It.Value.IsEmpty())
+			if (TSet<UPCGComponent*>* It = InMap.Find(Key))
 			{
-				InCandidateToRemove.Add(It.Key);
+				It->Remove(InComponent);
+				if (It->IsEmpty())
+				{
+					OnEmpty(Key);
+				}
 			}
-		}
-	};
+		};
 
-	auto RemoveKeysFromMap = [InComponent]<typename Key>(TMap<Key, TSet<UPCGComponent*>>& InMap, TSet<Key>& InCandidateToRemove, const TSet<Key>& SetToIterateOn)
-	{
-		for (const Key& KeyIt : SetToIterateOn)
+		for (const TSoftObjectPtr<UObject>& ObjectPtr : *OptionalObjectsToUntrack)
 		{
-			if (TSet<UPCGComponent*>* ComponentSetPtr = InMap.Find(KeyIt))
+			if (AActor* Actor = Cast<AActor>(ObjectPtr.Get()))
 			{
-				ComponentSetPtr->Remove(InComponent);
-				if (ComponentSetPtr->IsEmpty())
-				{
-					InCandidateToRemove.Add(KeyIt);
-				}
+				RemoveFromMap(CulledTrackedActorsToComponentsMap, Actor, [&ActorCandidatesForUntrack](AActor* Actor) { ActorCandidatesForUntrack.Add(Actor); });
+				RemoveFromMap(AlwaysTrackedActorsToComponentsMap, Actor, [&ActorCandidatesForUntrack](AActor* Actor) { ActorCandidatesForUntrack.Add(Actor); });
+			}
+			else
+			{
+				RemoveFromMap(TrackedObjectsToComponentsMap, ObjectPtr, [this](const TSoftObjectPtr<UObject>& ObjectPtr) { UnregisterObject(ObjectPtr); });
 			}
 		}
-	};
-
-	if (OptionalActorsToUntrack)
-	{
-		RemoveKeysFromMap(CulledTrackedActorsToComponentsMap, CandidatesForUntrack, *OptionalActorsToUntrack);
-		RemoveKeysFromMap(AlwaysTrackedActorsToComponentsMap, CandidatesForUntrack, *OptionalActorsToUntrack);
-	}
-	else
-	{
-		RemoveAllFromMap(CulledTrackedActorsToComponentsMap, CandidatesForUntrack);
-		RemoveAllFromMap(AlwaysTrackedActorsToComponentsMap, CandidatesForUntrack);
-	}
-
-	if (OptionalKeysToUntrack)
-	{
-		RemoveKeysFromMap(KeysToComponentsMap, KeysToRemove, *OptionalKeysToUntrack);
 	}
 	else
 	{
-		RemoveAllFromMap(KeysToComponentsMap, KeysToRemove);
-	}
+		auto RemoveAllFromMap = [InComponent](auto& InMap, auto OnEmpty)
+		{
+			for (auto& It : InMap)
+			{
+				It.Value.Remove(InComponent);
+				if (It.Value.IsEmpty())
+				{
+					OnEmpty(It.Key);
+				}
+			}
+		};
 
-	for (const FPCGActorSelectionKey& Key : KeysToRemove)
-	{
-		KeysToComponentsMap.Remove(Key);
+		RemoveAllFromMap(CulledTrackedActorsToComponentsMap, [&ActorCandidatesForUntrack](const TObjectKey<AActor>& ActorKey) { ActorCandidatesForUntrack.Add(ActorKey.ResolveObjectPtr()); });
+		RemoveAllFromMap(AlwaysTrackedActorsToComponentsMap, [&ActorCandidatesForUntrack](const TObjectKey<AActor>& ActorKey) { ActorCandidatesForUntrack.Add(ActorKey.ResolveObjectPtr()); });
+		RemoveAllFromMap(TrackedObjectsToComponentsMap, [this](const TSoftObjectPtr<UObject>& ObjectPtr) { UnregisterObject(ObjectPtr); });
 	}
 
 	// We also need to untrack actors that doesn't have any component that tracks them.
-	auto ShouldBeRemoved = [](const TObjectKey<AActor>& InActor, TMap<TObjectKey<AActor>, TSet<UPCGComponent*>>& InMap)
+	auto ShouldBeRemoved = [](const AActor* InActor, TMap<TObjectKey<AActor>, TSet<UPCGComponent*>>& InMap)
 	{
 		TSet<UPCGComponent*>* RegisteredComponents = InMap.Find(InActor);
 		return !RegisteredComponents || RegisteredComponents->IsEmpty();
 	};
 
-	for (TObjectKey<AActor> Candidate : CandidatesForUntrack)
+	for (AActor* Candidate : ActorCandidatesForUntrack)
 	{
-		if (ShouldBeRemoved(Candidate, CulledTrackedActorsToComponentsMap) && ShouldBeRemoved(Candidate, AlwaysTrackedActorsToComponentsMap))
+		if (!Candidate)
+		{
+			continue;
+		}
+
+		if (Candidate && ShouldBeRemoved(Candidate, AlwaysTrackedActorsToComponentsMap) && ShouldBeRemoved(Candidate, CulledTrackedActorsToComponentsMap))
 		{
-			UnregisterActor(Candidate.ResolveObjectPtr());
+			UnregisterActor(Candidate);
 		}
 	}
 }
@@ -1083,7 +1106,7 @@ void FPCGActorAndComponentMapping::UnregisterTracking(UPCGComponent* InComponent
 		return;
 	}
 
-	UnregisterTracking(InComponent, nullptr, nullptr);
+	UnregisterTracking(InComponent, nullptr);
 
 	InComponent->OnPCGGraphGeneratedDelegate.RemoveAll(this);
 	InComponent->OnPCGGraphCleanedDelegate.RemoveAll(this);
@@ -1103,6 +1126,7 @@ void FPCGActorAndComponentMapping::RegisterTrackingCallbacks()
 	GEngine->OnLevelActorDeleted().AddRaw(this, &FPCGActorAndComponentMapping::OnActorDeleted);
 	FCoreUObjectDelegates::OnObjectPropertyChanged.AddRaw(this, &FPCGActorAndComponentMapping::OnObjectPropertyChanged);
 	FCoreUObjectDelegates::OnPreObjectPropertyChanged.AddRaw(this, &FPCGActorAndComponentMapping::OnPreObjectPropertyChanged);
+	FCoreUObjectDelegates::OnObjectPreSave.AddRaw(this, &FPCGActorAndComponentMapping::OnObjectSaved);
 
 	UWorld* World = PCGSubsystem ? PCGSubsystem->GetWorld() : nullptr;
 	// Need the World condition for static analysis...
@@ -1120,6 +1144,7 @@ void FPCGActorAndComponentMapping::TeardownTrackingCallbacks()
 	GEngine->OnLevelActorDeleted().RemoveAll(this);
 	FCoreUObjectDelegates::OnObjectPropertyChanged.RemoveAll(this);
 	FCoreUObjectDelegates::OnPreObjectPropertyChanged.RemoveAll(this);
+	FCoreUObjectDelegates::OnObjectPreSave.RemoveAll(this);
 
 	UWorld* World = PCGSubsystem ? PCGSubsystem->GetWorld() : nullptr;
 	// Need the World condition for static analysis...
@@ -1300,7 +1325,7 @@ bool FPCGActorAndComponentMapping::AddOrUpdateTrackedActor(AActor* InActor)
 		}
 
 		bool bTrackingIsCulled = false;
-		if (PCGComponent && PCGComponent->IsActorTracked(InActor, bTrackingIsCulled))
+		if (PCGComponent && PCGComponent->IsObjectTracked(InActor, bTrackingIsCulled))
 		{
 			if (bTrackingIsCulled)
 			{
@@ -1393,6 +1418,11 @@ bool FPCGActorAndComponentMapping::UnregisterActor(AActor* InActor)
 	}
 }
 
+bool FPCGActorAndComponentMapping::UnregisterObject(const TSoftObjectPtr<UObject>& InObject)
+{
+	return TrackedObjectsToComponentsMap.Remove(InObject) > 0;
+}
+
 void FPCGActorAndComponentMapping::OnActorUnloaded(AActor& InActor)
 {
 	// Don't dirty on unload (to mirror the behavior in load)
@@ -1547,6 +1577,18 @@ void FPCGActorAndComponentMapping::OnPreObjectPropertyChanged(UObject* InObject,
 	TempTrackedActorTags = TSet<FName>(Actor->Tags);
 }
 
+void FPCGActorAndComponentMapping::OnObjectSaved(UObject* InObject, FObjectPreSaveContext InObjectSaveContext)
+{
+	// Only trigger a refresh a new user data and if it is a data table
+	// We only track data table because we probably will catch other changes with OnObjectPropertyChanged.
+	// To avoid to force the check multiple times (on change + on save)
+	if (!InObjectSaveContext.IsProceduralSave() && Cast<UDataTable>(InObject))
+	{
+		FPropertyChangedEvent Event{ nullptr };
+		OnObjectPropertyChanged(InObject, Event);
+	}
+}
+
 void FPCGActorAndComponentMapping::OnObjectPropertyChanged(UObject* InObject, FPropertyChangedEvent& InEvent)
 {
 	TRACE_CPUPROFILER_EVENT_SCOPE(FPCGActorAndComponentMapping::OnObjectPropertyChanged);
@@ -1583,6 +1625,18 @@ void FPCGActorAndComponentMapping::OnObjectPropertyChanged(UObject* InObject, FP
 	// If we don't find any actor, try to see if it is a dependency
 	if (!Actor)
 	{
+		if (TSet<UPCGComponent*>* Components = TrackedObjectsToComponentsMap.Find(InObject))
+		{
+			for (UPCGComponent* Component : *Components)
+			{
+				if (ClearCache(InObject, Component, /*bIntersect=*/ false, {}, InObject))
+				{
+					Component->DirtyGenerated();
+					Component->Refresh();
+				}
+			}
+		}
+
 		for (const TPair<TObjectKey<AActor>, TSet<TObjectPtr<UObject>>>& TrackedActor : TrackedActorsToDependenciesMap)
 		{
 			if (TrackedActor.Value.Contains(InObject))
@@ -1697,7 +1751,7 @@ void FPCGActorAndComponentMapping::OnActorChanged(AActor* InActor, bool bInHasMo
 				return;
 			}
 
-			if (ClearCacheForActor(InActor, ComponentRef.Component, /*bIntersect=*/true, RemovedTags, InOriginatingChangeObject))
+			if (ClearCache(InActor, ComponentRef.Component, /*bIntersect=*/true, RemovedTags, InOriginatingChangeObject))
 			{
 				ComponentRef.Component->DirtyGenerated(DirtyFlag);
 				DirtyComponents.Add(ComponentRef.Component);
@@ -1734,7 +1788,7 @@ void FPCGActorAndComponentMapping::OnActorChanged(AActor* InActor, bool bInHasMo
 			bool bWasDirtied = false;
 
 			// Since when we clear the cache for a settings, we clear it all, it's only necessary to do it once on the original component, then only dirty the local that intersects.
-			if (ClearCacheForActor(InActor, ComponentRef.Component, /*bIntersect=*/true, RemovedTags, InOriginatingChangeObject))
+			if (ClearCache(InActor, ComponentRef.Component, /*bIntersect=*/true, RemovedTags, InOriginatingChangeObject))
 			{
 				ForAllIntersectingPartitionActors(Overlap, [InActor, Component = ComponentRef.Component, &bWasDirtied, DirtyFlag, InOriginatingChangeObject](APCGPartitionActor* InPartitionActor) -> void
 				{
@@ -1790,7 +1844,7 @@ void FPCGActorAndComponentMapping::OnActorChanged(AActor* InActor, bool bInHasMo
 			{
 				const FBox ComponentBounds = PCGComponent->GetGridBounds();
 				const bool bIntersect = ActorBounds.Intersect(ComponentBounds) || (OldActorBoundsPtr && OldActorBoundsPtr->Intersect(ComponentBounds));
-				bWasDirtied = ClearCacheForActor(InActor, PCGComponent, /*bIntersect=*/bIntersect, RemovedTags, InOriginatingChangeObject);
+				bWasDirtied = ClearCache(InActor, PCGComponent, /*bIntersect=*/bIntersect, RemovedTags, InOriginatingChangeObject);
 			}
 
 			if (bWasDirtied || bOwnerChanged)
@@ -1906,11 +1960,11 @@ void FPCGActorAndComponentMapping::UpdateActorDependencies(AActor* InActor)
 		UMaterialInterface::StaticClass()
 	};
 
-	TSet<TObjectPtr<UObject>>& DependenciesSet = TrackedActorsToDependenciesMap.FindOrAdd(InActor);
-	DependenciesSet.Empty();
 
 	if (!PCGActorAndComponentMapping::CVarDisableObjectDependenciesTracking.GetValueOnAnyThread())
 	{
+		TSet<TObjectPtr<UObject>>& DependenciesSet = TrackedActorsToDependenciesMap.FindOrAdd(InActor);
+		DependenciesSet.Empty();
 		PCGHelpers::GatherDependencies(InActor, DependenciesSet, 1, ExcludedClasses);
 	}
 }
@@ -1930,9 +1984,9 @@ bool FPCGActorAndComponentMapping::IsActorTracked(const AActor* InActor) const
 	return InActor && (CulledTrackedActorsToComponentsMap.Contains(InActor) || AlwaysTrackedActorsToComponentsMap.Contains(InActor));
 }
 
-bool FPCGActorAndComponentMapping::ClearCacheForActor(const AActor* InActor, const UPCGComponent* InComponent, const bool bIntersect, const TSet<FName>& InRemovedTags, const UObject* InOriginatingChange) const
+bool FPCGActorAndComponentMapping::ClearCache(const UObject* InObject, const UPCGComponent* InComponent, const bool bIntersect, const TSet<FName>& InRemovedTags, const UObject* InOriginatingChange) const
 {
-	check(InActor && InComponent);
+	check(InObject && InComponent);
 
 	if (!PCGSubsystem)
 	{
@@ -1940,9 +1994,9 @@ bool FPCGActorAndComponentMapping::ClearCacheForActor(const AActor* InActor, con
 	}
 
 	// Special case for the landscape. No settings associated, but we should dirty if the component is tracking the landscape.
-	bool bShouldDirty = InComponent->ShouldTrackLandscape() && InActor->IsA<ALandscapeProxy>();
+	bool bShouldDirty = InComponent->ShouldTrackLandscape() && InObject->IsA<ALandscapeProxy>();
 
-	for (const UPCGSettings* Settings : InComponent->GatherSettingsTrackingActor(InActor, bIntersect, InRemovedTags, InOriginatingChange))
+	for (const UPCGSettings* Settings : InComponent->GatherSettingsTracking(InObject, bIntersect, InRemovedTags, InOriginatingChange))
 	{
 		if (!Settings) 
 		{
diff --git a/Engine/Plugins/Experimental/PCG/Source/PCG/Private/PCGComponent.cpp b/Engine/Plugins/Experimental/PCG/Source/PCG/Private/PCGComponent.cpp
--- a/Engine/Plugins/Experimental/PCG/Source/PCG/Private/PCGComponent.cpp
+++ b/Engine/Plugins/Experimental/PCG/Source/PCG/Private/PCGComponent.cpp
@@ -2042,24 +2042,26 @@ bool UPCGComponent::ShouldGenerateBPPCGAddedToWorld() const
 	}
 }
 
-bool UPCGComponent::IsActorTracked(AActor* InActor, bool& bOutIsCulled) const
+bool UPCGComponent::IsObjectTracked(const TSoftObjectPtr<UObject>& InObjectPtr, bool& bOutIsCulled) const
 {
-	check(InActor);
+	check(!InObjectPtr.IsNull());
 
 	if (!GetOwner())
 	{
 		return false;
 	}
 
+	const UObject* InObject = InObjectPtr.Get();
+
 	// We should always track the owner of the component, without culling
-	if (GetOwner() == InActor)
+	if (GetOwner() == InObject)
 	{
 		bOutIsCulled = false;
 		return true;
 	}
 
 	// If we track the landscape using legacy methods and it is a landscape, it should be tracked as culled
-	if (InActor->IsA<ALandscapeProxy>() && ShouldTrackLandscape())
+	if (InObject && InObject->IsA<ALandscapeProxy>() && ShouldTrackLandscape())
 	{
 		bOutIsCulled = true;
 		return true;
@@ -2069,7 +2071,7 @@ bool UPCGComponent::IsActorTracked(AActor* InActor, bool& bOutIsCulled) const
 
 	for (const TPair<FPCGActorSelectionKey, bool>& It : CachedTrackedKeysToCulling)
 	{
-		if (It.Key.IsMatching(InActor, this))
+		if (It.Key.IsMatching(InObjectPtr, this))
 		{
 			bOutIsCulled = It.Value;
 			bFound = true;
@@ -2678,11 +2680,11 @@ UPCGSubsystem* UPCGComponent::GetSubsystem() const
 }
 
 #if WITH_EDITOR
-TArray<const UPCGSettings*> UPCGComponent::GatherSettingsTrackingActor(const AActor* InActor, bool bIntersect, const TSet<FName>& InRemovedTags, const UObject* InOriginatingChangeObject) const
+TArray<const UPCGSettings*> UPCGComponent::GatherSettingsTracking(const UObject* InObject, bool bIntersect, const TSet<FName>& InRemovedTags, const UObject* InOriginatingChangeObject) const
 {
-	TRACE_CPUPROFILER_EVENT_SCOPE(UPCGComponent::GatherSettingsTrackingActor);
+	TRACE_CPUPROFILER_EVENT_SCOPE(UPCGComponent::GatherSettingsTracking);
 
-	if (!InActor)
+	if (!InObject)
 	{
 		return {};
 	}
@@ -2695,7 +2697,7 @@ TArray<const UPCGSettings*> UPCGComponent::GatherSettingsTrackingActor(const AAc
 
 		const bool bRemovedTagIsTracked = (Key.Selection == EPCGActorSelection::ByTag) && InRemovedTags.Contains(Key.Tag);
 
-		if (It.Key.IsMatching(InActor, this) || bRemovedTagIsTracked)
+		if (It.Key.IsMatching(InObject, this) || bRemovedTagIsTracked)
 		{
 			// Extra care if the change originates from a PCGComponent. Only dirty if we are tracking a PCG component.
 			if (InOriginatingChangeObject && InOriginatingChangeObject->IsA<UPCGComponent>() 
diff --git a/Engine/Plugins/Experimental/PCG/Source/PCG/Public/Elements/IO/PCGDataTableElement.h b/Engine/Plugins/Experimental/PCG/Source/PCG/Public/Elements/IO/PCGDataTableElement.h
--- a/Engine/Plugins/Experimental/PCG/Source/PCG/Public/Elements/IO/PCGDataTableElement.h
+++ b/Engine/Plugins/Experimental/PCG/Source/PCG/Public/Elements/IO/PCGDataTableElement.h
@@ -22,6 +22,7 @@ public:
 	virtual FText GetDefaultNodeTitle() const override;
 	virtual FText GetNodeTooltipText() const override;
 	virtual bool HasDynamicPins() const override { return true; }
+	virtual void GetTrackedActorKeys(FPCGActorSelectionKeyToSettingsMap& OutKeysToSettings, TArray<TObjectPtr<const UPCGGraph>>& OutVisitedGraphs) const override;
 #endif
 
 	virtual EPCGDataType GetCurrentPinTypes(const UPCGPin* InPin) const override;
diff --git a/Engine/Plugins/Experimental/PCG/Source/PCG/Public/Elements/PCGActorSelector.h b/Engine/Plugins/Experimental/PCG/Source/PCG/Public/Elements/PCGActorSelector.h
--- a/Engine/Plugins/Experimental/PCG/Source/PCG/Public/Elements/PCGActorSelector.h
+++ b/Engine/Plugins/Experimental/PCG/Source/PCG/Public/Elements/PCGActorSelector.h
@@ -18,6 +18,7 @@ enum class EPCGActorSelection : uint8
 	// Deprecated - actor labels are unavailable in shipping builds
 	ByName UMETA(Hidden),
 	ByClass,
+	ByPath UMETA(Hidden), // Hidden because actors are not tracked by paths.
 	Unknown UMETA(Hidden)
 };
 
@@ -36,6 +37,12 @@ enum class EPCGActorFilter : uint8
 	Original,
 };
 
+/**
+* Structure to specify a selection criteria for an object/actor
+* Object can be selected using the EPCGActorSelection::ByClass or EPCGActorSelection::ByPath
+* Actors have more selection with Self/Parent/Root/Original and also EPCGActorSelection::ByTag
+* TODO: Might want to rename it, since it is not limited to Actors anymore.
+*/
 struct FPCGActorSelectionKey
 {
 	FPCGActorSelectionKey() = default;
@@ -46,10 +53,12 @@ struct FPCGActorSelectionKey
 	explicit FPCGActorSelectionKey(FName InTag);
 	explicit FPCGActorSelectionKey(TSubclassOf<AActor> InSelectionClass);
 
+	static FPCGActorSelectionKey CreateFromPath(const FSoftObjectPath& InObjectPath);
+
 	bool operator==(const FPCGActorSelectionKey& InOther) const;
 
 	friend uint32 GetTypeHash(const FPCGActorSelectionKey& In);
-	bool IsMatching(const AActor* InActor, const UPCGComponent* InComponent) const;
+	bool IsMatching(const TSoftObjectPtr<UObject>& InObjectPtr, const UPCGComponent* InComponent) const;
 
 	void SetExtraDependency(const UClass* InExtraDependency);
 
@@ -58,6 +67,9 @@ struct FPCGActorSelectionKey
 	FName Tag = NAME_None;
 	TSubclassOf<AActor> ActorSelectionClass = nullptr;
 
+	// If the Selection is ByPath, contain the path to select.
+	FSoftObjectPath ObjectPath;
+
 	// If it should track a specific object dependency instead of an actor. For example, GetActorData with GetPCGComponent data.
 	const UClass* OptionalExtraDependency = nullptr;
 };
diff --git a/Engine/Plugins/Experimental/PCG/Source/PCG/Public/Elements/PCGGetPropertyFromObjectPath.h b/Engine/Plugins/Experimental/PCG/Source/PCG/Public/Elements/PCGGetPropertyFromObjectPath.h
--- a/Engine/Plugins/Experimental/PCG/Source/PCG/Public/Elements/PCGGetPropertyFromObjectPath.h
+++ b/Engine/Plugins/Experimental/PCG/Source/PCG/Public/Elements/PCGGetPropertyFromObjectPath.h
@@ -25,6 +25,7 @@ public:
 	virtual EPCGSettingsType GetType() const override { return EPCGSettingsType::Param; }
 	virtual bool CanEditChange(const FProperty* InProperty) const override;
 	virtual bool IsPinUsedByNodeExecution(const UPCGPin* InPin) const override;
+	virtual void GetTrackedActorKeys(FPCGActorSelectionKeyToSettingsMap& OutKeysToSettings, TArray<TObjectPtr<const UPCGGraph>>& OutVisitedGraphs) const override;
 #endif
 	virtual FString GetAdditionalTitleInformation() const override;
 
diff --git a/Engine/Plugins/Experimental/PCG/Source/PCG/Public/Elements/PCGIndirectionElement.h b/Engine/Plugins/Experimental/PCG/Source/PCG/Public/Elements/PCGIndirectionElement.h
--- a/Engine/Plugins/Experimental/PCG/Source/PCG/Public/Elements/PCGIndirectionElement.h
+++ b/Engine/Plugins/Experimental/PCG/Source/PCG/Public/Elements/PCGIndirectionElement.h
@@ -29,6 +29,7 @@ public:
 	virtual FText GetDefaultNodeTitle() const override;
 	virtual FText GetNodeTooltipText() const override;
 	virtual EPCGSettingsType GetType() const override { return EPCGSettingsType::Generic; }
+	virtual void GetTrackedActorKeys(FPCGActorSelectionKeyToSettingsMap& OutKeysToSettings, TArray<TObjectPtr<const UPCGGraph>>& OutVisitedGraphs) const override;
 #endif // WITH_EDITOR
 	virtual FString GetAdditionalTitleInformation() const override;
 	virtual bool HasFlippedTitleLines() const override { return true; }
diff --git a/Engine/Plugins/Experimental/PCG/Source/PCG/Public/PCGActorAndComponentMapping.h b/Engine/Plugins/Experimental/PCG/Source/PCG/Public/PCGActorAndComponentMapping.h
--- a/Engine/Plugins/Experimental/PCG/Source/PCG/Public/PCGActorAndComponentMapping.h
+++ b/Engine/Plugins/Experimental/PCG/Source/PCG/Public/PCGActorAndComponentMapping.h
@@ -13,6 +13,7 @@
 #include "Templates/Function.h"
 #include "UObject/ObjectKey.h"
 #include "UObject/ObjectPtr.h"
+#include "UObject/ObjectSaveContext.h"
 #include "UObject/SoftObjectPtr.h"
 #include "UObject/WeakObjectPtr.h"
 
@@ -136,6 +137,9 @@ private:
 	/* Return true if the actor was tracked. */
 	bool UnregisterActor(AActor* InActor);
 
+	/* Return true if the object was tracked. */
+	bool UnregisterObject(const TSoftObjectPtr<UObject>& InObject);
+
 	void OnActorAdded(AActor* InActor);
 	void OnActorLoaded(AActor& InActor);
 	void OnActorAdded_Internal(AActor* InActor, bool bShouldDirty, int32 LevelInstanceDepth, bool bForceAddDelayedActor = false);
@@ -148,6 +152,7 @@ private:
 	void ApplyLandscapeChanges(ALandscapeProxy* InLandscape);
 	void OnPreObjectPropertyChanged(UObject* InObject, const FEditPropertyChain& InEditPropertyChain);
 	void OnObjectPropertyChanged(UObject* InObject, FPropertyChangedEvent& InEvent);
+	void OnObjectSaved(UObject* InObject, FObjectPreSaveContext InObjectSaveContext);
 	void OnPCGGraphGeneratedOrCleaned(UPCGComponent* InComponent);
 
 	/** Remap the tracking in case of BP components. */
@@ -157,7 +162,7 @@ private:
 	void UnregisterTracking(UPCGComponent* InComponent);
 
 	/** Unregister tracking when a component is removed or keys are unregistered */
-	void UnregisterTracking(UPCGComponent* InComponent, const TSet<TObjectKey<AActor>>* OptionalActorsToUntrack, const TSet<FPCGActorSelectionKey>* OptionalKeysToUntrack);
+	void UnregisterTracking(UPCGComponent* InComponent, const TSet<TSoftObjectPtr<UObject>>* OptionalObjectsToUntrack);
 
 	/** Trigger an update when the actor changed. 
 	* Can specify if the actor has moved to also update components that were at its previous position.
@@ -172,7 +177,7 @@ private:
 	bool IsActorTracked(const AActor* InActor) const;
 
 	/** Gather all settings from a given component that track the actor, and clear the cache for them. Returns true if we should dirty afterwards (aka at least one settings was cleared and/or landscape changed). */
-	bool ClearCacheForActor(const AActor* InActor, const UPCGComponent* InComponent, const bool bIntersect, const TSet<FName>& InRemovedTags, const UObject* InOriginatingChange) const;
+	bool ClearCache(const UObject* InObject, const UPCGComponent* InComponent, const bool bIntersect, const TSet<FName>& InRemovedTags, const UObject* InOriginatingChange) const;
 #endif // WITH_EDITOR
 
 private:
@@ -213,8 +218,8 @@ private:
 	/** Same mapping but for always tracked actors */
 	TMap<TObjectKey<AActor>, TSet<UPCGComponent*>> AlwaysTrackedActorsToComponentsMap;
 
-	/** Keep the list of the keys already tracked, to avoid requerying the actors everytime */
-	TMap<FPCGActorSelectionKey, TSet<UPCGComponent*>> KeysToComponentsMap;
+	/** Object tracking. Use soft pointers, because tracked objects might not be yet loaded. */
+	TMap<TSoftObjectPtr<UObject>, TSet<UPCGComponent*>> TrackedObjectsToComponentsMap;
 
 	/** Finally keep a mapping between actors and their position to know if an actor move in/out of a component tracking bounds. Only kept for actors that need to be culled. */
 	TMap<TObjectKey<AActor>, FBox> TrackedActorToPositionMap;
diff --git a/Engine/Plugins/Experimental/PCG/Source/PCG/Public/PCGComponent.h b/Engine/Plugins/Experimental/PCG/Source/PCG/Public/PCGComponent.h
--- a/Engine/Plugins/Experimental/PCG/Source/PCG/Public/PCGComponent.h
+++ b/Engine/Plugins/Experimental/PCG/Source/PCG/Public/PCGComponent.h
@@ -299,7 +299,7 @@ public:
 	/** Did the given node produce one or more data items in the given stack in a previous execution. */
 	bool HasNodeProducedData(const UPCGNode* InNode, const FPCGStack& Stack) const;
 
-	bool IsActorTracked(AActor* InActor, bool& bOutIsCulled) const;
+	bool IsObjectTracked(const TSoftObjectPtr<UObject>& InObjectPtr, bool& bOutIsCulled) const;
 
 	/** Know if we need to force a generation, in case of BP added to the world in editor */
 	bool ShouldGenerateBPPCGAddedToWorld() const;
@@ -437,14 +437,14 @@ private:
 	/** Returns true if something changed in the tracking. */
 	bool UpdateTrackingCache(TArray<FPCGActorSelectionKey>* OptionalChangedKeys = nullptr);
 
-	/** Gather all the settings that are tracking the given actor. If the settings should cull and there is no intersection, it will not be added to the list.
-	* @param InActor                    The actor the component is tracking
+	/** Gather all the settings that are tracking the given object. If it is an actor, the settings should cull and there is no intersection, it will not be added to the list.
+	* @param InObject                   The object the component is tracking. Can be an Actor (for culling).
 	* @param bIntersect                 If the actor is intersecting with the component. Useful for culling.
 	* @param InRemovedTags              List of tags that were removed, because the actor won't have those tags on it, but we still need to gather the settings that were tracking this tag. Can be empty.
 	* @param InOriginatingChangeObject: Optional pointer on the originating object that triggered the update. Useful to track PCGComponents generation/cleanup. Can be null.
 	* @return                           Array of all the settings that track the actor.
 	*/
-	TArray<const UPCGSettings*> GatherSettingsTrackingActor(const AActor* InActor, const bool bIntersect, const TSet<FName>& InRemovedTags, const UObject* InOriginatingChangeObject) const;
+	TArray<const UPCGSettings*> GatherSettingsTracking(const UObject* InObject, const bool bIntersect, const TSet<FName>& InRemovedTags, const UObject* InOriginatingChangeObject) const;
 
 	bool ShouldTrackLandscape() const;
 
