[PCG] Trivial update of comment that was causing confusion

diff --git a/Engine/Plugins/Experimental/PCG/Source/PCG/Private/Elements/PCGBranch.cpp b/Engine/Plugins/Experimental/PCG/Source/PCG/Private/Elements/PCGBranch.cpp
--- a/Engine/Plugins/Experimental/PCG/Source/PCG/Private/Elements/PCGBranch.cpp
+++ b/Engine/Plugins/Experimental/PCG/Source/PCG/Private/Elements/PCGBranch.cpp
@@ -32,7 +32,7 @@ EPCGChangeType UPCGBranchSettings::GetChangeTypeForProperty(const FName& InPrope
 {
 	EPCGChangeType ChangeType = Super::GetChangeTypeForProperty(InPropertyName) | EPCGChangeType::Cosmetic;
 
-	// Grid sizes are processed during graph compilation and is part of the graph structure.
+	// Static branches are processed during graph compilation and are part of the graph structure.
 	if (InPropertyName == GET_MEMBER_NAME_CHECKED(UPCGBranchSettings, bEnabled)
 		|| InPropertyName == GET_MEMBER_NAME_CHECKED(UPCGBranchSettings, bOutputToB))
 	{
