[PCG] Graph cache debugging - fix cache entry count displayed in node title after multiple recent changes to change propagation and node titles broke it.

diff --git a/Engine/Plugins/Experimental/PCG/Source/PCG/Private/PCGNode.cpp b/Engine/Plugins/Experimental/PCG/Source/PCG/Private/PCGNode.cpp
--- a/Engine/Plugins/Experimental/PCG/Source/PCG/Private/PCGNode.cpp
+++ b/Engine/Plugins/Experimental/PCG/Source/PCG/Private/PCGNode.cpp
@@ -355,7 +355,7 @@ FText UPCGNode::GetNodeTitle(EPCGNodeTitleType TitleType) const
 
 		if (bDebuggingEnabled)
 		{
-			SecondaryTitleLine += FString::Format(TEXT(" [{1}]"), { NumCacheEntries });
+			SecondaryTitleLine += FString::Format(TEXT(" [{0}]"), { NumCacheEntries });
 		}
 #endif
 
diff --git a/Engine/Plugins/Experimental/PCG/Source/PCGEditor/Private/PCGEditor.cpp b/Engine/Plugins/Experimental/PCG/Source/PCGEditor/Private/PCGEditor.cpp
--- a/Engine/Plugins/Experimental/PCG/Source/PCGEditor/Private/PCGEditor.cpp
+++ b/Engine/Plugins/Experimental/PCG/Source/PCGEditor/Private/PCGEditor.cpp
@@ -2357,6 +2357,21 @@ void FPCGEditor::OnComponentGenerationCompleteOrCancelled()
 {
 	DebugObjectWidget->RefreshDebugObjects();
 	DebugObjectTreeWidget->RequestRefresh();
+
+	// If we are debugging the graph cache then we need to refresh the cache count displayed in the title after every generation.
+	if (UPCGSubsystem* Subsystem = GetSubsystem())
+	{
+		if (Subsystem->IsGraphCacheDebuggingEnabled() && PCGEditorGraph)
+		{
+			for (UEdGraphNode* EdGraphNode : PCGEditorGraph->Nodes)
+			{
+				if (UPCGEditorGraphNodeBase* PCGEditorGraphNode = Cast<UPCGEditorGraphNodeBase>(EdGraphNode))
+				{
+					PCGEditorGraphNode->ReconstructNode();
+				}
+			}
+		}
+	}
 }
 
 UPCGSubsystem* FPCGEditor::GetSubsystem()
