[PCG] Fix subgraph node asset subtitle to be correctly dependent on subgraph override pin, and fix FullTitle being used in Attribute List View

Also always take the primary title as the input for renaming on request from users.

diff --git a/Engine/Plugins/Experimental/PCG/Source/PCG/Private/PCGSubgraph.cpp b/Engine/Plugins/Experimental/PCG/Source/PCG/Private/PCGSubgraph.cpp
--- a/Engine/Plugins/Experimental/PCG/Source/PCG/Private/PCGSubgraph.cpp
+++ b/Engine/Plugins/Experimental/PCG/Source/PCG/Private/PCGSubgraph.cpp
@@ -279,12 +279,13 @@ UPCGNode* UPCGSubgraphSettings::CreateNode() const
 
 FString UPCGSubgraphSettings::GetAdditionalTitleInformation() const
 {
-	const UPCGNode* Node = Cast<UPCGNode>(GetOuter());
-	if (Node && Node->IsInputPinConnected(PCGPinConstants::DefaultInputLabel))
+#if WITH_EDITOR
+	if (IsPropertyOverriddenByPin(GET_MEMBER_NAME_CHECKED(UPCGSubgraphSettings, SubgraphOverride)))
 	{
-		// If subgraph is dynamic, don't emit any extra info here.
+		// Subgraphs with the subgraph override pin connected should not display any asset path.
 		return FString();
 	}
+#endif
 
 	if (UPCGGraph* TargetSubgraph = GetSubgraph())
 	{
diff --git a/Engine/Plugins/Experimental/PCG/Source/PCGEditor/Private/SPCGEditorGraphAttributeListView.cpp b/Engine/Plugins/Experimental/PCG/Source/PCGEditor/Private/SPCGEditorGraphAttributeListView.cpp
--- a/Engine/Plugins/Experimental/PCG/Source/PCGEditor/Private/SPCGEditorGraphAttributeListView.cpp
+++ b/Engine/Plugins/Experimental/PCG/Source/PCGEditor/Private/SPCGEditorGraphAttributeListView.cpp
@@ -565,7 +565,7 @@ void SPCGEditorGraphAttributeListView::OnInspectedNodeChanged(UPCGEditorGraphNod
 
 	if (PCGEditorGraphNode.IsValid())
 	{
-		NodeNameTextBlock->SetText(PCGEditorGraphNode->GetNodeTitle(ENodeTitleType::FullTitle));
+		NodeNameTextBlock->SetText(PCGEditorGraphNode->GetNodeTitle(ENodeTitleType::ListView));
 		NodeNameTextBlock->SetToolTipText(PCGEditorGraphNode->GetTooltipText());
 	}
 	else
diff --git a/Engine/Plugins/Experimental/PCG/Source/PCGEditor/Private/SPCGEditorGraphNode.cpp b/Engine/Plugins/Experimental/PCG/Source/PCGEditor/Private/SPCGEditorGraphNode.cpp
--- a/Engine/Plugins/Experimental/PCG/Source/PCGEditor/Private/SPCGEditorGraphNode.cpp
+++ b/Engine/Plugins/Experimental/PCG/Source/PCGEditor/Private/SPCGEditorGraphNode.cpp
@@ -580,13 +580,6 @@ void SPCGEditorGraphNode::OnNodeRenameInitiated()
 {
 	if (InlineEditableText.IsValid())
 	{
-		// If titles are flipped, the primary title widget will contain the generated title. Inject
-		// the other title here so the edit action will change the authored title.
-		if (PCGEditorGraphNode->HasFlippedTitleLines())
-		{
-			InlineEditableText->SetText(PCGEditorGraphNode->GetAuthoredTitleLine());
-		}
-
 		InlineEditableText->EnterEditingMode();
 	}
 }
