[PCG] Empty Input to Surface Sampler should now result in a no operation, not a warning

#rnx

diff --git a/Engine/Plugins/Experimental/PCG/Source/PCG/Private/Elements/PCGSurfaceSampler.cpp b/Engine/Plugins/Experimental/PCG/Source/PCG/Private/Elements/PCGSurfaceSampler.cpp
--- a/Engine/Plugins/Experimental/PCG/Source/PCG/Private/Elements/PCGSurfaceSampler.cpp
+++ b/Engine/Plugins/Experimental/PCG/Source/PCG/Private/Elements/PCGSurfaceSampler.cpp
@@ -331,8 +331,7 @@ namespace PCGSurfaceSamplerHelpers
 		// If there are no surfaces to sample, early out
 		if (SurfaceInputs.IsEmpty())
 		{
-			PCGE_LOG_C(Warning, GraphAndLog, Context, LOCTEXT("NoSurfaceInputs", "No Surface input was provided."));
-			return EPCGTimeSliceInitResult::AbortExecution;
+			return EPCGTimeSliceInitResult::NoOperation;
 		}
 
 		// Early out on invalid settings
