Refactored the accessor keys on metadata attributes to rely on entries on the metadata directly instead of the ones from the entry key to value map

#rnx

diff --git a/Engine/Plugins/Experimental/PCG/Source/PCG/Private/Elements/Metadata/PCGMetadataOpElementBase.cpp b/Engine/Plugins/Experimental/PCG/Source/PCG/Private/Elements/Metadata/PCGMetadataOpElementBase.cpp
--- a/Engine/Plugins/Experimental/PCG/Source/PCG/Private/Elements/Metadata/PCGMetadataOpElementBase.cpp
+++ b/Engine/Plugins/Experimental/PCG/Source/PCG/Private/Elements/Metadata/PCGMetadataOpElementBase.cpp
@@ -440,6 +440,15 @@ bool FPCGMetadataElementBase::PrepareDataInternal(FPCGContext* Context) const
 					PCGE_LOG(Error, GraphAndLog, FText::Format(LOCTEXT("CantCreateDefaultValue", "Pin '{0}' supports default value but we could not create it."), FText::FromName(CurrentPinLabel)));
 					return EPCGTimeSliceInitResult::AbortExecution;
 				}
+				else
+				{
+					// Need to make sure the param data has at least one entry
+					UPCGMetadata* DefaultParamMetadata = CastChecked<UPCGParamData>(DefaultData.Data)->Metadata;
+					if (DefaultParamMetadata->GetLocalItemCount() == 0)
+					{
+						DefaultParamMetadata->AddEntry();
+					}
+				}
 			}
 
 			if (CurrentPinInputData.IsEmpty())
diff --git a/Engine/Plugins/Experimental/PCG/Source/PCG/Private/Metadata/Accessors/PCGAttributeAccessorHelpers.cpp b/Engine/Plugins/Experimental/PCG/Source/PCG/Private/Metadata/Accessors/PCGAttributeAccessorHelpers.cpp
--- a/Engine/Plugins/Experimental/PCG/Source/PCG/Private/Metadata/Accessors/PCGAttributeAccessorHelpers.cpp
+++ b/Engine/Plugins/Experimental/PCG/Source/PCG/Private/Metadata/Accessors/PCGAttributeAccessorHelpers.cpp
@@ -19,26 +19,14 @@
 
 namespace PCGAttributeAccessorHelpers
 {
-	void ExtractMetadataAtribute(UPCGData* InData, FName Name, UPCGMetadata*& OutMetadata, FPCGMetadataAttributeBase*& OutAttribute)
+	/** Retrieves metadata object and the corresponding attribute if any for the provided data. */
+	void ExtractMetadataAttribute(UPCGData* InData, FName Name, UPCGMetadata*& OutMetadata, FPCGMetadataAttributeBase*& OutAttribute)
 	{
-		OutMetadata = nullptr;
-		OutAttribute = nullptr;
-
-		if (UPCGSpatialData* SpatialData = Cast<UPCGSpatialData>(InData))
-		{
-			OutMetadata = SpatialData->Metadata;
-		}
-		else if (UPCGParamData* ParamData = Cast<UPCGParamData>(InData))
-		{
-			OutMetadata = ParamData->Metadata;
-		}
-
-		if (OutMetadata)
-		{
-			OutAttribute = OutMetadata->GetMutableAttribute(Name);
-		}
+		OutMetadata = InData->MutableMetadata();
+		OutAttribute = OutMetadata ? OutMetadata->GetMutableAttribute(Name) : nullptr;
 	}
 
+	// This template creates a const or non-const accessor on an attribute depending on the provided templated type. 
 	template <
 		typename AccessorType,
 		typename = typename std::enable_if_t<
@@ -130,7 +118,7 @@ namespace PCGAttributeAccessorHelpers
 			FPCGMetadataAttributeBase* Attribute = nullptr;
 
 			// It is OK to const_cast here, since we will create a const accessor if the input is const.
-			ExtractMetadataAtribute(const_cast<UPCGData*>(InData), Name, Metadata, Attribute);
+			ExtractMetadataAttribute(const_cast<UPCGData*>(InData), Name, Metadata, Attribute);
 
 			// To simplify the code here & below we'll get a non-const accessor but force readonly if it should be
 			Accessor = CreateAttributeAccessorImpl<std::remove_const_t<AccessorType>>(Attribute, Metadata, /*bForceReadOnly=*/std::is_const_v<AccessorType>);
@@ -193,33 +181,25 @@ namespace PCGAttributeAccessorHelpers
 				return MakeUnique<FPCGAttributeAccessorKeysPoints>(View);
 			}
 		}
-
-		UPCGMetadata* Metadata = nullptr;
-		FPCGMetadataAttributeBase* Attribute = nullptr;
-
-		// It is OK to const_cast here, since we will create const keys if the input is const.
-		ExtractMetadataAtribute(const_cast<UPCGData*>(InData), InSelector.GetName(), Metadata, Attribute);
-
-		if (Attribute)
+		else if (InData) // e.g. is UPCGData and non-null
 		{
 			if constexpr (std::is_const_v<KeysType>)
 			{
-				return MakeUnique<FPCGAttributeAccessorKeysEntries>(static_cast<const FPCGMetadataAttributeBase*>(Attribute));
+				if (const UPCGMetadata* Metadata = InData->ConstMetadata())
+				{
+					return MakeUnique<FPCGAttributeAccessorKeysEntries>(Metadata);
+				}
 			}
 			else
 			{
-				return MakeUnique<FPCGAttributeAccessorKeysEntries>(Attribute);
+				if (UPCGMetadata* Metadata = InData->MutableMetadata())
+				{
+					return MakeUnique<FPCGAttributeAccessorKeysEntries>(Metadata);
+				}
 			}
 		}
-		else if (Metadata && InSelector.GetSelection() == EPCGAttributePropertySelection::ExtraProperty && InSelector.GetExtraProperty() == EPCGExtraProperties::Index)
-		{
-			// Special case for Indexes, we will use the metadata to get the keys
-			return MakeUnique<FPCGAttributeAccessorKeysEntries>(Metadata);
-		}
-		else
-		{
-			return TUniquePtr<KeysType>();
-		}
+
+		return TUniquePtr<KeysType>();
 	}
 
 	TUniquePtr<IPCGAttributeAccessor> CreateChainAccessor(TUniquePtr<IPCGAttributeAccessor> InAccessor, FName Name, bool& bOutSuccess)
diff --git a/Engine/Plugins/Experimental/PCG/Source/PCG/Private/Metadata/Accessors/PCGAttributeAccessorKeys.cpp b/Engine/Plugins/Experimental/PCG/Source/PCG/Private/Metadata/Accessors/PCGAttributeAccessorKeys.cpp
--- a/Engine/Plugins/Experimental/PCG/Source/PCG/Private/Metadata/Accessors/PCGAttributeAccessorKeys.cpp
+++ b/Engine/Plugins/Experimental/PCG/Source/PCG/Private/Metadata/Accessors/PCGAttributeAccessorKeys.cpp
@@ -7,32 +7,10 @@
 #include "PCGPoint.h"
 
 //////////////////////////////////////////////////////////////////// 
-
 FPCGAttributeAccessorKeysEntries::FPCGAttributeAccessorKeysEntries(const FPCGMetadataAttributeBase* Attribute)
-	: IPCGAttributeAccessorKeys(/*bInReadOnly=*/ false)
+	: FPCGAttributeAccessorKeysEntries(Attribute->GetMetadata())
 {
-	const FPCGMetadataAttributeBase* Current = Attribute;
-	while (Current)
-	{
-		TArray<PCGMetadataEntryKey> Temp;
-		Current->GetEntryToValueKeyMap_NotThreadSafe().GenerateKeyArray(Temp);
-		ExtractedEntries.Append(Temp);
-		Current = Current->GetParent();
-	}
-
-	// If the attribute doesn't have any entry, re-try with metadata entries.
-	if (Attribute && ExtractedEntries.IsEmpty())
-	{
-		InitializeFromMetadata(Attribute->GetMetadata());
-	}
-
-	// If the attribute still doesn't have any entry, we will always take the default value.
-	if (ExtractedEntries.IsEmpty())
-	{
-		ExtractedEntries.Add(PCGInvalidEntryKey);
-	}
-
-	Entries = TArrayView<PCGMetadataEntryKey>(ExtractedEntries);
+	// Deprecated
 }
 
 FPCGAttributeAccessorKeysEntries::FPCGAttributeAccessorKeysEntries(PCGMetadataEntryKey EntryKey)
@@ -49,6 +27,20 @@ FPCGAttributeAccessorKeysEntries::FPCGAttributeAccessorKeysEntries(const UPCGMet
 	Entries = TArrayView<PCGMetadataEntryKey>(ExtractedEntries);
 }
 
+FPCGAttributeAccessorKeysEntries::FPCGAttributeAccessorKeysEntries(UPCGMetadata* Metadata)
+	: IPCGAttributeAccessorKeys(/*bInReadOnly=*/ false)
+{
+	InitializeFromMetadata(Metadata);
+
+	// If the entries are still empty, we will always take the default value.
+	if (ExtractedEntries.IsEmpty())
+	{
+		ExtractedEntries.Add(PCGInvalidEntryKey);
+	}
+
+	Entries = TArrayView<PCGMetadataEntryKey>(ExtractedEntries);
+}
+
 FPCGAttributeAccessorKeysEntries::FPCGAttributeAccessorKeysEntries(const TArrayView<PCGMetadataEntryKey>& InEntries)
 	: IPCGAttributeAccessorKeys(/*bInReadOnly=*/false)
 	, Entries(InEntries)
diff --git a/Engine/Plugins/Experimental/PCG/Source/PCG/Private/Tests/Accessors/PCGAttributeAccessorTest.cpp b/Engine/Plugins/Experimental/PCG/Source/PCG/Private/Tests/Accessors/PCGAttributeAccessorTest.cpp
--- a/Engine/Plugins/Experimental/PCG/Source/PCG/Private/Tests/Accessors/PCGAttributeAccessorTest.cpp
+++ b/Engine/Plugins/Experimental/PCG/Source/PCG/Private/Tests/Accessors/PCGAttributeAccessorTest.cpp
@@ -141,7 +141,7 @@ bool FPCGAttributeAccessorSingleGetDefaultTest::RunTest(const FString& Parameter
 bool FPCGAttributeAccessorSingleGetTest::RunTest(const FString& Parameters)
 {
 	FPCGAttributeAccessorTestHelpers::AttributeData AttributeData;
-	FPCGAttributeAccessorKeysEntries Keys(&AttributeData.FloatAttribute);
+	FPCGAttributeAccessorKeysEntries Keys(AttributeData.FloatAttribute.GetMetadata());
 
 	if (!TestEqual("Right number of entries", Keys.GetNum(), FPCGAttributeAccessorTestHelpers::ValuesSize))
 	{
@@ -181,7 +181,7 @@ bool FPCGAttributeAccessorSingleGetTest::RunTest(const FString& Parameters)
 bool FPCGAttributeAccessorSingleGetOutsideRangeTest::RunTest(const FString& Parameters)
 {
 	FPCGAttributeAccessorTestHelpers::AttributeData AttributeData;
-	FPCGAttributeAccessorKeysEntries Keys(&AttributeData.FloatAttribute);
+	FPCGAttributeAccessorKeysEntries Keys(AttributeData.FloatAttribute.GetMetadata());
 
 	float DefaultFloat = 1.0f;
 
@@ -255,7 +255,7 @@ bool FPCGAttributeAccessorGetRangeDefaultTest::RunTest(const FString& Parameters
 bool FPCGAttributeAccessorGetRangeTest::RunTest(const FString& Parameters)
 {
 	FPCGAttributeAccessorTestHelpers::AttributeData AttributeData;
-	FPCGAttributeAccessorKeysEntries Keys(&AttributeData.FloatAttribute);
+	FPCGAttributeAccessorKeysEntries Keys(AttributeData.FloatAttribute.GetMetadata());
 
 	int32 Index = 0;
 
@@ -309,7 +309,7 @@ bool FPCGAttributeAccessorGetRangeTest::RunTest(const FString& Parameters)
 bool FPCGAttributeAccessorGetRangeOutsideRangeTest::RunTest(const FString& Parameters)
 {
 	FPCGAttributeAccessorTestHelpers::AttributeData AttributeData;
-	FPCGAttributeAccessorKeysEntries Keys(&AttributeData.FloatAttribute);
+	FPCGAttributeAccessorKeysEntries Keys(AttributeData.FloatAttribute.GetMetadata());
 
 	int32 OutsideRangeIndex = FPCGAttributeAccessorTestHelpers::ValuesSize;
 
@@ -392,7 +392,7 @@ bool FPCGAttributeAccessorSingleSetDefaultTest::RunTest(const FString& Parameter
 bool FPCGAttributeAccessorSingleSetTest::RunTest(const FString& Parameters)
 {
 	FPCGAttributeAccessorTestHelpers::AttributeData AttributeData;
-	FPCGAttributeAccessorKeysEntries Keys(&AttributeData.FloatAttribute);
+	FPCGAttributeAccessorKeysEntries Keys(const_cast<UPCGMetadata*>(AttributeData.FloatAttribute.GetMetadata()));
 
 	if (!TestTrue("Set float attribute", AttributeData.FloatAccessor->Set(*AttributeData.TempFloats, Keys))
 		|| !TestEqual("float value 0", AttributeData.FloatAttribute.GetValueFromItemKey(0), *AttributeData.TempFloats))
@@ -443,7 +443,7 @@ bool FPCGAttributeAccessorSingleInvalidKeySetTest::RunTest(const FString& Parame
 bool FPCGAttributeAccessorSetRangeTest::RunTest(const FString& Parameters)
 {
 	FPCGAttributeAccessorTestHelpers::AttributeData AttributeData;
-	FPCGAttributeAccessorKeysEntries Keys(&AttributeData.FloatAttribute);
+	FPCGAttributeAccessorKeysEntries Keys(const_cast<UPCGMetadata*>(AttributeData.FloatAttribute.GetMetadata()));
 
 	const int32 Index = 0;
 
@@ -477,7 +477,7 @@ bool FPCGAttributeAccessorSetRangeTest::RunTest(const FString& Parameters)
 bool FPCGAttributeAccessorSetRangeOutsideRangeTest::RunTest(const FString& Parameters)
 {
 	FPCGAttributeAccessorTestHelpers::AttributeData AttributeData;
-	FPCGAttributeAccessorKeysEntries Keys(&AttributeData.FloatAttribute);
+	FPCGAttributeAccessorKeysEntries Keys(AttributeData.FloatAttribute.GetMetadata());
 
 	const int32 Index = 0;
 
diff --git a/Engine/Plugins/Experimental/PCG/Source/PCG/Private/Tests/Elements/Metadata/PCGMetadataBooleanOpTest.cpp b/Engine/Plugins/Experimental/PCG/Source/PCG/Private/Tests/Elements/Metadata/PCGMetadataBooleanOpTest.cpp
--- a/Engine/Plugins/Experimental/PCG/Source/PCG/Private/Tests/Elements/Metadata/PCGMetadataBooleanOpTest.cpp
+++ b/Engine/Plugins/Experimental/PCG/Source/PCG/Private/Tests/Elements/Metadata/PCGMetadataBooleanOpTest.cpp
@@ -41,10 +41,12 @@ bool FPCGMetadataBooleanOpTest::RunTest(const FString& Parameters)
 	ParamData1->Metadata->CreateAttribute<bool>(TrueAttribute, /*DefaultValue=*/ true, bAllowInterpolation, bOverrideParent);
 	ParamData1->Metadata->CreateAttribute<bool>(FalseAttribute, /*DefaultValue=*/ false, bAllowInterpolation, bOverrideParent);
 	ParamData1->Metadata->CreateAttribute<int64>(InvalidAttribute, /*DefaultValue=*/ 0, bAllowInterpolation, bOverrideParent);
+	ParamData1->Metadata->AddEntry();
 
 	ParamData2->Metadata->CreateAttribute<bool>(TrueAttribute, /*DefaultValue=*/ true, bAllowInterpolation, bOverrideParent);
 	ParamData2->Metadata->CreateAttribute<bool>(FalseAttribute, /*DefaultValue=*/ false, bAllowInterpolation, bOverrideParent);
 	ParamData2->Metadata->CreateAttribute<int64>(InvalidAttribute, /*DefaultValue=*/ 0, bAllowInterpolation, bOverrideParent);
+	ParamData2->Metadata->AddEntry();
 
 	auto ValidateOp = [this, &MetadataBooleanElement, &TestData, &Settings](bool bExpectedResult, bool bIsValid = true) -> bool
 	{
diff --git a/Engine/Plugins/Experimental/PCG/Source/PCG/Public/Elements/Metadata/PCGMetadataOpElementBase.h b/Engine/Plugins/Experimental/PCG/Source/PCG/Public/Elements/Metadata/PCGMetadataOpElementBase.h
--- a/Engine/Plugins/Experimental/PCG/Source/PCG/Public/Elements/Metadata/PCGMetadataOpElementBase.h
+++ b/Engine/Plugins/Experimental/PCG/Source/PCG/Public/Elements/Metadata/PCGMetadataOpElementBase.h
@@ -273,7 +273,15 @@ namespace NAryOperation
 			using OutType = decltype(OutValue);
 
 			TArray<OutType, TInlineAllocator<DefaultChunkSize>> OutValues;
-			OutValues.SetNum(Range);
+			if constexpr (std::is_trivially_copyable_v<OutType>)
+			{
+				OutValues.SetNumUninitialized(Range);
+			}
+			else
+			{
+				OutValues.SetNum(Range);
+			}
+			
 			OutValues[0] = OutValue;
 
 			for (int32 i = 1; i < Range; ++i)
@@ -333,7 +341,15 @@ namespace NAryOperation
 		bool bSuccess = true;
 
 		TArray<InputType, TInlineAllocator<DefaultChunkSize>> InputValues;
-		InputValues.SetNum(Range);
+		if constexpr (std::is_trivially_copyable_v<InputType>)
+		{
+			InputValues.SetNumUninitialized(Range);
+		}
+		else
+		{
+			InputValues.SetNum(Range);
+		}
+		
 		if (InOptions.bUseDefaultKey)
 		{
 			FPCGAttributeAccessorKeysEntries DefaultAccessorKey(PCGInvalidEntryKey);
diff --git a/Engine/Plugins/Experimental/PCG/Source/PCG/Public/Metadata/Accessors/PCGAttributeAccessorKeys.h b/Engine/Plugins/Experimental/PCG/Source/PCG/Public/Metadata/Accessors/PCGAttributeAccessorKeys.h
--- a/Engine/Plugins/Experimental/PCG/Source/PCG/Public/Metadata/Accessors/PCGAttributeAccessorKeys.h
+++ b/Engine/Plugins/Experimental/PCG/Source/PCG/Public/Metadata/Accessors/PCGAttributeAccessorKeys.h
@@ -138,13 +138,16 @@ protected:
 class PCG_API FPCGAttributeAccessorKeysEntries : public IPCGAttributeAccessorKeys
 {
 public:
+	UE_DEPRECATED(5.5, "This key accessor is deprecated and replaced by the one taking a const or non-const UPCGMetadata object instead")
 	explicit FPCGAttributeAccessorKeysEntries(const FPCGMetadataAttributeBase* Attribute);
+
 	explicit FPCGAttributeAccessorKeysEntries(PCGMetadataEntryKey EntryKey);
 	explicit FPCGAttributeAccessorKeysEntries(const TArrayView<PCGMetadataEntryKey>& InEntries);
 	explicit FPCGAttributeAccessorKeysEntries(const TArrayView<const PCGMetadataEntryKey>& InEntries);
 
-	// Read-only case where we will just iterate on all the entries in the metadata.
+	// Iterates on all the entries in the metadata.
 	explicit FPCGAttributeAccessorKeysEntries(const UPCGMetadata* Metadata);
+	explicit FPCGAttributeAccessorKeysEntries(UPCGMetadata* Metadata);
 
 	virtual int32 GetNum() const override { return Entries.Num(); }
 
