Refactored and optimized trivial merge operations, esp. in cases where a lot of points had no initial non-default metadata.
Added threading to the merge operation by passing an optional context

#rnx

diff --git a/Engine/Plugins/Experimental/PCG/Source/PCG/Private/Data/PCGUnionData.cpp b/Engine/Plugins/Experimental/PCG/Source/PCG/Private/Data/PCGUnionData.cpp
--- a/Engine/Plugins/Experimental/PCG/Source/PCG/Private/Data/PCGUnionData.cpp
+++ b/Engine/Plugins/Experimental/PCG/Source/PCG/Private/Data/PCGUnionData.cpp
@@ -308,7 +308,7 @@ const UPCGPointData* UPCGUnionData::CreatePointData(FPCGContext* Context) const
 
 					if (OutMetadata && DatumPointMetadata && DatumPointMetadata->GetAttributeCount() > 0)
 					{
-						OutMetadata->SetPointAttributes(MakeArrayView(DatumPointData->GetPoints()), DatumPointMetadata, TargetPointsSubset);
+						OutMetadata->SetPointAttributes(MakeArrayView(DatumPointData->GetPoints()), DatumPointMetadata, TargetPointsSubset, Context);
 					}
 				}
 			}
diff --git a/Engine/Plugins/Experimental/PCG/Source/PCG/Private/Elements/PCGMergeAttributes.cpp b/Engine/Plugins/Experimental/PCG/Source/PCG/Private/Elements/PCGMergeAttributes.cpp
--- a/Engine/Plugins/Experimental/PCG/Source/PCG/Private/Elements/PCGMergeAttributes.cpp
+++ b/Engine/Plugins/Experimental/PCG/Source/PCG/Private/Elements/PCGMergeAttributes.cpp
@@ -65,13 +65,19 @@ bool FPCGMergeAttributesElement::ExecuteInternal(FPCGContext* Context) const
 			continue;
 		}
 
+		const UPCGMetadata* SourceMetadata = SourceData->Metadata;
+		const int64 ParamItemCount = SourceData->Metadata->GetLocalItemCount();
+
+		if (ParamItemCount <= 0)
+		{
+			continue;
+		}
+
 		if (!MergedOutput)
 		{
 			MergedOutput = &Outputs.Add_GetRef(Source);
 			continue;
 		}
-
-		const UPCGMetadata* SourceMetadata = SourceData->Metadata;
 		
 		// When we're merging the 2nd element, create the actual merged attribute set
 		if (!MergedAttributeSet)
@@ -87,14 +93,18 @@ bool FPCGMergeAttributesElement::ExecuteInternal(FPCGContext* Context) const
 		MergedAttributeSet->Metadata->AddAttributes(SourceMetadata);
 
 		// - Merge entries
-		const int64 ParamItemCount = SourceData->Metadata->GetLocalItemCount();
-
+		TArray<PCGMetadataEntryKey, TInlineAllocator<256>> SourceEntryKeys;
+		SourceEntryKeys.SetNumUninitialized(ParamItemCount);
 		for (int64 LocalItemKey = 0; LocalItemKey < ParamItemCount; ++LocalItemKey)
 		{
-			PCGMetadataEntryKey MergedItemKey = MergedAttributeSet->Metadata->AddEntry();
-			MergedAttributeSet->Metadata->SetAttributes(LocalItemKey, SourceMetadata, MergedItemKey);
+			SourceEntryKeys[LocalItemKey] = LocalItemKey;
 		}
 
+		TArray<PCGMetadataEntryKey, TInlineAllocator<256>> EntryKeys;
+		EntryKeys.Init(PCGInvalidEntryKey, ParamItemCount);
+
+		MergedAttributeSet->Metadata->SetAttributes(SourceEntryKeys, SourceMetadata, EntryKeys, Context);
+
 		// - Merge tags too (to be in line with the Merge points node)
 		MergedOutput->Tags.Append(Source.Tags);
 	}
diff --git a/Engine/Plugins/Experimental/PCG/Source/PCG/Private/Elements/PCGMergeElement.cpp b/Engine/Plugins/Experimental/PCG/Source/PCG/Private/Elements/PCGMergeElement.cpp
--- a/Engine/Plugins/Experimental/PCG/Source/PCG/Private/Elements/PCGMergeElement.cpp
+++ b/Engine/Plugins/Experimental/PCG/Source/PCG/Private/Elements/PCGMergeElement.cpp
@@ -51,14 +51,11 @@ bool FPCGMergeElement::ExecuteInternal(FPCGContext* Context) const
 	TArray<FPCGTaggedData> Sources = Context->InputData.GetInputsByPin(PCGPinConstants::DefaultInputLabel);
 	TArray<FPCGTaggedData>& Outputs = Context->OutputData.TaggedData;
 
-	if (Sources.IsEmpty())
-	{
-		return true;
-	}
-
 	UPCGPointData* TargetPointData = nullptr;
 	FPCGTaggedData* TargetTaggedData = nullptr;
 
+	int32 TotalPointCount = 0;
+
 	// Prepare data & metadata
 	// Done in two passes for futureproofing - expecting changes in the metadata attribute creation vs. usage in points
 	for (const FPCGTaggedData& Source : Sources)
@@ -71,15 +68,29 @@ bool FPCGMergeElement::ExecuteInternal(FPCGContext* Context) const
 			continue;
 		}
 
-		if (!TargetPointData)
+		if (SourcePointData->GetPoints().IsEmpty())
 		{
-			TargetPointData = NewObject<UPCGPointData>();
-			TargetPointData->InitializeFromData(SourcePointData, nullptr, bMergeMetadata);
+			continue;
+		}
+
+		TotalPointCount += SourcePointData->GetPoints().Num();
 
-			TargetTaggedData = &(Outputs.Emplace_GetRef(Source));
+		// First data that's valid - will be the original output so we don't do a copy unless we actually need it
+		if (!TargetTaggedData)
+		{
+			TargetTaggedData = &(Outputs.Add_GetRef(Source));
+		}
+		else if (!TargetPointData)
+		{
+			// Second valid data - we'll create the actual merged data at this point
+			check(TargetTaggedData);
+
+			TargetPointData = NewObject<UPCGPointData>();
+			TargetPointData->InitializeFromData(CastChecked<const UPCGPointData>(TargetTaggedData->Data), nullptr, bMergeMetadata);
 			TargetTaggedData->Data = TargetPointData;
 		}
-		else
+
+		if (TargetPointData)
 		{
 			if (bMergeMetadata)
 			{
@@ -91,14 +102,15 @@ bool FPCGMergeElement::ExecuteInternal(FPCGContext* Context) const
 		}
 	}
 
-	// No valid input types
+	// If there was no valid input or only one, there's nothing to do here
 	if (!TargetPointData)
 	{
 		return true;
 	}
 
 	TArray<FPCGPoint>& TargetPoints = TargetPointData->GetMutablePoints();
-	
+	TargetPoints.Reserve(TotalPointCount);
+
 	for(int32 SourceIndex = 0; SourceIndex < Sources.Num(); ++SourceIndex)
 	{
 		const UPCGPointData* SourcePointData = Cast<const UPCGPointData>(Sources[SourceIndex].Data);
@@ -121,7 +133,7 @@ bool FPCGMergeElement::ExecuteInternal(FPCGContext* Context) const
 
 			if (bMergeMetadata && TargetPointData->Metadata && SourcePointData->Metadata && SourcePointData->Metadata->GetAttributeCount() > 0)
 			{
-				TargetPointData->Metadata->SetPointAttributes(MakeArrayView(SourcePointData->GetPoints()), SourcePointData->Metadata, TargetPointsSubset);
+				TargetPointData->Metadata->SetPointAttributes(MakeArrayView(SourcePointData->GetPoints()), SourcePointData->Metadata, TargetPointsSubset, Context);
 			}
 		}
 	}
diff --git a/Engine/Plugins/Experimental/PCG/Source/PCG/Private/Metadata/PCGMetadata.cpp b/Engine/Plugins/Experimental/PCG/Source/PCG/Private/Metadata/PCGMetadata.cpp
--- a/Engine/Plugins/Experimental/PCG/Source/PCG/Private/Metadata/PCGMetadata.cpp
+++ b/Engine/Plugins/Experimental/PCG/Source/PCG/Private/Metadata/PCGMetadata.cpp
@@ -2,6 +2,7 @@
 
 #include "Metadata/PCGMetadata.h"
 
+#include "PCGContext.h"
 #include "PCGData.h"
 #include "PCGPoint.h"
 #include "Elements/Metadata/PCGMetadataElementCommon.h"
@@ -9,6 +10,7 @@
 #include "Metadata/PCGAttributePropertySelector.h"
 
 #include "Algo/Transform.h"
+#include "Async/ParallelFor.h"
 
 #include UE_INLINE_GENERATED_CPP_BY_NAME(PCGMetadata)
 
@@ -1178,23 +1180,25 @@ void UPCGMetadata::SetAttributes(PCGMetadataEntryKey InKey, const UPCGMetadata*
 	AttributeLock.ReadUnlock();
 }
 
-void UPCGMetadata::SetPointAttributes(const TArrayView<const FPCGPoint>& InPoints, const UPCGMetadata* InMetadata, const TArrayView<FPCGPoint>& OutPoints)
+void UPCGMetadata::SetPointAttributes(const TArrayView<const FPCGPoint>& InPoints, const UPCGMetadata* InMetadata, const TArrayView<FPCGPoint>& OutPoints, FPCGContext* OptionalContext)
 {
 	if (!InMetadata || InMetadata->GetAttributeCount() == 0 || GetAttributeCount() == 0)
 	{
 		return;
 	}
 
+	TRACE_CPUPROFILER_EVENT_SCOPE(UPCGMetadata::SetPointAttributes);
+
 	check(InPoints.Num() == OutPoints.Num());
 
 	// Extract the metadata entry keys from the in & out points
-	TArray<PCGMetadataEntryKey> InKeys;
-	TArray<PCGMetadataEntryKey> OutKeys;
+	TArray<PCGMetadataEntryKey, TInlineAllocator<256>> InKeys;
+	TArray<PCGMetadataEntryKey, TInlineAllocator<256>> OutKeys;
 
 	Algo::Transform(InPoints, InKeys, [](const FPCGPoint& Point) { return Point.MetadataEntry; });
 	Algo::Transform(OutPoints, OutKeys, [](const FPCGPoint& Point) { return Point.MetadataEntry; });
 
-	SetAttributes(InKeys, InMetadata, OutKeys);
+	SetAttributes(InKeys, InMetadata, OutKeys, OptionalContext);
 
 	// Write back the keys on the points
 	for (int KeyIndex = 0; KeyIndex < OutKeys.Num(); ++KeyIndex)
@@ -1203,72 +1207,125 @@ void UPCGMetadata::SetPointAttributes(const TArrayView<const FPCGPoint>& InPoint
 	}
 }
 
-void UPCGMetadata::SetAttributes(const TArrayView<PCGMetadataEntryKey>& InKeys, const UPCGMetadata* InMetadata, const TArrayView<PCGMetadataEntryKey>& OutKeys)
+void UPCGMetadata::SetAttributes(const TArrayView<const PCGMetadataEntryKey>& InOriginalKeys, const UPCGMetadata* InMetadata, const TArrayView<PCGMetadataEntryKey>& OutOriginalKeys, FPCGContext* OptionalContext)
 {
 	if (!InMetadata || InMetadata->GetAttributeCount() == 0 || GetAttributeCount() == 0)
 	{
 		return;
 	}
 
-	check(InKeys.Num() == OutKeys.Num());
+	TRACE_CPUPROFILER_EVENT_SCOPE(UPCGMetadata::SetAttributes);
 
-	for (int32 KeyIndex = 0; KeyIndex < InKeys.Num(); ++KeyIndex)
-	{
-		InitializeOnSet(OutKeys[KeyIndex], InKeys[KeyIndex], InMetadata);
-	}
+	check(InOriginalKeys.Num() == OutOriginalKeys.Num());
 
-	// Rarely-used convenience array that might be used multiple times, will be kept out of the loop to limit the number of allocations if needed
-	TArray<PCGMetadataValueKey> ValueKeys;
+	// There are a few things we can do to optimize here -
+	// basically, we don't need to set attributes more than once for a given <in, out> pair
+	TArray<PCGMetadataEntryKey, TInlineAllocator<256>> InKeys;
+	TArray<PCGMetadataEntryKey, TInlineAllocator<256>> OutKeys;
 
-	AttributeLock.ReadLock();
-	for(const TPair<FName, FPCGMetadataAttributeBase*>& AttributePair : Attributes)
 	{
-		const FName& AttributeName = AttributePair.Key;
-		FPCGMetadataAttributeBase* Attribute = AttributePair.Value;
+		TRACE_CPUPROFILER_EVENT_SCOPE(UPCGMetadata::SetAttributes::CreateDeduplicatedKeys);
+		TMap<TPair<PCGMetadataEntryKey, PCGMetadataEntryKey>, int> PairMapping;
 
-		if (const FPCGMetadataAttributeBase* OtherAttribute = InMetadata->GetConstAttribute(AttributeName))
+		for (int KeyIndex = 0; KeyIndex < InOriginalKeys.Num(); ++KeyIndex)
 		{
-			if (!PCG::Private::IsBroadcastableOrConstructible(OtherAttribute->GetTypeId(), Attribute->GetTypeId()))
-			{
-				UE_LOG(LogPCG, Error, TEXT("Metadata type mismatch with attribute '%s'"), *AttributeName.ToString());
-				continue;
-			}
+			PCGMetadataEntryKey InKey = InOriginalKeys[KeyIndex];
+			PCGMetadataEntryKey& OutKey = OutOriginalKeys[KeyIndex];
 
-			if (Attribute == OtherAttribute)
+			if (int* MatchingPairIndex = PairMapping.Find(TPair<PCGMetadataEntryKey, PCGMetadataEntryKey>(InKey, OutKey)))
 			{
-				ValueKeys.Reset();
-				Attribute->GetValueKeys(InKeys, ValueKeys);
-				Attribute->SetValuesFromValueKeys(OutKeys, ValueKeys);
+				OutKey = *MatchingPairIndex;
 			}
 			else
 			{
-				// Create accessor for the other attribute
-				TUniquePtr<const IPCGAttributeAccessor> OtherAttributeAccessor = PCGAttributeAccessorHelpers::CreateConstAccessor(OtherAttribute, InMetadata);
-				FPCGAttributeAccessorKeysEntries OtherAttributeKeys(InKeys);
+				int NewIndex = InKeys.Add(InKey);
 
-				if (!OtherAttributeAccessor)
+				PairMapping.Emplace(TPair<PCGMetadataEntryKey, PCGMetadataEntryKey>(InKey, OutKey), NewIndex);
+				OutKeys.Add(OutKey);
+				OutKey = NewIndex;
+			}
+		}
+	}
+
+	{
+		TRACE_CPUPROFILER_EVENT_SCOPE(UPCGMetadata::SetAttributes::InitializeOnSet);
+
+		for (int32 KeyIndex = 0; KeyIndex < InKeys.Num(); ++KeyIndex)
+		{
+			InitializeOnSet(OutKeys[KeyIndex], InKeys[KeyIndex], InMetadata);
+		}
+	}
+
+	AttributeLock.ReadLock();
+	int32 AttributeOffset = 0;
+	const int32 AttributesPerDispatch = OptionalContext ? FMath::Max(1, OptionalContext->AsyncState.NumAvailableTasks) : 1;
+
+	while (AttributeOffset < Attributes.Num())
+	{
+		TArray<FName> AttributeNames;
+		TArray<EPCGMetadataTypes> AttributeTypes;
+		GetAttributes(AttributeNames, AttributeTypes);
+
+		const int32 AttributeCountInCurrentDispatch = FMath::Min(AttributesPerDispatch, Attributes.Num() - AttributeOffset);
+		ParallelFor(AttributeCountInCurrentDispatch, [this, OptionalContext, AttributeOffset, InMetadata, &AttributeNames, &InKeys, &OutKeys](int32 WorkerIndex)
+		{
+			const FName AttributeName = AttributeNames[AttributeOffset + WorkerIndex];
+			FPCGMetadataAttributeBase* Attribute = Attributes[AttributeName];
+
+			if (const FPCGMetadataAttributeBase* OtherAttribute = InMetadata->GetConstAttribute(AttributeName))
+			{
+				if (!PCG::Private::IsBroadcastableOrConstructible(OtherAttribute->GetTypeId(), Attribute->GetTypeId()))
 				{
-					continue;
+					PCGE_LOG_C(Error, GraphAndLog, OptionalContext, FText::Format(NSLOCTEXT("PCGMetadata", "TypeMismatch", "Metadata type mismatch with attribute '{0}'"), FText::FromName(AttributeName)));
+					return;
 				}
 
-				auto GetAndSetValues = [Attribute, &OutKeys, &OtherAttributeAccessor, &OtherAttributeKeys](auto Dummy) -> bool
+				if (Attribute == OtherAttribute)
 				{
-					using Type = decltype(Dummy);
+					TArray<PCGMetadataValueKey> ValueKeys;
+					Attribute->GetValueKeys(InKeys, ValueKeys);
+					Attribute->SetValuesFromValueKeys(OutKeys, ValueKeys);
+				}
+				else
+				{
+					// Create accessor for the other attribute
+					TUniquePtr<const IPCGAttributeAccessor> OtherAttributeAccessor = PCGAttributeAccessorHelpers::CreateConstAccessor(OtherAttribute, InMetadata);
+
+					TArrayView<PCGMetadataEntryKey> InKeysView(InKeys);
+					FPCGAttributeAccessorKeysEntries OtherAttributeKeys(InKeysView);
 
-					auto SetValues = [Attribute, &OutKeys](const TArrayView<Type>& View, const int32 Start, const int32 Range)
+					if (!OtherAttributeAccessor)
 					{
-						TArrayView<PCGMetadataEntryKey> Keys(OutKeys.GetData() + Start, Range);
-						static_cast<FPCGMetadataAttribute<Type>*>(Attribute)->SetValues(Keys, View);
-					};
+						return;
+					}
 
-					return PCGMetadataElementCommon::ApplyOnAccessorRange<Type>(OtherAttributeKeys, *OtherAttributeAccessor, SetValues, EPCGAttributeAccessorFlags::AllowBroadcast | EPCGAttributeAccessorFlags::AllowConstructible);
-				};
+					auto GetAndSetValues = [Attribute, &OutKeys, &OtherAttributeAccessor, &OtherAttributeKeys](auto Dummy) -> bool
+					{
+						using Type = decltype(Dummy);
+
+						auto SetValues = [Attribute, &OutKeys](const TArrayView<Type>& View, const int32 Start, const int32 Range)
+						{
+							TArrayView<PCGMetadataEntryKey> Keys(OutKeys.GetData() + Start, Range);
+							static_cast<FPCGMetadataAttribute<Type>*>(Attribute)->SetValues(Keys, View);
+						};
 
-				PCGMetadataAttribute::CallbackWithRightType(Attribute->GetTypeId(), GetAndSetValues);
+						return PCGMetadataElementCommon::ApplyOnAccessorRange<Type>(OtherAttributeKeys, *OtherAttributeAccessor, SetValues, EPCGAttributeAccessorFlags::AllowBroadcast | EPCGAttributeAccessorFlags::AllowConstructible);
+					};
+
+					PCGMetadataAttribute::CallbackWithRightType(Attribute->GetTypeId(), GetAndSetValues);
+				}
 			}
-		}
+		});
+
+		AttributeOffset += AttributeCountInCurrentDispatch;
 	}
 	AttributeLock.ReadUnlock();
+
+	// Finally, copy back the actual out keys to the original out keys
+	for (PCGMetadataEntryKey& OutKey : OutOriginalKeys)
+	{
+		OutKey = OutKeys[OutKey];
+	}
 }
 
 void UPCGMetadata::MergeAttributesByKey(int64 KeyA, const UPCGMetadata* MetadataA, int64 KeyB, const UPCGMetadata* MetadataB, int64 TargetKey, EPCGMetadataOp Op, int64& OutKey)
diff --git a/Engine/Plugins/Experimental/PCG/Source/PCG/Public/Elements/PCGDataFromActor.h b/Engine/Plugins/Experimental/PCG/Source/PCG/Public/Elements/PCGDataFromActor.h
--- a/Engine/Plugins/Experimental/PCG/Source/PCG/Public/Elements/PCGDataFromActor.h
+++ b/Engine/Plugins/Experimental/PCG/Source/PCG/Public/Elements/PCGDataFromActor.h
@@ -92,9 +92,8 @@ public:
 #endif // WITH_EDITORONLY_DATA
 };
 
-class FPCGDataFromActorContext : public FPCGContext
+struct FPCGDataFromActorContext : public FPCGContext
 {
-public:
 	TArray<AActor*> FoundActors;
 	bool bPerformedQuery = false;
 };
diff --git a/Engine/Plugins/Experimental/PCG/Source/PCG/Public/Metadata/Accessors/PCGAttributeAccessor.h b/Engine/Plugins/Experimental/PCG/Source/PCG/Public/Metadata/Accessors/PCGAttributeAccessor.h
--- a/Engine/Plugins/Experimental/PCG/Source/PCG/Public/Metadata/Accessors/PCGAttributeAccessor.h
+++ b/Engine/Plugins/Experimental/PCG/Source/PCG/Public/Metadata/Accessors/PCGAttributeAccessor.h
@@ -69,6 +69,7 @@ public:
 
 		int LastDefaultKeyIndex = INDEX_NONE;
 
+		// Implementation note: this is a stripped down version of UPCGMetadata::InitializeOnSet
 		for(int EntryIndex = 0; EntryIndex < EntryKeys.Num(); ++EntryIndex)
 		{
 			PCGMetadataEntryKey& EntryKey = *EntryKeys[EntryIndex];
@@ -83,6 +84,10 @@ public:
 					LastDefaultKeyIndex = EntryIndex;
 				}
 			}
+			else if (EntryKey < Metadata->GetItemKeyCountForParent())
+			{
+				EntryKey = Metadata->AddEntry(EntryKey);
+			}
 		}
 
 		Attribute->SetValues(EntryKeys, InValues);
diff --git a/Engine/Plugins/Experimental/PCG/Source/PCG/Public/Metadata/PCGMetadata.h b/Engine/Plugins/Experimental/PCG/Source/PCG/Public/Metadata/PCGMetadata.h
--- a/Engine/Plugins/Experimental/PCG/Source/PCG/Public/Metadata/PCGMetadata.h
+++ b/Engine/Plugins/Experimental/PCG/Source/PCG/Public/Metadata/PCGMetadata.h
@@ -9,7 +9,7 @@
 #include "PCGMetadata.generated.h"
 
 struct FPCGPoint;
-
+struct FPCGContext;
 
 UCLASS(BlueprintType)
 class PCG_API UPCGMetadata : public UObject
@@ -255,12 +255,12 @@ public:
 	void AccumulateWeightedAttributes(PCGMetadataEntryKey InKey, const UPCGMetadata* InMetadata, float Weight, bool bSetNonInterpolableAttributes, PCGMetadataEntryKey& OutKey);
 
 	void SetAttributes(PCGMetadataEntryKey InKey, const UPCGMetadata* InMetadata, PCGMetadataEntryKey& OutKey);
-	void SetAttributes(const TArrayView<PCGMetadataEntryKey>& InKeys, const UPCGMetadata* InMetadata, const TArrayView<PCGMetadataEntryKey>& OutKeys);
+	void SetAttributes(const TArrayView<const PCGMetadataEntryKey>& InKeys, const UPCGMetadata* InMetadata, const TArrayView<PCGMetadataEntryKey>& OutKeys, FPCGContext* OptionalContext = nullptr);
 
 	/** Attributes operations - shorthand for points */
 	void MergePointAttributes(const FPCGPoint& InPointA, const FPCGPoint& InPointB, FPCGPoint& OutPoint, EPCGMetadataOp Op);
 	void MergePointAttributesSubset(const FPCGPoint& InPointA, const UPCGMetadata* InMetadataA, const UPCGMetadata* InMetadataSubetA, const FPCGPoint& InPointB, const UPCGMetadata* InMetadataB, const UPCGMetadata* InMetadataSubsetb, FPCGPoint& OutPoint, EPCGMetadataOp Op);
-	void SetPointAttributes(const TArrayView<const FPCGPoint>& InPoints, const UPCGMetadata* InMetadata, const TArrayView<FPCGPoint>& OutPoints);
+	void SetPointAttributes(const TArrayView<const FPCGPoint>& InPoints, const UPCGMetadata* InMetadata, const TArrayView<FPCGPoint>& OutPoints, FPCGContext* OptionalContext = nullptr);
 
 	/** Blueprint-friend versions */
 	UFUNCTION(BlueprintCallable, Category = "PCG|Metadata")
