Cleanup after removing the need for entries in attributes

diff --git a/Engine/Plugins/Experimental/PCG/Source/PCG/Private/Elements/PCGAttributeGetFromPointIndexElement.cpp b/Engine/Plugins/Experimental/PCG/Source/PCG/Private/Elements/PCGAttributeGetFromPointIndexElement.cpp
--- a/Engine/Plugins/Experimental/PCG/Source/PCG/Private/Elements/PCGAttributeGetFromPointIndexElement.cpp
+++ b/Engine/Plugins/Experimental/PCG/Source/PCG/Private/Elements/PCGAttributeGetFromPointIndexElement.cpp
@@ -140,7 +140,7 @@ bool FPCGAttributeGetFromPointIndexElement::ExecuteInternal(FPCGContext* Context
 					return false;
 				}
 
-				NewAttribute->SetValue(OutputParamData->Metadata->AddEntry(), Value);
+				OutputParamData->Metadata->AddEntry();
 
 				return true;
 			};
diff --git a/Engine/Plugins/Experimental/PCG/Source/PCG/Private/Elements/PCGAttributeReduceElement.cpp b/Engine/Plugins/Experimental/PCG/Source/PCG/Private/Elements/PCGAttributeReduceElement.cpp
--- a/Engine/Plugins/Experimental/PCG/Source/PCG/Private/Elements/PCGAttributeReduceElement.cpp
+++ b/Engine/Plugins/Experimental/PCG/Source/PCG/Private/Elements/PCGAttributeReduceElement.cpp
@@ -288,6 +288,8 @@ bool FPCGAttributeReduceElement::ExecuteInternal(FPCGContext* Context) const
 
 			if (bSuccess)
 			{
+				// Implementation note: since the default value does not match the value computed here
+				// and because we might have multiple entries, we need to set it in the attribute
 				TypedNewAttribute->SetValue(OutputParams->Metadata->AddEntry(), OutputValue);
 			}
 
diff --git a/Engine/Plugins/Experimental/PCG/Source/PCG/Private/Elements/PCGAttributeSelectElement.cpp b/Engine/Plugins/Experimental/PCG/Source/PCG/Private/Elements/PCGAttributeSelectElement.cpp
--- a/Engine/Plugins/Experimental/PCG/Source/PCG/Private/Elements/PCGAttributeSelectElement.cpp
+++ b/Engine/Plugins/Experimental/PCG/Source/PCG/Private/Elements/PCGAttributeSelectElement.cpp
@@ -402,8 +402,8 @@ bool FPCGAttributeSelectElement::ExecuteInternal(FPCGContext* Context) const
 
 			if (bSuccess)
 			{
+				OutputParamData->Metadata->AddEntry();
 				NewAttribute->SetDefaultValue(OutputValue);
-				NewAttribute->SetValue(OutputParamData->Metadata->AddEntry(), OutputValue);
 			}
 			else
 			{
diff --git a/Engine/Plugins/Experimental/PCG/Source/PCG/Private/Elements/PCGCreateAttribute.cpp b/Engine/Plugins/Experimental/PCG/Source/PCG/Private/Elements/PCGCreateAttribute.cpp
--- a/Engine/Plugins/Experimental/PCG/Source/PCG/Private/Elements/PCGCreateAttribute.cpp
+++ b/Engine/Plugins/Experimental/PCG/Source/PCG/Private/Elements/PCGCreateAttribute.cpp
@@ -417,17 +417,6 @@ bool FPCGCreateAttributeElement::ExecuteInternal(FPCGContext* Context) const
 
 					PCGMetadataElementCommon::CopyFromAccessorToAccessor(Params);
 
-					// In case of param data, do some padding at the end if needed
-					if (bIsParamData)
-					{
-						int32 NumInKeys = InputKeys->GetNum();
-						int32 Padding = OutputKeys->GetNum() - NumInKeys;
-						for (int64 Key = 0; Key < Padding; ++Key)
-						{
-							Attribute->SetValueFromValueKey(PCGMetadataEntryKey(NumInKeys + Key), PCGDefaultValueKey);
-						}
-					}
-
 					return Attribute;
 				};
 
@@ -452,16 +441,6 @@ bool FPCGCreateAttributeElement::ExecuteInternal(FPCGContext* Context) const
 			{
 				Metadata->AddEntry();
 			}
-
-			// Also make sure if the attribute has no entries, to map all the metadata entry to the default value key
-			// Very important to do in the case of multi-entry attribute set.
-			if (Attribute->GetNumberOfEntries() == 0)
-			{
-				for (PCGMetadataEntryKey Key = 0; Key < Metadata->GetLocalItemCount(); ++Key)
-				{
-					Attribute->SetValueFromValueKey(Key, PCGDefaultValueKey);
-				}
-			}
 		}
 
 		TArray<FPCGTaggedData>& Outputs = Context->OutputData.TaggedData;
diff --git a/Engine/Plugins/Experimental/PCG/Source/PCG/Private/Elements/PCGCreateTargetActor.cpp b/Engine/Plugins/Experimental/PCG/Source/PCG/Private/Elements/PCGCreateTargetActor.cpp
--- a/Engine/Plugins/Experimental/PCG/Source/PCG/Private/Elements/PCGCreateTargetActor.cpp
+++ b/Engine/Plugins/Experimental/PCG/Source/PCG/Private/Elements/PCGCreateTargetActor.cpp
@@ -314,7 +314,7 @@ bool FPCGCreateTargetActorElement::ExecuteInternal(FPCGContext* Context) const
 	check(ParamData && ParamData->Metadata);
 	FPCGMetadataAttribute<FSoftObjectPath>* ActorPathAttribute = ParamData->Metadata->CreateAttribute<FSoftObjectPath>(PCGPointDataConstants::ActorReferenceAttribute, GeneratedActorPath, /*bAllowsInterpolation=*/false, /*bOverrideParent=*/false);
 	check(ActorPathAttribute);
-	ActorPathAttribute->SetValue(ParamData->Metadata->AddEntry(), GeneratedActorPath);	
+	ParamData->Metadata->AddEntry();
 
 	// Add param data to output and we're done
 	Context->OutputData.TaggedData.Emplace_GetRef().Data = ParamData;
diff --git a/Engine/Plugins/Experimental/PCG/Source/PCG/Private/Elements/PCGDataNum.cpp b/Engine/Plugins/Experimental/PCG/Source/PCG/Private/Elements/PCGDataNum.cpp
--- a/Engine/Plugins/Experimental/PCG/Source/PCG/Private/Elements/PCGDataNum.cpp
+++ b/Engine/Plugins/Experimental/PCG/Source/PCG/Private/Elements/PCGDataNum.cpp
@@ -65,7 +65,7 @@ bool FPCGDataNumElement::ExecuteInternal(FPCGContext* Context) const
 		return true;
 	}
 
-	NumAttribute->SetValue(OutputParamData->Metadata->AddEntry(), InputDataCount);
+	OutputParamData->Metadata->AddEntry();
 
 	FPCGTaggedData& Output = Context->OutputData.TaggedData.Emplace_GetRef();
 	Output.Data = OutputParamData;
diff --git a/Engine/Plugins/Experimental/PCG/Source/PCG/Private/Elements/PCGGetLoopIndex.cpp b/Engine/Plugins/Experimental/PCG/Source/PCG/Private/Elements/PCGGetLoopIndex.cpp
--- a/Engine/Plugins/Experimental/PCG/Source/PCG/Private/Elements/PCGGetLoopIndex.cpp
+++ b/Engine/Plugins/Experimental/PCG/Source/PCG/Private/Elements/PCGGetLoopIndex.cpp
@@ -83,7 +83,7 @@ bool FPCGGetLoopIndexElement::ExecuteInternal(FPCGContext* Context) const
 	FPCGMetadataAttribute<int>* LoopIndexAttribute = LoopIndexParamData->Metadata->CreateAttribute<int>(PCGGetLoopIndexConstants::LoopIndexAttributeName, LoopIndex, /*bAllowInterpolation=*/false, /*bOverrideParent=*/false);
 	check(LoopIndexAttribute);
 
-	LoopIndexAttribute->SetValue(LoopIndexParamData->Metadata->AddEntry(), LoopIndex);
+	LoopIndexParamData->Metadata->AddEntry();
 
 	return true;
 }
diff --git a/Engine/Plugins/Experimental/PCG/Source/PCG/Private/Elements/PCGNumberOfElements.cpp b/Engine/Plugins/Experimental/PCG/Source/PCG/Private/Elements/PCGNumberOfElements.cpp
--- a/Engine/Plugins/Experimental/PCG/Source/PCG/Private/Elements/PCGNumberOfElements.cpp
+++ b/Engine/Plugins/Experimental/PCG/Source/PCG/Private/Elements/PCGNumberOfElements.cpp
@@ -101,6 +101,7 @@ bool FPCGNumberOfElementsBaseElement<DataType>::ExecuteInternal(FPCGContext* Con
 		}
 
 		check(OutputParamData && NewAttribute);
+		// Implementation note: since we might have multiple entries, we must set the actual value to a new entry
 		NewAttribute->SetValue(OutputParamData->Metadata->AddEntry(), GetNum(InputData));
 	}
 
