Fixed issues in DuplicatePoint - the relative transform was not applied in the right order.
Added direction-relative displacement to build things that rotate a bit better

diff --git a/Engine/Plugins/Experimental/PCG/Source/PCG/Private/Elements/PCGDuplicatePoint.cpp b/Engine/Plugins/Experimental/PCG/Source/PCG/Private/Elements/PCGDuplicatePoint.cpp
--- a/Engine/Plugins/Experimental/PCG/Source/PCG/Private/Elements/PCGDuplicatePoint.cpp
+++ b/Engine/Plugins/Experimental/PCG/Source/PCG/Private/Elements/PCGDuplicatePoint.cpp
@@ -49,7 +49,8 @@ bool FPCGDuplicatePointElement::ExecuteInternal(FPCGContext* Context) const
 		const TArray<FPCGPoint>& InputPoints = InputPointData->GetPoints();
 
 		// Determines whether or not to include the source point in data
-		const int DuplicatesPerPoint = Iterations + (Settings->bOutputSourcePoint ? 1 : 0);
+		const bool bKeepSourcePoint = Settings->bOutputSourcePoint;
+		const int DuplicatesPerPoint = Iterations + (bKeepSourcePoint ? 1 : 0);
 		const int NumIterations = DuplicatesPerPoint * InputPoints.Num();
 		const int FirstDuplicateIndex = Settings->bOutputSourcePoint ? 0 : 1;
 
@@ -66,41 +67,65 @@ bool FPCGDuplicatePointElement::ExecuteInternal(FPCGContext* Context) const
 		OutputPoints.SetNumUninitialized(NumIterations);
 		Output.Data = OutPointData;
 
-		auto ProcessPoint = [Settings, FirstDuplicateIndex, OutPointData, InputPointData, &Direction, &InputPoints, &OutputPoints](const int32 ReadIndex, const int32 WriteIndex)
-		{
-			const FPCGPoint& InPoint = InputPoints[ReadIndex % InputPoints.Num()];
-			FPCGPoint& OutputPoint = OutputPoints[WriteIndex];
-			const int DuplicateIndex = FirstDuplicateIndex + ReadIndex / InputPoints.Num();
+		const FTransform& SourceDuplicateTransform = Settings->PointTransform;
 
-			if (DuplicateIndex == 0)
-			{
-				OutputPoint = InPoint;
-			}
-			else
+		if (Settings->bDirectionAppliedInRelativeSpace)
+		{
+			auto ProcessPoint = [SourceDuplicateTransform, DuplicatesPerPoint, bKeepSourcePoint, &Direction, &InputPoints, &OutputPoints](const int32 ReadIndex, const int32 WriteIndex)
 			{
-				const UPCGMetadata* InMetadata = InputPointData->Metadata;
-				UPCGMetadata* OutMetadata = OutPointData->Metadata;
+				const FPCGPoint& ReadPoint = InputPoints[ReadIndex];
+				TArrayView<FPCGPoint> WritePoints(OutputPoints.GetData() + WriteIndex * DuplicatesPerPoint, DuplicatesPerPoint);
+
+				const FTransform DuplicateAxisTransform = FTransform((ReadPoint.BoundsMax - ReadPoint.BoundsMin) * Direction);
+				const FTransform DuplicateTransform = DuplicateAxisTransform * SourceDuplicateTransform;
+				FTransform CurrentTransform = ReadPoint.Transform;
+
+				int WritePointIndex = 0;
 
-				UPCGMetadataAccessorHelpers::CopyPoint(InPoint, OutputPoint, /*bCopyMetadata=*/true, InMetadata, OutMetadata);
+				if (bKeepSourcePoint)
+				{
+					WritePoints[WritePointIndex++] = ReadPoint;
+				}
 
-				OutputPoint.Transform = FTransform(InPoint.Transform.GetRotation(), InPoint.Transform.GetLocation(), FVector::One());
-				OutputPoint.BoundsMin *= InPoint.Transform.GetScale3D();
-				OutputPoint.BoundsMax *= InPoint.Transform.GetScale3D();
+				while (WritePointIndex < DuplicatesPerPoint)
+				{
+					FPCGPoint& WritePoint = WritePoints[WritePointIndex++];
+					WritePoint = ReadPoint;
 
-				const FVector DuplicateLocationOffset = ((OutputPoint.BoundsMax - OutputPoint.BoundsMin) * Direction + Settings->PointTransform.GetLocation()) * DuplicateIndex;
-				const FRotator DuplicateRotationOffset = Settings->PointTransform.Rotator() * DuplicateIndex;
-				const FVector DuplicateScaleMultiplier = FVector(
-					FMath::Pow(Settings->PointTransform.GetScale3D().X, DuplicateIndex), 
-					FMath::Pow(Settings->PointTransform.GetScale3D().Y, DuplicateIndex), 
-					FMath::Pow(Settings->PointTransform.GetScale3D().Z, DuplicateIndex));
-				const FTransform DuplicateTransformOffset(DuplicateRotationOffset, DuplicateLocationOffset, DuplicateScaleMultiplier);
+					CurrentTransform = DuplicateTransform * CurrentTransform;
+					WritePoint.Transform = CurrentTransform;
+					WritePoint.Seed = PCGHelpers::ComputeSeedFromPosition(CurrentTransform.GetLocation());
+				}
+			};
 
-				OutputPoint.Transform *= DuplicateTransformOffset;
-				OutputPoint.Seed = PCGHelpers::ComputeSeedFromPosition(OutputPoint.Transform.GetLocation());
-			}
-		};
+			FPCGAsync::AsyncProcessingOneToOneEx(&Context->AsyncState, InputPoints.Num(), /*Initialize=*/[]() {}, ProcessPoint, /*bEnableTimeSlicing=*/false);
+		}
+		else
+		{
+			auto ProcessPoint = [SourceDuplicateTransform, FirstDuplicateIndex, OutPointData, InputPointData, &Direction, &InputPoints, &OutputPoints](const int32 ReadIndex, const int32 WriteIndex)
+			{
+				const FPCGPoint& InPoint = InputPoints[ReadIndex % InputPoints.Num()];
+				FPCGPoint& OutputPoint = OutputPoints[WriteIndex];
+				const int DuplicateIndex = FirstDuplicateIndex + ReadIndex / InputPoints.Num();
 
-		FPCGAsync::AsyncProcessingOneToOneEx(&Context->AsyncState, NumIterations, /*Initialize=*/[]() {}, ProcessPoint, /*bEnableTimeSlicing=*/false);
+				OutputPoint = InPoint;
+
+				if (DuplicateIndex != 0)
+				{
+					const FVector DuplicateLocationOffset = ((OutputPoint.BoundsMax - OutputPoint.BoundsMin) * Direction + SourceDuplicateTransform.GetLocation()) * DuplicateIndex;
+					const FRotator DuplicateRotationOffset = SourceDuplicateTransform.Rotator() * DuplicateIndex;
+					const FVector DuplicateScaleMultiplier = FVector(
+						FMath::Pow(SourceDuplicateTransform.GetScale3D().X, DuplicateIndex), 
+						FMath::Pow(SourceDuplicateTransform.GetScale3D().Y, DuplicateIndex), 
+						FMath::Pow(SourceDuplicateTransform.GetScale3D().Z, DuplicateIndex));
+
+					OutputPoint.Transform = FTransform(DuplicateRotationOffset, DuplicateLocationOffset, DuplicateScaleMultiplier) * InPoint.Transform;
+					OutputPoint.Seed = PCGHelpers::ComputeSeedFromPosition(OutputPoint.Transform.GetLocation());
+				}
+			};
+
+			FPCGAsync::AsyncProcessingOneToOneEx(&Context->AsyncState, NumIterations, /*Initialize=*/[]() {}, ProcessPoint, /*bEnableTimeSlicing=*/false);
+		}
 	}
 
 	return true;
diff --git a/Engine/Plugins/Experimental/PCG/Source/PCG/Public/Elements/PCGDuplicatePoint.h b/Engine/Plugins/Experimental/PCG/Source/PCG/Public/Elements/PCGDuplicatePoint.h
--- a/Engine/Plugins/Experimental/PCG/Source/PCG/Public/Elements/PCGDuplicatePoint.h
+++ b/Engine/Plugins/Experimental/PCG/Source/PCG/Public/Elements/PCGDuplicatePoint.h
@@ -37,6 +37,10 @@ public:
 	/** Direction to stack point duplicates. */
 	UPROPERTY(BlueprintReadWrite, EditAnywhere, Category = Settings, meta = (ClampMin = "-1.0", ClampMax = "1.0", PCG_Overridable))
 	FVector Direction = FVector(0.0, 0.0, 1.0);
+
+	/** Controls whether the axis displacement will be made in relative space or not */
+	UPROPERTY(BlueprintReadWrite, EditAnywhere, Category = Settings, meta = (PCG_Overridable))
+	bool bDirectionAppliedInRelativeSpace = false;
 	
 	/** Include the source point. */
 	UPROPERTY(BlueprintReadWrite, EditAnywhere, Category = Settings, meta = (PCG_Overridable))
