[PCG] Fix crash related to PCG Component eluding superclass PostEditChangeProperty

Crash hits a checkf when undoing a scale change on a PCG Volume

#rnx

diff --git a/Engine/Plugins/Experimental/PCG/Source/PCG/Private/PCGComponent.cpp b/Engine/Plugins/Experimental/PCG/Source/PCG/Private/PCGComponent.cpp
--- a/Engine/Plugins/Experimental/PCG/Source/PCG/Private/PCGComponent.cpp
+++ b/Engine/Plugins/Experimental/PCG/Source/PCG/Private/PCGComponent.cpp
@@ -1506,6 +1506,7 @@ void UPCGComponent::PostEditChangeProperty(FPropertyChangedEvent& PropertyChange
 {
 	if (!PropertyChangedEvent.Property || !IsValid(this))
 	{
+		Super::PostEditChangeProperty(PropertyChangedEvent);
 		return;
 	}
 
