#UE Do not reset loaders when renaming actors in UPCGSpawnActorSettings::RefreshTemplateActor. This function can be called during postload

diff --git a/Engine/Plugins/Experimental/PCG/Source/PCG/Private/Elements/PCGSpawnActor.cpp b/Engine/Plugins/Experimental/PCG/Source/PCG/Private/Elements/PCGSpawnActor.cpp
--- a/Engine/Plugins/Experimental/PCG/Source/PCG/Private/Elements/PCGSpawnActor.cpp
+++ b/Engine/Plugins/Experimental/PCG/Source/PCG/Private/Elements/PCGSpawnActor.cpp
@@ -329,7 +329,7 @@ void UPCGSpawnActorSettings::RefreshTemplateActor()
 				Options.bNotifyObjectReplacement = true;
 				UEngine::CopyPropertiesForUnrelatedObjects(TemplateActor, NewTemplateActor, Options);
 
-				TemplateActor->Rename(nullptr, GetTransientPackage(), REN_DontCreateRedirectors);
+				TemplateActor->Rename(nullptr, GetTransientPackage(), REN_DontCreateRedirectors|REN_ForceNoResetLoaders);
 
 				TMap<UObject*, UObject*> OldToNew;
 				OldToNew.Emplace(TemplateActor, NewTemplateActor);
