Fixed loctext conflict

#rnx

diff --git a/Engine/Plugins/Experimental/PCG/Source/PCG/Private/Elements/PCGDeleteAttributesElement.cpp b/Engine/Plugins/Experimental/PCG/Source/PCG/Private/Elements/PCGDeleteAttributesElement.cpp
--- a/Engine/Plugins/Experimental/PCG/Source/PCG/Private/Elements/PCGDeleteAttributesElement.cpp
+++ b/Engine/Plugins/Experimental/PCG/Source/PCG/Private/Elements/PCGDeleteAttributesElement.cpp
@@ -8,7 +8,7 @@
 #include "Data/PCGSpatialData.h"
 #include "Helpers/PCGHelpers.h"
 
-#define LOCTEXT_NAMESPACE "PCGAttributeFilterElement"
+#define LOCTEXT_NAMESPACE "PCGDeleteAttributesElement"
 
 namespace PCGAttributeFilterConstants
 {
@@ -77,7 +77,7 @@ FString UPCGDeleteAttributesSettings::GetAdditionalTitleInformation() const
 	}
 	else if (Operation == EPCGAttributeFilterOperation::DeleteSelectedAttributes)
 	{
-		OperationString = LOCTEXT("OperationKeep", "Delete").ToString();
+		OperationString = LOCTEXT("OperationDelete", "Delete").ToString();
 	}
 	else
 	{
@@ -91,11 +91,11 @@ FString UPCGDeleteAttributesSettings::GetAdditionalTitleInformation() const
 	}
 	else if (AttributesToKeep.IsEmpty())
 	{
-		return FString::Printf(TEXT("%s (%s)"), *OperationString, *LOCTEXT("NoAttributes", "none").ToString());
+		return FString::Printf(TEXT("%s (%s)"), *OperationString, *LOCTEXT("NoAttributes", "None").ToString());
 	}
 	else
 	{
-		return FString::Printf(TEXT("%s (%s)"), *OperationString, *LOCTEXT("KeepMultipleAttributes", "multiple").ToString());
+		return FString::Printf(TEXT("%s (%s)"), *OperationString, *LOCTEXT("KeepMultipleAttributes", "Multiple").ToString());
 	}
 }
 
