[PCG] Added a new Switch node

* Supports switch operations on int, string, and enum

TODO: Unit Tests CL to follow before commit

diff --git a/Engine/Plugins/Experimental/PCG/Source/PCG/Private/Elements/ControlFlow/PCGBooleanSelect.cpp b/Engine/Plugins/Experimental/PCG/Source/PCG/Private/Elements/ControlFlow/PCGBooleanSelect.cpp
new file mode 100644
--- /dev/null
+++ b/Engine/Plugins/Experimental/PCG/Source/PCG/Private/Elements/ControlFlow/PCGBooleanSelect.cpp
@@ -0,0 +1,93 @@
+// Copyright Epic Games, Inc. All Rights Reserved.
+
+#include "Elements/ControlFlow/PCGBooleanSelect.h"
+
+#include "PCGCommon.h"
+#include "PCGContext.h"
+#include "PCGPin.h"
+#include "Elements/PCGGather.h"
+
+#define LOCTEXT_NAMESPACE "FPCGBooleanSelectElement"
+
+namespace PCGBooleanSelectConstants
+{
+	const FName InputLabelA = TEXT("Input A");
+	const FName InputLabelB = TEXT("Input B");
+	const FText NodeTitleBase = LOCTEXT("NodeTitle", "Select");
+}
+
+#if WITH_EDITOR
+FText UPCGBooleanSelectSettings::GetDefaultNodeTitle() const
+{
+	// TODO: This should statically update or dynamically update, if overridden, for which input was selected, ie. Select (Input A)
+	return PCGBooleanSelectConstants::NodeTitleBase;
+}
+
+FText UPCGBooleanSelectSettings::GetNodeTooltipText() const
+{
+	return LOCTEXT("NodeTooltip", "Control flow node that will select all input data on either Pin A or Pin B only, based on the 'Use Input B' property - which can also be overridden.");
+}
+#endif // WITH_EDITOR
+
+EPCGDataType UPCGBooleanSelectSettings::GetCurrentPinTypes(const UPCGPin* Pin) const
+{
+	// Output pin depends on the union of both input pins 
+	if (Pin->IsOutputPin())
+	{
+		const EPCGDataType InputTypeUnion = GetTypeUnionOfIncidentEdges(PCGBooleanSelectConstants::InputLabelA) | GetTypeUnionOfIncidentEdges(PCGBooleanSelectConstants::InputLabelB);
+		return InputTypeUnion != EPCGDataType::None ? InputTypeUnion : EPCGDataType::Any;
+	}
+
+	return Super::GetCurrentPinTypes(Pin);
+}
+
+TArray<FPCGPinProperties> UPCGBooleanSelectSettings::InputPinProperties() const
+{
+	TArray<FPCGPinProperties> PinProperties;
+	PinProperties.Emplace(PCGBooleanSelectConstants::InputLabelA,
+		EPCGDataType::Any,
+		/*bInAllowMultipleConnections=*/true,
+		/*bAllowMultipleData=*/true,
+		LOCTEXT("FirstInputPinTooltip", "Will only be used if 'Use Input B' (overridable) is false"));
+	PinProperties.Emplace(PCGBooleanSelectConstants::InputLabelB,
+		EPCGDataType::Any,
+		/*bInAllowMultipleConnections=*/true,
+		/*bAllowMultipleData=*/true,
+		LOCTEXT("SecondInputPinTooltip", "Will only be used if 'Use Input B' (overridable) is true"));
+
+	return PinProperties;
+}
+
+TArray<FPCGPinProperties> UPCGBooleanSelectSettings::OutputPinProperties() const
+{
+	TArray<FPCGPinProperties> PinProperties;
+	PinProperties.Emplace(PCGPinConstants::DefaultOutputLabel,
+		EPCGDataType::Any,
+		/*bInAllowMultipleConnections=*/true,
+		/*bAllowMultipleData=*/true,
+		LOCTEXT("OutputPinTooltip", "All input will gathered into a single data collection"));
+
+	return PinProperties;
+}
+
+FPCGElementPtr UPCGBooleanSelectSettings::CreateElement() const
+{
+	return MakeShared<FPCGBooleanSelectElement>();
+}
+
+bool FPCGBooleanSelectElement::ExecuteInternal(FPCGContext* Context) const
+{
+	TRACE_CPUPROFILER_EVENT_SCOPE(FPCGBooleanSelectElement::ExecuteInternal);
+
+	const UPCGBooleanSelectSettings* Settings = Context->GetInputSettings<UPCGBooleanSelectSettings>();
+	check(Settings);
+
+	const FName SelectedPinLabel = Settings->bUseInputB ? PCGBooleanSelectConstants::InputLabelB : PCGBooleanSelectConstants::InputLabelA;
+
+	// Reuse the functionality of the Gather Node
+	Context->OutputData = PCGGather::GatherDataForPin(Context->InputData, SelectedPinLabel);
+
+	return true;
+}
+
+#undef LOCTEXT_NAMESPACE
\ No newline at end of file
diff --git a/Engine/Plugins/Experimental/PCG/Source/PCG/Private/Elements/ControlFlow/PCGBranch.cpp b/Engine/Plugins/Experimental/PCG/Source/PCG/Private/Elements/ControlFlow/PCGBranch.cpp
new file mode 100644
--- /dev/null
+++ b/Engine/Plugins/Experimental/PCG/Source/PCG/Private/Elements/ControlFlow/PCGBranch.cpp
@@ -0,0 +1,97 @@
+// Copyright Epic Games, Inc. All Rights Reserved.
+
+#include "Elements/ControlFlow/PCGBranch.h"
+
+#include "PCGCommon.h"
+#include "PCGContext.h"
+#include "PCGPin.h"
+#include "Elements/PCGGather.h"
+
+#define LOCTEXT_NAMESPACE "FPCGBranchElement"
+
+namespace PCGBranchConstants
+{
+	const FName OutputLabelA = TEXT("Output A");
+	const FName OutputLabelB = TEXT("Output B");
+	const FText NodeTitleBase = LOCTEXT("NodeTitle", "Branch");
+}
+
+#if WITH_EDITOR
+FText UPCGBranchSettings::GetDefaultNodeTitle() const
+{
+	// TODO: This should statically update or dynamically update, if overridden, for which branch was taken, ie. Branch (Output A)
+	return PCGBranchConstants::NodeTitleBase;
+}
+
+FText UPCGBranchSettings::GetNodeTooltipText() const
+{
+	return LOCTEXT("NodeTooltip", "Control flow node that will route the input to either Output A or Output B, based on the 'Output To B' property - which can also be overridden.");
+}
+
+EPCGChangeType UPCGBranchSettings::GetChangeTypeForProperty(const FName& InPropertyName) const
+{
+	EPCGChangeType ChangeType = Super::GetChangeTypeForProperty(InPropertyName) | EPCGChangeType::Cosmetic;
+
+	// Static branches are processed during graph compilation and are part of the graph structure.
+	if (InPropertyName == GET_MEMBER_NAME_CHECKED(UPCGBranchSettings, bEnabled)
+		|| InPropertyName == GET_MEMBER_NAME_CHECKED(UPCGBranchSettings, bOutputToB))
+	{
+		ChangeType |= EPCGChangeType::Structural;
+	}
+
+	return ChangeType;
+}
+#endif // WITH_EDITOR
+
+TArray<FPCGPinProperties> UPCGBranchSettings::OutputPinProperties() const
+{
+	TArray<FPCGPinProperties> PinProperties;
+	PinProperties.Emplace(PCGBranchConstants::OutputLabelA,
+		EPCGDataType::Any,
+		/*bInAllowMultipleConnections=*/true,
+		/*bAllowMultipleData=*/true,
+		LOCTEXT("OutputPinTooltipA", "Will only route input if 'Output To B' (overridable) is false"));
+	PinProperties.Emplace(PCGBranchConstants::OutputLabelB,
+		EPCGDataType::Any,
+		/*bInAllowMultipleConnections=*/true,
+		/*bAllowMultipleData=*/true,
+		LOCTEXT("OutputPinTooltipB", "Will only route input if 'Output To B' (overridable) is true"));
+
+	return PinProperties;
+}
+
+FPCGElementPtr UPCGBranchSettings::CreateElement() const
+{
+	return MakeShared<FPCGBranchElement>();
+}
+
+bool UPCGBranchSettings::IsPinStaticallyActive(const FName& PinLabel) const
+{
+	// Dynamic branches are never known in advance - assume all branches are active prior to execution.
+	if (IsPropertyOverriddenByPin(GET_MEMBER_NAME_CHECKED(UPCGBranchSettings, bOutputToB)))
+	{
+		return true;
+	}
+
+	// Branch must be both enabled and set to B in order for output pin B to be active.
+	const FName ActiveOutputPinLabel = (bEnabled && bOutputToB) ? PCGBranchConstants::OutputLabelB : PCGBranchConstants::OutputLabelA;
+
+	return PinLabel == ActiveOutputPinLabel;
+}
+
+bool FPCGBranchElement::ExecuteInternal(FPCGContext* Context) const
+{
+	TRACE_CPUPROFILER_EVENT_SCOPE(FPCGBranchElement::ExecuteInternal);
+
+	const UPCGBranchSettings* Settings = Context->GetInputSettings<UPCGBranchSettings>();
+	check(Settings);
+
+	const FName SelectedPinLabel = Settings->bOutputToB ? PCGBranchConstants::OutputLabelB : PCGBranchConstants::OutputLabelA;
+
+	// Reuse the functionality of the Gather Node
+	Context->OutputData = PCGGather::GatherDataForPin(Context->InputData, PCGPinConstants::DefaultInputLabel, SelectedPinLabel);
+
+	return true;
+}
+
+#undef LOCTEXT_NAMESPACE
diff --git a/Engine/Plugins/Experimental/PCG/Source/PCG/Private/Elements/ControlFlow/PCGMultiSelect.cpp b/Engine/Plugins/Experimental/PCG/Source/PCG/Private/Elements/ControlFlow/PCGMultiSelect.cpp
new file mode 100644
--- /dev/null
+++ b/Engine/Plugins/Experimental/PCG/Source/PCG/Private/Elements/ControlFlow/PCGMultiSelect.cpp
@@ -0,0 +1,274 @@
+// Copyright Epic Games, Inc. All Rights Reserved.
+
+#include "Elements/ControlFlow/PCGMultiSelect.h"
+
+#include "PCGContext.h"
+#include "PCGModule.h"
+#include "Elements/PCGGather.h"
+
+#define LOCTEXT_NAMESPACE "FPCGMultiSelectElement"
+
+namespace PCGMultiSelectConstants
+{
+	const FText NodeTitleBase = LOCTEXT("NodeTitleBase", "Select (Multi)");
+}
+
+void UPCGMultiSelectSettings::PostLoad()
+{
+	Super::PostLoad();
+	CachePinLabels();
+}
+
+#if WITH_EDITOR
+void UPCGMultiSelectSettings::PostEditChangeProperty(FPropertyChangedEvent& PropertyChangedEvent)
+{
+	Super::PostEditChangeProperty(PropertyChangedEvent);
+
+	// Only need to change the pin labels if the options have changed
+	if (PropertyChangedEvent.GetMemberPropertyName() == GET_MEMBER_NAME_CHECKED(UPCGMultiSelectSettings, IntOptions) ||
+		PropertyChangedEvent.GetMemberPropertyName() == GET_MEMBER_NAME_CHECKED(UPCGMultiSelectSettings, StringOptions) ||
+		PropertyChangedEvent.GetMemberPropertyName() == GET_MEMBER_NAME_CHECKED(UPCGMultiSelectSettings, EnumSelection))
+	{
+		CachePinLabels();
+	}
+}
+
+FText UPCGMultiSelectSettings::GetDefaultNodeTitle() const
+{
+	return PCGMultiSelectConstants::NodeTitleBase;
+}
+
+FText UPCGMultiSelectSettings::GetNodeTooltipText() const
+{
+	return LOCTEXT("NodeTooltip", "Control flow node that will select all input data on a single input pin that matches a given selection mode and corresponding 'selection' property - which can also be overridden.");
+}
+#endif // WITH_EDITOR
+
+EPCGDataType UPCGMultiSelectSettings::GetCurrentPinTypes(const UPCGPin* Pin) const
+{
+	// Output pin depends on the union of all non-advanced input pins. Advanced pins are dynamic user overrides and
+	// should not affect the output pin type.
+	if (Pin->IsOutputPin())
+	{
+		EPCGDataType InputTypeUnion = EPCGDataType::None;
+
+		const TArray<TObjectPtr<UPCGPin>>& InputPins = Pin->Node->GetInputPins();
+
+		for (int Index = 0; Index < InputPins.Num(); ++Index)
+		{
+			if (!InputPins[Index]->Properties.bAdvancedPin)
+			{
+				InputTypeUnion |= GetTypeUnionOfIncidentEdges(InputPins[Index]->Properties.Label);
+			}
+		}
+
+		return InputTypeUnion != EPCGDataType::None ? InputTypeUnion : EPCGDataType::Any;
+	}
+
+	return Super::GetCurrentPinTypes(Pin);
+}
+
+FString UPCGMultiSelectSettings::GetAdditionalTitleInformation() const
+{
+	switch (SelectionMode)
+	{
+		case EPCGControlFlowSelectionMode::Integer:
+		{
+			FString Subtitle = PCGControlFlowConstants::SubtitleInt.ToString();
+			if (!IsPropertyOverriddenByPin(GET_MEMBER_NAME_CHECKED(UPCGMultiSelectSettings, IntegerSelection)))
+			{
+				Subtitle += FString::Format(TEXT(": {0}"), {IntegerSelection});
+			}
+
+			return Subtitle;
+		}
+
+		case EPCGControlFlowSelectionMode::Enum:
+			if (EnumSelection.Class)
+			{
+				FString Subtitle = PCGControlFlowConstants::SubtitleEnum.ToString();
+				if (!IsPropertyOverriddenByPin({GET_MEMBER_NAME_CHECKED(UPCGMultiSelectSettings, EnumSelection), GET_MEMBER_NAME_CHECKED(FEnumSelector, Value)}))
+				{
+					Subtitle += EnumSelection.Class->GetNameStringByValue(EnumSelection.Value);
+				}
+
+				return Subtitle;
+			}
+			break;
+
+		case EPCGControlFlowSelectionMode::String:
+			{
+				FString Subtitle = PCGControlFlowConstants::SubtitleString.ToString();
+				if (!IsPropertyOverriddenByPin(GET_MEMBER_NAME_CHECKED(UPCGMultiSelectSettings, StringSelection)))
+				{
+					Subtitle += FString::Format(TEXT(": {0}"), {StringSelection});
+				}
+
+				return Subtitle;
+			}
+
+		default:
+			checkNoEntry();
+			break;
+	}
+
+	return PCGMultiSelectConstants::NodeTitleBase.ToString();
+}
+
+TArray<FPCGPinProperties> UPCGMultiSelectSettings::InputPinProperties() const
+{
+	TArray<FPCGPinProperties> PinProperties;
+
+	switch (SelectionMode)
+	{
+		case EPCGControlFlowSelectionMode::Integer:
+			for (const int32 Value : IntOptions)
+			{
+				PinProperties.Emplace(FName(FString::FromInt(Value)));
+			}
+			break;
+		case EPCGControlFlowSelectionMode::String:
+			for (const FString& Value : StringOptions)
+			{
+				PinProperties.Emplace(FName(Value));
+			}
+			break;
+		case EPCGControlFlowSelectionMode::Enum:
+			// -1 to bypass the MAX value
+			for (int32 Index = 0; EnumSelection.Class && Index < EnumSelection.Class->NumEnums() - 1; ++Index)
+			{
+#if WITH_EDITOR
+				bool const bHidden = EnumSelection.Class->HasMetaData(TEXT("Hidden"), Index) || EnumSelection.Class->HasMetaData(TEXT("Spacer"), Index);
+				if (!bHidden)
+				{
+					PinProperties.Emplace(FName(EnumSelection.Class->GetDisplayNameTextByIndex(Index).ToString()));
+				}
+#else // WITH_EDITOR
+				// HasMetaData is editor only, so there will be extra pins at runtime, but that should be okay
+				PinProperties.Emplace(FName(EnumSelection.Class->GetDisplayNameTextByIndex(Index).ToString()));
+#endif // WITH_EDITOR
+			}
+			break;
+		default:
+			break;
+	}
+
+	PinProperties.Emplace(PCGControlFlowConstants::DefaultPathPinLabel);
+
+	return PinProperties;
+}
+
+TArray<FPCGPinProperties> UPCGMultiSelectSettings::OutputPinProperties() const
+{
+	TArray<FPCGPinProperties> PinProperties;
+	PinProperties.Emplace(PCGPinConstants::DefaultOutputLabel,
+		EPCGDataType::Any,
+		/*bInAllowMultipleConnections=*/true,
+		/*bAllowMultipleData=*/true,
+		LOCTEXT("OutputPinTooltip", "All input on the selected input pin will be forwarded directly to the output."));
+
+	return PinProperties;
+}
+
+FPCGElementPtr UPCGMultiSelectSettings::CreateElement() const
+{
+	return MakeShared<FPCGMultiSelectElement>();
+}
+
+bool UPCGMultiSelectSettings::IsValuePresent(const int32 Value) const
+{
+	return IntOptions.Contains(Value);
+}
+
+bool UPCGMultiSelectSettings::IsValuePresent(const FString& Value) const
+{
+	return StringOptions.Contains(Value);
+}
+
+bool UPCGMultiSelectSettings::IsValuePresent(const int64 Value) const
+{
+	if (!EnumSelection.Class)
+	{
+		return false;
+	}
+
+	const int64 Index = EnumSelection.Class->GetIndexByValue(Value);
+	return Index != INDEX_NONE && Index < EnumSelection.Class->NumEnums() - 1;
+}
+
+bool UPCGMultiSelectSettings::GetSelectedPinLabel(FName& OutSelectedPinLabel) const
+{
+	if (CachedPinLabels.IsEmpty())
+	{
+		return false;
+	}
+
+	int32 Index = INDEX_NONE;
+	if (SelectionMode == EPCGControlFlowSelectionMode::Integer && IsValuePresent(IntegerSelection))
+	{
+		Index = IntOptions.IndexOfByKey(IntegerSelection);
+	}
+	else if (SelectionMode == EPCGControlFlowSelectionMode::String && IsValuePresent(StringSelection))
+	{
+		Index = StringOptions.IndexOfByKey(StringSelection);
+	}
+	else if (SelectionMode == EPCGControlFlowSelectionMode::Enum && IsValuePresent(EnumSelection.Value))
+	{
+		// To account for hidden enums missing from the pin properties, find the pin by label instead of index
+		const FName PinLabel(EnumSelection.Class->GetDisplayNameTextByValue(EnumSelection.Value).ToString());
+		for (int i = 0; i < CachedPinLabels.Num(); ++i)
+		{
+			if (CachedPinLabels[i] == PinLabel)
+			{
+				Index = i;
+				break;
+			}
+		}
+	}
+	else
+	{
+		OutSelectedPinLabel = PCGControlFlowConstants::DefaultPathPinLabel;
+
+		return true;
+	}
+
+	if (Index < 0 || Index >= CachedPinLabels.Num())
+	{
+		return false;
+	}
+
+	OutSelectedPinLabel = CachedPinLabels[Index];
+
+	return true;
+}
+
+void UPCGMultiSelectSettings::CachePinLabels()
+{
+	CachedPinLabels.Empty();
+	Algo::Transform(InputPinProperties(), CachedPinLabels, [](const FPCGPinProperties& Property)
+	{
+		return Property.Label;
+	});
+}
+
+bool FPCGMultiSelectElement::ExecuteInternal(FPCGContext* Context) const
+{
+	TRACE_CPUPROFILER_EVENT_SCOPE(FPCGMultiSelectElement::ExecuteInternal);
+
+	const UPCGMultiSelectSettings* Settings = Context->GetInputSettings<UPCGMultiSelectSettings>();
+	check(Settings);
+
+	FName SelectedPinLabel;
+	if (!Settings->GetSelectedPinLabel(SelectedPinLabel))
+	{
+		PCGE_LOG_C(Error, GraphAndLog, Context, LOCTEXT("ValueDoesNotExist", "Selected value is not a valid option."));
+		return true;
+	}
+
+	// Reuse the functionality of the Gather node
+	Context->OutputData = PCGGather::GatherDataForPin(Context->InputData, SelectedPinLabel);
+
+	return true;
+}
+
+#undef LOCTEXT_NAMESPACE
\ No newline at end of file
diff --git a/Engine/Plugins/Experimental/PCG/Source/PCG/Private/Elements/ControlFlow/PCGSwitch.cpp b/Engine/Plugins/Experimental/PCG/Source/PCG/Private/Elements/ControlFlow/PCGSwitch.cpp
new file mode 100644
--- /dev/null
+++ b/Engine/Plugins/Experimental/PCG/Source/PCG/Private/Elements/ControlFlow/PCGSwitch.cpp
@@ -0,0 +1,301 @@
+// Copyright Epic Games, Inc. All Rights Reserved.
+
+#include "Elements/ControlFlow/PCGSwitch.h"
+
+#include "PCGContext.h"
+#include "PCGModule.h"
+#include "Elements/PCGGather.h"
+
+#define LOCTEXT_NAMESPACE "FPCGSwitchElement"
+
+namespace PCGSwitchConstants
+{
+	const FText NodeTitleBase = LOCTEXT("NodeTitleBase", "Switch");
+}
+
+void UPCGSwitchSettings::PostLoad()
+{
+	Super::PostLoad();
+	CachePinLabels();
+}
+
+#if WITH_EDITOR
+void UPCGSwitchSettings::PostEditChangeProperty(FPropertyChangedEvent& PropertyChangedEvent)
+{
+	Super::PostEditChangeProperty(PropertyChangedEvent);
+
+	// Only need to change the pin labels if the options have changed
+	if (PropertyChangedEvent.GetMemberPropertyName() == GET_MEMBER_NAME_CHECKED(UPCGSwitchSettings, IntOptions) ||
+		PropertyChangedEvent.GetMemberPropertyName() == GET_MEMBER_NAME_CHECKED(UPCGSwitchSettings, StringOptions) ||
+		PropertyChangedEvent.GetMemberPropertyName() == GET_MEMBER_NAME_CHECKED(UPCGSwitchSettings, EnumSelection))
+	{
+		CachePinLabels();
+	}
+}
+
+FText UPCGSwitchSettings::GetDefaultNodeTitle() const
+{
+	return PCGSwitchConstants::NodeTitleBase;
+}
+
+FText UPCGSwitchSettings::GetNodeTooltipText() const
+{
+	return LOCTEXT("NodeTooltip", "Control flow node that passes through input data to a specific output pin that matches a given selection mode and corresponding 'selection' property - which can also be overridden.");
+}
+#endif // WITH_EDITOR
+
+bool UPCGSwitchSettings::IsPinStaticallyActive(const FName& PinLabel) const
+{
+	// Dynamic branches are never known in advance - assume all branches are active prior to execution.
+	if (IsSwitchDynamic())
+	{
+		return true;
+	}
+
+	FName ActiveOutputPinLabel;
+	if (!GetSelectedPinLabel(ActiveOutputPinLabel))
+	{
+		return false;
+	}
+
+	return PinLabel == ActiveOutputPinLabel;
+}
+
+FString UPCGSwitchSettings::GetAdditionalTitleInformation() const
+{
+	switch (SelectionMode)
+	{
+		case EPCGControlFlowSelectionMode::Integer:
+		{
+			FString Subtitle = PCGControlFlowConstants::SubtitleInt.ToString();
+			if (!IsPropertyOverriddenByPin(GET_MEMBER_NAME_CHECKED(UPCGSwitchSettings, IntegerSelection)))
+			{
+				Subtitle += FString::Format(TEXT(": {0}"), {IntegerSelection});
+			}
+
+			return Subtitle;
+		}
+
+		case EPCGControlFlowSelectionMode::Enum:
+		{
+			if (EnumSelection.Class)
+			{
+				FString Subtitle = PCGControlFlowConstants::SubtitleEnum.ToString();
+				if (!IsPropertyOverriddenByPin({GET_MEMBER_NAME_CHECKED(UPCGSwitchSettings, EnumSelection), GET_MEMBER_NAME_CHECKED(FEnumSelector, Value)}))
+				{
+					Subtitle += EnumSelection.Class->GetNameStringByValue(EnumSelection.Value);
+				}
+
+				return Subtitle;
+			}
+			break;
+		}
+
+		case EPCGControlFlowSelectionMode::String:
+		{
+			FString Subtitle = PCGControlFlowConstants::SubtitleString.ToString();
+			if (!IsPropertyOverriddenByPin(GET_MEMBER_NAME_CHECKED(UPCGSwitchSettings, StringSelection)))
+			{
+				Subtitle += FString::Format(TEXT(": {0}"), {StringSelection});
+			}
+
+			return Subtitle;
+		}
+
+		default:
+			checkNoEntry();
+			break;
+	}
+
+	return PCGSwitchConstants::NodeTitleBase.ToString();
+}
+
+#if WITH_EDITOR
+EPCGChangeType UPCGSwitchSettings::GetChangeTypeForProperty(const FName& InPropertyName) const
+{
+	EPCGChangeType ChangeType = Super::GetChangeTypeForProperty(InPropertyName) | EPCGChangeType::Cosmetic;
+
+	if (InPropertyName == GET_MEMBER_NAME_CHECKED(UPCGSwitchSettings, bEnabled) ||
+		InPropertyName == GET_MEMBER_NAME_CHECKED(UPCGSwitchSettings, IntegerSelection) ||
+		InPropertyName == GET_MEMBER_NAME_CHECKED(UPCGSwitchSettings, StringSelection) ||
+		InPropertyName == GET_MEMBER_NAME_CHECKED(UPCGSwitchSettings, EnumSelection) ||
+		InPropertyName == GET_MEMBER_NAME_CHECKED(UPCGSwitchSettings, SelectionMode))
+	{
+		ChangeType |= EPCGChangeType::Structural;
+	}
+
+	return ChangeType;
+}
+#endif // WITH_EDITOR
+
+TArray<FPCGPinProperties> UPCGSwitchSettings::InputPinProperties() const
+{
+	TArray<FPCGPinProperties> PinProperties;
+	PinProperties.Emplace(PCGPinConstants::DefaultInputLabel,
+		EPCGDataType::Any,
+		/*bInAllowMultipleConnections=*/true,
+		/*bAllowMultipleData=*/true,
+		LOCTEXT("OutputPinTooltip", "All input will be forwarded directly to the selected output pin."));
+
+	return PinProperties;
+}
+
+TArray<FPCGPinProperties> UPCGSwitchSettings::OutputPinProperties() const
+{
+	TArray<FPCGPinProperties> PinProperties;
+
+	switch (SelectionMode)
+	{
+		case EPCGControlFlowSelectionMode::Integer:
+			for (const int32 Value : IntOptions)
+			{
+				PinProperties.Emplace(FName(FString::FromInt(Value)));
+			}
+			break;
+
+		case EPCGControlFlowSelectionMode::String:
+			for (const FString& Value : StringOptions)
+			{
+				PinProperties.Emplace(FName(Value));
+			}
+			break;
+
+		case EPCGControlFlowSelectionMode::Enum:
+			// -1 to bypass the MAX value
+			for (int32 Index = 0; EnumSelection.Class && Index < EnumSelection.Class->NumEnums() - 1; ++Index)
+			{
+#if WITH_EDITOR
+				bool const bHidden = EnumSelection.Class->HasMetaData(TEXT("Hidden"), Index) || EnumSelection.Class->HasMetaData(TEXT("Spacer"), Index);
+				if (!bHidden)
+				{
+					PinProperties.Emplace(FName(EnumSelection.Class->GetDisplayNameTextByIndex(Index).ToString()));
+				}
+#else // WITH_EDITOR
+				// HasMetaData is editor only, so there will be extra pins at runtime, but that should be okay
+				PinProperties.Emplace(FName(EnumSelection.Class->GetDisplayNameTextByIndex(Index).ToString()));
+#endif // WITH_EDITOR
+			}
+			break;
+
+		default:
+			break;
+	}
+
+	PinProperties.Emplace(PCGControlFlowConstants::DefaultPathPinLabel);
+
+	return PinProperties;
+}
+
+FPCGElementPtr UPCGSwitchSettings::CreateElement() const
+{
+	return MakeShared<FPCGSwitchElement>();
+}
+
+bool UPCGSwitchSettings::IsSwitchDynamic() const
+{
+	bool bIsDynamic = false;
+
+	bIsDynamic |= (SelectionMode == EPCGControlFlowSelectionMode::Integer && IsPropertyOverriddenByPin(GET_MEMBER_NAME_CHECKED(UPCGSwitchSettings, IntegerSelection)));
+	bIsDynamic |= (SelectionMode == EPCGControlFlowSelectionMode::String && IsPropertyOverriddenByPin(GET_MEMBER_NAME_CHECKED(UPCGSwitchSettings, StringSelection)));
+	bIsDynamic |= (SelectionMode == EPCGControlFlowSelectionMode::Enum && IsPropertyOverriddenByPin({GET_MEMBER_NAME_CHECKED(UPCGSwitchSettings, EnumSelection), GET_MEMBER_NAME_CHECKED(FEnumSelector, Value)}));
+
+	return bIsDynamic;
+}
+
+bool UPCGSwitchSettings::IsValuePresent(const int32 Value) const
+{
+	return IntOptions.Contains(Value);
+}
+
+bool UPCGSwitchSettings::IsValuePresent(const FString& Value) const
+{
+	return StringOptions.Contains(Value);
+}
+
+bool UPCGSwitchSettings::IsValuePresent(const int64 Value) const
+{
+	if (!EnumSelection.Class)
+	{
+		return false;
+	}
+
+	const int64 Index = EnumSelection.Class->GetIndexByValue(Value);
+	return Index != INDEX_NONE && Index < EnumSelection.Class->NumEnums() - 1;
+}
+
+bool UPCGSwitchSettings::GetSelectedPinLabel(FName& OutSelectedPinLabel) const
+{
+	if (CachedPinLabels.IsEmpty())
+	{
+		return false;
+	}
+
+	int32 Index = INDEX_NONE;
+	if (SelectionMode == EPCGControlFlowSelectionMode::Integer && IsValuePresent(IntegerSelection))
+	{
+		Index = IntOptions.IndexOfByKey(IntegerSelection);
+	}
+	else if (SelectionMode == EPCGControlFlowSelectionMode::String && IsValuePresent(StringSelection))
+	{
+		Index = StringOptions.IndexOfByKey(StringSelection);
+	}
+	else if (SelectionMode == EPCGControlFlowSelectionMode::Enum && IsValuePresent(EnumSelection.Value))
+	{
+		// A "hidden" value could be selected that wasn't cached, so do a name-wise comparison
+		const FName PinLabel(EnumSelection.Class->GetDisplayNameTextByValue(EnumSelection.Value).ToString());
+		for (int i = 0; i < CachedPinLabels.Num(); ++i)
+		{
+			if (CachedPinLabels[i] == PinLabel)
+			{
+				Index = i;
+				break;
+			}
+		}
+	}
+	else
+	{
+		OutSelectedPinLabel = PCGControlFlowConstants::DefaultPathPinLabel;
+
+		return true;
+	}
+
+	if (Index < 0 || Index >= CachedPinLabels.Num())
+	{
+		return false;
+	}
+
+	OutSelectedPinLabel = CachedPinLabels[Index];
+
+	return true;
+}
+
+void UPCGSwitchSettings::CachePinLabels()
+{
+	CachedPinLabels.Empty();
+	Algo::Transform(OutputPinProperties(), CachedPinLabels, [](const FPCGPinProperties& Property)
+	{
+		return Property.Label;
+	});
+}
+
+bool FPCGSwitchElement::ExecuteInternal(FPCGContext* Context) const
+{
+	TRACE_CPUPROFILER_EVENT_SCOPE(FPCGSwitchElement::ExecuteInternal);
+
+	const UPCGSwitchSettings* Settings = Context->GetInputSettings<UPCGSwitchSettings>();
+	check(Settings);
+
+	FName SelectedPinLabel;
+	if (!Settings->GetSelectedPinLabel(SelectedPinLabel))
+	{
+		PCGE_LOG_C(Error, GraphAndLog, Context, LOCTEXT("ValueDoesNotExist", "Selected value is not a valid option."));
+		return true;
+	}
+
+	// Reuse the functionality of the Gather node
+	Context->OutputData = PCGGather::GatherDataForPin(Context->InputData, PCGPinConstants::DefaultInputLabel, SelectedPinLabel);
+
+	return true;
+}
+
+#undef LOCTEXT_NAMESPACE
\ No newline at end of file
diff --git a/Engine/Plugins/Experimental/PCG/Source/PCG/Private/Elements/PCGBooleanSelect.cpp b/Engine/Plugins/Experimental/PCG/Source/PCG/Private/Elements/PCGBooleanSelect.cpp
deleted file mode 100644
diff --git a/Engine/Plugins/Experimental/PCG/Source/PCG/Private/Elements/PCGBranch.cpp b/Engine/Plugins/Experimental/PCG/Source/PCG/Private/Elements/PCGBranch.cpp
deleted file mode 100644
diff --git a/Engine/Plugins/Experimental/PCG/Source/PCG/Private/Elements/PCGMultiSelect.cpp b/Engine/Plugins/Experimental/PCG/Source/PCG/Private/Elements/PCGMultiSelect.cpp
deleted file mode 100644
diff --git a/Engine/Plugins/Experimental/PCG/Source/PCG/Private/Tests/Elements/PCGBooleanSelectTest.cpp b/Engine/Plugins/Experimental/PCG/Source/PCG/Private/Tests/Elements/PCGBooleanSelectTest.cpp
--- a/Engine/Plugins/Experimental/PCG/Source/PCG/Private/Tests/Elements/PCGBooleanSelectTest.cpp
+++ b/Engine/Plugins/Experimental/PCG/Source/PCG/Private/Tests/Elements/PCGBooleanSelectTest.cpp
@@ -2,7 +2,7 @@
 
 #include "PCGContext.h"
 #include "Data/PCGPointData.h"
-#include "Elements/PCGBooleanSelect.h"
+#include "Elements/ControlFlow/PCGBooleanSelect.h"
 #include "Tests/PCGTestsCommon.h"
 
 #if WITH_EDITOR
@@ -27,15 +27,13 @@ protected:
 
 		UTEST_TRUE("Test element created and valid", TestElement.IsValid());
 
-		UPCGPointData* FirstPointData = PCGTestsCommon::CreateEmptyPointData();
 		FPCGTaggedData& FirstTaggedData = TestData.InputData.TaggedData.Emplace_GetRef();
-		FirstTaggedData.Data = FirstPointData;
+		FirstTaggedData.Data = PCGTestsCommon::CreateEmptyPointData();
 		FirstTaggedData.Tags.Emplace(FirstDataTag);
 		FirstTaggedData.Pin = FirstInputPinLabel;
 
-		UPCGPointData* SecondPointData = PCGTestsCommon::CreateEmptyPointData();
 		FPCGTaggedData& SecondTaggedData = TestData.InputData.TaggedData.Emplace_GetRef();
-		SecondTaggedData.Data = SecondPointData;
+		SecondTaggedData.Data = PCGTestsCommon::CreateEmptyPointData();
 		SecondTaggedData.Tags.Emplace(SecondDataTag);
 		SecondTaggedData.Pin = SecondInputPinLabel;
 
diff --git a/Engine/Plugins/Experimental/PCG/Source/PCG/Private/Tests/Elements/PCGMultiSelectTest.cpp b/Engine/Plugins/Experimental/PCG/Source/PCG/Private/Tests/Elements/PCGMultiSelectTest.cpp
--- a/Engine/Plugins/Experimental/PCG/Source/PCG/Private/Tests/Elements/PCGMultiSelectTest.cpp
+++ b/Engine/Plugins/Experimental/PCG/Source/PCG/Private/Tests/Elements/PCGMultiSelectTest.cpp
@@ -2,7 +2,7 @@
 
 #include "PCGContext.h"
 #include "Data/PCGPointData.h"
-#include "Elements/PCGMultiSelect.h"
+#include "Elements/ControlFlow/PCGMultiSelect.h"
 #include "Tests/PCGTestsCommon.h"
 
 #if WITH_EDITOR
@@ -47,18 +47,16 @@ protected:
 		{
 			if (!SelectElementPinProperties[i].bAdvancedPin && SelectElementPinProperties[i].Label != DefaultPathPinLabel)
 			{
-				UPCGPointData* PointData = PCGTestsCommon::CreateEmptyPointData();
 				FPCGTaggedData& TaggedData = TestData.InputData.TaggedData.Emplace_GetRef();
-				TaggedData.Data = PointData;
+				TaggedData.Data = PCGTestsCommon::CreateEmptyPointData();
 				TaggedData.Tags.Emplace(CreateDataTag(i));
 				TaggedData.Pin = SelectElementPinProperties[i].Label;
 			}
 		}
 
 		// Add an input to the default path
-		UPCGPointData* PointData = PCGTestsCommon::CreateEmptyPointData();
 		FPCGTaggedData& TaggedData = TestData.InputData.TaggedData.Emplace_GetRef();
-		TaggedData.Data = PointData;
+		TaggedData.Data = PCGTestsCommon::CreateEmptyPointData();
 		TaggedData.Tags.Emplace(DefaultPathPinDataTag);
 		TaggedData.Pin = DefaultPathPinLabel;
 
diff --git a/Engine/Plugins/Experimental/PCG/Source/PCG/Public/Elements/ControlFlow/PCGBooleanSelect.h b/Engine/Plugins/Experimental/PCG/Source/PCG/Public/Elements/ControlFlow/PCGBooleanSelect.h
new file mode 100644
--- /dev/null
+++ b/Engine/Plugins/Experimental/PCG/Source/PCG/Public/Elements/ControlFlow/PCGBooleanSelect.h
@@ -0,0 +1,43 @@
+// Copyright Epic Games, Inc. All Rights Reserved.
+#pragma once
+
+#include "PCGSettings.h"
+
+#include "PCGBooleanSelect.generated.h"
+
+/**
+ * Selects data from either input pin, based on a boolean condition.
+ */
+UCLASS(BlueprintType, ClassGroup = (Procedural), meta=(Keywords = "if bool select"))
+class UPCGBooleanSelectSettings : public UPCGSettings
+{
+	GENERATED_BODY()
+
+public:
+	//~Begin UPCGSettings interface
+#if WITH_EDITOR
+	virtual FName GetDefaultNodeName() const override { return FName(TEXT("Select")); }
+	virtual FText GetDefaultNodeTitle() const override;
+	virtual FText GetNodeTooltipText() const override;
+	virtual EPCGSettingsType GetType() const override { return EPCGSettingsType::ControlFlow; }
+	virtual bool HasDynamicPins() const override { return true; }
+#endif
+
+	virtual EPCGDataType GetCurrentPinTypes(const UPCGPin* Pin) const override;
+
+protected:
+	virtual TArray<FPCGPinProperties> InputPinProperties() const override;
+	virtual TArray<FPCGPinProperties> OutputPinProperties() const override;
+	virtual FPCGElementPtr CreateElement() const override;
+	//~End UPCGSettings interface
+
+public:
+	UPROPERTY(BlueprintReadWrite, EditAnywhere, Category=Settings, meta=(PCG_Overridable))
+	bool bUseInputB = false;
+};
+
+class FPCGBooleanSelectElement : public IPCGElement
+{
+protected:
+	virtual bool ExecuteInternal(FPCGContext* Context) const override;
+};
diff --git a/Engine/Plugins/Experimental/PCG/Source/PCG/Public/Elements/ControlFlow/PCGBranch.h b/Engine/Plugins/Experimental/PCG/Source/PCG/Public/Elements/ControlFlow/PCGBranch.h
new file mode 100644
--- /dev/null
+++ b/Engine/Plugins/Experimental/PCG/Source/PCG/Public/Elements/ControlFlow/PCGBranch.h
@@ -0,0 +1,45 @@
+// Copyright Epic Games, Inc. All Rights Reserved.
+#pragma once
+
+#include "PCGSettings.h"
+
+#include "PCGBranch.generated.h"
+
+/**
+ * Routes input data to one of two outputs, based on a boolean condition.
+ */
+UCLASS(BlueprintType, ClassGroup = (Procedural), meta=(Keywords = "if bool switch"))
+class UPCGBranchSettings : public UPCGSettings
+{
+	GENERATED_BODY()
+
+public:
+	//~Begin UPCGSettings interface
+#if WITH_EDITOR
+	virtual FName GetDefaultNodeName() const override { return FName(TEXT("Branch")); }
+	virtual FText GetDefaultNodeTitle() const override;
+	virtual FText GetNodeTooltipText() const override;
+	virtual EPCGSettingsType GetType() const override { return EPCGSettingsType::ControlFlow; }
+	virtual bool HasDynamicPins() const override { return true; }
+#endif
+
+	virtual bool IsPinStaticallyActive(const FName& PinLabel) const override;
+
+protected:
+#if WITH_EDITOR
+	virtual EPCGChangeType GetChangeTypeForProperty(const FName& InPropertyName) const override;
+#endif
+	virtual TArray<FPCGPinProperties> OutputPinProperties() const override;
+	virtual FPCGElementPtr CreateElement() const override;
+	//~End UPCGSettings interface
+
+public:
+	UPROPERTY(BlueprintReadWrite, EditAnywhere, Category=Settings, meta=(PCG_Overridable))
+	bool bOutputToB = false;
+};
+
+class FPCGBranchElement : public IPCGElement
+{
+protected:
+	virtual bool ExecuteInternal(FPCGContext* Context) const override;
+};
diff --git a/Engine/Plugins/Experimental/PCG/Source/PCG/Public/Elements/ControlFlow/PCGControlFlow.h b/Engine/Plugins/Experimental/PCG/Source/PCG/Public/Elements/ControlFlow/PCGControlFlow.h
new file mode 100644
--- /dev/null
+++ b/Engine/Plugins/Experimental/PCG/Source/PCG/Public/Elements/ControlFlow/PCGControlFlow.h
@@ -0,0 +1,33 @@
+// Copyright Epic Games, Inc. All Rights Reserved.
+
+#pragma once
+
+#include "PCGControlFlow.generated.h"
+
+USTRUCT(BlueprintType, meta = (Hidden))
+struct PCG_API FEnumSelector
+{
+	GENERATED_BODY()
+
+	UPROPERTY(DisplayName="Enum Class", meta=(PCG_NotOverridable))
+	UEnum* Class = nullptr;
+
+	UPROPERTY(DisplayName="Enum Value")
+	int64 Value = 0;
+};
+
+UENUM()
+enum class EPCGControlFlowSelectionMode : uint8
+{
+	Integer,
+	Enum,
+	String
+};
+
+namespace PCGControlFlowConstants
+{
+	inline const FText SubtitleInt = NSLOCTEXT("FPCGControlFlow", "SubtitleInt", "Integer Selection");
+	inline const FText SubtitleEnum = NSLOCTEXT("FPCGControlFlow", "SubtitleEnum", "Enum Selection");
+	inline const FText SubtitleString = NSLOCTEXT("FPCGControlFlow", "SubtitleString", "String Selection");
+	inline const FName DefaultPathPinLabel = TEXT("Default");
+}
diff --git a/Engine/Plugins/Experimental/PCG/Source/PCG/Public/Elements/ControlFlow/PCGMultiSelect.h b/Engine/Plugins/Experimental/PCG/Source/PCG/Public/Elements/ControlFlow/PCGMultiSelect.h
new file mode 100644
--- /dev/null
+++ b/Engine/Plugins/Experimental/PCG/Source/PCG/Public/Elements/ControlFlow/PCGMultiSelect.h
@@ -0,0 +1,92 @@
+// Copyright Epic Games, Inc. All Rights Reserved.
+
+#pragma once
+
+#include "PCGSettings.h"
+#include "PCGControlFlow.h"
+
+#include "PCGMultiSelect.generated.h"
+
+/**
+ * Selects data from any number of input pins, based on a static selection criteria (Int/String/Enum) 
+ */
+UCLASS(BlueprintType, ClassGroup = (Procedural), meta=(Keywords = "if multi switch enum"))
+class UPCGMultiSelectSettings : public UPCGSettings
+{
+	GENERATED_BODY()
+
+public:
+	//~Begin UObject interface
+	virtual void PostLoad() override;
+
+#if WITH_EDITOR
+	virtual void PostEditChangeProperty(FPropertyChangedEvent& PropertyChangedEvent) override;
+	//~End UObject interface
+
+	//~Begin UPCGSettings interface
+	virtual FName GetDefaultNodeName() const override { return FName(TEXT("Select (Multi)")); }
+	virtual FText GetDefaultNodeTitle() const override;
+	virtual FText GetNodeTooltipText() const override;
+	virtual EPCGSettingsType GetType() const override { return EPCGSettingsType::ControlFlow; }
+	virtual bool HasDynamicPins() const override { return true; }
+#endif // WITH_EDITOR
+
+	virtual EPCGDataType GetCurrentPinTypes(const UPCGPin* Pin) const override;
+	virtual FString GetAdditionalTitleInformation() const override;
+	virtual bool HasFlippedTitleLines() const override { return true; }
+
+protected:
+	virtual TArray<FPCGPinProperties> InputPinProperties() const override;
+	virtual TArray<FPCGPinProperties> OutputPinProperties() const override;
+	virtual FPCGElementPtr CreateElement() const override;
+	//~End UPCGSettings interface
+
+public:
+	/** Determines the type of value to be used to select an input. */
+	UPROPERTY(BlueprintReadWrite, EditAnywhere, Category=Settings)
+	EPCGControlFlowSelectionMode SelectionMode = EPCGControlFlowSelectionMode::Integer;
+
+	/** Determines which input will be selected if the selection mode is Integer. */
+	UPROPERTY(BlueprintReadWrite, EditAnywhere, Category=Settings, meta=(PCG_Overridable, EditConditionHides, EditCondition="SelectionMode == EPCGControlFlowSelectionMode::Integer"))
+	int32 IntegerSelection = 0;
+
+	/** Determines the available input pin selection options. */
+	UPROPERTY(BlueprintReadWrite, EditAnywhere, Category=Settings, meta=(EditConditionHides, EditCondition="SelectionMode == EPCGControlFlowSelectionMode::Integer"))
+	TArray<int32> IntOptions = {0};
+
+	/** Determines which input will be selected if the selection mode is String. */
+	UPROPERTY(BlueprintReadWrite, EditAnywhere, Category=Settings, meta=(PCG_Overridable, EditConditionHides, NoResetToDefault, EditCondition="SelectionMode == EPCGControlFlowSelectionMode::String"))
+	FString StringSelection;
+
+	/** Determines the available input pin selection options. */
+	UPROPERTY(BlueprintReadWrite, EditAnywhere, Category=Settings, meta=(EditConditionHides, EditCondition="SelectionMode == EPCGControlFlowSelectionMode::String"))
+	TArray<FString> StringOptions;
+
+	/** Determines which input pin will be selected if the selection mode is Enum. */
+	UPROPERTY(BlueprintReadWrite, EditAnywhere, Category=Settings, meta=(PCG_Overridable, EditConditionHides, EditCondition="SelectionMode == EPCGControlFlowSelectionMode::Enum"))
+	FEnumSelector EnumSelection;
+
+	/** Cached pin labels for use during the selection process. */
+	UPROPERTY(Transient)
+	TArray<FName> CachedPinLabels;
+
+	/** Returns true if the integer value exists in the user defined array. */
+	bool IsValuePresent(int32 Value) const;
+
+	/** Returns true if the string value exists in the user defined array. */
+	bool IsValuePresent(const FString& Value) const;
+
+	/** Returns true if the enum value exists within the enum class. */
+	bool IsValuePresent(int64 Value) const;
+
+	/** Helper function to use the appropriate selection value to determine the current selection. Returns true if it succeeds or false if the value is invalid. */
+	bool GetSelectedPinLabel(FName& OutSelectedPinLabel) const;
+
+	void CachePinLabels();
+};
+
+class FPCGMultiSelectElement : public IPCGElement
+{
+protected:
+	virtual bool ExecuteInternal(FPCGContext* Context) const override;
+};
diff --git a/Engine/Plugins/Experimental/PCG/Source/PCG/Public/Elements/ControlFlow/PCGSwitch.h b/Engine/Plugins/Experimental/PCG/Source/PCG/Public/Elements/ControlFlow/PCGSwitch.h
new file mode 100644
--- /dev/null
+++ b/Engine/Plugins/Experimental/PCG/Source/PCG/Public/Elements/ControlFlow/PCGSwitch.h
@@ -0,0 +1,97 @@
+// Copyright Epic Games, Inc. All Rights Reserved.
+
+#pragma once
+
+#include "PCGControlFlow.h"
+#include "PCGSettings.h"
+
+#include "PCGSwitch.generated.h"
+
+/**
+ * Routes data from the input pin, to a specific output pin based on a selection criteria (Int/String/Enum) 
+ */
+UCLASS(BlueprintType, ClassGroup = (Procedural), meta=(Keywords = "if multi enum branch"))
+class UPCGSwitchSettings : public UPCGSettings
+{
+	GENERATED_BODY()
+
+public:
+	//~Begin UObject interface
+	virtual void PostLoad() override;
+
+#if WITH_EDITOR
+	virtual void PostEditChangeProperty(FPropertyChangedEvent& PropertyChangedEvent) override;
+	//~End UObject interface
+
+	//~Begin UPCGSettings interface
+	virtual FName GetDefaultNodeName() const override { return FName(TEXT("Switch")); }
+	virtual FText GetDefaultNodeTitle() const override;
+	virtual FText GetNodeTooltipText() const override;
+	virtual EPCGSettingsType GetType() const override { return EPCGSettingsType::ControlFlow; }
+	virtual bool HasDynamicPins() const override { return true; }
+#endif // WITH_EDITOR
+
+	virtual bool IsPinStaticallyActive(const FName& PinLabel) const override;
+	virtual FString GetAdditionalTitleInformation() const override;
+	virtual bool HasFlippedTitleLines() const override { return true; }
+
+protected:
+#if WITH_EDITOR
+	virtual EPCGChangeType GetChangeTypeForProperty(const FName& InPropertyName) const override;
+#endif
+	virtual TArray<FPCGPinProperties> InputPinProperties() const override;
+	virtual TArray<FPCGPinProperties> OutputPinProperties() const override;
+	virtual FPCGElementPtr CreateElement() const override;
+	//~End UPCGSettings interface
+
+	/** Switch is dynamic if selection value is overridden / cannot be determined prior to execution. */
+	virtual bool IsSwitchDynamic() const;
+
+public:
+	/** Determines the type of value to be used to select an output. */
+	UPROPERTY(BlueprintReadWrite, EditAnywhere, Category=Settings)
+	EPCGControlFlowSelectionMode SelectionMode = EPCGControlFlowSelectionMode::Integer;
+
+	/** Determines which output will be selected if the selection mode is Integer. */
+	UPROPERTY(BlueprintReadWrite, EditAnywhere, Category=Settings, meta=(PCG_Overridable, EditConditionHides, EditCondition="SelectionMode == EPCGControlFlowSelectionMode::Integer"))
+	int32 IntegerSelection = 0;
+
+	/** Determines the available output pin selection options. */
+	UPROPERTY(BlueprintReadWrite, EditAnywhere, Category=Settings, meta=(EditConditionHides, EditCondition="SelectionMode == EPCGControlFlowSelectionMode::Integer"))
+	TArray<int32> IntOptions = {0};
+
+	/** Determines which output will be selected if the selection mode is String. */
+	UPROPERTY(BlueprintReadWrite, EditAnywhere, Category=Settings, meta=(PCG_Overridable, EditConditionHides, NoResetToDefault, EditCondition="SelectionMode == EPCGControlFlowSelectionMode::String"))
+	FString StringSelection;
+
+	/** Determines the available output pin selection options. */
+	UPROPERTY(BlueprintReadWrite, EditAnywhere, Category=Settings, meta=(EditConditionHides, EditCondition="SelectionMode == EPCGControlFlowSelectionMode::String"))
+	TArray<FString> StringOptions;
+
+	/** Determines which output pin will be selected if the selection mode is Enum. */
+	UPROPERTY(BlueprintReadWrite, EditAnywhere, Category=Settings, meta=(PCG_Overridable, EditConditionHides, EditCondition="SelectionMode == EPCGControlFlowSelectionMode::Enum"))
+	FEnumSelector EnumSelection;
+
+	/** Cached pin labels for use during the selection process. */
+	TArray<FName> CachedPinLabels;
+
+	/** Returns true if the integer value exists in the user defined options. */
+	bool IsValuePresent(int32 Value) const;
+
+	/** Returns true if the string value exists in the user defined options. */
+	bool IsValuePresent(const FString& Value) const;
+
+	/** Returns true if the enum value exists within the selected enum class. */
+	bool IsValuePresent(int64 Value) const;
+
+	/** Helper function to use the appropriate selection value to determine the current selection. Returns true if it succeeds or false if the value is invalid. */
+	bool GetSelectedPinLabel(FName& OutSelectedPinLabel) const;
+
+	void CachePinLabels();
+};
+
+class FPCGSwitchElement : public IPCGElement
+{
+protected:
+	virtual bool ExecuteInternal(FPCGContext* Context) const override;
+};
diff --git a/Engine/Plugins/Experimental/PCG/Source/PCG/Public/Elements/PCGBooleanSelect.h b/Engine/Plugins/Experimental/PCG/Source/PCG/Public/Elements/PCGBooleanSelect.h
deleted file mode 100644
diff --git a/Engine/Plugins/Experimental/PCG/Source/PCG/Public/Elements/PCGBranch.h b/Engine/Plugins/Experimental/PCG/Source/PCG/Public/Elements/PCGBranch.h
deleted file mode 100644
diff --git a/Engine/Plugins/Experimental/PCG/Source/PCG/Public/Elements/PCGMultiSelect.h b/Engine/Plugins/Experimental/PCG/Source/PCG/Public/Elements/PCGMultiSelect.h
deleted file mode 100644
diff --git a/Engine/Plugins/Experimental/PCG/Source/PCG/Public/Helpers/PCGPropertyHelpers.h b/Engine/Plugins/Experimental/PCG/Source/PCG/Public/Helpers/PCGPropertyHelpers.h
--- a/Engine/Plugins/Experimental/PCG/Source/PCG/Public/Helpers/PCGPropertyHelpers.h
+++ b/Engine/Plugins/Experimental/PCG/Source/PCG/Public/Helpers/PCGPropertyHelpers.h
@@ -7,30 +7,15 @@
 #include "Metadata/Accessors/IPCGAttributeAccessor.h"
 #include "Metadata/Accessors/PCGAttributeAccessorHelpers.h"
 #include "Metadata/Accessors/PCGAttributeAccessorKeys.h"
-#include "UObject/UnrealType.h"
 
 #include "UObject/UnrealType.h" // IWYU pragma: keep
 
-#include "PCGPropertyHelpers.generated.h"
-
 enum class EPCGMetadataTypes : uint8;
 
 struct FPCGContext;
 class UPCGData;
 class UPCGParamData;
 
-USTRUCT(BlueprintType, meta = (Hidden))
-struct PCG_API FEnumSelector
-{
-	GENERATED_BODY()
-
-	UPROPERTY(DisplayName="Enum Class", meta=(PCG_NotOverridable))
-	UEnum* Class = nullptr;
-
-	UPROPERTY(DisplayName="Enum Value")
-	int64 Value = 0;
-};
-
 namespace PCGPropertyHelpers
 {
 	/**
diff --git a/Engine/Plugins/Experimental/PCG/Source/PCGEditor/Private/Details/EnumSelectorDetails.cpp b/Engine/Plugins/Experimental/PCG/Source/PCGEditor/Private/Details/EnumSelectorDetails.cpp
--- a/Engine/Plugins/Experimental/PCG/Source/PCGEditor/Private/Details/EnumSelectorDetails.cpp
+++ b/Engine/Plugins/Experimental/PCG/Source/PCGEditor/Private/Details/EnumSelectorDetails.cpp
@@ -2,7 +2,7 @@
 
 #include "Details/EnumSelectorDetails.h"
 
-#include "Helpers/PCGPropertyHelpers.h"
+#include "Elements/ControlFlow/PCGControlFlow.h"
 
 #include "DetailLayoutBuilder.h"
 #include "DetailWidgetRow.h"
