Removed the per-node tag filtering and tagging as this was marked deprecated already in previous versions.
Removed the 'remove duplicates' option on the debug settings since it served no real purpose.
Moved around the determinism test info into the debug settings to cleanup the details view.

diff --git a/Engine/Plugins/Experimental/PCG/Source/PCG/Private/PCGElement.cpp b/Engine/Plugins/Experimental/PCG/Source/PCG/Private/PCGElement.cpp
--- a/Engine/Plugins/Experimental/PCG/Source/PCG/Private/PCGElement.cpp
+++ b/Engine/Plugins/Experimental/PCG/Source/PCG/Private/PCGElement.cpp
@@ -8,7 +8,6 @@
 #include "PCGSubsystem.h"
 #include "Data/PCGPointData.h"
 #include "Elements/PCGDebugElement.h"
-#include "Elements/PCGSelfPruning.h"
 
 #include "HAL/IConsoleManager.h"
 #include "Utils/PCGExtraCapture.h"
@@ -145,37 +144,8 @@ void IPCGElement::PreExecute(FPCGContext* Context) const
 	{
 		//Pass-through - no execution
 		DisabledPassThroughData(Context);
-
 		Context->CurrentPhase = EPCGExecutionPhase::PostExecute;
 	}
-	else
-	{
-		// Perform input filtering
-		/** TODO - Placeholder feature */
-		if (!Settings->FilterOnTags.IsEmpty())
-		{
-			// Move any of the inputs that don't have the tags to the outputs as a pass-through
-			// NOTE: this breaks a bit the ordering of inputs, however, there's no obvious way around it
-			TArray<FPCGTaggedData> FilteredTaggedData;
-			for (FPCGTaggedData& TaggedData : Context->InputData.TaggedData)
-			{
-				if (TaggedData.Tags.Intersect(Settings->FilterOnTags).IsEmpty())
-				{
-					if (Settings->bPassThroughFilteredOutInputs)
-					{
-						Context->OutputData.TaggedData.Add(TaggedData);
-					}
-				}
-				else // input has the required tags
-				{
-					FilteredTaggedData.Add(TaggedData);
-				}
-			}
-
-			Context->InputData.TaggedData = MoveTemp(FilteredTaggedData);
-			Context->BypassedOutputCount = Context->OutputData.TaggedData.Num();
-		}
-	}
 }
 
 bool IPCGElement::PrepareDataInternal(FPCGContext* Context) const
@@ -190,16 +160,6 @@ void IPCGElement::PostExecute(FPCGContext* Context) const
 
 	const UPCGSettingsInterface* SettingsInterface = Context->GetInputSettingsInterface();
 	const UPCGSettings* Settings = SettingsInterface ? SettingsInterface->GetSettings() : nullptr;
-
-	// Apply tags on output
-	/** TODO - Placeholder feature */
-	if (Settings && !Settings->TagsAppliedOnOutput.IsEmpty())
-	{
-		for (int32 TaggedDataIdx = Context->BypassedOutputCount; TaggedDataIdx < Context->OutputData.TaggedData.Num(); ++TaggedDataIdx)
-		{
-			Context->OutputData.TaggedData[TaggedDataIdx].Tags.Append(Settings->TagsAppliedOnOutput);
-		}
-	}
 	
 	// Output data Crc
 	{
@@ -214,24 +174,6 @@ void IPCGElement::PostExecute(FPCGContext* Context) const
 		Context->OutputData.Crc = Context->OutputData.ComputeCrc(bShouldComputeFullOutputDataCrc);
 	}
 
-	// Additional debug things (check for duplicates),
-#if WITH_EDITOR
-	if (SettingsInterface && SettingsInterface->DebugSettings.bCheckForDuplicates)
-	{
-		FPCGDataCollection ElementInputs = Context->InputData;
-		FPCGDataCollection ElementOutputs = Context->OutputData;
-
-		Context->InputData = ElementOutputs;
-		Context->OutputData = FPCGDataCollection();
-
-		PCGE_LOG(Verbose, LogOnly, LOCTEXT("PerformingDuplicatePointTest", "Performing remove duplicate points test (perf warning)"));
-		PCGSelfPruningElement::Execute(Context, EPCGSelfPruningType::RemoveDuplicates, 0.0f, false);
-
-		Context->InputData = ElementInputs;
-		Context->OutputData = ElementOutputs;
-	}
-#endif
-
 	Context->CurrentPhase = EPCGExecutionPhase::Done;
 }
 
diff --git a/Engine/Plugins/Experimental/PCG/Source/PCG/Public/PCGDebug.h b/Engine/Plugins/Experimental/PCG/Source/PCG/Public/PCGDebug.h
--- a/Engine/Plugins/Experimental/PCG/Source/PCG/Public/PCGDebug.h
+++ b/Engine/Plugins/Experimental/PCG/Source/PCG/Public/PCGDebug.h
@@ -36,10 +36,6 @@ public:
 	UPROPERTY(BlueprintReadWrite, EditAnywhere, Category = Debug, meta = (EditCondition = "bDisplayProperties", EditConditionHides, HideEditConditionToggle))
 	TSoftObjectPtr<UMaterialInterface> MaterialOverride;
 
-	/** Warning: enabling this flag will have severe performance impact */
-	UPROPERTY(BlueprintReadWrite, EditAnywhere, Category = Debug, meta = (EditCondition = "bDisplayProperties", EditConditionHides, HideEditConditionToggle))
-	bool bCheckForDuplicates = false;
-
 	TSoftObjectPtr<UMaterialInterface> GetMaterial() const;
 
 #if WITH_EDITORONLY_DATA
diff --git a/Engine/Plugins/Experimental/PCG/Source/PCG/Public/PCGSettings.h b/Engine/Plugins/Experimental/PCG/Source/PCG/Public/PCGSettings.h
--- a/Engine/Plugins/Experimental/PCG/Source/PCG/Public/PCGSettings.h
+++ b/Engine/Plugins/Experimental/PCG/Source/PCG/Public/PCGSettings.h
@@ -176,7 +176,7 @@ public:
 	FPCGDebugVisualizationSettings DebugSettings;
 
 	/** If a debugger is attached, triggers a breakpoint inside IPCGElement::Execute(). Editor only. Transient. */
-	UPROPERTY(Transient, DuplicateTransient, EditAnywhere, BlueprintReadWrite, Category = Debug, AdvancedDisplay, meta = (EditCondition = bDisplayDebuggingProperties, EditConditionHides, HideEditConditionToggle))
+	UPROPERTY(Transient, DuplicateTransient, EditAnywhere, BlueprintReadWrite, Category = Debug, AdvancedDisplay, meta = (DisplayName = "Break In Debugger", EditCondition = bDisplayDebuggingProperties, EditConditionHides, HideEditConditionToggle))
 	bool bBreakDebugger = false;
 
 	// This can be set false by inheriting nodes to hide the debugging properties.
@@ -334,23 +334,11 @@ public:
 	UPROPERTY(EditAnywhere, BlueprintReadWrite, Category = Settings, meta=(EditCondition=bUseSeed, EditConditionHides, PCG_Overridable))
 	int Seed = 0xC35A9631; // Default seed is a random prime number, but will be overriden for new settings based on the class type name hash, making each settings class have a different default seed.
 
-	/** Warning - this is deprecated and will be removed soon since we have a Filter By Tag node for this specific purpose */
-	UPROPERTY(EditAnywhere, BlueprintReadWrite, Category = "Tags")
-	TSet<FString> FilterOnTags;
-
-	/** Warning - this is deprecated and will be removed soon since we have a Filter By Tag node for this specific purpose */
-	UPROPERTY(EditAnywhere, BlueprintReadWrite, Category = "Tags")
-	bool bPassThroughFilteredOutInputs = true;
-
-	/** Applies the specified tags on the output data. Note - this might be replaced by a dedicated Tagging node */
-	UPROPERTY(EditAnywhere, BlueprintReadWrite, Category = "Tags")
-	TSet<FString> TagsAppliedOnOutput;
-
 #if WITH_EDITORONLY_DATA
 	UPROPERTY()
 	EPCGSettingsExecutionMode ExecutionMode_DEPRECATED = EPCGSettingsExecutionMode::Enabled;
 
-	UPROPERTY(EditAnywhere, BlueprintReadWrite, Category = Determinism, meta = (ShowOnlyInnerProperties))
+	UPROPERTY(EditAnywhere, BlueprintReadWrite, Category = Debug, AdvancedDisplay, meta = (DisplayName="Determinism", NoResetToDefault))
 	FPCGDeterminismSettings DeterminismSettings;
 #endif
 
