[PCG] Add accessors for FColor and FLinearColor to convert them to vector 4

#rnx

diff --git a/Engine/Plugins/Experimental/PCG/Source/PCG/Private/Metadata/Accessors/PCGAttributeAccessorHelpers.cpp b/Engine/Plugins/Experimental/PCG/Source/PCG/Private/Metadata/Accessors/PCGAttributeAccessorHelpers.cpp
--- a/Engine/Plugins/Experimental/PCG/Source/PCG/Private/Metadata/Accessors/PCGAttributeAccessorHelpers.cpp
+++ b/Engine/Plugins/Experimental/PCG/Source/PCG/Private/Metadata/Accessors/PCGAttributeAccessorHelpers.cpp
@@ -336,6 +336,14 @@ namespace PCGAttributeAccessorHelpers
 			{
 				return Functor(Signature<FPCGPropertyStructAccessor<FSoftClassPath>>{}, StructProperty);
 			}
+			else if (StructProperty->Struct == TBaseStructure<FLinearColor>::Get())
+			{
+				return Functor(Signature<FPCGLinearColorAccessor>{}, StructProperty);
+			}
+			else if (StructProperty->Struct == TBaseStructure<FColor>::Get())
+			{
+				return Functor(Signature<FPCGColorAccessor>{}, StructProperty);
+			}
 			else if (StructProperty->Struct && StructProperty->Struct->IsChildOf(FPCGAttributePropertySelector::StaticStruct()))
 			{
 				return Functor(Signature<FPCGAttributePropertySelectorAccessor>{}, StructProperty);
diff --git a/Engine/Plugins/Experimental/PCG/Source/PCG/Private/Metadata/Accessors/PCGPropertyAccessor.cpp b/Engine/Plugins/Experimental/PCG/Source/PCG/Private/Metadata/Accessors/PCGPropertyAccessor.cpp
--- a/Engine/Plugins/Experimental/PCG/Source/PCG/Private/Metadata/Accessors/PCGPropertyAccessor.cpp
+++ b/Engine/Plugins/Experimental/PCG/Source/PCG/Private/Metadata/Accessors/PCGPropertyAccessor.cpp
@@ -30,7 +30,7 @@ TArray<void*> PCGPropertyAccessor::GetContainerKeys(int32 Index, int32 Range, IP
 	return ContainerKeys;
 }
 
-IPCGPropertyChainAccessor::IPCGPropertyChainAccessor(const FProperty* Property, TArray<const FProperty*>&& ExtraProperties)
+IPCGPropertyChain::IPCGPropertyChain(const FProperty* Property, TArray<const FProperty*>&& ExtraProperties)
 	: PropertyChain(std::forward<TArray<const FProperty*>>(ExtraProperties))
 {
 	// Fix property chain
diff --git a/Engine/Plugins/Experimental/PCG/Source/PCG/Private/Tests/Elements/PCGGetActorPropertyTest.cpp b/Engine/Plugins/Experimental/PCG/Source/PCG/Private/Tests/Elements/PCGGetActorPropertyTest.cpp
--- a/Engine/Plugins/Experimental/PCG/Source/PCG/Private/Tests/Elements/PCGGetActorPropertyTest.cpp
+++ b/Engine/Plugins/Experimental/PCG/Source/PCG/Private/Tests/Elements/PCGGetActorPropertyTest.cpp
@@ -215,8 +215,10 @@ bool FPCGPropertyToParamDataPropertyTypeTest::RunTest(const FString& Parameters)
 	const FSoftClassPath SoftClassPathValue{ UPCGDummyGetPropertyTest::StaticClass() };
 
 	const FVector2D Vector2Value = { 1.0, 2.0 };
-	const FColor ColorValue = FColor::Cyan;
-	const FColor SecondColorValue = FColor::Yellow;
+	const FPCGTestMyColorStruct PCGColorValue{ 1.0, 1.0, 0.0, 1.0 };
+	const FPCGTestMyColorStruct SecondPCGColorValue{ 1.0, 0.0, 1.0, 1.0 };
+	const FColor ColorValue = FColor::White;
+	const FLinearColor LinearColorValue = FLinearColor::Blue;
 
 	APCGUnitTestDummyActor* Actor = Cast<APCGUnitTestDummyActor>(TestData.TestActor);
 	Actor->IntProperty = 42;
@@ -238,9 +240,11 @@ bool FPCGPropertyToParamDataPropertyTypeTest::RunTest(const FString& Parameters)
 	Actor->ObjectProperty = ObjectValue;
 	Actor->Vector2Property = Vector2Value;
 	Actor->ColorProperty = ColorValue;
+	Actor->LinearColorProperty = LinearColorValue;
+	Actor->PCGColorProperty = PCGColorValue;
 	Actor->ArrayOfIntsProperty = { 42, 43, 44 };
 	Actor->ArrayOfVectorsProperty = { VectorValue, SecondVectorValue };
-	Actor->ArrayOfStructsProperty = { ColorValue, SecondColorValue };
+	Actor->ArrayOfStructsProperty = { PCGColorValue, SecondPCGColorValue };
 	Actor->ArrayOfObjectsProperty = { ObjectValue, SecondObjectValue };
 	Actor->DummyStruct.FloatProperty = 1.2f;
 	Actor->DummyStruct.IntArrayProperty = { 5, 6, 7 };
@@ -274,10 +278,14 @@ bool FPCGPropertyToParamDataPropertyTypeTest::RunTest(const FString& Parameters)
 	bSuccess &= VerifyAttributeValueValid(this, TestData, GET_MEMBER_NAME_CHECKED(APCGUnitTestDummyActor, ClassProperty), FSoftClassPath(UPCGDummyGetPropertyTest::StaticClass()), ExtraTestWhat);
 	bSuccess &= VerifyAttributeValueValid(this, TestData, GET_MEMBER_NAME_CHECKED(APCGUnitTestDummyActor, ObjectProperty), FSoftObjectPath(ObjectValue->GetPathName()), ExtraTestWhat);
 
-	// Struct Property Extracted - Colors
+	// Colors as Vector 4
+	bSuccess &= VerifyAttributeValueValid(this, TestData, GET_MEMBER_NAME_CHECKED(APCGUnitTestDummyActor, ColorProperty), FVector4(1.0, 1.0, 1.0, 1.0), ExtraTestWhat);
+	bSuccess &= VerifyAttributeValueValid(this, TestData, GET_MEMBER_NAME_CHECKED(APCGUnitTestDummyActor, LinearColorProperty), FVector4(0.0, 0.0, 1.0, 1.0), ExtraTestWhat);
+
+	// Struct Property Extracted - Dummy Color struct
 	// Extracting int will always yield a int64 and extracting floats with yield doubles. Here color is u8, so cast all of them to int64
 	const TArray<FName> ColorPropertyNames = { TEXT("R"), TEXT("G"), TEXT("B"), TEXT("A") };
-	bSuccess &= VerifyAttributeValuesValid(this, TestData, GET_MEMBER_NAME_CHECKED(APCGUnitTestDummyActor, ColorProperty), ColorPropertyNames, ExtraTestWhat, (int64)ColorValue.R, (int64)ColorValue.G, (int64)ColorValue.B, (int64)ColorValue.A);
+	bSuccess &= VerifyAttributeValuesValid(this, TestData, GET_MEMBER_NAME_CHECKED(APCGUnitTestDummyActor, PCGColorProperty), ColorPropertyNames, ExtraTestWhat, PCGColorValue.R, PCGColorValue.G, PCGColorValue.B, PCGColorValue.A);
 
 	// Object Property Extracted - UPCGDummyGetPropertyTest
 	const TArray<FName> ObjectPropertyNames = { GET_MEMBER_NAME_CHECKED(UPCGDummyGetPropertyTest, Int64Property), GET_MEMBER_NAME_CHECKED(UPCGDummyGetPropertyTest, DoubleProperty) };
@@ -290,7 +298,7 @@ bool FPCGPropertyToParamDataPropertyTypeTest::RunTest(const FString& Parameters)
 
 	// Arrays of extracted properties
 	bSuccess &= VerifyAttributeValuesValid(this, TestData, GET_MEMBER_NAME_CHECKED(APCGUnitTestDummyActor, ArrayOfStructsProperty), ColorPropertyNames, ExtraTestWhat, 
-		(int64)ColorValue.R, (int64)ColorValue.G, (int64)ColorValue.B, (int64)ColorValue.A, (int64)SecondColorValue.R, (int64)SecondColorValue.G, (int64)SecondColorValue.B, (int64)SecondColorValue.A);
+		PCGColorValue.R, PCGColorValue.G, PCGColorValue.B, PCGColorValue.A, SecondPCGColorValue.R, SecondPCGColorValue.G, SecondPCGColorValue.B, SecondPCGColorValue.A);
 	bSuccess &= VerifyAttributeValuesValid(this, TestData, GET_MEMBER_NAME_CHECKED(APCGUnitTestDummyActor, ArrayOfObjectsProperty), ObjectPropertyNames, ExtraTestWhat, ObjectValue->Int64Property, ObjectValue->DoubleProperty, SecondObjectValue->Int64Property, SecondObjectValue->DoubleProperty);
 
 	// Extractors
diff --git a/Engine/Plugins/Experimental/PCG/Source/PCG/Private/Tests/Elements/PCGGetActorPropertyTest.h b/Engine/Plugins/Experimental/PCG/Source/PCG/Private/Tests/Elements/PCGGetActorPropertyTest.h
--- a/Engine/Plugins/Experimental/PCG/Source/PCG/Private/Tests/Elements/PCGGetActorPropertyTest.h
+++ b/Engine/Plugins/Experimental/PCG/Source/PCG/Private/Tests/Elements/PCGGetActorPropertyTest.h
@@ -38,6 +38,24 @@ struct FPCGDummyGetPropertyStruct
 	FPCGDummyGetPropertyLevel2Struct Level2Struct;
 };
 
+USTRUCT(BlueprintType, meta = (Hidden))
+struct FPCGTestMyColorStruct
+{
+	GENERATED_BODY()
+
+	UPROPERTY(BlueprintReadOnly, Category = "")
+	double B = 0;
+
+	UPROPERTY(BlueprintReadOnly, Category = "")
+	double G = 0;
+
+	UPROPERTY(BlueprintReadOnly, Category = "")
+	double R = 0;
+
+	UPROPERTY(BlueprintReadOnly, Category = "")
+	double A = 0;
+};
+
 UCLASS(MinimalAPI, NotBlueprintable, NotPlaceable, NotBlueprintType, Transient, HideDropdown, meta = (Hidden))
 class UPCGDummyGetPropertyTest : public UObject
 {
@@ -114,6 +132,12 @@ public:
 	UPROPERTY(BlueprintReadOnly, Category = "")
 	FColor ColorProperty;
 
+	UPROPERTY(BlueprintReadOnly, Category = "")
+	FLinearColor LinearColorProperty;
+
+	UPROPERTY(BlueprintReadOnly, Category = "")
+	FPCGTestMyColorStruct PCGColorProperty;
+
 	UPROPERTY(BlueprintReadOnly, Category = "")
 	TArray<int32> ArrayOfIntsProperty;
 
@@ -121,7 +145,7 @@ public:
 	TArray<FVector> ArrayOfVectorsProperty;
 
 	UPROPERTY(BlueprintReadOnly, Category = "")
-	TArray<FColor> ArrayOfStructsProperty;
+	TArray<FPCGTestMyColorStruct> ArrayOfStructsProperty;
 
 	UPROPERTY(BlueprintReadOnly, Category = "")
 	TArray<TObjectPtr<UPCGDummyGetPropertyTest>> ArrayOfObjectsProperty;
diff --git a/Engine/Plugins/Experimental/PCG/Source/PCG/Public/Metadata/Accessors/PCGPropertyAccessor.h b/Engine/Plugins/Experimental/PCG/Source/PCG/Public/Metadata/Accessors/PCGPropertyAccessor.h
--- a/Engine/Plugins/Experimental/PCG/Source/PCG/Public/Metadata/Accessors/PCGPropertyAccessor.h
+++ b/Engine/Plugins/Experimental/PCG/Source/PCG/Public/Metadata/Accessors/PCGPropertyAccessor.h
@@ -7,6 +7,7 @@
 #include "Metadata/PCGAttributePropertySelector.h"
 
 #include "Containers/UnrealString.h" // IWYU pragma: keep
+#include "Math/Color.h" // IWYU pragma: keep
 #include "UObject/SoftObjectPath.h" // IWYU pragma: keep
 #include "UObject/UnrealType.h" // IWYU pragma: keep
 
@@ -87,13 +88,13 @@ namespace PCGPropertyAccessor
 /**
 * Interface for Property chain to factorize ctor, fix the chain and storing the property chain
 */
-class IPCGPropertyChainAccessor
+class IPCGPropertyChain
 {
 public:
-	virtual ~IPCGPropertyChainAccessor() = default;
+	virtual ~IPCGPropertyChain() = default;
 
 protected:
-	IPCGPropertyChainAccessor(const FProperty* Property, TArray<const FProperty*>&& ExtraProperties);
+	IPCGPropertyChain(const FProperty* Property, TArray<const FProperty*>&& ExtraProperties);
 
 	const TArray<const FProperty*>& GetPropertyChain() const { return PropertyChain; }
 
@@ -108,7 +109,7 @@ private:
 * Key supported: Generic object
 */
 template <typename T>
-class FPCGNumericPropertyAccessor : public IPCGAttributeAccessorT<FPCGNumericPropertyAccessor<T>>, IPCGPropertyChainAccessor
+class FPCGNumericPropertyAccessor : public IPCGAttributeAccessorT<FPCGNumericPropertyAccessor<T>>, IPCGPropertyChain
 {
 public:
 	using Type = T;
@@ -116,7 +117,7 @@ public:
 
 	FPCGNumericPropertyAccessor(const FNumericProperty* InProperty, TArray<const FProperty*>&& ExtraProperties = {})
 		: Super(/*bInReadOnly=*/ false)
-		, IPCGPropertyChainAccessor(InProperty, std::forward<TArray<const FProperty*>>(ExtraProperties))
+		, IPCGPropertyChain(InProperty, std::forward<TArray<const FProperty*>>(ExtraProperties))
 		, Property(InProperty)
 	{
 		static_assert(std::is_integral_v<T> || std::is_floating_point_v<T>);
@@ -163,7 +164,7 @@ private:
 * Do not instanciate it manually, use PCGAttributeAccessorHelpers::CreatePropertyAccessor.
 * Key supported: Generic object
 */
-class FPCGEnumPropertyAccessor : public IPCGAttributeAccessorT<FPCGEnumPropertyAccessor>, IPCGPropertyChainAccessor
+class FPCGEnumPropertyAccessor : public IPCGAttributeAccessorT<FPCGEnumPropertyAccessor>, IPCGPropertyChain
 {
 public:
 	using Type = int64;
@@ -171,7 +172,7 @@ public:
 
 	FPCGEnumPropertyAccessor(const FEnumProperty* InProperty, TArray<const FProperty*>&& ExtraProperties = {})
 		: Super(/*bInReadOnly=*/ false)
-		, IPCGPropertyChainAccessor(InProperty, std::forward<TArray<const FProperty*>>(ExtraProperties))
+		, IPCGPropertyChain(InProperty, std::forward<TArray<const FProperty*>>(ExtraProperties))
 		, Property(InProperty)
 	{
 		check(Property);
@@ -192,7 +193,7 @@ private:
 * Key supported: Generic object
 */
 template <typename T>
-class FPCGPropertyStructAccessor : public IPCGAttributeAccessorT<FPCGPropertyStructAccessor<T>>, IPCGPropertyChainAccessor
+class FPCGPropertyStructAccessor : public IPCGAttributeAccessorT<FPCGPropertyStructAccessor<T>>, IPCGPropertyChain
 {
 public:
 	using Type = T;
@@ -200,7 +201,7 @@ public:
 
 	FPCGPropertyStructAccessor(const FStructProperty* InProperty, TArray<const FProperty*>&& ExtraProperties = {})
 		: Super(/*bInReadOnly=*/ false)
-		, IPCGPropertyChainAccessor(InProperty, std::forward<TArray<const FProperty*>>(ExtraProperties))
+		, IPCGPropertyChain(InProperty, std::forward<TArray<const FProperty*>>(ExtraProperties))
 		, Property(InProperty)
 	{
 		static_assert(PCG::Private::IsPCGType<T>());
@@ -234,7 +235,7 @@ private:
 * Key supported: Generic object
 */
 template <typename T, typename PropertyType>
-class FPCGPropertyAccessor : public IPCGAttributeAccessorT<FPCGPropertyAccessor<T, PropertyType>>, IPCGPropertyChainAccessor
+class FPCGPropertyAccessor : public IPCGAttributeAccessorT<FPCGPropertyAccessor<T, PropertyType>>, IPCGPropertyChain
 {
 public:
 	using Type = T;
@@ -242,7 +243,7 @@ public:
 
 	FPCGPropertyAccessor(const PropertyType* InProperty, TArray<const FProperty*>&& ExtraProperties = {})
 		: Super(/*bInReadOnly=*/ false)
-		, IPCGPropertyChainAccessor(InProperty, std::forward<TArray<const FProperty*>>(ExtraProperties))
+		, IPCGPropertyChain(InProperty, std::forward<TArray<const FProperty*>>(ExtraProperties))
 		, Property(InProperty)
 	{
 		static_assert(PCG::Private::IsPCGType<T>());
@@ -276,7 +277,7 @@ private:
 * Key supported: Generic object
 */
 template <typename T>
-class FPCGPropertyPathAccessor : public IPCGAttributeAccessorT<FPCGPropertyPathAccessor<T>>, IPCGPropertyChainAccessor
+class FPCGPropertyPathAccessor : public IPCGAttributeAccessorT<FPCGPropertyPathAccessor<T>>, IPCGPropertyChain
 {
 public:
 	using Type = FString;
@@ -284,7 +285,7 @@ public:
 
 	FPCGPropertyPathAccessor(const FProperty* InProperty, TArray<const FProperty*>&& ExtraProperties = {})
 		: Super(/*bInReadOnly=*/ false)
-		, IPCGPropertyChainAccessor(InProperty, std::forward<TArray<const FProperty*>>(ExtraProperties))
+		, IPCGPropertyChain(InProperty, std::forward<TArray<const FProperty*>>(ExtraProperties))
 		, Property(InProperty)
 	{
 		static_assert(std::is_same_v<FSoftObjectPath, T> || std::is_same_v<FSoftClassPath, T>);
@@ -316,7 +317,7 @@ private:
 * Do not instantiate it manually, use PCGAttributeAccessorHelpers::CreatePropertyAccessor.
 * Key supported: Generic object
 */
-class FPCGPropertySoftObjectPathAccessor : public IPCGAttributeAccessorT<FPCGPropertySoftObjectPathAccessor>, IPCGPropertyChainAccessor
+class FPCGPropertySoftObjectPathAccessor : public IPCGAttributeAccessorT<FPCGPropertySoftObjectPathAccessor>, IPCGPropertyChain
 {
 public:
 	using Type = FSoftObjectPath;
@@ -324,7 +325,7 @@ public:
 
 	FPCGPropertySoftObjectPathAccessor(const FSoftObjectProperty* InProperty, TArray<const FProperty*>&& ExtraProperties = {})
 		: Super(/*bInReadOnly=*/ false)
-		, IPCGPropertyChainAccessor(InProperty, std::forward<TArray<const FProperty*>>(ExtraProperties))
+		, IPCGPropertyChain(InProperty, std::forward<TArray<const FProperty*>>(ExtraProperties))
 		, Property(InProperty)
 	{
 		check(Property);
@@ -355,7 +356,7 @@ private:
 * Do not instantiate it manually, use PCGAttributeAccessorHelpers::CreatePropertyAccessor.
 * Key supported: Generic object
 */
-class FPCGPropertySoftClassPathAccessor : public IPCGAttributeAccessorT<FPCGPropertySoftClassPathAccessor>, IPCGPropertyChainAccessor
+class FPCGPropertySoftClassPathAccessor : public IPCGAttributeAccessorT<FPCGPropertySoftClassPathAccessor>, IPCGPropertyChain
 {
 public:
 	using Type = FSoftClassPath;
@@ -363,7 +364,7 @@ public:
 
 	FPCGPropertySoftClassPathAccessor(const FSoftClassProperty* InProperty, TArray<const FProperty*>&& ExtraProperties = {})
 		: Super(/*bInReadOnly=*/ false)
-		, IPCGPropertyChainAccessor(InProperty, std::forward<TArray<const FProperty*>>(ExtraProperties))
+		, IPCGPropertyChain(InProperty, std::forward<TArray<const FProperty*>>(ExtraProperties))
 		, Property(InProperty)
 	{
 		check(Property);
@@ -396,7 +397,7 @@ private:
 * Key supported: Generic object
 */
 template <typename PropertyType>
-class FPCGPropertyObjectPtrAccessor : public IPCGAttributeAccessorT<FPCGPropertyObjectPtrAccessor<PropertyType>>, IPCGPropertyChainAccessor
+class FPCGPropertyObjectPtrAccessor : public IPCGAttributeAccessorT<FPCGPropertyObjectPtrAccessor<PropertyType>>, IPCGPropertyChain
 {
 public:
 	using Type = std::conditional_t<std::is_same_v<PropertyType, FClassProperty>, FSoftClassPath, FSoftObjectPath>;
@@ -404,7 +405,7 @@ public:
 
 	FPCGPropertyObjectPtrAccessor(const PropertyType* InProperty, TArray<const FProperty*>&& ExtraProperties = {})
 		: Super(/*bInReadOnly=*/ false)
-		, IPCGPropertyChainAccessor(InProperty, std::forward<TArray<const FProperty*>>(ExtraProperties))
+		, IPCGPropertyChain(InProperty, std::forward<TArray<const FProperty*>>(ExtraProperties))
 		, Property(InProperty)
 	{
 		// Making sure it is the right properties.
@@ -434,9 +435,9 @@ private:
 
 /**
 * Special accessor to support attribute selector overrides. Interface with a string.
-* Key supported: All
+* Key supported: Generic object
 */
-class FPCGAttributePropertySelectorAccessor : public IPCGAttributeAccessorT<FPCGAttributePropertySelectorAccessor>, IPCGPropertyChainAccessor
+class FPCGAttributePropertySelectorAccessor : public IPCGAttributeAccessorT<FPCGAttributePropertySelectorAccessor>, IPCGPropertyChain
 {
 public:
 	using Type = FString;
@@ -444,7 +445,7 @@ public:
 
 	FPCGAttributePropertySelectorAccessor(const FStructProperty* InProperty, TArray<const FProperty*>&& ExtraProperties = {})
 		: Super(/*bInReadOnly=*/ false)
-		, IPCGPropertyChainAccessor(InProperty, std::forward<TArray<const FProperty*>>(ExtraProperties))
+		, IPCGPropertyChain(InProperty, std::forward<TArray<const FProperty*>>(ExtraProperties))
 		, Property(InProperty)
 	{
 		ensure(InProperty && InProperty->Struct && InProperty->Struct->IsA(FPCGAttributePropertySelector::StaticStruct()->GetClass()));
@@ -470,3 +471,81 @@ private:
 	const FStructProperty* Property = nullptr;
 };
 
+/**
+* Special accessor to support linear color overrides. Interface with a vector 4. Will be output as RGBA.
+* Key supported: Generic object
+*/
+class FPCGLinearColorAccessor : public IPCGAttributeAccessorT<FPCGLinearColorAccessor>, IPCGPropertyChain
+{
+public:
+	using Type = FVector4;
+	using Super = IPCGAttributeAccessorT<FPCGLinearColorAccessor>;
+
+	FPCGLinearColorAccessor(const FStructProperty* InProperty, TArray<const FProperty*>&& ExtraProperties = {})
+		: Super(/*bInReadOnly=*/ false)
+		, IPCGPropertyChain(InProperty, std::forward<TArray<const FProperty*>>(ExtraProperties))
+		, Property(InProperty)
+	{
+		ensure(InProperty && InProperty->Struct && InProperty->Struct == TBaseStructure<FLinearColor>::Get());
+	}
+
+	bool GetRangeImpl(TArrayView<FVector4> OutValues, int32 Index, const IPCGAttributeAccessorKeys& Keys) const
+	{
+		return PCGPropertyAccessor::IterateGet(GetPropertyChain(), OutValues, Index, Keys, [this](const void* PropertyAddressData) -> FVector4
+		{
+			const FLinearColor* Value = reinterpret_cast<const FLinearColor*>(PropertyAddressData);
+			return FVector4(Value->R, Value->G, Value->B, Value->A);
+		});
+	}
+
+	bool SetRangeImpl(TArrayView<const FVector4> InValues, int32 Index, IPCGAttributeAccessorKeys& Keys, EPCGAttributeAccessorFlags Flags)
+	{
+		return PCGPropertyAccessor::IterateSet(GetPropertyChain(), InValues, Index, Keys, [this](void* PropertyAddressData, const FVector4& Value) -> void
+		{
+			*reinterpret_cast<FLinearColor*>(PropertyAddressData) = FLinearColor(Value);
+		});
+	}
+
+private:
+	const FStructProperty* Property = nullptr;
+};
+
+/**
+* Special accessor to support color overrides. Interface with a vector 4. Will remap [0;255] to [0.0,1.0]
+* Key supported: Generic object
+*/
+class FPCGColorAccessor : public IPCGAttributeAccessorT<FPCGColorAccessor>, IPCGPropertyChain
+{
+public:
+	using Type = FVector4;
+	using Super = IPCGAttributeAccessorT<FPCGColorAccessor>;
+
+	FPCGColorAccessor(const FStructProperty* InProperty, TArray<const FProperty*>&& ExtraProperties = {})
+		: Super(/*bInReadOnly=*/ false)
+		, IPCGPropertyChain(InProperty, std::forward<TArray<const FProperty*>>(ExtraProperties))
+		, Property(InProperty)
+	{
+		ensure(InProperty && InProperty->Struct && InProperty->Struct == TBaseStructure<FColor>::Get());
+	}
+
+	bool GetRangeImpl(TArrayView<FVector4> OutValues, int32 Index, const IPCGAttributeAccessorKeys& Keys) const
+	{
+		return PCGPropertyAccessor::IterateGet(GetPropertyChain(), OutValues, Index, Keys, [this](const void* PropertyAddressData) -> FVector4
+		{
+			constexpr double Inv255 = 1.0 / 255.0;
+			const FColor* Value = reinterpret_cast<const FColor*>(PropertyAddressData);
+			return FVector4(Value->R * Inv255, Value->G * Inv255, Value->B * Inv255, Value->A * Inv255);
+		});
+	}
+
+	bool SetRangeImpl(TArrayView<const FVector4> InValues, int32 Index, IPCGAttributeAccessorKeys& Keys, EPCGAttributeAccessorFlags Flags)
+	{
+		return PCGPropertyAccessor::IterateSet(GetPropertyChain(), InValues, Index, Keys, [this](void* PropertyAddressData, const FVector4& Value) -> void
+		{
+			*reinterpret_cast<FColor*>(PropertyAddressData) = FLinearColor(Value).QuantizeRound();
+		});
+	}
+
+private:
+	const FStructProperty* Property = nullptr;
+};
\ No newline at end of file
