Made sure we don't reset the loaders when renaming objects during PostLoad

#rnx

diff --git a/Engine/Plugins/Experimental/PCG/Source/PCG/Private/Elements/PCGCreateTargetActor.cpp b/Engine/Plugins/Experimental/PCG/Source/PCG/Private/Elements/PCGCreateTargetActor.cpp
--- a/Engine/Plugins/Experimental/PCG/Source/PCG/Private/Elements/PCGCreateTargetActor.cpp
+++ b/Engine/Plugins/Experimental/PCG/Source/PCG/Private/Elements/PCGCreateTargetActor.cpp
@@ -157,7 +157,7 @@ void UPCGCreateTargetActor::RefreshTemplateActor()
 				Options.bNotifyObjectReplacement = true;
 				UEngine::CopyPropertiesForUnrelatedObjects(TemplateActor, NewTemplateActor, Options);
 
-				TemplateActor->Rename(nullptr, GetTransientPackage(), REN_DontCreateRedirectors);
+				TemplateActor->Rename(nullptr, GetTransientPackage(), REN_DontCreateRedirectors | REN_ForceNoResetLoaders);
 
 				TMap<UObject*, UObject*> OldToNew;
 				OldToNew.Emplace(TemplateActor, NewTemplateActor);
diff --git a/Engine/Plugins/Experimental/PCG/Source/PCG/Private/PCGGraph.cpp b/Engine/Plugins/Experimental/PCG/Source/PCG/Private/PCGGraph.cpp
--- a/Engine/Plugins/Experimental/PCG/Source/PCG/Private/PCGGraph.cpp
+++ b/Engine/Plugins/Experimental/PCG/Source/PCG/Private/PCGGraph.cpp
@@ -517,7 +517,7 @@ UPCGNode* UPCGGraph::AddNodeInstance(UPCGSettings* InSettings)
 
 	if (Node)
 	{
-		SettingsInstance->Rename(nullptr, Node);
+		SettingsInstance->Rename(nullptr, Node, REN_ForceNoResetLoaders | REN_DontCreateRedirectors);
 		SettingsInstance->SetFlags(RF_Transactional);
 	}
 
@@ -536,7 +536,7 @@ UPCGNode* UPCGGraph::AddNodeCopy(UPCGSettings* InSettings, UPCGSettings*& Defaul
 
 	if (SettingsCopy)
 	{
-		SettingsCopy->Rename(nullptr, NewNode);
+		SettingsCopy->Rename(nullptr, NewNode, REN_ForceNoResetLoaders | REN_DontCreateRedirectors);
 	}
 
 	DefaultNodeSettings = SettingsCopy;
@@ -674,14 +674,14 @@ void UPCGGraph::AddNodes_Internal(TArrayView<UPCGNode*> InNodes)
 	for (UPCGNode* Node : InNodes)
 	{
 		check(Node);
-		Node->Rename(nullptr, this);
+		Node->Rename(nullptr, this, REN_ForceNoResetLoaders | REN_DontCreateRedirectors);
 
 #if WITH_EDITOR
 		const FName DefaultNodeName = Node->GetSettings()->GetDefaultNodeName();
 		if (DefaultNodeName != NAME_None)
 		{
 			FName NodeName = MakeUniqueObjectName(this, UPCGNode::StaticClass(), DefaultNodeName);
-			Node->Rename(*NodeName.ToString());
+			Node->Rename(*NodeName.ToString(), nullptr, REN_ForceNoResetLoaders | REN_DontCreateRedirectors);
 		}
 #endif
 
diff --git a/Engine/Plugins/Experimental/PCG/Source/PCG/Private/RuntimeGen/PCGRuntimeGenScheduler.cpp b/Engine/Plugins/Experimental/PCG/Source/PCG/Private/RuntimeGen/PCGRuntimeGenScheduler.cpp
--- a/Engine/Plugins/Experimental/PCG/Source/PCG/Private/RuntimeGen/PCGRuntimeGenScheduler.cpp
+++ b/Engine/Plugins/Experimental/PCG/Source/PCG/Private/RuntimeGen/PCGRuntimeGenScheduler.cpp
@@ -852,7 +852,7 @@ void FPCGRuntimeGenScheduler::CleanupComponent(const FGridGenerationKey& Generat
 			}
 
 #if WITH_EDITOR
-			PartitionActor->Rename(nullptr, PartitionActor->GetOuter(), REN_NonTransactional | REN_DoNotDirty);
+			PartitionActor->Rename(nullptr, PartitionActor->GetOuter(), REN_NonTransactional | REN_DoNotDirty | REN_ForceNoResetLoaders);
 			PartitionActor->SetActorLabel(*PCGRuntimeGenSchedulerConstants::PooledPartitionActorName);
 #endif
 			PartitionActorPool.Push(PartitionActor);
@@ -865,7 +865,7 @@ void FPCGRuntimeGenScheduler::CleanupComponent(const FGridGenerationKey& Generat
 			}
 
 #if WITH_EDITOR
-			PartitionActor->Rename(nullptr, PartitionActor->GetOuter(), REN_NonTransactional | REN_DoNotDirty);
+			PartitionActor->Rename(nullptr, PartitionActor->GetOuter(), REN_NonTransactional | REN_DoNotDirty | REN_ForceNoResetLoaders);
 #endif
 			World->DestroyActor(PartitionActor);
 		}
@@ -1069,7 +1069,7 @@ APCGPartitionActor* FPCGRuntimeGenScheduler::GetPartitionActorFromPool(uint32 Gr
 #if WITH_EDITOR
 	const FName ActorName = *APCGPartitionActor::GetRuntimeGenActorName(GridSize, GridCoords);
 
-	PartitionActor->Rename(*ActorName.ToString(), PartitionActor->GetOuter(), REN_NonTransactional | REN_DoNotDirty);
+	PartitionActor->Rename(*ActorName.ToString(), PartitionActor->GetOuter(), REN_NonTransactional | REN_DoNotDirty | REN_ForceNoResetLoaders);
 	PartitionActor->SetActorLabel(ActorName.ToString());
 #endif
 
@@ -1122,7 +1122,7 @@ void FPCGRuntimeGenScheduler::ResetPartitionActorPoolToSize(uint32 NewPoolSize)
 	for (APCGPartitionActor* PartitionActor : PartitionActorPool)
 	{
 #if WITH_EDITOR
-		PartitionActor->Rename(nullptr, PartitionActor->GetOuter(), REN_NonTransactional | REN_DoNotDirty);
+		PartitionActor->Rename(nullptr, PartitionActor->GetOuter(), REN_NonTransactional | REN_DoNotDirty | REN_ForceNoResetLoaders);
 #endif
 		World->DestroyActor(PartitionActor);
 	}
