World actor override cleanup and adding missing call to super post edit change property

#rnx

diff --git a/Engine/Plugins/Experimental/PCG/Source/PCG/Private/PCGWorldActor.cpp b/Engine/Plugins/Experimental/PCG/Source/PCG/Private/PCGWorldActor.cpp
--- a/Engine/Plugins/Experimental/PCG/Source/PCG/Private/PCGWorldActor.cpp
+++ b/Engine/Plugins/Experimental/PCG/Source/PCG/Private/PCGWorldActor.cpp
@@ -316,6 +316,8 @@ void APCGWorldActor::PostLoad()
 
 void APCGWorldActor::PostEditChangeProperty(FPropertyChangedEvent& PropertyChangedEvent)
 {
+	Super::PostEditChangeProperty(PropertyChangedEvent);
+	
 	if (PropertyChangedEvent.GetPropertyName() == GET_MEMBER_NAME_CHECKED(APCGWorldActor, PartitionGridSize)
 		|| PropertyChangedEvent.GetPropertyName() == GET_MEMBER_NAME_CHECKED(APCGWorldActor, bUse2DGrid))
 	{
diff --git a/Engine/Plugins/Experimental/PCG/Source/PCG/Public/PCGWorldActor.h b/Engine/Plugins/Experimental/PCG/Source/PCG/Public/PCGWorldActor.h
--- a/Engine/Plugins/Experimental/PCG/Source/PCG/Public/PCGWorldActor.h
+++ b/Engine/Plugins/Experimental/PCG/Source/PCG/Public/PCGWorldActor.h
@@ -37,11 +37,11 @@ public:
 	void MergeFrom(APCGWorldActor* OtherWorldActor);
 
 #if WITH_EDITOR
-	virtual bool CanChangeIsSpatiallyLoadedFlag() const { return false; }
+	virtual bool CanChangeIsSpatiallyLoadedFlag() const override { return false; }
 	virtual bool IsUserManaged() const override { return false; }
 	virtual bool ShouldExport() override { return false; }
 	virtual bool ShouldImport(FStringView ActorPropString, bool IsMovingLevel) override { return false; }
-	virtual void BeginCacheForCookedPlatformData(const ITargetPlatform* TargetPlatform);
+	virtual void BeginCacheForCookedPlatformData(const ITargetPlatform* TargetPlatform) override;
 	//~End AActor Interface
 
 	static APCGWorldActor* CreatePCGWorldActor(UWorld* InWorld);
