[PCG] Add GetPCGComponentData typed getter
* Stop local component executions storing their generated graph output data on the original component (we now distinguish between data dependencies and execution dependencies when scheduling)
* Cleaned up & documented GetActorData settings

diff --git a/Engine/Plugins/Experimental/PCG/Source/PCG/Private/Elements/PCGDataFromActor.cpp b/Engine/Plugins/Experimental/PCG/Source/PCG/Private/Elements/PCGDataFromActor.cpp
--- a/Engine/Plugins/Experimental/PCG/Source/PCG/Private/Elements/PCGDataFromActor.cpp
+++ b/Engine/Plugins/Experimental/PCG/Source/PCG/Private/Elements/PCGDataFromActor.cpp
@@ -3,6 +3,7 @@
 #include "Elements/PCGDataFromActor.h"
 
 #include "PCGComponent.h"
+#include "PCGCustomVersion.h"
 #include "PCGSubsystem.h"
 #include "Data/PCGPointData.h"
 #include "Data/PCGSpatialData.h"
@@ -22,6 +23,107 @@
 namespace PCGDataFromActorConstants
 {
 	static const FName SinglePointPinLabel = TEXT("Single Point");
+	static const FString PCGComponentDataGridSizeTagPrefix = TEXT("PCG_GridSize_");
+}
+
+namespace PCGDataFromActorHelpers
+{
+	/**
+	 * Get the PCG Components associated with an actor. Optionally, also search for any local components associated with components
+	 * on the actor using the 'bGetLocalComponents' flag. By default, gets data on all grids, but alternatively you can provide a
+	 * set of 'AllowedGrids' to match against.
+	 *
+	 * If 'bMustOverlap' is true, it will only collect components which overlap with the given 'OverlappingBounds'. Note that this
+	 * overlap does not include bounds which are only touching, with no overlapping volume.
+	 */
+	static TInlineComponentArray<UPCGComponent*, 1> GetPCGComponentsFromActor(
+		AActor* Actor,
+		UPCGSubsystem* Subsystem,
+		bool bGetLocalComponents = false,
+		bool bGetAllGrids = true,
+		int32 AllowedGrids = (int32)EPCGHiGenGrid::Uninitialized,
+		bool bMustOverlap = false,
+		const FBox& OverlappingBounds = FBox())
+	{
+		TInlineComponentArray<UPCGComponent*, 1> PCGComponents;
+
+		if (!Actor || !Subsystem)
+		{
+			return PCGComponents;
+		}
+
+		Actor->GetComponents(PCGComponents);
+
+		if (bMustOverlap)
+		{
+			// Remove actor components that do not overlap the source bounds.
+			// Note: This assumes that a local component always lies inside the bounds of its original component,
+			// which is true at the time of writing, but may not always be the case (e.g. "truly" unbounded execution).
+			for (int I = PCGComponents.Num() - 1; I >= 0; I--)
+			{
+				const FBox ComponentBounds = PCGComponents[I]->GetGridBounds();
+
+				// We reject overlaps with zero volume instead of simply checking Intersect(...) to avoid bounds which touch but do not overlap.
+				if (OverlappingBounds.Overlap(ComponentBounds).GetVolume() <= 0)
+				{
+					PCGComponents.RemoveAtSwap(I);
+				}
+			}
+		}
+
+		TArray<UPCGComponent*> LocalComponents;
+
+		if (bGetLocalComponents)
+		{
+			// Collect the local components for each actor PCG component.
+			for (UPCGComponent* Component : PCGComponents)
+			{
+				if (Component && Component->IsPartitioned())
+				{
+					Subsystem->ForAllRegisteredLocalComponents(Component, [&LocalComponents, bGetAllGrids, AllowedGrids, bMustOverlap, &OverlappingBounds](UPCGComponent* LocalComponent)
+					{
+						if (bGetAllGrids || (AllowedGrids & (int32)LocalComponent->GetGenerationGrid()))
+						{
+							if (bMustOverlap)
+							{
+								const FBox LocalBounds = LocalComponent->GetGridBounds();
+
+								// We reject overlaps with zero volume instead of simply checking Intersect(...) to avoid bounds which touch but do not overlap.
+								if (OverlappingBounds.Overlap(LocalBounds).GetVolume() > 0)
+								{
+									LocalComponents.Add(LocalComponent);
+								}
+							}
+							else
+							{
+								LocalComponents.Add(LocalComponent);
+							}
+						}
+					});
+				}
+			}
+		}
+
+		// Remove the actor's PCG components if they aren't on an allowed grid size.
+		// Implementation note: We delay removing these components until now because they may have had local components on an allowed grid size.
+		if (!bGetAllGrids)
+		{
+			for (int I = PCGComponents.Num() - 1; I >= 0; I--)
+			{
+				if (!(AllowedGrids & (int32)PCGComponents[I]->GetGenerationGridSize()))
+				{
+					PCGComponents.RemoveAtSwap(I);
+				}
+			}
+		}
+
+		if (bGetLocalComponents)
+		{
+			PCGComponents.Append(LocalComponents);
+		}
+
+		return PCGComponents;
+	}
 }
 
 #if WITH_EDITOR
@@ -36,6 +138,17 @@ void UPCGDataFromActorSettings::GetTrackedActorKeys(FPCGActorSelectionKeyToSetti
 	OutKeysToSettings.FindOrAdd(Key).Emplace(this, bTrackActorsOnlyWithinBounds);
 }
 
+void UPCGDataFromActorSettings::ApplyDeprecation(UPCGNode* InOutNode)
+{
+	if (DataVersion < FPCGCustomVersion::GetPCGComponentDataMustOverlapSourceComponentByDefault)
+	{
+		// Old versions of GetActorData did not require found components to overlap self, but going forward it's a more efficient default.
+		bComponentsMustOverlapSelf = false;
+	}
+
+	Super::ApplyDeprecation(InOutNode);
+}
+
 FText UPCGDataFromActorSettings::GetNodeTooltipText() const
 {
 	return LOCTEXT("DataFromActorTooltip", "Builds a collection of PCG-compatible data from the selected actors.");
@@ -254,11 +367,27 @@ void FPCGDataFromActorElement::GatherWaitTasks(AActor* FoundActor, FPCGContext*
 		return;
 	}
 
+	const UPCGDataFromActorSettings* Settings = Context->GetInputSettings<UPCGDataFromActorSettings>();
+	check(Settings);
+
+	UPCGComponent* SourceComponent = Context->SourceComponent.IsValid() ? Context->SourceComponent.Get() : nullptr;
+
+	if (!SourceComponent)
+	{
+		return;
+	}
+
 	// We will prevent gathering the current execution - this task cannot wait on itself
-	AActor* ThisOwner = ((Context && Context->SourceComponent.IsValid()) ? Context->SourceComponent->GetOwner() : nullptr);
+	AActor* ThisOwner = SourceComponent->GetOwner();
 
-	TInlineComponentArray<UPCGComponent*, 1> PCGComponents;
-	FoundActor->GetComponents(PCGComponents);
+	TInlineComponentArray<UPCGComponent*, 1> PCGComponents = PCGDataFromActorHelpers::GetPCGComponentsFromActor(
+		FoundActor,
+		SourceComponent->GetSubsystem(),
+		/*bGetLocalComponents=*/true,
+		Settings->bGetDataOnAllGrids,
+		Settings->AllowedGrids,
+		Settings->bComponentsMustOverlapSelf,
+		(Settings->bComponentsMustOverlapSelf && SourceComponent) ? SourceComponent->GetGridBounds() : FBox());
 
 	for (UPCGComponent* Component : PCGComponents)
 	{
@@ -356,12 +485,14 @@ void FPCGDataFromActorElement::ProcessActor(FPCGContext* Context, const UPCGData
 	check(Context);
 	check(Settings);
 
-	if (!FoundActor || !IsValid(FoundActor))
+	UPCGComponent* SourceComponent = Context->SourceComponent.IsValid() ? Context->SourceComponent.Get() : nullptr;
+
+	if (!FoundActor || !IsValid(FoundActor) || !SourceComponent)
 	{
 		return;
 	}
 
-	AActor* ThisOwner = ((Context && Context->SourceComponent.Get()) ? Context->SourceComponent->GetOwner() : nullptr);
+	AActor* ThisOwner = SourceComponent->GetOwner();
 	TInlineComponentArray<UPCGComponent*, 1> PCGComponents;
 	bool bHasGeneratedPCGData = false;
 	FProperty* FoundProperty = nullptr;
@@ -370,7 +501,14 @@ void FPCGDataFromActorElement::ProcessActor(FPCGContext* Context, const UPCGData
 
 	if (bCanGetDataFromComponent && (Settings->Mode == EPCGGetDataFromActorMode::GetDataFromPCGComponent || Settings->Mode == EPCGGetDataFromActorMode::GetDataFromPCGComponentOrParseComponents))
 	{
-		FoundActor->GetComponents(PCGComponents);
+		PCGComponents = PCGDataFromActorHelpers::GetPCGComponentsFromActor(
+			FoundActor,
+			SourceComponent->GetSubsystem(),
+			/*bGetLocalComponents=*/true,
+			Settings->bGetDataOnAllGrids,
+			Settings->AllowedGrids,
+			Settings->bComponentsMustOverlapSelf,
+			Settings->bComponentsMustOverlapSelf ? SourceComponent->GetGridBounds() : FBox());
 
 		for (UPCGComponent* Component : PCGComponents)
 		{
@@ -419,6 +557,7 @@ void FPCGDataFromActorElement::ProcessActor(FPCGContext* Context, const UPCGData
 			{
 				FPCGTaggedData& DuplicatedTaggedData = Outputs.Add_GetRef(TaggedData);
 				DuplicatedTaggedData.Data = Cast<UPCGData>(StaticDuplicateObject(TaggedData.Data, GetTransientPackage()));
+				DuplicatedTaggedData.Tags.Add(PCGDataFromActorConstants::PCGComponentDataGridSizeTagPrefix + FString::FromInt(PCGHiGenGrid::GridToGridSize(Component->GetGenerationGrid())));
 			}
 			//Outputs.Append(Component->GetGeneratedGraphOutput().TaggedData);
 		}
@@ -452,7 +591,7 @@ void FPCGDataFromActorElement::ProcessActor(FPCGContext* Context, const UPCGData
 	else
 	{
 		const bool bParseActor = (Settings->Mode != EPCGGetDataFromActorMode::GetSinglePoint);
-		FPCGDataCollection Collection = UPCGComponent::CreateActorPCGDataCollection(FoundActor, Context->SourceComponent.Get(), Settings->GetDataFilter(), bParseActor);
+		FPCGDataCollection Collection = UPCGComponent::CreateActorPCGDataCollection(FoundActor, SourceComponent, Settings->GetDataFilter(), bParseActor);
 		Outputs += Collection.TaggedData;
 	}
 
@@ -460,7 +599,7 @@ void FPCGDataFromActorElement::ProcessActor(FPCGContext* Context, const UPCGData
 	if (Settings->bAlsoOutputSinglePointData && (Settings->Mode == EPCGGetDataFromActorMode::GetDataFromPCGComponent || Settings->Mode == EPCGGetDataFromActorMode::GetDataFromPCGComponentOrParseComponents))
 	{
 		const bool bParseActor = false;
-		FPCGDataCollection Collection = UPCGComponent::CreateActorPCGDataCollection(FoundActor, Context->SourceComponent.Get(), EPCGDataType::Any, bParseActor);
+		FPCGDataCollection Collection = UPCGComponent::CreateActorPCGDataCollection(FoundActor, SourceComponent, EPCGDataType::Any, bParseActor);
 		for (const FPCGTaggedData& SinglePointData : Collection.TaggedData)
 		{
 			FPCGTaggedData& OutSinglePoint = Outputs.Add_GetRef(SinglePointData);
@@ -474,11 +613,11 @@ void FPCGDataFromActorElement::GetDependenciesCrc(const FPCGDataCollection& InIn
 	FPCGCrc Crc;
 	IPCGElement::GetDependenciesCrc(InInput, InSettings, InComponent, Crc);
 
-	// If we track self or original, we are dependant on the actor data
+	// If we track self or original, we are dependent on the actor data
 	if (const UPCGDataFromActorSettings* Settings = Cast<const UPCGDataFromActorSettings>(InSettings))
 	{
 		const bool bDependsOnSelfOrHierarchy = (Settings->ActorSelector.ActorFilter == EPCGActorFilter::Self || Settings->ActorSelector.ActorFilter == EPCGActorFilter::Original);
-		const bool bDependsOnSelfBounds = (Settings->ActorSelector.bMustOverlapSelf);
+		const bool bDependsOnSelfBounds = Settings->ActorSelector.bMustOverlapSelf;
 
 		if (InComponent && (bDependsOnSelfOrHierarchy || bDependsOnSelfBounds))
 		{
@@ -490,6 +629,17 @@ void FPCGDataFromActorElement::GetDependenciesCrc(const FPCGDataCollection& InIn
 				Crc.Combine(ActorData->GetOrComputeCrc(/*bFullDataCrc=*/false));
 			}
 		}
+
+		const bool bDependsOnComponentData = Settings->Mode == EPCGGetDataFromActorMode::GetDataFromPCGComponent || Settings->Mode == EPCGGetDataFromActorMode::GetDataFromPCGComponentOrParseComponents;
+		const bool bDependsOnLocalComponentBounds = Settings->bComponentsMustOverlapSelf || !Settings->bGetDataOnAllGrids;
+
+		if (InComponent && bDependsOnComponentData && bDependsOnLocalComponentBounds)
+		{
+			if (const UPCGData* LocalActorData = InComponent->GetActorPCGData())
+			{
+				Crc.Combine(LocalActorData->GetOrComputeCrc(/*bFullDataCrc=*/false));
+			}
+		}
 	}
 
 	OutCrc = Crc;
diff --git a/Engine/Plugins/Experimental/PCG/Source/PCG/Private/Elements/PCGTypedGetter.cpp b/Engine/Plugins/Experimental/PCG/Source/PCG/Private/Elements/PCGTypedGetter.cpp
--- a/Engine/Plugins/Experimental/PCG/Source/PCG/Private/Elements/PCGTypedGetter.cpp
+++ b/Engine/Plugins/Experimental/PCG/Source/PCG/Private/Elements/PCGTypedGetter.cpp
@@ -198,4 +198,34 @@ FText UPCGGetPrimitiveSettings::GetNodeTooltipText() const
 }
 #endif
 
+UPCGGetPCGComponentSettings::UPCGGetPCGComponentSettings()
+{
+	bDisplayModeSettings = false;
+	Mode = EPCGGetDataFromActorMode::GetDataFromPCGComponent;
+
+	// It doesn't make sense (for dependency reasons) to look for components on self/same hierarchy, so we can just hide those options.
+	ActorSelector.bShowActorFilter = false;
+	ActorSelector.ActorFilter = EPCGActorFilter::AllWorldActors;
+	ActorSelector.bShowIgnoreSelfAndChildren = false;
+	ActorSelector.bIgnoreSelfAndChildren = true;
+}
+
+TArray<FPCGPinProperties> UPCGGetPCGComponentSettings::OutputPinProperties() const
+{
+	TArray<FPCGPinProperties> PinProperties = Super::OutputPinProperties();
+	check(!PinProperties.IsEmpty());
+	PinProperties[0].AllowedTypes = EPCGDataType::Any;
+
+	return PinProperties;
+}
+
+#if WITH_EDITOR
+FText UPCGGetPCGComponentSettings::GetNodeTooltipText() const
+{
+	return LOCTEXT("GetPCGComponentTooltip", "Builds a collection of data from other PCG components on the selected actors.\n"
+		"Automatically tags each output with the grid size it was collected from, prefixed by \"PCG_GridSize_\" (e.g.PCG_GridSize_12800).\n"
+		"Note: a component cannot get component data from itself or other components in its hierarchy, as it could create a circular dependency.");
+}
+#endif
+
 #undef LOCTEXT_NAMESPACE
diff --git a/Engine/Plugins/Experimental/PCG/Source/PCG/Private/Graph/PCGGraphExecutor.cpp b/Engine/Plugins/Experimental/PCG/Source/PCG/Private/Graph/PCGGraphExecutor.cpp
--- a/Engine/Plugins/Experimental/PCG/Source/PCG/Private/Graph/PCGGraphExecutor.cpp
+++ b/Engine/Plugins/Experimental/PCG/Source/PCG/Private/Graph/PCGGraphExecutor.cpp
@@ -401,24 +401,30 @@ bool FPCGGraphExecutor::IsGraphCurrentlyExecuting(UPCGGraph* InGraph)
 	return bAnyPresent;
 }
 
-FPCGTaskId FPCGGraphExecutor::ScheduleGeneric(TFunction<bool()> InOperation, UPCGComponent* InSourceComponent, const TArray<FPCGTaskId>& TaskDependencies)
+FPCGTaskId FPCGGraphExecutor::ScheduleGeneric(TFunction<bool()> InOperation, UPCGComponent* InSourceComponent, const TArray<FPCGTaskId>& TaskExecutionDependencies)
 {
-	// Since we have no context, the generic task will consume no input
-	constexpr bool bConsumeInputData = false;
+	// Since we have no context, the generic task will consume no input (no data dependencies).
 	return ScheduleGenericWithContext([Operation = MoveTemp(InOperation)](FPCGContext*) -> bool
-		{
+	{
 			return Operation();
-		}, InSourceComponent, TaskDependencies, bConsumeInputData);
+	}, InSourceComponent, TaskExecutionDependencies, /*TaskDataDependencies=*/{});
 }
 
-FPCGTaskId FPCGGraphExecutor::ScheduleGenericWithContext(TFunction<bool(FPCGContext*)> InOperation, UPCGComponent* InSourceComponent, const TArray<FPCGTaskId>& TaskDependencies, bool bConsumeInputData)
+FPCGTaskId FPCGGraphExecutor::ScheduleGenericWithContext(TFunction<bool(FPCGContext*)> InOperation, UPCGComponent* InSourceComponent, const TArray<FPCGTaskId>& TaskDataDependencies, const TArray<FPCGTaskId>& TaskExecutionDependencies)
 {
 	// Build task & element to hold the operation to perform
 	FPCGGraphTask Task;
-	for (FPCGTaskId TaskDependency : TaskDependencies)
+
+	for (FPCGTaskId TaskDependency : TaskExecutionDependencies)
+	{
+		ensure(TaskDependency != InvalidPCGTaskId);
+		Task.Inputs.Emplace(TaskDependency, /*InPin=*/nullptr, /*OutPin=*/nullptr, /*bConsumeInputData=*/false);
+	}
+
+	for (FPCGTaskId TaskDependency : TaskDataDependencies)
 	{
 		ensure(TaskDependency != InvalidPCGTaskId);
-		Task.Inputs.Emplace(TaskDependency, nullptr, nullptr, bConsumeInputData);
+		Task.Inputs.Emplace(TaskDependency, /*InPin=*/nullptr, /*OutPin=*/nullptr, /*bConsumeInputData=*/true);
 	}
 
 	Task.SourceComponent = InSourceComponent;
diff --git a/Engine/Plugins/Experimental/PCG/Source/PCG/Private/Graph/PCGGraphExecutor.h b/Engine/Plugins/Experimental/PCG/Source/PCG/Private/Graph/PCGGraphExecutor.h
--- a/Engine/Plugins/Experimental/PCG/Source/PCG/Private/Graph/PCGGraphExecutor.h
+++ b/Engine/Plugins/Experimental/PCG/Source/PCG/Private/Graph/PCGGraphExecutor.h
@@ -119,15 +119,15 @@ public:
 	bool IsGraphCurrentlyExecuting(UPCGGraph* InGraph);
 
 	// Back compatibility function. Use ScheduleGenericWithContext
-	FPCGTaskId ScheduleGeneric(TFunction<bool()> InOperation, UPCGComponent* InSourceComponent, const TArray<FPCGTaskId>& TaskDependencies);
+	FPCGTaskId ScheduleGeneric(TFunction<bool()> InOperation, UPCGComponent* InSourceComponent, const TArray<FPCGTaskId>& TaskExecutionDependencies);
 
 	/** General job scheduling
-	*  @param InOperation:       Callback that takes a Context as argument and returns true if the task is done, false otherwise
-	*  @param InSourceComponent: PCG component associated with this task. Can be null.
-	*  @param TaskDependencies:  List of all the dependencies for this task.
-	*  @param bConsumeInputData: If your task need a context, but don't need the input data, set this flag to false. Default is true.
+	*  @param InOperation:               Callback that takes a Context as argument and returns true if the task is done, false otherwise
+	*  @param InSourceComponent:         PCG component associated with this task. Can be null.
+	*  @param TaskExecutionDependencies: Task will wait on these tasks to execute and won't take their output data as input.
+	*  @param TaskDataDependencies:      Task will wait on these tasks to execute and will take their output data as input.
 	*/
-	FPCGTaskId ScheduleGenericWithContext(TFunction<bool(FPCGContext*)> InOperation, UPCGComponent* InSourceComponent, const TArray<FPCGTaskId>& TaskDependencies, bool bConsumeInputData = true);
+	FPCGTaskId ScheduleGenericWithContext(TFunction<bool(FPCGContext*)> InOperation, UPCGComponent* InSourceComponent, const TArray<FPCGTaskId>& TaskExecutionDependencies, const TArray<FPCGTaskId>& TaskDataDependencies);
 
 	/** Gets data in the output results. Returns false if data is not ready. */
 	bool GetOutputData(FPCGTaskId InTaskId, FPCGDataCollection& OutData);
diff --git a/Engine/Plugins/Experimental/PCG/Source/PCG/Private/PCGSubsystem.cpp b/Engine/Plugins/Experimental/PCG/Source/PCG/Private/PCGSubsystem.cpp
--- a/Engine/Plugins/Experimental/PCG/Source/PCG/Private/PCGSubsystem.cpp
+++ b/Engine/Plugins/Experimental/PCG/Source/PCG/Private/PCGSubsystem.cpp
@@ -500,7 +500,13 @@ FPCGTaskId UPCGSubsystem::ScheduleComponent(UPCGComponent* PCGComponent, EPCGHiG
 	}
 #endif // WITH_EDITOR
 
-	TArray<FPCGTaskId> AllTasks;
+	// Execution dependencies require a task to finish executing before the dependent task.
+	TArray<FPCGTaskId> ExecutionDependencyTasks;
+
+	// Data dependencies act as execution dependencies, but will also have their output consumed by the waiting task. For a component, this means
+	// it will store the output data into its managed resources, which, for an original component, should not include the local component generation tasks,
+	// since those resources should be managed locally.
+	TArray<FPCGTaskId> DataDependencyTasks;
 
 	// Schedule generation of original component if is is non-partitioned, or if it has nodes that will execute at the Unbounded level.
 	FPCGTaskId OriginalComponentTask = InvalidPCGTaskId;
@@ -522,7 +528,7 @@ FPCGTaskId UPCGSubsystem::ScheduleComponent(UPCGComponent* PCGComponent, EPCGHiG
 		OriginalComponentTask = PCGComponent->CreateGenerateTask(/*bForce=*/bSave, InDependencies);
 		if (OriginalComponentTask != InvalidPCGTaskId)
 		{
-			AllTasks.Add(OriginalComponentTask);
+			DataDependencyTasks.Add(OriginalComponentTask);
 		}
 	}
 
@@ -569,14 +575,15 @@ FPCGTaskId UPCGSubsystem::ScheduleComponent(UPCGComponent* PCGComponent, EPCGHiG
 			return LocalComponent->GenerateInternal(/*bForce=*/bSave, LocalComponent->GetGenerationGrid(), EPCGComponentGenerationTrigger::GenerateOnDemand, Dependencies);
 		};
 
-		AllTasks.Append(ActorAndComponentMapping.DispatchToRegisteredLocalComponents(PCGComponent, LocalGenerateTask));
+		ExecutionDependencyTasks.Append(ActorAndComponentMapping.DispatchToRegisteredLocalComponents(PCGComponent, LocalGenerateTask));
 	}
 
-	if (!AllTasks.IsEmpty())
+	if (!ExecutionDependencyTasks.IsEmpty() || !DataDependencyTasks.IsEmpty())
 	{
 		TWeakObjectPtr<UPCGComponent> ComponentPtr(PCGComponent);
 
-		return GraphExecutor->ScheduleGenericWithContext([ComponentPtr](FPCGContext* Context) {
+		return GraphExecutor->ScheduleGenericWithContext([ComponentPtr](FPCGContext* Context)
+		{
 			if (UPCGComponent* Component = ComponentPtr.Get())
 			{
 				// If the component is not valid anymore, just early out.
@@ -590,7 +597,7 @@ FPCGTaskId UPCGSubsystem::ScheduleComponent(UPCGComponent* PCGComponent, EPCGHiG
 			}
 
 			return true;
-			}, PCGComponent, AllTasks);
+		}, PCGComponent, DataDependencyTasks, ExecutionDependencyTasks);
 	}
 	else
 	{
@@ -727,16 +734,16 @@ TStatId UPCGSubsystem::GetStatId() const
 	RETURN_QUICK_DECLARE_CYCLE_STAT(UPCGSubsystem, STATGROUP_Tickables);
 }
 
-FPCGTaskId UPCGSubsystem::ScheduleGeneric(TFunction<bool()> InOperation, UPCGComponent* Component, const TArray<FPCGTaskId>& TaskDependencies)
+FPCGTaskId UPCGSubsystem::ScheduleGeneric(TFunction<bool()> InOperation, UPCGComponent* SourceComponent, const TArray<FPCGTaskId>& TaskExecutionDependencies)
 {
 	check(GraphExecutor);
-	return GraphExecutor->ScheduleGeneric(InOperation, Component, TaskDependencies);
+	return GraphExecutor->ScheduleGeneric(InOperation, SourceComponent, TaskExecutionDependencies);
 }
 
-FPCGTaskId UPCGSubsystem::ScheduleGenericWithContext(TFunction<bool(FPCGContext*)> InOperation, UPCGComponent* Component, const TArray<FPCGTaskId>& TaskDependencies, bool bConsumeInputData)
+FPCGTaskId UPCGSubsystem::ScheduleGenericWithContext(TFunction<bool(FPCGContext*)> InOperation, UPCGComponent* SourceComponent, const TArray<FPCGTaskId>& TaskExecutionDependencies, const TArray<FPCGTaskId>& TaskDataDependencies)
 {
 	check(GraphExecutor);
-	return GraphExecutor->ScheduleGenericWithContext(InOperation, Component, TaskDependencies, bConsumeInputData);
+	return GraphExecutor->ScheduleGenericWithContext(InOperation, SourceComponent, TaskExecutionDependencies, TaskDataDependencies);
 }
 
 void UPCGSubsystem::CancelGeneration(UPCGComponent* Component)
diff --git a/Engine/Plugins/Experimental/PCG/Source/PCG/Public/Elements/PCGActorSelector.h b/Engine/Plugins/Experimental/PCG/Source/PCG/Public/Elements/PCGActorSelector.h
--- a/Engine/Plugins/Experimental/PCG/Source/PCG/Public/Elements/PCGActorSelector.h
+++ b/Engine/Plugins/Experimental/PCG/Source/PCG/Public/Elements/PCGActorSelector.h
@@ -74,41 +74,46 @@ struct FPCGActorSelectionKey
 	const UClass* OptionalExtraDependency = nullptr;
 };
 
+/** Helper struct for organizing queries against the world to gather actors. */
 USTRUCT(BlueprintType)
 struct FPCGActorSelectorSettings
 {
 	GENERATED_BODY()
 
 	/** Which actors to consider. */
-	UPROPERTY(BlueprintReadWrite, EditAnywhere, Category = Settings, meta = (EditCondition = "bShowActorFilter", EditConditionHides, HideEditConditionToggle))
+	UPROPERTY(BlueprintReadWrite, EditAnywhere, Category = "Actor Selector Settings", meta = (EditCondition = "bShowActorFilter", EditConditionHides, HideEditConditionToggle))
 	EPCGActorFilter ActorFilter = EPCGActorFilter::Self;
 
-	UPROPERTY(BlueprintReadWrite, EditAnywhere, Category = Settings, meta = (EditCondition = "ActorFilter==EPCGActorFilter::AllWorldActors", EditConditionHides))
+	/** Filters out actors that do not overlap the source component bounds. */
+	UPROPERTY(BlueprintReadWrite, EditAnywhere, Category = "Actor Selector Settings", meta = (EditCondition = "ActorFilter==EPCGActorFilter::AllWorldActors", EditConditionHides))
 	bool bMustOverlapSelf = false;
 
 	/** Whether to consider child actors. */
-	UPROPERTY(BlueprintReadWrite, EditAnywhere, Category = Settings, meta = (EditCondition = "bShowIncludeChildren && ActorFilter!=EPCGActorFilter::AllWorldActors", EditConditionHides))
+	UPROPERTY(BlueprintReadWrite, EditAnywhere, Category = "Actor Selector Settings", meta = (EditCondition = "bShowIncludeChildren && ActorFilter!=EPCGActorFilter::AllWorldActors", EditConditionHides))
 	bool bIncludeChildren = false;
 
 	/** Enables/disables fine-grained actor filtering options. */
-	UPROPERTY(BlueprintReadWrite, EditAnywhere, Category = Settings, meta = (EditCondition = "ActorFilter!=EPCGActorFilter::AllWorldActors && bIncludeChildren", EditConditionHides))
+	UPROPERTY(BlueprintReadWrite, EditAnywhere, Category = "Actor Selector Settings", meta = (EditCondition = "ActorFilter!=EPCGActorFilter::AllWorldActors && bIncludeChildren", EditConditionHides))
 	bool bDisableFilter = false;
 
-	UPROPERTY(BlueprintReadWrite, EditAnywhere, Category = Settings, meta = (EditCondition = "bShowActorSelection && (ActorFilter==EPCGActorFilter::AllWorldActors || (bIncludeChildren && !bDisableFilter))", EditConditionHides))
+	/** How to select when filtering actors. */
+	UPROPERTY(BlueprintReadWrite, EditAnywhere, Category = "Actor Selector Settings", meta = (EditCondition = "bShowActorSelection && (ActorFilter==EPCGActorFilter::AllWorldActors || (bIncludeChildren && !bDisableFilter))", EditConditionHides))
 	EPCGActorSelection ActorSelection = EPCGActorSelection::ByTag;
 
-	UPROPERTY(BlueprintReadWrite, EditAnywhere, Category = Settings, meta = (EditCondition = "bShowActorSelection && (ActorFilter==EPCGActorFilter::AllWorldActors || (bIncludeChildren && !bDisableFilter)) && ActorSelection==EPCGActorSelection::ByTag", EditConditionHides))
+	/** Tag to match against when filtering actors. */
+	UPROPERTY(BlueprintReadWrite, EditAnywhere, Category = "Actor Selector Settings", meta = (EditCondition = "bShowActorSelection && (ActorFilter==EPCGActorFilter::AllWorldActors || (bIncludeChildren && !bDisableFilter)) && ActorSelection==EPCGActorSelection::ByTag", EditConditionHides))
 	FName ActorSelectionTag;
 
-	UPROPERTY(BlueprintReadWrite, EditAnywhere, Category = Settings, meta = (EditCondition = "bShowActorSelectionClass && bShowActorSelection && (ActorFilter==EPCGActorFilter::AllWorldActors || (bIncludeChildren && !bDisableFilter)) && ActorSelection==EPCGActorSelection::ByClass", EditConditionHides, AllowAbstract = "true"))
+	/** Actor class to match against when filtering actors. */
+	UPROPERTY(BlueprintReadWrite, EditAnywhere, Category = "Actor Selector Settings", meta = (EditCondition = "bShowActorSelectionClass && bShowActorSelection && (ActorFilter==EPCGActorFilter::AllWorldActors || (bIncludeChildren && !bDisableFilter)) && ActorSelection==EPCGActorSelection::ByClass", EditConditionHides, AllowAbstract = "true"))
 	TSubclassOf<AActor> ActorSelectionClass;
 
 	/** If true processes all matching actors, otherwise returns data from first match. */
-	UPROPERTY(BlueprintReadWrite, EditAnywhere, Category = Settings, meta = (EditCondition = "bShowSelectMultiple && ActorFilter==EPCGActorFilter::AllWorldActors && ActorSelection!=EPCGActorSelection::ByName", EditConditionHides))
+	UPROPERTY(BlueprintReadWrite, EditAnywhere, Category = "Actor Selector Settings", meta = (EditCondition = "bShowSelectMultiple && ActorFilter==EPCGActorFilter::AllWorldActors && ActorSelection!=EPCGActorSelection::ByName", EditConditionHides))
 	bool bSelectMultiple = false;
 
-	/** If true, ignores results found from within this actor's hierarchy */
-	UPROPERTY(BlueprintReadWrite, EditAnywhere, Category = Settings, meta = (EditCondition = "ActorFilter==EPCGActorFilter::AllWorldActors", EditConditionHides))
+	/** If true, ignores results found from within this actor's hierarchy. */
+	UPROPERTY(BlueprintReadWrite, EditAnywhere, Category = "Actor Selector Settings", meta = (EditCondition = "bShowIgnoreSelfAndChildren && ActorFilter==EPCGActorFilter::AllWorldActors", EditConditionHides))
 	bool bIgnoreSelfAndChildren = false;
 
 	// Properties used to hide some fields when used in different contexts
@@ -127,6 +132,9 @@ struct FPCGActorSelectorSettings
 	UPROPERTY(Transient, meta = (EditCondition = false, EditConditionHides))
 	bool bShowSelectMultiple = true;
 
+	UPROPERTY(Transient, meta = (EditCondition = false, EditConditionHides))
+	bool bShowIgnoreSelfAndChildren = true;
+
 #if WITH_EDITOR	
 	FText GetTaskNameSuffix() const;
 	FName GetTaskName(const FText& Prefix) const;
diff --git a/Engine/Plugins/Experimental/PCG/Source/PCG/Public/Elements/PCGDataFromActor.h b/Engine/Plugins/Experimental/PCG/Source/PCG/Public/Elements/PCGDataFromActor.h
--- a/Engine/Plugins/Experimental/PCG/Source/PCG/Public/Elements/PCGDataFromActor.h
+++ b/Engine/Plugins/Experimental/PCG/Source/PCG/Public/Elements/PCGDataFromActor.h
@@ -11,11 +11,11 @@
 UENUM()
 enum class EPCGGetDataFromActorMode : uint8
 {
-	ParseActorComponents,
-	GetSinglePoint,
-	GetDataFromProperty,
-	GetDataFromPCGComponent,
-	GetDataFromPCGComponentOrParseComponents
+	ParseActorComponents UMETA(Tooltip = "Parse the found actor(s) for relevant components such as Primitives, Splines, and Volumes."),
+	GetSinglePoint UMETA(Tooltip = "Produces a single point per actor with the actor transform and bounds."),
+	GetDataFromProperty UMETA(Tooltip = "Gets a data collection from an actor property."),
+	GetDataFromPCGComponent UMETA(Tooltip = "Copy generated output from other PCG components on the found actor(s)."),
+	GetDataFromPCGComponentOrParseComponents UMETA(Tooltip = "Attempts to copy generated output from other PCG components on the found actor(s), otherwise, falls back to parsing actor components.")
 };
 
 /** Builds a collection of PCG-compatible data from the selected actors. */
@@ -33,6 +33,7 @@ public:
 	virtual EPCGSettingsType GetType() const override { return EPCGSettingsType::Spatial; }
 	virtual void GetTrackedActorKeys(FPCGActorSelectionKeyToSettingsMap& OutKeysToSettings, TArray<TObjectPtr<const UPCGGraph>>& OutVisitedGraphs) const override;
 	virtual bool HasDynamicPins() const override { return true; }
+	virtual void ApplyDeprecation(UPCGNode* InOutNode) override;
 #endif
 	virtual EPCGDataType GetCurrentPinTypes(const UPCGPin* InPin) const override;
 
@@ -63,31 +64,49 @@ public:
 	/** Override this to change the default value the selector will revert to when changing the actor selection type */
 	virtual TSubclassOf<AActor> GetDefaultActorSelectorClass() const;
 
+	/** Describes which actors to select for data collection. */
 	UPROPERTY(BlueprintReadWrite, EditAnywhere, Category = Settings, meta = (ShowOnlyInnerProperties))
 	FPCGActorSelectorSettings ActorSelector;
 
-	UPROPERTY(BlueprintReadWrite, EditAnywhere, Category = Settings, meta = (EditCondition = bDisplayModeSettings, EditConditionHides, HideEditConditionToggle))
+	/** Describes what kind of data we will collect from the found actor(s). */
+	UPROPERTY(BlueprintReadWrite, EditAnywhere, Category = "Data Retrieval Settings", meta = (EditCondition = bDisplayModeSettings, EditConditionHides, HideEditConditionToggle))
 	EPCGGetDataFromActorMode Mode = EPCGGetDataFromActorMode::ParseActorComponents;
 
-	UPROPERTY(BlueprintReadWrite, EditAnywhere, Category = Settings, meta = (EditCondition = "Mode == EPCGGetDataFromActorMode::GetDataFromPCGComponent || Mode == EPCGGetDataFromActorMode::GetDataFromPCGComponentOrParseComponents", EditConditionHides))
+	/** Also produces a single point data at the actor location. */
+	UPROPERTY(BlueprintReadWrite, EditAnywhere, Category = "Data Retrieval Settings", meta = (EditCondition = "Mode == EPCGGetDataFromActorMode::GetDataFromPCGComponent || Mode == EPCGGetDataFromActorMode::GetDataFromPCGComponentOrParseComponents", EditConditionHides))
 	bool bAlsoOutputSinglePointData = false;
 
-	UPROPERTY(BlueprintReadWrite, EditAnywhere, Category = Settings, meta = (EditCondition = "Mode == EPCGGetDataFromActorMode::GetSinglePoint", EditConditionHides))
+	/** Only get data from components which overlap with the bounds of your source component. */
+	UPROPERTY(BlueprintReadWrite, EditAnywhere, Category = "Data Retrieval Settings", meta = (EditCondition = "Mode == EPCGGetDataFromActorMode::GetDataFromPCGComponent || Mode == EPCGGetDataFromActorMode::GetDataFromPCGComponentOrParseComponents", EditConditionHides))
+	bool bComponentsMustOverlapSelf = true;
+
+	/** Get data from all grid sizes if there is a partitioned PCG component on the actor, instead of a specific set of grid sizes. */
+	UPROPERTY(BlueprintReadWrite, EditAnywhere, Category = "Data Retrieval Settings", meta = (EditCondition = "Mode == EPCGGetDataFromActorMode::GetDataFromPCGComponent || Mode == EPCGGetDataFromActorMode::GetDataFromPCGComponentOrParseComponents", EditConditionHides))
+	bool bGetDataOnAllGrids = true;
+
+	/** Select which grid sizes to consider when collecting data from partitioned PCG components. */
+	UPROPERTY(BlueprintReadWrite, EditAnywhere, Category = "Data Retrieval Settings", meta = (Bitmask, BitmaskEnum = "/Script/PCG.EPCGHiGenGrid", EditCondition = "!bGetDataOnAllGrids && Mode == EPCGGetDataFromActorMode::GetDataFromPCGComponent || Mode == EPCGGetDataFromActorMode::GetDataFromPCGComponentOrParseComponents", EditConditionHides))
+	int32 AllowedGrids = int32(EPCGHiGenGrid::Uninitialized);
+
+	/** Merges all the single point data outputs into a single point data. */
+	UPROPERTY(BlueprintReadWrite, EditAnywhere, Category = "Data Retrieval Settings", meta = (EditCondition = "Mode == EPCGGetDataFromActorMode::GetSinglePoint", EditConditionHides))
 	bool bMergeSinglePointData = false;
 
 	// This can be set false by inheriting nodes to hide the 'Mode' property.
 	UPROPERTY(Transient, meta = (EditCondition = false, EditConditionHides))
 	bool bDisplayModeSettings = true;
 
-	UPROPERTY(BlueprintReadWrite, EditAnywhere, Category = Settings, meta = (EditCondition = "Mode == EPCGGetDataFromActorMode::GetDataFromPCGComponent || Mode == EPCGGetDataFromActorMode::GetDataFromPCGComponentOrParseComponents", EditConditionHides))
+	/** Provide pin names to match against the found component output pins. Data will automatically be wired to the expected pin if the name comparison succeeds. All unmatched pins will go into the standard out pin. */
+	UPROPERTY(BlueprintReadWrite, EditAnywhere, Category = "Data Retrieval Settings", meta = (EditCondition = "Mode == EPCGGetDataFromActorMode::GetDataFromPCGComponent || Mode == EPCGGetDataFromActorMode::GetDataFromPCGComponentOrParseComponents", EditConditionHides))
 	TArray<FName> ExpectedPins;
 
-	UPROPERTY(BlueprintReadWrite, EditAnywhere, Category = Settings, meta = (EditCondition = "Mode == EPCGGetDataFromActorMode::GetDataFromProperty", EditConditionHides))
+	/** The property name on the found actor to create a data collection from. */
+	UPROPERTY(BlueprintReadWrite, EditAnywhere, Category = "Data Retrieval Settings", meta = (EditCondition = "Mode == EPCGGetDataFromActorMode::GetDataFromProperty", EditConditionHides))
 	FName PropertyName = NAME_None;
 
 #if WITH_EDITORONLY_DATA
 	/** If this is checked, found actors that are outside component bounds will not trigger a refresh. Only works for tags for now in editor. */
-	UPROPERTY(BlueprintReadOnly, EditAnywhere, Category = Settings)
+	UPROPERTY(BlueprintReadOnly, EditAnywhere, Category = "Data Retrieval Settings")
 	bool bTrackActorsOnlyWithinBounds = true;
 #endif // WITH_EDITORONLY_DATA
 };
diff --git a/Engine/Plugins/Experimental/PCG/Source/PCG/Public/Elements/PCGTypedGetter.h b/Engine/Plugins/Experimental/PCG/Source/PCG/Public/Elements/PCGTypedGetter.h
--- a/Engine/Plugins/Experimental/PCG/Source/PCG/Public/Elements/PCGTypedGetter.h
+++ b/Engine/Plugins/Experimental/PCG/Source/PCG/Public/Elements/PCGTypedGetter.h
@@ -23,7 +23,7 @@ public:
 	//~Begin UPCGSettings interface
 #if WITH_EDITOR
 	virtual FName GetDefaultNodeName() const override { return FName(TEXT("GetLandscapeData")); }
-	virtual FText GetDefaultNodeTitle() const override { return NSLOCTEXT("PCGGetLandscapeSettings", "NodeTitle", "Get Landscape Data"); }
+	virtual FText GetDefaultNodeTitle() const override { return NSLOCTEXT("PCGGetLandscapeElement", "NodeTitle", "Get Landscape Data"); }
 	virtual FText GetNodeTooltipText() const override;
 #endif
 
@@ -72,7 +72,7 @@ public:
 	//~Begin UPCGSettings interface
 #if WITH_EDITOR
 	virtual FName GetDefaultNodeName() const override { return FName(TEXT("GetSplineData")); }
-	virtual FText GetDefaultNodeTitle() const override { return NSLOCTEXT("PCGGetSplineSettings", "NodeTitle", "Get Spline Data"); }
+	virtual FText GetDefaultNodeTitle() const override { return NSLOCTEXT("PCGGetSplineElement", "NodeTitle", "Get Spline Data"); }
 	virtual FText GetNodeTooltipText() const override;
 #endif
 
@@ -98,7 +98,7 @@ public:
 	//~Begin UPCGSettings interface
 #if WITH_EDITOR
 	virtual FName GetDefaultNodeName() const override { return FName(TEXT("GetVolumeData")); }
-	virtual FText GetDefaultNodeTitle() const override { return NSLOCTEXT("PCGGetVolumeSettings", "NodeTitle", "Get Volume Data"); }
+	virtual FText GetDefaultNodeTitle() const override { return NSLOCTEXT("PCGGetVolumeElement", "NodeTitle", "Get Volume Data"); }
 	virtual FText GetNodeTooltipText() const override;
 #endif
 
@@ -124,7 +124,7 @@ public:
 	//~Begin UPCGSettings interface
 #if WITH_EDITOR
 	virtual FName GetDefaultNodeName() const override { return FName(TEXT("GetPrimitiveData")); }
-	virtual FText GetDefaultNodeTitle() const override { return NSLOCTEXT("PCGGetPrimitiveSettings", "NodeTitle", "Get Primitive Data"); }
+	virtual FText GetDefaultNodeTitle() const override { return NSLOCTEXT("PCGGetPrimitiveElement", "NodeTitle", "Get Primitive Data"); }
 	virtual FText GetNodeTooltipText() const override;
 #endif
 
@@ -137,3 +137,34 @@ public:
 	virtual EPCGDataType GetDataFilter() const override { return EPCGDataType::Primitive; }
 	//~End UPCGDataFromActorSettings
 };
+
+/**
+ * Builds a collection of data from other PCG components on the selected actors. Automatically tags each output with the grid size it was collected
+ * from, prefixed by "PCG_GridSize_" (e.g. PCG_GridSize_12800).
+ *
+ * Note: a component cannot get component data from itself or other components in its hierarchy, as it could create a circular dependency.
+ */
+UCLASS(BlueprintType, ClassGroup = (Procedural))
+class PCG_API UPCGGetPCGComponentSettings : public UPCGDataFromActorSettings
+{
+	GENERATED_BODY()
+
+public:
+	UPCGGetPCGComponentSettings();
+
+	//~Begin UPCGSettings interface
+#if WITH_EDITOR
+	virtual FName GetDefaultNodeName() const override { return FName(TEXT("GetPCGComponentData")); }
+	virtual FText GetDefaultNodeTitle() const override { return NSLOCTEXT("PCGGetPCGComponentElement", "NodeTitle", "Get PCG Component Data"); }
+	virtual FText GetNodeTooltipText() const override;
+#endif
+
+protected:
+	virtual TArray<FPCGPinProperties> OutputPinProperties() const override;
+	//~End UPCGSettings
+
+public:
+	//~Begin UPCGDataFromActorSettings interface
+	virtual EPCGDataType GetDataFilter() const override { return EPCGDataType::Any; }
+	//~End UPCGDataFromActorSettings
+};
diff --git a/Engine/Plugins/Experimental/PCG/Source/PCG/Public/PCGCommon.h b/Engine/Plugins/Experimental/PCG/Source/PCG/Public/PCGCommon.h
--- a/Engine/Plugins/Experimental/PCG/Source/PCG/Public/PCGCommon.h
+++ b/Engine/Plugins/Experimental/PCG/Source/PCG/Public/PCGCommon.h
@@ -161,7 +161,7 @@ enum class EPCGCoordinateSpace : uint8
 };
 
 /** Describes one or more target execution grids. */
-UENUM(meta = (Bitflags))
+UENUM(meta = (Bitflags, UseEnumValuesAsMaskValuesInEditor = "true"))
 enum class EPCGHiGenGrid : uint32
 {
 	Uninitialized = 0 UMETA(Hidden),
diff --git a/Engine/Plugins/Experimental/PCG/Source/PCG/Public/PCGCustomVersion.h b/Engine/Plugins/Experimental/PCG/Source/PCG/Public/PCGCustomVersion.h
--- a/Engine/Plugins/Experimental/PCG/Source/PCG/Public/PCGCustomVersion.h
+++ b/Engine/Plugins/Experimental/PCG/Source/PCG/Public/PCGCustomVersion.h
@@ -65,6 +65,9 @@ struct PCG_API FPCGCustomVersion
 		// Added 'out' filter pins on filter by tag & by type
 		UpdateFilterNodeOutputPins = 18,
 
+		// Added 'bComponentsMustOverlapSelf' to GetActorData when the mode collects PCG component data
+		GetPCGComponentDataMustOverlapSourceComponentByDefault = 19,
+
 		// -----<new versions can be added above this line>-------------------------------------------------
 		VersionPlusOne,
 		LatestVersion = VersionPlusOne - 1
diff --git a/Engine/Plugins/Experimental/PCG/Source/PCG/Public/PCGSubsystem.h b/Engine/Plugins/Experimental/PCG/Source/PCG/Public/PCGSubsystem.h
--- a/Engine/Plugins/Experimental/PCG/Source/PCG/Public/PCGSubsystem.h
+++ b/Engine/Plugins/Experimental/PCG/Source/PCG/Public/PCGSubsystem.h
@@ -117,19 +117,19 @@ public:
 	FPCGTaskId ScheduleGraph(UPCGComponent* SourceComponent, const TArray<FPCGTaskId>& Dependencies);
 
 	/** General job scheduling
-	*  @param InOperation:       Callback that returns true if the task is done, false otherwise.
-	*  @param SourceComponent: PCG component associated with this task. Can be null.
-	*  @param TaskDependencies:  List of all the dependencies for this task.
+	*  @param InOperation:               Callback that returns true if the task is done, false otherwise.
+	*  @param SourceComponent:           PCG component associated with this task. Can be null.
+	*  @param TaskExecutionDependencies: Task will wait on these tasks to execute and won't take their output data as input.
 	*/
-	FPCGTaskId ScheduleGeneric(TFunction<bool()> InOperation, UPCGComponent* SourceComponent, const TArray<FPCGTaskId>& TaskDependencies);
+	FPCGTaskId ScheduleGeneric(TFunction<bool()> InOperation, UPCGComponent* SourceComponent, const TArray<FPCGTaskId>& TaskExecutionDependencies);
 
 	/** General job scheduling with context
-	*  @param InOperation:       Callback that takes a Context as argument and returns true if the task is done, false otherwise.
-	*  @param SourceComponent: PCG component associated with this task. Can be null.
-	*  @param TaskDependencies:  List of all the dependencies for this task.
-	*  @param bConsumeInputData: If your task need a context, but don't need the input data, set this flag to false. Default is true.
+	*  @param InOperation:               Callback that takes a Context as argument and returns true if the task is done, false otherwise.
+	*  @param SourceComponent:           PCG component associated with this task. Can be null.
+	*  @param TaskExecutionDependencies: Task will wait on these tasks to execute and won't take their output data as input.
+	*  @param TaskDataDependencies:      Task will wait on these tasks to execute and will take their output data as input.
 	*/
-	FPCGTaskId ScheduleGenericWithContext(TFunction<bool(FPCGContext*)> InOperation, UPCGComponent* SourceComponent, const TArray<FPCGTaskId>& TaskDependencies, bool bConsumeInputData = true);
+	FPCGTaskId ScheduleGenericWithContext(TFunction<bool(FPCGContext*)> InOperation, UPCGComponent* SourceComponent, const TArray<FPCGTaskId>& TaskExecutionDependencies, const TArray<FPCGTaskId>& TaskDataDependencies);
 
 	/** Cancels currently running generation */
 	void CancelGeneration(UPCGComponent* Component);
