[PCG] Fix debug object selection getting lost when moving debug flag around.

Only force removal of PCG components if higen grid sizes change. We added a new change type to communicate this.

To make this work, we also made the runtime gen system support a more refresh-like path - previously a refresh would completely clear out the PAs and recreate them. Now the PAs are marked dirty and scheduled for generation which is much nicer.

Also ensure PA grid updated when changing & undoing Grid Size changes, which required adding some code here and there to the refresh paths that respond to the GenerationGrid change type.

diff --git a/Engine/Plugins/Experimental/PCG/Source/PCG/Private/Elements/PCGHiGenGridSize.cpp b/Engine/Plugins/Experimental/PCG/Source/PCG/Private/Elements/PCGHiGenGridSize.cpp
--- a/Engine/Plugins/Experimental/PCG/Source/PCG/Private/Elements/PCGHiGenGridSize.cpp
+++ b/Engine/Plugins/Experimental/PCG/Source/PCG/Private/Elements/PCGHiGenGridSize.cpp
@@ -103,7 +103,7 @@ EPCGChangeType UPCGHiGenGridSizeSettings::GetChangeTypeForProperty(const FName&
 	if (InPropertyName == GET_MEMBER_NAME_CHECKED(UPCGHiGenGridSizeSettings, bEnabled)
 		|| InPropertyName == GET_MEMBER_NAME_CHECKED(UPCGHiGenGridSizeSettings, HiGenGridSize))
 	{
-		ChangeType |= EPCGChangeType::Structural;
+		ChangeType |= EPCGChangeType::Structural | EPCGChangeType::GenerationGrid;
 	}
 
 	return ChangeType;
diff --git a/Engine/Plugins/Experimental/PCG/Source/PCG/Private/PCGComponent.cpp b/Engine/Plugins/Experimental/PCG/Source/PCG/Private/PCGComponent.cpp
--- a/Engine/Plugins/Experimental/PCG/Source/PCG/Private/PCGComponent.cpp
+++ b/Engine/Plugins/Experimental/PCG/Source/PCG/Private/PCGComponent.cpp
@@ -141,7 +141,7 @@ void UPCGComponent::SetGraphInterfaceLocal(UPCGGraphInterface* InGraphInterface)
 	if (ensure(GraphInstance))
 	{
 		GraphInstance->SetGraph(InGraphInterface);
-		RefreshAfterGraphChanged(GraphInstance, /*bIsStructural=*/true, /*bDirtyInputs=*/true);
+		RefreshAfterGraphChanged(GraphInstance, EPCGChangeType::Structural | EPCGChangeType::GenerationGrid);
 	}
 }
 
@@ -1415,13 +1415,10 @@ TStructOnScope<FActorComponentInstanceData> UPCGComponent::GetComponentInstanceD
 
 void UPCGComponent::OnGraphChanged(UPCGGraphInterface* InGraph, EPCGChangeType ChangeType)
 {
-	const bool bIsStructural = ((ChangeType & (EPCGChangeType::Edge | EPCGChangeType::Structural)) != EPCGChangeType::None);
-	const bool bDirtyInputs = bIsStructural || ((ChangeType & EPCGChangeType::Input) != EPCGChangeType::None);
-
-	RefreshAfterGraphChanged(InGraph, bIsStructural, bDirtyInputs);
+	RefreshAfterGraphChanged(InGraph, ChangeType);
 }
 
-void UPCGComponent::RefreshAfterGraphChanged(UPCGGraphInterface* InGraph, bool bIsStructural, bool bDirtyInputs)
+void UPCGComponent::RefreshAfterGraphChanged(UPCGGraphInterface* InGraph, EPCGChangeType ChangeType)
 {
 	if (InGraph != GraphInstance)
 	{
@@ -1430,6 +1427,9 @@ void UPCGComponent::RefreshAfterGraphChanged(UPCGGraphInterface* InGraph, bool b
 
 	const bool bHasGraph = (InGraph && InGraph->GetGraph());
 
+	const bool bIsStructural = ((ChangeType & (EPCGChangeType::Edge | EPCGChangeType::Structural)) != EPCGChangeType::None);
+	const bool bDirtyInputs = bIsStructural || ((ChangeType & EPCGChangeType::Input) != EPCGChangeType::None);
+
 #if WITH_EDITOR
 	// In editor, since we've changed the graph, we might have changed the tracked actor tags as well
 	if (!PCGHelpers::IsRuntimeOrPIE())
@@ -1445,9 +1445,10 @@ void UPCGComponent::RefreshAfterGraphChanged(UPCGGraphInterface* InGraph, bool b
 		}
 
 		DirtyGenerated(bDirtyInputs ? (EPCGComponentDirtyFlag::Actor | EPCGComponentDirtyFlag::Landscape) : EPCGComponentDirtyFlag::None);
+
 		if (bHasGraph)
 		{
-			Refresh(bIsStructural);
+			Refresh(ChangeType);
 		}
 		else
 		{
@@ -1456,6 +1457,7 @@ void UPCGComponent::RefreshAfterGraphChanged(UPCGGraphInterface* InGraph, bool b
 		}
 
 		ClearInspectionData();
+
 		return;
 	}
 #endif
@@ -1464,7 +1466,7 @@ void UPCGComponent::RefreshAfterGraphChanged(UPCGGraphInterface* InGraph, bool b
 	{
 		if (UPCGSubsystem* Subsystem = GetSubsystem())
 		{
-			Subsystem->RefreshRuntimeGenComponent(this, /*bRemovePartitionActors=*/true);
+			Subsystem->RefreshRuntimeGenComponent(this, ChangeType);
 		}
 	}
 	else
@@ -1493,7 +1495,7 @@ void UPCGComponent::PreEditChange(FProperty* PropertyAboutToChange)
 			if (UPCGSubsystem* Subsystem = GetSubsystem())
 			{
 				// When toggling off of GenerateAtRuntime, we should flush the RuntimeGenScheduler state for this component.
-				Subsystem->RefreshRuntimeGenComponent(this, /*bRemovePartitionActors=*/true);
+				Subsystem->RefreshRuntimeGenComponent(this, EPCGChangeType::GenerationGrid);
 			}
 
 			// Reset to the the editing mode we were in before entering GenerateAtRuntime mode.
@@ -1674,7 +1676,7 @@ void UPCGComponent::PostEditUndo()
 		// operation removes the component, a valid refresh task ID is set but the refresh task itself will fail
 		// and leave the valid task ID hanging on the component. Forcing here means if we later retrieve this state
 		// from the undo/redo buffer, the refresh will be forced which will reset the state.
-		Refresh(/*bIsStructural=*/true, /*bCancelExistingRefresh=*/true);
+		Refresh(EPCGChangeType::Structural, /*bCancelExistingRefresh=*/true);
 	}
 
 	Super::PostEditUndo();
@@ -1931,7 +1933,7 @@ bool UPCGComponent::HasNodeProducedData(const UPCGNode* InNode, const FPCGStack&
 	return StacksThatProducedData && StacksThatProducedData->Contains(Stack);
 }
 
-void UPCGComponent::Refresh(bool bStructural, bool bCancelExistingRefresh)
+void UPCGComponent::Refresh(EPCGChangeType ChangeType, bool bCancelExistingRefresh)
 {
 	// Disable auto-refreshing on preview actors until we have something more robust on the execution side.
 	if (GetOwner() && GetOwner()->bIsEditorPreviewActor)
@@ -1944,8 +1946,7 @@ void UPCGComponent::Refresh(bool bStructural, bool bCancelExistingRefresh)
 	{
 		if (UPCGSubsystem* Subsystem = GetSubsystem())
 		{
-			// TODO: We only need to remove PAs if the grid sizes changed. Is there a reliable way to know that?
-			Subsystem->RefreshRuntimeGenComponent(this, /*bRemovePartitionActors=*/true);
+			Subsystem->RefreshRuntimeGenComponent(this, ChangeType);
 		}
 
 		return;
@@ -1979,6 +1980,7 @@ void UPCGComponent::Refresh(bool bStructural, bool bCancelExistingRefresh)
 		// Cancel an already existing generation if either the change is structural in nature (which requires a recompilation, so a full-rescheduling)
 		// or if the generation is already started
 		const bool bGenerationWasInProgress = IsGenerationInProgress();
+		const bool bStructural = !!(ChangeType & EPCGChangeType::Structural);
 		bool bNeedToCancelCurrentTasks = (CurrentGenerationTask != InvalidPCGTaskId && (bStructural || bGenerationWasInProgress));
 
 		// Cancel an already existing refresh if caller allows this
@@ -1998,7 +2000,10 @@ void UPCGComponent::Refresh(bool bStructural, bool bCancelExistingRefresh)
 		// then the bGenerated flag will be false, which will prevent a subsequent update here
 		if (CurrentRefreshTask == InvalidPCGTaskId && CurrentCleanupTask == InvalidPCGTaskId)
 		{
-			CurrentRefreshTask = Subsystem->ScheduleRefresh(this, bGenerationWasInProgress);
+			// Always force a regeneration if generation grids might be affected - refresh is insufficient.
+			const bool bGridChanged = !!(ChangeType & EPCGChangeType::GenerationGrid);
+			const bool bForceRegen = bGenerationWasInProgress || bGridChanged;
+			CurrentRefreshTask = Subsystem->ScheduleRefresh(this, bForceRegen, /*bForceCleanup=*/bGridChanged);
 		}
 	}
 }
@@ -2090,7 +2095,7 @@ bool UPCGComponent::IsObjectTracked(const TSoftObjectPtr<UObject>& InObjectPtr,
 	return bFound;
 }
 
-void UPCGComponent::OnRefresh(bool bForceRefresh)
+void UPCGComponent::OnRefresh(bool bForceRefresh, bool bForceCleanup)
 {
 	check(!IsManagedByRuntimeGenSystem());
 
@@ -2103,13 +2108,10 @@ void UPCGComponent::OnRefresh(bool bForceRefresh)
 	const bool bWasGenerated = bGenerated;
 	const bool bWasGeneratedOrGenerating = bWasGenerated || bForceRefresh;
 
-	if (IsPartitioned())
+	// If we are partitioned but we have resources, we need to force a cleanup
+	if (bForceCleanup || (IsPartitioned() && !GeneratedResources.IsEmpty()))
 	{
-		// If we are partitioned but we have resources, we need to force a cleanup
-		if (!GeneratedResources.IsEmpty())
-		{
-			CleanupLocalImmediate(/*bRemoveComponents=*/true);
-		}
+		CleanupLocalImmediate(/*bRemoveComponents=*/true);
 	}
 
 	if (Subsystem)
diff --git a/Engine/Plugins/Experimental/PCG/Source/PCG/Private/PCGGraph.cpp b/Engine/Plugins/Experimental/PCG/Source/PCG/Private/PCGGraph.cpp
--- a/Engine/Plugins/Experimental/PCG/Source/PCG/Private/PCGGraph.cpp
+++ b/Engine/Plugins/Experimental/PCG/Source/PCG/Private/PCGGraph.cpp
@@ -551,15 +551,22 @@ void UPCGGraph::OnNodeAdded(UPCGNode* InNode)
 void UPCGGraph::OnNodesAdded(TArrayView<UPCGNode*> InNodes)
 {
 #if WITH_EDITOR
+	EPCGChangeType ChangeType = EPCGChangeType::Structural;
+
 	for (UPCGNode* Node : InNodes)
 	{
 		if (Node)
 		{
 			Node->OnNodeChangedDelegate.AddUObject(this, &UPCGGraph::OnNodeChanged);
+
+			if (Node->GetSettings() && Node->GetSettings()->IsA<UPCGHiGenGridSizeSettings>())
+			{
+				ChangeType |= EPCGChangeType::GenerationGrid;
+			}
 		}
 	}
 
-	NotifyGraphChanged(EPCGChangeType::Structural);
+	NotifyGraphChanged(ChangeType);
 #endif
 }
 
@@ -571,15 +578,19 @@ void UPCGGraph::OnNodeRemoved(UPCGNode* InNode)
 void UPCGGraph::OnNodesRemoved(TArrayView<UPCGNode*> InNodes)
 {
 #if WITH_EDITOR
+	bool bAnyGridSizeNodes = false;
+
 	for (UPCGNode* Node : InNodes)
 	{
 		if (Node)
 		{
 			Node->OnNodeChangedDelegate.RemoveAll(this);
+
+			bAnyGridSizeNodes |= !!Cast<UPCGHiGenGridSizeSettings>(Node->GetSettings());
 		}
 	}
 
-	NotifyGraphChanged(EPCGChangeType::Structural);
+	NotifyGraphChanged(bAnyGridSizeNodes ? (EPCGChangeType::Structural | EPCGChangeType::GenerationGrid) : EPCGChangeType::Structural);
 #endif
 }
 
@@ -869,7 +880,8 @@ bool UPCGGraph::RemoveOutboundEdges(UPCGNode* InNode, const FName& OutboundLabel
 void UPCGGraph::ForceNotificationForEditor()
 {
 	// Queue up the delayed change
-	NotifyGraphChanged(EPCGChangeType::Structural);
+	NotifyGraphChanged(EPCGChangeType::Structural | EPCGChangeType::GenerationGrid);
+
 	if (bUserPausedNotificationsInGraphEditor)
 	{
 		EnableNotificationsForEditor();
@@ -1105,6 +1117,8 @@ void UPCGGraph::OnNodeChanged(UPCGNode* InNode, EPCGChangeType ChangeType)
 		// Update node to grid size map for grid size changes.
 		if (Cast<UPCGHiGenGridSizeSettings>(InNode->GetSettings()))
 		{
+			ChangeType |= EPCGChangeType::GenerationGrid;
+
 			FWriteScopeLock Lock(NodeToGridSizeLock);
 			NodeToGridSize.Reset();
 		}
@@ -1182,7 +1196,7 @@ void UPCGGraph::PostEditChangeProperty(FPropertyChangedEvent& PropertyChangedEve
 		|| PropertyName == GET_MEMBER_NAME_CHECKED(UPCGGraph, bUseHierarchicalGeneration))
 	{
 		// The higen settings change the structure of the graph (presence or absence of links between grid levels).
-		NotifyGraphChanged(EPCGChangeType::Structural);
+		NotifyGraphChanged(EPCGChangeType::Structural | EPCGChangeType::GenerationGrid);
 	}
 
 	NumberOfUserParametersPreEdit = 0;
diff --git a/Engine/Plugins/Experimental/PCG/Source/PCG/Private/PCGSettings.cpp b/Engine/Plugins/Experimental/PCG/Source/PCG/Private/PCGSettings.cpp
--- a/Engine/Plugins/Experimental/PCG/Source/PCG/Private/PCGSettings.cpp
+++ b/Engine/Plugins/Experimental/PCG/Source/PCG/Private/PCGSettings.cpp
@@ -540,9 +540,17 @@ void UPCGSettings::PostEditChangeProperty(FPropertyChangedEvent& PropertyChanged
 		InitializeCachedOverridableParams(/*bReset=*/true);
 	}
 
-	if (PropertyChangedEvent.GetPropertyName() != GET_MEMBER_NAME_CHECKED(UPCGSettings, DeterminismSettings))
+	const FName PropertyName = PropertyChangedEvent.GetPropertyName();
+	if (PropertyName != GET_MEMBER_NAME_CHECKED(UPCGSettings, DeterminismSettings))
 	{
-		OnSettingsChangedDelegate.Broadcast(this, GetChangeTypeForProperty(PropertyChangedEvent.GetPropertyName()));
+		// If we have a property name then get the change type for that property, otherwise assume deepest change type.
+		EPCGChangeType ChangeType = EPCGChangeType::Structural | EPCGChangeType::GenerationGrid;
+		if (PropertyName != NAME_None)
+		{
+			ChangeType = GetChangeTypeForProperty(PropertyChangedEvent.GetPropertyName());
+		}
+
+		OnSettingsChangedDelegate.Broadcast(this, ChangeType);
 	}
 
 	CacheCrc();
diff --git a/Engine/Plugins/Experimental/PCG/Source/PCG/Private/PCGSubsystem.cpp b/Engine/Plugins/Experimental/PCG/Source/PCG/Private/PCGSubsystem.cpp
--- a/Engine/Plugins/Experimental/PCG/Source/PCG/Private/PCGSubsystem.cpp
+++ b/Engine/Plugins/Experimental/PCG/Source/PCG/Private/PCGSubsystem.cpp
@@ -812,7 +812,7 @@ void UPCGSubsystem::CancelAllGeneration()
 	}
 }
 
-void UPCGSubsystem::RefreshRuntimeGenComponent(UPCGComponent* RuntimeComponent, bool bRemovePartitionActors)
+void UPCGSubsystem::RefreshRuntimeGenComponent(UPCGComponent* RuntimeComponent, EPCGChangeType ChangeType)
 {
 	if (!ensure(RuntimeComponent && RuntimeComponent->IsManagedByRuntimeGenSystem()))
 	{
@@ -821,6 +821,8 @@ void UPCGSubsystem::RefreshRuntimeGenComponent(UPCGComponent* RuntimeComponent,
 
 	if (ensure(RuntimeGenScheduler))
 	{
+		// Only need to remove PAs if the grid sizes have changed.
+		const bool bRemovePartitionActors = !!(ChangeType & EPCGChangeType::GenerationGrid);
 		RuntimeGenScheduler->RefreshComponent(RuntimeComponent, bRemovePartitionActors);
 	}
 }
@@ -1103,16 +1105,16 @@ void UPCGSubsystem::CreatePartitionActorsWithinBounds(const FBox& InBounds, cons
 	}
 }
 
-FPCGTaskId UPCGSubsystem::ScheduleRefresh(UPCGComponent* Component, bool bForceRegen)
+FPCGTaskId UPCGSubsystem::ScheduleRefresh(UPCGComponent* Component, bool bForceRegen, bool bForceCleanup)
 {
 	check(Component && !Component->IsManagedByRuntimeGenSystem());
 
 	TWeakObjectPtr<UPCGComponent> ComponentPtr(Component);
 
-	auto RefreshTask = [ComponentPtr, bForceRegen]() {
+	auto RefreshTask = [ComponentPtr, bForceRegen, bForceCleanup]() {
 		if (UPCGComponent* Component = ComponentPtr.Get())
 		{
-			Component->OnRefresh(bForceRegen);
+			Component->OnRefresh(bForceRegen, bForceCleanup);
 		}
 		return true;
 	};
diff --git a/Engine/Plugins/Experimental/PCG/Source/PCG/Private/RuntimeGen/PCGRuntimeGenScheduler.cpp b/Engine/Plugins/Experimental/PCG/Source/PCG/Private/RuntimeGen/PCGRuntimeGenScheduler.cpp
--- a/Engine/Plugins/Experimental/PCG/Source/PCG/Private/RuntimeGen/PCGRuntimeGenScheduler.cpp
+++ b/Engine/Plugins/Experimental/PCG/Source/PCG/Private/RuntimeGen/PCGRuntimeGenScheduler.cpp
@@ -557,71 +557,85 @@ void FPCGRuntimeGenScheduler::TickScheduleGeneration(TMap<FGridGenerationKey, do
 		// If the Grid is unbounded, we have a non-partitioned or unbounded component.
 		if (Grid == EPCGHiGenGrid::Unbounded)
 		{
-			if (ensure(OriginalComponent && !OriginalComponent->bGenerated) && !OriginalComponent->IsGenerating())
+			if (ensure(OriginalComponent) && !OriginalComponent->IsGenerating())
 			{
 				if (PCGRuntimeGenSchedulerHelpers::CVarRuntimeGenerationEnableDebugging.GetValueOnAnyThread() && OriginalComponent->GetOwner())
 				{
 					UE_LOG(LogPCG, Warning, TEXT("[RUNTIMEGEN] GENERATE: '%s' (priority %lf)"), *OriginalComponent->GetOwner()->GetActorNameOrLabel(), Priority);
 				}
 
-				OriginalComponent->GenerateLocal(EPCGComponentGenerationTrigger::GenerateAtRuntime, /*bForce=*/false, Grid);
+				// Force to refresh if the component is already generated.
+				OriginalComponent->GenerateLocal(EPCGComponentGenerationTrigger::GenerateAtRuntime, /*bForce=*/true, Grid);
 			}
 		}
 		// Otherwise we have a local component.
 		else
 		{
-			TObjectPtr<APCGPartitionActor> PartitionActor = nullptr;
+			// Grab local component and PA if they exist already.
+			UPCGComponent* LocalComponent = ActorAndComponentMapping->GetLocalComponent(GridSize, GridCoords, OriginalComponent, /*bRuntimeGenerated=*/true);
+			TObjectPtr<APCGPartitionActor> PartitionActor = LocalComponent ? Cast<APCGPartitionActor>(LocalComponent->GetOwner()) : nullptr;
 
-			if (PCGRuntimeGenSchedulerHelpers::CVarRuntimeGenerationEnablePooling.GetValueOnAnyThread())
+			if (!LocalComponent || !ensure(PartitionActor))
 			{
-				if (PCGRuntimeGenSchedulerHelpers::CVarRuntimeGenerationEnableDebugging.GetValueOnAnyThread())
+				// Local component & PA do not exist, create them.
+				if (PCGRuntimeGenSchedulerHelpers::CVarRuntimeGenerationEnablePooling.GetValueOnAnyThread())
 				{
-					UE_LOG(LogPCG, Warning, TEXT("[RUNTIMEGEN] UNPOOL PARTITION ACTOR: '%s' (priority %lf, %d remaining out of %d)"),
-						*APCGPartitionActor::GetRuntimeGenActorName(GridSize, GridCoords),
-						Priority,
-						PartitionActorPool.Num(),
-						PartitionActorPoolSize);
+					if (PCGRuntimeGenSchedulerHelpers::CVarRuntimeGenerationEnableDebugging.GetValueOnAnyThread())
+					{
+						UE_LOG(LogPCG, Warning, TEXT("[RUNTIMEGEN] UNPOOL PARTITION ACTOR: '%s' (priority %lf, %d remaining out of %d)"),
+							*APCGPartitionActor::GetRuntimeGenActorName(GridSize, GridCoords),
+							Priority,
+							PartitionActorPool.Num(),
+							PartitionActorPoolSize);
+					}
+
+					// Get RuntimeGenPA from pool.
+					PartitionActor = GetPartitionActorFromPool(GridSize, GridCoords);
 				}
+				else
+				{
+					if (PCGRuntimeGenSchedulerHelpers::CVarRuntimeGenerationEnableDebugging.GetValueOnAnyThread())
+					{
+						UE_LOG(LogPCG, Warning, TEXT("[RUNTIMEGEN] CREATE PARTITION ACTOR: '%s' (priority %lf)"),
+							*APCGPartitionActor::GetRuntimeGenActorName(GridSize, GridCoords),
+							Priority);
+					}
 
-				// Get RuntimeGenPA from pool.
-				PartitionActor = GetPartitionActorFromPool(GridSize, GridCoords);
-			}
-			else
-			{
-				if (PCGRuntimeGenSchedulerHelpers::CVarRuntimeGenerationEnableDebugging.GetValueOnAnyThread())
+					// Find or Create RuntimeGenPA.
+					PartitionActor = FindOrCreatePartitionActor(GridSize, GridCoords);
+				}
+
+				if (!ensure(PartitionActor))
 				{
-					UE_LOG(LogPCG, Warning, TEXT("[RUNTIMEGEN] CREATE PARTITION ACTOR: '%s' (priority %lf)"),
-						*APCGPartitionActor::GetRuntimeGenActorName(GridSize, GridCoords),
-						Priority);
+					continue;
 				}
 
-				// Find or Create RuntimeGenPA.
-				PartitionActor = FindOrCreatePartitionActor(GridSize, GridCoords);
-			}
+				// Update component mapping for this PA (add local component).
+				{
+					FWriteScopeLock WriteLock(ActorAndComponentMapping->ComponentToRuntimeGenPartitionActorsMapLock);
+					TSet<TObjectPtr<APCGPartitionActor>>* PartitionActorsPtr = ActorAndComponentMapping->ComponentToRuntimeGenPartitionActorsMap.Find(OriginalComponent);
 
-			if (!ensure(PartitionActor))
-			{
-				continue;
-			}
+					if (!PartitionActorsPtr)
+					{
+						PartitionActorsPtr = &ActorAndComponentMapping->ComponentToRuntimeGenPartitionActorsMap.Emplace(OriginalComponent);
+					}
 
-			// Update component mapping for this PA (add local component).
-			{
-				FWriteScopeLock WriteLock(ActorAndComponentMapping->ComponentToRuntimeGenPartitionActorsMapLock);
-				TSet<TObjectPtr<APCGPartitionActor>>* PartitionActorsPtr = ActorAndComponentMapping->ComponentToRuntimeGenPartitionActorsMap.Find(OriginalComponent);
+					// Log this original component before setting up the PA, so that we early out from RefreshComponent if it gets called
+					// in the AddGraphInstance call below.
+					OriginalComponentBeingGenerated = OriginalComponent;
 
-				if (!PartitionActorsPtr)
-				{
-					PartitionActorsPtr = &ActorAndComponentMapping->ComponentToRuntimeGenPartitionActorsMap.Emplace(OriginalComponent);
+					PartitionActor->AddGraphInstance(OriginalComponent);
+
+					OriginalComponentBeingGenerated = nullptr;
+
+					PartitionActorsPtr->Add(PartitionActor);
 				}
 
-				PartitionActor->AddGraphInstance(OriginalComponent);
-				PartitionActorsPtr->Add(PartitionActor);
+				// Generate local component.
+				LocalComponent = PartitionActor->GetLocalComponent(OriginalComponent);
 			}
 
-			// Generate local component.
-			UPCGComponent* LocalComponent = PartitionActor->GetLocalComponent(OriginalComponent);
-
-			if (ensure(LocalComponent && !LocalComponent->IsGenerating()))
+			if (ensure(LocalComponent) && !LocalComponent->IsGenerating())
 			{
 				if (PCGRuntimeGenSchedulerHelpers::CVarRuntimeGenerationEnableDebugging.GetValueOnAnyThread())
 				{
@@ -641,7 +655,8 @@ void FPCGRuntimeGenScheduler::TickScheduleGeneration(TMap<FGridGenerationKey, do
 					}
 				}
 
-				LocalComponent->GenerateLocal(EPCGComponentGenerationTrigger::GenerateAtRuntime, /*bForce=*/false, LocalComponent->GetGenerationGrid(), Dependencies);
+				// Force to refresh if the component is already generated.
+				LocalComponent->GenerateLocal(EPCGComponentGenerationTrigger::GenerateAtRuntime, /*bForce=*/true, LocalComponent->GetGenerationGrid(), Dependencies);
 			}
 		}
 
@@ -881,107 +896,163 @@ void FPCGRuntimeGenScheduler::RefreshComponent(UPCGComponent* InComponent, bool
 		return;
 	}
 
-	const APCGPartitionActor* PartitionActor = Cast<APCGPartitionActor>(InComponent->GetOwner());
+	APCGPartitionActor* PartitionActor = Cast<APCGPartitionActor>(InComponent->GetOwner());
 	const bool bIsLocalComponent = PartitionActor != nullptr;
-
 	UPCGComponent* OriginalComponent = bIsLocalComponent ? PartitionActor->GetOriginalComponent(InComponent) : InComponent;
 
-	auto RefreshLocalComponent = [this, OriginalComponent, bRemovePartitionActors](UPCGComponent* LocalComponent)
+	// If we are mid way through setting up an original component, early out from this refresh.
+	if (!OriginalComponent || OriginalComponent == OriginalComponentBeingGenerated)
+	{
+		return;
+	}
+
+	const bool bLoggingEnabled = PCGRuntimeGenSchedulerHelpers::CVarRuntimeGenerationEnableDebugging.GetValueOnAnyThread();
+
+	// Useful because we can run into generation order issues if components are left to continue generating.
+	if (InComponent->IsGenerating())
 	{
-		check(LocalComponent);
-		APCGPartitionActor* PartitionActor = CastChecked<APCGPartitionActor>(LocalComponent->GetOwner());
+		InComponent->CancelGeneration();
+	}
 
-		// Find the specific generation key for this component, if it exists, cleanup and generate.
-		FGridGenerationKey LocalComponentKey(PartitionActor->GetPCGGridSize(), PartitionActor->GetGridCoord(), OriginalComponent);
+	if (!bRemovePartitionActors)
+	{
+		// Refresh path - mark component dirty and removed generated keys which will cause it to be scheduled for regeneration.
 
-		if (GeneratedComponents.Find(LocalComponentKey))
+		// Remove from generated component set, component will be regenerated later (and in grid order so that e.g. unbounded
+		// is generated first).
+		if (PartitionActor)
 		{
-			if (PCGRuntimeGenSchedulerHelpers::CVarRuntimeGenerationEnableDebugging.GetValueOnAnyThread())
+			if (bLoggingEnabled)
 			{
-				UE_LOG(LogPCG, Warning, TEXT("[RUNTIMEGEN] REFRESH LOCAL COMPONENT: '%s'"), *PartitionActor->GetActorNameOrLabel());
+				UE_LOG(LogPCG, Warning, TEXT("[RUNTIMEGEN] SHALLOW REFRESH LOCAL COMPONENT: '%s'"), *PartitionActor->GetActorNameOrLabel());
 			}
 
-			if (bRemovePartitionActors)
+			GeneratedComponents.Remove({ PartitionActor->GetPCGGridSize(), PartitionActor->GetGridCoord(), OriginalComponent });
+
+			InComponent->CleanupLocalImmediate(/*bRemoveComponents=*/false);
+		}
+		else
+		{
+			// Remove original component
+			GeneratedComponents.Remove({ PCGHiGenGrid::UnboundedGridSize(), FIntVector(0), OriginalComponent });
+
+			// Remove local components
+			TSet<FGridGenerationKey> KeysToRemove;
+
+			for (const FGridGenerationKey& Key : GeneratedComponents)
 			{
-				CleanupComponent(LocalComponentKey, LocalComponent);
+				if (Key.GetOriginalComponent() == InComponent)
+				{
+					// TODO - clean up local immediate will have a flag in the future to clean up the local components on its own, so this call to CleanupLocalImmediate will not be required
+					UPCGComponent* LocalComponent = ActorAndComponentMapping->GetLocalComponent(
+						Key.GetGridSize(),
+						Key.GetGridCoords(),
+						OriginalComponent,
+						/*bRuntimeGenerated=*/true);
+
+					if (bLoggingEnabled && LocalComponent && LocalComponent->GetOwner())
+					{
+						UE_LOG(LogPCG, Warning, TEXT("[RUNTIMEGEN] SHALLOW REFRESH LOCAL COMPONENT: '%s'"), *LocalComponent->GetOwner()->GetActorNameOrLabel());
+					}
+
+					if (ensure(LocalComponent))
+					{
+						LocalComponent->CleanupLocalImmediate(/*bRemoveComponents=*/false);
+					}
+
+					KeysToRemove.Add(Key);
+				}
 			}
-			else
+
+			GeneratedComponents = GeneratedComponents.Difference(KeysToRemove);
+
+			if (bLoggingEnabled && OriginalComponent->GetOwner())
 			{
-				LocalComponent->CleanupLocalImmediate(/*bRemoveComponents=*/true);
-				LocalComponent->GenerateLocal(EPCGComponentGenerationTrigger::GenerateAtRuntime, /*bForce=*/false, LocalComponent->GetGenerationGrid());
+				UE_LOG(LogPCG, Warning, TEXT("[RUNTIMEGEN] SHALLOW REFRESH COMPONENT: '%s' PARTITIONED: %d"),
+					*OriginalComponent->GetOwner()->GetActorNameOrLabel(),
+					OriginalComponent->IsPartitioned() ? 1 : 0);
 			}
-		}
-	};
 
-	if (bIsLocalComponent)
-	{
-		RefreshLocalComponent(InComponent);
+			InComponent->CleanupLocalImmediate(/*bRemoveComponents=*/false);
+		}
 	}
 	else
 	{
-		// Grid sizes might have changed, so we should defensively create grid guids.
-		CreateGridGuidsForComponent(OriginalComponent);
-
-		TArray<FGridGenerationKey> GenerationKeys;
+		// Full cleanout path - cleanup existing components and return actors to the pool.
 
-		for (FGridGenerationKey GenerationKey : GeneratedComponents)
+		auto RefreshLocalComponent = [this, OriginalComponent, bLoggingEnabled](UPCGComponent* LocalComponent)
 		{
-			if (GenerationKey.GetOriginalComponent() == OriginalComponent)
+			check(LocalComponent);
+			APCGPartitionActor* PartitionActor = CastChecked<APCGPartitionActor>(LocalComponent->GetOwner());
+
+			// Find the specific generation key for this component, if it exists, cleanup and generate.
+			FGridGenerationKey LocalComponentKey(PartitionActor->GetPCGGridSize(), PartitionActor->GetGridCoord(), OriginalComponent);
+
+			if (GeneratedComponents.Find(LocalComponentKey))
 			{
-				GenerationKeys.Add(GenerationKey);
+				if (bLoggingEnabled)
+				{
+					UE_LOG(LogPCG, Warning, TEXT("[RUNTIMEGEN] DEEP REFRESH LOCAL COMPONENT: '%s'"), *PartitionActor->GetActorNameOrLabel());
+				}
+
+				CleanupComponent(LocalComponentKey, LocalComponent);
 			}
-		}
+		};
 
-		if (!bRemovePartitionActors)
+		if (bIsLocalComponent)
 		{
-			// Sort GenerationKeys by GridSize to guarantee the largest grid sizes execute first.
-			// This is only necessary if we are not doing a full cleanup, since full cleanup regeneration is already sorted by the scheduler.
-			GenerationKeys.Sort([](const FGridGenerationKey& Lhs, const FGridGenerationKey& Rhs)
-			{
-				return Lhs.GetGridSize() > Rhs.GetGridSize();
-			});
+			RefreshLocalComponent(InComponent);
 		}
-
-		// Gather all generated components which originated from this original component.
-		for (FGridGenerationKey GenerationKey : GenerationKeys)
+		else
 		{
-			const uint32 GridSize = GenerationKey.GetGridSize();
-			const EPCGHiGenGrid Grid = PCGHiGenGrid::GridSizeToGrid(GridSize);
+			TArray<FGridGenerationKey> GenerationKeys;
 
-			// If the Grid is unbounded, we have a non-partitioned or unbounded component.
-			if (Grid == EPCGHiGenGrid::Unbounded)
+			for (FGridGenerationKey GenerationKey : GeneratedComponents)
 			{
-				if (PCGRuntimeGenSchedulerHelpers::CVarRuntimeGenerationEnableDebugging.GetValueOnAnyThread() && OriginalComponent->GetOwner())
+				if (GenerationKey.GetOriginalComponent() == OriginalComponent)
 				{
-					UE_LOG(LogPCG, Warning, TEXT("[RUNTIMEGEN] REFRESH COMPONENT: '%s'"), *OriginalComponent->GetOwner()->GetActorNameOrLabel());
+					GenerationKeys.Add(GenerationKey);
 				}
+			}
 
-				if (bRemovePartitionActors)
+			// Gather all generated components which originated from this original component.
+			for (FGridGenerationKey GenerationKey : GenerationKeys)
+			{
+				const uint32 GridSize = GenerationKey.GetGridSize();
+				const EPCGHiGenGrid Grid = PCGHiGenGrid::GridSizeToGrid(GridSize);
+
+				// If the Grid is unbounded, we have a non-partitioned or unbounded component.
+				if (Grid == EPCGHiGenGrid::Unbounded)
 				{
+					if (bLoggingEnabled && OriginalComponent->GetOwner())
+					{
+						UE_LOG(LogPCG, Warning, TEXT("[RUNTIMEGEN] DEEP REFRESH COMPONENT: '%s' PARTITIONED: %d"),
+							*OriginalComponent->GetOwner()->GetActorNameOrLabel(),
+							OriginalComponent->IsPartitioned() ? 1 : 0);
+					}
+
 					CleanupComponent(GenerationKey, OriginalComponent);
 				}
+				// Otherwise we have a local component.
 				else
 				{
-					OriginalComponent->CleanupLocalImmediate(/*bRemoveComponents=*/true);
-					OriginalComponent->GenerateLocal(EPCGComponentGenerationTrigger::GenerateAtRuntime, /*bForce=*/false, Grid);
-				}
-			}
-			// Otherwise we have a local component.
-			else
-			{
-				const FIntVector GridCoords = GenerationKey.GetGridCoords();
+					const FIntVector GridCoords = GenerationKey.GetGridCoords();
 
-				UPCGComponent* LocalComponent = ActorAndComponentMapping->GetLocalComponent(GridSize, GridCoords, OriginalComponent, /*bRuntimeGenerated=*/true);
-				if (!ensure(LocalComponent))
-				{
-					UE_LOG(LogPCG, Error, TEXT("[RUNTIMEGEN] Generated local component could not be retrieved: %u, %d_%d_%d"), GridSize, GridCoords.X, GridCoords.Y, GridCoords.Z);
-					continue;
-				}
+					UPCGComponent* LocalComponent = ActorAndComponentMapping->GetLocalComponent(GridSize, GridCoords, OriginalComponent, /*bRuntimeGenerated=*/true);
+					if (!ensure(LocalComponent))
+					{
+						UE_LOG(LogPCG, Error, TEXT("[RUNTIMEGEN] Generated local component could not be retrieved: %u, %d_%d_%d"), GridSize, GridCoords.X, GridCoords.Y, GridCoords.Z);
+						continue;
+					}
 
-				RefreshLocalComponent(LocalComponent);
+					RefreshLocalComponent(LocalComponent);
+				}
 			}
 		}
+	}
 
+	if (!bIsLocalComponent)
+	{
 		// When an original/non-partitioned component is refreshed, we need to dirty the state.
 		bAnyRuntimeGenComponentsExistDirty = true;
 	}
diff --git a/Engine/Plugins/Experimental/PCG/Source/PCG/Public/PCGCommon.h b/Engine/Plugins/Experimental/PCG/Source/PCG/Public/PCGCommon.h
--- a/Engine/Plugins/Experimental/PCG/Source/PCG/Public/PCGCommon.h
+++ b/Engine/Plugins/Experimental/PCG/Source/PCG/Public/PCGCommon.h
@@ -19,6 +19,8 @@ enum class EPCGChangeType : uint8
 	Edge = 1 << 3,
 	Node = 1 << 4,
 	Structural = 1 << 5,
+	/** Anything related to generation grids - changing grid size or adding/removing grid size nodes. */
+	GenerationGrid = 1 << 6,
 };
 ENUM_CLASS_FLAGS(EPCGChangeType);
 
diff --git a/Engine/Plugins/Experimental/PCG/Source/PCG/Public/PCGComponent.h b/Engine/Plugins/Experimental/PCG/Source/PCG/Public/PCGComponent.h
--- a/Engine/Plugins/Experimental/PCG/Source/PCG/Public/PCGComponent.h
+++ b/Engine/Plugins/Experimental/PCG/Source/PCG/Public/PCGComponent.h
@@ -272,9 +272,9 @@ public:
 
 #if WITH_EDITOR
 	/** Schedules refresh of the component. If bCancelExistingRefresh is true, any existing refresh is cancelled and a new one is scheduled. */
-	void Refresh(bool bStructural = false, bool bCancelExistingRefresh = false);
+	void Refresh(EPCGChangeType ChangeType = EPCGChangeType::None, bool bCancelExistingRefresh = false);
 
-	void OnRefresh(bool bForceRefresh);
+	void OnRefresh(bool bForceRefresh, bool bForceCleanup);
 
 	void StartGenerationInProgress();
 	void StopGenerationInProgress();
@@ -422,7 +422,7 @@ private:
 	void GetManagedResources(TArray<TObjectPtr<UPCGManagedResource>>& Resources) const;
 	void SetManagedResources(const TArray<TObjectPtr<UPCGManagedResource>>& Resources);
 
-	void RefreshAfterGraphChanged(UPCGGraphInterface* InGraph, bool bIsStructural, bool bDirtyInputs);
+	void RefreshAfterGraphChanged(UPCGGraphInterface* InGraph, EPCGChangeType ChangeType);
 	void OnGraphChanged(UPCGGraphInterface* InGraph, EPCGChangeType ChangeType);
 
 #if WITH_EDITOR
diff --git a/Engine/Plugins/Experimental/PCG/Source/PCG/Public/PCGSubsystem.h b/Engine/Plugins/Experimental/PCG/Source/PCG/Public/PCGSubsystem.h
--- a/Engine/Plugins/Experimental/PCG/Source/PCG/Public/PCGSubsystem.h
+++ b/Engine/Plugins/Experimental/PCG/Source/PCG/Public/PCGSubsystem.h
@@ -138,7 +138,7 @@ public:
 	void CancelGeneration(UPCGGraph* Graph);
 
 	/** Asks the runtime generation scheduler to refresh a given GenerateAtRuntime component. bRemovePartitionActors will also perform a full cleanup of PAs and local components. */
-	void RefreshRuntimeGenComponent(UPCGComponent* RuntimeComponent, bool bRemovePartitionActors = false);
+	void RefreshRuntimeGenComponent(UPCGComponent* RuntimeComponent, EPCGChangeType ChangeType = EPCGChangeType::None);
 
 	/** Returns true if there are any tasks for this graph currently scheduled or executing. */
 	bool IsGraphCurrentlyExecuting(UPCGGraph* Graph);
@@ -186,7 +186,7 @@ public:
 #if WITH_EDITOR
 public:
 	/** Schedule refresh on the current or next frame */
-	FPCGTaskId ScheduleRefresh(UPCGComponent* SourceComponent, bool bForceRefresh);
+	FPCGTaskId ScheduleRefresh(UPCGComponent* SourceComponent, bool bForceRefresh, bool bForceCleanup);
 
 	/** Schedules an operation to cleanup the graph in the given bounds */
 	FPCGTaskId CleanupGraph(UPCGComponent* Component, const FBox& InBounds, bool bRemoveComponents, bool bSave);
diff --git a/Engine/Plugins/Experimental/PCG/Source/PCG/Public/RuntimeGen/PCGRuntimeGenScheduler.h b/Engine/Plugins/Experimental/PCG/Source/PCG/Public/RuntimeGen/PCGRuntimeGenScheduler.h
--- a/Engine/Plugins/Experimental/PCG/Source/PCG/Public/RuntimeGen/PCGRuntimeGenScheduler.h
+++ b/Engine/Plugins/Experimental/PCG/Source/PCG/Public/RuntimeGen/PCGRuntimeGenScheduler.h
@@ -121,4 +121,10 @@ private:
 	/** Track the existence of runtime gen components to avoid unnecessary computation when there is no work to do. */
 	bool bAnyRuntimeGenComponentsExist = false;
 	bool bAnyRuntimeGenComponentsExistDirty = false;
+
+	/** Setting up a PA calls APCGPartitionActor::AddGraphInstance which later calls RefreshComponent, which can create
+	* an infinite refresh loop. To break this loop we write the OC pointer to this variable, and if Refresh gets called for
+	* this OC we early out. Basically we don't respond to refresh calls for a component we are midway through setting up.
+	*/
+	const UPCGComponent* OriginalComponentBeingGenerated = nullptr;
 };
diff --git a/Engine/Plugins/Experimental/PCG/Source/PCGEditor/Private/Details/PCGComponentDetails.cpp b/Engine/Plugins/Experimental/PCG/Source/PCGEditor/Private/Details/PCGComponentDetails.cpp
--- a/Engine/Plugins/Experimental/PCG/Source/PCGEditor/Private/Details/PCGComponentDetails.cpp
+++ b/Engine/Plugins/Experimental/PCG/Source/PCGEditor/Private/Details/PCGComponentDetails.cpp
@@ -241,7 +241,8 @@ FReply FPCGComponentDetails::OnRefreshClicked()
 	{
 		if (Component.IsValid())
 		{
-			Component.Get()->Refresh();
+			// Trigger the deepest refresh - re-initialize the PAs.
+			Component.Get()->Refresh(EPCGChangeType::Structural | EPCGChangeType::GenerationGrid, /*bCancelExistingRefresh=*/true);
 		}
 	}
 
diff --git a/Engine/Plugins/Experimental/PCG/Source/PCGEditor/Private/PCGEditor.cpp b/Engine/Plugins/Experimental/PCG/Source/PCGEditor/Private/PCGEditor.cpp
--- a/Engine/Plugins/Experimental/PCG/Source/PCGEditor/Private/PCGEditor.cpp
+++ b/Engine/Plugins/Experimental/PCG/Source/PCGEditor/Private/PCGEditor.cpp
@@ -302,7 +302,7 @@ void FPCGEditor::UpdateDebugAfterComponentSelection(UPCGComponent* InOldComponen
 			if (UPCGSubsystem* Subsystem = GetSubsystem())
 			{
 				// We don't want to do a full cleanup if we're setting the debug object, since full cleanup destroys the component, which is the debug object itself!
-				Subsystem->RefreshRuntimeGenComponent(Component, /*bRemovePartitionActors=*/false);
+				Subsystem->RefreshRuntimeGenComponent(Component);
 			}
 		}
 		else
@@ -461,7 +461,8 @@ void FPCGEditor::PostUndo(bool bSuccess)
 	{
 		if (PCGGraphBeingEdited)
 		{
-			PCGGraphBeingEdited->NotifyGraphChanged(EPCGChangeType::Structural);
+			// Deepest change type to catch all types of change (like redoing adding a grid size node or etc).
+			PCGGraphBeingEdited->NotifyGraphChanged(EPCGChangeType::Structural | EPCGChangeType::GenerationGrid);
 		}
 
 		if (GraphEditorWidget.IsValid())
