[PCG] SplineSampler improvements
* Fix last spline point missing when sampling non-looped splines
* Provide option to have segment index as a metadata attribute
* Provide boolean to know if a point is a control point (Subdivision Mode only)
* Provide option to have control point tangents as metadata attributes

diff --git a/Engine/Plugins/Experimental/PCG/Source/PCG/Private/Data/PCGLandscapeSplineData.cpp b/Engine/Plugins/Experimental/PCG/Source/PCG/Private/Data/PCGLandscapeSplineData.cpp
--- a/Engine/Plugins/Experimental/PCG/Source/PCG/Private/Data/PCGLandscapeSplineData.cpp
+++ b/Engine/Plugins/Experimental/PCG/Source/PCG/Private/Data/PCGLandscapeSplineData.cpp
@@ -7,6 +7,7 @@
 #include "Elements/PCGSplineSampler.h"
 #include "Helpers/PCGHelpers.h"
 
+#include "LandscapeSplineControlPoint.h"
 #include "LandscapeSplineSegment.h"
 
 #include UE_INLINE_GENERATED_CPP_BY_NAME(PCGLandscapeSplineData)
@@ -66,7 +67,7 @@ FVector::FReal UPCGLandscapeSplineData::GetSegmentLength(int SegmentIndex) const
 {
 	check(Spline.IsValid());
 	check(SegmentIndex >= 0 && SegmentIndex < Spline->GetSegments().Num());
-	
+
 	const ULandscapeSplineSegment* Segment = Spline->GetSegments()[SegmentIndex];
 	const TArray<FLandscapeSplineInterpPoint>& InterpPoints = Segment->GetPoints();
 	FVector::FReal Length = 0;
@@ -211,6 +212,25 @@ FVector::FReal UPCGLandscapeSplineData::GetCurvatureAtDistance(int SegmentIndex,
 	return 0;
 }
 
+void UPCGLandscapeSplineData::GetTangentsAtSegmentStart(int SegmentIndex, FVector& OutArriveTangent, FVector& OutLeaveTangent) const
+{
+	check(Spline.IsValid());
+
+	const TArray<TObjectPtr<ULandscapeSplineSegment>>& Segments = Spline->GetSegments();
+
+	const TObjectPtr<ULandscapeSplineSegment> PreviousSegment = SegmentIndex > 0 ? Segments[SegmentIndex - 1] : nullptr;
+	const TObjectPtr<ULandscapeSplineSegment> CurrentSegment = SegmentIndex < Segments.Num() ? Segments[SegmentIndex] : nullptr;
+
+	// Arrive tangent lives in the end-point of the previous segment
+	const TObjectPtr<ULandscapeSplineControlPoint> ArrivePoint = PreviousSegment ? PreviousSegment->Connections[1].ControlPoint : nullptr;
+
+	// Leave tangent lives in the start-point of the current segment
+	const TObjectPtr<ULandscapeSplineControlPoint> LeavePoint = CurrentSegment ? CurrentSegment->Connections[0].ControlPoint : nullptr;
+
+	OutArriveTangent = ArrivePoint ? ArrivePoint->Rotation.Vector() * -PreviousSegment->Connections[1].TangentLen : FVector::Zero();
+	OutLeaveTangent = LeavePoint ? LeavePoint->Rotation.Vector() * CurrentSegment->Connections[0].TangentLen : FVector::Zero();
+}
+
 const UPCGPointData* UPCGLandscapeSplineData::CreatePointData(FPCGContext* Context) const
 {
 	check(Spline.IsValid());
@@ -271,7 +291,7 @@ bool UPCGLandscapeSplineData::SamplePoint(const FTransform& InTransform, const F
 	{
 		// Considering the landscape spline always exists on the landscape,
 		// we'll ignore the Z component of the input here for the bounds check.
-		if(!PCGHelpers::IsInsideBoundsXY(Segment->GetBounds(), Position))
+		if (!Segment->GetBounds().IsInsideOrOnXY(Position))
 		{
 			continue;
 		}
diff --git a/Engine/Plugins/Experimental/PCG/Source/PCG/Private/Data/PCGPolyLineData.cpp b/Engine/Plugins/Experimental/PCG/Source/PCG/Private/Data/PCGPolyLineData.cpp
--- a/Engine/Plugins/Experimental/PCG/Source/PCG/Private/Data/PCGPolyLineData.cpp
+++ b/Engine/Plugins/Experimental/PCG/Source/PCG/Private/Data/PCGPolyLineData.cpp
@@ -19,3 +19,9 @@ FBox UPCGPolyLineData::GetBounds() const
 	
 	return Bounds;
 }
+
+void UPCGPolyLineData::GetTangentsAtSegmentStart(int SegmentIndex, FVector& OutArriveTangent, FVector& OutLeaveTangent) const
+{
+	OutArriveTangent = FVector::Zero();
+	OutLeaveTangent = FVector::Zero();
+}
diff --git a/Engine/Plugins/Experimental/PCG/Source/PCG/Private/Data/PCGSplineData.cpp b/Engine/Plugins/Experimental/PCG/Source/PCG/Private/Data/PCGSplineData.cpp
--- a/Engine/Plugins/Experimental/PCG/Source/PCG/Private/Data/PCGSplineData.cpp
+++ b/Engine/Plugins/Experimental/PCG/Source/PCG/Private/Data/PCGSplineData.cpp
@@ -116,6 +116,13 @@ FVector::FReal UPCGSplineData::GetCurvatureAtDistance(int SegmentIndex, FVector:
 	return FMath::Sign(RightVector | CurvatureVector) * Curvature;
 }
 
+void UPCGSplineData::GetTangentsAtSegmentStart(int SegmentIndex, FVector& OutArriveTangent, FVector& OutLeaveTangent) const
+{
+	check(SplineStruct.SplineCurves.Position.Points.IsValidIndex(SegmentIndex));
+	OutArriveTangent = SplineStruct.SplineCurves.Position.Points[SegmentIndex].ArriveTangent;
+	OutLeaveTangent = SplineStruct.SplineCurves.Position.Points[SegmentIndex].LeaveTangent;
+}
+
 const UPCGPointData* UPCGSplineData::CreatePointData(FPCGContext* Context) const
 {
 	TRACE_CPUPROFILER_EVENT_SCOPE(UPCGSplineData::CreatePointData);
diff --git a/Engine/Plugins/Experimental/PCG/Source/PCG/Private/Elements/PCGSplineSampler.cpp b/Engine/Plugins/Experimental/PCG/Source/PCG/Private/Elements/PCGSplineSampler.cpp
--- a/Engine/Plugins/Experimental/PCG/Source/PCG/Private/Elements/PCGSplineSampler.cpp
+++ b/Engine/Plugins/Experimental/PCG/Source/PCG/Private/Elements/PCGSplineSampler.cpp
@@ -317,6 +317,10 @@ namespace PCGSplineSampler
 		FTransform LocalTransform;
 		FBox Box = FBox::BuildAABB(FVector::ZeroVector, FVector::OneVector);
 		FVector::FReal Curvature = 0;
+		int SegmentIndex = 0;
+		int SubsegmentIndex = 0;
+		FVector ArriveTangent = FVector::Zero();
+		FVector LeaveTangent = FVector::Zero();
 		FVector::FReal PreviousDeltaAngle = 0;
 		FVector::FReal NextDeltaAngle = 0;
 	};
@@ -351,18 +355,15 @@ namespace PCGSplineSampler
 	struct FStepSampler
 	{
 		FStepSampler(const UPCGPolyLineData* InLineData, const FPCGSplineSamplerParams& Params)
-			: LineData(InLineData), bComputeCurvature(Params.bComputeCurvature)
+			: LineData(InLineData)
+			, bComputeCurvature(Params.bComputeCurvature)
 		{
 			check(LineData);
 			CurrentSegmentIndex = 0;
 		}
 
 		virtual void Step(FSamplerResult& OutSamplerResult) = 0;
-
-		bool IsDone() const
-		{
-			return CurrentSegmentIndex >= LineData->GetNumSegments();
-		}
+		virtual bool IsDone() const = 0;
 
 		const UPCGPolyLineData* LineData = nullptr;
 		int CurrentSegmentIndex = 0;
@@ -373,6 +374,7 @@ namespace PCGSplineSampler
 	{
 		FSubdivisionStepSampler(const UPCGPolyLineData* InLineData, const FPCGSplineSamplerParams& Params)
 			: FStepSampler(InLineData, Params)
+			, bComputeTangents(Params.bComputeTangents)
 		{
 			NumSegments = LineData->GetNumSegments();
 			SubdivisionsPerSegment = Params.SubdivisionsPerSegment;
@@ -383,21 +385,45 @@ namespace PCGSplineSampler
 
 		virtual void Step(FSamplerResult& OutResult) override
 		{
-			const FVector::FReal SegmentLength = LineData->GetSegmentLength(CurrentSegmentIndex);
+			const int PreviousSegmentIndex = (CurrentSegmentIndex > 0 ? CurrentSegmentIndex : NumSegments) - 1;
+
+			// To capture the last key point on the spline, we sample the point at the end of the previous segment.
+			const bool bLastKeyPoint = CurrentSegmentIndex == LineData->GetNumSegments();
+			const int SegmentIndex = bLastKeyPoint ? PreviousSegmentIndex : CurrentSegmentIndex;
+
+			const FVector::FReal SegmentLength = LineData->GetSegmentLength(SegmentIndex);
 			const FVector::FReal SegmentStep = SegmentLength / (SubdivisionsPerSegment + 1);
+			const FVector::FReal DistanceAlongSegment = bLastKeyPoint ? SegmentLength : SubpointIndex * SegmentStep;
 
 			FBox& OutBox = OutResult.Box;
 			FTransform& OutTransform = OutResult.LocalTransform;
-			OutTransform = LineData->GetTransformAtDistance(CurrentSegmentIndex, SubpointIndex * SegmentStep, /*bWorldSpace=*/false, &OutBox);
+			OutTransform = LineData->GetTransformAtDistance(SegmentIndex, DistanceAlongSegment, /*bWorldSpace=*/false, &OutBox);
+			OutResult.SegmentIndex = LineData->IsClosed() ? CurrentSegmentIndex : SegmentIndex;
+			OutResult.SubsegmentIndex = SubpointIndex;
 
 			if (bComputeCurvature)
 			{
-				OutResult.Curvature = LineData->GetCurvatureAtDistance(CurrentSegmentIndex, SubpointIndex * SegmentStep);
+				OutResult.Curvature = LineData->GetCurvatureAtDistance(SegmentIndex, DistanceAlongSegment);
+			}
+
+			if (bComputeTangents)
+			{
+				// Control points have actual Arrive and Leave tangents
+				if (SubpointIndex == 0)
+				{
+					LineData->GetTangentsAtSegmentStart(CurrentSegmentIndex, OutResult.ArriveTangent, OutResult.LeaveTangent);
+				}
+				else
+				{
+					// For a non-control-point, we can get the normalized tangent at least.
+					const FVector Forward = OutTransform.GetRotation().GetForwardVector();
+					OutResult.ArriveTangent = Forward;
+					OutResult.LeaveTangent = Forward;
+				}
 			}
 
 			if (SubpointIndex == 0)
 			{
-				const int PreviousSegmentIndex = (CurrentSegmentIndex > 0 ? CurrentSegmentIndex : NumSegments) - 1;
 				const FVector::FReal PreviousSegmentLength = LineData->GetSegmentLength(PreviousSegmentIndex);
 				FTransform PreviousSegmentEndTransform = LineData->GetTransformAtDistance(PreviousSegmentIndex, PreviousSegmentLength, /*bWorldSpace=*/false);
 
@@ -418,16 +444,24 @@ namespace PCGSplineSampler
 			OutBox.Max.X *= 0.5 * SegmentStep / OutTransform.GetScale3D().X;
 
 			++SubpointIndex;
-			if (SubpointIndex > SubdivisionsPerSegment)
+			if (SubpointIndex > SubdivisionsPerSegment || bLastKeyPoint)
 			{
 				SubpointIndex = 0;
 				++CurrentSegmentIndex;
 			}
 		}
 
+		virtual bool IsDone() const override
+		{
+			// Subdivision Sampler is not done until it captures the final control point, which does not have a valid segment associated unless the spline is a closed loop.
+			const int NumSegmentsWhenDone = ((LineData && LineData->IsClosed()) || NumSegments == 0) ? NumSegments : NumSegments + 1;
+			return CurrentSegmentIndex >= NumSegmentsWhenDone;
+		}
+
 		int NumSegments;
 		int SubdivisionsPerSegment;
 		int SubpointIndex;
+		bool bComputeTangents = false;
 	};
 
 	struct FDistanceStepSampler : public FStepSampler
@@ -445,6 +479,7 @@ namespace PCGSplineSampler
 			FTransform& OutTransform = OutResult.LocalTransform;
 			FBox& OutBox = OutResult.Box;
 			OutTransform = LineData->GetTransformAtDistance(CurrentSegmentIndex, CurrentDistance, /*bWorldSpace=*/false, &OutBox);
+			OutResult.SegmentIndex = CurrentSegmentIndex;
 
 			// Set min/max to half of extent
 			OutBox.Min.X *= 0.5 * DistanceIncrement / OutTransform.GetScale3D().X;
@@ -471,6 +506,13 @@ namespace PCGSplineSampler
 			}
 		}
 
+		virtual bool IsDone() const override
+		{
+			// Distance Sampler is not done until it captures the final control point, which does not have a valid segment associated unless the spline is a closed loop.
+			const int NumSegments = LineData->IsClosed() ? LineData->GetNumSegments() - 1 : LineData->GetNumSegments();
+			return CurrentSegmentIndex >= NumSegments;
+		}
+
 		FVector::FReal DistanceIncrement;
 		FVector::FReal CurrentDistance;
 	};
@@ -487,20 +529,47 @@ namespace PCGSplineSampler
 			ProjectionParams = InProjectionParams;
 
 			// Initialize metadata accessors if needed
-			if (OutPointData && OutPointData->Metadata && (Params.bComputeDirectionDelta || Params.bComputeCurvature))
+			if (OutPointData
+				&& OutPointData->Metadata
+				&& (Params.bComputeDirectionDelta
+					|| Params.bComputeCurvature
+					|| Params.bComputeSegmentIndex
+					|| Params.bComputeSubsegmentIndex
+					|| Params.bComputeTangents))
 			{
 				constexpr double DefaultValue = 0.0;
 				if (Params.bComputeDirectionDelta)
 				{
-					NextDirectionDeltaAttribute = OutPointData->Metadata->FindOrCreateAttribute(Params.NextDirectionDeltaAttribute, DefaultValue);
+					NextDirectionDeltaAttribute = OutPointData->Metadata->FindOrCreateAttribute<double>(Params.NextDirectionDeltaAttribute, DefaultValue);
 					bSetMetadata |= (NextDirectionDeltaAttribute != nullptr);
 				}
 
 				if (Params.bComputeCurvature)
 				{
-					CurvatureAttribute = OutPointData->Metadata->FindOrCreateAttribute(Params.CurvatureAttribute, DefaultValue);
+					CurvatureAttribute = OutPointData->Metadata->FindOrCreateAttribute<double>(Params.CurvatureAttribute, DefaultValue);
 					bSetMetadata |= (CurvatureAttribute != nullptr);
 				}
+
+				if (Params.bComputeSegmentIndex)
+				{
+					SegmentIndexAttribute = OutPointData->Metadata->FindOrCreateAttribute<int>(Params.SegmentIndexAttribute, static_cast<int>(DefaultValue));
+					bSetMetadata |= (SegmentIndexAttribute != nullptr);
+				}
+
+				if (Params.bComputeSubsegmentIndex)
+				{
+					SubsegmentIndexAttribute = OutPointData->Metadata->FindOrCreateAttribute<int>(Params.SubsegmentIndexAttribute, static_cast<int>(DefaultValue));
+					bSetMetadata |= (SubsegmentIndexAttribute != nullptr);
+				}
+
+				if (Params.bComputeTangents)
+				{
+					ArriveTangentAttribute = OutPointData->Metadata->FindOrCreateAttribute<FVector>(Params.ArriveTangentAttribute, FVector::Zero());
+					bSetMetadata |= (ArriveTangentAttribute != nullptr);
+
+					LeaveTangentAttribute = OutPointData->Metadata->FindOrCreateAttribute<FVector>(Params.LeaveTangentAttribute, FVector::Zero());
+					bSetMetadata |= (LeaveTangentAttribute != nullptr);
+				}
 			}
 		}
 
@@ -522,6 +591,26 @@ namespace PCGSplineSampler
 			{
 				CurvatureAttribute->SetValue(OutPoint.MetadataEntry, InResult.Curvature);
 			}
+
+			if (SegmentIndexAttribute)
+			{
+				SegmentIndexAttribute->SetValue(OutPoint.MetadataEntry, InResult.SegmentIndex);
+			}
+			
+			if (SubsegmentIndexAttribute)
+			{
+				SubsegmentIndexAttribute->SetValue(OutPoint.MetadataEntry, InResult.SubsegmentIndex);
+			}
+
+			if (ArriveTangentAttribute)
+			{
+				ArriveTangentAttribute->SetValue(OutPoint.MetadataEntry, InResult.ArriveTangent);
+			}
+
+			if (LeaveTangentAttribute)
+			{
+				LeaveTangentAttribute->SetValue(OutPoint.MetadataEntry, InResult.LeaveTangent);
+			}
 		}
 
 		virtual void Sample(const FSamplerResult& InResult, UPCGPointData* OutPointData)
@@ -565,6 +654,10 @@ namespace PCGSplineSampler
 		bool bSetMetadata = false;
 		FPCGMetadataAttribute<double>* NextDirectionDeltaAttribute = nullptr;
 		FPCGMetadataAttribute<double>* CurvatureAttribute = nullptr;
+		FPCGMetadataAttribute<int>* SegmentIndexAttribute = nullptr;
+		FPCGMetadataAttribute<int>* SubsegmentIndexAttribute = nullptr;
+		FPCGMetadataAttribute<FVector>* LeaveTangentAttribute = nullptr;
+		FPCGMetadataAttribute<FVector>* ArriveTangentAttribute = nullptr;
 	};
 
 	/** Samples in a volume surrounding the poly line. */
diff --git a/Engine/Plugins/Experimental/PCG/Source/PCG/Public/Data/PCGLandscapeSplineData.h b/Engine/Plugins/Experimental/PCG/Source/PCG/Public/Data/PCGLandscapeSplineData.h
--- a/Engine/Plugins/Experimental/PCG/Source/PCG/Public/Data/PCGLandscapeSplineData.h
+++ b/Engine/Plugins/Experimental/PCG/Source/PCG/Public/Data/PCGLandscapeSplineData.h
@@ -28,6 +28,7 @@ public:
 	virtual FVector::FReal GetSegmentLength(int SegmentIndex) const override;
 	virtual FTransform GetTransformAtDistance(int SegmentIndex, FVector::FReal Distance, bool bWorldSpace = true, FBox* OutBounds = nullptr) const override;
 	virtual FVector::FReal GetCurvatureAtDistance(int SegmentIndex, FVector::FReal Distance) const override;
+	virtual void GetTangentsAtSegmentStart(int SegmentIndex, FVector& OutArriveTangent, FVector& OutLeaveTangent) const override;
 	//~End UPCGPolyLineData interface
 
 	//~Begin UPCGSpatialDataWithPointCache interface
diff --git a/Engine/Plugins/Experimental/PCG/Source/PCG/Public/Data/PCGPolyLineData.h b/Engine/Plugins/Experimental/PCG/Source/PCG/Public/Data/PCGPolyLineData.h
--- a/Engine/Plugins/Experimental/PCG/Source/PCG/Public/Data/PCGPolyLineData.h
+++ b/Engine/Plugins/Experimental/PCG/Source/PCG/Public/Data/PCGPolyLineData.h
@@ -26,6 +26,8 @@ public:
 	virtual FVector::FReal GetSegmentLength(int SegmentIndex) const PURE_VIRTUAL(UPCGPolyLineData::GetSegmentLength, return 0;);
 	virtual FTransform GetTransformAtDistance(int SegmentIndex, FVector::FReal Distance, bool bWorldSpace = true, FBox* OutBounds = nullptr) const PURE_VIRTUAL(UPCGPolyLine::GetTransformAtDistance, return FTransform(););
 	virtual FVector::FReal GetCurvatureAtDistance(int SegmentIndex, FVector::FReal Distance) const { return 0; }
+	virtual void GetTangentsAtSegmentStart(int SegmentIndex, FVector& OutArriveTangent, FVector& OutLeaveTangent) const;
+	virtual bool IsClosed() const { return false; }
 
 	virtual FVector GetLocationAtDistance(int SegmentIndex, FVector::FReal Distance, bool bWorldSpace = true) const { return GetTransformAtDistance(SegmentIndex, Distance, bWorldSpace).GetLocation(); }
 };
diff --git a/Engine/Plugins/Experimental/PCG/Source/PCG/Public/Data/PCGSplineData.h b/Engine/Plugins/Experimental/PCG/Source/PCG/Public/Data/PCGSplineData.h
--- a/Engine/Plugins/Experimental/PCG/Source/PCG/Public/Data/PCGSplineData.h
+++ b/Engine/Plugins/Experimental/PCG/Source/PCG/Public/Data/PCGSplineData.h
@@ -37,6 +37,8 @@ public:
 	virtual FVector GetLocationAtDistance(int SegmentIndex, FVector::FReal Distance, bool bWorldSpace = true) const override;
 	virtual FTransform GetTransformAtDistance(int SegmentIndex, FVector::FReal Distance, bool bWorldSpace = true, FBox* OutBounds = nullptr) const override;
 	virtual FVector::FReal GetCurvatureAtDistance(int SegmentIndex, FVector::FReal Distance) const override;
+	virtual void GetTangentsAtSegmentStart(int SegmentIndex, FVector& OutArriveTangent, FVector& OutLeaveTangent) const override;
+	virtual bool IsClosed() const override { return SplineStruct.bClosedLoop; }
 	//~End UPCGPolyLineData interface
 
 	//~Begin UPCGSpatialDataWithPointCache interface
diff --git a/Engine/Plugins/Experimental/PCG/Source/PCG/Public/Elements/PCGSplineSampler.h b/Engine/Plugins/Experimental/PCG/Source/PCG/Public/Elements/PCGSplineSampler.h
--- a/Engine/Plugins/Experimental/PCG/Source/PCG/Public/Elements/PCGSplineSampler.h
+++ b/Engine/Plugins/Experimental/PCG/Source/PCG/Public/Elements/PCGSplineSampler.h
@@ -73,16 +73,41 @@ struct PCG_API FPCGSplineSamplerParams
 	UPROPERTY(BlueprintReadWrite, EditAnywhere, Category = Settings, meta = (InlineEditConditionToggle))
 	bool bComputeDirectionDelta = false;
 
-	/** Attribute that wil contain the delta angle to the next point on the spline w.r.t to the current's point Up vector. */
+	/** Attribute that will contain the delta angle to the next point on the spline w.r.t to the current's point Up vector. */
 	UPROPERTY(BlueprintReadWrite, EditAnywhere, Category = Settings, meta = (EditCondition = "bComputeDirectionDelta"))
-	FName NextDirectionDeltaAttribute = "NextDirectionDelta";
+	FName NextDirectionDeltaAttribute = TEXT("NextDirectionDelta");
 
 	UPROPERTY(BlueprintReadWrite, EditAnywhere, Category = Settings, meta = (InlineEditConditionToggle))
 	bool bComputeCurvature = false;
 
 	/** Attribute that will contain the curvature. Note that the radius of curvature is defined as 1/Curvature, and might need you to scale to world units */
 	UPROPERTY(BlueprintReadWrite, EditAnywhere, Category = Settings, meta = (EditCondition = "bComputeCurvature"))
-	FName CurvatureAttribute = "Curvature";
+	FName CurvatureAttribute = TEXT("Curvature");
+
+	UPROPERTY(BlueprintReadWrite, EditAnywhere, Category = Settings, meta = (InlineEditConditionToggle))
+	bool bComputeSegmentIndex = false;
+
+	/** Attribute that will contain the spline segment index. */
+	UPROPERTY(BlueprintReadWrite, EditAnywhere, Category = Settings, meta = (EditCondition = "bComputeSegmentIndex"))
+	FName SegmentIndexAttribute = TEXT("SegmentIndex");
+
+	UPROPERTY(BlueprintReadWrite, EditAnywhere, Category = Settings, meta = (InlineEditConditionToggle))
+	bool bComputeSubsegmentIndex = false;
+
+	/** Attribute that will contain the sub-segment index of a point on the spline. When the sub-segment index is 0, the point is a control point on the actual spline. Only applies to Subdivision mode. */
+	UPROPERTY(BlueprintReadWrite, EditAnywhere, Category = Settings, meta = (EditCondition = "bComputeSubsegmentIndex"))
+	FName SubsegmentIndexAttribute = TEXT("SubsegmentIndex");
+
+	UPROPERTY(BlueprintReadWrite, EditAnywhere, Category = Settings)
+	bool bComputeTangents = false;
+
+	/** Attribute that will contain the arrive tangent vector. For control points, this will be the actual arrive tangent. For non-control points, this will only be the normalized tangent at this point. Only applies to Subdivision mode. */
+	UPROPERTY(BlueprintReadWrite, EditAnywhere, Category = Settings, meta = (EditCondition = "bComputeTangents"))
+	FName ArriveTangentAttribute = TEXT("ArriveTangent");
+
+	/** Attribute that will contain the leave tangent vector. For control points, this will be the actual leave tangent. For non-control points, this will only be the normalized tangent at this point. Only applies to Subdivision mode. */
+	UPROPERTY(BlueprintReadWrite, EditAnywhere, Category = Settings, meta = (EditCondition = "bComputeTangents"))
+	FName LeaveTangentAttribute = TEXT("LeaveTangent");
 
 	/** If no Bounding Shape input is provided, the actor bounds are used to limit the sample generation domain.
 	* This option allows ignoring the actor bounds and generating over the entire spline. Use with caution as this
@@ -103,11 +128,11 @@ struct PCG_API FPCGSplineSamplerParams
 	UPROPERTY(BlueprintReadWrite, EditAnywhere, Category = Settings, meta = (EditCondition = "Dimension==EPCGSplineSamplingDimension::OnInterior"))
 	bool bTreatSplineAsPolyline = false;
 
-	/** Determines the orientation of interior points */
+	/** Determines the orientation of interior points. */
 	UPROPERTY(BlueprintReadWrite, EditAnywhere, Category = Settings, meta = (EditCondition = "Dimension==EPCGSplineSamplingDimension::OnInterior"))
 	EPCGSplineSamplingInteriorOrientation InteriorOrientation = EPCGSplineSamplingInteriorOrientation::Uniform;
 
-	/** Project sample points onto one possible surface given by the spline boundary */
+	/** Project sample points onto one possible surface given by the spline boundary. */
 	UPROPERTY(EditAnywhere, Category = Settings, meta = (EditCondition = "Dimension==EPCGSplineSamplingDimension::OnInterior"))
 	bool bProjectOntoSurface = false;
 
