[PCG] Fix minor UI issue with Switch and Select (multi) nodes

#rnx

diff --git a/Engine/Plugins/Experimental/PCG/Source/PCG/Private/Elements/ControlFlow/PCGMultiSelect.cpp b/Engine/Plugins/Experimental/PCG/Source/PCG/Private/Elements/ControlFlow/PCGMultiSelect.cpp
--- a/Engine/Plugins/Experimental/PCG/Source/PCG/Private/Elements/ControlFlow/PCGMultiSelect.cpp
+++ b/Engine/Plugins/Experimental/PCG/Source/PCG/Private/Elements/ControlFlow/PCGMultiSelect.cpp
@@ -86,15 +86,18 @@ FString UPCGMultiSelectSettings::GetAdditionalTitleInformation() const
 		case EPCGControlFlowSelectionMode::Enum:
 			if (EnumSelection.Class)
 			{
-				FString Subtitle = PCGControlFlowConstants::SubtitleEnum.ToString();
+				FString Subtitle = EnumSelection.Class->GetName();
 				if (!IsPropertyOverriddenByPin({GET_MEMBER_NAME_CHECKED(UPCGMultiSelectSettings, EnumSelection), GET_MEMBER_NAME_CHECKED(FEnumSelector, Value)}))
 				{
-					Subtitle += EnumSelection.Class->GetNameStringByValue(EnumSelection.Value);
+					Subtitle += FString::Format(TEXT(": {0}"), {EnumSelection.Class->GetNameStringByValue(EnumSelection.Value)});
 				}
 
 				return Subtitle;
 			}
-			break;
+			else
+			{
+				return PCGControlFlowConstants::SubtitleEnum.ToString();
+			}
 
 		case EPCGControlFlowSelectionMode::String:
 			{
diff --git a/Engine/Plugins/Experimental/PCG/Source/PCG/Private/Elements/ControlFlow/PCGSwitch.cpp b/Engine/Plugins/Experimental/PCG/Source/PCG/Private/Elements/ControlFlow/PCGSwitch.cpp
--- a/Engine/Plugins/Experimental/PCG/Source/PCG/Private/Elements/ControlFlow/PCGSwitch.cpp
+++ b/Engine/Plugins/Experimental/PCG/Source/PCG/Private/Elements/ControlFlow/PCGSwitch.cpp
@@ -80,15 +80,18 @@ FString UPCGSwitchSettings::GetAdditionalTitleInformation() const
 		{
 			if (EnumSelection.Class)
 			{
-				FString Subtitle = PCGControlFlowConstants::SubtitleEnum.ToString();
+				FString Subtitle = EnumSelection.Class->GetName();
 				if (!IsPropertyOverriddenByPin({GET_MEMBER_NAME_CHECKED(UPCGSwitchSettings, EnumSelection), GET_MEMBER_NAME_CHECKED(FEnumSelector, Value)}))
 				{
-					Subtitle += EnumSelection.Class->GetNameStringByValue(EnumSelection.Value);
+					Subtitle += FString::Format(TEXT(": {0}"), {EnumSelection.Class->GetNameStringByValue(EnumSelection.Value)});
 				}
 
 				return Subtitle;
 			}
-			break;
+			else
+			{
+				return PCGControlFlowConstants::SubtitleEnum.ToString();
+			}
 		}
 
 		case EPCGControlFlowSelectionMode::String:
