[PCG] Renaming FPCGActorSelectionKey to FPCGSelectionKey and FPCGActorSelectionKeyToSettingsMap to FPCGSelectionKeyToSettingsMap

diff --git a/Engine/Plugins/Experimental/PCG/Source/PCG/Private/Elements/IO/PCGDataTableElement.cpp b/Engine/Plugins/Experimental/PCG/Source/PCG/Private/Elements/IO/PCGDataTableElement.cpp
--- a/Engine/Plugins/Experimental/PCG/Source/PCG/Private/Elements/IO/PCGDataTableElement.cpp
+++ b/Engine/Plugins/Experimental/PCG/Source/PCG/Private/Elements/IO/PCGDataTableElement.cpp
@@ -13,14 +13,14 @@
 #define LOCTEXT_NAMESPACE "PCGLoadDataTable"
 
 #if WITH_EDITOR
-void UPCGLoadDataTableSettings::GetTrackedActorKeys(FPCGActorSelectionKeyToSettingsMap& OutKeysToSettings, TArray<TObjectPtr<const UPCGGraph>>& OutVisitedGraphs) const
+void UPCGLoadDataTableSettings::GetTrackedActorKeys(FPCGSelectionKeyToSettingsMap& OutKeysToSettings, TArray<TObjectPtr<const UPCGGraph>>& OutVisitedGraphs) const
 {
 	if (DataTable.IsNull())
 	{
 		return;
 	}
 
-	FPCGActorSelectionKey Key = FPCGActorSelectionKey::CreateFromPath(DataTable.ToSoftObjectPath());
+	FPCGSelectionKey Key = FPCGSelectionKey::CreateFromPath(DataTable.ToSoftObjectPath());
 
 	OutKeysToSettings.FindOrAdd(Key).Emplace(this, /*bCulling=*/false);
 }
diff --git a/Engine/Plugins/Experimental/PCG/Source/PCG/Private/Elements/PCGActorSelector.cpp b/Engine/Plugins/Experimental/PCG/Source/PCG/Private/Elements/PCGActorSelector.cpp
--- a/Engine/Plugins/Experimental/PCG/Source/PCG/Private/Elements/PCGActorSelector.cpp
+++ b/Engine/Plugins/Experimental/PCG/Source/PCG/Private/Elements/PCGActorSelector.cpp
@@ -70,7 +70,7 @@ namespace PCGActorSelector
 			break;
 		}
 
-		return [](AActor* Actor) -> bool { return false; };
+		return {};
 	}
 
 	TArray<AActor*> FindActors(const FPCGActorSelectorSettings& Settings, const UPCGComponent* InComponent, const TFunction<bool(const AActor*)>& BoundsCheck, const TFunction<bool(const AActor*)>& SelfIgnoreCheck)
@@ -100,6 +100,11 @@ namespace PCGActorSelector
 		// It will modify the FoundActor pointer to the found actor, if found.
 		TFunction<bool(AActor*)> FilteringFunction = PCGActorSelector::GetFilteringFunction(Settings, BoundsCheck, SelfIgnoreCheck, FoundActors);
 
+		if (!FilteringFunction)
+		{
+			return FoundActors;
+		}
+
 		// In case of iterating over all actors in the world, call our filtering function and get out.
 		if (Settings.ActorFilter == EPCGActorFilter::AllWorldActors)
 		{
@@ -204,42 +209,42 @@ namespace PCGActorSelector
 	}
 }
 
-FPCGActorSelectionKey::FPCGActorSelectionKey(EPCGActorFilter InFilter)
+FPCGSelectionKey::FPCGSelectionKey(EPCGActorFilter InFilter)
 {
 	check(InFilter != EPCGActorFilter::AllWorldActors);
 	ActorFilter = InFilter;
 }
 
-FPCGActorSelectionKey::FPCGActorSelectionKey(FName InTag)
+FPCGSelectionKey::FPCGSelectionKey(FName InTag)
 {
 	Selection = EPCGActorSelection::ByTag;
 	Tag = InTag;
 	ActorFilter = EPCGActorFilter::AllWorldActors;
 }
 
-FPCGActorSelectionKey::FPCGActorSelectionKey(TSubclassOf<AActor> InSelectionClass)
+FPCGSelectionKey::FPCGSelectionKey(TSubclassOf<UObject> InSelectionClass)
 {
 	Selection = EPCGActorSelection::ByClass;
-	ActorSelectionClass = InSelectionClass;
+	SelectionClass = InSelectionClass;
 	ActorFilter = EPCGActorFilter::AllWorldActors;
 }
 
-FPCGActorSelectionKey FPCGActorSelectionKey::CreateFromPath(const FSoftObjectPath& InObjectPath)
+FPCGSelectionKey FPCGSelectionKey::CreateFromPath(const FSoftObjectPath& InObjectPath)
 {
-	FPCGActorSelectionKey Res{};
-	Res.Selection = EPCGActorSelection::ByPath;
-	Res.ObjectPath = InObjectPath;
-	Res.ActorFilter = EPCGActorFilter::AllWorldActors;
+	FPCGSelectionKey Key{};
+	Key.Selection = EPCGActorSelection::ByPath;
+	Key.ObjectPath = InObjectPath;
+	Key.ActorFilter = EPCGActorFilter::AllWorldActors;
 
-	return Res;
+	return Key;
 }
 
-void FPCGActorSelectionKey::SetExtraDependency(const UClass* InExtraDependency)
+void FPCGSelectionKey::SetExtraDependency(const UClass* InExtraDependency)
 {
 	OptionalExtraDependency = InExtraDependency;
 }
 
-bool FPCGActorSelectionKey::IsMatching(const TSoftObjectPtr<UObject>& InObjectPtr, const UPCGComponent* InComponent) const
+bool FPCGSelectionKey::IsMatching(const TSoftObjectPtr<UObject>& InObjectPtr, const UPCGComponent* InComponent) const
 {
 	if (InObjectPtr.IsNull())
 	{
@@ -276,7 +281,7 @@ bool FPCGActorSelectionKey::IsMatching(const TSoftObjectPtr<UObject>& InObjectPt
 		return InActor && InActor->ActorHasTag(Tag);
 	}
 	case EPCGActorSelection::ByClass:
-		return InObject && InObject->GetClass()->IsChildOf(ActorSelectionClass);
+		return InObject && InObject->GetClass()->IsChildOf(SelectionClass);
 	case EPCGActorSelection::ByPath:
 		return InObjectPtr.ToSoftObjectPath() == ObjectPath;
 	default:
@@ -284,7 +289,7 @@ bool FPCGActorSelectionKey::IsMatching(const TSoftObjectPtr<UObject>& InObjectPt
 	}
 }
 
-bool FPCGActorSelectionKey::operator==(const FPCGActorSelectionKey& InOther) const
+bool FPCGSelectionKey::operator==(const FPCGSelectionKey& InOther) const
 {
 	if (ActorFilter != InOther.ActorFilter || Selection != InOther.Selection || OptionalExtraDependency != InOther.OptionalExtraDependency)
 	{
@@ -296,7 +301,7 @@ bool FPCGActorSelectionKey::operator==(const FPCGActorSelectionKey& InOther) con
 	case EPCGActorSelection::ByTag:
 		return Tag == InOther.Tag;
 	case EPCGActorSelection::ByClass:
-		return ActorSelectionClass == InOther.ActorSelectionClass;
+		return SelectionClass == InOther.SelectionClass;
 	case EPCGActorSelection::ByPath:
 		return ObjectPath == InOther.ObjectPath;
 	case EPCGActorSelection::Unknown: // Fall-through
@@ -310,11 +315,11 @@ bool FPCGActorSelectionKey::operator==(const FPCGActorSelectionKey& InOther) con
 	}
 }
 
-uint32 GetTypeHash(const FPCGActorSelectionKey& In)
+uint32 GetTypeHash(const FPCGSelectionKey& In)
 {
 	uint32 HashResult = HashCombine(GetTypeHash(In.ActorFilter), GetTypeHash(In.Selection));
 	HashResult = HashCombine(HashResult, GetTypeHash(In.Tag));
-	HashResult = HashCombine(HashResult, GetTypeHash(In.ActorSelectionClass));
+	HashResult = HashCombine(HashResult, GetTypeHash(In.SelectionClass));
 	HashResult = HashCombine(HashResult, GetTypeHash(In.OptionalExtraDependency));
 	HashResult = HashCombine(HashResult, GetTypeHash(In.ObjectPath));
 
@@ -349,31 +354,36 @@ FName FPCGActorSelectorSettings::GetTaskName(const FText& Prefix) const
 }
 #endif // WITH_EDITOR
 
-FPCGActorSelectionKey FPCGActorSelectorSettings::GetAssociatedKey() const
+FPCGSelectionKey FPCGActorSelectorSettings::GetAssociatedKey() const
 {
 	if (ActorFilter != EPCGActorFilter::AllWorldActors)
 	{
-		return FPCGActorSelectionKey(ActorFilter);
+		return FPCGSelectionKey(ActorFilter);
 	}
 
 	switch (ActorSelection)
 	{
 	case EPCGActorSelection::ByTag:
-		return FPCGActorSelectionKey(ActorSelectionTag);
+		return FPCGSelectionKey(ActorSelectionTag);
 	case EPCGActorSelection::ByClass:
-		return FPCGActorSelectionKey(ActorSelectionClass);
+		return FPCGSelectionKey(ActorSelectionClass);
 	default:
-		return FPCGActorSelectionKey();
+		return FPCGSelectionKey();
 	}
 }
 
-FPCGActorSelectorSettings FPCGActorSelectorSettings::ReconstructFromKey(const FPCGActorSelectionKey& InKey)
+FPCGActorSelectorSettings FPCGActorSelectorSettings::ReconstructFromKey(const FPCGSelectionKey& InKey)
 {
+	if (InKey.SelectionClass && !InKey.SelectionClass->IsChildOf<AActor>())
+	{
+		return FPCGActorSelectorSettings{};
+	}
+
 	FPCGActorSelectorSettings Result{};
 	Result.ActorFilter = InKey.ActorFilter;
 	Result.ActorSelection = InKey.Selection;
 	Result.ActorSelectionTag = InKey.Tag;
-	Result.ActorSelectionClass = InKey.ActorSelectionClass;
+	Result.ActorSelectionClass = InKey.SelectionClass;
 
 	return Result;
 }
\ No newline at end of file
diff --git a/Engine/Plugins/Experimental/PCG/Source/PCG/Private/Elements/PCGDataFromActor.cpp b/Engine/Plugins/Experimental/PCG/Source/PCG/Private/Elements/PCGDataFromActor.cpp
--- a/Engine/Plugins/Experimental/PCG/Source/PCG/Private/Elements/PCGDataFromActor.cpp
+++ b/Engine/Plugins/Experimental/PCG/Source/PCG/Private/Elements/PCGDataFromActor.cpp
@@ -127,9 +127,9 @@ namespace PCGDataFromActorHelpers
 }
 
 #if WITH_EDITOR
-void UPCGDataFromActorSettings::GetTrackedActorKeys(FPCGActorSelectionKeyToSettingsMap& OutKeysToSettings, TArray<TObjectPtr<const UPCGGraph>>& OutVisitedGraphs) const
+void UPCGDataFromActorSettings::GetTrackedActorKeys(FPCGSelectionKeyToSettingsMap& OutKeysToSettings, TArray<TObjectPtr<const UPCGGraph>>& OutVisitedGraphs) const
 {
-	FPCGActorSelectionKey Key = ActorSelector.GetAssociatedKey();
+	FPCGSelectionKey Key = ActorSelector.GetAssociatedKey();
 	if (Mode == EPCGGetDataFromActorMode::GetDataFromPCGComponent || Mode == EPCGGetDataFromActorMode::GetDataFromPCGComponentOrParseComponents)
 	{
 		Key.SetExtraDependency(UPCGComponent::StaticClass());
diff --git a/Engine/Plugins/Experimental/PCG/Source/PCG/Private/Elements/PCGExecuteBlueprint.cpp b/Engine/Plugins/Experimental/PCG/Source/PCG/Private/Elements/PCGExecuteBlueprint.cpp
--- a/Engine/Plugins/Experimental/PCG/Source/PCG/Private/Elements/PCGExecuteBlueprint.cpp
+++ b/Engine/Plugins/Experimental/PCG/Source/PCG/Private/Elements/PCGExecuteBlueprint.cpp
@@ -500,11 +500,11 @@ EPCGSettingsType UPCGBlueprintSettings::GetType() const
 	}
 }
 
-void UPCGBlueprintSettings::GetTrackedActorKeys(FPCGActorSelectionKeyToSettingsMap& OutKeysToSettings, TArray<TObjectPtr<const UPCGGraph>>& OutVisitedGraphs) const
+void UPCGBlueprintSettings::GetTrackedActorKeys(FPCGSelectionKeyToSettingsMap& OutKeysToSettings, TArray<TObjectPtr<const UPCGGraph>>& OutVisitedGraphs) const
 {
 	for (const FName& Tag : TrackedActorTags)
 	{
-		OutKeysToSettings.FindOrAdd(FPCGActorSelectionKey(Tag)).Emplace(this, bTrackActorsOnlyWithinBounds);
+		OutKeysToSettings.FindOrAdd(FPCGSelectionKey(Tag)).Emplace(this, bTrackActorsOnlyWithinBounds);
 	}
 }
 
diff --git a/Engine/Plugins/Experimental/PCG/Source/PCG/Private/Elements/PCGGetActorProperty.cpp b/Engine/Plugins/Experimental/PCG/Source/PCG/Private/Elements/PCGGetActorProperty.cpp
--- a/Engine/Plugins/Experimental/PCG/Source/PCG/Private/Elements/PCGGetActorProperty.cpp
+++ b/Engine/Plugins/Experimental/PCG/Source/PCG/Private/Elements/PCGGetActorProperty.cpp
@@ -15,7 +15,7 @@
 #define LOCTEXT_NAMESPACE "PCGPropertyToParamDataElement"
 
 #if WITH_EDITOR
-void UPCGGetActorPropertySettings::GetTrackedActorKeys(FPCGActorSelectionKeyToSettingsMap& OutKeysToSettings, TArray<TObjectPtr<const UPCGGraph>>& OutVisitedGraphs) const
+void UPCGGetActorPropertySettings::GetTrackedActorKeys(FPCGSelectionKeyToSettingsMap& OutKeysToSettings, TArray<TObjectPtr<const UPCGGraph>>& OutVisitedGraphs) const
 {
 	OutKeysToSettings.FindOrAdd(ActorSelector.GetAssociatedKey()).Emplace(this, bTrackActorsOnlyWithinBounds);
 }
diff --git a/Engine/Plugins/Experimental/PCG/Source/PCG/Private/Elements/PCGGetPropertyFromObjectPath.cpp b/Engine/Plugins/Experimental/PCG/Source/PCG/Private/Elements/PCGGetPropertyFromObjectPath.cpp
--- a/Engine/Plugins/Experimental/PCG/Source/PCG/Private/Elements/PCGGetPropertyFromObjectPath.cpp
+++ b/Engine/Plugins/Experimental/PCG/Source/PCG/Private/Elements/PCGGetPropertyFromObjectPath.cpp
@@ -13,7 +13,7 @@
 #define LOCTEXT_NAMESPACE "PCGGetPropertyFromObjectPathElement"
 
 #if WITH_EDITOR
-void UPCGGetPropertyFromObjectPathSettings::GetTrackedActorKeys(FPCGActorSelectionKeyToSettingsMap& OutKeysToSettings, TArray<TObjectPtr<const UPCGGraph>>& OutVisitedGraphs) const
+void UPCGGetPropertyFromObjectPathSettings::GetTrackedActorKeys(FPCGSelectionKeyToSettingsMap& OutKeysToSettings, TArray<TObjectPtr<const UPCGGraph>>& OutVisitedGraphs) const
 {
 	for (const FSoftObjectPath& ObjectPath : ObjectPathsToExtract)
 	{
@@ -22,7 +22,7 @@ void UPCGGetPropertyFromObjectPathSettings::GetTrackedActorKeys(FPCGActorSelecti
 			continue;
 		}
 
-		FPCGActorSelectionKey Key = FPCGActorSelectionKey::CreateFromPath(ObjectPath);
+		FPCGSelectionKey Key = FPCGSelectionKey::CreateFromPath(ObjectPath);
 
 		OutKeysToSettings.FindOrAdd(Key).Emplace(this, /*bCulling=*/false);
 	}
diff --git a/Engine/Plugins/Experimental/PCG/Source/PCG/Private/Elements/PCGIndirectionElement.cpp b/Engine/Plugins/Experimental/PCG/Source/PCG/Private/Elements/PCGIndirectionElement.cpp
--- a/Engine/Plugins/Experimental/PCG/Source/PCG/Private/Elements/PCGIndirectionElement.cpp
+++ b/Engine/Plugins/Experimental/PCG/Source/PCG/Private/Elements/PCGIndirectionElement.cpp
@@ -17,14 +17,14 @@ namespace PCGIndirectionSettings
 }
 
 #if WITH_EDITOR
-void UPCGIndirectionSettings::GetTrackedActorKeys(FPCGActorSelectionKeyToSettingsMap& OutKeysToSettings, TArray<TObjectPtr<const UPCGGraph>>& OutVisitedGraphs) const
+void UPCGIndirectionSettings::GetTrackedActorKeys(FPCGSelectionKeyToSettingsMap& OutKeysToSettings, TArray<TObjectPtr<const UPCGGraph>>& OutVisitedGraphs) const
 {
 	if (Settings.IsNull())
 	{
 		return;
 	}
 
-	FPCGActorSelectionKey Key = FPCGActorSelectionKey::CreateFromPath(Settings.ToSoftObjectPath());
+	FPCGSelectionKey Key = FPCGSelectionKey::CreateFromPath(Settings.ToSoftObjectPath());
 
 	OutKeysToSettings.FindOrAdd(Key).Emplace(this, /*bCulling=*/false);
 }
diff --git a/Engine/Plugins/Experimental/PCG/Source/PCG/Private/PCGActorAndComponentMapping.cpp b/Engine/Plugins/Experimental/PCG/Source/PCG/Private/PCGActorAndComponentMapping.cpp
--- a/Engine/Plugins/Experimental/PCG/Source/PCG/Private/PCGActorAndComponentMapping.cpp
+++ b/Engine/Plugins/Experimental/PCG/Source/PCG/Private/PCGActorAndComponentMapping.cpp
@@ -837,7 +837,7 @@ void FPCGActorAndComponentMapping::RegisterTracking(UPCGComponent* InComponent)
 	}
 }
 
-void FPCGActorAndComponentMapping::UpdateTracking(UPCGComponent* InComponent, bool bInShouldDirtyActors, const TArray<FPCGActorSelectionKey>* ChangedKeys)
+void FPCGActorAndComponentMapping::UpdateTracking(UPCGComponent* InComponent, bool bInShouldDirtyActors, const TArray<FPCGSelectionKey>* ChangedKeys)
 {
 	// Discard BP templates, local components and invalid component
 	if (!IsValid(InComponent) || !InComponent->GetOwner() || InComponent->GetOwner()->IsA<APCGPartitionActor>())
@@ -861,7 +861,7 @@ void FPCGActorAndComponentMapping::UpdateTracking(UPCGComponent* InComponent, bo
 	}
 
 	// If no keys are provided, update all tracking keys.
-	TArray<FPCGActorSelectionKey> AllKeys;
+	TArray<FPCGSelectionKey> AllKeys;
 	if (ChangedKeys == nullptr)
 	{
 		InComponent->CachedTrackedKeysToSettings.GenerateKeyArray(AllKeys);
@@ -880,7 +880,7 @@ void FPCGActorAndComponentMapping::UpdateTracking(UPCGComponent* InComponent, bo
 	TSet<TSoftObjectPtr<UObject>> CandidatesForTracking;
 	TSet<TSoftObjectPtr<UObject>> CandidatesForUntracking;
 
-	auto GatherObjects = [this, InComponent, &CandidatesForTracking, &CandidatesForUntracking](const FPCGActorSelectionKey& InKey, bool bInShouldUntrack)
+	auto GatherObjects = [this, InComponent, &CandidatesForTracking, &CandidatesForUntracking](const FPCGSelectionKey& InKey, bool bInShouldUntrack)
 	{
 		if (InKey.Selection == EPCGActorSelection::ByPath)
 		{
@@ -930,7 +930,7 @@ void FPCGActorAndComponentMapping::UpdateTracking(UPCGComponent* InComponent, bo
 		}
 	};
 
-	for (const FPCGActorSelectionKey& Key : *ChangedKeys)
+	for (const FPCGSelectionKey& Key : *ChangedKeys)
 	{
 		const bool bShouldUntrack = !InComponent->CachedTrackedKeysToSettings.Contains(Key);
 		GatherObjects(Key, bShouldUntrack);
diff --git a/Engine/Plugins/Experimental/PCG/Source/PCG/Private/PCGComponent.cpp b/Engine/Plugins/Experimental/PCG/Source/PCG/Private/PCGComponent.cpp
--- a/Engine/Plugins/Experimental/PCG/Source/PCG/Private/PCGComponent.cpp
+++ b/Engine/Plugins/Experimental/PCG/Source/PCG/Private/PCGComponent.cpp
@@ -1437,7 +1437,7 @@ void UPCGComponent::RefreshAfterGraphChanged(UPCGGraphInterface* InGraph, EPCGCh
 		if (UPCGSubsystem* Subsystem = GetSubsystem())
 		{
 			// Don't update the tracking if nothing changed for the tracking.
-			TArray<FPCGActorSelectionKey> ChangedKeys;
+			TArray<FPCGSelectionKey> ChangedKeys;
 			if (UpdateTrackingCache(&ChangedKeys))
 			{
 				Subsystem->UpdateComponentTracking(this, /*bInShouldDirtyActors=*/ true, &ChangedKeys);
@@ -1573,7 +1573,7 @@ void UPCGComponent::PostEditChangeProperty(FPropertyChangedEvent& PropertyChange
 	{
 		if (UPCGSubsystem* Subsystem = GetSubsystem())
 		{
-			TArray<FPCGActorSelectionKey> ChangedKeys;
+			TArray<FPCGSelectionKey> ChangedKeys;
 			if (UpdateTrackingCache(&ChangedKeys))
 			{
 				Subsystem->UpdateComponentTracking(this, /*bInShouldDirtyActors=*/true, &ChangedKeys);
@@ -1682,7 +1682,7 @@ void UPCGComponent::PostEditUndo()
 	Super::PostEditUndo();
 }
 
-bool UPCGComponent::UpdateTrackingCache(TArray<FPCGActorSelectionKey>* OptionalChangedKeys)
+bool UPCGComponent::UpdateTrackingCache(TArray<FPCGSelectionKey>* OptionalChangedKeys)
 {
 	// Without an owner, it probably means we are in a BP template, so no need to update the tracking cache.
 	// Same for local components, as we will use the cache of the original component.
@@ -1692,7 +1692,7 @@ bool UPCGComponent::UpdateTrackingCache(TArray<FPCGActorSelectionKey>* OptionalC
 	}
 
 	// Store in a temporary map to detect key changes.
-	TMap<FPCGActorSelectionKey, bool> NewTrackedKeysToCulling;
+	TMap<FPCGSelectionKey, bool> NewTrackedKeysToCulling;
 
 	int32 FoundKeys = 0;
 
@@ -1703,7 +1703,7 @@ bool UPCGComponent::UpdateTrackingCache(TArray<FPCGActorSelectionKey>* OptionalC
 		// Also add a key for the landscape, with settings null and always culled, if we should track the landscape
 		if (ShouldTrackLandscape())
 		{
-			FPCGActorSelectionKey LandscapeKey = FPCGActorSelectionKey(ALandscapeProxy::StaticClass());
+			FPCGSelectionKey LandscapeKey = FPCGSelectionKey(ALandscapeProxy::StaticClass());
 			CachedTrackedKeysToSettings.FindOrAdd(LandscapeKey).Emplace(/*Settings*/nullptr, /*bIsCulled*/true);
 		}
 
@@ -1711,9 +1711,9 @@ bool UPCGComponent::UpdateTrackingCache(TArray<FPCGActorSelectionKey>* OptionalC
 		// Note that is only impact the fact that we track (or not) this tag.
 		// If a setting is marked as "should cull", it will only be dirtied (at least by default), if the actor with the
 		// given tag intersect with the component.
-		for (const TPair<FPCGActorSelectionKey, TArray<FPCGSettingsAndCulling>>& It : CachedTrackedKeysToSettings)
+		for (const TPair<FPCGSelectionKey, TArray<FPCGSettingsAndCulling>>& It : CachedTrackedKeysToSettings)
 		{
-			const FPCGActorSelectionKey& Key = It.Key;
+			const FPCGSelectionKey& Key = It.Key;
 
 			// Should cull only if all the settings requires a cull.
 			const bool bShouldCull = Algo::AllOf(It.Value, [](const FPCGSettingsAndCulling& SettingsAndCullingPair) { return SettingsAndCullingPair.Value; });
@@ -1742,7 +1742,7 @@ bool UPCGComponent::UpdateTrackingCache(TArray<FPCGActorSelectionKey>* OptionalC
 		{
 			OptionalChangedKeys->Reserve(OptionalChangedKeys->Num() + CachedTrackedKeysToCulling.Num());
 
-			for (const TPair<FPCGActorSelectionKey, bool>& It : CachedTrackedKeysToCulling)
+			for (const TPair<FPCGSelectionKey, bool>& It : CachedTrackedKeysToCulling)
 			{
 				OptionalChangedKeys->Add(It.Key);
 			}
@@ -2075,7 +2075,7 @@ bool UPCGComponent::IsObjectTracked(const TSoftObjectPtr<UObject>& InObjectPtr,
 
 	bool bFound = false;
 
-	for (const TPair<FPCGActorSelectionKey, bool>& It : CachedTrackedKeysToCulling)
+	for (const TPair<FPCGSelectionKey, bool>& It : CachedTrackedKeysToCulling)
 	{
 		if (It.Key.IsMatching(InObjectPtr, this))
 		{
@@ -2696,7 +2696,7 @@ TArray<const UPCGSettings*> UPCGComponent::GatherSettingsTracking(const UObject*
 
 	for (const auto& It : CachedTrackedKeysToSettings)
 	{
-		const FPCGActorSelectionKey& Key = It.Key;
+		const FPCGSelectionKey& Key = It.Key;
 
 		const bool bRemovedTagIsTracked = (Key.Selection == EPCGActorSelection::ByTag) && InRemovedTags.Contains(Key.Tag);
 
diff --git a/Engine/Plugins/Experimental/PCG/Source/PCG/Private/PCGGraph.cpp b/Engine/Plugins/Experimental/PCG/Source/PCG/Private/PCGGraph.cpp
--- a/Engine/Plugins/Experimental/PCG/Source/PCG/Private/PCGGraph.cpp
+++ b/Engine/Plugins/Experimental/PCG/Source/PCG/Private/PCGGraph.cpp
@@ -1003,16 +1003,16 @@ void UPCGGraph::RemoveExtraEditorNode(const UObject* InNode)
 	ExtraEditorNodes.Remove(const_cast<UObject*>(InNode));
 }
 
-FPCGActorSelectionKeyToSettingsMap UPCGGraph::GetTrackedActorKeysToSettings() const
+FPCGSelectionKeyToSettingsMap UPCGGraph::GetTrackedActorKeysToSettings() const
 {
-	FPCGActorSelectionKeyToSettingsMap TagsToSettings;
+	FPCGSelectionKeyToSettingsMap TagsToSettings;
 	TArray<TObjectPtr<const UPCGGraph>> VisitedGraphs;
 
 	GetTrackedActorKeysToSettings(TagsToSettings, VisitedGraphs);
 	return TagsToSettings;
 }
 
-void UPCGGraph::GetTrackedActorKeysToSettings(FPCGActorSelectionKeyToSettingsMap& OutTagsToSettings, TArray<TObjectPtr<const UPCGGraph>>& OutVisitedGraphs) const
+void UPCGGraph::GetTrackedActorKeysToSettings(FPCGSelectionKeyToSettingsMap& OutTagsToSettings, TArray<TObjectPtr<const UPCGGraph>>& OutVisitedGraphs) const
 {
 	if (OutVisitedGraphs.Contains(this))
 	{
@@ -1035,7 +1035,7 @@ void UPCGGraph::GetTrackedActorKeysToSettings(FPCGActorSelectionKeyToSettingsMap
 	// We need to loop because we can have multiple keys that have Filter to self or original.
 	for (auto& It : OutTagsToSettings)
 	{
-		const FPCGActorSelectionKey& Key = It.Key;
+		const FPCGSelectionKey& Key = It.Key;
 		TArray<FPCGSettingsAndCulling>& SettingsAndCullingArray = It.Value;
 		if (Key.ActorFilter == EPCGActorFilter::Self || Key.ActorFilter == EPCGActorFilter::Original)
 		{
diff --git a/Engine/Plugins/Experimental/PCG/Source/PCG/Private/PCGSubgraph.cpp b/Engine/Plugins/Experimental/PCG/Source/PCG/Private/PCGSubgraph.cpp
--- a/Engine/Plugins/Experimental/PCG/Source/PCG/Private/PCGSubgraph.cpp
+++ b/Engine/Plugins/Experimental/PCG/Source/PCG/Private/PCGSubgraph.cpp
@@ -155,7 +155,7 @@ void UPCGBaseSubgraphSettings::PostEditChangeProperty(struct FPropertyChangedEve
 	}
 }
 
-void UPCGBaseSubgraphSettings::GetTrackedActorKeys(FPCGActorSelectionKeyToSettingsMap& OutKeysToSettings, TArray<TObjectPtr<const UPCGGraph>>& VisitedGraphs) const
+void UPCGBaseSubgraphSettings::GetTrackedActorKeys(FPCGSelectionKeyToSettingsMap& OutKeysToSettings, TArray<TObjectPtr<const UPCGGraph>>& VisitedGraphs) const
 {
 	if (UPCGGraph* Subgraph = GetSubgraph())
 	{
diff --git a/Engine/Plugins/Experimental/PCG/Source/PCG/Public/Elements/IO/PCGDataTableElement.h b/Engine/Plugins/Experimental/PCG/Source/PCG/Public/Elements/IO/PCGDataTableElement.h
--- a/Engine/Plugins/Experimental/PCG/Source/PCG/Public/Elements/IO/PCGDataTableElement.h
+++ b/Engine/Plugins/Experimental/PCG/Source/PCG/Public/Elements/IO/PCGDataTableElement.h
@@ -22,7 +22,7 @@ public:
 	virtual FText GetDefaultNodeTitle() const override;
 	virtual FText GetNodeTooltipText() const override;
 	virtual bool HasDynamicPins() const override { return true; }
-	virtual void GetTrackedActorKeys(FPCGActorSelectionKeyToSettingsMap& OutKeysToSettings, TArray<TObjectPtr<const UPCGGraph>>& OutVisitedGraphs) const override;
+	virtual void GetTrackedActorKeys(FPCGSelectionKeyToSettingsMap& OutKeysToSettings, TArray<TObjectPtr<const UPCGGraph>>& OutVisitedGraphs) const override;
 #endif
 
 	virtual EPCGDataType GetCurrentPinTypes(const UPCGPin* InPin) const override;
diff --git a/Engine/Plugins/Experimental/PCG/Source/PCG/Public/Elements/PCGActorSelector.h b/Engine/Plugins/Experimental/PCG/Source/PCG/Public/Elements/PCGActorSelector.h
--- a/Engine/Plugins/Experimental/PCG/Source/PCG/Public/Elements/PCGActorSelector.h
+++ b/Engine/Plugins/Experimental/PCG/Source/PCG/Public/Elements/PCGActorSelector.h
@@ -40,24 +40,23 @@ enum class EPCGActorFilter : uint8
 /**
 * Structure to specify a selection criteria for an object/actor
 * Object can be selected using the EPCGActorSelection::ByClass or EPCGActorSelection::ByPath
-* Actors have more selection with Self/Parent/Root/Original and also EPCGActorSelection::ByTag
-* TODO: Might want to rename it, since it is not limited to Actors anymore.
+* Actors have more options for selection with Self/Parent/Root/Original and also EPCGActorSelection::ByTag
 */
-struct FPCGActorSelectionKey
+struct FPCGSelectionKey
 {
-	FPCGActorSelectionKey() = default;
+	FPCGSelectionKey() = default;
 
 	// For all filters others than AllWorldActor. For AllWorldActors Filter, use the other constructors.
-	explicit FPCGActorSelectionKey(EPCGActorFilter InFilter);
+	explicit FPCGSelectionKey(EPCGActorFilter InFilter);
 
-	explicit FPCGActorSelectionKey(FName InTag);
-	explicit FPCGActorSelectionKey(TSubclassOf<AActor> InSelectionClass);
+	explicit FPCGSelectionKey(FName InTag);
+	explicit FPCGSelectionKey(TSubclassOf<UObject> InSelectionClass);
 
-	static FPCGActorSelectionKey CreateFromPath(const FSoftObjectPath& InObjectPath);
+	static FPCGSelectionKey CreateFromPath(const FSoftObjectPath& InObjectPath);
 
-	bool operator==(const FPCGActorSelectionKey& InOther) const;
+	bool operator==(const FPCGSelectionKey& InOther) const;
 
-	friend uint32 GetTypeHash(const FPCGActorSelectionKey& In);
+	friend uint32 GetTypeHash(const FPCGSelectionKey& In);
 	bool IsMatching(const TSoftObjectPtr<UObject>& InObjectPtr, const UPCGComponent* InComponent) const;
 
 	void SetExtraDependency(const UClass* InExtraDependency);
@@ -65,7 +64,7 @@ struct FPCGActorSelectionKey
 	EPCGActorFilter ActorFilter = EPCGActorFilter::AllWorldActors;
 	EPCGActorSelection Selection = EPCGActorSelection::Unknown;
 	FName Tag = NAME_None;
-	TSubclassOf<AActor> ActorSelectionClass = nullptr;
+	TSubclassOf<UObject> SelectionClass = nullptr;
 
 	// If the Selection is ByPath, contain the path to select.
 	FSoftObjectPath ObjectPath;
@@ -140,8 +139,8 @@ struct FPCGActorSelectorSettings
 	FName GetTaskName(const FText& Prefix) const;
 #endif
 
-	FPCGActorSelectionKey GetAssociatedKey() const;
-	static FPCGActorSelectorSettings ReconstructFromKey(const FPCGActorSelectionKey& InKey);
+	FPCGSelectionKey GetAssociatedKey() const;
+	static FPCGActorSelectorSettings ReconstructFromKey(const FPCGSelectionKey& InKey);
 };
 
 namespace PCGActorSelector
diff --git a/Engine/Plugins/Experimental/PCG/Source/PCG/Public/Elements/PCGDataFromActor.h b/Engine/Plugins/Experimental/PCG/Source/PCG/Public/Elements/PCGDataFromActor.h
--- a/Engine/Plugins/Experimental/PCG/Source/PCG/Public/Elements/PCGDataFromActor.h
+++ b/Engine/Plugins/Experimental/PCG/Source/PCG/Public/Elements/PCGDataFromActor.h
@@ -31,7 +31,7 @@ public:
 	virtual FText GetDefaultNodeTitle() const override { return NSLOCTEXT("PCGDataFromActorSettings", "NodeTitle", "Get Actor Data"); }
 	virtual FText GetNodeTooltipText() const override;
 	virtual EPCGSettingsType GetType() const override { return EPCGSettingsType::Spatial; }
-	virtual void GetTrackedActorKeys(FPCGActorSelectionKeyToSettingsMap& OutKeysToSettings, TArray<TObjectPtr<const UPCGGraph>>& OutVisitedGraphs) const override;
+	virtual void GetTrackedActorKeys(FPCGSelectionKeyToSettingsMap& OutKeysToSettings, TArray<TObjectPtr<const UPCGGraph>>& OutVisitedGraphs) const override;
 	virtual bool HasDynamicPins() const override { return true; }
 	virtual void ApplyDeprecation(UPCGNode* InOutNode) override;
 #endif
diff --git a/Engine/Plugins/Experimental/PCG/Source/PCG/Public/Elements/PCGExecuteBlueprint.h b/Engine/Plugins/Experimental/PCG/Source/PCG/Public/Elements/PCGExecuteBlueprint.h
--- a/Engine/Plugins/Experimental/PCG/Source/PCG/Public/Elements/PCGExecuteBlueprint.h
+++ b/Engine/Plugins/Experimental/PCG/Source/PCG/Public/Elements/PCGExecuteBlueprint.h
@@ -301,7 +301,7 @@ public:
 	virtual FText GetDefaultNodeTitle() const override { return NSLOCTEXT("PCGBlueprintSettings", "NodeTitle", "Execute Blueprint"); }
 	virtual FLinearColor GetNodeTitleColor() const override;
 	virtual EPCGSettingsType GetType() const override;
-	virtual void GetTrackedActorKeys(FPCGActorSelectionKeyToSettingsMap& OutKeysToSettings, TArray<TObjectPtr<const UPCGGraph>>& OutVisitedGraphs) const override;
+	virtual void GetTrackedActorKeys(FPCGSelectionKeyToSettingsMap& OutKeysToSettings, TArray<TObjectPtr<const UPCGGraph>>& OutVisitedGraphs) const override;
 	virtual UObject* GetJumpTargetForDoubleClick() const override;
 	virtual void ApplyDeprecationBeforeUpdatePins(UPCGNode* InOutNode, TArray<TObjectPtr<UPCGPin>>& InputPins, TArray<TObjectPtr<UPCGPin>>& OutputPins) override;
 	virtual TArray<FPCGPreConfiguredSettingsInfo> GetPreconfiguredInfo() const override;
diff --git a/Engine/Plugins/Experimental/PCG/Source/PCG/Public/Elements/PCGGetActorProperty.h b/Engine/Plugins/Experimental/PCG/Source/PCG/Public/Elements/PCGGetActorProperty.h
--- a/Engine/Plugins/Experimental/PCG/Source/PCG/Public/Elements/PCGGetActorProperty.h
+++ b/Engine/Plugins/Experimental/PCG/Source/PCG/Public/Elements/PCGGetActorProperty.h
@@ -31,7 +31,7 @@ public:
 	virtual FName GetDefaultNodeName() const override { return FName(TEXT("GetActorProperty")); }
 	virtual FText GetDefaultNodeTitle() const override { return NSLOCTEXT("PCGPropertyToParamDataSettings", "NodeTitle", "Get Actor Property"); }
 	virtual EPCGSettingsType GetType() const override { return EPCGSettingsType::Param; }
-	virtual void GetTrackedActorKeys(FPCGActorSelectionKeyToSettingsMap& OutKeysToSettings, TArray<TObjectPtr<const UPCGGraph>>& OutVisitedGraphs) const override;
+	virtual void GetTrackedActorKeys(FPCGSelectionKeyToSettingsMap& OutKeysToSettings, TArray<TObjectPtr<const UPCGGraph>>& OutVisitedGraphs) const override;
 #endif
 
 	virtual FString GetAdditionalTitleInformation() const override;
diff --git a/Engine/Plugins/Experimental/PCG/Source/PCG/Public/Elements/PCGGetPropertyFromObjectPath.h b/Engine/Plugins/Experimental/PCG/Source/PCG/Public/Elements/PCGGetPropertyFromObjectPath.h
--- a/Engine/Plugins/Experimental/PCG/Source/PCG/Public/Elements/PCGGetPropertyFromObjectPath.h
+++ b/Engine/Plugins/Experimental/PCG/Source/PCG/Public/Elements/PCGGetPropertyFromObjectPath.h
@@ -25,7 +25,7 @@ public:
 	virtual EPCGSettingsType GetType() const override { return EPCGSettingsType::Param; }
 	virtual bool CanEditChange(const FProperty* InProperty) const override;
 	virtual bool IsPinUsedByNodeExecution(const UPCGPin* InPin) const override;
-	virtual void GetTrackedActorKeys(FPCGActorSelectionKeyToSettingsMap& OutKeysToSettings, TArray<TObjectPtr<const UPCGGraph>>& OutVisitedGraphs) const override;
+	virtual void GetTrackedActorKeys(FPCGSelectionKeyToSettingsMap& OutKeysToSettings, TArray<TObjectPtr<const UPCGGraph>>& OutVisitedGraphs) const override;
 #endif
 	virtual FString GetAdditionalTitleInformation() const override;
 
diff --git a/Engine/Plugins/Experimental/PCG/Source/PCG/Public/Elements/PCGIndirectionElement.h b/Engine/Plugins/Experimental/PCG/Source/PCG/Public/Elements/PCGIndirectionElement.h
--- a/Engine/Plugins/Experimental/PCG/Source/PCG/Public/Elements/PCGIndirectionElement.h
+++ b/Engine/Plugins/Experimental/PCG/Source/PCG/Public/Elements/PCGIndirectionElement.h
@@ -29,7 +29,7 @@ public:
 	virtual FText GetDefaultNodeTitle() const override;
 	virtual FText GetNodeTooltipText() const override;
 	virtual EPCGSettingsType GetType() const override { return EPCGSettingsType::Generic; }
-	virtual void GetTrackedActorKeys(FPCGActorSelectionKeyToSettingsMap& OutKeysToSettings, TArray<TObjectPtr<const UPCGGraph>>& OutVisitedGraphs) const override;
+	virtual void GetTrackedActorKeys(FPCGSelectionKeyToSettingsMap& OutKeysToSettings, TArray<TObjectPtr<const UPCGGraph>>& OutVisitedGraphs) const override;
 #endif // WITH_EDITOR
 	virtual FString GetAdditionalTitleInformation() const override;
 	virtual bool HasFlippedTitleLines() const override { return true; }
diff --git a/Engine/Plugins/Experimental/PCG/Source/PCG/Public/PCGActorAndComponentMapping.h b/Engine/Plugins/Experimental/PCG/Source/PCG/Public/PCGActorAndComponentMapping.h
--- a/Engine/Plugins/Experimental/PCG/Source/PCG/Public/PCGActorAndComponentMapping.h
+++ b/Engine/Plugins/Experimental/PCG/Source/PCG/Public/PCGActorAndComponentMapping.h
@@ -55,7 +55,7 @@ public:
 	void AddDelayedActors();
 
 	void RegisterTracking(UPCGComponent* InComponent);
-	void UpdateTracking(UPCGComponent* InComponent, bool bInShouldDirtyActors, const TArray<FPCGActorSelectionKey>* OptionalChangedKeys = nullptr);
+	void UpdateTracking(UPCGComponent* InComponent, bool bInShouldDirtyActors, const TArray<FPCGSelectionKey>* OptionalChangedKeys = nullptr);
 
 	/** Notify that we exited the Landscape edit mode. */
 	PCG_API void NotifyLandscapeEditModeExited();
diff --git a/Engine/Plugins/Experimental/PCG/Source/PCG/Public/PCGComponent.h b/Engine/Plugins/Experimental/PCG/Source/PCG/Public/PCGComponent.h
--- a/Engine/Plugins/Experimental/PCG/Source/PCG/Public/PCGComponent.h
+++ b/Engine/Plugins/Experimental/PCG/Source/PCG/Public/PCGComponent.h
@@ -435,7 +435,7 @@ private:
 	void SetupCallbacksOnCreation();
 
 	/** Returns true if something changed in the tracking. */
-	bool UpdateTrackingCache(TArray<FPCGActorSelectionKey>* OptionalChangedKeys = nullptr);
+	bool UpdateTrackingCache(TArray<FPCGSelectionKey>* OptionalChangedKeys = nullptr);
 
 	/** Gather all the settings that are tracking the given object. If it is an actor, the settings should cull and there is no intersection, it will not be added to the list.
 	* @param InObject                   The object the component is tracking. Can be an Actor (for culling).
@@ -515,9 +515,9 @@ private:
 #if WITH_EDITOR
 	int32 InspectionCounter = 0;
 	FBox LastGeneratedBoundsPriorToUndo = FBox(EForceInit::ForceInit);
-	FPCGActorSelectionKeyToSettingsMap CachedTrackedKeysToSettings;
+	FPCGSelectionKeyToSettingsMap CachedTrackedKeysToSettings;
 
-	TMap<FPCGActorSelectionKey, bool> CachedTrackedKeysToCulling;
+	TMap<FPCGSelectionKey, bool> CachedTrackedKeysToCulling;
 #endif
 
 #if WITH_EDITORONLY_DATA
diff --git a/Engine/Plugins/Experimental/PCG/Source/PCG/Public/PCGGraph.h b/Engine/Plugins/Experimental/PCG/Source/PCG/Public/PCGGraph.h
--- a/Engine/Plugins/Experimental/PCG/Source/PCG/Public/PCGGraph.h
+++ b/Engine/Plugins/Experimental/PCG/Source/PCG/Public/PCGGraph.h
@@ -312,8 +312,8 @@ public:
 #endif
 
 #if WITH_EDITOR
-	FPCGActorSelectionKeyToSettingsMap GetTrackedActorKeysToSettings() const;
-	void GetTrackedActorKeysToSettings(FPCGActorSelectionKeyToSettingsMap& OutKeysToSettings, TArray<TObjectPtr<const UPCGGraph>>& OutVisitedGraphs) const;
+	FPCGSelectionKeyToSettingsMap GetTrackedActorKeysToSettings() const;
+	void GetTrackedActorKeysToSettings(FPCGSelectionKeyToSettingsMap& OutKeysToSettings, TArray<TObjectPtr<const UPCGGraph>>& OutVisitedGraphs) const;
 #endif
 
 	/** Size of grid on which this node should be executed. Nodes execute at the minimum of all input grid sizes. */
diff --git a/Engine/Plugins/Experimental/PCG/Source/PCG/Public/PCGSettings.h b/Engine/Plugins/Experimental/PCG/Source/PCG/Public/PCGSettings.h
--- a/Engine/Plugins/Experimental/PCG/Source/PCG/Public/PCGSettings.h
+++ b/Engine/Plugins/Experimental/PCG/Source/PCG/Public/PCGSettings.h
@@ -21,7 +21,7 @@ class UPCGNode;
 class UPCGSettings;
 
 using FPCGSettingsAndCulling = TPair<TWeakObjectPtr<const UPCGSettings>, bool>;
-using FPCGActorSelectionKeyToSettingsMap = TMap<FPCGActorSelectionKey, TArray<FPCGSettingsAndCulling>>;
+using FPCGSelectionKeyToSettingsMap = TMap<FPCGSelectionKey, TArray<FPCGSettingsAndCulling>>;
 
 UENUM()
 enum class EPCGSettingsExecutionMode : uint8
@@ -279,7 +279,7 @@ public:
 	virtual bool GetPinExtraIcon(const UPCGPin* InPin, FName& OutExtraIcon, FText& OutTooltip) const;
 
 	/** Derived classes must implement this to communicate dependencies on external actors */
-	virtual void GetTrackedActorKeys(FPCGActorSelectionKeyToSettingsMap& OutKeysToSettings, TArray<TObjectPtr<const UPCGGraph>>& OutVisitedGraphs) const {}
+	virtual void GetTrackedActorKeys(FPCGSelectionKeyToSettingsMap& OutKeysToSettings, TArray<TObjectPtr<const UPCGGraph>>& OutVisitedGraphs) const {}
 
 	/** Override this class to provide an UObject to jump to in case of double click on node
 	 *  ie. returning a blueprint instance will open the given blueprint in its editor.
diff --git a/Engine/Plugins/Experimental/PCG/Source/PCG/Public/PCGSubgraph.h b/Engine/Plugins/Experimental/PCG/Source/PCG/Public/PCGSubgraph.h
--- a/Engine/Plugins/Experimental/PCG/Source/PCG/Public/PCGSubgraph.h
+++ b/Engine/Plugins/Experimental/PCG/Source/PCG/Public/PCGSubgraph.h
@@ -51,7 +51,7 @@ public:
 
 protected:
 #if WITH_EDITOR
-	virtual void GetTrackedActorKeys(FPCGActorSelectionKeyToSettingsMap& OutKeysToSettings, TArray<TObjectPtr<const UPCGGraph>>& OutVisitedGraphs) const override;
+	virtual void GetTrackedActorKeys(FPCGSelectionKeyToSettingsMap& OutKeysToSettings, TArray<TObjectPtr<const UPCGGraph>>& OutVisitedGraphs) const override;
 	virtual EPCGChangeType GetChangeTypeForProperty(const FName& InPropertyName) const override;
 #endif
 
diff --git a/Engine/Plugins/Experimental/PCG/Source/PCG/Public/PCGSubsystem.h b/Engine/Plugins/Experimental/PCG/Source/PCG/Public/PCGSubsystem.h
--- a/Engine/Plugins/Experimental/PCG/Source/PCG/Public/PCGSubsystem.h
+++ b/Engine/Plugins/Experimental/PCG/Source/PCG/Public/PCGSubsystem.h
@@ -202,7 +202,7 @@ public:
 	void NotifyGraphChanged(UPCGGraph* InGraph);
 
 	/** Update the tracking on a given component. */
-	void UpdateComponentTracking(UPCGComponent* InComponent, bool bShouldDirtyActors, const TArray<FPCGActorSelectionKey>* OptionalChangedKeys = nullptr) { ActorAndComponentMapping.UpdateTracking(InComponent, bShouldDirtyActors, OptionalChangedKeys); }
+	void UpdateComponentTracking(UPCGComponent* InComponent, bool bShouldDirtyActors, const TArray<FPCGSelectionKey>* OptionalChangedKeys = nullptr) { ActorAndComponentMapping.UpdateTracking(InComponent, bShouldDirtyActors, OptionalChangedKeys); }
 
 	/** Propagates transient state change from an original component to the relevant partition actors */
 	void PropagateEditingModeToLocalComponents(UPCGComponent* InOriginalComponent, EPCGEditorDirtyMode EditingMode);
