[PCG] Make Threshold attribute on point filter overridable

#rnx

diff --git a/Engine/Plugins/Experimental/PCG/Source/PCG/Public/Elements/PCGAttributeFilter.h b/Engine/Plugins/Experimental/PCG/Source/PCG/Public/Elements/PCGAttributeFilter.h
--- a/Engine/Plugins/Experimental/PCG/Source/PCG/Public/Elements/PCGAttributeFilter.h
+++ b/Engine/Plugins/Experimental/PCG/Source/PCG/Public/Elements/PCGAttributeFilter.h
@@ -38,7 +38,7 @@ struct PCG_API FPCGAttributeFilterThresholdSettings
 	UPROPERTY(BlueprintReadWrite, EditAnywhere, Category = Settings, meta = (PCG_Overridable))
 	bool bUseConstantThreshold = false;
 
-	UPROPERTY(BlueprintReadWrite, EditAnywhere, Category = Settings, meta = (EditCondition = "!bUseConstantThreshold", EditConditionHides, PCG_NotOverridable))
+	UPROPERTY(BlueprintReadWrite, EditAnywhere, Category = Settings, meta = (EditCondition = "!bUseConstantThreshold", EditConditionHides, PCG_Overridable))
 	FPCGAttributePropertyInputSelector ThresholdAttribute;
 
 	/** For Point Data, enabling this option will use sampling rather than comparing points 1 to 1 directly. For other spatial data, this is always true, and for attribute sets, always false. */
@@ -102,7 +102,7 @@ public:
 	UPROPERTY(BlueprintReadWrite, EditAnywhere, Category = Settings, meta = (PCG_Overridable))
 	bool bUseConstantThreshold = false;
 
-	UPROPERTY(BlueprintReadWrite, EditAnywhere, Category = Settings, meta = (EditCondition = "!bUseConstantThreshold", EditConditionHides, PCG_NotOverridable))
+	UPROPERTY(BlueprintReadWrite, EditAnywhere, Category = Settings, meta = (EditCondition = "!bUseConstantThreshold", EditConditionHides, PCG_Overridable))
 	FPCGAttributePropertyInputSelector ThresholdAttribute;
 
 	/** If the threshold data is Point data, it will sample input points in threshold data. Always true with Spatial data.*/
