Multipoint sampling/projecting interface change, for future optimization purposes but also to prevent performance issues due to locking (physics) in some cases

diff --git a/Engine/Plugins/Experimental/PCG/Config/BasePCG.ini b/Engine/Plugins/Experimental/PCG/Config/BasePCG.ini
--- a/Engine/Plugins/Experimental/PCG/Config/BasePCG.ini
+++ b/Engine/Plugins/Experimental/PCG/Config/BasePCG.ini
@@ -14,6 +14,8 @@
 +FunctionRedirects=(OldName="PCGDataFunctionLibrary.GetTaggedParams",NewName="PCGDataFunctionLibrary.GetParamsByTag")
 +FunctionRedirects=(OldName="PCGMetadata.InitializeAsCopy",NewName="PCGMetadata.K2_InitializeAsCopy")
 +FunctionRedirects=(OldName="PCGMetadata.InitializeAsCopyWithAttributeFilter",NewName="PCGMetadata.K2_InitializeAsCopyWithAttributeFilter")
++FunctionRedirects=(OldName="PCGSpatialData.SamplePoint",NewName="PCGSpatialData.K2_SamplePoint")
++FunctionRedirects=(OldName="PCGSpatialData.ProjectPoint",NewName="PCGSpatialData.K2_ProjectPoint")
 
 +PropertyRedirects=(OldName="PCGSplineSamplerSettings.Params",NewName="PCGSplineSamplerSettings.SamplerParams")
 +PropertyRedirects=(OldName="PCGProjectionSettings.Params",NewName="PCGProjectionSettings.ProjectionParams")
diff --git a/Engine/Plugins/Experimental/PCG/Source/PCG/Private/Data/PCGDifferenceData.cpp b/Engine/Plugins/Experimental/PCG/Source/PCG/Private/Data/PCGDifferenceData.cpp
--- a/Engine/Plugins/Experimental/PCG/Source/PCG/Private/Data/PCGDifferenceData.cpp
+++ b/Engine/Plugins/Experimental/PCG/Source/PCG/Private/Data/PCGDifferenceData.cpp
@@ -1,8 +1,11 @@
 // Copyright Epic Games, Inc. All Rights Reserved.
 
 #include "Data/PCGDifferenceData.h"
+
+#include "PCGContext.h"
 #include "Data/PCGPointData.h"
 #include "Data/PCGSpatialData.h"
+#include "Data/PCGSpatialDataTpl.h"
 #include "Data/PCGUnionData.h"
 #include "Helpers/PCGAsync.h"
 
@@ -239,13 +242,12 @@ const UPCGPointData* UPCGDifferenceData::CreatePointData(FPCGContext* Context) c
 	}
 
 	const UPCGMetadata* SourceMetadata = SourcePointData->Metadata;
+	const TArray<FPCGPoint>& SourcePoints = SourcePointData->GetPoints();
 
 	UPCGPointData* Data = NewObject<UPCGPointData>();
 	Data->InitializeFromData(this, SourceMetadata);
 	
 	UPCGMetadata* OutMetadata = Data->Metadata;
-
-	const TArray<FPCGPoint>& SourcePoints = SourcePointData->GetPoints();
 	TArray<FPCGPoint>& TargetPoints = Data->GetMutablePoints();
 
 	const UPCGMetadata* DifferenceMetadata = GetDifference()->Metadata;
@@ -256,34 +258,55 @@ const UPCGPointData* UPCGDifferenceData::CreatePointData(FPCGContext* Context) c
 		TempDiffMetadata->Initialize(DifferenceMetadata);
 	}
 
-	FPCGAsync::AsyncPointProcessing(Context, SourcePoints.Num(), TargetPoints, [this, Data, OutMetadata, SourcePointData, SourceMetadata, TempDiffMetadata, &SourcePoints](int32 Index, FPCGPoint& OutPoint)
+	constexpr int ChunkSize = FPCGSpatialDataProcessing::DefaultSamplePointsChunkSize;
+
+	auto ChunkSamplePoints = [this, SourceMetadata, TempDiffMetadata, OutMetadata](const TArrayView<TPair<FTransform, FBox>>& Samples, const TArrayView<const FPCGPoint>& SourcePoints, TArray<FPCGPoint, TInlineAllocator<ChunkSize>>& OutPoints)
 	{
-		const FPCGPoint& Point = SourcePoints[Index];
+		const int NumPoints = Samples.Num();
 
-		FPCGPoint PointFromDiff;
-		if (GetDifference() && GetDifference()->SamplePoint(Point.Transform, Point.GetLocalBounds(), PointFromDiff, TempDiffMetadata))
+		TArray<FPCGPoint, TInlineAllocator<ChunkSize>> PointsFromDiff;
+		PointsFromDiff.SetNum(NumPoints);
+
+		check(GetDifference());
+		GetDifference()->SamplePoints(Samples, PointsFromDiff, TempDiffMetadata);
+
+		TArray<FPCGPoint, TInlineAllocator<ChunkSize>> KeptPoints;
+		TArray<FPCGPoint, TInlineAllocator<ChunkSize>> RejectedPoints;
+
+		const bool bBinaryDensity = (DensityFunction == EPCGDifferenceDensityFunction::Binary);
+
+		for (int PointIndex = 0; PointIndex < NumPoints; ++PointIndex)
 		{
-			const bool bBinaryDensity = (DensityFunction == EPCGDifferenceDensityFunction::Binary);
+			const FPCGPoint& Point = SourcePoints[PointIndex];
+			FPCGPoint& PointFromDiff = PointsFromDiff[PointIndex];
 
-			OutPoint = Point;
-			OutPoint.Density = bBinaryDensity ? 0 : FMath::Max(0, Point.Density - PointFromDiff.Density);
+			const float Density = (bBinaryDensity && PointFromDiff.Density > 0) ? 0.0f : Point.Density - PointFromDiff.Density;
 
-			if (TempDiffMetadata && OutPoint.Density > 0 && PointFromDiff.MetadataEntry != PCGInvalidEntryKey)
+			if (Density > 0)
 			{
-				OutMetadata->MergePointAttributesSubset(Point, SourceMetadata, SourceMetadata, PointFromDiff, TempDiffMetadata, TempDiffMetadata, OutPoint, EPCGMetadataOp::Sub);
+				FPCGPoint& OutPoint = KeptPoints.Add_GetRef(Point);
+				OutPoint.Density = Density;
+
+				// TODO: create an array-based MergePointsAttributeSubset
+				if (TempDiffMetadata && PointFromDiff.MetadataEntry != PCGInvalidEntryKey)
+				{
+					OutMetadata->MergePointAttributesSubset(Point, SourceMetadata, SourceMetadata, PointFromDiff, TempDiffMetadata, TempDiffMetadata, OutPoint, EPCGMetadataOp::Sub);
+				}
+			}
+			else if(bKeepZeroDensityPoints)
+			{
+				FPCGPoint& RejectedPoint = RejectedPoints.Add_GetRef(Point);
+				RejectedPoint.Density = 0;
 			}
-
-			return OutPoint.Density > 0 || bKeepZeroDensityPoints;
-		}
-		else
-		{
-			OutPoint = Point;
-			return true;
 		}
-	});
 
-	UE_LOG(LogPCG, Verbose, TEXT("Difference generated %d points from %d source points"), TargetPoints.Num(), SourcePointData->GetPoints().Num());
+		OutPoints.Append(KeptPoints);
+		OutPoints.Append(RejectedPoints);
+	};
 
+	FPCGSpatialDataProcessing::SampleBasedRangeProcessing<ChunkSize>(Context ? &Context->AsyncState : nullptr, ChunkSamplePoints, SourcePoints, TargetPoints);
+	
+	UE_LOG(LogPCG, Verbose, TEXT("Difference generated %d points from %d source points"), TargetPoints.Num(), SourcePointData->GetPoints().Num());	
 	return Data;
 }
 
diff --git a/Engine/Plugins/Experimental/PCG/Source/PCG/Private/Data/PCGIntersectionData.cpp b/Engine/Plugins/Experimental/PCG/Source/PCG/Private/Data/PCGIntersectionData.cpp
--- a/Engine/Plugins/Experimental/PCG/Source/PCG/Private/Data/PCGIntersectionData.cpp
+++ b/Engine/Plugins/Experimental/PCG/Source/PCG/Private/Data/PCGIntersectionData.cpp
@@ -1,8 +1,11 @@
 // Copyright Epic Games, Inc. All Rights Reserved.
 
 #include "Data/PCGIntersectionData.h"
+
+#include "PCGContext.h"
 #include "Data/PCGPointData.h"
 #include "Data/PCGSpatialData.h"
+#include "Data/PCGSpatialDataTpl.h"
 #include "Helpers/PCGAsync.h"
 #include "Helpers/PCGHelpers.h"
 
@@ -215,9 +218,10 @@ UPCGPointData* UPCGIntersectionData::CreateAndFilterPointData(FPCGContext* Conte
 	}
 
 	const TArray<FPCGPoint>& SourcePoints = SourcePointData->GetPoints();
+	const UPCGMetadata* SourceMetadata = SourcePointData->Metadata;
 
 	UPCGPointData* Data = NewObject<UPCGPointData>();
-	Data->InitializeFromData(this, SourcePointData->Metadata);
+	Data->InitializeFromData(this, SourceMetadata);
 	Data->Metadata->AddAttributes(Y->Metadata);
 
 	UPCGMetadata* TempYMetadata = nullptr;
@@ -227,43 +231,55 @@ UPCGPointData* UPCGIntersectionData::CreateAndFilterPointData(FPCGContext* Conte
 		TempYMetadata->Initialize(Y->Metadata);
 	}
 
-	const bool bPointDataHasCommonAttributes = (SourcePointData->Metadata && Y->Metadata && SourcePointData->Metadata->HasCommonAttributes(Y->Metadata));
+	const bool bPointDataHasCommonAttributes = (SourceMetadata && Y->Metadata && SourceMetadata->HasCommonAttributes(Y->Metadata));
 
 	TArray<FPCGPoint>& TargetPoints = Data->GetMutablePoints();
 
-	FPCGAsync::AsyncPointProcessing(Context, SourcePoints.Num(), TargetPoints, [this, Data, SourcePointData, &SourcePoints, Y, TempYMetadata, bPointDataHasCommonAttributes](int32 Index, FPCGPoint& OutPoint)
+	constexpr int ChunkSize = FPCGSpatialDataProcessing::DefaultSamplePointsChunkSize;
+
+	auto ChunkSamplePoints = [this, Data, SourceMetadata, Y, TempYMetadata, bPointDataHasCommonAttributes](const TArrayView<TPair<FTransform, FBox>>& Samples, const TArrayView<const FPCGPoint>& SourcePoints, TArray<FPCGPoint, TInlineAllocator<ChunkSize>>& OutPoints)
 	{
-		TRACE_CPUPROFILER_EVENT_SCOPE(UPCGIntersectionData::CreateAndFilterPointData::Iteration);
-		const FPCGPoint& Point = SourcePoints[Index];
+		const int NumPoints = Samples.Num();
+
+		TArray<FPCGPoint, TInlineAllocator<ChunkSize>> PointsFromY;
+		PointsFromY.SetNum(NumPoints);
+
+		Y->SamplePoints(Samples, PointsFromY, TempYMetadata);
+
+		TArray<FPCGPoint, TInlineAllocator<ChunkSize>> KeptPoints;
+		TArray<FPCGPoint, TInlineAllocator<ChunkSize>> RejectedPoints;
 
-		FPCGPoint PointFromY;
-		if (!Y->SamplePoint(Point.Transform, Point.GetLocalBounds(), PointFromY, TempYMetadata))
+		// Filter points based on output density
+		for (int PointIndex = 0; PointIndex < NumPoints; ++PointIndex)
 		{
-			if (!bKeepZeroDensityPoints)
+			const FPCGPoint& Point = SourcePoints[PointIndex];
+			FPCGPoint& PointFromY = PointsFromY[PointIndex];
+
+			if (PointFromY.Density > 0)
 			{
-				return false;
+				FPCGPoint& KeptPoint = KeptPoints.Add_GetRef(Point); // note: not the sampled point
+				KeptPoint.Density = PCGIntersectionDataMaths::ComputeDensity(Point.Density, PointFromY.Density, DensityFunction);
+				KeptPoint.Color = Point.Color * PointFromY.Color;
+
+				// TODO: create an array-based MergePointsAttributeSubset..
+				// If either the point from Y has metadata or the merge would be a non-trivial value, then perform the full merge
+				if (Data->Metadata && (bPointDataHasCommonAttributes || PointFromY.MetadataEntry != PCGInvalidEntryKey))
+				{
+					Data->Metadata->MergePointAttributesSubset(Point, SourceMetadata, SourceMetadata, PointFromY, TempYMetadata, TempYMetadata, KeptPoint, EPCGMetadataOp::Min);
+				}
 			}
-			else
+			else if (bKeepZeroDensityPoints)
 			{
-				// Point is rejected, mark its density to zero
-				PointFromY.Density = 0;
-				PointFromY.MetadataEntry = PCGInvalidEntryKey;
-				PointFromY.Color = FVector4::One();
+				FPCGPoint& RejectedPoint = RejectedPoints.Add_GetRef(Point); // note: not the sampled point
+				RejectedPoint.Density = 0;
 			}
 		}
 
-		OutPoint = Point;
-		OutPoint.Density = PCGIntersectionDataMaths::ComputeDensity(Point.Density, PointFromY.Density, DensityFunction);
-		OutPoint.Color = Point.Color * PointFromY.Color;
-
-		// If either the point from Y has metadata or the merge would be a non-trivial value, then perform the full merge
-		if (Data->Metadata && (bPointDataHasCommonAttributes || PointFromY.MetadataEntry != PCGInvalidEntryKey))
-		{
-			Data->Metadata->MergePointAttributesSubset(Point, SourcePointData->Metadata, SourcePointData->Metadata, PointFromY, TempYMetadata, TempYMetadata, OutPoint, EPCGMetadataOp::Min);
-		}
+		OutPoints.Append(KeptPoints);
+		OutPoints.Append(RejectedPoints);
+	};
 
-		return true;
-	});
+	FPCGSpatialDataProcessing::SampleBasedRangeProcessing<ChunkSize>(Context ? &Context->AsyncState : nullptr, ChunkSamplePoints, SourcePoints, TargetPoints);
 
 	UE_LOG(LogPCG, Verbose, TEXT("Intersection generated %d points from %d source points"), TargetPoints.Num(), SourcePoints.Num());
 
diff --git a/Engine/Plugins/Experimental/PCG/Source/PCG/Private/Data/PCGLandscapeData.cpp b/Engine/Plugins/Experimental/PCG/Source/PCG/Private/Data/PCGLandscapeData.cpp
--- a/Engine/Plugins/Experimental/PCG/Source/PCG/Private/Data/PCGLandscapeData.cpp
+++ b/Engine/Plugins/Experimental/PCG/Source/PCG/Private/Data/PCGLandscapeData.cpp
@@ -4,16 +4,21 @@
 
 #include "PCGSubsystem.h"
 #include "Data/PCGPointData.h"
+#include "Data/PCGSpatialDataTpl.h"
 #include "Data/PCGSurfaceData.h"
 #include "Data/PCGWorldData.h"
 #include "Grid/PCGLandscapeCache.h"
 #include "Helpers/PCGHelpers.h"
 
+#include "ChaosInterfaceWrapperCore.h"
 #include "Landscape.h"
 #include "LandscapeInfo.h"
 #include "LandscapeProxy.h"
+#include "Chaos/ChaosEngineInterface.h"
+#include "Chaos/PhysicsObjectCollisionInterface.h"
 #include "Engine/World.h"
 #include "PhysicalMaterials/PhysicalMaterial.h"
+#include "PhysicsEngine/PhysicsObjectExternalInterface.h"
 
 #include UE_INLINE_GENERATED_CPP_BY_NAME(PCGLandscapeData)
 
@@ -50,7 +55,8 @@ void UPCGLandscapeData::Initialize(const TArray<TWeakObjectPtr<ALandscapeProxy>>
 		Landscapes.Emplace(Landscape.Get());
 
 		// Build landscape info list
-		LandscapeInfos.Emplace(PCGHelpers::GetLandscapeBounds(Landscape.Get()), Landscape->GetLandscapeInfo());
+		LandscapeInfos.AddUnique(Landscape->GetLandscapeInfo());
+		BoundsToLandscapeInfos.Emplace(PCGHelpers::GetLandscapeBounds(Landscape.Get()), Landscape->GetLandscapeInfo());
 	}
 
 	check(!Landscapes.IsEmpty());
@@ -120,7 +126,8 @@ void UPCGLandscapeData::PostLoad()
 		Landscape.LoadSynchronous();
 		if (Landscape.Get())
 		{
-			LandscapeInfos.Emplace(PCGHelpers::GetLandscapeBounds(Landscape.Get()), Landscape->GetLandscapeInfo());
+			LandscapeInfos.AddUnique(Landscape->GetLandscapeInfo());
+			BoundsToLandscapeInfos.Emplace(PCGHelpers::GetLandscapeBounds(Landscape.Get()), Landscape->GetLandscapeInfo());
 
 			if (!FirstLandscape)
 			{
@@ -228,6 +235,128 @@ bool UPCGLandscapeData::SamplePoint(const FTransform& InTransform, const FBox& I
 	return false;
 }
 
+void UPCGLandscapeData::SamplePoints(const TArrayView<const TPair<FTransform, FBox>>& Samples, const TArrayView<FPCGPoint>& OutPoints, UPCGMetadata* OutMetadata) const
+{
+	// Implementation note:
+	// We will first build a list of all relevant landscsape collision components and the samples to test against them
+	constexpr int32 ChunkSize = FPCGSpatialDataProcessing::DefaultSamplePointsChunkSize;
+	TMap<ULandscapeHeightfieldCollisionComponent*, TArray<int, TInlineAllocator<ChunkSize>>> LandscapeCollisionComponentsToSamples;
+	TMap<const ULandscapeInfo*, FTransform> LandscapeTransformsMap;
+
+	for (int SampleIndex = 0; SampleIndex < Samples.Num(); ++SampleIndex)
+	{
+		const TPair<FTransform, FBox>& Sample = Samples[SampleIndex];
+		const ULandscapeInfo* LandscapeInfo = GetLandscapeInfo(Sample.Key.GetLocation());
+
+		// Implementation note: we reset the density here to simplify the early return cases + the points that wouldn't be kept at the end of the process
+		OutPoints[SampleIndex].Density = 0;
+
+		if (!LandscapeInfo || !LandscapeInfo->GetLandscapeProxy())
+		{
+			continue;
+		}
+
+		if (!LandscapeTransformsMap.Contains(LandscapeInfo))
+		{
+			LandscapeTransformsMap.Add(LandscapeInfo, LandscapeInfo->GetLandscapeProxy()->LandscapeActorToWorld());
+		}
+
+		const FTransform& LandscapeTransform = LandscapeTransformsMap[LandscapeInfo];
+
+		// Transform Box in local space -> box in world space -> box in landscape space
+		const FTransform BoundsTransformInLanscapeSpace = Sample.Key.GetRelativeTransform(LandscapeTransform);
+		FBox BoundsInLanscapeSpace = Sample.Value.TransformBy(BoundsTransformInLanscapeSpace);
+
+		// The landscape is transformed so that its coordinates are [0, ComponentSizeQuads], so we'll compute our min/max bounds here in landscape local space down below
+		// Gather all landscape heightfield components we need to test
+		const int ComponentMapKeyMinX = FMath::FloorToInt(BoundsInLanscapeSpace.Min.X / LandscapeInfo->ComponentSizeQuads);
+		const int ComponentMapKeyMaxX = FMath::FloorToInt(BoundsInLanscapeSpace.Max.X / LandscapeInfo->ComponentSizeQuads);
+		const int ComponentMapKeyMinY = FMath::FloorToInt(BoundsInLanscapeSpace.Min.Y / LandscapeInfo->ComponentSizeQuads);
+		const int ComponentMapKeyMaxY = FMath::FloorToInt(BoundsInLanscapeSpace.Max.Y / LandscapeInfo->ComponentSizeQuads);
+
+		for (int X = ComponentMapKeyMinX; X <= ComponentMapKeyMaxX; ++X)
+		{
+			for (int Y = ComponentMapKeyMinY; Y <= ComponentMapKeyMaxY; ++Y)
+			{
+				if (ULandscapeHeightfieldCollisionComponent* CollisionComponent = LandscapeInfo->XYtoCollisionComponentMap.FindRef(FIntPoint(X, Y)))
+				{
+					LandscapeCollisionComponentsToSamples.FindOrAdd(CollisionComponent).Add(SampleIndex);
+				}
+			}
+		}
+	}
+
+	if (LandscapeCollisionComponentsToSamples.IsEmpty())
+	{
+		return;
+	}
+
+	TArray<FCollisionShape, TInlineAllocator<ChunkSize>> CollisionShapes;
+	TArray<FPhysicsShapeAdapter_Chaos, TInlineAllocator<ChunkSize>> CollisionShapeAdapters;
+	TBitArray<TInlineAllocator<ChunkSize>> KeptSamples(false, Samples.Num());
+	CollisionShapes.Reserve(Samples.Num());
+
+	for (const TPair<FTransform, FBox>& Sample : Samples)
+	{
+		FCollisionShape& CollisionShape = CollisionShapes.Emplace_GetRef();
+		CollisionShape.SetBox(FVector3f(Sample.Value.GetExtent() * Sample.Key.GetScale3D()));
+		CollisionShapeAdapters.Emplace(Sample.Key.GetRotation(), CollisionShape);
+	}
+
+	// For each landscape collision component, lock, test all points, repeat.
+	TArray<ChaosInterface::FOverlapHit> OverlapHits;
+	for (const auto& ComponentToSamples : LandscapeCollisionComponentsToSamples)
+	{
+		ULandscapeHeightfieldCollisionComponent* Component = ComponentToSamples.Key;
+		const TArray<int, TInlineAllocator<ChunkSize>>& SampleIndices = ComponentToSamples.Value;
+
+		// Implementation note: this is an exploded version of OverlapComponentWithResult so we lock only once per chunk
+		// TODO: Replace this by the proper API call once it is available
+		TArray<Chaos::FPhysicsObjectHandle> Objects = Component->GetAllPhysicsObjects();
+		FLockedReadPhysicsObjectExternalInterface Interface = FPhysicsObjectExternalInterface::LockRead(Objects);
+		Objects = Objects.FilterByPredicate(
+			[&Interface](Chaos::FPhysicsObjectHandle Handle)
+			{
+				return !Interface->AreAllDisabled({ &Handle, 1 });
+			}
+		);
+
+		Chaos::FPhysicsObjectCollisionInterface_External CollisionInterface{ Interface.GetInterface() };
+
+		for (int ShapeIndex : SampleIndices)
+		{
+			if (KeptSamples[ShapeIndex])
+			{
+				continue;
+			}
+
+			const FPhysicsGeometry& Geometry = CollisionShapeAdapters[ShapeIndex].GetGeometry();
+			const FTransform& SampleTransform = Samples[ShapeIndex].Key;
+
+			if (CollisionInterface.ShapeOverlap(Objects, Geometry, { SampleTransform.GetRotation(), SampleTransform.GetLocation() }, OverlapHits))
+			{
+				if (!OverlapHits.IsEmpty())
+				{
+					KeptSamples[ShapeIndex] = true;
+				}
+
+				OverlapHits.Reset();
+			}
+		}
+	}
+
+	// Finally, write back the data to the OutPoints
+	for (int SampleIndex = 0; SampleIndex < Samples.Num(); ++SampleIndex)
+	{
+		FPCGPoint& OutPoint = OutPoints[SampleIndex];
+		if (KeptSamples[SampleIndex])
+		{
+			new(&OutPoint) FPCGPoint(Samples[SampleIndex].Key, /*Density=*/1.0f, /*Seed=*/0);
+			OutPoint.SetLocalBounds(Samples[SampleIndex].Value);
+		}
+	}
+}
+
 bool UPCGLandscapeData::ProjectPoint(const FTransform& InTransform, const FBox& InBounds, const FPCGProjectionParams& InParams, FPCGPoint& OutPoint, UPCGMetadata* OutMetadata) const
 {
 	if (!LandscapeCache)
@@ -366,13 +495,7 @@ const UPCGPointData* UPCGLandscapeData::CreatePointData(FPCGContext* Context, co
 
 	// Most proxies we gathered will have the same landscape info, we shouldn't loop multiple times
 	// on them, unless we add the box filtering - but even then, depending on the transform we could have overlaps
-	TSet<ULandscapeInfo*> AllLandscapeInfos;
-	for (const TPair<FBox, ULandscapeInfo*>& LandscapeInfoPair : LandscapeInfos)
-	{
-		AllLandscapeInfos.Add(LandscapeInfoPair.Value);
-	}
-
-	for(ULandscapeInfo* LandscapeInfo : AllLandscapeInfos)
+	for (ULandscapeInfo* LandscapeInfo : LandscapeInfos)
 	{
 		ALandscapeProxy* LandscapeProxy = LandscapeInfo ? LandscapeInfo->GetLandscapeProxy() : nullptr;
 		if (!LandscapeProxy)
@@ -472,12 +595,12 @@ const ULandscapeInfo* UPCGLandscapeData::GetLandscapeInfo(const FVector& InPosit
 	// Early out
 	if (LandscapeInfos.Num() == 1)
 	{
-		return LandscapeInfos[0].Value;
+		return LandscapeInfos[0];
 	}
 
 	// As discussed in the header, this loop here is the reason why we do not really support overlapping landscapes.
 	// TODO: we could maybe improve on this if we find the "nearest" landscape on a Z perspective, but this might still lead to issues
-	for (const TPair<FBox, ULandscapeInfo*>& LandscapeInfoPair : LandscapeInfos)
+	for (const TPair<FBox, ULandscapeInfo*>& LandscapeInfoPair : BoundsToLandscapeInfos)
 	{
 		if (PCGHelpers::IsInsideBoundsXY(LandscapeInfoPair.Key, InPosition))
 		{
@@ -498,6 +621,7 @@ UPCGSpatialData* UPCGLandscapeData::CopyInternal() const
 	NewLandscapeData->Bounds = Bounds;
 	NewLandscapeData->DataProps = DataProps;
 	NewLandscapeData->LandscapeInfos = LandscapeInfos;
+	NewLandscapeData->BoundsToLandscapeInfos = BoundsToLandscapeInfos;
 	NewLandscapeData->LandscapeCache = LandscapeCache;
 
 	return NewLandscapeData;
diff --git a/Engine/Plugins/Experimental/PCG/Source/PCG/Private/Data/PCGSpatialData.cpp b/Engine/Plugins/Experimental/PCG/Source/PCG/Private/Data/PCGSpatialData.cpp
--- a/Engine/Plugins/Experimental/PCG/Source/PCG/Private/Data/PCGSpatialData.cpp
+++ b/Engine/Plugins/Experimental/PCG/Source/PCG/Private/Data/PCGSpatialData.cpp
@@ -104,6 +104,32 @@ float UPCGSpatialData::GetDensityAtPosition(const FVector& InPosition) const
 	}
 }
 
+bool UPCGSpatialData::K2_SamplePoint(const FTransform& InTransform, const FBox& InBounds, FPCGPoint& OutPoint, UPCGMetadata* OutMetadata) const
+{
+	return SamplePoint(InTransform, InBounds, OutPoint, OutMetadata);
+}
+
+bool UPCGSpatialData::K2_ProjectPoint(const FTransform& InTransform, const FBox& InBounds, const FPCGProjectionParams& InParams, FPCGPoint& OutPoint, UPCGMetadata* OutMetadata) const
+{
+	return ProjectPoint(InTransform, InBounds, InParams, OutPoint, OutMetadata);
+}
+
+void UPCGSpatialData::SamplePoints(const TArrayView<const TPair<FTransform, FBox>>& InSamples, const TArrayView<FPCGPoint>& OutPoints, UPCGMetadata* OutMetadata) const
+{
+	TRACE_CPUPROFILER_EVENT_SCOPE(UPCGSpatialData::SamplePoints);
+	check(InSamples.Num() == OutPoints.Num());
+	for(int Index = 0; Index < InSamples.Num(); ++Index)
+	{
+		const TPair<FTransform, FBox>& Sample = InSamples[Index];
+		FPCGPoint& OutPoint = OutPoints[Index];
+
+		if (!SamplePoint(Sample.Key, Sample.Value, OutPoint, OutMetadata))
+		{
+			OutPoint.Density = 0;
+		}
+	}
+}
+
 bool UPCGSpatialData::ProjectPoint(const FTransform& InTransform, const FBox& InBounds, const FPCGProjectionParams& InParams, FPCGPoint& OutPoint, UPCGMetadata* OutMetadata) const
 {
 	// Fallback implementation - calls SamplePoint because SamplePoint was being used for projection previously.
@@ -132,6 +158,22 @@ bool UPCGSpatialData::ProjectPoint(const FTransform& InTransform, const FBox& In
 	return bResult;
 }
 
+void UPCGSpatialData::ProjectPoints(const TArrayView<const TPair<FTransform, FBox>>& InSamples, const FPCGProjectionParams& InParams, const TArrayView<FPCGPoint>& OutPoints, UPCGMetadata* OutMetadata) const
+{
+	TRACE_CPUPROFILER_EVENT_SCOPE(UPCGSpatialData::ProjectPoints);
+	check(InSamples.Num() == OutPoints.Num());
+	for (int Index = 0; Index < InSamples.Num(); ++Index)
+	{
+		const TPair<FTransform, FBox>& Sample = InSamples[Index];
+		FPCGPoint& OutPoint = OutPoints[Index];
+
+		if (!ProjectPoint(Sample.Key, Sample.Value, InParams, OutPoint, OutMetadata))
+		{
+			OutPoint.Density = 0;
+		}
+	}
+}
+
 UPCGIntersectionData* UPCGSpatialData::IntersectWith(const UPCGSpatialData* InOther) const
 {
 	UPCGIntersectionData* IntersectionData = NewObject<UPCGIntersectionData>();
diff --git a/Engine/Plugins/Experimental/PCG/Source/PCG/Public/Data/PCGLandscapeData.h b/Engine/Plugins/Experimental/PCG/Source/PCG/Public/Data/PCGLandscapeData.h
--- a/Engine/Plugins/Experimental/PCG/Source/PCG/Public/Data/PCGLandscapeData.h
+++ b/Engine/Plugins/Experimental/PCG/Source/PCG/Public/Data/PCGLandscapeData.h
@@ -63,6 +63,7 @@ public:
 	virtual FBox GetBounds() const override;
 	virtual FBox GetStrictBounds() const override;
 	virtual bool SamplePoint(const FTransform& Transform, const FBox& Bounds, FPCGPoint& OutPoint, UPCGMetadata* OutMetadata) const override;
+	virtual void SamplePoints(const TArrayView<const TPair<FTransform, FBox>>& Samples, const TArrayView<FPCGPoint>& OutPoints, UPCGMetadata* OutMetadata) const override;
 	virtual bool ProjectPoint(const FTransform& InTransform, const FBox& InBounds, const FPCGProjectionParams& InParams, FPCGPoint& OutPoint, UPCGMetadata* OutMetadata) const override;
 	virtual bool HasNonTrivialTransform() const override { return true; }
 protected:
@@ -104,7 +105,8 @@ protected:
 
 private:
 	// Transient data
-	TArray<TPair<FBox, ULandscapeInfo*>> LandscapeInfos;
+	TArray<TPair<FBox, ULandscapeInfo*>> BoundsToLandscapeInfos;
+	TArray<ULandscapeInfo*> LandscapeInfos;
 	UPCGLandscapeCache* LandscapeCache = nullptr;
 };
 
diff --git a/Engine/Plugins/Experimental/PCG/Source/PCG/Public/Data/PCGSpatialData.h b/Engine/Plugins/Experimental/PCG/Source/PCG/Public/Data/PCGSpatialData.h
--- a/Engine/Plugins/Experimental/PCG/Source/PCG/Public/Data/PCGSpatialData.h
+++ b/Engine/Plugins/Experimental/PCG/Source/PCG/Public/Data/PCGSpatialData.h
@@ -91,12 +91,29 @@ public:
 	virtual const UPCGPointData* ToPointData(FPCGContext* Context, const FBox& InBounds = FBox(EForceInit::ForceInit)) const PURE_VIRTUAL(UPCGSpatialData::ToPointData, return nullptr;);
 
 	/** Sample rotation, scale and other attributes from this data at the query position. Returns true if Transform location and Bounds overlaps this data. */
-	UFUNCTION(BlueprintCallable, Category = SpatialData)
 	virtual bool SamplePoint(const FTransform& Transform, const FBox& Bounds, FPCGPoint& OutPoint, UPCGMetadata* OutMetadata) const PURE_VIRTUAL(UPCGSpatialData::SamplePoint, return false;);
 
+	/** Performs multiple samples at the same time.
+	* Contrary to the single SamplePoint call, this is expected to set the density to 0 for points that were not overlapping - but the other properties can be anything.
+	* The OutPoints arrays is expected pre-allocated to the size of the Samples.
+	*/
+	virtual void SamplePoints(const TArrayView<const TPair<FTransform, FBox>>& Samples, const TArrayView<FPCGPoint>& OutPoints, UPCGMetadata* OutMetadata) const;
+
+	/** Sample rotation, scale and other attributes from this data at the query position. Returns true if Transform location and Bounds overlaps this data. */
+	UFUNCTION(BlueprintCallable, Category = SpatialData, meta = (DisplayName="SamplePoint"))
+	bool K2_SamplePoint(const FTransform& Transform, const FBox& Bounds, FPCGPoint& OutPoint, UPCGMetadata* OutMetadata) const;
+
 	/** Project the query point onto this data, and sample point and metadata information at the projected position. Returns true if successful. */
-	UFUNCTION(BlueprintCallable, Category = SpatialData)
 	virtual bool ProjectPoint(const FTransform& InTransform, const FBox& InBounds, const FPCGProjectionParams& InParams, FPCGPoint& OutPoint, UPCGMetadata* OutMetadata) const;
+	
+	/** Performs multiple projections of samples at the same time.
+	* Contrary to the single ProjectPoint call, this is expected to set the density to 0 for points that were not overlapping - but the other properties can be anything.
+	* The OutPoints arrays is expected pre-allocated to the size of the Samples.
+	*/
+	virtual void ProjectPoints(const TArrayView<const TPair<FTransform, FBox>>& Samples, const FPCGProjectionParams& InParams, const TArrayView<FPCGPoint>& OutPoints, UPCGMetadata* OutMetadata) const;
+
+	UFUNCTION(BlueprintCallable, Category = SpatialData, meta = (DisplayName="ProjectPoint"))
+	bool K2_ProjectPoint(const FTransform& InTransform, const FBox& InBounds, const FPCGProjectionParams& InParams, FPCGPoint& OutPoint, UPCGMetadata* OutMetadata) const;
 
 	/** Returns true if the data has a non-trivial transform */
 	UFUNCTION(BlueprintCallable, Category = SpatialData)
diff --git a/Engine/Plugins/Experimental/PCG/Source/PCG/Public/Data/PCGSpatialDataTpl.h b/Engine/Plugins/Experimental/PCG/Source/PCG/Public/Data/PCGSpatialDataTpl.h
new file mode 100644
--- /dev/null
+++ b/Engine/Plugins/Experimental/PCG/Source/PCG/Public/Data/PCGSpatialDataTpl.h
@@ -0,0 +1,66 @@
+// Copyright Epic Games, Inc. All Rights Reserved.
+
+#pragma once
+
+#include "PCGPoint.h"
+#include "Helpers/PCGAsync.h"
+
+#include "Algo/Transform.h"
+
+namespace FPCGSpatialDataProcessing
+{
+	constexpr int32 DefaultSamplePointsChunkSize = 256;
+
+	template<int ChunkSize, typename ProcessRangeFunc>
+	void SampleBasedRangeProcessing(FPCGAsyncState* AsyncState, ProcessRangeFunc&& InProcessRange, const TArray<FPCGPoint>& SourcePoints, TArray<FPCGPoint>& OutPoints)
+	{
+		const int NumIterations = SourcePoints.Num();
+
+		auto Initialize = [&OutPoints, NumIterations]()
+		{
+			OutPoints.SetNumUninitialized(NumIterations);
+		};
+
+		auto ProcessRange = [RangeFunc = MoveTemp(InProcessRange), &SourcePoints, &OutPoints](int32 StartReadIndex, int32 StartWriteIndex, int32 Count) -> int32
+		{
+			ensure(Count >= 0 && Count <= ChunkSize);
+
+			TArrayView<const FPCGPoint> IterationPoints(SourcePoints.GetData() + StartReadIndex, Count);
+			TArray<TPair<FTransform, FBox>, TInlineAllocator<ChunkSize>> Samples;
+			TArray<FPCGPoint, TInlineAllocator<ChunkSize>> RangeOutputPoints;
+
+			{
+				TRACE_CPUPROFILER_EVENT_SCOPE(FPCGSpatialDataProcessing::SamplePoints::PrepareSamples)
+				Algo::Transform(IterationPoints, Samples, [](const FPCGPoint& Point) { return TPair<FTransform, FBox>(Point.Transform, Point.GetLocalBounds()); });
+			}
+
+			{
+				TRACE_CPUPROFILER_EVENT_SCOPE(FPCGSpatialDataProcessing::SamplePoints::RangeFunc)
+				RangeFunc(Samples, IterationPoints, RangeOutputPoints);
+			}
+
+			{
+				TRACE_CPUPROFILER_EVENT_SCOPE(FPCGSpatialDataProcessing::SamplePoints::RangeCopyResults)
+				// Copy back the points
+				FMemory::Memcpy((void*)(OutPoints.GetData() + StartWriteIndex), (void*)(RangeOutputPoints.GetData()), sizeof(FPCGPoint) * RangeOutputPoints.Num());
+			}
+
+			return RangeOutputPoints.Num();
+		};
+
+		auto MoveDataRange = [&OutPoints](int32 ReadIndex, int32 WriteIndex, int32 Count)
+		{
+			// Implementation note: if the array is to be moved to a range partially overlapping itself, it's important not to use Memcpy here
+			FMemory::Memmove((void*)(OutPoints.GetData() + WriteIndex), (void*)(OutPoints.GetData() + ReadIndex), sizeof(FPCGPoint) * Count);
+		};
+
+		auto Finished = [&OutPoints](int32 Count)
+		{
+			// Shrinking can have a big impact on the performance, but without it, we can also hold a big chunk of wasted memory.
+			// Might revisit later if the performance impact is too big.
+			OutPoints.SetNum(Count);
+		};
+
+		ensure(FPCGAsync::AsyncProcessingRangeEx(AsyncState, NumIterations, Initialize, ProcessRange, MoveDataRange, Finished, /*bEnableTimeSlicing=*/false, ChunkSize));
+	}
+}
\ No newline at end of file
diff --git a/Engine/Plugins/Experimental/PCG/Source/PCG/Public/Helpers/PCGAsync.h b/Engine/Plugins/Experimental/PCG/Source/PCG/Public/Helpers/PCGAsync.h
--- a/Engine/Plugins/Experimental/PCG/Source/PCG/Public/Helpers/PCGAsync.h
+++ b/Engine/Plugins/Experimental/PCG/Source/PCG/Public/Helpers/PCGAsync.h
@@ -93,8 +93,8 @@ namespace FPCGAsync
 
 	namespace Private
 	{
-	template <typename InitializeFunc, typename InnerLoopFunc, typename MoveDataFunc, typename FinishedFunc>
-	bool AsyncProcessing(FPCGAsyncState& AsyncState, int32 NumIterations, InitializeFunc&& Initialize, InnerLoopFunc&& IterationInnerLoop, MoveDataFunc&& MoveData, FinishedFunc&& Finished, const bool bInEnableTimeSlicing, const int32 InChunkSize)
+	template <typename InitializeFunc, typename InnerLoopFunc, typename MoveDataRangeFunc, typename FinishedFunc>
+	bool AsyncProcessing(FPCGAsyncState& AsyncState, int32 NumIterations, InitializeFunc&& Initialize, InnerLoopFunc&& IterationInnerLoop, MoveDataRangeFunc&& MoveDataRange, FinishedFunc&& Finished, const bool bInEnableTimeSlicing, const int32 InChunkSize)
 	{
 		TRACE_CPUPROFILER_EVENT_SCOPE(FPCGAsync::AsyncProcessing);
 
@@ -278,10 +278,8 @@ namespace FPCGAsync
 						else
 						{
 							// Otherwise collapse
-							for (int32 i = 0; i < NumberOfElementsWritten; ++i)
-							{
-								MoveData(ReadStartIndexForCurrentChunkToCollapse + i, AsyncState.AsyncCurrentWriteIndex++);
-							}
+							MoveDataRange(ReadStartIndexForCurrentChunkToCollapse, AsyncState.AsyncCurrentWriteIndex, NumberOfElementsWritten);
+							AsyncState.AsyncCurrentWriteIndex += NumberOfElementsWritten;
 						}
 
 						AsyncState.AsyncCurrentReadIndex += ChunkSize;
@@ -317,6 +315,48 @@ namespace FPCGAsync
 	}
 	}
 
+	/**
+	* A Helper for generic parallel loops, with support for timeslicing, specialized to work on ranges. This version only uses indexes allowing more flexible usage
+	* to process many arrays in parallel or use it for other batch updating.
+	* 
+	* Work will be separated in chunks, that will be processed in parallel. Main thread will then collapse incoming data from async tasks.
+	* Will use AsyncState.ShouldStop() to stop execution if timeslicing is enabled.
+	* Important info: 
+	*   - We will finish to process and collapse data for all data already in process, even if we need to stop. To mitigate this, try to use small chunk sizes.
+	*   - To avoid infinite loops (when we should stop even before starting working), we will at least process 1 chunk of data per thread.
+	*   - To have async tasks, you need to have at least 3 available threads (main thread + 2 futures). Otherwise, we will only process on the main thread, without collapse.
+	* 
+	* @param AsyncState - The context containing the information about how many tasks we can launch, async read/write index for the current job and a function to know if we need to stop processing.
+	* @param NumIterations - The number of calls that will be done to the provided function, also an upper bound on the number of data generated.
+	* @param Initialize - Signature: void(). A function that will be called once on the first timeslice, where you can reserve data for processing
+	* @param ProcessRange - Signature: bool(int32 StartReadIndex, int32 StartWriteIndex, int32 Count). A function that processes a range of values and returns the number of written values 
+	* @param MoveDataRange - Signature: void(int32 ReadIndex, int32 WriteIndex, int32 Count). If the processing filters points, this will be used to move elements in chunk from one range to another 
+	* @param Finished - Signature: void(int32 Count). Called once on finished, and tells you the total count of points written.
+	* @param bEnableTimeSlicing - If false, we will not stop until all the processing is done.
+	* @param ChunkSize - Size of the chunks to cut the input data with
+	* @returns true if the processing is done, false otherwise. Use this to know if you need to reschedule the task.
+	*/
+	template <typename InitializeFunc, typename ProcessRangeFunc, typename MoveDataRangeFunc, typename FinishedFunc>
+	bool AsyncProcessingRangeEx(FPCGAsyncState* AsyncState, int32 NumIterations, InitializeFunc&& Initialize, ProcessRangeFunc&& ProcessRange, MoveDataRangeFunc&& MoveDataRange, FinishedFunc&& Finished, const bool bEnableTimeSlicing, const int32 ChunkSize = 64)
+	{
+		if (AsyncState && !AsyncState->bIsRunningAsyncCall)
+		{
+			AsyncState->bIsRunningAsyncCall = true;
+			const bool bIsDone = Private::AsyncProcessing(*AsyncState, NumIterations, Initialize, ProcessRange, MoveDataRange, Finished, bEnableTimeSlicing, ChunkSize);
+			AsyncState->bIsRunningAsyncCall = false;
+			return bIsDone;
+		}
+		else
+		{
+			// Can't use time slicing without an async state or while running in another async call (it will mess up with async indexes). 
+			// We also force using one thread (the current one).
+			FPCGAsyncState DummyState;
+			DummyState.NumAvailableTasks = 1;
+			DummyState.bIsRunningAsyncCall = true;
+			return Private::AsyncProcessing(DummyState, NumIterations, Initialize, ProcessRange, MoveDataRange, Finished, /*bEnableTimeSlicing=*/false, ChunkSize);
+		}
+	}
+
 	/**
 	* A Helper for generic parallel loops, with support for timeslicing. This version only uses indexes allowing more flexible usage
 	* to process many arrays in parallel or use it for other batch updating.
@@ -356,22 +396,15 @@ namespace FPCGAsync
 			return NumPointsWritten;
 		};
 
-		if (AsyncState && !AsyncState->bIsRunningAsyncCall)
-		{
-			AsyncState->bIsRunningAsyncCall = true;
-			const bool bIsDone = Private::AsyncProcessing(*AsyncState, NumIterations, Initialize, IterationInnerLoop, MoveData, Finished, bEnableTimeSlicing, ChunkSize);
-			AsyncState->bIsRunningAsyncCall = false;
-			return bIsDone;
-		}
-		else
+		auto MoveDataRange = [Func = MoveTemp(MoveData)](int32 ReadIndex, int32 WriteIndex, int32 Count)
 		{
-			// Can't use time slicing without an async state or while running in another async call (it will mess up with async indexes). 
-			// We also force using one thread (the current one).
-			FPCGAsyncState DummyState;
-			DummyState.NumAvailableTasks = 1;
-			DummyState.bIsRunningAsyncCall = true;
-			return Private::AsyncProcessing(DummyState, NumIterations, Initialize, IterationInnerLoop, MoveData, Finished, /*bEnableTimeSlicing=*/false, ChunkSize);
-		}
+			for (int Index = 0; Index < Count; ++Index)
+			{
+				Func(ReadIndex + Index, WriteIndex + Index);
+			}
+		};
+
+		return AsyncProcessingRangeEx(AsyncState, NumIterations, Initialize, IterationInnerLoop, MoveDataRange, Finished, bEnableTimeSlicing, ChunkSize);
 	}
  
  	/**
@@ -409,13 +442,13 @@ namespace FPCGAsync
 			return Count;
 		};
 
-		auto MoveData = [](int32, int32) { ensure(false); };
+		auto MoveDataRange = [](int32, int32, int32) { ensure(false); };
 		auto Finished = [NumIterations](int32 Count) { ensure(NumIterations == Count); };
 
 		if (AsyncState && !AsyncState->bIsRunningAsyncCall)
 		{
 			AsyncState->bIsRunningAsyncCall = true;
-			const bool bIsDone = Private::AsyncProcessing(*AsyncState, NumIterations, Initialize, IterationInnerLoop, MoveData, Finished, bEnableTimeSlicing, ChunkSize);
+			const bool bIsDone = Private::AsyncProcessing(*AsyncState, NumIterations, Initialize, IterationInnerLoop, MoveDataRange, Finished, bEnableTimeSlicing, ChunkSize);
 			AsyncState->bIsRunningAsyncCall = false;
 			return bIsDone;
 		}
@@ -426,7 +459,7 @@ namespace FPCGAsync
 			FPCGAsyncState DummyState;
 			DummyState.NumAvailableTasks = 1;
 			DummyState.bIsRunningAsyncCall = true;
-			return Private::AsyncProcessing(DummyState, NumIterations, Initialize, IterationInnerLoop, MoveData, Finished, /*bEnableTimeSlicing=*/false, ChunkSize);
+			return Private::AsyncProcessing(DummyState, NumIterations, Initialize, IterationInnerLoop, MoveDataRange, Finished, /*bEnableTimeSlicing=*/false, ChunkSize);
 		}
 	}
 
@@ -471,7 +504,7 @@ namespace FPCGAsync
 		{
 			// Shrinking can have a big impact on the performance, but without it, we can also hold a big chunk of wasted memory.
 			// Might revisit later if the performance impact is too big.
-			OutData.SetNum(Count);			
+			OutData.SetNum(Count);
 		};
 
 		return AsyncProcessingEx(AsyncState, NumIterations, Initialize, IterationInnerLoop, MoveData, Finished, bEnableTimeSlicing, ChunkSize);
