Extended filter by attribute capabilities (allows to check properties on points too)

#rnx

diff --git a/Engine/Plugins/Experimental/PCG/Source/PCG/Private/Elements/PCGFilterByAttribute.cpp b/Engine/Plugins/Experimental/PCG/Source/PCG/Private/Elements/PCGFilterByAttribute.cpp
--- a/Engine/Plugins/Experimental/PCG/Source/PCG/Private/Elements/PCGFilterByAttribute.cpp
+++ b/Engine/Plugins/Experimental/PCG/Source/PCG/Private/Elements/PCGFilterByAttribute.cpp
@@ -4,7 +4,8 @@
 
 #include "PCGContext.h"
 #include "PCGData.h"
-#include "Metadata/PCGMetadata.h"
+#include "Metadata/Accessors/IPCGAttributeAccessor.h"
+#include "Metadata/Accessors/PCGAttributeAccessorHelpers.h"
 
 #define LOCTEXT_NAMESPACE "PCGFilterByAttributeElement"
 
@@ -42,7 +43,8 @@ bool FPCGFilterByAttributeElement::ExecuteInternal(FPCGContext* Context) const
 	const UPCGFilterByAttributeSettings* Settings = Context->GetInputSettings<UPCGFilterByAttributeSettings>();
 	check(Settings);
 
-	const FName Attribute = FName(Settings->Attribute);
+	FPCGAttributePropertySelector Selector;
+	Selector.Update(Settings->Attribute.ToString());
 
 	TArray<FPCGTaggedData> Inputs = Context->InputData.GetInputsByPin(PCGPinConstants::DefaultInputLabel);
 	TArray<FPCGTaggedData>& Outputs = Context->OutputData.TaggedData;
@@ -51,9 +53,15 @@ bool FPCGFilterByAttributeElement::ExecuteInternal(FPCGContext* Context) const
 	{
 		FPCGTaggedData& Output = Outputs.Add_GetRef(Input);
 		Output.Pin = PCGPinConstants::DefaultOutFilterLabel;
-		
-		const UPCGMetadata* Metadata = Input.Data ? Input.Data->ConstMetadata() : nullptr;
-		if (Metadata && Metadata->HasAttribute(Attribute))
+
+		if (!Input.Data)
+		{
+			continue;
+		}
+
+		TUniquePtr<const IPCGAttributeAccessor> Accessor = PCGAttributeAccessorHelpers::CreateConstAccessor(Input.Data, Selector, /*bQuiet=*/true);
+
+		if (Accessor && Accessor.IsValid())
 		{
 			Output.Pin = PCGPinConstants::DefaultInFilterLabel;
 		}
diff --git a/Engine/Plugins/Experimental/PCG/Source/PCG/Private/Metadata/Accessors/PCGAttributeAccessorHelpers.cpp b/Engine/Plugins/Experimental/PCG/Source/PCG/Private/Metadata/Accessors/PCGAttributeAccessorHelpers.cpp
--- a/Engine/Plugins/Experimental/PCG/Source/PCG/Private/Metadata/Accessors/PCGAttributeAccessorHelpers.cpp
+++ b/Engine/Plugins/Experimental/PCG/Source/PCG/Private/Metadata/Accessors/PCGAttributeAccessorHelpers.cpp
@@ -74,7 +74,7 @@ namespace PCGAttributeAccessorHelpers
 			std::is_const_v<AccessorType> == std::is_const_v<DataType>
 		>
 	>
-	TUniquePtr<AccessorType> CreateAccessorImpl(DataType* InData, const FPCGAttributePropertySelector& InSelector)
+	TUniquePtr<AccessorType> CreateAccessorImpl(DataType* InData, const FPCGAttributePropertySelector& InSelector, bool bQuiet)
 	{
 		const FName Name = InSelector.GetName();
 		TUniquePtr<IPCGAttributeAccessor> Accessor;
@@ -99,7 +99,11 @@ namespace PCGAttributeAccessorHelpers
 
 			if (!Accessor.IsValid())
 			{
-				UE_LOG(LogPCG, Error, TEXT("[PCGAttributeAccessorHelpers::ConstAccessor] Expected to select an extra property but the data doesn't support this property."));
+				if (!bQuiet)
+				{
+					UE_LOG(LogPCG, Error, TEXT("[PCGAttributeAccessorHelpers::ConstAccessor] Expected to select an extra property but the data doesn't support this property."));
+				}
+				
 				return TUniquePtr<IPCGAttributeAccessor>();
 			}
 		}
@@ -108,7 +112,11 @@ namespace PCGAttributeAccessorHelpers
 		// We can't continue if it is a property wanted.
 		if (InSelector.GetSelection() == EPCGAttributePropertySelection::PointProperty && !Accessor.IsValid())
 		{
-			UE_LOG(LogPCG, Error, TEXT("[PCGAttributeAccessorHelpers::CreateAccessor] Expected to select a property but the data doesn't support this property."));
+			if (!bQuiet)
+			{
+				UE_LOG(LogPCG, Error, TEXT("[PCGAttributeAccessorHelpers::CreateAccessor] Expected to select a property but the data doesn't support this property."));
+			}
+			
 			return TUniquePtr<IPCGAttributeAccessor>();
 		}
 
@@ -136,7 +144,11 @@ namespace PCGAttributeAccessorHelpers
 			Accessor = CreateChainAccessor(std::move(Accessor), FName(ExtraName), bSuccess);
 			if (!bSuccess)
 			{
-				UE_LOG(LogPCG, Error, TEXT("[PCGAttributeAccessorHelpers::CreateAccessor] Extra selectors don't match existing properties."));
+				if (!bQuiet)
+				{
+					UE_LOG(LogPCG, Error, TEXT("[PCGAttributeAccessorHelpers::CreateAccessor] Extra selectors don't match existing properties."));
+				}
+				
 				return TUniquePtr<AccessorType>();
 			}
 		}
@@ -146,7 +158,11 @@ namespace PCGAttributeAccessorHelpers
 		{
 			if (Accessor && Accessor->IsReadOnly())
 			{
-				UE_LOG(LogPCG, Error, TEXT("[PCGAttributeAccessorHelpers::CreateAccessor] Attribute can not be written into, since it is read-only."));
+				if (!bQuiet)
+				{
+					UE_LOG(LogPCG, Error, TEXT("[PCGAttributeAccessorHelpers::CreateAccessor] Attribute can not be written into, since it is read-only."));
+				}
+				
 				return TUniquePtr<AccessorType>();
 			}
 
@@ -583,24 +599,24 @@ TUniquePtr<const IPCGAttributeAccessor> PCGAttributeAccessorHelpers::CreateConst
 	return TUniquePtr<const IPCGAttributeAccessor>{};
 }
 
-TUniquePtr<const IPCGAttributeAccessor> PCGAttributeAccessorHelpers::CreateConstAccessor(const UPCGData* InData, const FPCGAttributePropertySelector& InSelector)
+TUniquePtr<const IPCGAttributeAccessor> PCGAttributeAccessorHelpers::CreateConstAccessor(const UPCGData* InData, const FPCGAttributePropertySelector& InSelector, bool bQuiet)
 {
-	return CreateAccessorImpl<const IPCGAttributeAccessor, const UPCGData>(InData, InSelector);
+	return CreateAccessorImpl<const IPCGAttributeAccessor, const UPCGData>(InData, InSelector, bQuiet);
 }
 
-TUniquePtr<const IPCGAttributeAccessor> PCGAttributeAccessorHelpers::CreateConstAccessor(const FPCGMetadataAttributeBase* InAttribute, const UPCGMetadata* InMetadata)
+TUniquePtr<const IPCGAttributeAccessor> PCGAttributeAccessorHelpers::CreateConstAccessor(const FPCGMetadataAttributeBase* InAttribute, const UPCGMetadata* InMetadata, bool bQuiet)
 {
-	return CreateAttributeAccessorImpl<const IPCGAttributeAccessor>(const_cast<FPCGMetadataAttributeBase*>(InAttribute), const_cast<UPCGMetadata*>(InMetadata));
+	return CreateAttributeAccessorImpl<const IPCGAttributeAccessor>(const_cast<FPCGMetadataAttributeBase*>(InAttribute), const_cast<UPCGMetadata*>(InMetadata), bQuiet);
 }
 
-TUniquePtr<IPCGAttributeAccessor> PCGAttributeAccessorHelpers::CreateAccessor(UPCGData* InData, const FPCGAttributePropertySelector& InSelector)
+TUniquePtr<IPCGAttributeAccessor> PCGAttributeAccessorHelpers::CreateAccessor(UPCGData* InData, const FPCGAttributePropertySelector& InSelector, bool bQuiet)
 {
-	return CreateAccessorImpl<IPCGAttributeAccessor, UPCGData>(InData, InSelector);
+	return CreateAccessorImpl<IPCGAttributeAccessor, UPCGData>(InData, InSelector, bQuiet);
 }
 
-TUniquePtr<IPCGAttributeAccessor> PCGAttributeAccessorHelpers::CreateAccessor(FPCGMetadataAttributeBase* InAttribute, UPCGMetadata* InMetadata)
+TUniquePtr<IPCGAttributeAccessor> PCGAttributeAccessorHelpers::CreateAccessor(FPCGMetadataAttributeBase* InAttribute, UPCGMetadata* InMetadata, bool bQuiet)
 {
-	return CreateAttributeAccessorImpl<IPCGAttributeAccessor>(InAttribute, InMetadata);
+	return CreateAttributeAccessorImpl<IPCGAttributeAccessor>(InAttribute, InMetadata, bQuiet);
 }
 
 TUniquePtr<const IPCGAttributeAccessorKeys> PCGAttributeAccessorHelpers::CreateConstKeys(const UPCGData* InData, const FPCGAttributePropertySelector& InSelector)
diff --git a/Engine/Plugins/Experimental/PCG/Source/PCG/Public/Metadata/Accessors/PCGAttributeAccessorHelpers.h b/Engine/Plugins/Experimental/PCG/Source/PCG/Public/Metadata/Accessors/PCGAttributeAccessorHelpers.h
--- a/Engine/Plugins/Experimental/PCG/Source/PCG/Public/Metadata/Accessors/PCGAttributeAccessorHelpers.h
+++ b/Engine/Plugins/Experimental/PCG/Source/PCG/Public/Metadata/Accessors/PCGAttributeAccessorHelpers.h
@@ -62,12 +62,12 @@ namespace PCGAttributeAccessorHelpers
 	* but in the case of properties, it either has to be the appropriate type or null.
 	* Make sure to update your selector before-hand if you want to support "@Last"
 	*/
-	PCG_API TUniquePtr<const IPCGAttributeAccessor> CreateConstAccessor(const UPCGData* InData, const FPCGAttributePropertySelector& InSelector);
+	PCG_API TUniquePtr<const IPCGAttributeAccessor> CreateConstAccessor(const UPCGData* InData, const FPCGAttributePropertySelector& InSelector, bool bQuiet = false);
 
 	/** 
 	* Creates a const accessor to an attribute without requiring a selector.
 	*/
-	PCG_API TUniquePtr<const IPCGAttributeAccessor> CreateConstAccessor(const FPCGMetadataAttributeBase* InAttribute, const UPCGMetadata* InMetadata);
+	PCG_API TUniquePtr<const IPCGAttributeAccessor> CreateConstAccessor(const FPCGMetadataAttributeBase* InAttribute, const UPCGMetadata* InMetadata, bool bQuiet = false);
 
 	/**
 	* Creates a accessor to the property or attribute pointed at by the InSelector.
@@ -75,12 +75,12 @@ namespace PCGAttributeAccessorHelpers
 	* but in the case of properties, it either has to be the appropriate type or null.
 	* Make sure to update your selector before-hand if you want to support "@Source". Otherwise the creation will fail.
 	*/
-	PCG_API TUniquePtr<IPCGAttributeAccessor> CreateAccessor(UPCGData* InData, const FPCGAttributePropertySelector& InSelector);
+	PCG_API TUniquePtr<IPCGAttributeAccessor> CreateAccessor(UPCGData* InData, const FPCGAttributePropertySelector& InSelector, bool bQuiet = false);
 
 	/**
 	* Creates an accessor to an attribute without requiring a selector.
 	*/
-	PCG_API TUniquePtr<IPCGAttributeAccessor> CreateAccessor(FPCGMetadataAttributeBase* InAttribute, UPCGMetadata* InMetadata);
+	PCG_API TUniquePtr<IPCGAttributeAccessor> CreateAccessor(FPCGMetadataAttributeBase* InAttribute, UPCGMetadata* InMetadata, bool bQuiet = false);
 
 	PCG_API TUniquePtr<const IPCGAttributeAccessorKeys> CreateConstKeys(const UPCGData* InData, const FPCGAttributePropertySelector& InSelector);
 	PCG_API TUniquePtr<IPCGAttributeAccessorKeys> CreateKeys(UPCGData* InData, const FPCGAttributePropertySelector& InSelector);
