[PCG] Tweak sorting order of debug object tree entries.

Makes items with children come after items without children, to avoid confusing situations where some debug objects appear far away from their parent when the tree is expanded.

diff --git a/Engine/Plugins/Experimental/PCG/Source/PCGEditor/Private/SPCGEditorGraphDebugObjectTree.cpp b/Engine/Plugins/Experimental/PCG/Source/PCGEditor/Private/SPCGEditorGraphDebugObjectTree.cpp
--- a/Engine/Plugins/Experimental/PCG/Source/PCGEditor/Private/SPCGEditorGraphDebugObjectTree.cpp
+++ b/Engine/Plugins/Experimental/PCG/Source/PCGEditor/Private/SPCGEditorGraphDebugObjectTree.cpp
@@ -106,7 +106,15 @@ void FPCGEditorGraphDebugObjectItem::SortChildren(bool bIsAscending, bool bIsRec
 		// Support for sorting by loop index.
 		const int32 IndexLHS = InLHS->GetSortPriority();
 		const int32 IndexRHS = InRHS->GetSortPriority();
-		if (IndexLHS == INDEX_NONE || IndexRHS == INDEX_NONE)
+
+		const int HasChildrenLHS = InLHS->Children.Num() ? 1 : 0;
+		const int HasChildrenRHS = InRHS->Children.Num() ? 1 : 0;
+
+		if (HasChildrenLHS != HasChildrenRHS)
+		{
+			return HasChildrenLHS < HasChildrenRHS;
+		}
+		else if (IndexLHS == INDEX_NONE || IndexRHS == INDEX_NONE)
 		{
 			return (InLHS->GetLabel() < InRHS->GetLabel()) == bIsAscending;
 		}
