[PCG] Remove debug object dropdown which, replace with button to spawn the debug object tree tab.

diff --git a/Engine/Plugins/Experimental/PCG/Source/PCGEditor/Private/PCGEditor.cpp b/Engine/Plugins/Experimental/PCG/Source/PCGEditor/Private/PCGEditor.cpp
--- a/Engine/Plugins/Experimental/PCG/Source/PCGEditor/Private/PCGEditor.cpp
+++ b/Engine/Plugins/Experimental/PCG/Source/PCGEditor/Private/PCGEditor.cpp
@@ -31,7 +31,6 @@
 #include "PCGEditorUtils.h"
 #include "SPCGEditorGraphAttributeListView.h"
 #include "SPCGEditorGraphDebugObjectTree.h"
-#include "SPCGEditorGraphDebugObjectWidget.h"
 #include "SPCGEditorGraphDeterminism.h"
 #include "SPCGEditorGraphFind.h"
 #include "SPCGEditorGraphLogView.h"
@@ -133,7 +132,6 @@ void FPCGEditor::Initialize(const EToolkitMode::Type InMode, const TSharedPtr<cl
 
 	GraphEditorWidget = CreateGraphEditorWidget();
 	PaletteWidget = CreatePaletteWidget();
-	DebugObjectWidget = CreateDebugObjectWidget();
 	DebugObjectTreeWidget = CreateDebugObjectTreeWidget();
 	FindWidget = CreateFindWidget();
 	AttributesWidget = CreateAttributesWidget();
@@ -224,7 +222,7 @@ UPCGEditorGraph* FPCGEditor::GetPCGEditorGraph()
 	return PCGEditorGraph;
 }
 
-void FPCGEditor::SetStackBeingInspected(const FPCGStack& FullStack, FPCGDebugObjectSelectionMethod SelectionMethod)
+void FPCGEditor::SetStackBeingInspected(const FPCGStack& FullStack)
 {
 	UPCGComponent* OldComponent = PCGComponentBeingInspected.Get();
 	UPCGComponent* NewComponent = const_cast<UPCGComponent*>(FullStack.GetRootComponent());
@@ -256,15 +254,6 @@ void FPCGEditor::SetStackBeingInspected(const FPCGStack& FullStack, FPCGDebugObj
 			}
 		}
 
-		if (SelectionMethod == FPCGDebugObjectSelectionMethod::DebugObjectTree)
-		{
-			DebugObjectWidget->SetDebugObjectSelection(StackBeingInspected);
-		}
-		else
-		{
-			DebugObjectTreeWidget->SetDebugObjectSelection(StackBeingInspected);
-		}
-
 		UpdateDebugAfterComponentSelection(OldComponent, NewComponent, bNewComponentStartedInspecting);
 
 		check(PCGEditorGraph);
@@ -388,7 +377,7 @@ void FPCGEditor::RegisterTabSpawners(const TSharedRef<FTabManager>& InTabManager
 		.SetDisplayName(LOCTEXT("PaletteTab", "Palette"))
 		.SetGroup(WorkspaceMenuCategoryRef);
 
-	InTabManager->RegisterTabSpawner(FPCGEditor_private::DebugObjectID, FOnSpawnTab::CreateSP(this, &FPCGEditor::SpawnTab_DebugObject))
+	InTabManager->RegisterTabSpawner(FPCGEditor_private::DebugObjectID, FOnSpawnTab::CreateSP(this, &FPCGEditor::SpawnTab_DebugObjectTree))
 		.SetDisplayName(LOCTEXT("DebugTab", "Debug Object Tree"))
 		.SetGroup(WorkspaceMenuCategoryRef);
 
@@ -507,11 +496,12 @@ void FPCGEditor::RegisterToolbar() const
 		const FPCGEditorCommands& PCGEditorCommands = FPCGEditorCommands::Get();
 		const FToolMenuInsert InsertAfterAssetSection("Asset", EToolMenuInsertType::After);
 		FToolMenuSection& Section = ToolBar->AddSection("PCGToolbar", TAttribute<FText>(), InsertAfterAssetSection);
+
 		Section.AddEntry(FToolMenuEntry::InitToolBarButton(
 			PCGEditorCommands.Find,
 			TAttribute<FText>(),
 			TAttribute<FText>(),
-			FSlateIcon(FAppStyle::GetAppStyleSetName(), "BlueprintEditor.FindInBlueprint"))); // TODO, use own application style?
+			FSlateIcon(FAppStyle::GetAppStyleSetName(), "BlueprintEditor.FindInBlueprint")));
 
 		Section.AddEntry(FToolMenuEntry::InitToolBarButton(
 			FPCGEditorCommands::Get().PauseAutoRegeneration,
@@ -531,20 +521,14 @@ void FPCGEditor::RegisterToolbar() const
 			TAttribute<FText>(),
 			FSlateIcon(FAppStyle::GetAppStyleSetName(), "Icons.Toolbar.Stop")));
 
-		Section.AddSeparator(NAME_None);
-		Section.AddDynamicEntry("Debugging", FNewToolMenuSectionDelegate::CreateLambda([](FToolMenuSection& InSection)
-		{
-			const UPCGEditorMenuContext* Context = InSection.FindContext<UPCGEditorMenuContext>();
-			if (Context && Context->PCGEditor.IsValid())
-			{
-				InSection.AddEntry(FToolMenuEntry::InitWidget(
-					"SelectedDebugObjectWidget",
-					Context->PCGEditor.Pin()->DebugObjectWidget.ToSharedRef(),
-					FText::GetEmpty()));
-			}
-		}));
+		Section.AddEntry(FToolMenuEntry::InitToolBarButton(
+			FPCGEditorCommands::Get().OpenDebugObjectTreeTab,
+			TAttribute<FText>(),
+			TAttribute<FText>(),
+			FSlateIcon(FAppStyle::GetAppStyleSetName(), "LevelEditor.CreateBlankBlueprintClass")));
 
 		Section.AddSeparator(NAME_None);
+
 		Section.AddEntry(FToolMenuEntry::InitToolBarButton(
 			 PCGEditorCommands.RunDeterminismGraphTest,
 			 TAttribute<FText>(),
@@ -552,6 +536,7 @@ void FPCGEditor::RegisterToolbar() const
 			 FSlateIcon(FAppStyle::GetAppStyleSetName(), "BlueprintDebugger.TabIcon")));
 
 		Section.AddSeparator(NAME_None);
+
 		Section.AddEntry(FToolMenuEntry::InitToolBarButton(
 			 PCGEditorCommands.EditGraphSettings,
 			 TAttribute<FText>(),
@@ -583,6 +568,13 @@ void FPCGEditor::BindCommands()
 		FExecuteAction::CreateSP(this, &FPCGEditor::OnCancelExecution_Clicked),
 		FCanExecuteAction::CreateSP(this, &FPCGEditor::IsCurrentlyGenerating));
 
+	// Left on UI as a disabled button if debug object tree tab already open. This is a deliberate
+	// hint for 5.4 to help direct users to use the tree.
+	ToolkitCommands->MapAction(
+		PCGEditorCommands.OpenDebugObjectTreeTab,
+		FExecuteAction::CreateSP(this, &FPCGEditor::OnOpenDebugObjectTreeTab_Clicked),
+		FCanExecuteAction::CreateSP(this, &FPCGEditor::IsDebugObjectTreeTabClosed));
+
 	ToolkitCommands->MapAction(
 		PCGEditorCommands.RunDeterminismGraphTest,
 		FExecuteAction::CreateSP(this, &FPCGEditor::OnDeterminismGraphTest),
@@ -714,6 +706,16 @@ bool FPCGEditor::IsCurrentlyGenerating() const
 	return false;
 }
 
+bool FPCGEditor::IsDebugObjectTreeTabClosed() const
+{
+	return !TabManager.IsValid() || !TabManager->FindExistingLiveTab(FPCGEditor_private::DebugObjectID).IsValid();
+}
+
+void FPCGEditor::OnOpenDebugObjectTreeTab_Clicked()
+{
+	TabManager->TryInvokeTab(FPCGEditor_private::DebugObjectID);
+}
+
 bool FPCGEditor::CanRunDeterminismNodeTest() const
 {
 	check(GraphEditorWidget.IsValid());
@@ -2132,11 +2134,6 @@ TSharedRef<SPCGEditorGraphNodePalette> FPCGEditor::CreatePaletteWidget()
 	return SNew(SPCGEditorGraphNodePalette);
 }
 
-TSharedRef<SPCGEditorGraphDebugObjectWidget> FPCGEditor::CreateDebugObjectWidget()
-{
-	return SNew(SPCGEditorGraphDebugObjectWidget, SharedThis(this));
-}
-
 TSharedRef<SPCGEditorGraphDebugObjectTree> FPCGEditor::CreateDebugObjectTreeWidget()
 {
 	return SNew(SPCGEditorGraphDebugObjectTree, SharedThis(this));
@@ -2356,7 +2353,6 @@ void FPCGEditor::OnGraphStructureChanged(UPCGGraphInterface* InGraph)
 
 void FPCGEditor::OnComponentGenerationCompleteOrCancelled()
 {
-	DebugObjectWidget->RefreshDebugObjects();
 	DebugObjectTreeWidget->RequestRefresh();
 
 	// If we are debugging the graph cache then we need to refresh the cache count displayed in the title after every generation.
@@ -2389,10 +2385,6 @@ void FPCGEditor::OnMapChanged(UWorld* InWorld, EMapChangeType InMapChangedType)
 		{
 			DebugObjectTreeWidget->RequestRefresh();
 		}
-		if (DebugObjectWidget.IsValid())
-		{
-			DebugObjectWidget->RefreshDebugObjects();
-		}
 
 		// Subsystem has been torn down and rebuilt.
 		if (UPCGSubsystem* Subsystem = GetSubsystem())
@@ -2404,16 +2396,11 @@ void FPCGEditor::OnMapChanged(UWorld* InWorld, EMapChangeType InMapChangedType)
 
 void FPCGEditor::OnLevelActorDeleted(AActor* InActor)
 {
+	// Forward call as this makes an effort to retain the selection if the selected component has not been deleted.
 	if (DebugObjectTreeWidget.IsValid())
 	{
 		DebugObjectTreeWidget->RequestRefresh();
 	}
-
-	// Forward call as this makes an effort to retain the selection if the selected component has not been deleted.
-	if (DebugObjectWidget.IsValid())
-	{
-		DebugObjectWidget->OnLevelActorDeleted(InActor);
-	}
 }
 
 TSharedRef<SDockTab> FPCGEditor::SpawnTab_GraphEditor(const FSpawnTabArgs& Args)
@@ -2446,7 +2433,7 @@ TSharedRef<SDockTab> FPCGEditor::SpawnTab_Palette(const FSpawnTabArgs& Args)
 		];
 }
 
-TSharedRef<SDockTab> FPCGEditor::SpawnTab_DebugObject(const FSpawnTabArgs& Args)
+TSharedRef<SDockTab> FPCGEditor::SpawnTab_DebugObjectTree(const FSpawnTabArgs& Args)
 {
 	return SNew(SDockTab)
 		.Label(LOCTEXT("PCGDebugObjectTitle", "Debug Object"))
diff --git a/Engine/Plugins/Experimental/PCG/Source/PCGEditor/Private/PCGEditor.h b/Engine/Plugins/Experimental/PCG/Source/PCGEditor/Private/PCGEditor.h
--- a/Engine/Plugins/Experimental/PCG/Source/PCGEditor/Private/PCGEditor.h
+++ b/Engine/Plugins/Experimental/PCG/Source/PCGEditor/Private/PCGEditor.h
@@ -17,7 +17,6 @@ class FUICommandList;
 class IDetailsView;
 class SGraphEditor;
 class SPCGEditorGraphAttributeListView;
-class SPCGEditorGraphDebugObjectWidget;
 class SPCGEditorGraphDebugObjectTree;
 class SPCGEditorGraphDeterminismListView;
 class SPCGEditorGraphFind;
@@ -34,13 +33,6 @@ class UPCGNode;
 DECLARE_MULTICAST_DELEGATE_OneParam(FOnInspectedStackChanged, const FPCGStack&);
 DECLARE_MULTICAST_DELEGATE_OneParam(FOnInspectedNodeChanged, UPCGEditorGraphNodeBase*);
 
-/** Used to communicate how a debug object was selected so that other selection widgets can be synchronized. */
-enum class FPCGDebugObjectSelectionMethod : uint8
-{
-	DebugObjectTree,
-	DebugObjectDropdown,
-};
-
 class FPCGEditor : public FAssetEditorToolkit, public FGCObject, public FSelfRegisteringEditorUndoClient
 {
 public:
@@ -54,7 +46,7 @@ public:
 	static UPCGEditorGraph* GetPCGEditorGraph(UPCGGraph* InGraph);
 
 	/** Sets the execution stack that want to inspect. */
-	void SetStackBeingInspected(const FPCGStack& FullStack, FPCGDebugObjectSelectionMethod SelectionMethod);
+	void SetStackBeingInspected(const FPCGStack& FullStack);
 
 	/** Gets the PCG component we are debugging */
 	UPCGComponent* GetPCGComponentBeingInspected() const { return const_cast<UPCGComponent*>(StackBeingInspected.GetRootComponent()); }
@@ -153,9 +145,15 @@ private:
 	/** Cancels the current execution of the selected graph */
 	void OnCancelExecution_Clicked();
 
-	/* Returns true if inspected graph is currently scheduled or executing */
+	/** Returns true if inspected graph is currently scheduled or executing */
 	bool IsCurrentlyGenerating() const;
 
+	/** Returns true if the debug object tree tab is not currently open. */
+	bool IsDebugObjectTreeTabClosed() const;
+
+	/** Opens the debug object tree tab if it is not open already. */
+	void OnOpenDebugObjectTreeTab_Clicked();
+
 	/** Can determinism be tested on the selected node(s) */
 	bool CanRunDeterminismNodeTest() const;
 	/** Run the determinism test on the selected node(s) */
@@ -247,9 +245,6 @@ private:
 	/** Create new palette widget */
 	TSharedRef<SPCGEditorGraphNodePalette> CreatePaletteWidget();
 
-	/** Create debug object combo box */
-	TSharedRef<SPCGEditorGraphDebugObjectWidget> CreateDebugObjectWidget();
-
 	/** Create new debug object tree widget */
 	TSharedRef<SPCGEditorGraphDebugObjectTree> CreateDebugObjectTreeWidget();
 
@@ -311,7 +306,7 @@ private:
 	TSharedRef<SDockTab> SpawnTab_GraphEditor(const FSpawnTabArgs& Args);
 	TSharedRef<SDockTab> SpawnTab_PropertyDetails(const FSpawnTabArgs& Args);
 	TSharedRef<SDockTab> SpawnTab_Palette(const FSpawnTabArgs& Args);
-	TSharedRef<SDockTab> SpawnTab_DebugObject(const FSpawnTabArgs& Args);
+	TSharedRef<SDockTab> SpawnTab_DebugObjectTree(const FSpawnTabArgs& Args);
 	TSharedRef<SDockTab> SpawnTab_Attributes(const FSpawnTabArgs& Args);
 	TSharedRef<SDockTab> SpawnTab_Find(const FSpawnTabArgs& Args);
 	TSharedRef<SDockTab> SpawnTab_Determinism(const FSpawnTabArgs& Args);
@@ -321,7 +316,6 @@ private:
 	TSharedPtr<SGraphEditor> GraphEditorWidget;
 	TSharedPtr<IDetailsView> PropertyDetailsWidget;
 	TSharedPtr<SPCGEditorGraphNodePalette> PaletteWidget;
-	TSharedPtr<SPCGEditorGraphDebugObjectWidget> DebugObjectWidget;
 	TSharedPtr<SPCGEditorGraphDebugObjectTree> DebugObjectTreeWidget;
 	TSharedPtr<SPCGEditorGraphFind> FindWidget;
 	TSharedPtr<SPCGEditorGraphAttributeListView> AttributesWidget;
diff --git a/Engine/Plugins/Experimental/PCG/Source/PCGEditor/Private/PCGEditorCommands.cpp b/Engine/Plugins/Experimental/PCG/Source/PCGEditor/Private/PCGEditorCommands.cpp
--- a/Engine/Plugins/Experimental/PCG/Source/PCGEditor/Private/PCGEditorCommands.cpp
+++ b/Engine/Plugins/Experimental/PCG/Source/PCGEditor/Private/PCGEditorCommands.cpp
@@ -23,6 +23,7 @@ void FPCGEditorCommands::RegisterCommands()
 	UI_COMMAND(Find, "Find", "Finds PCG nodes and comments in the current graph", EUserInterfaceActionType::Button, FInputChord(EModifierKey::Control, EKeys::F));
 	UI_COMMAND(PauseAutoRegeneration, "Pause Regen", "Pause automatic regeneration of the current graph.", EUserInterfaceActionType::ToggleButton, FInputChord(EModifierKey::Alt, EKeys::R));
 	UI_COMMAND(ForceGraphRegeneration, "Force Regen", "Manually force a regeneration of the current graph.\nCtrl-click will also perform a flush cache before the regeneration.", EUserInterfaceActionType::Button, FInputChord());
+	UI_COMMAND(OpenDebugObjectTreeTab, "Debug Object Tree", "Open the Debug Object Tree tab to display and select graph invocations to debug.", EUserInterfaceActionType::Button, FInputChord());
 	UI_COMMAND(RunDeterminismNodeTest, "Run Determinism Test on Node", "Evaluate the current node for determinism.", EUserInterfaceActionType::Button, FInputChord(EModifierKey::Alt, EKeys::T));
 	UI_COMMAND(RunDeterminismGraphTest, "Graph Determinism Test", "Evaluate the current graph for determinism.", EUserInterfaceActionType::Button, FInputChord());
 	UI_COMMAND(EditGraphSettings, "Graph Settings", "Edit the graph settings.", EUserInterfaceActionType::ToggleButton, FInputChord());
diff --git a/Engine/Plugins/Experimental/PCG/Source/PCGEditor/Private/PCGEditorCommands.h b/Engine/Plugins/Experimental/PCG/Source/PCGEditor/Private/PCGEditorCommands.h
--- a/Engine/Plugins/Experimental/PCG/Source/PCGEditor/Private/PCGEditorCommands.h
+++ b/Engine/Plugins/Experimental/PCG/Source/PCGEditor/Private/PCGEditorCommands.h
@@ -19,6 +19,7 @@ public:
 	TSharedPtr<FUICommandInfo> Find;
 	TSharedPtr<FUICommandInfo> PauseAutoRegeneration;
 	TSharedPtr<FUICommandInfo> ForceGraphRegeneration;
+	TSharedPtr<FUICommandInfo> OpenDebugObjectTreeTab;
 	TSharedPtr<FUICommandInfo> RunDeterminismNodeTest;
 	TSharedPtr<FUICommandInfo> RunDeterminismGraphTest;
 	TSharedPtr<FUICommandInfo> EditGraphSettings;
diff --git a/Engine/Plugins/Experimental/PCG/Source/PCGEditor/Private/SPCGEditorGraphDebugObjectTree.cpp b/Engine/Plugins/Experimental/PCG/Source/PCGEditor/Private/SPCGEditorGraphDebugObjectTree.cpp
--- a/Engine/Plugins/Experimental/PCG/Source/PCGEditor/Private/SPCGEditorGraphDebugObjectTree.cpp
+++ b/Engine/Plugins/Experimental/PCG/Source/PCGEditor/Private/SPCGEditorGraphDebugObjectTree.cpp
@@ -896,13 +896,13 @@ void SPCGEditorGraphDebugObjectTree::OnSelectionChanged(FPCGEditorGraphDebugObje
 
 	if (!InItem)
 	{
-		PCGEditor.Pin()->SetStackBeingInspected(FPCGStack(), FPCGDebugObjectSelectionMethod::DebugObjectTree);
+		PCGEditor.Pin()->SetStackBeingInspected(FPCGStack());
 		return;
 	}
 
 	if (const FPCGStack* PCGStack = InItem->GetPCGStack())
 	{
-		PCGEditor.Pin()->SetStackBeingInspected(*PCGStack, FPCGDebugObjectSelectionMethod::DebugObjectTree);
+		PCGEditor.Pin()->SetStackBeingInspected(*PCGStack);
 	}
 }
 
diff --git a/Engine/Plugins/Experimental/PCG/Source/PCGEditor/Private/SPCGEditorGraphDebugObjectWidget.cpp b/Engine/Plugins/Experimental/PCG/Source/PCGEditor/Private/SPCGEditorGraphDebugObjectWidget.cpp
deleted file mode 100644
diff --git a/Engine/Plugins/Experimental/PCG/Source/PCGEditor/Private/SPCGEditorGraphDebugObjectWidget.h b/Engine/Plugins/Experimental/PCG/Source/PCGEditor/Private/SPCGEditorGraphDebugObjectWidget.h
deleted file mode 100644
