[PCG] Add Replace Tag node for data tags

diff --git a/Engine/Plugins/Experimental/PCG/Source/PCG/Private/Elements/PCGReplaceTags.cpp b/Engine/Plugins/Experimental/PCG/Source/PCG/Private/Elements/PCGReplaceTags.cpp
new file mode 100644
--- /dev/null
+++ b/Engine/Plugins/Experimental/PCG/Source/PCG/Private/Elements/PCGReplaceTags.cpp
@@ -0,0 +1,117 @@
+// Copyright Epic Games, Inc. All Rights Reserved.
+
+#include "Elements/PCGReplaceTags.h"
+
+#include "PCGContext.h"
+#include "PCGData.h"
+#include "PCGModule.h"
+#include "PCGPin.h"
+#include "Helpers/PCGHelpers.h"
+
+#define LOCTEXT_NAMESPACE "PCGReplaceTagsElement"
+
+#if WITH_EDITOR
+FName UPCGReplaceTagsSettings::GetDefaultNodeName() const
+{
+	return TEXT("ReplaceTags");
+}
+
+FText UPCGReplaceTagsSettings::GetDefaultNodeTitle() const
+{
+	return LOCTEXT("NodeTitle", "Replace Tags");
+}
+#endif // WITH_EDITOR
+
+FString UPCGReplaceTagsSettings::GetAdditionalTitleInformation() const
+{
+	TArray<FString> TagsToReplace = PCGHelpers::GetStringArrayFromCommaSeparatedString(SelectedTags);
+	TArray<FString> TagsToProcess = PCGHelpers::GetStringArrayFromCommaSeparatedString(ReplacedTags);
+
+	if (TagsToReplace.Num() == 1 && TagsToProcess.Num() == 1)
+	{
+		return FString::Printf(TEXT("%s -> %s"), *TagsToReplace[0], *TagsToProcess[0]);
+	}
+	else if (TagsToReplace.Num() > 1 && TagsToProcess.Num() == 1)
+	{
+		return FString::Printf(TEXT("Multiple -> %s"), *TagsToProcess[0]);
+	}
+	else if (TagsToReplace.Num() > 1 && TagsToProcess.Num() > 1)
+	{
+		return FString(TEXT("Multiple -> Multiple"));
+	}
+	else
+	{
+		return FString();
+	}
+}
+
+TArray<FPCGPinProperties> UPCGReplaceTagsSettings::OutputPinProperties() const
+{
+	TArray<FPCGPinProperties> PinProperties;
+	PinProperties.Emplace(PCGPinConstants::DefaultOutputLabel, EPCGDataType::Any);
+
+	return PinProperties;
+}
+
+FPCGElementPtr UPCGReplaceTagsSettings::CreateElement() const
+{
+	return MakeShared<FPCGReplaceTagsElement>();
+}
+
+bool FPCGReplaceTagsElement::ExecuteInternal(FPCGContext* Context) const
+{
+	TRACE_CPUPROFILER_EVENT_SCOPE(FPCGReplaceTagsElement::Execute);
+	check(Context);
+
+	const UPCGReplaceTagsSettings* Settings = Context->GetInputSettings<UPCGReplaceTagsSettings>();
+	check(Settings);
+
+	const TArray<FPCGTaggedData> Inputs = Context->InputData.GetInputsByPin(PCGPinConstants::DefaultInputLabel);
+	TArray<FPCGTaggedData>& Outputs = Context->OutputData.TaggedData;
+
+	const TArray<FString> SelectedTags = PCGHelpers::GetStringArrayFromCommaSeparatedString(Settings->SelectedTags);
+	const TArray<FString> ReplacedTags = PCGHelpers::GetStringArrayFromCommaSeparatedString(Settings->ReplacedTags);
+	if (SelectedTags.IsEmpty() && ReplacedTags.IsEmpty())
+	{
+		// Forward the input and early out
+		Outputs = Inputs;
+		PCGE_LOG(Verbose, LogOnly, LOCTEXT("InvalidNumberOfTags", "No selected or replaced tags, nothing to do."));
+		return true;
+	}
+
+	if (ReplacedTags.Num() > 1 && SelectedTags.Num() != ReplacedTags.Num())
+	{
+		Outputs = Inputs;
+		PCGE_LOG(Error, GraphAndLog, LOCTEXT("InvalidMappingSupport", "Input only supports 1:1, N:1, and N:N mappings from source to replaced tags, Input data discarded."));
+		return true;
+	}
+
+	const bool bNToNTagMapping = SelectedTags.Num() == ReplacedTags.Num();
+	const bool bNToOneMapping = ReplacedTags.Num() == 1;
+	
+	for (const FPCGTaggedData& Input : Inputs)
+	{
+		FPCGTaggedData& Output = Outputs.Add_GetRef(Input);
+
+		for (int i = 0; i < SelectedTags.Num(); ++i)
+		{
+			if (Input.Tags.Contains(SelectedTags[i]))
+			{
+				Output.Tags.Remove(SelectedTags[i]);
+
+				if (bNToNTagMapping)
+				{
+					Output.Tags.Add(ReplacedTags[i]);
+				}
+				else if (bNToOneMapping)
+				{
+					Output.Tags.Add(ReplacedTags[0]);
+				}
+			}
+		}
+	}
+
+	return true;
+}
+
+#undef LOCTEXT_NAMESPACE
\ No newline at end of file
diff --git a/Engine/Plugins/Experimental/PCG/Source/PCG/Private/Tests/Elements/PCGReplaceTagTest.cpp b/Engine/Plugins/Experimental/PCG/Source/PCG/Private/Tests/Elements/PCGReplaceTagTest.cpp
new file mode 100644
--- /dev/null
+++ b/Engine/Plugins/Experimental/PCG/Source/PCG/Private/Tests/Elements/PCGReplaceTagTest.cpp
@@ -0,0 +1,138 @@
+// Copyright Epic Games, Inc. All Rights Reserved.
+
+#include "Tests/PCGTestsCommon.h"
+
+#include "PCGComponent.h"
+#include "PCGContext.h"
+#include "Data/PCGPointData.h"
+
+#include "Elements/PCGReplaceTags.h"
+
+class PCGReplaceTagsTestBase : public FPCGTestBaseClass
+{
+public:
+	using FPCGTestBaseClass::FPCGTestBaseClass;
+
+	struct FTestParameters
+	{
+		bool ExpectedToHaveOneOutput = true;
+		FString SelectedTags;
+		FString ReplacedTags;
+		TSet<FString> ExpectedOutputTags;
+	};
+
+	bool GenerateTestDataRunAndValidate(const FTestParameters& Parameters)
+	{
+		PCGTestsCommon::FTestData TestData;
+		PCGTestsCommon::GenerateSettings<UPCGReplaceTagsSettings>(TestData);
+		UPCGReplaceTagsSettings* Settings = CastChecked<UPCGReplaceTagsSettings>(TestData.Settings);
+		Settings->SelectedTags = Parameters.SelectedTags;
+		Settings->ReplacedTags = Parameters.ReplacedTags;
+
+		FPCGElementPtr TestElement = TestData.Settings->GetElement();
+		TUniquePtr<FPCGContext> Context = TestData.InitializeTestContext();
+
+		FPCGTaggedData& InputTaggedData = Context->InputData.TaggedData.Emplace_GetRef();
+		InputTaggedData.Data = PCGTestsCommon::CreatePointData();
+		InputTaggedData.Pin = PCGPinConstants::DefaultInputLabel;
+		InputTaggedData.Tags = TSet<FString>{ TEXT("Awesome"), TEXT("Fun"), TEXT("Cool") };
+
+		while (!TestElement->Execute(Context.Get())) {}
+
+		// Testing output
+		const TArray<FPCGTaggedData>& Outputs = Context->OutputData.TaggedData;
+
+		UTEST_EQUAL("Number of Outputs", Outputs.Num(), Parameters.ExpectedToHaveOneOutput? 1 : 0);
+
+		if (Parameters.ExpectedToHaveOneOutput)
+		{
+			// Testing tag info
+			const TSet<FString> OutputTagSet = Context->OutputData.TaggedData[0].Tags;
+
+			UTEST_EQUAL("Tag count", OutputTagSet.Num(), Parameters.ExpectedOutputTags.Num());
+			UTEST_TRUE("Output Tags", OutputTagSet.Array() == Parameters.ExpectedOutputTags.Array());
+		}
+
+		return true;
+	}
+};
+
+IMPLEMENT_CUSTOM_SIMPLE_AUTOMATION_TEST(FPCGReplaceTagTest_OneToOne, PCGReplaceTagsTestBase, "Plugins.PCG.ReplaceTag.OneToOne", PCGTestsCommon::TestFlags)
+
+bool FPCGReplaceTagTest_OneToOne::RunTest(const FString& Parameters)
+{
+	FTestParameters OneToOneParameters{};
+	OneToOneParameters.SelectedTags = TEXT("Awesome");
+	OneToOneParameters.ReplacedTags = TEXT("Neat");
+	OneToOneParameters.ExpectedToHaveOneOutput = true;
+	OneToOneParameters.ExpectedOutputTags = TSet<FString>{ TEXT("Neat"), TEXT("Fun"), TEXT("Cool") };
+	
+	return GenerateTestDataRunAndValidate(OneToOneParameters);
+}
+
+IMPLEMENT_CUSTOM_SIMPLE_AUTOMATION_TEST(FPCGReplaceTagTest_NToOne, PCGReplaceTagsTestBase, "Plugins.PCG.ReplaceTag.NToOne", PCGTestsCommon::TestFlags)
+
+bool FPCGReplaceTagTest_NToOne::RunTest(const FString& Parameters)
+{
+	FTestParameters NtoOneParameters{};
+	NtoOneParameters.SelectedTags = TEXT("Awesome, Fun");
+	NtoOneParameters.ReplacedTags = TEXT("Neat");
+	NtoOneParameters.ExpectedToHaveOneOutput = true;
+	NtoOneParameters.ExpectedOutputTags = TSet<FString>{ TEXT("Neat"), TEXT("Cool") };
+
+	return GenerateTestDataRunAndValidate(NtoOneParameters);
+}
+
+IMPLEMENT_CUSTOM_SIMPLE_AUTOMATION_TEST(FPCGReplaceTagTest_NToN, PCGReplaceTagsTestBase, "Plugins.PCG.ReplaceTag.NToN", PCGTestsCommon::TestFlags)
+
+bool FPCGReplaceTagTest_NToN::RunTest(const FString& Parameters)
+{
+	FTestParameters NtoNParameters{};
+	NtoNParameters.SelectedTags = TEXT("Awesome, Fun, Cool");
+	NtoNParameters.ReplacedTags = TEXT("Neat, Great, Amazing");
+	NtoNParameters.ExpectedToHaveOneOutput = true;
+	NtoNParameters.ExpectedOutputTags = TSet<FString>{ TEXT("Neat"), TEXT("Great"), TEXT("Amazing")};
+
+	return GenerateTestDataRunAndValidate(NtoNParameters);
+}
+
+IMPLEMENT_CUSTOM_SIMPLE_AUTOMATION_TEST(FPCGReplaceTagTest_OneToNError, PCGReplaceTagsTestBase, "Plugins.PCG.ReplaceTag.OneToNError", PCGTestsCommon::TestFlags)
+
+bool FPCGReplaceTagTest_OneToNError::RunTest(const FString& Parameters)
+{
+	AddExpectedError(TEXT("Input only supports 1:1, N:1, and N:N mappings from source to replaced tags, discarded."));
+
+	FTestParameters OneToNErrorParameters{};
+	OneToNErrorParameters.SelectedTags = TEXT("Awesome");
+	OneToNErrorParameters.ReplacedTags = TEXT("Neat, Great");
+	OneToNErrorParameters.ExpectedToHaveOneOutput = false;
+	OneToNErrorParameters.ExpectedOutputTags = TSet<FString>();
+
+	return GenerateTestDataRunAndValidate(OneToNErrorParameters);
+}
+
+IMPLEMENT_CUSTOM_SIMPLE_AUTOMATION_TEST(FPCGReplaceTagTest_RemoveTags, PCGReplaceTagsTestBase, "Plugins.PCG.ReplaceTag.RemoveTags", PCGTestsCommon::TestFlags)
+
+bool FPCGReplaceTagTest_RemoveTags::RunTest(const FString& Parameters)
+{
+	FTestParameters RemoveTagsParameters{};
+	RemoveTagsParameters.SelectedTags = TEXT("Awesome");
+	RemoveTagsParameters.ReplacedTags = FString();
+	RemoveTagsParameters.ExpectedToHaveOneOutput = true;
+	RemoveTagsParameters.ExpectedOutputTags = TSet<FString>{ TEXT("Fun"), TEXT("Cool") };
+
+	return GenerateTestDataRunAndValidate(RemoveTagsParameters);
+}
+
+IMPLEMENT_CUSTOM_SIMPLE_AUTOMATION_TEST(FPCGReplaceTagTest_ForwardInput, PCGReplaceTagsTestBase, "Plugins.PCG.ReplaceTag.ForwardInput", PCGTestsCommon::TestFlags)
+
+bool FPCGReplaceTagTest_ForwardInput::RunTest(const FString& Parameters)
+{
+	FTestParameters ForwardInputParameters{};
+	ForwardInputParameters.SelectedTags = FString();
+	ForwardInputParameters.ReplacedTags = FString();
+	ForwardInputParameters.ExpectedToHaveOneOutput = true;
+	ForwardInputParameters.ExpectedOutputTags = TSet<FString>{ TEXT("Awesome"), TEXT("Fun"), TEXT("Cool") };
+
+	return GenerateTestDataRunAndValidate(ForwardInputParameters);
+}
\ No newline at end of file
diff --git a/Engine/Plugins/Experimental/PCG/Source/PCG/Public/Elements/PCGReplaceTags.h b/Engine/Plugins/Experimental/PCG/Source/PCG/Public/Elements/PCGReplaceTags.h
new file mode 100644
--- /dev/null
+++ b/Engine/Plugins/Experimental/PCG/Source/PCG/Public/Elements/PCGReplaceTags.h
@@ -0,0 +1,49 @@
+// Copyright Epic Games, Inc. All Rights Reserved.
+
+#pragma once
+
+#include "PCGSettings.h"
+
+#include "PCGReplaceTags.generated.h"
+
+/** Replaces the tags on the input data. */
+UCLASS(BlueprintType, ClassGroup = (Procedural))
+class PCG_API UPCGReplaceTagsSettings : public UPCGSettings
+{
+	GENERATED_BODY()
+
+public:
+	//~Begin UPCGSettings interface
+#if WITH_EDITOR
+	virtual FName GetDefaultNodeName() const override;
+	virtual FText GetDefaultNodeTitle() const override;
+	virtual EPCGSettingsType GetType() const override { return EPCGSettingsType::Generic; }
+	virtual bool HasDynamicPins() const override { return true; }
+#endif
+	virtual FString GetAdditionalTitleInformation() const override;
+
+protected:
+#if WITH_EDITOR
+	virtual EPCGChangeType GetChangeTypeForProperty(const FName& InPropertyName) const override { return Super::GetChangeTypeForProperty(InPropertyName) | EPCGChangeType::Cosmetic; }
+#endif
+	virtual TArray<FPCGPinProperties> OutputPinProperties() const override;
+	virtual FPCGElementPtr CreateElement() const override;
+	//~End UPCGSettings interface
+
+public:
+	/** Comma-separated list of tags to be replaced from the input data.
+	* If no tags are replaced, then these tags will be removed. 
+	*/
+	UPROPERTY(BlueprintReadWrite, EditAnywhere, Category = Settings, meta = (PCG_Overridable))
+	FString SelectedTags;
+
+	/** Comma-separated list of new tags to replace the selected tags. */
+	UPROPERTY(BlueprintReadWrite, EditAnywhere, Category = Settings, meta = (PCG_Overridable))
+	FString ReplacedTags;
+};
+
+class FPCGReplaceTagsElement : public IPCGElement
+{
+protected:
+	virtual bool ExecuteInternal(FPCGContext* Context) const override;
+};
\ No newline at end of file
