[PCG] Cleanup in Add/Create Attributes

* Remove deprecated var that are more than 1 year old
* Remove base class that was more confusing than helping.
* Common code is now centralized in the cpp, as free functions.
* Add the option to copy all attributes in add attributes.
* Since the memebers of each settings are still there with the same name, no deprecation needed.

#rnx

diff --git a/Engine/Plugins/Experimental/PCG/Source/PCG/Private/Elements/Metadata/PCGMetadataElement.cpp b/Engine/Plugins/Experimental/PCG/Source/PCG/Private/Elements/Metadata/PCGMetadataElement.cpp
--- a/Engine/Plugins/Experimental/PCG/Source/PCG/Private/Elements/Metadata/PCGMetadataElement.cpp
+++ b/Engine/Plugins/Experimental/PCG/Source/PCG/Private/Elements/Metadata/PCGMetadataElement.cpp
@@ -7,8 +7,7 @@
 #include "PCGParamData.h"
 #include "Data/PCGPointData.h"
 #include "Data/PCGSpatialData.h"
-#include "Elements/Metadata/PCGMetadataElementCommon.h"
-#include "Metadata/Accessors/PCGAttributeAccessorHelpers.h"
+#include "Helpers/PCGMetadataHelpers.h"
 
 #include UE_INLINE_GENERATED_CPP_BY_NAME(PCGMetadataElement)
 
@@ -194,151 +193,18 @@ bool FPCGMetadataOperationElement::ExecuteInternal(FPCGContext* Context) const
 			continue;
 		}
 
-		// Prepare attribute accessors
-		TArray<TPair<FPCGAttributePropertyInputSelector, FPCGAttributePropertyOutputSelector>> AttributeSelectors;
-
+		UPCGPointData* SampledData = CastChecked<UPCGPointData>(OriginalData->DuplicateData());
+		bool bSuccess = false;
 		if (Settings->bCopyAllAttributes)
 		{
-			TArray<FName> AttributeNames;
-			TArray<EPCGMetadataTypes> AttributeTypes;
-			SourceMetadata->GetAttributes(AttributeNames, AttributeTypes);
-
-			for (const FName& AttributeName : AttributeNames)
-			{
-				TPair<FPCGAttributePropertyInputSelector, FPCGAttributePropertyOutputSelector>& Selectors = AttributeSelectors.Emplace_GetRef();
-				Selectors.Key.SetAttributeName(AttributeName);
-				Selectors.Value.SetAttributeName(AttributeName);
-			}
+			bSuccess = PCGMetadataHelpers::CopyAllAttributes(SourceData, SampledData, Context);
 		}
 		else
 		{
-			FPCGAttributePropertyInputSelector InputSource = Settings->InputSource.CopyAndFixLast(SourceData);
-			// Deprecation, old behavior was if the Target was None, we took LastCreated in the SourceData
-			FPCGAttributePropertyOutputSelector OutputTarget = Settings->OutputTarget.CopyAndFixSource(&InputSource, SourceData);
-
-			AttributeSelectors.Emplace(MoveTemp(InputSource), MoveTemp(OutputTarget));
-		}
-
-		UPCGPointData* SampledData = nullptr;
-
-		for (const auto& SelectorPair : AttributeSelectors)
-		{
-			const FPCGAttributePropertyInputSelector& InputSource = SelectorPair.Key;
-			const FPCGAttributePropertyOutputSelector& OutputTarget = SelectorPair.Value;
-
-			const FName LocalSourceAttribute = InputSource.GetName();
-			const FName LocalDestinationAttribute = OutputTarget.GetName();
-
-			if (InputSource.GetSelection() == EPCGAttributePropertySelection::Attribute && !SourceMetadata->HasAttribute(LocalSourceAttribute))
-			{
-				PCGE_LOG(Warning, GraphAndLog, FText::Format(LOCTEXT("InputMissingAttribute", "Input does not have the '{0}' attribute"), FText::FromName(LocalSourceAttribute)));
-				continue;
-			}
-
-			const TArray<FPCGPoint>& Points = OriginalData->GetPoints();
-			const int OriginalPointCount = Points.Num();
-
-			if (!SampledData)
-			{
-				SampledData = NewObject<UPCGPointData>();
-				SampledData->InitializeFromData(OriginalData);
-
-				SampledData->GetMutablePoints() = Points;
-			}
-
-			TArray<FPCGPoint>& SampledPoints = SampledData->GetMutablePoints();
-
-			// If it is attribute to attribute, just copy the attributes, if they exist and are valid
-			// Only do that if it is really attribute to attribute, without any extra accessor. Any extra accessor will behave as a property.
-			const bool bInputHasAnyExtra = !InputSource.GetExtraNames().IsEmpty();
-			const bool bOutputHasAnyExtra = !OutputTarget.GetExtraNames().IsEmpty();
-			if (!bInputHasAnyExtra && !bOutputHasAnyExtra && Settings->InputSource.GetSelection() == EPCGAttributePropertySelection::Attribute && OutputTarget.GetSelection() == EPCGAttributePropertySelection::Attribute)
-			{
-				if (!SourceAttributeSet && LocalSourceAttribute == LocalDestinationAttribute)
-				{
-					// Nothing to do if we try to copy an attribute into itself in the original point data.
-					continue;
-				}
-
-				const FPCGMetadataAttributeBase* OriginalAttribute = SourceMetadata->GetConstAttribute(LocalSourceAttribute);
-				check(OriginalAttribute);
-
-				// We only copy entries/values if we copy from the input spatial metadata
-				// If it is from the source attribute set, we don't copy (and all points will have the same default value, value from the attribute set)
-				const bool bCopyEntriesAndValues = (SourceAttributeSet == nullptr);
-				if (!SampledData->Metadata->CopyAttribute(OriginalAttribute, LocalDestinationAttribute, /*bKeepParent=*/ false, /*bCopyEntries=*/ bCopyEntriesAndValues, /*bCopyValues=*/ bCopyEntriesAndValues))
-				{
-					PCGE_LOG(Warning, GraphAndLog, FText::Format(LOCTEXT("FailedCopyToNewAttribute", "Failed to copy to new attribute '{0}'"), FText::FromName(LocalDestinationAttribute)));
-				}
-
-				continue;
-			}
-
-			TUniquePtr<const IPCGAttributeAccessor> InputAccessor = PCGAttributeAccessorHelpers::CreateConstAccessor(SourceData, InputSource);
-			TUniquePtr<const IPCGAttributeAccessorKeys> InputKeys = PCGAttributeAccessorHelpers::CreateConstKeys(SourceData, InputSource);
-
-			if (!InputAccessor.IsValid() || !InputKeys.IsValid())
-			{
-				PCGE_LOG(Warning, GraphAndLog, LOCTEXT("FailedToCreateInputAccessor", "Failed to create input accessor or iterator"));
-				continue;
-			}
-
-			// If the target is an attribute, only create a new one if the attribute doesn't already exist or we have any extra.
-			// If it exist or have any extra, it will try to write to it.
-			if (!bOutputHasAnyExtra && OutputTarget.GetSelection() == EPCGAttributePropertySelection::Attribute && !SampledData->Metadata->HasAttribute(LocalDestinationAttribute))
-			{
-				auto CreateAttribute = [SampledData, LocalDestinationAttribute](auto Dummy)
-				{
-					using AttributeType = decltype(Dummy);
-					return PCGMetadataElementCommon::ClearOrCreateAttribute<AttributeType>(SampledData->Metadata, LocalDestinationAttribute) != nullptr;
-				};
-
-				if (!PCGMetadataAttribute::CallbackWithRightType(InputAccessor->GetUnderlyingType(), CreateAttribute))
-				{
-					PCGE_LOG(Warning, GraphAndLog, FText::Format(LOCTEXT("FailedToCreateNewAttribute", "Failed to create new attribute '{0}'"), FText::FromName(LocalDestinationAttribute)));
-					continue;
-				}
-			}
-
-			TUniquePtr<IPCGAttributeAccessor> OutputAccessor = PCGAttributeAccessorHelpers::CreateAccessor(SampledData, OutputTarget);
-			TUniquePtr<IPCGAttributeAccessorKeys> OutputKeys = PCGAttributeAccessorHelpers::CreateKeys(SampledData, OutputTarget);
-
-			if (!OutputAccessor.IsValid() || !OutputKeys.IsValid())
-			{
-				PCGE_LOG(Warning, GraphAndLog, LOCTEXT("FailedToCreateOutputAccessor", "Failed to create output accessor or iterator"));
-				continue;
-			}
-
-			if (OutputAccessor->IsReadOnly())
-			{
-				PCGE_LOG(Warning, GraphAndLog, FText::Format(LOCTEXT("OutputAccessorIsReadOnly", "Attribute/Property '{0}' is read only."), OutputTarget.GetDisplayText()));
-				continue;
-			}
-
-			// Final verification, if we can put the value of input into output
-			if (!PCG::Private::IsBroadcastable(InputAccessor->GetUnderlyingType(), OutputAccessor->GetUnderlyingType()))
-			{
-				PCGE_LOG(Error, GraphAndLog, LOCTEXT("CannotBroadcastTypes", "Cannot broadcast input type into output type"));
-				continue;
-			}
-
-			// At this point, we are ready.
-			PCGMetadataElementCommon::FCopyFromAccessorToAccessorParams Params;
-			Params.InKeys = InputKeys.Get();
-			Params.InAccessor = InputAccessor.Get();
-			Params.OutKeys = OutputKeys.Get();
-			Params.OutAccessor = OutputAccessor.Get();
-			Params.IterationCount = PCGMetadataElementCommon::FCopyFromAccessorToAccessorParams::Out;
-			Params.Flags = EPCGAttributeAccessorFlags::AllowBroadcast | EPCGAttributeAccessorFlags::AllowConstructible;
-
-			if (!PCGMetadataElementCommon::CopyFromAccessorToAccessor(Params))
-			{
-				PCGE_LOG(Warning, GraphAndLog, LOCTEXT("ErrorGettingSettingValues", "Error while getting/setting values"));
-				continue;
-			}
+			bSuccess = PCGMetadataHelpers::CopyAttributes(SourceData, Settings->InputSource, SampledData, Settings->OutputTarget, /*bSameOrigin=*/!SourceAttributeSet, Context);
 		}
 
-		if (SampledData)
+		if (bSuccess)
 		{
 			Output.Data = SampledData;
 		}
diff --git a/Engine/Plugins/Experimental/PCG/Source/PCG/Private/Elements/PCGCreateAttribute.cpp b/Engine/Plugins/Experimental/PCG/Source/PCG/Private/Elements/PCGCreateAttribute.cpp
--- a/Engine/Plugins/Experimental/PCG/Source/PCG/Private/Elements/PCGCreateAttribute.cpp
+++ b/Engine/Plugins/Experimental/PCG/Source/PCG/Private/Elements/PCGCreateAttribute.cpp
@@ -22,48 +22,28 @@ namespace PCGCreateAttributeConstants
 	const FText NodeTitleCreateAttribute = LOCTEXT("NodeTitleCreateAttribute", "Create Attribute");
 	const FName AttributesLabel = TEXT("Attributes");
 	const FText AttributesTooltip = LOCTEXT("AttributesTooltip", "Optional Attribute Set to create from. Not used if not connected.");
+	const FText ErrorCreatingAttributeMessage = LOCTEXT("ErrorCreatingAttribute", "Error while creating attribute '{0}'");
 }
 
-void UPCGCreateAttributeBaseSettings::PostLoad()
+namespace PCGCreateAttribute
 {
-	Super::PostLoad();
-
-#if WITH_EDITOR
-	if ((Type_DEPRECATED != EPCGMetadataTypes::Double) || (DoubleValue_DEPRECATED != 0.0))
+	FPCGMetadataAttributeBase* ClearOrCreateAttribute(const FPCGMetadataTypesConstantStruct& AttributeTypes, UPCGMetadata* Metadata, const FName OutputAttributeName)
 	{
-		AttributeTypes.Type = Type_DEPRECATED;
-		AttributeTypes.DoubleValue = DoubleValue_DEPRECATED;
-		AttributeTypes.FloatValue = FloatValue_DEPRECATED;
-		AttributeTypes.IntValue = IntValue_DEPRECATED;
-		AttributeTypes.Int32Value = Int32Value_DEPRECATED;
-		AttributeTypes.Vector2Value = Vector2Value_DEPRECATED;
-		AttributeTypes.VectorValue = VectorValue_DEPRECATED;
-		AttributeTypes.Vector4Value = Vector4Value_DEPRECATED;
-		AttributeTypes.RotatorValue = RotatorValue_DEPRECATED;
-		AttributeTypes.QuatValue = QuatValue_DEPRECATED;
-		AttributeTypes.TransformValue = TransformValue_DEPRECATED;
-		AttributeTypes.BoolValue = BoolValue_DEPRECATED;
-		AttributeTypes.StringValue = StringValue_DEPRECATED;
-		AttributeTypes.NameValue = NameValue_DEPRECATED;
-
-		Type_DEPRECATED = EPCGMetadataTypes::Double;
-		DoubleValue_DEPRECATED = 0.0;
-	}
+		check(Metadata);
 
-	if (SourceParamAttributeName_DEPRECATED != NAME_None)
-	{
-		InputSource.SetAttributeName(SourceParamAttributeName_DEPRECATED);
-		SourceParamAttributeName_DEPRECATED = NAME_None;
-	}
+		auto CreateAttribute = [Metadata, OutputAttributeName](auto&& Value) -> FPCGMetadataAttributeBase*
+		{
+			return PCGMetadataElementCommon::ClearOrCreateAttribute(Metadata, OutputAttributeName, std::forward<decltype(Value)>(Value));
+		};
 
-	AttributeTypes.OnPostLoad();
-#endif // WITH_EDITOR
+		return AttributeTypes.Dispatcher(CreateAttribute);
+	}
 }
 
-EPCGDataType UPCGCreateAttributeBaseSettings::GetCurrentPinTypes(const UPCGPin* InPin) const
+EPCGDataType UPCGAddAttributeSettings::GetCurrentPinTypes(const UPCGPin* InPin) const
 {
 	check(InPin);
-	if (!HasDynamicPins() || !InPin->IsOutputPin())
+	if (!InPin->IsOutputPin())
 	{
 		return Super::GetCurrentPinTypes(InPin);
 	}
@@ -74,12 +54,12 @@ EPCGDataType UPCGCreateAttributeBaseSettings::GetCurrentPinTypes(const UPCGPin*
 }
 
 #if WITH_EDITOR
-bool UPCGCreateAttributeBaseSettings::IsPinUsedByNodeExecution(const UPCGPin* InPin) const
+bool UPCGAddAttributeSettings::IsPinUsedByNodeExecution(const UPCGPin* InPin) const
 {
 	return !InPin || (InPin->Properties.Label != PCGCreateAttributeConstants::AttributesLabel) || InPin->IsConnected();
 }
 
-bool UPCGCreateAttributeBaseSettings::CanEditChange(const FProperty* InProperty) const
+bool UPCGAddAttributeSettings::CanEditChange(const FProperty* InProperty) const
 {
 	if (!InProperty || !Super::CanEditChange(InProperty))
 	{
@@ -89,7 +69,7 @@ bool UPCGCreateAttributeBaseSettings::CanEditChange(const FProperty* InProperty)
 	const UPCGNode* Node = Cast<UPCGNode>(GetOuter());
 	const bool AttributesPinIsConnected = Node ? Node->IsInputPinConnected(PCGCreateAttributeConstants::AttributesLabel) : false;
 
-	if (InProperty->GetFName() == GET_MEMBER_NAME_CHECKED(UPCGCreateAttributeBaseSettings, InputSource))
+	if (InProperty->GetFName() == GET_MEMBER_NAME_CHECKED(UPCGAddAttributeSettings, InputSource))
 	{
 		return AttributesPinIsConnected;
 	}
@@ -102,38 +82,44 @@ bool UPCGCreateAttributeBaseSettings::CanEditChange(const FProperty* InProperty)
 }
 #endif // WITH_EDITOR
 
-FPCGElementPtr UPCGCreateAttributeBaseSettings::CreateElement() const
+FPCGElementPtr UPCGAddAttributeSettings::CreateElement() const
 {
-	return MakeShared<FPCGCreateAttributeElement>();
+	return MakeShared<FPCGAddAttributeElement>();
 }
 
-FString UPCGCreateAttributeBaseSettings::GetAdditionalTitleInformationInternal(FName NodeName) const
+FString UPCGAddAttributeSettings::GetAdditionalTitleInformation() const
 {
 	if (HasAnyFlags(RF_ClassDefaultObject))
 	{
-		return NodeName.ToString();
+		return PCGCreateAttributeConstants::NodeNameAddAttribute.ToString();
 	}
 
 	const UPCGNode* Node = Cast<UPCGNode>(GetOuter());
 	const bool bAttributesPinIsConnected = Node ? Node->IsInputPinConnected(PCGCreateAttributeConstants::AttributesLabel) : false;
 
-	const FName OutputAttributeName = GetOutputAttributeName(bAttributesPinIsConnected ? &InputSource : nullptr, nullptr);
-	const FName SourceParamAttributeName = InputSource.GetName();
-
-	if (ShouldAddAttributesPin() && bAttributesPinIsConnected)
+	if (bAttributesPinIsConnected)
 	{
-		if ((OutputAttributeName == NAME_None) && (SourceParamAttributeName == NAME_None))
+		if (bCopyAllAttributes)
 		{
-			return NodeName.ToString();
+			return LOCTEXT("AllAttributes", "All Attributes").ToString();
 		}
 		else
 		{
-			return ((OutputAttributeName == NAME_None) ? SourceParamAttributeName : OutputAttributeName).ToString();
+			const FName SourceParamAttributeName = InputSource.GetName();
+			const FName OutputAttributeName = OutputTarget.CopyAndFixSource(&InputSource, nullptr).GetName();
+			if ((OutputAttributeName == NAME_None) && (SourceParamAttributeName == NAME_None))
+			{
+				return PCGCreateAttributeConstants::NodeNameAddAttribute.ToString();
+			}
+			else
+			{
+				return ((OutputAttributeName == NAME_None) ? SourceParamAttributeName : OutputAttributeName).ToString();
+			}
 		}
 	}
 	else
 	{
-		return FString::Printf(TEXT("%s: %s"), *OutputAttributeName.ToString(), *AttributeTypes.ToString());
+		return FString::Printf(TEXT("%s: %s"), *OutputTarget.GetName().ToString(), *AttributeTypes.ToString());
 	}
 }
 
@@ -142,16 +128,17 @@ UPCGAddAttributeSettings::UPCGAddAttributeSettings()
 	OutputTarget.SetAttributeName(NAME_None);
 }
 
-FString UPCGAddAttributeSettings::GetAdditionalTitleInformation() const
-{
-	return GetAdditionalTitleInformationInternal(PCGCreateAttributeConstants::NodeNameAddAttribute);
-}
-
 void UPCGAddAttributeSettings::PostLoad()
 {
 	Super::PostLoad();
 
 #if WITH_EDITOR
+	if (SourceParamAttributeName_DEPRECATED != NAME_None)
+	{
+		InputSource.SetAttributeName(SourceParamAttributeName_DEPRECATED);
+		SourceParamAttributeName_DEPRECATED = NAME_None;
+	}
+
 	if (OutputAttributeName_DEPRECATED != NAME_None)
 	{
 		OutputTarget.SetAttributeName(OutputAttributeName_DEPRECATED);
@@ -200,17 +187,10 @@ TArray<FPCGPinProperties> UPCGAddAttributeSettings::InputPinProperties() const
 
 TArray<FPCGPinProperties> UPCGAddAttributeSettings::OutputPinProperties() const
 {
-	FPCGPinProperties PinProperties;
-	PinProperties.Label = PCGPinConstants::DefaultOutputLabel;
-	PinProperties.AllowedTypes = EPCGDataType::Any;
-
-	return { PinProperties };
-}
+	TArray<FPCGPinProperties> PinProperties;
+	PinProperties.Emplace(PCGPinConstants::DefaultOutputLabel, EPCGDataType::Any);
 
-UPCGCreateAttributeSetSettings::UPCGCreateAttributeSetSettings()
-{
-	// No input pin to grab source param from
-	bDisplayFromSourceParamSetting = false;
+	return PinProperties;
 }
 
 void UPCGCreateAttributeSetSettings::PostLoad()
@@ -253,214 +233,131 @@ TArray<FPCGPinProperties> UPCGCreateAttributeSetSettings::OutputPinProperties()
 
 FString UPCGCreateAttributeSetSettings::GetAdditionalTitleInformation() const
 {
-	return GetAdditionalTitleInformationInternal(PCGCreateAttributeConstants::NodeNameCreateAttribute);
+	if (HasAnyFlags(RF_ClassDefaultObject))
+	{
+		return PCGCreateAttributeConstants::NodeNameAddAttribute.ToString();
+	}
+
+	return FString::Printf(TEXT("%s: %s"), *OutputTarget.GetName().ToString(), *AttributeTypes.ToString());
 }
 
-bool FPCGCreateAttributeElement::ExecuteInternal(FPCGContext* Context) const
+FPCGElementPtr UPCGCreateAttributeSetSettings::CreateElement() const
 {
-	TRACE_CPUPROFILER_EVENT_SCOPE(FPCGCreateAttributeElement::Execute);
+	return MakeShared<FPCGCreateAttributeElement>();
+}
+
+bool FPCGAddAttributeElement::ExecuteInternal(FPCGContext* Context) const
+{
+	TRACE_CPUPROFILER_EVENT_SCOPE(FPCGAddAttributeElement::Execute);
 
 	check(Context);
 
-	const UPCGCreateAttributeBaseSettings* Settings = Context->GetInputSettings<UPCGCreateAttributeBaseSettings>();
+	const UPCGAddAttributeSettings* Settings = Context->GetInputSettings<UPCGAddAttributeSettings>();
 	check(Settings);
 
+	const bool bAttributesPinIsConnected = Context->Node ? Context->Node->IsInputPinConnected(PCGCreateAttributeConstants::AttributesLabel) : false;
 	TArray<FPCGTaggedData> SourceParams = Context->InputData.GetInputsByPin(PCGCreateAttributeConstants::AttributesLabel);
-	const UPCGParamData* SourceParamData = nullptr;
-	FName SourceParamAttributeName = NAME_None;
-	FName OutputAttributeName = NAME_None;
+	TArray<FPCGTaggedData> Inputs = Context->InputData.GetInputsByPin(PCGPinConstants::DefaultInputLabel);
 
-	FPCGAttributePropertyInputSelector InputSource{};
+	const FName OutputAttributeName = Settings->OutputTarget.GetName();
 
-	if (!SourceParams.IsEmpty())
+	// If we add from a constant
+	if (SourceParams.IsEmpty() && !bAttributesPinIsConnected)
 	{
-		SourceParamData = CastChecked<UPCGParamData>(SourceParams[0].Data);
-
-		if (!SourceParamData->Metadata)
+		for (int32 i = 0; i < Inputs.Num(); ++i)
 		{
-			PCGE_LOG(Error, GraphAndLog, LOCTEXT("ParamMissingMetadata", "Source Attribute Set data does not have metadata"));
-			return true;
-		}
+			const UPCGData* InData = Inputs[i].Data;
+			if (!InData || !InData->ConstMetadata())
+			{
+				continue;
+			}
 
-		InputSource = Settings->InputSource.CopyAndFixLast(SourceParamData);
+			UPCGData* OutputData = InData->DuplicateData();
+			check(OutputData);
+			UPCGMetadata* OutputMetadata = OutputData->MutableMetadata();
+			if (!PCGCreateAttribute::ClearOrCreateAttribute(Settings->AttributeTypes, OutputMetadata, OutputAttributeName))
+			{
+				PCGE_LOG(Error, GraphAndLog, FText::Format(PCGCreateAttributeConstants::ErrorCreatingAttributeMessage, FText::FromName(OutputAttributeName)));
+				return true;
+			}
 
-		SourceParamAttributeName = InputSource.GetName();
-		OutputAttributeName = Settings->GetOutputAttributeName(&InputSource, SourceParamData);
+			// Making sure we have at least one entry.
+			if (OutputMetadata && OutputMetadata->GetItemCountForChild() == 0)
+			{
+				OutputMetadata->AddEntry();
+			}
 
-		if (!SourceParamData->Metadata->HasAttribute(SourceParamAttributeName))
-		{
-			PCGE_LOG(Error, GraphAndLog, FText::Format(LOCTEXT("ParamMissingAttribute", "Source Attribute Set data does not have an attribute '{0}'"), FText::FromName(SourceParamAttributeName)));
-			return true;
+			FPCGTaggedData& NewData = Context->OutputData.TaggedData.Emplace_GetRef();
+			NewData.Data = OutputData;
 		}
-	}
-	else
-	{
-		OutputAttributeName = Settings->GetOutputAttributeName(nullptr, nullptr);
-	}
-
-	TArray<FPCGTaggedData> Inputs = Context->InputData.GetInputsByPin(PCGPinConstants::DefaultInputLabel);
 
-	// Only if we have no input connections we'll switch from "adding attribute" to "creating param data"
-	const UPCGPin* InputPin = Context->Node ? Context->Node->GetInputPin(PCGPinConstants::DefaultInputLabel) : nullptr;
-	const bool bHasInputConnections = InputPin && InputPin->EdgeCount() > 0;
+		return true;
+	}
 
-	// If the input is empty, we will create a new ParamData.
-	// We can re-use this newly object as the output
-	bool bCanReuseInputData = false;
-	bool bIsParamData = false;
-	if (!bHasInputConnections)
+	// Otherwise, is is like a copy
+	const UPCGParamData* SourceParamData = !SourceParams.IsEmpty() ? Cast<UPCGParamData>(SourceParams[0].Data) : nullptr;
+	if (!SourceParamData)
 	{
-		ensure(Inputs.IsEmpty());
-
-		FPCGTaggedData& NewData = Inputs.Emplace_GetRef();
-		NewData.Data = NewObject<UPCGParamData>();
-		NewData.Pin = PCGPinConstants::DefaultInputLabel;
-		bCanReuseInputData = true;
+		// Nothing to do
+		Context->OutputData.TaggedData = Inputs;
+		return true;
 	}
 
-	for (const FPCGTaggedData& InputTaggedData : Inputs)
+	for (int32 i = 0; i < Inputs.Num(); ++i)
 	{
-		const UPCGData* InputData = InputTaggedData.Data;
-		UPCGData* OutputData = nullptr;
-
-		UPCGMetadata* Metadata = nullptr;
-
-		if (const UPCGSpatialData* InputSpatialData = Cast<UPCGSpatialData>(InputData))
-		{
-			UPCGSpatialData* NewSpatialData = InputSpatialData->DuplicateData(/*bInitializeFromData=*/false);
-			NewSpatialData->InitializeFromData(InputSpatialData, /*InMetadataParentOverride=*/ nullptr, /*bInheritMetadata=*/true);
-
-			OutputData = NewSpatialData;
-			Metadata = NewSpatialData->Metadata;
-		}
-		else if (const UPCGParamData* InputParamData = Cast<UPCGParamData>(InputData))
-		{
-			// If we can reuse input data, it is safe to const_cast, as it was created by ourselves above.
-			UPCGParamData* NewParamData = bCanReuseInputData ? const_cast<UPCGParamData*>(InputParamData) : NewObject<UPCGParamData>();
-			NewParamData->Metadata->InitializeAsCopy(bCanReuseInputData ? nullptr : InputParamData->Metadata);
-
-			OutputData = NewParamData;
-			Metadata = NewParamData->Metadata;
-			bIsParamData = true;
-		}
-		else
+		const UPCGData* InputData = Inputs[i].Data;
+		if (!InputData)
 		{
-			PCGE_LOG(Error, GraphAndLog, LOCTEXT("InvalidInputData", "Invalid data as input. Only Spatial and Attribute Set data supported."));
 			continue;
 		}
 
-		FPCGMetadataAttributeBase* Attribute = nullptr;
+		UPCGData* TargetData = InputData->DuplicateData();
+		FPCGTaggedData& Output = Context->OutputData.TaggedData.Add_GetRef(Inputs[i]);
 
-		if (SourceParamData)
+		bool bSuccess = false;
+		if (Settings->bCopyAllAttributes)
 		{
-			// We need accessors if we have a multi entry source attribute or we have extractors
-			const bool bIsMultiEntries = SourceParamData->Metadata->GetLocalItemCount() > 1;
-			const bool bNeedAccessors = bIsMultiEntries || !InputSource.GetExtraNames().IsEmpty();
-
-			// If no accessor, copy over the attribute
-			if (!bNeedAccessors)
-			{
-				const FPCGMetadataAttributeBase* SourceAttribute = SourceParamData->Metadata->GetConstAttribute(SourceParamAttributeName);
-				// Presence of attribute was already checked before, this should not return null
-				check(SourceAttribute);
-
-				// Copy the attribute using the first entry of the source attribute as the default value (there is just a single entry or none). If there is no first entry, will be the default value anyway.
-				auto CreateAttribute = [Metadata, SourceAttribute, OutputAttributeName](auto Dummy) -> FPCGMetadataAttributeBase*
-				{
-					using AttributeType = decltype(Dummy);
-					AttributeType DefaultValue = static_cast<const FPCGMetadataAttribute<AttributeType>*>(SourceAttribute)->GetValue(PCGMetadataEntryKey(0));
-					return PCGMetadataElementCommon::ClearOrCreateAttribute<AttributeType>(Metadata, OutputAttributeName, DefaultValue);
-				};
-
-				Attribute = PCGMetadataAttribute::CallbackWithRightType(SourceAttribute->GetTypeId(), std::move(CreateAttribute));
-			}
-			else // Create a new attribute of the accessed field's type manually
-			{
-				TUniquePtr<const IPCGAttributeAccessor> InputAccessor = PCGAttributeAccessorHelpers::CreateConstAccessor(SourceParamData, InputSource);
-				TUniquePtr<const IPCGAttributeAccessorKeys> InputKeys = PCGAttributeAccessorHelpers::CreateConstKeys(SourceParamData, InputSource);
-				if (!InputAccessor.IsValid() || !InputKeys.IsValid())
-				{
-					PCGE_LOG(Error, GraphAndLog, LOCTEXT("FailedToCreateInputAccessor", "Failed to create input accessor"));
-					return true;
-				}
-
-				auto CreateOutputAttribute = [Metadata, OutputAttributeName, &InputAccessor, &InputKeys, OutputData, bIsParamData, bIsMultiEntries]<typename Type>(Type Dummy) -> FPCGMetadataAttributeBase*
-				{
-					// If we have multiple entries, use the zero value as a default value.
-					// Otherwise get the first entry of the source param as default value.
-					Type DefaultValue = PCG::Private::MetadataTraits<Type>::ZeroValue();
-					if (!bIsMultiEntries)
-					{
-						InputAccessor->Get<Type>(DefaultValue, FPCGAttributeAccessorKeysEntries(PCGMetadataEntryKey(0)));
-					}
-
-					FPCGMetadataAttribute<Type>* Attribute = PCGMetadataElementCommon::ClearOrCreateAttribute<Type>(Metadata, OutputAttributeName, DefaultValue);
-
-					FPCGAttributePropertySelector OutputSelector;
-					OutputSelector.SetAttributeName(OutputAttributeName);
-
-					TUniquePtr<IPCGAttributeAccessor> OutputAccessor = PCGAttributeAccessorHelpers::CreateAccessor(OutputData, OutputSelector);
-					TUniquePtr<IPCGAttributeAccessorKeys> OutputKeys = PCGAttributeAccessorHelpers::CreateKeys(OutputData, OutputSelector);
-
-					// We just created the attribute, this should not fail
-					if (!ensure(OutputAccessor.IsValid() && OutputKeys.IsValid()))
-					{
-						return nullptr;
-					}
-
-					PCGMetadataElementCommon::FCopyFromAccessorToAccessorParams Params;
-					Params.InKeys = InputKeys.Get();
-					Params.InAccessor = InputAccessor.Get();
-					Params.OutKeys = OutputKeys.Get();
-					Params.OutAccessor = OutputAccessor.Get();
-					Params.IterationCount = PCGMetadataElementCommon::FCopyFromAccessorToAccessorParams::Min;
-
-					PCGMetadataElementCommon::CopyFromAccessorToAccessor(Params);
-
-					return Attribute;
-				};
-
-				Attribute = PCGMetadataAttribute::CallbackWithRightType(InputAccessor->GetUnderlyingType(), CreateOutputAttribute);
-			}
+			bSuccess = PCGMetadataHelpers::CopyAllAttributes(SourceParamData, TargetData, Context);
 		}
 		else
 		{
-			Attribute = ClearOrCreateAttribute(Settings, Metadata, OutputAttributeName);
-		}
-
-		if (!Attribute)
-		{
-			PCGE_LOG(Error, GraphAndLog, FText::Format(LOCTEXT("ErrorCreatingAttribute", "Error while creating attribute '{0}'"), FText::FromName(OutputAttributeName)));
-			continue;
+			bSuccess = PCGMetadataHelpers::CopyAttributes(SourceParamData, Settings->InputSource, TargetData, Settings->OutputTarget, /*bSameOrigin=*/false, Context);
 		}
 
-		if (bIsParamData)
+		if (bSuccess)
 		{
-			// Making sure the metadata has at least one entry.
-			if (Metadata->GetLocalItemCount() == 0)
-			{
-				Metadata->AddEntry();
-			}
+			Output.Data = TargetData;
 		}
-
-		TArray<FPCGTaggedData>& Outputs = Context->OutputData.TaggedData;
-		FPCGTaggedData& Output = Outputs.Emplace_GetRef();
-		Output.Data = OutputData;
 	}
 
 	return true;
 }
 
-FPCGMetadataAttributeBase* FPCGCreateAttributeElement::ClearOrCreateAttribute(const UPCGCreateAttributeBaseSettings* Settings, UPCGMetadata* Metadata, const FName OutputAttributeName) const
+bool FPCGCreateAttributeElement::ExecuteInternal(FPCGContext* Context) const
 {
-	check(Metadata);
+	TRACE_CPUPROFILER_EVENT_SCOPE(FPCGCreateAttributeElement::Execute);
+
+	check(Context);
+
+	const UPCGCreateAttributeSetSettings* Settings = Context->GetInputSettings<UPCGCreateAttributeSetSettings>();
+	check(Settings);
 
-	auto CreateAttribute = [Settings, Metadata, OutputAttributeName](auto&& Value) -> FPCGMetadataAttributeBase*
+	FName OutputAttributeName = Settings->OutputTarget.GetName();
+
+	UPCGParamData* OutputData = NewObject<UPCGParamData>();
+	check(OutputData && OutputData->Metadata);
+	OutputData->Metadata->AddEntry();
+
+	if (!PCGCreateAttribute::ClearOrCreateAttribute(Settings->AttributeTypes, OutputData->Metadata, OutputAttributeName))
 	{
-		return PCGMetadataElementCommon::ClearOrCreateAttribute(Metadata, OutputAttributeName, std::forward<decltype(Value)>(Value));
-	};
+		PCGE_LOG(Error, GraphAndLog, FText::Format(PCGCreateAttributeConstants::ErrorCreatingAttributeMessage, FText::FromName(OutputAttributeName)));
+		return true;
+	}
 
-	return Settings->AttributeTypes.Dispatcher(CreateAttribute);
+	FPCGTaggedData& NewData = Context->OutputData.TaggedData.Emplace_GetRef();
+	NewData.Data = OutputData;
+
+	return true;
 }
 
 #undef LOCTEXT_NAMESPACE
diff --git a/Engine/Plugins/Experimental/PCG/Source/PCG/Private/Helpers/PCGMetadataHelpers.cpp b/Engine/Plugins/Experimental/PCG/Source/PCG/Private/Helpers/PCGMetadataHelpers.cpp
--- a/Engine/Plugins/Experimental/PCG/Source/PCG/Private/Helpers/PCGMetadataHelpers.cpp
+++ b/Engine/Plugins/Experimental/PCG/Source/PCG/Private/Helpers/PCGMetadataHelpers.cpp
@@ -4,8 +4,14 @@
 
 #include "PCGParamData.h"
 #include "Data/PCGSpatialData.h"
+#include "Elements/Metadata/PCGMetadataElementCommon.h"
+#include "Metadata/Accessors/IPCGAttributeAccessor.h"
+#include "Metadata/Accessors/PCGAttributeAccessorHelpers.h"	
+#include "Metadata/Accessors/PCGAttributeAccessorKeys.h"
 #include "Metadata/PCGMetadata.h"
 
+#define LOCTEXT_NAMESPACE "PCGMetadataHelpers"
+
 namespace PCGMetadataHelpers
 {
 	bool HasSameRoot(const UPCGMetadata* Metadata1, const UPCGMetadata* Metadata2)
@@ -26,34 +32,12 @@ namespace PCGMetadataHelpers
 
 	const UPCGMetadata* GetConstMetadata(const UPCGData* InData)
 	{
-		if (const UPCGSpatialData* SpatialData = Cast<UPCGSpatialData>(InData))
-		{
-			return SpatialData->ConstMetadata();
-		}
-		else if (const UPCGParamData* ParamData = Cast<UPCGParamData>(InData))
-		{
-			return ParamData->ConstMetadata();
-		}
-		else
-		{
-			return nullptr;
-		}
+		return InData ? InData->ConstMetadata() : nullptr;
 	}
 
 	UPCGMetadata* GetMutableMetadata(UPCGData* InData)
 	{
-		if (UPCGSpatialData* SpatialData = Cast<UPCGSpatialData>(InData))
-		{
-			return SpatialData->MutableMetadata();
-		}
-		else if (UPCGParamData* ParamData = Cast<UPCGParamData>(InData))
-		{
-			return ParamData->MutableMetadata();
-		}
-		else
-		{
-			return nullptr;
-		}
+		return InData ? InData->MutableMetadata() : nullptr;
 	}
 
 	bool CreateObjectPathGetter(const FPCGMetadataAttributeBase* InAttributeBase, TFunction<void(int64, FSoftObjectPath&)>& OutGetter)
@@ -124,4 +108,184 @@ namespace PCGMetadataHelpers
 
 		return false;
 	}
+
+	bool CopyAttributes(UPCGData* TargetData, const UPCGData* SourceData, const TArray<TPair<FPCGAttributePropertyInputSelector, FPCGAttributePropertyOutputSelector>>& AttributeSelectors, bool bSameOrigin, FPCGContext* OptionalContext)
+	{
+		check(TargetData && SourceData);
+		const UPCGMetadata* SourceMetadata = SourceData->ConstMetadata();
+		UPCGMetadata* TargetMetadata = TargetData->MutableMetadata();
+
+		if (!SourceMetadata || !TargetMetadata)
+		{
+			return false;
+		}
+
+		bool bSuccess = false;
+
+		for (const auto& SelectorPair : AttributeSelectors)
+		{
+			const FPCGAttributePropertyInputSelector& InputSource = SelectorPair.Key;
+			const FPCGAttributePropertyOutputSelector& OutputTarget = SelectorPair.Value;
+
+			const FName LocalSourceAttribute = InputSource.GetName();
+			const FName LocalDestinationAttribute = OutputTarget.GetName();
+
+			if (InputSource.GetSelection() == EPCGAttributePropertySelection::Attribute && !SourceMetadata->HasAttribute(LocalSourceAttribute))
+			{
+				PCGLog::LogWarningOnGraph(FText::Format(LOCTEXT("InputMissingAttribute", "Input does not have the '{0}' attribute"), FText::FromName(LocalSourceAttribute)), OptionalContext);
+				continue;
+			}
+
+			// We need accessors if we have a multi entry source attribute or we have extractors
+			const bool bIsMultiEntries = SourceData->IsA<UPCGParamData>() && SourceMetadata->GetLocalItemCount() > 1;
+			const bool bInputHasAnyExtra = !InputSource.GetExtraNames().IsEmpty();
+			const bool bOutputHasAnyExtra = !OutputTarget.GetExtraNames().IsEmpty();
+			const bool bSourceIsAttribute = InputSource.GetSelection() == EPCGAttributePropertySelection::Attribute;
+			const bool bTargetIsAttribute = OutputTarget.GetSelection() == EPCGAttributePropertySelection::Attribute;
+
+			const bool bNeedAccessors = bIsMultiEntries || bInputHasAnyExtra || bOutputHasAnyExtra || !bSourceIsAttribute || !bTargetIsAttribute;
+
+			// If no accessor, copy over the attribute
+			if (!bNeedAccessors)
+			{
+				if (bSameOrigin && LocalSourceAttribute == LocalDestinationAttribute)
+				{
+					// Nothing to do if we try to copy an attribute into itself in the original data.
+					continue;
+				}
+
+				const FPCGMetadataAttributeBase* SourceAttribute = SourceMetadata->GetConstAttribute(LocalSourceAttribute);
+				// Presence of attribute was already checked before, this should not return null
+				check(SourceAttribute);
+
+				// Copy the attribute using the first entry of the source attribute as the default value (there is just a single entry or none). If there is no first entry, will be the default value anyway.
+				auto CreateAttribute = [TargetMetadata, SourceAttribute, LocalDestinationAttribute](auto Dummy) -> FPCGMetadataAttributeBase*
+				{
+					using AttributeType = decltype(Dummy);
+					AttributeType DefaultValue = static_cast<const FPCGMetadataAttribute<AttributeType>*>(SourceAttribute)->GetValue(PCGMetadataEntryKey(0));
+					return PCGMetadataElementCommon::ClearOrCreateAttribute<AttributeType>(TargetMetadata, LocalDestinationAttribute, DefaultValue);
+				};
+
+				if (!PCGMetadataAttribute::CallbackWithRightType(SourceAttribute->GetTypeId(), std::move(CreateAttribute)))
+				{
+					PCGLog::LogWarningOnGraph(FText::Format(LOCTEXT("FailedCreateNewAttribute", "Failed to create new attribute '{0}'"), FText::FromName(LocalDestinationAttribute)));
+					continue;
+				}
+			}
+			else // Create a new attribute of the accessed field's type manually
+			{
+				TUniquePtr<const IPCGAttributeAccessor> InputAccessor = PCGAttributeAccessorHelpers::CreateConstAccessor(SourceData, InputSource);
+				TUniquePtr<const IPCGAttributeAccessorKeys> InputKeys = PCGAttributeAccessorHelpers::CreateConstKeys(SourceData, InputSource);
+
+				if (!InputAccessor.IsValid() || !InputKeys.IsValid())
+				{
+					PCGLog::LogWarningOnGraph(LOCTEXT("FailedToCreateInputAccessor", "Failed to create input accessor or iterator"), OptionalContext);
+					continue;
+				}
+
+				// If the target is an attribute, only create a new one if the attribute we don't have any extra.
+				// If it has any extra, it will try to write to it.
+				if (!bOutputHasAnyExtra && bTargetIsAttribute)
+				{
+					auto CreateAttribute = [TargetMetadata, LocalDestinationAttribute](auto Dummy)
+					{
+						using AttributeType = decltype(Dummy);
+						return PCGMetadataElementCommon::ClearOrCreateAttribute<AttributeType>(TargetMetadata, LocalDestinationAttribute) != nullptr;
+					};
+
+					if (!PCGMetadataAttribute::CallbackWithRightType(InputAccessor->GetUnderlyingType(), CreateAttribute))
+					{
+						PCGLog::LogWarningOnGraph(FText::Format(LOCTEXT("FailedToCreateNewAttribute", "Failed to create new attribute '{0}'"), FText::FromName(LocalDestinationAttribute)), OptionalContext);
+						continue;
+					}
+				}
+
+				TUniquePtr<IPCGAttributeAccessor> OutputAccessor = PCGAttributeAccessorHelpers::CreateAccessor(TargetData, OutputTarget);
+				TUniquePtr<IPCGAttributeAccessorKeys> OutputKeys = PCGAttributeAccessorHelpers::CreateKeys(TargetData, OutputTarget);
+
+				if (!OutputAccessor.IsValid() || !OutputKeys.IsValid())
+				{
+					PCGLog::LogWarningOnGraph(LOCTEXT("FailedToCreateOutputAccessor", "Failed to create output accessor or iterator"), OptionalContext);
+					continue;
+				}
+
+				if (OutputAccessor->IsReadOnly())
+				{
+					PCGLog::LogWarningOnGraph(FText::Format(LOCTEXT("OutputAccessorIsReadOnly", "Attribute/Property '{0}' is read only."), OutputTarget.GetDisplayText()), OptionalContext);
+					continue;
+				}
+
+				// Final verification, if we can put the value of input into output
+				if (!PCG::Private::IsBroadcastableOrConstructible(InputAccessor->GetUnderlyingType(), OutputAccessor->GetUnderlyingType()))
+				{
+					PCGLog::LogErrorOnGraph(FText::Format(LOCTEXT("CannotConvertTypes", "Cannot convert input type {0} into output type {1}"), PCG::Private::GetTypeNameText(InputAccessor->GetUnderlyingType()), PCG::Private::GetTypeNameText(OutputAccessor->GetUnderlyingType())), OptionalContext);
+					continue;
+				}
+
+				// At this point, we are ready.
+				PCGMetadataElementCommon::FCopyFromAccessorToAccessorParams Params;
+				Params.InKeys = InputKeys.Get();
+				Params.InAccessor = InputAccessor.Get();
+				Params.OutKeys = OutputKeys.Get();
+				Params.OutAccessor = OutputAccessor.Get();
+				Params.IterationCount = PCGMetadataElementCommon::FCopyFromAccessorToAccessorParams::Out;
+				Params.Flags = EPCGAttributeAccessorFlags::AllowBroadcast | EPCGAttributeAccessorFlags::AllowConstructible;
+
+				if (!PCGMetadataElementCommon::CopyFromAccessorToAccessor(Params))
+				{
+					PCGLog::LogWarningOnGraph(LOCTEXT("ErrorGettingSettingValues", "Error while getting/setting values"), OptionalContext);
+					continue;
+				}
+			}
+
+			bSuccess = true;
+		}
+
+		return bSuccess;
+	}
+
+	bool CopyAttributes(const UPCGData* SourceData, const FPCGAttributePropertyInputSelector& _InputSource, UPCGData* TargetData, const FPCGAttributePropertyOutputSelector& _OutputTarget, bool bSameOrigin, FPCGContext* OptionalContext)
+	{
+		if (!TargetData || !SourceData)
+		{
+			return false;
+		}
+
+		TArray<TPair<FPCGAttributePropertyInputSelector, FPCGAttributePropertyOutputSelector>> AttributeSelectors;
+		FPCGAttributePropertyInputSelector InputSource = _InputSource.CopyAndFixLast(SourceData);
+		FPCGAttributePropertyOutputSelector OutputTarget = _OutputTarget.CopyAndFixSource(&InputSource, SourceData);
+
+		AttributeSelectors.Emplace(MoveTemp(InputSource), MoveTemp(OutputTarget));
+		return CopyAttributes(TargetData, SourceData, AttributeSelectors, bSameOrigin, OptionalContext);
+	}
+
+	bool CopyAllAttributes(const UPCGData* SourceData, UPCGData* TargetData, FPCGContext* OptionalContext)
+	{
+		if (!TargetData || !SourceData)
+		{
+			return false;
+		}
+
+		const UPCGMetadata* SourceMetadata = SourceData->ConstMetadata();
+		if (!SourceMetadata)
+		{
+			return false;
+		}
+
+		TArray<TPair<FPCGAttributePropertyInputSelector, FPCGAttributePropertyOutputSelector>> AttributeSelectors;
+		TArray<FName> AttributeNames;
+		TArray<EPCGMetadataTypes> AttributeTypes;
+		SourceMetadata->GetAttributes(AttributeNames, AttributeTypes);
+
+		for (const FName& AttributeName : AttributeNames)
+		{
+			TPair<FPCGAttributePropertyInputSelector, FPCGAttributePropertyOutputSelector>& Selectors = AttributeSelectors.Emplace_GetRef();
+			Selectors.Key.SetAttributeName(AttributeName);
+			Selectors.Value.SetAttributeName(AttributeName);
+		}
+
+		return CopyAttributes(TargetData, SourceData, AttributeSelectors, /*bSameOrigin=*/false, OptionalContext);
+	}
 }
+
+#undef LOCTEXT_NAMESPACE
\ No newline at end of file
diff --git a/Engine/Plugins/Experimental/PCG/Source/PCG/Public/Elements/PCGCreateAttribute.h b/Engine/Plugins/Experimental/PCG/Source/PCG/Public/Elements/PCGCreateAttribute.h
--- a/Engine/Plugins/Experimental/PCG/Source/PCG/Public/Elements/PCGCreateAttribute.h
+++ b/Engine/Plugins/Experimental/PCG/Source/PCG/Public/Elements/PCGCreateAttribute.h
@@ -12,108 +12,13 @@
 class FPCGMetadataAttributeBase;
 class UPCGMetadata;
 
-/**
-* Adds an attribute to Spatial data or to an Attribute Set, or creates a new Attribute Set if no
-* input is provided.
-* 
-* Note: This need to be updated if we ever add new types.
-*/
-UCLASS(Abstract, BlueprintType, ClassGroup = (Procedural))
-class PCG_API UPCGCreateAttributeBaseSettings : public UPCGSettings
-{
-	GENERATED_BODY()
-public:
-	//~Begin UObject interface
-	virtual void PostLoad() override;
-	//~End UObject interface
-
-	//~Begin UPCGSettings interface
-#if WITH_EDITOR
-	virtual EPCGSettingsType GetType() const override { return EPCGSettingsType::Metadata; }
-	virtual bool IsPinUsedByNodeExecution(const UPCGPin* InPin) const override;
-	virtual bool CanEditChange(const FProperty* InProperty) const override;
-#endif // WITH_EDITOR
-	virtual EPCGDataType GetCurrentPinTypes(const UPCGPin* InPin) const override;
-
-protected:
-#if WITH_EDITOR
-	virtual EPCGChangeType GetChangeTypeForProperty(const FName& InPropertyName) const override { return Super::GetChangeTypeForProperty(InPropertyName) | EPCGChangeType::Cosmetic; }
-#endif
-	//~End UPCGSettings interface
-
-public:
-	virtual FName GetOutputAttributeName(const FPCGAttributePropertyInputSelector* InSource, const UPCGData* InSourceData) const PURE_VIRTUAL(UPCGCreateAttributeBaseSettings::GetOutputAttributeName, return NAME_None;);
-
-	// This can be set false by inheriting nodes to hide the 'From Source Param' property.
-	UPROPERTY(Transient, meta = (EditCondition = false, EditConditionHides))
-	bool bDisplayFromSourceParamSetting = true;
-
-	UPROPERTY(BlueprintReadWrite, EditAnywhere, Category = Settings, meta = (EditCondition = "bDisplayFromSourceParamSetting", EditConditionHides, HideEditConditionToggle, PCG_DiscardPropertySelection, PCG_DiscardExtraSelection, PCG_Overridable))
-	FPCGAttributePropertyInputSelector InputSource;
-
-	UPROPERTY(BlueprintReadWrite, EditAnywhere, Category = Settings, meta = (ShowOnlyInnerProperties))
-	FPCGMetadataTypesConstantStruct AttributeTypes;
-
-#if WITH_EDITORONLY_DATA
-	UPROPERTY()
-	EPCGMetadataTypes Type_DEPRECATED = EPCGMetadataTypes::Double;
-
-	UPROPERTY()
-	float FloatValue_DEPRECATED = 0.0;
-
-	UPROPERTY()
-	int32 Int32Value_DEPRECATED = 0;
-
-	UPROPERTY()
-	double DoubleValue_DEPRECATED = 0.0;
-
-	UPROPERTY()
-	int64 IntValue_DEPRECATED = 0;
-
-	UPROPERTY()
-	FVector2D Vector2Value_DEPRECATED = FVector2D::ZeroVector;
-
-	UPROPERTY()
-	FVector VectorValue_DEPRECATED = FVector::ZeroVector;
-
-	UPROPERTY()
-	FVector4 Vector4Value_DEPRECATED = FVector4::Zero();
-
-	UPROPERTY()
-	FQuat QuatValue_DEPRECATED = FQuat::Identity;
 
-	UPROPERTY()
-	FTransform TransformValue_DEPRECATED = FTransform::Identity;
-
-	UPROPERTY()
-	FString StringValue_DEPRECATED = "";
-
-	UPROPERTY()
-	bool BoolValue_DEPRECATED = false;
-
-	UPROPERTY()
-	FRotator RotatorValue_DEPRECATED = FRotator::ZeroRotator;
-
-	UPROPERTY()
-	FName NameValue_DEPRECATED = NAME_None;
-
-	UPROPERTY()
-	bool bKeepExistingAttributes_DEPRECATED = false;
-
-	UPROPERTY()
-	FName SourceParamAttributeName_DEPRECATED = NAME_None;
-#endif // WITH_EDITORONLY_DATA
-
-protected:
-	FString GetAdditionalTitleInformationInternal(FName NodeName) const;
-
-	virtual FPCGElementPtr CreateElement() const override;
-	virtual bool ShouldAddAttributesPin() const { return true; }
-};
-
-/* Add a new attribute. */
+/** Add a new attribute to a spatial data or an attribute set.
+* New attribute can be a constant, hardcoded in the node, or can come from another Attribute Set.
+* Can also add all the attributes coming from the other Attribute Set.
+*/
 UCLASS(BlueprintType, ClassGroup = (Procedural))
-class PCG_API UPCGAddAttributeSettings : public UPCGCreateAttributeBaseSettings
+class PCG_API UPCGAddAttributeSettings : public UPCGSettings
 {
 	GENERATED_BODY()
 
@@ -130,7 +35,12 @@ public:
 	virtual FText GetDefaultNodeTitle() const override;
 	virtual bool HasDynamicPins() const override { return true; }
 	virtual void ApplyDeprecation(UPCGNode* InOutNode) override;
+	virtual EPCGSettingsType GetType() const override { return EPCGSettingsType::Param; }
+	virtual bool IsPinUsedByNodeExecution(const UPCGPin* InPin) const override;
+	virtual bool CanEditChange(const FProperty* InProperty) const override;
 #endif // WITH_EDITOR
+
+	virtual EPCGDataType GetCurrentPinTypes(const UPCGPin* InPin) const override;
 	virtual FString GetAdditionalTitleInformation() const override;
 
 protected:
@@ -140,14 +50,22 @@ protected:
 	//~End UPCGSettings interface
 
 public:
-	//~Begin UPCGCreateAttributeBaseSettings interface
-	virtual FName GetOutputAttributeName(const FPCGAttributePropertyInputSelector* InSource, const UPCGData* InSourceData) const { return OutputTarget.CopyAndFixSource(InSource, InSourceData).GetName(); }
-	//~End UPCGCreateAttributeBaseSettings interface
+	UPROPERTY(BlueprintReadWrite, EditAnywhere, Category = Settings, meta = (EditCondition = "!bCopyAllAttributes", HideEditConditionToggle, PCG_DiscardPropertySelection, PCG_DiscardExtraSelection, PCG_Overridable))
+	FPCGAttributePropertyInputSelector InputSource;
 
-	UPROPERTY(BlueprintReadWrite, EditAnywhere, Category = Settings, meta = (PCG_DiscardPropertySelection, PCG_DiscardExtraSelection, PCG_Overridable))
+	UPROPERTY(BlueprintReadWrite, EditAnywhere, Category = Settings, meta = (EditCondition = "!bCopyAllAttributes", PCG_DiscardPropertySelection, PCG_DiscardExtraSelection, PCG_Overridable))
 	FPCGAttributePropertyOutputSelector OutputTarget;
 
+	UPROPERTY(BlueprintReadWrite, EditAnywhere, Category = Settings, meta = (ShowOnlyInnerProperties))
+	FPCGMetadataTypesConstantStruct AttributeTypes;
+
+	UPROPERTY(BlueprintReadWrite, EditAnywhere, Category = Settings, meta = (PCG_Overridable))
+	bool bCopyAllAttributes = false;
+
 #if WITH_EDITORONLY_DATA
+	UPROPERTY()
+	FName SourceParamAttributeName_DEPRECATED = NAME_None;
+
 	UPROPERTY()
 	FName OutputAttributeName_DEPRECATED = NAME_None;
 #endif // WITH_EDITORONLY_DATA
@@ -155,17 +73,16 @@ public:
 protected:
 	virtual TArray<FPCGPinProperties> InputPinProperties() const override;
 	virtual TArray<FPCGPinProperties> OutputPinProperties() const override;
+	virtual FPCGElementPtr CreateElement() const override;
 };
 
 /* Creates a new Attribute Set. */
 UCLASS(BlueprintType, ClassGroup = (Procedural))
-class PCG_API UPCGCreateAttributeSetSettings : public UPCGCreateAttributeBaseSettings
+class PCG_API UPCGCreateAttributeSetSettings : public UPCGSettings
 {
 	GENERATED_BODY()
 
 public:
-	UPCGCreateAttributeSetSettings();
-
 	//~Begin UObject interface
 	virtual void PostLoad() override;
 	//~End UObject interface
@@ -174,7 +91,7 @@ public:
 #if WITH_EDITOR
 	virtual FName GetDefaultNodeName() const override;
 	virtual FText GetDefaultNodeTitle() const override;
-	virtual bool HasDynamicPins() const override { return false; }
+	virtual EPCGSettingsType GetType() const override { return EPCGSettingsType::Param; }
 #endif
 	virtual FString GetAdditionalTitleInformation() const override;
 
@@ -185,13 +102,12 @@ protected:
 	//~End UPCGSettings interface
 
 public:
+	UPROPERTY(BlueprintReadWrite, EditAnywhere, Category = Settings, meta = (ShowOnlyInnerProperties))
+	FPCGMetadataTypesConstantStruct AttributeTypes;
+
 	UPROPERTY(BlueprintReadWrite, EditAnywhere, Category = Settings, meta = (PCG_DiscardPropertySelection, PCG_DiscardExtraSelection, PCG_Overridable))
 	FPCGAttributePropertyOutputNoSourceSelector OutputTarget;
 
-	//~Begin UPCGCreateAttributeBaseSettings interface
-	virtual FName GetOutputAttributeName(const FPCGAttributePropertyInputSelector*, const UPCGData*) const { return OutputTarget.GetName(); }
-	//~End UPCGCreateAttributeBaseSettings interface
-
 #if WITH_EDITORONLY_DATA
 	UPROPERTY()
 	FName OutputAttributeName_DEPRECATED = NAME_None;
@@ -200,18 +116,19 @@ public:
 protected:
 	virtual TArray<FPCGPinProperties> InputPinProperties() const override;
 	virtual TArray<FPCGPinProperties> OutputPinProperties() const override;
-
-	virtual bool ShouldAddAttributesPin() const override { return false; }
+	virtual FPCGElementPtr CreateElement() const override;
 };
 
 class FPCGCreateAttributeElement : public IPCGElement
 {
 protected:
 	virtual bool ExecuteInternal(FPCGContext* Context) const override;
+};
 
-private:
-	/* Create (or clear) an attribute named by OutputAttributeName and depending on the selected type. Value can be overridden by params. Default value will be set to the specified value. */
-	FPCGMetadataAttributeBase* ClearOrCreateAttribute(const UPCGCreateAttributeBaseSettings* Settings, UPCGMetadata* Metadata, const FName OutputAttributeName) const;
+class FPCGAddAttributeElement : public IPCGElement
+{
+protected:
+	virtual bool ExecuteInternal(FPCGContext* Context) const override;
 };
 
 #if UE_ENABLE_INCLUDE_ORDER_DEPRECATED_IN_5_2
diff --git a/Engine/Plugins/Experimental/PCG/Source/PCG/Public/Helpers/PCGMetadataHelpers.h b/Engine/Plugins/Experimental/PCG/Source/PCG/Public/Helpers/PCGMetadataHelpers.h
--- a/Engine/Plugins/Experimental/PCG/Source/PCG/Public/Helpers/PCGMetadataHelpers.h
+++ b/Engine/Plugins/Experimental/PCG/Source/PCG/Public/Helpers/PCGMetadataHelpers.h
@@ -4,11 +4,14 @@
 
 #include "HAL/Platform.h"
 
-template <typename FuncType> class TFunction;
+class FPCGMetadataAttributeBase;
 class UPCGData;
 class UPCGMetadata;
-class FPCGMetadataAttributeBase;
+struct FPCGAttributePropertyInputSelector;
+struct FPCGAttributePropertyOutputSelector;
+struct FPCGContext;
 struct FSoftObjectPath;
+template <typename FuncType> class TFunction;
 
 namespace PCGMetadataHelpers
 {
@@ -24,6 +27,14 @@ namespace PCGMetadataHelpers
 
 	/** Create a lambda that will construct a soft object path from an underlying attribute of type FSoftObjectPath or FString. Returns true if successful. */
 	UE_NODISCARD PCG_API bool CreateObjectOrClassPathGetter(const FPCGMetadataAttributeBase* InAttributeBase, TFunction<void(int64, FSoftObjectPath&)>& OutGetter);
+
+	/** Copy the attribute coming from Source Data into Target Data.
+	* Can set bSameOrigin to true if SourceData and TargetData have the same origin (if TargetData was initialized from SourceData). 
+	*/
+	PCG_API bool CopyAttributes(const UPCGData* SourceData, const FPCGAttributePropertyInputSelector& InputSource, UPCGData* TargetData, const FPCGAttributePropertyOutputSelector& OutputTarget, bool bSameOrigin, FPCGContext* OptionalContext = nullptr);
+
+	/** Copy all the attributes coming from Source Data into Target Data. */
+	PCG_API bool CopyAllAttributes(const UPCGData* SourceData, UPCGData* TargetData, FPCGContext* OptionalContext = nullptr);
 }
 
 #if UE_ENABLE_INCLUDE_ORDER_DEPRECATED_IN_5_2
