[PCG] Support actor functions that take Point + Metadata as parameters
* Add UPCGFunctionPrototypes as a BlueprintLibrary
* Provide PostSpawn function names for PCGSpawnActor that take Point + Metadata as parameters. These execute per point.

diff --git a/Engine/Plugins/Experimental/PCG/Source/PCG/Private/Elements/PCGApplyOnActor.cpp b/Engine/Plugins/Experimental/PCG/Source/PCG/Private/Elements/PCGApplyOnActor.cpp
--- a/Engine/Plugins/Experimental/PCG/Source/PCG/Private/Elements/PCGApplyOnActor.cpp
+++ b/Engine/Plugins/Experimental/PCG/Source/PCG/Private/Elements/PCGApplyOnActor.cpp
@@ -48,7 +48,7 @@ bool FPCGApplyOnActorElement::ExecuteInternal(FPCGContext* Context) const
 
 	PCGActorPropertyOverrideHelpers::ApplyOverridesFromParams(Settings->PropertyOverrideDescriptions, TargetActor, PCGApplyOnActorConstants::ActorPropertyOverridesLabel, Context);
 
-	for (UFunction* Function : PCGHelpers::FindUserFunctions(TargetActor->GetClass(), Settings->PostProcessFunctionNames, Context))
+	for (UFunction* Function : PCGHelpers::FindUserFunctions(TargetActor->GetClass(), Settings->PostProcessFunctionNames, { UPCGFunctionPrototypes::GetPrototypeWithNoParams() }, Context))
 	{
 		TargetActor->ProcessEvent(Function, nullptr);
 	}
diff --git a/Engine/Plugins/Experimental/PCG/Source/PCG/Private/Elements/PCGCreateSpline.cpp b/Engine/Plugins/Experimental/PCG/Source/PCG/Private/Elements/PCGCreateSpline.cpp
--- a/Engine/Plugins/Experimental/PCG/Source/PCG/Private/Elements/PCGCreateSpline.cpp
+++ b/Engine/Plugins/Experimental/PCG/Source/PCG/Private/Elements/PCGCreateSpline.cpp
@@ -218,7 +218,7 @@ bool FPCGCreateSplineElement::ExecuteInternal(FPCGContext* Context) const
 		// Execute PostProcess Functions
 		if (SplineActor)
 		{
-			for (UFunction* Function : PCGHelpers::FindUserFunctions(SplineActor->GetClass(), Settings->PostProcessFunctionNames, Context))
+			for (UFunction* Function : PCGHelpers::FindUserFunctions(SplineActor->GetClass(), Settings->PostProcessFunctionNames, { UPCGFunctionPrototypes::GetPrototypeWithNoParams() }, Context))
 			{
 				SplineActor->ProcessEvent(Function, nullptr);
 			}
diff --git a/Engine/Plugins/Experimental/PCG/Source/PCG/Private/Elements/PCGCreateTargetActor.cpp b/Engine/Plugins/Experimental/PCG/Source/PCG/Private/Elements/PCGCreateTargetActor.cpp
--- a/Engine/Plugins/Experimental/PCG/Source/PCG/Private/Elements/PCGCreateTargetActor.cpp
+++ b/Engine/Plugins/Experimental/PCG/Source/PCG/Private/Elements/PCGCreateTargetActor.cpp
@@ -295,7 +295,7 @@ bool FPCGCreateTargetActorElement::ExecuteInternal(FPCGContext* Context) const
 	// Apply property overrides to the GeneratedActor
 	PCGActorPropertyOverrideHelpers::ApplyOverridesFromParams(Settings->PropertyOverrideDescriptions, GeneratedActor, PCGCreateTargetActorConstants::ActorPropertyOverridesLabel, Context);
 
-	for (UFunction* Function : PCGHelpers::FindUserFunctions(GeneratedActor->GetClass(), Settings->PostProcessFunctionNames, Context))
+	for (UFunction* Function : PCGHelpers::FindUserFunctions(GeneratedActor->GetClass(), Settings->PostProcessFunctionNames, { UPCGFunctionPrototypes::GetPrototypeWithNoParams() }, Context))
 	{
 		GeneratedActor->ProcessEvent(Function, nullptr);
 	}
diff --git a/Engine/Plugins/Experimental/PCG/Source/PCG/Private/Elements/PCGSpawnActor.cpp b/Engine/Plugins/Experimental/PCG/Source/PCG/Private/Elements/PCGSpawnActor.cpp
--- a/Engine/Plugins/Experimental/PCG/Source/PCG/Private/Elements/PCGSpawnActor.cpp
+++ b/Engine/Plugins/Experimental/PCG/Source/PCG/Private/Elements/PCGSpawnActor.cpp
@@ -829,7 +829,14 @@ void FPCGSpawnActorElement::SpawnActors(FPCGSubgraphContext* Context, AActor* Ta
 		PCGHelpers::GetGeneratedActorsFolderPath(TargetActor, GeneratedActorsFolderPath);
 #endif
 
-		const TArray<UFunction*> PostSpawnFunctions = PCGHelpers::FindUserFunctions(InTemplateActorClass, Settings->PostSpawnFunctionNames, Context);
+		const UFunction* FunctionPrototypeWithNoParams = UPCGFunctionPrototypes::GetPrototypeWithNoParams();
+		const UFunction* FunctionPrototypeWithPointAndMetadata = UPCGFunctionPrototypes::GetPrototypeWithPointAndMetadata();
+
+		const TArray<UFunction*> PostSpawnFunctions = PCGHelpers::FindUserFunctions(
+			InTemplateActorClass,
+			Settings->PostSpawnFunctionNames,
+			{ FunctionPrototypeWithNoParams, FunctionPrototypeWithPointAndMetadata },
+			Context);
 
 		bool bAllActorOverridesSucceeded = true;
 
@@ -854,7 +861,15 @@ void FPCGSpawnActorElement::SpawnActors(FPCGSubgraphContext* Context, AActor* Ta
 
 			for (UFunction* PostSpawnFunction : PostSpawnFunctions)
 			{
-				GeneratedActor->ProcessEvent(PostSpawnFunction, nullptr);
+				if (PostSpawnFunction->IsSignatureCompatibleWith(FunctionPrototypeWithNoParams))
+				{
+					GeneratedActor->ProcessEvent(PostSpawnFunction, nullptr);
+				}
+				else if (PostSpawnFunction->IsSignatureCompatibleWith(FunctionPrototypeWithPointAndMetadata))
+				{
+					TPair<FPCGPoint, const UPCGMetadata*> PointAndMetadata = { Point, PointData->ConstMetadata() };
+					GeneratedActor->ProcessEvent(PostSpawnFunction, &PointAndMetadata);
+				}
 			}
 
 			ManagedActors->GeneratedActors.Add(GeneratedActor);
diff --git a/Engine/Plugins/Experimental/PCG/Source/PCG/Private/Elements/PCGStaticMeshSpawner.cpp b/Engine/Plugins/Experimental/PCG/Source/PCG/Private/Elements/PCGStaticMeshSpawner.cpp
--- a/Engine/Plugins/Experimental/PCG/Source/PCG/Private/Elements/PCGStaticMeshSpawner.cpp
+++ b/Engine/Plugins/Experimental/PCG/Source/PCG/Private/Elements/PCGStaticMeshSpawner.cpp
@@ -302,7 +302,7 @@ bool FPCGStaticMeshSpawnerElement::ExecuteInternal(FPCGContext* InContext) const
 	{
 		if (AActor* TargetActor = Settings->TargetActor.Get() ? Settings->TargetActor.Get() : Context->GetTargetActor(nullptr))
 		{
-			for (UFunction* Function : PCGHelpers::FindUserFunctions(TargetActor->GetClass(), Settings->PostProcessFunctionNames, Context))
+			for (UFunction* Function : PCGHelpers::FindUserFunctions(TargetActor->GetClass(), Settings->PostProcessFunctionNames, { UPCGFunctionPrototypes::GetPrototypeWithNoParams() }, Context))
 			{
 				TargetActor->ProcessEvent(Function, nullptr);
 			}
diff --git a/Engine/Plugins/Experimental/PCG/Source/PCG/Private/Helpers/PCGHelpers.cpp b/Engine/Plugins/Experimental/PCG/Source/PCG/Private/Helpers/PCGHelpers.cpp
--- a/Engine/Plugins/Experimental/PCG/Source/PCG/Private/Helpers/PCGHelpers.cpp
+++ b/Engine/Plugins/Experimental/PCG/Source/PCG/Private/Helpers/PCGHelpers.cpp
@@ -492,7 +492,7 @@ namespace PCGHelpers
 #endif
 	}
 
-	TArray<UFunction*> FindUserFunctions(TSubclassOf<AActor> ActorClass, const TArray<FName>& FunctionNames, const FPCGContext* InContext)
+	TArray<UFunction*> FindUserFunctions(TSubclassOf<AActor> ActorClass, const TArray<FName>& FunctionNames, const TArray<const UFunction*>& FunctionPrototypes, const FPCGContext* InContext)
 	{
 		TArray<UFunction*> Functions;
 
@@ -510,20 +510,25 @@ namespace PCGHelpers
 
 			if (UFunction* Function = ActorClass->FindFunctionByName(FunctionName))
 			{
-				// TODO: support Parameters
-				if (Function->NumParms != 0)
-				{
-					PCGLog::LogWarningOnGraph(FText::Format(LOCTEXT("ParametersMissing", "Function '{0}' in class '{1}' has parameters, but only parameter-less functions are supported. Call will be skipped."), FText::FromName(FunctionName), FText::FromName(ActorClass->GetFName())), InContext);
-				}
 #if WITH_EDITOR
-				else if (!Function->GetBoolMetaData(TEXT("CallInEditor")))
+				if (!Function->GetBoolMetaData(TEXT("CallInEditor")))
 				{
 					PCGLog::LogWarningOnGraph(FText::Format(LOCTEXT("CallInEditorFailed", "Function '{0}' in class '{1}' requires CallInEditor to be true while in-editor."), FText::FromName(FunctionName), FText::FromName(ActorClass->GetFName())), InContext);
+					continue;
 				}
 #endif
-				else
+				for (const UFunction* Prototype : FunctionPrototypes)
+				{
+					if (Function->IsSignatureCompatibleWith(Prototype))
+					{
+						Functions.Add(Function);
+						break;
+					}
+				}
+
+				if (Functions.IsEmpty() || Functions.Last() != Function)
 				{
-					Functions.Add(Function);
+					PCGLog::LogWarningOnGraph(FText::Format(LOCTEXT("ParametersIncorrect", "Function '{0}' in class '{1}' has incorrect parameters."), FText::FromName(FunctionName), FText::FromName(ActorClass->GetFName())), InContext);
 				}
 			}
 			else
diff --git a/Engine/Plugins/Experimental/PCG/Source/PCG/Public/Elements/PCGSpawnActor.h b/Engine/Plugins/Experimental/PCG/Source/PCG/Public/Elements/PCGSpawnActor.h
--- a/Engine/Plugins/Experimental/PCG/Source/PCG/Public/Elements/PCGSpawnActor.h
+++ b/Engine/Plugins/Experimental/PCG/Source/PCG/Public/Elements/PCGSpawnActor.h
@@ -44,7 +44,10 @@ public:
 	UPROPERTY(BlueprintReadOnly, EditAnywhere, Category = Settings, meta = (OnlyPlaceable, DisallowCreateNew))
 	TSubclassOf<AActor> TemplateActorClass = nullptr;
 
-	/* Can specify a list of functions from the template class to be called on each actor spawned, in order. Need to be parameter-less and with "CallInEditor" flag enabled.*/
+	/** 
+	 * Can specify a list of functions from the template class to be called on each actor spawned, in order. Need to have "CallInEditor" flag enabled
+	 * and have either no parameters or exactly the parameters PCGPoint and PCGMetadata
+	 */
 	UPROPERTY(BlueprintReadWrite, EditAnywhere, Category = Settings, meta = (EditCondition = "Option!=EPCGSpawnActorOption::CollapseActors", EditConditionHides))
 	TArray<FName> PostSpawnFunctionNames;
 
diff --git a/Engine/Plugins/Experimental/PCG/Source/PCG/Public/Helpers/PCGHelpers.h b/Engine/Plugins/Experimental/PCG/Source/PCG/Public/Helpers/PCGHelpers.h
--- a/Engine/Plugins/Experimental/PCG/Source/PCG/Public/Helpers/PCGHelpers.h
+++ b/Engine/Plugins/Experimental/PCG/Source/PCG/Public/Helpers/PCGHelpers.h
@@ -3,17 +3,23 @@
 #pragma once
 
 #include "PCGCommon.h"
+#include "PCGPoint.h"
 
+#include "Kismet/BlueprintFunctionLibrary.h"
 #include "Math/Box.h"
 
+#include "PCGHelpers.generated.h"
+
 class AActor;
 class APCGWorldActor;
 class ALandscape;
 class ALandscapeProxy;
 class UPCGComponent;
 class UPCGGraph;
+class UPCGMetadata;
 class UWorld;
 struct FPCGContext;
+struct FPCGPoint;
 
 namespace PCGHelpers
 {
@@ -72,6 +78,27 @@ namespace PCGHelpers
 
 	PCG_API void AttachToParent(AActor* InActorToAttach, AActor* InParent, EPCGAttachOptions AttachOptions, const FString& GeneratedPath = FString());
 
-	/** Finds functions on the actor matching the provided function names. Functions must be marked as CallInEditor and take no parameters. */
-	PCG_API TArray<UFunction*> FindUserFunctions(TSubclassOf<AActor> ActorClass, const TArray<FName>& FunctionNames, const FPCGContext* InContext = nullptr);
+	/**
+	 * Finds functions on the actor matching the provided function names. Functions must be marked as CallInEditor
+	 * and have parameters matching one of the provided prototypes. Some prototypes are provided in UPCGFunctionPrototypes.
+	 */
+	PCG_API TArray<UFunction*> FindUserFunctions(TSubclassOf<AActor> ActorClass, const TArray<FName>& FunctionNames, const TArray<const UFunction*>& FunctionPrototypes, const FPCGContext* InContext = nullptr);
+};
+
+/** Holds function prototypes used to match against actor function signatures. */
+UCLASS()
+class PCG_API UPCGFunctionPrototypes : public UBlueprintFunctionLibrary
+{
+	GENERATED_BODY()
+
+public:
+	static UFunction* GetPrototypeWithNoParams() { return FindObject<UFunction>(StaticClass(), TEXT("PrototypeWithNoParams")); }
+	static UFunction* GetPrototypeWithPointAndMetadata() { return FindObject<UFunction>(StaticClass(), TEXT("PrototypeWithPointAndMetadata")); }
+
+private:
+	UFUNCTION()
+	void PrototypeWithNoParams() {}
+
+	UFUNCTION()
+	void PrototypeWithPointAndMetadata(FPCGPoint Point, const UPCGMetadata* Metadata) {}
 };
