[PCG] Reorder debug object tree sorting priorities so that an explicit sort order is always the primary sort key.

This ensures that loop iteration items are always displayed in loop order independent of children.

diff --git a/Engine/Plugins/Experimental/PCG/Source/PCGEditor/Private/SPCGEditorGraphDebugObjectTree.cpp b/Engine/Plugins/Experimental/PCG/Source/PCGEditor/Private/SPCGEditorGraphDebugObjectTree.cpp
--- a/Engine/Plugins/Experimental/PCG/Source/PCGEditor/Private/SPCGEditorGraphDebugObjectTree.cpp
+++ b/Engine/Plugins/Experimental/PCG/Source/PCGEditor/Private/SPCGEditorGraphDebugObjectTree.cpp
@@ -103,25 +103,25 @@ void FPCGEditorGraphDebugObjectItem::SortChildren(bool bIsAscending, bool bIsRec
 {
 	Children.Sort([bIsAscending](const FPCGEditorGraphDebugObjectItemPtr& InLHS, const FPCGEditorGraphDebugObjectItemPtr& InRHS)
 	{
-		// Support for sorting by loop index.
+		// If both items have an explicit sort priority like a loop index, this is the primary sort key.
 		const int32 IndexLHS = InLHS->GetSortPriority();
 		const int32 IndexRHS = InRHS->GetSortPriority();
+		if (IndexLHS != INDEX_NONE && IndexRHS != INDEX_NONE)
+		{
+			return (IndexLHS < IndexRHS) == bIsAscending;
+		}
 
+		// Next sort priority is presence or not of children. Items without children are shown first to reduce the possibility that
+		// a child item ends up displayed far away from its parent item when the tree is expanded.
 		const int HasChildrenLHS = InLHS->Children.Num() ? 1 : 0;
 		const int HasChildrenRHS = InRHS->Children.Num() ? 1 : 0;
-
 		if (HasChildrenLHS != HasChildrenRHS)
 		{
-			return HasChildrenLHS < HasChildrenRHS;
-		}
-		else if (IndexLHS == INDEX_NONE || IndexRHS == INDEX_NONE)
-		{
-			return (InLHS->GetLabel() < InRHS->GetLabel()) == bIsAscending;
-		}
-		else
-		{
-			return (IndexLHS < IndexRHS) == bIsAscending;
+			return (HasChildrenLHS < HasChildrenRHS) == bIsAscending;
 		}
+
+		// Otherwise fall back to alphanumeric order.
+		return (InLHS->GetLabel() < InRHS->GetLabel()) == bIsAscending;
 	});
 
 	if (bIsRecursive)
