[PCG] Allow sampling density falloff for SplineSampler Interior even when only 2 or 3 control points on the spline

#rnx

diff --git a/Engine/Plugins/Experimental/PCG/Source/PCG/Private/Elements/PCGSplineSampler.cpp b/Engine/Plugins/Experimental/PCG/Source/PCG/Private/Elements/PCGSplineSampler.cpp
--- a/Engine/Plugins/Experimental/PCG/Source/PCG/Private/Elements/PCGSplineSampler.cpp
+++ b/Engine/Plugins/Experimental/PCG/Source/PCG/Private/Elements/PCGSplineSampler.cpp
@@ -874,9 +874,11 @@ namespace PCGSplineSampler
 		const FVector::FReal MaxDimension = FMath::Max(Spline->Bounds.BoxExtent.X, Spline->Bounds.BoxExtent.Y) * 2.f;
 		const FVector::FReal MaxDimensionSquared = MaxDimension * MaxDimension;
 
+		const int NumSegments = LineData->GetNumSegments();
+
 		const FRichCurve* DensityFalloffCurve = Params.InteriorDensityFalloffCurve.GetRichCurveConst();
-		const bool bGenerateMedialAxis = DensityFalloffCurve != nullptr && DensityFalloffCurve->GetNumKeys() > 0;
-		const bool bFindNearestSplineKey = Params.InteriorOrientation == EPCGSplineSamplingInteriorOrientation::FollowCurvature || (bGenerateMedialAxis && !Params.bTreatSplineAsPolyline);
+		const bool bComputeDensityFalloff = DensityFalloffCurve != nullptr && DensityFalloffCurve->GetNumKeys() > 0 && NumSegments > 1;
+		const bool bFindNearestSplineKey = Params.InteriorOrientation == EPCGSplineSamplingInteriorOrientation::FollowCurvature || (bComputeDensityFalloff && !Params.bTreatSplineAsPolyline);
 		const bool bProjectOntoSurface = Params.bProjectOntoSurface || bFindNearestSplineKey;
 
 		const FVector::FReal BoundExtents = Params.InteriorBorderSampleSpacing * 0.5f;
@@ -890,7 +892,6 @@ namespace PCGSplineSampler
 		if (Params.bTreatSplineAsPolyline)
 		{
 			// Treat spline interface points as vertices of a polyline
-			const int NumSegments = LineData->GetNumSegments();
 			SplineSamplePoints.Reserve(NumSegments);
 
 			for (int32 SegmentIndex = 0; SegmentIndex < NumSegments; ++SegmentIndex)
@@ -919,79 +920,94 @@ namespace PCGSplineSampler
 		}
 
 		TArray<TTuple<FVector2D, FVector2D>> MedialAxisEdges;
+		FVector2D Centroid;
 
-		// Compute the Medial Axis as a subset of the Voronoi Diagram of the PolyLine points
-		if (bGenerateMedialAxis)
+		if (bComputeDensityFalloff)
 		{
-			TArray<FVector> PolygonPoints; // Top-down 2D projection polygon of the spline points
-			if (Params.bTreatSplineAsPolyline)
+			// Compute the Medial Axis as a subset of the Voronoi Diagram of the PolyLine points. Only works for >= 4 points.
+			if (NumSegments >= 4)
 			{
-				// If we already computed the polygon, use a copy instead of generating it again
-				PolygonPoints.Reserve(SplineSamplePoints.Num());
-				for (const FVector& SplinePoint : SplineSamplePoints)
+				TArray<FVector> PolygonPoints; // Top-down 2D projection polygon of the spline points
+				if (Params.bTreatSplineAsPolyline)
 				{
-					FVector& PolygonPoint = PolygonPoints.Add_GetRef(SplinePoint);
-					PolygonPoint.Z = MinPoint.Z;
+					// If we already computed the polygon, use a copy instead of generating it again
+					PolygonPoints.Reserve(SplineSamplePoints.Num());
+					for (const FVector& SplinePoint : SplineSamplePoints)
+					{
+						FVector& PolygonPoint = PolygonPoints.Add_GetRef(SplinePoint);
+						PolygonPoint.Z = MinPoint.Z;
+					}
 				}
-			}
-			else
-			{
-				const int NumSegments = LineData->GetNumSegments();
-				PolygonPoints.Reserve(NumSegments);
-
-				for (int32 SegmentIndex = 0; SegmentIndex < NumSegments; ++SegmentIndex)
+				else
 				{
-					FVector& PolygonPoint = PolygonPoints.Add_GetRef(LineData->GetLocationAtDistance(SegmentIndex, 0, /*bWorldSpace=*/false));
-					PolygonPoint.Z = MinPoint.Z;
+					PolygonPoints.Reserve(NumSegments);
+
+					for (int32 SegmentIndex = 0; SegmentIndex < NumSegments; ++SegmentIndex)
+					{
+						FVector& PolygonPoint = PolygonPoints.Add_GetRef(LineData->GetLocationAtDistance(SegmentIndex, 0, /*bWorldSpace=*/false));
+						PolygonPoint.Z = MinPoint.Z;
+					}
 				}
-			}
 
-			TArray<FVector2D> PolygonPoints2D;
-			PolygonPoints2D.Reserve(PolygonPoints.Num());
+				TArray<FVector2D> PolygonPoints2D;
+				PolygonPoints2D.Reserve(PolygonPoints.Num());
 
-			for (const FVector& Point : PolygonPoints)
-			{
-				PolygonPoints2D.Add(FVector2D(Point));
-			}
+				for (const FVector& Point : PolygonPoints)
+				{
+					PolygonPoints2D.Add(FVector2D(Point));
+				}
 
-			TArray<TTuple<FVector, FVector>> VoronoiEdges;
-			TArray<int32> CellMember;
+				TArray<TTuple<FVector, FVector>> VoronoiEdges;
+				TArray<int32> CellMember;
 
-			{
-				TRACE_CPUPROFILER_EVENT_SCOPE(FPCGSplineSamplerElement::Execute::GetVoronoiEdges);
-				GetVoronoiEdges(PolygonPoints, FBox(MinPoint, FVector(MaxPoint.X, MaxPoint.Y, MinPoint.Z)), VoronoiEdges, CellMember);
-			}
+				{
+					TRACE_CPUPROFILER_EVENT_SCOPE(FPCGSplineSamplerElement::Execute::GetVoronoiEdges);
+					GetVoronoiEdges(PolygonPoints, FBox(MinPoint, FVector(MaxPoint.X, MaxPoint.Y, MinPoint.Z)), VoronoiEdges, CellMember);
+				}
 
-			// Find the subset of the Voronoi Diagram which composes the Medial Axis
-			for (const TTuple<FVector, FVector>& Edge : VoronoiEdges)
-			{
-				// Discard any edges which intersect the polygon
-				bool bDiscard = false;
-				for (int32 PointIndex = 0; PointIndex < PolygonPoints.Num(); ++PointIndex)
+				// Find the subset of the Voronoi Diagram which composes the Medial Axis
+				for (const TTuple<FVector, FVector>& Edge : VoronoiEdges)
 				{
-					FVector2D IntersectionPoint;
-					if (PCGSplineSamplerHelpers::SegmentIntersection2D(PolygonPoints2D[PointIndex], PolygonPoints2D[(PointIndex + 1) % PolygonPoints2D.Num()], FVector2D(Edge.Get<0>()), FVector2D(Edge.Get<1>()), IntersectionPoint))
+					// Discard any edges which intersect the polygon
+					bool bDiscard = false;
+					for (int32 PointIndex = 0; PointIndex < PolygonPoints.Num(); ++PointIndex)
 					{
-						bDiscard = true;
-						break;
+						FVector2D IntersectionPoint;
+						if (PCGSplineSamplerHelpers::SegmentIntersection2D(PolygonPoints2D[PointIndex], PolygonPoints2D[(PointIndex + 1) % PolygonPoints2D.Num()], FVector2D(Edge.Get<0>()), FVector2D(Edge.Get<1>()), IntersectionPoint))
+						{
+							bDiscard = true;
+							break;
+						}
 					}
-				}
 
-				if (bDiscard)
-				{
-					continue;
-				}
+					if (bDiscard)
+					{
+						continue;
+					}
 
-				// If either of the points lies within the polygon, the segment must lie within the polygon
-				if (PCGSplineSamplerHelpers::PointInsidePolygon2D(PolygonPoints2D, FVector2D(Edge.Get<0>()), MaxDimension))
-				{
-					MedialAxisEdges.Add(Edge);
+					// If either of the points lies within the polygon, the segment must lie within the polygon
+					if (PCGSplineSamplerHelpers::PointInsidePolygon2D(PolygonPoints2D, FVector2D(Edge.Get<0>()), MaxDimension))
+					{
+						MedialAxisEdges.Add(Edge);
+					}
 				}
 			}
-
-			if (MedialAxisEdges.IsEmpty())
+			else if (NumSegments == 3) // If we only have 3 points, use the centroid instead of the medial axis.
+			{
+				const FVector2D ControlPoint1 = FVector2D(LineData->GetLocationAtDistance(/*SegmentIndex=*/0, /*Distance=*/0.0f, /*bWorldSpace=*/false));
+				const FVector2D ControlPoint2 = FVector2D(LineData->GetLocationAtDistance(/*SegmentIndex=*/1, /*Distance=*/0.0f, /*bWorldSpace=*/false));
+				const FVector2D ControlPoint3 = FVector2D(LineData->GetLocationAtDistance(/*SegmentIndex=*/2, /*Distance=*/0.0f, /*bWorldSpace=*/false));
+				Centroid = (ControlPoint1 + ControlPoint2 + ControlPoint3) / 3.0f;
+			}
+			else if (NumSegments == 2) // If we only have 2 points, compute a third point, and use the centroid instead of the medial axis.
 			{
-				PCGE_LOG_C(Warning, GraphAndLog, Context, LOCTEXT("MedialAxisFailed", "Failed to compute medial axis in interior region, density fall-off will not be applied. This functionality requires a closed spline with at least 4 spline points."));
+				// Note: The midpoint of the first segment may not be collinear if the control point has meaningful tangents.
+				// Note: The centroid of these three points may not always lie inside the spline. This can happen if the spline is self-intersecting.
+				// TODO: It may be worth investigating super-sampling the spline to get >= 4 points and then compute the medial axis instead.
+				const FVector2D ControlPoint1 = FVector2D(LineData->GetLocationAtDistance(/*SegmentIndex=*/0, /*Distance=*/0.0f, /*bWorldSpace=*/false));
+				const FVector2D ControlPoint2 = FVector2D(LineData->GetLocationAtDistance(/*SegmentIndex=*/0, /*Distance=*/0.5f, /*bWorldSpace=*/false));
+				const FVector2D ControlPoint3 = FVector2D(LineData->GetLocationAtDistance(/*SegmentIndex=*/1, /*Distance=*/0.0f, /*bWorldSpace=*/false));
+				Centroid = (ControlPoint1 + ControlPoint2 + ControlPoint3) / 3.0f;
 			}
 		}
 
@@ -1114,20 +1130,27 @@ namespace PCGSplineSampler
 						// Calculate density fall off
 						float Density = 1.0f;
 
-						if (bGenerateMedialAxis && MedialAxisEdges.Num())
+						if (bComputeDensityFalloff)
 						{
 							FVector::FReal SmallestDistSquared = MaxDimensionSquared;
 
-							// Find distance from SampleLocation to MedialAxis
-							for (const TTuple<FVector2D, FVector2D>& Edge : MedialAxisEdges)
+							if (MedialAxisEdges.Num() > 0)
 							{
-								const FVector::FReal DistSquared = FMath::PointDistToSegmentSquared(FVector(SampleLocation, 0), FVector(Edge.Get<0>(), 0), FVector(Edge.Get<1>(), 0));
-
-								if (DistSquared < SmallestDistSquared)
+								// Find distance from SampleLocation to MedialAxis
+								for (const TTuple<FVector2D, FVector2D>& Edge : MedialAxisEdges)
 								{
-									SmallestDistSquared = DistSquared;
+									const FVector::FReal DistSquared = FMath::PointDistToSegmentSquared(FVector(SampleLocation, 0), FVector(Edge.Get<0>(), 0), FVector(Edge.Get<1>(), 0));
+
+									if (DistSquared < SmallestDistSquared)
+									{
+										SmallestDistSquared = DistSquared;
+									}
 								}
 							}
+							else if (NumSegments == 2 || NumSegments == 3) // If a centroid was computed instead of the medial axis, fallback to that.
+							{
+								SmallestDistSquared = FVector2D::DistSquared(SampleLocation, Centroid);
+							}
 
 							const FVector::FReal SmallestDist = FMath::Sqrt(SmallestDistSquared);
 
