Improved cancellation behavior, especially when we have tasks joining multiple graphs together.
Added 'abort' capabilities on PCGElements which allow special behavior when a task is not fully completed before being cancelled.
Soft-removed the CancelExecutionOnEmpty property as it was causing hard to track issues for users - property is still there and will be removed but it does not have an effect now.
Light refactoring around the graph cache, graph executor and subsystem to remove useless parameters.

diff --git a/Engine/Plugins/Experimental/PCG/Source/PCG/Private/Elements/PCGDataFromActor.cpp b/Engine/Plugins/Experimental/PCG/Source/PCG/Private/Elements/PCGDataFromActor.cpp
--- a/Engine/Plugins/Experimental/PCG/Source/PCG/Private/Elements/PCGDataFromActor.cpp
+++ b/Engine/Plugins/Experimental/PCG/Source/PCG/Private/Elements/PCGDataFromActor.cpp
@@ -335,12 +335,21 @@ bool FPCGDataFromActorElement::ExecuteInternal(FPCGContext* InContext) const
 					// Add a trivial task after these generations that wakes up this task
 					Context->bIsPaused = true;
 
-					Subsystem->ScheduleGeneric([Context]()
-					{
-						// Wake up the current task
-						Context->bIsPaused = false;
-						return true;
-					}, Context->SourceComponent.Get(), WaitOnTaskIds);
+					Subsystem->ScheduleGeneric(
+						[Context]() // Normal execution: Wake up the current task
+						{
+							Context->bIsPaused = false; 
+							return true;
+						}, 
+						[Context]() // On Abort: Wake up on abort, clear all results and mark as cancelled
+						{
+							Context->bIsPaused = false; 
+							Context->FoundActors.Reset();
+							Context->OutputData.bCancelExecution = true;
+							return true;
+						},
+						Context->SourceComponent.Get(),
+						WaitOnTaskIds);
 
 					return false;
 				}
diff --git a/Engine/Plugins/Experimental/PCG/Source/PCG/Private/Elements/PCGLoopElement.cpp b/Engine/Plugins/Experimental/PCG/Source/PCG/Private/Elements/PCGLoopElement.cpp
--- a/Engine/Plugins/Experimental/PCG/Source/PCG/Private/Elements/PCGLoopElement.cpp
+++ b/Engine/Plugins/Experimental/PCG/Source/PCG/Private/Elements/PCGLoopElement.cpp
@@ -290,12 +290,21 @@ bool FPCGLoopElement::ExecuteInternal(FPCGContext* InContext) const
 			Context->bScheduledSubgraph = true;
 			Context->bIsPaused = true;
 
-			Subsystem->ScheduleGeneric([Context]()
-			{
-				// Wake up the current task
-				Context->bIsPaused = false;
-				return true;
-			}, Context->SourceComponent.Get(), Context->SubgraphTaskIds);
+			Subsystem->ScheduleGeneric(
+				[Context]() // Normal execution: Wake up the current task
+				{
+					Context->bIsPaused = false;
+					return true;
+				},
+				[Context]() // On abort: wakeup and cancel, forget subgraphs
+				{
+					Context->bIsPaused = false;
+					Context->SubgraphTaskIds.Reset();
+					Context->OutputData.bCancelExecution = true;
+					return true;
+				},
+				Context->SourceComponent.Get(),
+				Context->SubgraphTaskIds);
 
 			return false;
 		}
diff --git a/Engine/Plugins/Experimental/PCG/Source/PCG/Private/Elements/PCGSpawnActor.cpp b/Engine/Plugins/Experimental/PCG/Source/PCG/Private/Elements/PCGSpawnActor.cpp
--- a/Engine/Plugins/Experimental/PCG/Source/PCG/Private/Elements/PCGSpawnActor.cpp
+++ b/Engine/Plugins/Experimental/PCG/Source/PCG/Private/Elements/PCGSpawnActor.cpp
@@ -592,11 +592,19 @@ bool FPCGSpawnActorElement::SpawnAndPrepareSubgraphs(FPCGSubgraphContext* Contex
 		Context->bScheduledSubgraph = true;
 		Context->bIsPaused = true;
 
-		Subsystem->ScheduleGeneric([Context]() {
-			// Wake up the current task
-			Context->bIsPaused = false;
-			return true;
-		}, Context->SourceComponent.Get(), Context->SubgraphTaskIds);
+		Subsystem->ScheduleGeneric(
+			[Context]() // Normal execution: Wake up the current task
+			{
+				Context->bIsPaused = false;
+				return true;
+			}, 
+			[Context]() // On Abort: Wake up & cancel
+			{
+				Context->bIsPaused = false;
+				Context->OutputData.bCancelExecution = true;
+			},
+			Context->SourceComponent.Get(), 
+			Context->SubgraphTaskIds);
 
 		return false;
 	}
diff --git a/Engine/Plugins/Experimental/PCG/Source/PCG/Private/Graph/PCGGraphCache.cpp b/Engine/Plugins/Experimental/PCG/Source/PCG/Private/Graph/PCGGraphCache.cpp
--- a/Engine/Plugins/Experimental/PCG/Source/PCG/Private/Graph/PCGGraphCache.cpp
+++ b/Engine/Plugins/Experimental/PCG/Source/PCG/Private/Graph/PCGGraphCache.cpp
@@ -27,11 +27,9 @@ static TAutoConsoleVariable<bool> CVarCacheMemoryBudgetEnabled(
 // Initial max number of entries graph cache
 static const int32 GPCGGraphCacheInitialCapacity = 65536;
 
-FPCGGraphCache::FPCGGraphCache(TWeakObjectPtr<UObject> InOwner)
+FPCGGraphCache::FPCGGraphCache()
 	: CacheData(GPCGGraphCacheInitialCapacity)
-	, Owner(InOwner)
 {
-	check(InOwner.Get());
 }
 
 FPCGGraphCache::~FPCGGraphCache()
@@ -39,13 +37,8 @@ FPCGGraphCache::~FPCGGraphCache()
 	ClearCache();
 }
 
-bool FPCGGraphCache::GetFromCache(const UPCGNode* InNode, const IPCGElement* InElement, const FPCGCrc& InDependenciesCrc, const FPCGDataCollection& InInput, const UPCGSettings* InSettings, const UPCGComponent* InComponent, FPCGDataCollection& OutOutput) const
+bool FPCGGraphCache::GetFromCache(const UPCGNode* InNode, const IPCGElement* InElement, const FPCGCrc& InDependenciesCrc, const UPCGComponent* InComponent, FPCGDataCollection& OutOutput) const
 {
-	if (!Owner.IsValid())
-	{
-		return false;
-	}
-
 	if(!InDependenciesCrc.IsValid())
 	{
 		UE_LOG(LogPCG, Warning, TEXT("Invalid dependencies passed to FPCGGraphCache::GetFromCache(), lookup aborted."));
@@ -83,9 +76,9 @@ bool FPCGGraphCache::GetFromCache(const UPCGNode* InNode, const IPCGElement* InE
 	}
 }
 
-void FPCGGraphCache::StoreInCache(const IPCGElement* InElement, const FPCGCrc& InDependenciesCrc, const FPCGDataCollection& InInput, const UPCGSettings* InSettings, const UPCGComponent* InComponent, const FPCGDataCollection& InOutput)
+void FPCGGraphCache::StoreInCache(const IPCGElement* InElement, const FPCGCrc& InDependenciesCrc, const FPCGDataCollection& InOutput)
 {
-	if (!Owner.IsValid() || !ensure(InDependenciesCrc.IsValid()))
+	if (!ensure(InDependenciesCrc.IsValid()))
 	{
 		return;
 	}
diff --git a/Engine/Plugins/Experimental/PCG/Source/PCG/Private/Graph/PCGGraphCache.h b/Engine/Plugins/Experimental/PCG/Source/PCG/Private/Graph/PCGGraphCache.h
--- a/Engine/Plugins/Experimental/PCG/Source/PCG/Private/Graph/PCGGraphCache.h
+++ b/Engine/Plugins/Experimental/PCG/Source/PCG/Private/Graph/PCGGraphCache.h
@@ -63,14 +63,14 @@ private:
 */
 struct FPCGGraphCache : public FGCObject
 {
-	explicit FPCGGraphCache(TWeakObjectPtr<UObject> InOwner);
+	FPCGGraphCache();
 	~FPCGGraphCache();
 
 	/** Returns true if data was found from the cache, in which case the outputs are written in OutOutput. InNode is optional and for logging only. */
-	bool GetFromCache(const UPCGNode* InNode, const IPCGElement* InElement, const FPCGCrc& InCrc, const FPCGDataCollection& InInput, const UPCGSettings* InSettings, const UPCGComponent* InComponent, FPCGDataCollection& OutOutput) const;
+	bool GetFromCache(const UPCGNode* InNode, const IPCGElement* InElement, const FPCGCrc& InCrc, const UPCGComponent* InComponent, FPCGDataCollection& OutOutput) const;
 
 	/** Stores data in the cache for later use */
-	void StoreInCache(const IPCGElement* InElement, const FPCGCrc& InCrc, const FPCGDataCollection& InInput, const UPCGSettings* InSettings, const UPCGComponent* InComponent, const FPCGDataCollection& InOutput);
+	void StoreInCache(const IPCGElement* InElement, const FPCGCrc& InCrc, const FPCGDataCollection& InOutput);
 
 	/** Removes all entries from the cache, unroots data, etc. */
 	void ClearCache();
@@ -107,7 +107,6 @@ private:
 	// Note: we are not going to serialize this as-is, since the pointers will change
 	// We will have to serialize on a node id basis most likely
 	TLruCache<FPCGCacheEntryKey, FPCGDataCollection> CacheData;
-	TWeakObjectPtr<UObject> Owner = nullptr;
 
 	/** Map from data UIDs to records. Provides ref counting and caches memory size. */
 	TMap<uint64, FCachedMemoryRecord> MemoryRecords;
diff --git a/Engine/Plugins/Experimental/PCG/Source/PCG/Private/Graph/PCGGraphExecutor.cpp b/Engine/Plugins/Experimental/PCG/Source/PCG/Private/Graph/PCGGraphExecutor.cpp
--- a/Engine/Plugins/Experimental/PCG/Source/PCG/Private/Graph/PCGGraphExecutor.cpp
+++ b/Engine/Plugins/Experimental/PCG/Source/PCG/Private/Graph/PCGGraphExecutor.cpp
@@ -4,7 +4,6 @@
 
 #include "PCGCommon.h"
 #include "PCGComponent.h"
-#include "PCGContext.h"
 #include "PCGCrc.h"
 #include "PCGGraph.h"
 #include "PCGInputOutputSettings.h"
@@ -13,9 +12,6 @@
 #include "PCGPin.h"
 #include "PCGSubsystem.h"
 #include "PCGWorldActor.h"
-#include "Graph/PCGGraphCache.h"
-#include "Graph/PCGGraphCompiler.h"
-#include "Graph/PCGStackContext.h"
 #include "Grid/PCGPartitionActor.h"
 #include "Helpers/PCGActorHelpers.h"
 #include "Helpers/PCGHelpers.h"
@@ -47,31 +43,34 @@ static TAutoConsoleVariable<float> CVarMaxPercentageOfThreadsToUse(
 	0.9f,
 	TEXT("Maximum percentage of number of threads for concurrent PCG processing"));
 
-static TAutoConsoleVariable<float> CVarTimePerFrame(
-	TEXT("pcg.FrameTime"),
-	1000.0f / 60.0f,
-	TEXT("Allocated time in ms per frame"));
-
-static TAutoConsoleVariable<float> CVarEditorTimePerFrame(
-	TEXT("pcg.EditorFrameTime"),
-	1000.0f / 20.0f,
-	TEXT("Allocated time in ms per frame when running in editor (non pie)"));
-
-static TAutoConsoleVariable<bool> CVarGraphMultithreading(
-	TEXT("pcg.GraphMultithreading"),
-	false,
-	TEXT("Controls whether the graph can dispatch multiple tasks at the same time"));
-
 static TAutoConsoleVariable<bool> CVarStripEmptyPointData(
 	TEXT("pcg.StripEmptyPointData"),
 	false,
 	TEXT("Will strip empty point data from being passed along as output"));
 
-FPCGGraphExecutor::FPCGGraphExecutor(UObject* InOwner)
-	: GraphCompiler(MakeUnique<FPCGGraphCompiler>())
-	, GraphCache(InOwner)
+namespace PCGGraphExecutor
+{
+	TAutoConsoleVariable<float> CVarTimePerFrame(
+		TEXT("pcg.FrameTime"),
+		1000.0f / 60.0f,
+		TEXT("Allocated time in ms per frame"));
+
+	TAutoConsoleVariable<bool> CVarGraphMultithreading(
+		TEXT("pcg.GraphMultithreading"),
+		false,
+		TEXT("Controls whether the graph can dispatch multiple tasks at the same time"));
+
 #if WITH_EDITOR
-	, GenerationProgressNotification(GetNotificationTextFormat())
+	TAutoConsoleVariable<float> CVarEditorTimePerFrame(
+		TEXT("pcg.EditorFrameTime"),
+		1000.0f / 20.0f,
+		TEXT("Allocated time in ms per frame when running in editor (non pie)"));
+#endif
+}
+
+FPCGGraphExecutor::FPCGGraphExecutor()
+#if WITH_EDITOR
+	: GenerationProgressNotification(GetNotificationTextFormat())
 #endif
 {
 }
@@ -91,7 +90,7 @@ FPCGGraphExecutor::~FPCGGraphExecutor()
 
 void FPCGGraphExecutor::Compile(UPCGGraph* Graph)
 {
-	GraphCompiler->Compile(Graph);
+	GraphCompiler.Compile(Graph);
 }
 
 FPCGTaskId FPCGGraphExecutor::Schedule(UPCGComponent* Component, const TArray<FPCGTaskId>& ExternalDependencies, const FPCGStack* InFromStack)
@@ -130,7 +129,7 @@ FPCGTaskId FPCGGraphExecutor::Schedule(
 
 	// Get compiled tasks from compiler
 	TSharedPtr<FPCGStackContext> StackContextPtr = MakeShared<FPCGStackContext>();
-	TArray<FPCGGraphTask> CompiledTasks = GraphCompiler->GetCompiledTasks(Graph, GenerationGridSize, *StackContextPtr);
+	TArray<FPCGGraphTask> CompiledTasks = GraphCompiler.GetCompiledTasks(Graph, GenerationGridSize, *StackContextPtr);
 
 	// Create the final stack context by including the current stack frames
 	if (InFromStack)
@@ -286,6 +285,20 @@ TSet<UPCGComponent*> FPCGGraphExecutor::Cancel(TFunctionRef<bool(TWeakObjectPtr<
 		return CancelledComponents;
 	}
 
+	auto TryAbortScheduledTasks = [](FPCGGraphScheduleTask& ScheduledTask)
+	{
+		if (ScheduledTask.bHasAbortCallbacks)
+		{
+			for (FPCGGraphTask& InternalTask : ScheduledTask.Tasks)
+			{
+				if (InternalTask.Element)
+				{
+					InternalTask.Element->Abort(InternalTask.Context);
+				}
+			}
+		}
+	};
+
 	TArray<FPCGTaskId> CancelledScheduledTasks;
 
 	bool bStableCancellationSet = false;
@@ -301,6 +314,8 @@ TSet<UPCGComponent*> FPCGGraphExecutor::Cancel(TFunctionRef<bool(TWeakObjectPtr<
 			if (CancelledComponents.Contains(ScheduledTask.SourceComponent.Get()))
 			{
 				CancelledScheduledTasks.Add(ScheduledTask.Tasks[ScheduledTask.LastTaskIndex].NodeId);
+
+				TryAbortScheduledTasks(ScheduledTask);
 				ScheduledTasks.RemoveAtSwap(ScheduledTaskIndex);
 			}
 		}
@@ -324,6 +339,8 @@ TSet<UPCGComponent*> FPCGGraphExecutor::Cancel(TFunctionRef<bool(TWeakObjectPtr<
 					}
 
 					CancelledScheduledTasks.Add(ScheduledTask.Tasks[ScheduledTask.LastTaskIndex].NodeId);
+
+					TryAbortScheduledTasks(ScheduledTask);
 					ScheduledTasks.RemoveAtSwap(ScheduledTaskIndex);
 				}
 			}
@@ -339,6 +356,11 @@ TSet<UPCGComponent*> FPCGGraphExecutor::Cancel(TFunctionRef<bool(TWeakObjectPtr<
 
 			if (CancelledComponents.Contains(Task.SourceComponent.Get()))
 			{
+				if (Task.Element)
+				{
+					Task.Element->Abort(Task.Context);
+				}
+
 				FPCGTaskId CancelledTaskId = Task.NodeId;
 				RemoveTaskFromInputSuccessors(CancelledTaskId, Task.Inputs);
 
@@ -354,6 +376,9 @@ TSet<UPCGComponent*> FPCGGraphExecutor::Cancel(TFunctionRef<bool(TWeakObjectPtr<
 			FPCGGraphActiveTask& Task = ActiveTasks[ActiveTaskIndex];
 			if (Task.Context && CancelledComponents.Contains(Task.Context->SourceComponent.Get()))
 			{
+				check(Task.Element);
+				Task.Element->Abort(Task.Context.Get());
+
 				FPCGTaskId CancelledTaskId = Task.NodeId;
 				Task.bWasCancelled = true;
 				RemoveTaskFromInputSuccessors(CancelledTaskId, Task.Inputs);
@@ -368,6 +393,9 @@ TSet<UPCGComponent*> FPCGGraphExecutor::Cancel(TFunctionRef<bool(TWeakObjectPtr<
 			FPCGGraphActiveTask& Task = SleepingTasks[SleepingTaskIndex];
 			if (Task.Context && CancelledComponents.Contains(Task.Context->SourceComponent.Get()))
 			{
+				check(Task.Element);
+				Task.Element->Abort(Task.Context.Get());
+
 				FPCGTaskId CancelledTaskId = Task.NodeId;
 				RemoveTaskFromInputSuccessors(CancelledTaskId, Task.Inputs);
 
@@ -402,15 +430,45 @@ bool FPCGGraphExecutor::IsGraphCurrentlyExecuting(UPCGGraph* InGraph)
 }
 
 FPCGTaskId FPCGGraphExecutor::ScheduleGeneric(TFunction<bool()> InOperation, UPCGComponent* InSourceComponent, const TArray<FPCGTaskId>& TaskExecutionDependencies)
+{
+	return ScheduleGeneric(
+		InOperation,
+		TFunction<void()>(),
+		InSourceComponent,
+		TaskExecutionDependencies);
+}
+
+FPCGTaskId FPCGGraphExecutor::ScheduleGeneric(TFunction<bool()> InOperation, TFunction<void()> InAbortOperation, UPCGComponent* InSourceComponent, const TArray<FPCGTaskId>& TaskExecutionDependencies)
 {
 	// Since we have no context, the generic task will consume no input (no data dependencies).
-	return ScheduleGenericWithContext([Operation = MoveTemp(InOperation)](FPCGContext*) -> bool
-	{
-			return Operation();
-	}, InSourceComponent, TaskExecutionDependencies, /*TaskDataDependencies=*/{});
+	return ScheduleGenericWithContext(
+		[Operation = MoveTemp(InOperation)](FPCGContext*) -> bool
+		{
+			return Operation && Operation();
+		}, 
+		[AbortOperation = MoveTemp(InAbortOperation)](FPCGContext*)
+		{
+			if(AbortOperation)
+			{
+				AbortOperation();
+			}
+		},
+		InSourceComponent,
+		TaskExecutionDependencies,
+		/*TaskDataDependencies=*/{});
 }
 
-FPCGTaskId FPCGGraphExecutor::ScheduleGenericWithContext(TFunction<bool(FPCGContext*)> InOperation, UPCGComponent* InSourceComponent, const TArray<FPCGTaskId>& TaskDataDependencies, const TArray<FPCGTaskId>& TaskExecutionDependencies)
+FPCGTaskId FPCGGraphExecutor::ScheduleGenericWithContext(TFunction<bool(FPCGContext*)> InOperation, UPCGComponent* InSourceComponent, const TArray<FPCGTaskId>& TaskExecutionDependencies, const TArray<FPCGTaskId>& TaskDataDependencies)
+{
+	return ScheduleGenericWithContext(
+		InOperation,
+		TFunction<void(FPCGContext*)>(),
+		InSourceComponent,
+		TaskExecutionDependencies,
+		TaskDataDependencies);
+}
+
+FPCGTaskId FPCGGraphExecutor::ScheduleGenericWithContext(TFunction<bool(FPCGContext*)> InOperation, TFunction<void(FPCGContext*)> InAbortOperation, UPCGComponent* InSourceComponent, const TArray<FPCGTaskId>& TaskExecutionDependencies, const TArray<FPCGTaskId>& TaskDataDependencies)
 {
 	// Build task & element to hold the operation to perform
 	FPCGGraphTask Task;
@@ -428,7 +486,7 @@ FPCGTaskId FPCGGraphExecutor::ScheduleGenericWithContext(TFunction<bool(FPCGCont
 	}
 
 	Task.SourceComponent = InSourceComponent;
-	Task.Element = MakeShared<FPCGGenericElement>(InOperation);
+	Task.Element = MakeShared<FPCGGenericElement>(InOperation, InAbortOperation);
 
 	ScheduleLock.Lock();
 
@@ -438,6 +496,7 @@ FPCGTaskId FPCGGraphExecutor::ScheduleGenericWithContext(TFunction<bool(FPCGCont
 	FPCGGraphScheduleTask& ScheduledTask = ScheduledTasks.Emplace_GetRef();
 	ScheduledTask.Tasks.Add(Task);
 	ScheduledTask.SourceComponent = InSourceComponent;
+	ScheduledTask.bHasAbortCallbacks = !!InAbortOperation;
 
 	ScheduleLock.Unlock();
 
@@ -517,19 +576,19 @@ void FPCGGraphExecutor::Execute()
 
 	const double StartTime = FPlatformTime::Seconds();
 
-	double VarTimePerFrame = CVarTimePerFrame.GetValueOnAnyThread() / 1000.0;
+	double VarTimePerFrame = PCGGraphExecutor::CVarTimePerFrame.GetValueOnAnyThread() / 1000.0;
 
 #if WITH_EDITOR
 	if (GEditor && !GEditor->IsPlaySessionInProgress())
 	{
-		VarTimePerFrame = CVarEditorTimePerFrame.GetValueOnAnyThread() / 1000.0;
+		VarTimePerFrame = PCGGraphExecutor::CVarEditorTimePerFrame.GetValueOnAnyThread() / 1000.0;
 	}
 #endif
 
 	const double EndTime = StartTime + VarTimePerFrame;
 	const float MaxPercentageOfThreadsToUse = FMath::Clamp(CVarMaxPercentageOfThreadsToUse.GetValueOnAnyThread(), 0.0f, 1.0f);
 	const int32 MaxNumThreads = FMath::Max(0, FMath::Min((int32)(FPlatformMisc::NumberOfCoresIncludingHyperthreads() * MaxPercentageOfThreadsToUse), CVarMaxNumTasks.GetValueOnAnyThread() - 1));
-	const bool bAllowMultiDispatch = CVarGraphMultithreading.GetValueOnAnyThread();
+	const bool bAllowMultiDispatch = PCGGraphExecutor::CVarGraphMultithreading.GetValueOnAnyThread();
 	const bool bGraphCacheDebuggingEnabled = IsGraphCacheDebuggingEnabled();
 	const bool bStripEmptyPointData = CVarStripEmptyPointData.GetValueOnAnyThread();
 
@@ -663,7 +722,7 @@ void FPCGGraphExecutor::Execute()
 				}
 
 				FPCGDataCollection CachedOutput;
-				const bool bResultAlreadyInCache = bCacheable && DependenciesCrc.IsValid() && GraphCache.GetFromCache(Task.Node, Task.Element.Get(), DependenciesCrc, TaskInput, TaskSettings, Task.SourceComponent.Get(), CachedOutput);
+				const bool bResultAlreadyInCache = bCacheable && DependenciesCrc.IsValid() && GraphCache.GetFromCache(Task.Node, Task.Element.Get(), DependenciesCrc, Task.SourceComponent.Get(), CachedOutput);
 #if WITH_EDITOR
 				const bool bNeedsToCreateActiveTask = !bResultAlreadyInCache || TaskSettingsInterface->bDebug;
 #else
@@ -819,7 +878,7 @@ void FPCGGraphExecutor::Execute()
 				if (ActiveTaskSettingsInterface && !bHasErrorOrWarning && ActiveTask.Element->IsCacheableInstance(ActiveTaskSettingsInterface))
 				{
 					const UPCGSettings* ActiveTaskSettings = ActiveTaskSettingsInterface ? ActiveTaskSettingsInterface->GetSettings() : nullptr;
-					GraphCache.StoreInCache(ActiveTask.Element.Get(), ActiveTask.Context->DependenciesCrc, ActiveTask.Context->InputData, ActiveTaskSettings, ActiveTask.Context->SourceComponent.Get(), ActiveTask.Context->OutputData);
+					GraphCache.StoreInCache(ActiveTask.Element.Get(), ActiveTask.Context->DependenciesCrc, ActiveTask.Context->OutputData);
 				}
 			}
 
@@ -1058,6 +1117,11 @@ bool FPCGGraphExecutor::CancelNextTasks(FPCGTaskId CancelledTask, TSet<UPCGCompo
 					bAddedComponents = true;
 				}
 
+				if (Task->Element)
+				{
+					Task->Element->Abort(Task->Context);
+				}
+
 				RemoveTaskFromInputSuccessors(Task->NodeId, Task->Inputs);
 				Tasks.Remove(Successor);
 			}
@@ -1161,11 +1225,6 @@ void FPCGGraphExecutor::BuildTaskInput(const FPCGGraphTask& Task, FPCGDataCollec
 		// Write input pin name (e.g. name of the output pin on the dependency node) Crc to uniquely identify inputs per-pin, or use a placeholder for symmetry.
 		Crc.Combine(Input.InPin ? GetTypeHash(Input.InPin->Properties.Label) : DefaultHashForNoInputPin);
 
-		if (TaskInput.TaggedData.Num() == TaggedDataOffset && InputCollection.bCancelExecutionOnEmpty)
-		{
-			TaskInput.bCancelExecution = true;
-		}
-
 		// Apply labelling on data; technically, we should ensure that we do this only for pass-through nodes,
 		// Otherwise we could also null out the label on the input...
 		if (Input.OutPin)
@@ -1300,7 +1359,7 @@ FPCGTaskId FPCGGraphExecutor::ScheduleDebugWithTaskCallback(UPCGComponent* InCom
 	const uint32 GenerationGridSize = bNonPartitionedComponent ? PCGHiGenGrid::UninitializedGridSize() : InComponent->GetGenerationGridSize();
 
 	FPCGStackContext DummyStackContext;
-	TArray<FPCGGraphTask> CompiledTasks = GraphCompiler->GetCompiledTasks(InComponent->GetGraph(), GenerationGridSize, DummyStackContext, /*bIsTopGraph=*/true);
+	TArray<FPCGGraphTask> CompiledTasks = GraphCompiler.GetCompiledTasks(InComponent->GetGraph(), GenerationGridSize, DummyStackContext, /*bIsTopGraph=*/true);
 	CompiledTasks.Pop(); // Remove the final task
 
 	// Set up all final dependencies for the entire execution
@@ -1385,10 +1444,7 @@ void FPCGGraphExecutor::ReleaseUnusedActors()
 
 void FPCGGraphExecutor::NotifyGraphChanged(UPCGGraph* InGraph)
 {
-	if (GraphCompiler)
-	{
-		GraphCompiler->NotifyGraphChanged(InGraph);
-	}
+	GraphCompiler.NotifyGraphChanged(InGraph);
 }
 
 void FPCGGraphExecutor::UpdateGenerationNotification()
@@ -1506,6 +1562,13 @@ FPCGGenericElement::FPCGGenericElement(TFunction<bool(FPCGContext*)> InOperation
 {
 }
 
+FPCGGenericElement::FPCGGenericElement(TFunction<bool(FPCGContext*)> InOperation, TFunction<void(FPCGContext*)> InAbortOperation, const FContextAllocator& InContextAllocator)
+	: Operation(InOperation)
+	, AbortOperation(InAbortOperation)
+	, ContextAllocator(InContextAllocator)
+{
+}
+
 FPCGContext* FPCGGenericElement::Initialize(const FPCGDataCollection& InputData, TWeakObjectPtr<UPCGComponent> SourceComponent, const UPCGNode* Node)
 {
 	check(ContextAllocator);
@@ -1520,7 +1583,16 @@ FPCGContext* FPCGGenericElement::Initialize(const FPCGDataCollection& InputData,
 bool FPCGGenericElement::ExecuteInternal(FPCGContext* Context) const
 {
 	TRACE_CPUPROFILER_EVENT_SCOPE(FPCGGenericElement::Execute);
-	return Operation(Context);
+	return Operation && Operation(Context);
+}
+
+void FPCGGenericElement::AbortInternal(FPCGContext* Context) const
+{
+	TRACE_CPUPROFILER_EVENT_SCOPE(FPCGGenericElement::Abort);
+	if(AbortOperation)
+	{
+		AbortOperation(Context);
+	}
 }
 
 namespace PCGGraphExecutor
diff --git a/Engine/Plugins/Experimental/PCG/Source/PCG/Private/Graph/PCGGraphExecutor.h b/Engine/Plugins/Experimental/PCG/Source/PCG/Private/Graph/PCGGraphExecutor.h
--- a/Engine/Plugins/Experimental/PCG/Source/PCG/Private/Graph/PCGGraphExecutor.h
+++ b/Engine/Plugins/Experimental/PCG/Source/PCG/Private/Graph/PCGGraphExecutor.h
@@ -6,6 +6,7 @@
 #include "PCGElement.h"
 #include "PCGSubsystem.h"
 #include "Graph/PCGGraphCache.h"
+#include "Graph/PCGGraphCompiler.h"
 #include "Graph/PCGStackContext.h"
 
 #include "UObject/GCObject.h"
@@ -24,6 +25,16 @@ struct FPCGStack;
 class FPCGStackContext;
 class FTextFormat;
 
+namespace PCGGraphExecutor
+{
+	extern PCG_API TAutoConsoleVariable<float> CVarTimePerFrame;
+	extern PCG_API TAutoConsoleVariable<bool> CVarGraphMultithreading;
+
+#if WITH_EDITOR
+	extern PCG_API TAutoConsoleVariable<float> CVarEditorTimePerFrame;
+#endif
+}
+
 struct FPCGGraphTaskInput
 {
 	FPCGGraphTaskInput(FPCGTaskId InTaskId, const UPCGPin* InInboundPin, const UPCGPin* InOutboundPin, bool bInProvideData = true)
@@ -68,6 +79,7 @@ struct FPCGGraphScheduleTask
 	TWeakObjectPtr<UPCGComponent> SourceComponent = nullptr;
 	int32 FirstTaskIndex = 0;
 	int32 LastTaskIndex = 0;
+	bool bHasAbortCallbacks = false;
 };
 
 struct FPCGGraphActiveTask
@@ -87,7 +99,7 @@ struct FPCGGraphActiveTask
 class FPCGGraphExecutor : public FGCObject
 {
 public:
-	explicit FPCGGraphExecutor(UObject* InOwner);
+	FPCGGraphExecutor();
 	~FPCGGraphExecutor();
 
 	/** Compile (and cache) a graph for later use. This call is threadsafe */
@@ -120,6 +132,7 @@ public:
 
 	// Back compatibility function. Use ScheduleGenericWithContext
 	FPCGTaskId ScheduleGeneric(TFunction<bool()> InOperation, UPCGComponent* InSourceComponent, const TArray<FPCGTaskId>& TaskExecutionDependencies);
+	FPCGTaskId ScheduleGeneric(TFunction<bool()> InOperation, TFunction<void()> InAbortOperation, UPCGComponent* InSourceComponent, const TArray<FPCGTaskId>& TaskExecutionDependencies);
 
 	/** General job scheduling
 	*  @param InOperation:               Callback that takes a Context as argument and returns true if the task is done, false otherwise
@@ -129,11 +142,23 @@ public:
 	*/
 	FPCGTaskId ScheduleGenericWithContext(TFunction<bool(FPCGContext*)> InOperation, UPCGComponent* InSourceComponent, const TArray<FPCGTaskId>& TaskExecutionDependencies, const TArray<FPCGTaskId>& TaskDataDependencies);
 
+	/** General job scheduling
+	*  @param InOperation:               Callback that takes a Context as argument and returns true if the task is done, false otherwise
+	*  @param InAbortOperation:          Callback that is called if the task is aborted (cancelled) before fully executed.
+	*  @param InSourceComponent:         PCG component associated with this task. Can be null.
+	*  @param TaskExecutionDependencies: Task will wait on these tasks to execute and won't take their output data as input.
+	*  @param TaskDataDependencies:      Task will wait on these tasks to execute and will take their output data as input.
+	*/
+	FPCGTaskId ScheduleGenericWithContext(TFunction<bool(FPCGContext*)> InOperation, TFunction<void(FPCGContext*)> InAbortOperation, UPCGComponent* InSourceComponent, const TArray<FPCGTaskId>& TaskExecutionDependencies, const TArray<FPCGTaskId>& TaskDataDependencies);
+
 	/** Gets data in the output results. Returns false if data is not ready. */
 	bool GetOutputData(FPCGTaskId InTaskId, FPCGDataCollection& OutData);
 
-	/** So the profiler can decode graph task ids **/
-	FPCGGraphCompiler* GetCompiler() const { return GraphCompiler.Get(); }
+	/** Accessor so PCG tools (e.g. profiler) can easily decode graph task ids **/
+	FPCGGraphCompiler& GetCompiler() { return GraphCompiler; }
+
+	/** Accessor so PCG tools (e.g. profiler) can easily decode graph task ids **/
+	const FPCGGraphCompiler& GetCompiler() const { return GraphCompiler; }
 
 #if WITH_EDITOR
 	FPCGTaskId ScheduleDebugWithTaskCallback(UPCGComponent* InComponent, TFunction<void(FPCGTaskId, const UPCGNode*, const FPCGDataCollection&)> TaskCompleteCallback);
@@ -185,7 +210,7 @@ private:
 #endif
 
 	/** Graph compiler that turns a graph into tasks */
-	TUniquePtr<FPCGGraphCompiler> GraphCompiler;
+	FPCGGraphCompiler GraphCompiler;
 
 	/** Graph results cache */
 	FPCGGraphCache GraphCache;
@@ -244,6 +269,14 @@ public:
 	{
 		return new FPCGContext();
 	});
+
+	FPCGGenericElement(
+		TFunction<bool(FPCGContext*)> InOperation,
+		TFunction<void(FPCGContext*)> InAbortOperation,
+		const FContextAllocator& InContextAllocator = (FContextAllocator)[](const FPCGDataCollection&, TWeakObjectPtr<UPCGComponent>, const UPCGNode*)
+	{
+		return new FPCGContext();
+	});
 	
 	virtual FPCGContext* Initialize(const FPCGDataCollection& InputData, TWeakObjectPtr<UPCGComponent> SourceComponent, const UPCGNode* Node) override;
 
@@ -254,6 +287,7 @@ protected:
 	// Important note: generic elements must always be run on the main thread
 	// as most of these will impact the editor in some way (loading, unloading, saving)
 	virtual bool ExecuteInternal(FPCGContext* Context) const override;
+	virtual void AbortInternal(FPCGContext* Context) const override;
 	virtual bool IsCancellable() const override { return false; }
 
 #if WITH_EDITOR
@@ -262,6 +296,7 @@ protected:
 
 private:
 	TFunction<bool(FPCGContext*)> Operation;
+	TFunction<void(FPCGContext*)> AbortOperation;
 
 	/** Creates a context object for this element. */
 	FContextAllocator ContextAllocator;
diff --git a/Engine/Plugins/Experimental/PCG/Source/PCG/Private/PCGComponent.cpp b/Engine/Plugins/Experimental/PCG/Source/PCG/Private/PCGComponent.cpp
--- a/Engine/Plugins/Experimental/PCG/Source/PCG/Private/PCGComponent.cpp
+++ b/Engine/Plugins/Experimental/PCG/Source/PCG/Private/PCGComponent.cpp
@@ -2896,7 +2896,7 @@ void UPCGComponent::ChangeTransientState(EPCGEditorDirtyMode NewEditingMode)
 
 bool UPCGComponent::GetStackContext(FPCGStackContext& OutStackContext) const
 {
-	const UPCGSubsystem* Subsystem = GetSubsystem();
+	UPCGSubsystem* Subsystem = GetSubsystem();
 	if (Subsystem && Subsystem->GetStackContext(this, OutStackContext))
 	{
 		FPCGStack ComponentStack;
diff --git a/Engine/Plugins/Experimental/PCG/Source/PCG/Private/PCGData.cpp b/Engine/Plugins/Experimental/PCG/Source/PCG/Private/PCGData.cpp
--- a/Engine/Plugins/Experimental/PCG/Source/PCG/Private/PCGData.cpp
+++ b/Engine/Plugins/Experimental/PCG/Source/PCG/Private/PCGData.cpp
@@ -418,7 +418,6 @@ void FPCGDataCollection::Reset()
 {
 	// Implementation note: We are assuming that there is no need to remove the data from the rootset here.
 	TaggedData.Reset();
-	bCancelExecutionOnEmpty = false;
 	bCancelExecution = false;
 }
 
diff --git a/Engine/Plugins/Experimental/PCG/Source/PCG/Private/PCGElement.cpp b/Engine/Plugins/Experimental/PCG/Source/PCG/Private/PCGElement.cpp
--- a/Engine/Plugins/Experimental/PCG/Source/PCG/Private/PCGElement.cpp
+++ b/Engine/Plugins/Experimental/PCG/Source/PCG/Private/PCGElement.cpp
@@ -177,6 +177,11 @@ void IPCGElement::PostExecute(FPCGContext* Context) const
 	Context->CurrentPhase = EPCGExecutionPhase::Done;
 }
 
+void IPCGElement::Abort(FPCGContext* Context) const
+{
+	AbortInternal(Context);
+}
+
 void IPCGElement::DisabledPassThroughData(FPCGContext* Context) const
 {
 	check(Context);
diff --git a/Engine/Plugins/Experimental/PCG/Source/PCG/Private/PCGSubgraph.cpp b/Engine/Plugins/Experimental/PCG/Source/PCG/Private/PCGSubgraph.cpp
--- a/Engine/Plugins/Experimental/PCG/Source/PCG/Private/PCGSubgraph.cpp
+++ b/Engine/Plugins/Experimental/PCG/Source/PCG/Private/PCGSubgraph.cpp
@@ -594,11 +594,21 @@ bool FPCGSubgraphElement::ExecuteInternal(FPCGContext* InContext) const
 					Context->bIsPaused = true;
 					
 					// add a trivial task after the output task that wakes up this task
-					Subsystem->ScheduleGeneric([Context]() {
-						// Wake up the current task
-						Context->bIsPaused = false;
-						return true;
-					}, Context->SourceComponent.Get(), Context->SubgraphTaskIds);
+					Subsystem->ScheduleGeneric(
+						[Context]() // Normal execution: Wake up the current task
+						{
+							Context->bIsPaused = false;
+							return true;
+						}, 
+						[Context]() // On Abort: Wake up and cancel the execution
+						{
+							Context->bIsPaused = false;
+							Context->OutputData.bCancelExecution = true;
+							Context->SubgraphTaskIds.Reset();
+							return true;
+						},
+						Context->SourceComponent.Get(), 
+						Context->SubgraphTaskIds);
 
 					return false;
 				}
diff --git a/Engine/Plugins/Experimental/PCG/Source/PCG/Private/PCGSubsystem.cpp b/Engine/Plugins/Experimental/PCG/Source/PCG/Private/PCGSubsystem.cpp
--- a/Engine/Plugins/Experimental/PCG/Source/PCG/Private/PCGSubsystem.cpp
+++ b/Engine/Plugins/Experimental/PCG/Source/PCG/Private/PCGSubsystem.cpp
@@ -172,7 +172,7 @@ void UPCGSubsystem::PostInitialize()
 
 	// Initialize graph executor
 	check(!GraphExecutor);
-	GraphExecutor = new FPCGGraphExecutor(this);
+	GraphExecutor = new FPCGGraphExecutor();
 
 	// Initialize runtime generation scheduler
 	check(!RuntimeGenScheduler);
@@ -597,7 +597,7 @@ FPCGTaskId UPCGSubsystem::ScheduleComponent(UPCGComponent* PCGComponent, EPCGHiG
 			}
 
 			return true;
-		}, PCGComponent, DataDependencyTasks, ExecutionDependencyTasks);
+		}, PCGComponent, ExecutionDependencyTasks, DataDependencyTasks);
 	}
 	else
 	{
@@ -740,12 +740,24 @@ FPCGTaskId UPCGSubsystem::ScheduleGeneric(TFunction<bool()> InOperation, UPCGCom
 	return GraphExecutor->ScheduleGeneric(InOperation, SourceComponent, TaskExecutionDependencies);
 }
 
+FPCGTaskId UPCGSubsystem::ScheduleGeneric(TFunction<bool()> InOperation, TFunction<void()> InAbortOperation, UPCGComponent* SourceComponent, const TArray<FPCGTaskId>& TaskExecutionDependencies)
+{
+	check(GraphExecutor);
+	return GraphExecutor->ScheduleGeneric(InOperation, InAbortOperation, SourceComponent, TaskExecutionDependencies);
+}
+
 FPCGTaskId UPCGSubsystem::ScheduleGenericWithContext(TFunction<bool(FPCGContext*)> InOperation, UPCGComponent* SourceComponent, const TArray<FPCGTaskId>& TaskExecutionDependencies, const TArray<FPCGTaskId>& TaskDataDependencies)
 {
 	check(GraphExecutor);
 	return GraphExecutor->ScheduleGenericWithContext(InOperation, SourceComponent, TaskExecutionDependencies, TaskDataDependencies);
 }
 
+FPCGTaskId UPCGSubsystem::ScheduleGenericWithContext(TFunction<bool(FPCGContext*)> InOperation, TFunction<void(FPCGContext*)> InAbortOperation, UPCGComponent* SourceComponent, const TArray<FPCGTaskId>& TaskExecutionDependencies, const TArray<FPCGTaskId>& TaskDataDependencies)
+{
+	check(GraphExecutor);
+	return GraphExecutor->ScheduleGenericWithContext(InOperation, InAbortOperation, SourceComponent, TaskExecutionDependencies, TaskDataDependencies);
+}
+
 void UPCGSubsystem::CancelGeneration(UPCGComponent* Component)
 {
 	check(GraphExecutor && IsInGameThread());
@@ -1493,17 +1505,17 @@ void UPCGSubsystem::ClearLandscapeCache()
 	}
 }
 
-FPCGGraphCompiler* UPCGSubsystem::GetGraphCompiler() const
+FPCGGraphCompiler* UPCGSubsystem::GetGraphCompiler()
 {
 	if (GraphExecutor)
 	{
-		return GraphExecutor->GetCompiler();
+		return &(GraphExecutor->GetCompiler());
 	}
 
 	return nullptr;
 }
 
-bool UPCGSubsystem::GetStackContext(const UPCGComponent* InComponent, FPCGStackContext& OutStackContext) const
+bool UPCGSubsystem::GetStackContext(const UPCGComponent* InComponent, FPCGStackContext& OutStackContext)
 {
 	if (InComponent && InComponent->GetGraph())
 	{
@@ -1577,11 +1589,7 @@ void UPCGSubsystem::FlushCache()
 	if (GraphExecutor)
 	{
 		GraphExecutor->GetCache().ClearCache();
-
-		if (FPCGGraphCompiler* Compiler = GraphExecutor->GetCompiler())
-		{
-			Compiler->ClearCache();
-		}
+		GraphExecutor->GetCompiler().ClearCache();
 	}
 
 #if WITH_EDITOR
diff --git a/Engine/Plugins/Experimental/PCG/Source/PCG/Private/Tests/Determinism/PCGDeterminismTestsCommon.cpp b/Engine/Plugins/Experimental/PCG/Source/PCG/Private/Tests/Determinism/PCGDeterminismTestsCommon.cpp
--- a/Engine/Plugins/Experimental/PCG/Source/PCG/Private/Tests/Determinism/PCGDeterminismTestsCommon.cpp
+++ b/Engine/Plugins/Experimental/PCG/Source/PCG/Private/Tests/Determinism/PCGDeterminismTestsCommon.cpp
@@ -96,7 +96,7 @@ namespace PCGDeterminismTests
 			PCGComponentCopy->RegisterComponentWithWorld(PCGActor->GetWorld());
 		}
 
-		FPCGGraphExecutor Executor = FPCGGraphExecutor(PCGComponentCopy);
+		FPCGGraphExecutor Executor;
 
 		auto ScheduleAndWaitForExecution = [&Executor, PCGComponentCopy](FPCGTaskId FinalTaskID)
 		{
@@ -107,7 +107,7 @@ namespace PCGDeterminismTests
 			Executor.GetCache().ClearCache();
 
 			// Schedule final task for waiting
-			Executor.ScheduleGeneric([&bTasksComplete]
+			Executor.ScheduleGeneric([&bTasksComplete]()
 			{
 				bTasksComplete = true;
 				return true;
diff --git a/Engine/Plugins/Experimental/PCG/Source/PCG/Private/Tests/Grid/PCGGraphExecutorTest.cpp b/Engine/Plugins/Experimental/PCG/Source/PCG/Private/Tests/Grid/PCGGraphExecutorTest.cpp
new file mode 100644
--- /dev/null
+++ b/Engine/Plugins/Experimental/PCG/Source/PCG/Private/Tests/Grid/PCGGraphExecutorTest.cpp
@@ -0,0 +1,175 @@
+// Copyright Epic Games, Inc. All Rights Reserved.
+
+#if WITH_EDITOR
+
+#include "PCGComponent.h"
+#include "Graph/PCGGraphExecutor.h"
+#include "Tests/PCGTestsCommon.h"
+
+IMPLEMENT_CUSTOM_SIMPLE_AUTOMATION_TEST(FPCGGraphExecutorTests, FPCGTestBaseClass, "Plugins.PCG.GraphExecutor.Tests", PCGTestsCommon::TestFlags)
+
+bool FPCGGraphExecutorTests::RunTest(const FString& Parameters)
+{
+	bool bTestPassed = true;
+
+	constexpr int ComponentsNum = 3;
+
+	TArray<PCGTestsCommon::FTestData> TestData;
+	TestData.SetNum(ComponentsNum);
+
+	FPCGGraphExecutor Executor;
+	TArray<TArray<FPCGTaskId>> Tasks;
+	TArray<TArray<bool>> bExecuted;
+	TArray<TArray<bool>> bAborted;
+
+	auto AddTasks = [&](int TasksPerComponent)
+	{
+		Tasks.SetNum(ComponentsNum);
+		bExecuted.SetNum(ComponentsNum);
+		bAborted.SetNum(ComponentsNum);
+
+		for (int ComponentIndex = 0; ComponentIndex < ComponentsNum; ++ComponentIndex)
+		{
+			TArray<FPCGTaskId>& ComponentTasks = Tasks[ComponentIndex];
+			ComponentTasks.SetNum(TasksPerComponent);
+			bExecuted[ComponentIndex].Init(false, TasksPerComponent);
+			bAborted[ComponentIndex].Init(false, TasksPerComponent);
+
+			for (int TaskIndex = 0; TaskIndex < TasksPerComponent; ++TaskIndex)
+			{
+				TArray<FPCGTaskId> PreviousTasks;
+
+				if (ComponentIndex != 0)
+				{
+					PreviousTasks.Add(Tasks[ComponentIndex - 1][TaskIndex]);
+				}
+
+				if (TaskIndex != 0)
+				{
+					PreviousTasks.Add(Tasks[ComponentIndex][TaskIndex - 1]);
+				}
+
+				ComponentTasks[TaskIndex] = Executor.ScheduleGeneric(
+					[&bExecuted, ComponentIndex, TaskIndex]()
+					{
+						bExecuted[ComponentIndex][TaskIndex] = true;
+						return true;
+					},
+					[&bAborted, ComponentIndex, TaskIndex]()
+					{
+						bAborted[ComponentIndex][TaskIndex] = true;
+					}, TestData[ComponentIndex].TestPCGComponent, PreviousTasks);
+			}
+		}
+	};
+
+	auto ExecuteNTimes = [&Executor](int NTimes)
+	{
+		for (int i = 0; i < NTimes; ++i)
+		{
+			Executor.Execute();
+		}
+	};
+
+	auto ExecuteAll = [&ExecuteNTimes]()
+	{
+		ExecuteNTimes(1024);
+	};
+
+	// First test: schedule all three tasks, and cancel them, checking that the appropriate number is cancelled.
+	{
+		AddTasks(1);
+		Executor.Cancel(TestData[0].TestPCGComponent);
+		ExecuteAll();
+
+		bTestPassed &= TestTrue(TEXT("Expect that no tasks are executed"), !bExecuted[0][0] && !bExecuted[1][0] && !bExecuted[2][0]);
+		bTestPassed &= TestTrue(TEXT("Expect that all tasks are aborted"), bAborted[0][0] && bAborted[1][0] && bAborted[2][0]);
+
+		AddTasks(1);
+		Executor.Cancel(TestData[1].TestPCGComponent);
+		ExecuteAll();
+
+		bTestPassed &= TestTrue(TEXT("Expect that only task from A is executed"), bExecuted[0][0] && !bExecuted[1][0] && !bExecuted[2][0]);
+		bTestPassed &= TestTrue(TEXT("Expect that only task from A is not aborted"), !bAborted[0][0] && bAborted[1][0] && bAborted[2][0]);
+
+		AddTasks(1);
+		Executor.Cancel(TestData[2].TestPCGComponent);
+		ExecuteAll();
+
+		bTestPassed &= TestTrue(TEXT("Expect that only task from C is not executed"), bExecuted[0][0] && bExecuted[1][0] && !bExecuted[2][0]);
+		bTestPassed &= TestTrue(TEXT("Expect that onlt task from C is aborted"), !bAborted[0][0] && !bAborted[1][0] && bAborted[2][0]);
+
+		AddTasks(1);
+		ExecuteAll();
+
+		bTestPassed &= TestTrue(TEXT("Expect that all tasks are executed"), bExecuted[0][0] && bExecuted[1][0] && bExecuted[2][0]);
+		bTestPassed &= TestTrue(TEXT("Expect that no tasks are aborted"), !bAborted[0][0] && !bAborted[1][0] && !bAborted[2][0]);
+	}
+
+	// Set editor frame time to zero so we can accurately control execution
+	bool PreviousWasMultithreaded = PCGGraphExecutor::CVarGraphMultithreading.GetValueOnAnyThread();
+	PCGGraphExecutor::CVarGraphMultithreading->Set(false);
+	float PreviousEditorFrameTime = PCGGraphExecutor::CVarEditorTimePerFrame.GetValueOnAnyThread();
+	PCGGraphExecutor::CVarEditorTimePerFrame->Set(0.0f);
+
+	// Second test: execute partially, then cancel
+	{
+		const int TasksPerComponent = 3;
+		const int ExecutionPhasesCount = 4;
+
+		for (int TaskIndex = 0; TaskIndex <= ComponentsNum * TasksPerComponent * ExecutionPhasesCount; ++TaskIndex)
+		{
+			int ComponentIndex = TaskIndex / (TasksPerComponent * ExecutionPhasesCount);
+			int ComponentTaskIndex = (TaskIndex / ExecutionPhasesCount) % TasksPerComponent;
+			int ExecutionPhaseIndex = TaskIndex % ExecutionPhasesCount;
+
+			AddTasks(TasksPerComponent);
+			ExecuteNTimes(TaskIndex);
+
+			FPCGTaskId FirstTaskCancelled = InvalidPCGTaskId;
+
+			if (ComponentIndex < ComponentsNum)
+			{
+				FirstTaskCancelled = Tasks[ComponentIndex][ComponentTaskIndex];
+				Executor.Cancel(TestData[ComponentIndex].TestPCGComponent);
+			}
+
+			// Finish execution
+			ExecuteAll();
+
+			// Perform validation
+			bool bTestIterationPassed = true;
+			for (int ValidationComponentIndex = 0; ValidationComponentIndex < ComponentsNum; ++ValidationComponentIndex)
+			{
+				for (int ValidationTaskIndex = 0; ValidationTaskIndex < TasksPerComponent; ++ValidationTaskIndex)
+				{
+					bool bShouldBeExecuted;
+					bool bShouldBeAborted;
+
+					if (Tasks[ValidationComponentIndex][ValidationTaskIndex] == FirstTaskCancelled)
+ 					{
+						bShouldBeExecuted = (ExecutionPhaseIndex == 3);
+						bShouldBeAborted = true;
+					}
+					else
+					{
+						bShouldBeExecuted = Tasks[ValidationComponentIndex][ValidationTaskIndex] < FirstTaskCancelled;
+						bShouldBeAborted = Tasks[ValidationComponentIndex][ValidationTaskIndex] >= FirstTaskCancelled;
+					}
+
+					bTestIterationPassed &= (bExecuted[ValidationComponentIndex][ValidationTaskIndex] == bShouldBeExecuted);
+					bTestIterationPassed &= (bAborted[ValidationComponentIndex][ValidationTaskIndex] == bShouldBeAborted);
+				}
+			}
+
+			bTestPassed &= TestTrue(FString::Printf(TEXT("Iteration %d failed"), TaskIndex), bTestIterationPassed);
+		}
+	}
+
+	PCGGraphExecutor::CVarGraphMultithreading->Set(PreviousWasMultithreaded);
+	PCGGraphExecutor::CVarEditorTimePerFrame->Set(PreviousEditorFrameTime);
+
+	return bTestPassed;
+}
+
+#endif // WITH_EDITOR
\ No newline at end of file
diff --git a/Engine/Plugins/Experimental/PCG/Source/PCG/Public/PCGData.h b/Engine/Plugins/Experimental/PCG/Source/PCG/Public/PCGData.h
--- a/Engine/Plugins/Experimental/PCG/Source/PCG/Public/PCGData.h
+++ b/Engine/Plugins/Experimental/PCG/Source/PCG/Public/PCGData.h
@@ -175,6 +175,7 @@ struct PCG_API FPCGDataCollection
 	UPROPERTY(BlueprintReadWrite, EditAnywhere, Category = Data)
 	TArray<FPCGTaggedData> TaggedData;
 
+	/** Deprecated - Will be removed in 5.4 */
 	UPROPERTY(BlueprintReadWrite, EditAnywhere, Category = Data)
 	bool bCancelExecutionOnEmpty = false;
 
diff --git a/Engine/Plugins/Experimental/PCG/Source/PCG/Public/PCGElement.h b/Engine/Plugins/Experimental/PCG/Source/PCG/Public/PCGElement.h
--- a/Engine/Plugins/Experimental/PCG/Source/PCG/Public/PCGElement.h
+++ b/Engine/Plugins/Experimental/PCG/Source/PCG/Public/PCGElement.h
@@ -96,6 +96,9 @@ public:
 	*/
 	bool Execute(FPCGContext* Context) const;
 
+	/** Public function called when an element is cancelled, passing its current context if any. */
+	void Abort(FPCGContext* Context) const;
+
 	/** Note: the following methods must be called from the main thread */
 #if WITH_EDITOR
 	void DebugDisplay(FPCGContext* Context) const;
@@ -110,6 +113,8 @@ protected:
 	virtual bool ExecuteInternal(FPCGContext* Context) const = 0;
 	/** This function will be called once and once only, at the end of an execution */
 	void PostExecute(FPCGContext* Context) const;
+	/** This function will be called once and only once if the element is aborted. The Context can be used to retrieve the current phase if needed. */
+	virtual void AbortInternal(FPCGContext* Context) const {};
 
 	/** Controls whether an element can skip its execution wholly when the input data has the cancelled tag */
 	virtual bool IsCancellable() const { return true; }
diff --git a/Engine/Plugins/Experimental/PCG/Source/PCG/Public/PCGSubsystem.h b/Engine/Plugins/Experimental/PCG/Source/PCG/Public/PCGSubsystem.h
--- a/Engine/Plugins/Experimental/PCG/Source/PCG/Public/PCGSubsystem.h
+++ b/Engine/Plugins/Experimental/PCG/Source/PCG/Public/PCGSubsystem.h
@@ -123,6 +123,14 @@ public:
 	*/
 	FPCGTaskId ScheduleGeneric(TFunction<bool()> InOperation, UPCGComponent* SourceComponent, const TArray<FPCGTaskId>& TaskExecutionDependencies);
 
+	/** General job scheduling
+	*  @param InOperation:               Callback that returns true if the task is done, false otherwise.
+	*  @param InAbortOperation:          Callback that will be called if the generic task is cancelled for any reason.
+	*  @param SourceComponent:           PCG component associated with this task. Can be null.
+	*  @param TaskExecutionDependencies: Task will wait on these tasks to execute and won't take their output data as input.
+	*/
+	FPCGTaskId ScheduleGeneric(TFunction<bool()> InOperation, TFunction<void()> InAbortOperation, UPCGComponent* SourceComponent, const TArray<FPCGTaskId>& TaskExecutionDependencies);
+
 	/** General job scheduling with context
 	*  @param InOperation:               Callback that takes a Context as argument and returns true if the task is done, false otherwise.
 	*  @param SourceComponent:           PCG component associated with this task. Can be null.
@@ -131,6 +139,15 @@ public:
 	*/
 	FPCGTaskId ScheduleGenericWithContext(TFunction<bool(FPCGContext*)> InOperation, UPCGComponent* SourceComponent, const TArray<FPCGTaskId>& TaskExecutionDependencies, const TArray<FPCGTaskId>& TaskDataDependencies);
 
+	/** General job scheduling with context
+	*  @param InOperation:               Callback that takes a Context as argument and returns true if the task is done, false otherwise.
+	*  @param InAbortOperation:          Callback that will be called if the generic task is cancelled for any reason.
+	*  @param SourceComponent:           PCG component associated with this task. Can be null.
+	*  @param TaskExecutionDependencies: Task will wait on these tasks to execute and won't take their output data as input.
+	*  @param TaskDataDependencies:      Task will wait on these tasks to execute and will take their output data as input.
+	*/
+	FPCGTaskId ScheduleGenericWithContext(TFunction<bool(FPCGContext*)> InOperation, TFunction<void(FPCGContext*)> InAbortOperation, UPCGComponent* SourceComponent, const TArray<FPCGTaskId>& TaskExecutionDependencies, const TArray<FPCGTaskId>& TaskDataDependencies);
+
 	/** Cancels currently running generation */
 	void CancelGeneration(UPCGComponent* Component);
 
@@ -223,10 +240,10 @@ public:
 	void ClearLandscapeCache();
 
 	/** Returns the graph compiler so we can figure out task info in the profiler view **/
-	FPCGGraphCompiler* GetGraphCompiler() const;
+	FPCGGraphCompiler* GetGraphCompiler();
 
 	/** Get the execution stack information for the given component. */
-	bool GetStackContext(const UPCGComponent* InComponent, FPCGStackContext& OutStackContext) const;
+	bool GetStackContext(const UPCGComponent* InComponent, FPCGStackContext& OutStackContext);
 
 	/** Returns how many times InElement is present in the cache. */
 	uint32 GetGraphCacheEntryCount(IPCGElement* InElement) const;
