[PCG] Fix bug where closed distance spline sampler does not sample last segment
#rnx

diff --git a/Engine/Plugins/Experimental/PCG/Source/PCG/Private/Elements/PCGSplineSampler.cpp b/Engine/Plugins/Experimental/PCG/Source/PCG/Private/Elements/PCGSplineSampler.cpp
--- a/Engine/Plugins/Experimental/PCG/Source/PCG/Private/Elements/PCGSplineSampler.cpp
+++ b/Engine/Plugins/Experimental/PCG/Source/PCG/Private/Elements/PCGSplineSampler.cpp
@@ -508,9 +508,7 @@ namespace PCGSplineSampler
 
 		virtual bool IsDone() const override
 		{
-			// Distance Sampler is not done until it captures the final control point, which does not have a valid segment associated unless the spline is a closed loop.
-			const int NumSegments = LineData->IsClosed() ? LineData->GetNumSegments() - 1 : LineData->GetNumSegments();
-			return CurrentSegmentIndex >= NumSegments;
+			return CurrentSegmentIndex >= LineData->GetNumSegments();
 		}
 
 		FVector::FReal DistanceIncrement;
