Added multiple details view and multiple attribute list views in PCG Editor.

diff --git a/Engine/Plugins/Experimental/PCG/Source/PCGEditor/Private/PCGEditor.cpp b/Engine/Plugins/Experimental/PCG/Source/PCGEditor/Private/PCGEditor.cpp
--- a/Engine/Plugins/Experimental/PCG/Source/PCGEditor/Private/PCGEditor.cpp
+++ b/Engine/Plugins/Experimental/PCG/Source/PCGEditor/Private/PCGEditor.cpp
@@ -2,7 +2,6 @@
 
 #include "PCGEditor.h"
 
-#include "DetailsViewArgs.h"
 #include "PCGComponent.h"
 #include "Framework/Application/SlateApplication.h"
 #include "PCGEdge.h"
@@ -31,6 +30,7 @@
 #include "PCGEditorUtils.h"
 #include "SPCGEditorGraphAttributeListView.h"
 #include "SPCGEditorGraphDebugObjectTree.h"
+#include "SPCGEditorGraphDetailsView.h"
 #include "SPCGEditorGraphDeterminism.h"
 #include "SPCGEditorGraphFind.h"
 #include "SPCGEditorGraphLogView.h"
@@ -70,10 +70,18 @@
 namespace FPCGEditor_private
 {
 	const FName GraphEditorID = FName(TEXT("GraphEditor"));
-	const FName PropertyDetailsID = FName(TEXT("PropertyDetails"));
+	const FName PropertyDetailsID[] = {
+		FName(TEXT("PropertyDetails")),
+		FName(TEXT("PropertyDetails2")),
+		FName(TEXT("PropertyDetails3")),
+		FName(TEXT("PropertyDetails4")) };
 	const FName PaletteID = FName(TEXT("Palette"));
 	const FName DebugObjectID = FName(TEXT("DebugObject"));
-	const FName AttributesID = FName(TEXT("Attributes"));
+	const FName AttributesID[] = {
+		FName(TEXT("Attributes")),
+		FName(TEXT("Attributes2")),
+		FName(TEXT("Attributes3")),
+		FName(TEXT("Attributes4")) };
 	const FName FindID = FName(TEXT("Find"));
 	const FName DeterminismID = FName(TEXT("Determinism"));
 	const FName ProfilingID = FName(TEXT("Profiling"));
@@ -117,24 +125,23 @@ void FPCGEditor::Initialize(const EToolkitMode::Type InMode, const TSharedPtr<cl
 	PCGGraphBeingEdited->PCGEditorGraph->SetEditor(SharedThis(this));
 	PCGEditorGraph = PCGGraphBeingEdited->PCGEditorGraph;
 
-	FPropertyEditorModule& PropertyEditorModule = FModuleManager::GetModuleChecked<FPropertyEditorModule>("PropertyEditor");
-
-	FDetailsViewArgs DetailsViewArgs;
-	DetailsViewArgs.NameAreaSettings = FDetailsViewArgs::HideNameArea;
-	DetailsViewArgs.bHideSelectionTip = true;
-	DetailsViewArgs.DefaultsOnlyVisibility = EEditDefaultsOnlyNodeVisibility::Automatic;
-	PropertyDetailsWidget = PropertyEditorModule.CreateDetailView(DetailsViewArgs);
-	PropertyDetailsWidget->SetObject(PCGGraphBeingEdited);
-
-	IDetailsView* RawPropertyDetailsPtr = PropertyDetailsWidget.Get();
-	PropertyDetailsWidget->SetIsPropertyReadOnlyDelegate(FIsPropertyReadOnly::CreateRaw(this, &FPCGEditor::IsReadOnlyProperty, RawPropertyDetailsPtr));
-	PropertyDetailsWidget->SetIsPropertyVisibleDelegate(FIsPropertyVisible::CreateRaw(this, &FPCGEditor::IsVisibleProperty, RawPropertyDetailsPtr));
+	for (int PropertyDetailsIndex = 0; PropertyDetailsIndex < 4; ++PropertyDetailsIndex)
+	{
+		TSharedRef<SPCGEditorGraphDetailsView> PropertyDetailsWidget = SNew(SPCGEditorGraphDetailsView);
+		PropertyDetailsWidget->SetObject(PCGGraphBeingEdited);
+		PropertyDetailsWidgets.Add(PropertyDetailsWidget);
+	}
 
 	GraphEditorWidget = CreateGraphEditorWidget();
 	PaletteWidget = CreatePaletteWidget();
 	DebugObjectTreeWidget = CreateDebugObjectTreeWidget();
 	FindWidget = CreateFindWidget();
-	AttributesWidget = CreateAttributesWidget();
+
+	for (int AttributesIndex = 0; AttributesIndex < 4; ++AttributesIndex)
+	{
+		AttributesWidgets.Add(CreateAttributesWidget());
+	}
+	
 	DeterminismWidget = CreateDeterminismWidget();
 	ProfilingWidget = CreateProfilingWidget();
 	LogWidget = CreateLogWidget();
@@ -169,7 +176,7 @@ void FPCGEditor::Initialize(const EToolkitMode::Type InMode, const TSharedPtr<cl
 					FTabManager::NewStack()
 					->SetSizeCoefficient(0.20f)
 					->SetHideTabWell(true)
-					->AddTab(FPCGEditor_private::PropertyDetailsID, ETabState::OpenedTab)
+					->AddTab(FPCGEditor_private::PropertyDetailsID[0], ETabState::OpenedTab)
 				)
 			)
 			->Split
@@ -188,7 +195,7 @@ void FPCGEditor::Initialize(const EToolkitMode::Type InMode, const TSharedPtr<cl
 					FTabManager::NewStack()
 					->SetSizeCoefficient(0.8)
 					->SetHideTabWell(true)
-					->AddTab(FPCGEditor_private::AttributesID, ETabState::OpenedTab)
+					->AddTab(FPCGEditor_private::AttributesID[0], ETabState::OpenedTab)
 					->AddTab(FPCGEditor_private::DeterminismID, ETabState::ClosedTab)
 					->AddTab(FPCGEditor_private::FindID, ETabState::ClosedTab)
 				)
@@ -283,7 +290,10 @@ void FPCGEditor::UpdateDebugAfterComponentSelection(UPCGComponent* InOldComponen
 
 	auto RefreshComponent = [](UPCGComponent* Component)
 	{
-		check(Component);
+		if (!ensure(Component))
+		{
+			return;
+		}
 
 		// GenerateAtRuntime components should be refreshed through the runtime gen scheduler.
 		if (Component->IsManagedByRuntimeGenSystem())
@@ -360,6 +370,8 @@ void FPCGEditor::JumpToNode(const UEdGraphNode* InNode)
 void FPCGEditor::RegisterTabSpawners(const TSharedRef<FTabManager>& InTabManager)
 {
 	WorkspaceMenuCategory = InTabManager->AddLocalWorkspaceMenuCategory(LOCTEXT("WorkspaceMenu_PCGEditor", "PCG Editor"));
+	TSharedRef<FWorkspaceItem> DetailsGroup = WorkspaceMenuCategory->AddGroup(LOCTEXT("WorkspaceMenu_PCGEditor_Details", "Details"));
+	TSharedRef<FWorkspaceItem> AttributesGroup = WorkspaceMenuCategory->AddGroup(LOCTEXT("WorkspaceMenu_PCGEditor_Attributes", "Attributes"));
 	const TSharedRef<FWorkspaceItem>& WorkspaceMenuCategoryRef = WorkspaceMenuCategory.ToSharedRef();
 
 	FAssetEditorToolkit::RegisterTabSpawners(InTabManager);
@@ -369,9 +381,21 @@ void FPCGEditor::RegisterTabSpawners(const TSharedRef<FTabManager>& InTabManager
 		.SetDisplayName(LOCTEXT("GraphTab", "Graph"))
 		.SetGroup(WorkspaceMenuCategoryRef);
 
-	InTabManager->RegisterTabSpawner(FPCGEditor_private::PropertyDetailsID, FOnSpawnTab::CreateSP(this, &FPCGEditor::SpawnTab_PropertyDetails))
-		.SetDisplayName(LOCTEXT("DetailsTab", "Details"))
-		.SetGroup(WorkspaceMenuCategoryRef);
+	InTabManager->RegisterTabSpawner(FPCGEditor_private::PropertyDetailsID[0], FOnSpawnTab::CreateSP(this, &FPCGEditor::SpawnTab_PropertyDetails, 0))
+		.SetDisplayName(LOCTEXT("DetailsTab", "Details 1"))
+		.SetGroup(DetailsGroup);
+
+	InTabManager->RegisterTabSpawner(FPCGEditor_private::PropertyDetailsID[1], FOnSpawnTab::CreateSP(this, &FPCGEditor::SpawnTab_PropertyDetails, 1))
+		.SetDisplayName(LOCTEXT("DetailsTab", "Details 2"))
+		.SetGroup(DetailsGroup);
+
+	InTabManager->RegisterTabSpawner(FPCGEditor_private::PropertyDetailsID[2], FOnSpawnTab::CreateSP(this, &FPCGEditor::SpawnTab_PropertyDetails, 2))
+		.SetDisplayName(LOCTEXT("DetailsTab", "Details 3"))
+		.SetGroup(DetailsGroup);
+
+	InTabManager->RegisterTabSpawner(FPCGEditor_private::PropertyDetailsID[3], FOnSpawnTab::CreateSP(this, &FPCGEditor::SpawnTab_PropertyDetails, 3))
+		.SetDisplayName(LOCTEXT("DetailsTab", "Details 4"))
+		.SetGroup(DetailsGroup);
 
 	InTabManager->RegisterTabSpawner(FPCGEditor_private::PaletteID, FOnSpawnTab::CreateSP(this, &FPCGEditor::SpawnTab_Palette))
 		.SetDisplayName(LOCTEXT("PaletteTab", "Palette"))
@@ -381,9 +405,21 @@ void FPCGEditor::RegisterTabSpawners(const TSharedRef<FTabManager>& InTabManager
 		.SetDisplayName(LOCTEXT("DebugTab", "Debug Object Tree"))
 		.SetGroup(WorkspaceMenuCategoryRef);
 
-	InTabManager->RegisterTabSpawner(FPCGEditor_private::AttributesID, FOnSpawnTab::CreateSP(this, &FPCGEditor::SpawnTab_Attributes))
-		.SetDisplayName(LOCTEXT("AttributesTab", "Attributes"))
-		.SetGroup(WorkspaceMenuCategoryRef);
+	InTabManager->RegisterTabSpawner(FPCGEditor_private::AttributesID[0], FOnSpawnTab::CreateSP(this, &FPCGEditor::SpawnTab_Attributes, 0))
+		.SetDisplayName(LOCTEXT("AttributesTab", "Attributes 1"))
+		.SetGroup(AttributesGroup);
+
+	InTabManager->RegisterTabSpawner(FPCGEditor_private::AttributesID[1], FOnSpawnTab::CreateSP(this, &FPCGEditor::SpawnTab_Attributes, 1))
+		.SetDisplayName(LOCTEXT("AttributesTab", "Attributes 2"))
+		.SetGroup(AttributesGroup);
+
+	InTabManager->RegisterTabSpawner(FPCGEditor_private::AttributesID[2], FOnSpawnTab::CreateSP(this, &FPCGEditor::SpawnTab_Attributes, 2))
+		.SetDisplayName(LOCTEXT("AttributesTab", "Attributes 3"))
+		.SetGroup(AttributesGroup);
+
+	InTabManager->RegisterTabSpawner(FPCGEditor_private::AttributesID[3], FOnSpawnTab::CreateSP(this, &FPCGEditor::SpawnTab_Attributes, 3))
+		.SetDisplayName(LOCTEXT("AttributesTab", "Attributes 4"))
+		.SetGroup(AttributesGroup);
 
 	InTabManager->RegisterTabSpawner(FPCGEditor_private::FindID, FOnSpawnTab::CreateSP(this, &FPCGEditor::SpawnTab_Find))
 		.SetDisplayName(LOCTEXT("FindTab", "Find"))
@@ -405,10 +441,16 @@ void FPCGEditor::RegisterTabSpawners(const TSharedRef<FTabManager>& InTabManager
 void FPCGEditor::UnregisterTabSpawners(const TSharedRef<class FTabManager>& InTabManager)
 {
 	InTabManager->UnregisterTabSpawner(FPCGEditor_private::GraphEditorID);
-	InTabManager->UnregisterTabSpawner(FPCGEditor_private::PropertyDetailsID);
+	InTabManager->UnregisterTabSpawner(FPCGEditor_private::PropertyDetailsID[0]);
+	InTabManager->UnregisterTabSpawner(FPCGEditor_private::PropertyDetailsID[1]);
+	InTabManager->UnregisterTabSpawner(FPCGEditor_private::PropertyDetailsID[2]);
+	InTabManager->UnregisterTabSpawner(FPCGEditor_private::PropertyDetailsID[3]);
 	InTabManager->UnregisterTabSpawner(FPCGEditor_private::PaletteID);
 	InTabManager->UnregisterTabSpawner(FPCGEditor_private::DebugObjectID);
-	InTabManager->UnregisterTabSpawner(FPCGEditor_private::AttributesID);
+	InTabManager->UnregisterTabSpawner(FPCGEditor_private::AttributesID[0]);
+	InTabManager->UnregisterTabSpawner(FPCGEditor_private::AttributesID[1]);
+	InTabManager->UnregisterTabSpawner(FPCGEditor_private::AttributesID[2]);
+	InTabManager->UnregisterTabSpawner(FPCGEditor_private::AttributesID[3]);
 	InTabManager->UnregisterTabSpawner(FPCGEditor_private::FindID);
 	InTabManager->UnregisterTabSpawner(FPCGEditor_private::DeterminismID);
 
@@ -875,13 +917,24 @@ void FPCGEditor::OnDeterminismGraphTest()
 
 void FPCGEditor::OnEditGraphSettings() const
 {
-	PropertyDetailsWidget->SetObject(PCGGraphBeingEdited);
+	for (TSharedPtr<SPCGEditorGraphDetailsView> PropertyDetailsWidget : PropertyDetailsWidgets)
+	{
+		PropertyDetailsWidget->SetObject(PCGGraphBeingEdited);
+	}
 }
 
 bool FPCGEditor::IsEditGraphSettingsToggled() const
 {
-	const TArray<TWeakObjectPtr<UObject>>& SelectedObjects = PropertyDetailsWidget->GetSelectedObjects();
-	return SelectedObjects.Num() == 1 && SelectedObjects[0] == PCGGraphBeingEdited.Get();
+	for (TSharedPtr<SPCGEditorGraphDetailsView> PropertyDetailsWidget : PropertyDetailsWidgets)
+	{
+		const TArray<TWeakObjectPtr<UObject>>& SelectedObjects = PropertyDetailsWidget->GetSelectedObjects();
+		if (SelectedObjects.Num() == 1 && SelectedObjects[0] == PCGGraphBeingEdited.Get())
+		{
+			return true;
+		}
+	}
+
+	return false;
 }
 
 bool FPCGEditor::CanCollapseNodesInSubgraph() const
@@ -1215,11 +1268,6 @@ void FPCGEditor::OnToggleInspected()
 		return;
 	}
 
-	if (PCGGraphNodeBeingInspected)
-	{
-		PCGGraphNodeBeingInspected->SetInspected(false);
-	}
-
 	UEdGraphNode* GraphNode = GraphEditorWidget->GetSingleSelectedNode();
 	UPCGEditorGraphNodeBase* PCGGraphNodeBase = Cast<UPCGEditorGraphNodeBase>(GraphNode);
 
@@ -1231,29 +1279,65 @@ void FPCGEditor::OnToggleInspected()
 		return;
 	}
 
-	if (PCGGraphNodeBase && PCGGraphNodeBase != PCGGraphNodeBeingInspected)
+	TArray<UPCGEditorGraphNodeBase*, TInlineAllocator<4>> InspectedNodesBefore;
+	for (TSharedPtr<SPCGEditorGraphAttributeListView> AttributeListView : AttributesWidgets)
 	{
-		PCGGraphNodeBeingInspected = PCGGraphNodeBase;
-		PCGGraphNodeBeingInspected->SetInspected(true);
+		InspectedNodesBefore.Add(AttributeListView->GetNodeBeingInspected());
 	}
-	else
+
+	bool bIsInspecting = false;
+
+	// If the selected node was previously inspected, stop inspecting it, and unselect it from the attribute list views
+	if (InspectedNodesBefore.Contains(PCGGraphNodeBase))
 	{
-		PCGGraphNodeBeingInspected = nullptr;
+		PCGGraphNodeBase->SetInspected(false);
+
+		for (TSharedPtr<SPCGEditorGraphAttributeListView> AttributeListView : AttributesWidgets)
+		{
+			if (AttributeListView->GetNodeBeingInspected() == PCGGraphNodeBase)
+			{
+				AttributeListView->SetNodeBeingInspected(nullptr);
+			}
+		}
 	}
+	else
+	{
+		TArray<UPCGEditorGraphNodeBase*, TInlineAllocator<4>> InspectedNodesAfter;
 
-	OnInspectedNodeChangedDelegate.Broadcast(PCGGraphNodeBeingInspected);
-	GetTabManager()->TryInvokeTab(FPCGEditor_private::AttributesID);
+		for (TSharedPtr<SPCGEditorGraphAttributeListView> AttributeListView : AttributesWidgets)
+		{
+			if (!AttributeListView->IsLocked())
+			{
+				AttributeListView->SetNodeBeingInspected(PCGGraphNodeBase);
+			}
 
-	DebugObjectTreeWidget->SetNodeBeingInspected(PCGGraphNodeBeingInspected ? PCGNode : nullptr);
-}
+			InspectedNodesAfter.Add(AttributeListView->GetNodeBeingInspected());
+		}
 
-bool FPCGEditor::CanToggleInspected() const
-{
-	if (PCGGraphNodeBeingInspected)
+		if (InspectedNodesAfter.Contains(PCGGraphNodeBase))
+		{
+			PCGGraphNodeBase->SetInspected(true);
+			bIsInspecting = true;
+		}
+
+		for (UPCGEditorGraphNodeBase* BeforeNode : InspectedNodesBefore)
+		{
+			if (!InspectedNodesAfter.Contains(BeforeNode) && BeforeNode)
+			{
+				BeforeNode->SetInspected(false);
+			}
+		}
+	}
+
+	if (bIsInspecting)
 	{
-		return true;
+		GetTabManager()->TryInvokeTab(FPCGEditor_private::AttributesID[0]);
+		DebugObjectTreeWidget->SetNodeBeingInspected(PCGNode);
 	}
+}
 
+bool FPCGEditor::CanToggleInspected() const
+{
 	if (!GraphEditorWidget.IsValid())
 	{
 		return false;
@@ -2197,9 +2281,12 @@ void FPCGEditor::OnSelectedNodesChanged(const TSet<UObject*>& NewSelection)
 		}
 	}
 
-	PropertyDetailsWidget->SetObjects(SelectedObjects, /*bForceRefresh=*/true);
+	for (TSharedPtr<SPCGEditorGraphDetailsView> PropertyDetailsWidget : PropertyDetailsWidgets)
+	{
+		PropertyDetailsWidget->SetObjects(SelectedObjects, /*bForceRefresh=*/true);
+	}
 
-	GetTabManager()->TryInvokeTab(FPCGEditor_private::PropertyDetailsID);
+	GetTabManager()->TryInvokeTab(FPCGEditor_private::PropertyDetailsID[0]);
 }
 
 void FPCGEditor::OnNodeTitleCommitted(const FText& NewText, ETextCommit::Type CommitInfo, UEdGraphNode* NodeBeingChanged)
@@ -2278,72 +2365,6 @@ void FPCGEditor::JumpToDefinition(const UClass* Class) const
 	}
 }
 
-bool FPCGEditor::IsReadOnlyProperty(const FPropertyAndParent& InPropertyAndParent, IDetailsView* InDetailsView) const
-{
-	// Everything is writeable when not in an instance
-	if (!InDetailsView ||
-		InPropertyAndParent.ParentProperties.IsEmpty() ||
-		InPropertyAndParent.ParentProperties.Last()->GetFName() != GET_MEMBER_NAME_CHECKED(UPCGSettingsInstance, Settings))
-	{
-		return false;
-	}
-
-	const TArray<TWeakObjectPtr<UObject>>& SelectedObjects = InDetailsView->GetSelectedObjects();
-	for (const TWeakObjectPtr<UObject>& SelectedObject : SelectedObjects)
-	{
-		if (!SelectedObject.IsValid())
-		{
-			continue;
-		}
-
-		if (UPCGSettingsInstance* Instance = Cast<UPCGSettingsInstance>(SelectedObject.Get()))
-		{
-			return true;
-		}
-	}
-
-	return false;
-}
-
-bool FPCGEditor::IsVisibleProperty(const FPropertyAndParent& InPropertyAndParent, IDetailsView* InDetailsView) const
-{
-	if (!InDetailsView)
-	{
-		return true;
-	}
-
-	// Currently never hide anything from the graph settings
-	if (InPropertyAndParent.Objects.Num() == 1 && Cast<UPCGGraph>(InPropertyAndParent.Objects[0]))
-	{
-		return true;
-	}
-
-	// Always hide asset info information
-	if (InPropertyAndParent.Property.HasMetaData(TEXT("Category")) &&
-		InPropertyAndParent.Property.GetMetaData(TEXT("Category")) == TEXT("AssetInfo"))
-	{
-		return false;
-	}
-
-	// Otherwise, everything is visible when not in an instance
-	if(	InPropertyAndParent.ParentProperties.IsEmpty() ||
-		InPropertyAndParent.ParentProperties.Last()->GetFName() != GET_MEMBER_NAME_CHECKED(UPCGSettingsInstance, Settings))
-	{
-		return true;
-	}
-
-	// Hide debug settings from the setting when showing the instance settings.
-	if (InPropertyAndParent.Property.GetFName() == GET_MEMBER_NAME_CHECKED(UPCGSettings, bEnabled) ||
-		InPropertyAndParent.Property.GetFName() == GET_MEMBER_NAME_CHECKED(UPCGSettings, bDebug) ||
-		(InPropertyAndParent.ParentProperties.Num() >= 2 &&
-			InPropertyAndParent.ParentProperties[InPropertyAndParent.ParentProperties.Num() - 2]->GetFName() == GET_MEMBER_NAME_CHECKED(UPCGSettings, DebugSettings)))
-	{
-		return false;
-	}
-
-	return true;
-}
-
 void FPCGEditor::OnGraphStructureChanged(UPCGGraphInterface* InGraph)
 {
 	check(PCGEditorGraph);
@@ -2413,13 +2434,16 @@ TSharedRef<SDockTab> FPCGEditor::SpawnTab_GraphEditor(const FSpawnTabArgs& Args)
 		];
 }
 
-TSharedRef<SDockTab> FPCGEditor::SpawnTab_PropertyDetails(const FSpawnTabArgs& Args)
+TSharedRef<SDockTab> FPCGEditor::SpawnTab_PropertyDetails(const FSpawnTabArgs& Args, int PropertyDetailsIndex)
 {
+	TAttribute<FText> Label = TAttribute<FText>::Create(TAttribute<FText>::FGetter::CreateRaw(this, &FPCGEditor::GetDetailsTabLabel, PropertyDetailsIndex));
+	TSharedPtr<SPCGEditorGraphDetailsView> DetailsView = PropertyDetailsWidgets[PropertyDetailsIndex];
+
 	return SNew(SDockTab)
-		.Label(LOCTEXT("PCGDetailsTitle", "Details"))
+		.Label(Label)
 		.TabColorScale(GetTabColorScale())
 		[
-			PropertyDetailsWidget.ToSharedRef()
+			DetailsView.ToSharedRef()
 		];
 }
 
@@ -2443,13 +2467,15 @@ TSharedRef<SDockTab> FPCGEditor::SpawnTab_DebugObjectTree(const FSpawnTabArgs& A
 		];
 }
 
-TSharedRef<SDockTab> FPCGEditor::SpawnTab_Attributes(const FSpawnTabArgs& Args)
+TSharedRef<SDockTab> FPCGEditor::SpawnTab_Attributes(const FSpawnTabArgs& Args, int AttributesIndex)
 {
+	TAttribute<FText> Label = TAttribute<FText>::Create(TAttribute<FText>::FGetter::CreateRaw(this, &FPCGEditor::GetAttributesTabLabel, AttributesIndex));
+
 	return SNew(SDockTab)
-		.Label(LOCTEXT("PCGAttributesTitle", "Attributes"))
+		.Label(Label)
 		.TabColorScale(GetTabColorScale())
 		[
-			AttributesWidget.ToSharedRef()
+			AttributesWidgets[AttributesIndex].ToSharedRef()
 		];
 }
 
@@ -2493,4 +2519,33 @@ TSharedRef<SDockTab> FPCGEditor::SpawnTab_Log(const FSpawnTabArgs& Args)
 		];
 }
 
+FText FPCGEditor::GetDetailsTabLabel(int DetailsIndex)
+{
+	if (DetailsIndex == 0)
+	{
+		return LOCTEXT("PCGDetailsTitle", "Details");
+	}
+	else
+	{
+		return FText::Format(LOCTEXT("PCGDetailsTitle_Multi", "Details {0}"), DetailsIndex + 1);
+	}
+}
+
+FText FPCGEditor::GetDetailsViewObjectName(int DetailsIndex)
+{
+	return LOCTEXT("PCGDetailsName", "This is a node name placeholder");
+}
+
+FText FPCGEditor::GetAttributesTabLabel(int AttributesIndex)
+{
+	if (AttributesIndex == 0)
+	{
+		return LOCTEXT("PCGAttributesTitle", "Attributes");
+	}
+	else
+	{
+		return FText::Format(LOCTEXT("PCGAttributesTitle_Multi", "Attributes {0}"), AttributesIndex + 1);
+	}
+}
+
 #undef LOCTEXT_NAMESPACE
diff --git a/Engine/Plugins/Experimental/PCG/Source/PCGEditor/Private/PCGEditor.h b/Engine/Plugins/Experimental/PCG/Source/PCGEditor/Private/PCGEditor.h
--- a/Engine/Plugins/Experimental/PCG/Source/PCGEditor/Private/PCGEditor.h
+++ b/Engine/Plugins/Experimental/PCG/Source/PCGEditor/Private/PCGEditor.h
@@ -12,12 +12,11 @@ class FSpawnTabArgs;
 enum class ECheckBoxState : uint8;
 namespace ETextCommit { enum Type : int; }
 
-struct FPropertyAndParent;
 class FUICommandList;
-class IDetailsView;
 class SGraphEditor;
 class SPCGEditorGraphAttributeListView;
 class SPCGEditorGraphDebugObjectTree;
+class SPCGEditorGraphDetailsView;
 class SPCGEditorGraphDeterminismListView;
 class SPCGEditorGraphFind;
 class SPCGEditorGraphLogView;
@@ -28,7 +27,6 @@ class UPCGComponent;
 class UPCGEditorGraph;
 class UPCGEditorGraphNodeBase;
 class UPCGGraph;
-class UPCGNode;
 
 DECLARE_MULTICAST_DELEGATE_OneParam(FOnInspectedStackChanged, const FPCGStack&);
 DECLARE_MULTICAST_DELEGATE_OneParam(FOnInspectedNodeChanged, UPCGEditorGraphNodeBase*);
@@ -89,7 +87,6 @@ public:
 	// ~End FAssetEditorToolkit interface
 
 	FOnInspectedStackChanged OnInspectedStackChangedDelegate;
-	FOnInspectedNodeChanged OnInspectedNodeChangedDelegate;
 
 protected:
 	// ~Begin FAssetEditorToolkit interface
@@ -286,12 +283,6 @@ private:
 	/** To be called everytime we need to replicate our extra nodes to the underlying PCGGraph */
 	void ReplicateExtraNodes() const;
 
-	/** Returns whether a property should be readonly (used for instances) */
-	bool IsReadOnlyProperty(const FPropertyAndParent& InPropertyAndParent, IDetailsView* InDetailsView) const;
-
-	/** Returns whether a property should be visible (used for instance vs. settings properties) */
-	bool IsVisibleProperty(const FPropertyAndParent& InPropertyAndParent, IDetailsView* InDetailsView) const;
-
 	void OnGraphStructureChanged(UPCGGraphInterface* InGraph);
 
 	/** Called when a component finishes executing. Useful for updating debugging tools/UIs. */
@@ -304,21 +295,25 @@ private:
 	void OnLevelActorDeleted(AActor* InActor);
 
 	TSharedRef<SDockTab> SpawnTab_GraphEditor(const FSpawnTabArgs& Args);
-	TSharedRef<SDockTab> SpawnTab_PropertyDetails(const FSpawnTabArgs& Args);
+	TSharedRef<SDockTab> SpawnTab_PropertyDetails(const FSpawnTabArgs& Args, int PropertyDetailsIndex);
 	TSharedRef<SDockTab> SpawnTab_Palette(const FSpawnTabArgs& Args);
 	TSharedRef<SDockTab> SpawnTab_DebugObjectTree(const FSpawnTabArgs& Args);
-	TSharedRef<SDockTab> SpawnTab_Attributes(const FSpawnTabArgs& Args);
+	TSharedRef<SDockTab> SpawnTab_Attributes(const FSpawnTabArgs& Args, int AttributesIndex);
 	TSharedRef<SDockTab> SpawnTab_Find(const FSpawnTabArgs& Args);
 	TSharedRef<SDockTab> SpawnTab_Determinism(const FSpawnTabArgs& Args);
 	TSharedRef<SDockTab> SpawnTab_Profiling(const FSpawnTabArgs& Args);
 	TSharedRef<SDockTab> SpawnTab_Log(const FSpawnTabArgs& Args);
 
+	FText GetDetailsTabLabel(int DetailsIndex);
+	FText GetDetailsViewObjectName(int DetailsIndex);
+	FText GetAttributesTabLabel(int AttributesIndex);
+
 	TSharedPtr<SGraphEditor> GraphEditorWidget;
-	TSharedPtr<IDetailsView> PropertyDetailsWidget;
+	TArray<TSharedPtr<SPCGEditorGraphDetailsView>> PropertyDetailsWidgets;
 	TSharedPtr<SPCGEditorGraphNodePalette> PaletteWidget;
 	TSharedPtr<SPCGEditorGraphDebugObjectTree> DebugObjectTreeWidget;
 	TSharedPtr<SPCGEditorGraphFind> FindWidget;
-	TSharedPtr<SPCGEditorGraphAttributeListView> AttributesWidget;
+	TArray<TSharedPtr<SPCGEditorGraphAttributeListView>> AttributesWidgets;
 	TSharedPtr<SPCGEditorGraphDeterminismListView> DeterminismWidget;
 	TSharedPtr<SPCGEditorGraphProfilingView> ProfilingWidget;
 	TSharedPtr<SPCGEditorGraphLogView> LogWidget;
@@ -330,5 +325,4 @@ private:
 
 	TWeakObjectPtr<UPCGComponent> PCGComponentBeingInspected;
 	FPCGStack StackBeingInspected;
-	UPCGEditorGraphNodeBase* PCGGraphNodeBeingInspected = nullptr;
 };
diff --git a/Engine/Plugins/Experimental/PCG/Source/PCGEditor/Private/SPCGEditorGraphAttributeListView.cpp b/Engine/Plugins/Experimental/PCG/Source/PCGEditor/Private/SPCGEditorGraphAttributeListView.cpp
--- a/Engine/Plugins/Experimental/PCG/Source/PCGEditor/Private/SPCGEditorGraphAttributeListView.cpp
+++ b/Engine/Plugins/Experimental/PCG/Source/PCGEditor/Private/SPCGEditorGraphAttributeListView.cpp
@@ -26,6 +26,7 @@
 #include "HAL/PlatformApplicationMisc.h"
 #include "Widgets/Images/SLayeredImage.h"
 #include "Widgets/Images/SThrobber.h"
+#include "Widgets/Input/SButton.h"
 #include "Widgets/Input/SComboBox.h"
 #include "Widgets/Input/SSearchBox.h"
 #include "Widgets/Layout/SScrollBox.h"
@@ -141,7 +142,7 @@ void FPCGListViewUpdater::AsyncSort()
 {
 	if (const FPCGColumnData* Data = ColumnData.Find(SortMode == EColumnSortMode::None ? PCGEditorGraphAttributeListView::NAME_Index : SortingColumn))
 	{
-		if (Data->DataAccessor.IsValid() && Data->DataKeys.IsValid())
+		if (Data->DataAccessor.IsValid() && Data->DataKeys.IsValid() && Data->DataKeys->GetNum() == ListViewItems.Num())
 		{
 			//lambda used here to get the index value of an item in the array for sorting
 			PCGAttributeAccessorHelpers::SortByAttribute(*Data->DataAccessor, *Data->DataKeys, ListViewItems, !(SortMode & EColumnSortMode::Descending), [this](int Index) { return ListViewItems[Index]->Index; });
@@ -311,7 +312,6 @@ SPCGEditorGraphAttributeListView::~SPCGEditorGraphAttributeListView()
 	if (PCGEditorPtr.IsValid())
 	{
 		PCGEditorPtr.Pin()->OnInspectedStackChangedDelegate.RemoveAll(this);
-		PCGEditorPtr.Pin()->OnInspectedNodeChangedDelegate.RemoveAll(this);
 	}
 }
 
@@ -321,7 +321,6 @@ void SPCGEditorGraphAttributeListView::Construct(const FArguments& InArgs, TShar
 	SortingColumn = PCGEditorGraphAttributeListView::NAME_Index;
 
 	PCGEditorPtr.Pin()->OnInspectedStackChangedDelegate.AddSP(this, &SPCGEditorGraphAttributeListView::OnInspectedStackChanged);
-	PCGEditorPtr.Pin()->OnInspectedNodeChangedDelegate.AddSP(this, &SPCGEditorGraphAttributeListView::OnInspectedNodeChanged);
 
 	TextFilter = MakeShareable(new FTextFilterExpressionEvaluator(ETextFilterExpressionEvaluatorMode::Complex));
 
@@ -380,6 +379,19 @@ void SPCGEditorGraphAttributeListView::Construct(const FArguments& InArgs, TShar
 
 	FilterImage->AddLayer(TAttribute<const FSlateBrush*>(this, &SPCGEditorGraphAttributeListView::GetFilterBadgeIcon));
 
+	SAssignNew(LockButton, SButton)
+		.ButtonStyle(FAppStyle::Get(), "SimpleButton")
+		.OnClicked(this, &SPCGEditorGraphAttributeListView::OnLockClick)
+		.ContentPadding(FMargin(4, 2))
+		.HAlign(HAlign_Center)
+		.VAlign(VAlign_Center)
+		.ToolTipText(LOCTEXT("LockSelectionButton_ToolTip", "Locks the current attribute list view to this selection"))
+		[
+			SNew(SImage)
+				.ColorAndOpacity(FSlateColor::UseForeground())
+				.Image(this, &SPCGEditorGraphAttributeListView::OnGetLockButtonImageResource)
+		];
+
 	SAssignNew(FilterButton, SComboButton)
 		.ForegroundColor(FSlateColor::UseStyle())
 		.HasDownArrow(false)
@@ -409,6 +421,12 @@ void SPCGEditorGraphAttributeListView::Construct(const FArguments& InArgs, TShar
 			+SHorizontalBox::Slot()
 			.AutoWidth()
 			.Padding(1.0f, 0.0f)
+			[
+				LockButton->AsShared()
+			]
+			+SHorizontalBox::Slot()
+			.AutoWidth()
+			.Padding(1.0f, 0.0f)
 			[
 				FilterButton->AsShared()
 			]
@@ -554,7 +572,12 @@ void SPCGEditorGraphAttributeListView::OnInspectedStackChanged(const FPCGStack&
 	}
 }
 
-void SPCGEditorGraphAttributeListView::OnInspectedNodeChanged(UPCGEditorGraphNodeBase* InPCGEditorGraphNode)
+UPCGEditorGraphNodeBase* SPCGEditorGraphAttributeListView::GetNodeBeingInspected() const
+{
+	return PCGEditorGraphNode.Get();
+}
+
+void SPCGEditorGraphAttributeListView::SetNodeBeingInspected(UPCGEditorGraphNodeBase* InPCGEditorGraphNode)
 {
 	if (PCGEditorGraphNode == InPCGEditorGraphNode)
 	{
@@ -574,6 +597,9 @@ void SPCGEditorGraphAttributeListView::OnInspectedNodeChanged(UPCGEditorGraphNod
 		NodeNameTextBlock->SetToolTipText(PCGEditorGraphAttributeListView::NoNodeInspectedToolTip);
 	}
 
+	// Always unlock when changing the node, to make sure we unlock when removing the inspected node
+	bIsLocked = false;
+
 	RequestRefresh();
 }
 
@@ -1394,4 +1420,15 @@ bool SPCGEditorGraphAttributeListView::CanCopySelectionToClipboard() const
 	return ListView->GetNumItemsSelected() > 0;
 }
 
+const FSlateBrush* SPCGEditorGraphAttributeListView::OnGetLockButtonImageResource() const
+{
+	return FAppStyle::GetBrush(bIsLocked ? TEXT("PropertyWindow.Locked") : TEXT("PropertyWindow.Unlocked"));
+}
+
+FReply SPCGEditorGraphAttributeListView::OnLockClick()
+{
+	bIsLocked = !bIsLocked;
+	return FReply::Handled();
+}
+
 #undef LOCTEXT_NAMESPACE
diff --git a/Engine/Plugins/Experimental/PCG/Source/PCGEditor/Private/SPCGEditorGraphAttributeListView.h b/Engine/Plugins/Experimental/PCG/Source/PCGEditor/Private/SPCGEditorGraphAttributeListView.h
--- a/Engine/Plugins/Experimental/PCG/Source/PCGEditor/Private/SPCGEditorGraphAttributeListView.h
+++ b/Engine/Plugins/Experimental/PCG/Source/PCGEditor/Private/SPCGEditorGraphAttributeListView.h
@@ -139,11 +139,15 @@ public:
 
 	void RequestRefresh() { bNeedsRefresh = true; }
 
+	UPCGEditorGraphNodeBase* GetNodeBeingInspected() const;
+	void SetNodeBeingInspected(UPCGEditorGraphNodeBase* InNode);
+
+	bool IsLocked() const { return bIsLocked; }
+
 private:
 	TSharedRef<SHeaderRow> CreateHeaderRowWidget() const;
 
 	void OnInspectedStackChanged(const FPCGStack& InPCGStack);
-	void OnInspectedNodeChanged(UPCGEditorGraphNodeBase* InPCGEditorGraphNode);
 
 	void OnGenerateUpdated(UPCGComponent* InPCGComponent);
 
@@ -204,6 +208,11 @@ private:
 	void CopySelectionToClipboard() const;
 	bool CanCopySelectionToClipboard() const;
 
+	/** @return the Slate brush to use for the lock image */
+	const FSlateBrush* OnGetLockButtonImageResource() const;
+
+	FReply OnLockClick();
+
 	/** Pointer back to the PCG editor that owns us */
 	TWeakPtr<FPCGEditor> PCGEditorPtr;
 
@@ -232,6 +241,7 @@ private:
 	TSharedPtr<STextBlock> NodeNameTextBlock;
 	TSharedPtr<STextBlock> InfoTextBlock;
 	TSharedPtr<SComboButton> FilterButton;
+	TSharedPtr<SButton> LockButton;
 
 	TArray<FName> HiddenAttributes;
 
@@ -242,7 +252,10 @@ private:
 	FName SortingColumn = NAME_None;
 	EColumnSortMode::Type SortMode = EColumnSortMode::Type::Ascending;
 
-	bool bNeedsRefresh = false;
+	bool bNeedsRefresh : 1 = false;
+
+	/** True if this property view is currently locked (i.e. The objects being observed are not changed automatically due to user selection)*/
+	bool bIsLocked : 1 = false;
 
 	TSharedPtr<FPCGListViewUpdater> CurrentUpdateTask = nullptr;
 };
diff --git a/Engine/Plugins/Experimental/PCG/Source/PCGEditor/Private/SPCGEditorGraphDetailsView.cpp b/Engine/Plugins/Experimental/PCG/Source/PCGEditor/Private/SPCGEditorGraphDetailsView.cpp
new file mode 100644
--- /dev/null
+++ b/Engine/Plugins/Experimental/PCG/Source/PCGEditor/Private/SPCGEditorGraphDetailsView.cpp
@@ -0,0 +1,241 @@
+// Copyright Epic Games, Inc. All Rights Reserved.
+
+#include "SPCGEditorGraphDetailsView.h"
+
+#include "PCGGraph.h"
+#include "PCGNode.h"
+#include "PCGSettings.h"
+
+#include "DetailsViewArgs.h"
+#include "IDetailsView.h"
+#include "PropertyEditorModule.h"
+#include "Layout/Visibility.h"
+#include "Modules/ModuleManager.h"
+#include "Widgets/DeclarativeSyntaxSupport.h"
+#include "Widgets/Images/SImage.h"
+#include "Widgets/Input/SButton.h"
+#include "Widgets/Text/STextBlock.h"
+
+#define LOCTEXT_NAMESPACE "PCGEditorGraphDetailsView"
+
+void SPCGEditorGraphDetailsView::Construct(const FArguments& InArgs)
+{
+	FDetailsViewArgs DetailsViewArgs;
+	DetailsViewArgs.NameAreaSettings = FDetailsViewArgs::HideNameArea;
+	DetailsViewArgs.bCustomFilterAreaLocation = true;
+	DetailsViewArgs.bHideSelectionTip = true;
+	DetailsViewArgs.DefaultsOnlyVisibility = EEditDefaultsOnlyNodeVisibility::Automatic;
+
+	FPropertyEditorModule& PropertyEditorModule = FModuleManager::GetModuleChecked<FPropertyEditorModule>("PropertyEditor");
+	DetailsView = PropertyEditorModule.CreateDetailView(DetailsViewArgs);
+
+	DetailsView->SetIsPropertyReadOnlyDelegate(FIsPropertyReadOnly::CreateSP(this, &SPCGEditorGraphDetailsView::IsReadOnlyProperty));
+	DetailsView->SetIsPropertyVisibleDelegate(FIsPropertyVisible::CreateSP(this, &SPCGEditorGraphDetailsView::IsVisibleProperty));
+
+	ChildSlot
+	[
+		SNew(SVerticalBox)
+		+SVerticalBox::Slot()
+		.Padding(0.f, 0.f, 0.f, 0.f)
+		.AutoHeight()
+		[
+			SNew(SHorizontalBox)
+			+SHorizontalBox::Slot()
+			.AutoWidth()
+			[
+				SNew(SButton)
+				.ButtonStyle(FAppStyle::Get(), "SimpleButton")
+				.OnClicked(this, &SPCGEditorGraphDetailsView::OnLockButtonClicked)
+				.ContentPadding(FMargin(4, 2))
+				.HAlign(HAlign_Center)
+				.VAlign(VAlign_Center)
+				.ToolTipText(LOCTEXT("LockSelectionButton_ToolTip", "Locks the current attribute list view to this selection"))
+				[
+					SNew(SImage)
+					.ColorAndOpacity(FSlateColor::UseForeground())
+					.Image(this, &SPCGEditorGraphDetailsView::GetLockIcon)
+				]
+			]
+			+SHorizontalBox::Slot()
+			.AutoWidth()
+			.HAlign(HAlign_Center)
+			.VAlign(VAlign_Center)
+			[
+				SNew(STextBlock)
+				.Visibility(this, &SPCGEditorGraphDetailsView::GetNameVisibility)
+				.Text(this, &SPCGEditorGraphDetailsView::GetName)
+			]
+			+SHorizontalBox::Slot()
+			[
+				DetailsView->GetFilterAreaWidget().ToSharedRef()
+			]
+		]
+		+SVerticalBox::Slot()
+		[
+			DetailsView.ToSharedRef()
+		]
+	];
+}
+
+void SPCGEditorGraphDetailsView::SetObject(UObject* InObject, bool bForceRefresh)
+{
+	TArray<UObject*> InObjects;
+	InObjects.Add(InObject);
+
+	SetObjects(InObjects, bForceRefresh);
+}
+
+void SPCGEditorGraphDetailsView::SetObjects(const TArray<UObject*>& InObjects, bool bForceRefresh, bool bOverrideLock)
+{
+	if (bOverrideLock || !bIsLocked)
+	{
+		DetailsView->SetObjects(InObjects, bForceRefresh, bOverrideLock);
+
+		// Unlock the view automatically if we are viewing nothing
+		if (InObjects.IsEmpty())
+		{
+			bIsLocked = false;
+		}
+	}
+}
+
+void SPCGEditorGraphDetailsView::SetObjects(const TArray<TWeakObjectPtr<UObject>>& InObjects, bool bForceRefresh, bool bOverrideLock)
+{
+	TArray<UObject*> Objects;
+	Objects.Reserve(InObjects.Num());
+	for (const TWeakObjectPtr<UObject>& InObject : InObjects)
+	{
+		if (InObject.IsValid())
+		{
+			Objects.Add(InObject.Get());
+		}
+	}
+
+	SetObjects(Objects, bForceRefresh, bOverrideLock);
+}
+
+const TArray<TWeakObjectPtr<UObject>>& SPCGEditorGraphDetailsView::GetSelectedObjects() const
+{
+	return DetailsView->GetSelectedObjects();
+}
+
+FReply SPCGEditorGraphDetailsView::OnLockButtonClicked()
+{
+	bIsLocked = !bIsLocked;
+	return FReply::Handled();
+}
+
+const FSlateBrush* SPCGEditorGraphDetailsView::GetLockIcon() const
+{
+	return FAppStyle::GetBrush(bIsLocked ? TEXT("PropertyWindow.Locked") : TEXT("PropertyWindow.Unlocked"));
+}
+
+EVisibility SPCGEditorGraphDetailsView::GetNameVisibility() const
+{
+	return bIsLocked ? EVisibility::Visible : EVisibility::Collapsed;
+}
+
+FText SPCGEditorGraphDetailsView::GetName() const
+{
+	const TArray<TWeakObjectPtr<UObject>>& SelectedObjects = GetSelectedObjects();
+	if (SelectedObjects.IsEmpty())
+	{
+		// Implementation note: this doesn't really happen since we reselect the graph settings when unselecting
+		return LOCTEXT("NoObjectsSelected", "Empty selection");
+	}
+	else if (SelectedObjects.Num() == 1)
+	{
+		UObject* SelectedObject = SelectedObjects[0].Get();
+		UPCGNode* OwnerNode = nullptr;
+
+		if (Cast<UPCGSettings>(SelectedObject) != nullptr || Cast<UPCGSettingsInstance>(SelectedObject) != nullptr)
+		{
+			OwnerNode = Cast<UPCGNode>(SelectedObject->GetOuter());
+		}
+
+		if (OwnerNode)
+		{
+			return OwnerNode->GetNodeTitle(EPCGNodeTitleType::ListView);
+		}
+		else if (UPCGGraph* Graph = Cast<UPCGGraph>(SelectedObject))
+		{
+			return LOCTEXT("GraphSettingsSelected", "Graph Settings");
+		}
+		else
+		{
+			return FText::FromName(SelectedObject->GetFName());
+		}
+	}
+	else
+	{
+		return FText::Format(LOCTEXT("MultipleObjectsSelectedFmt", "{0} nodes"), FText::AsNumber(SelectedObjects.Num()));
+	}
+}
+
+bool SPCGEditorGraphDetailsView::IsReadOnlyProperty(const FPropertyAndParent& InPropertyAndParent) const
+{
+	// Everything is writable when not in an instance
+	if (!DetailsView ||
+		InPropertyAndParent.ParentProperties.IsEmpty() ||
+		InPropertyAndParent.ParentProperties.Last()->GetFName() != GET_MEMBER_NAME_CHECKED(UPCGSettingsInstance, Settings))
+	{
+		return false;
+	}
+
+	const TArray<TWeakObjectPtr<UObject>>& SelectedObjects = DetailsView->GetSelectedObjects();
+	for (const TWeakObjectPtr<UObject>& SelectedObject : SelectedObjects)
+	{
+		if (!SelectedObject.IsValid())
+		{
+			continue;
+		}
+
+		if (UPCGSettingsInstance* Instance = Cast<UPCGSettingsInstance>(SelectedObject.Get()))
+		{
+			return true;
+		}
+	}
+
+	return false;
+}
+
+bool SPCGEditorGraphDetailsView::IsVisibleProperty(const FPropertyAndParent& InPropertyAndParent) const
+{
+	if (!DetailsView)
+	{
+		return true;
+	}
+
+	// Currently never hide anything from the graph settings
+	if (InPropertyAndParent.Objects.Num() == 1 && Cast<UPCGGraph>(InPropertyAndParent.Objects[0]))
+	{
+		return true;
+	}
+
+	// Always hide asset info information
+	if (InPropertyAndParent.Property.HasMetaData(TEXT("Category")) &&
+		InPropertyAndParent.Property.GetMetaData(TEXT("Category")) == TEXT("AssetInfo"))
+	{
+		return false;
+	}
+
+	// Otherwise, everything is visible when not in an instance
+	if (InPropertyAndParent.ParentProperties.IsEmpty() ||
+		InPropertyAndParent.ParentProperties.Last()->GetFName() != GET_MEMBER_NAME_CHECKED(UPCGSettingsInstance, Settings))
+	{
+		return true;
+	}
+
+	// Hide debug settings from the setting when showing the instance settings.
+	if (InPropertyAndParent.Property.GetFName() == GET_MEMBER_NAME_CHECKED(UPCGSettings, bEnabled) ||
+		InPropertyAndParent.Property.GetFName() == GET_MEMBER_NAME_CHECKED(UPCGSettings, bDebug) ||
+		(InPropertyAndParent.ParentProperties.Num() >= 2 &&
+			InPropertyAndParent.ParentProperties[InPropertyAndParent.ParentProperties.Num() - 2]->GetFName() == GET_MEMBER_NAME_CHECKED(UPCGSettings, DebugSettings)))
+	{
+		return false;
+	}
+
+	return true;
+}
+
+#undef LOCTEXT_NAMESPACE
\ No newline at end of file
diff --git a/Engine/Plugins/Experimental/PCG/Source/PCGEditor/Private/SPCGEditorGraphDetailsView.h b/Engine/Plugins/Experimental/PCG/Source/PCGEditor/Private/SPCGEditorGraphDetailsView.h
new file mode 100644
--- /dev/null
+++ b/Engine/Plugins/Experimental/PCG/Source/PCGEditor/Private/SPCGEditorGraphDetailsView.h
@@ -0,0 +1,57 @@
+// Copyright Epic Games, Inc. All Rights Reserved.
+
+#pragma once
+
+#include "Templates/SharedPointer.h"
+#include "Widgets/SCompoundWidget.h"
+
+class IDetailsView;
+struct FPropertyAndParent;
+
+/** Thin wrapper on top of a SDetailsView to add some PCG-specific logic around locking and improve layout */
+class SPCGEditorGraphDetailsView : public SCompoundWidget
+{
+public:
+	SLATE_BEGIN_ARGS(SPCGEditorGraphDetailsView)
+		{}	
+	SLATE_END_ARGS()
+
+	/** Constructs this widget with InArgs */
+	void Construct(const FArguments& InArgs);
+
+	TSharedPtr<IDetailsView> GetDetailsView() const { return DetailsView; }
+
+	//~Begin IDetailsView-like interface
+	/** Returns the lock state of the details view, e.g. whether the selected object should change on a new selection */
+	bool IsLocked() const { return bIsLocked; }
+
+	/** Changes the selected objects in the details view. Will unlock automatically. */
+	void SetObject(UObject* InObject, bool bForceRefresh = false);
+
+	/** Changes the selected objects in the details view. Will unlock automatically. */
+	void SetObjects(const TArray<UObject*>& InObjects, bool bForceRefresh = false, bool bOverrideLock = false);
+
+	/** Changes the selected objects in the details view. Will unlock automatically. */
+	void SetObjects(const TArray<TWeakObjectPtr<UObject>>& InObjects, bool bForceRefresh = false, bool bOverrideLock = false);
+
+	/** Retrieves the list of objects currently shown in the details view */
+	const TArray<TWeakObjectPtr<UObject>>& GetSelectedObjects() const;
+	//~End IDetailsView-like interface
+
+protected:
+	/** Returns whether a property should be readonly (used for instances) */
+	bool IsReadOnlyProperty(const FPropertyAndParent& InPropertyAndParent) const;
+	/** Returns whether a property should be visible (used for instance vs. settings properties) */
+	bool IsVisibleProperty(const FPropertyAndParent& InPropertyAndParent) const;
+
+	FReply OnLockButtonClicked();
+	const FSlateBrush* GetLockIcon() const;
+
+	/** Gets the current visibility state of the name section, which is related to whether the details view is locked. */
+	EVisibility GetNameVisibility() const;
+	/** Returns the name to show in the details-view name section, e.g. the settings title most of the time. */
+	FText GetName() const;
+
+	TSharedPtr<IDetailsView> DetailsView;
+	bool bIsLocked : 1 = false;
+};
\ No newline at end of file
