Removed loaded-preview components (from partition actors) that aren't active anymore from the debug object tree view

#rnx

diff --git a/Engine/Plugins/Experimental/PCG/Source/PCGEditor/Private/SPCGEditorGraphDebugObjectTree.cpp b/Engine/Plugins/Experimental/PCG/Source/PCGEditor/Private/SPCGEditorGraphDebugObjectTree.cpp
--- a/Engine/Plugins/Experimental/PCG/Source/PCGEditor/Private/SPCGEditorGraphDebugObjectTree.cpp
+++ b/Engine/Plugins/Experimental/PCG/Source/PCGEditor/Private/SPCGEditorGraphDebugObjectTree.cpp
@@ -748,7 +748,7 @@ void SPCGEditorGraphDebugObjectTree::RefreshTree()
 		}
 
 		UPCGComponent* PCGComponent = Cast<UPCGComponent>(PCGComponentObject);
-		if (!PCGComponent)
+		if (!PCGComponent || !PCGComponent->IsRegistered())
 		{
 			continue;
 		}
