[PCG] SplineSampler mode that samples a given number of evenly spaced points

diff --git a/Engine/Plugins/Experimental/PCG/Source/PCG/Private/Data/PCGPolyLineData.cpp b/Engine/Plugins/Experimental/PCG/Source/PCG/Private/Data/PCGPolyLineData.cpp
--- a/Engine/Plugins/Experimental/PCG/Source/PCG/Private/Data/PCGPolyLineData.cpp
+++ b/Engine/Plugins/Experimental/PCG/Source/PCG/Private/Data/PCGPolyLineData.cpp
@@ -20,6 +20,18 @@ FBox UPCGPolyLineData::GetBounds() const
 	return Bounds;
 }
 
+FVector::FReal UPCGPolyLineData::GetLength() const
+{
+	FVector::FReal Length = 0.0;
+
+	for (int SegmentIndex = 0; SegmentIndex < GetNumSegments(); ++SegmentIndex)
+	{
+		Length += GetSegmentLength(SegmentIndex);
+	}
+
+	return Length;
+}
+
 void UPCGPolyLineData::GetTangentsAtSegmentStart(int SegmentIndex, FVector& OutArriveTangent, FVector& OutLeaveTangent) const
 {
 	OutArriveTangent = FVector::Zero();
diff --git a/Engine/Plugins/Experimental/PCG/Source/PCG/Private/Elements/PCGSplineSampler.cpp b/Engine/Plugins/Experimental/PCG/Source/PCG/Private/Elements/PCGSplineSampler.cpp
--- a/Engine/Plugins/Experimental/PCG/Source/PCG/Private/Elements/PCGSplineSampler.cpp
+++ b/Engine/Plugins/Experimental/PCG/Source/PCG/Private/Elements/PCGSplineSampler.cpp
@@ -782,6 +782,28 @@ namespace PCGSplineSampler
 
 		FStepSampler* Sampler = ((Params.Mode == EPCGSplineSamplingMode::Subdivision) ? static_cast<FStepSampler*>(&SubdivisionSampler) : static_cast<FStepSampler*>(&DistanceSampler));
 
+		if (Params.Mode == EPCGSplineSamplingMode::NumberOfSamples)
+		{
+			if (Params.NumSamples <= 0)
+			{
+				return;
+			}
+
+			// Compute an increment which evenly distributes sample points along the length of the curve.
+			DistanceSampler.DistanceIncrement = LineData->GetLength() / (LineData->IsClosed() ? Params.NumSamples : FMath::Max(1, Params.NumSamples - 1));
+
+			if (LineData->IsClosed() || Params.NumSamples == 1)
+			{
+				// If the curve is closed or only has one sample, we nudge the DistanceIncrement slightly to avoid floating point error giving us an extra sample point.
+				DistanceSampler.DistanceIncrement += UE_DOUBLE_SMALL_NUMBER;
+			}
+			else
+			{
+				// If the curve is not closed and has more than one sample, we should nudge DistanceIncrement slightly lower so that we guarantee capturing the last sample point.
+				DistanceSampler.DistanceIncrement -= UE_DOUBLE_SMALL_NUMBER;
+			}
+		}
+
 		FDimensionSampler TrivialDimensionSampler(LineData, InBoundingShapeData, InProjectionTarget, InProjectionParams, Params, OutPointData);
 		FVolumeSampler VolumeSampler(LineData, InBoundingShapeData, InProjectionTarget, InProjectionParams, Params, OutPointData);
 
diff --git a/Engine/Plugins/Experimental/PCG/Source/PCG/Public/Data/PCGPolyLineData.h b/Engine/Plugins/Experimental/PCG/Source/PCG/Public/Data/PCGPolyLineData.h
--- a/Engine/Plugins/Experimental/PCG/Source/PCG/Public/Data/PCGPolyLineData.h
+++ b/Engine/Plugins/Experimental/PCG/Source/PCG/Public/Data/PCGPolyLineData.h
@@ -24,6 +24,7 @@ public:
 	virtual FTransform GetTransform() const { return FTransform::Identity; }
 	virtual int GetNumSegments() const PURE_VIRTUAL(UPCGPolyLineData::GetNumSegments, return 0;);
 	virtual FVector::FReal GetSegmentLength(int SegmentIndex) const PURE_VIRTUAL(UPCGPolyLineData::GetSegmentLength, return 0;);
+	virtual FVector::FReal GetLength() const;
 	virtual FTransform GetTransformAtDistance(int SegmentIndex, FVector::FReal Distance, bool bWorldSpace = true, FBox* OutBounds = nullptr) const PURE_VIRTUAL(UPCGPolyLine::GetTransformAtDistance, return FTransform(););
 	virtual FVector::FReal GetCurvatureAtDistance(int SegmentIndex, FVector::FReal Distance) const { return 0; }
 	virtual void GetTangentsAtSegmentStart(int SegmentIndex, FVector& OutArriveTangent, FVector& OutLeaveTangent) const;
diff --git a/Engine/Plugins/Experimental/PCG/Source/PCG/Public/Elements/PCGSplineSampler.h b/Engine/Plugins/Experimental/PCG/Source/PCG/Public/Elements/PCGSplineSampler.h
--- a/Engine/Plugins/Experimental/PCG/Source/PCG/Public/Elements/PCGSplineSampler.h
+++ b/Engine/Plugins/Experimental/PCG/Source/PCG/Public/Elements/PCGSplineSampler.h
@@ -17,7 +17,8 @@ UENUM()
 enum class EPCGSplineSamplingMode : uint8
 {
 	Subdivision = 0,
-	Distance
+	Distance,
+	NumberOfSamples UMETA(Tooltip = "Samples a specified number of times, evenly spaced around the spline.")
 };
 
 UENUM()
@@ -64,6 +65,9 @@ struct PCG_API FPCGSplineSamplerParams
 	UPROPERTY(BlueprintReadWrite, EditAnywhere, Category = Settings, meta = (ClampMin = "0.1", EditCondition = "Mode==EPCGSplineSamplingMode::Distance&&Dimension!=EPCGSplineSamplingDimension::OnInterior"))
 	float DistanceIncrement = 100.0f;
 
+	UPROPERTY(BlueprintReadWrite, EditAnywhere, Category = Settings, meta = (ClampMin = "0", EditCondition = "Mode==EPCGSplineSamplingMode::NumberOfSamples&&Dimension!=EPCGSplineSamplingDimension::OnInterior"))
+	int32 NumSamples = 8;
+
 	UPROPERTY(BlueprintReadWrite, EditAnywhere, Category = Settings, meta = (ClampMin = "0", EditCondition = "Dimension==EPCGSplineSamplingDimension::OnHorizontal||Dimension==EPCGSplineSamplingDimension::OnVolume"))
 	int32 NumPlanarSubdivisions = 8;
 
