Fixed localization errors

#rnx

diff --git a/Engine/Plugins/Experimental/PCG/Source/PCGEditor/Private/PCGEditor.cpp b/Engine/Plugins/Experimental/PCG/Source/PCGEditor/Private/PCGEditor.cpp
--- a/Engine/Plugins/Experimental/PCG/Source/PCGEditor/Private/PCGEditor.cpp
+++ b/Engine/Plugins/Experimental/PCG/Source/PCGEditor/Private/PCGEditor.cpp
@@ -382,19 +382,19 @@ void FPCGEditor::RegisterTabSpawners(const TSharedRef<FTabManager>& InTabManager
 		.SetGroup(WorkspaceMenuCategoryRef);
 
 	InTabManager->RegisterTabSpawner(FPCGEditor_private::PropertyDetailsID[0], FOnSpawnTab::CreateSP(this, &FPCGEditor::SpawnTab_PropertyDetails, 0))
-		.SetDisplayName(LOCTEXT("DetailsTab", "Details 1"))
+		.SetDisplayName(LOCTEXT("DetailsTab1", "Details 1"))
 		.SetGroup(DetailsGroup);
 
 	InTabManager->RegisterTabSpawner(FPCGEditor_private::PropertyDetailsID[1], FOnSpawnTab::CreateSP(this, &FPCGEditor::SpawnTab_PropertyDetails, 1))
-		.SetDisplayName(LOCTEXT("DetailsTab", "Details 2"))
+		.SetDisplayName(LOCTEXT("DetailsTab2", "Details 2"))
 		.SetGroup(DetailsGroup);
 
 	InTabManager->RegisterTabSpawner(FPCGEditor_private::PropertyDetailsID[2], FOnSpawnTab::CreateSP(this, &FPCGEditor::SpawnTab_PropertyDetails, 2))
-		.SetDisplayName(LOCTEXT("DetailsTab", "Details 3"))
+		.SetDisplayName(LOCTEXT("DetailsTab3", "Details 3"))
 		.SetGroup(DetailsGroup);
 
 	InTabManager->RegisterTabSpawner(FPCGEditor_private::PropertyDetailsID[3], FOnSpawnTab::CreateSP(this, &FPCGEditor::SpawnTab_PropertyDetails, 3))
-		.SetDisplayName(LOCTEXT("DetailsTab", "Details 4"))
+		.SetDisplayName(LOCTEXT("DetailsTab4", "Details 4"))
 		.SetGroup(DetailsGroup);
 
 	InTabManager->RegisterTabSpawner(FPCGEditor_private::PaletteID, FOnSpawnTab::CreateSP(this, &FPCGEditor::SpawnTab_Palette))
@@ -406,19 +406,19 @@ void FPCGEditor::RegisterTabSpawners(const TSharedRef<FTabManager>& InTabManager
 		.SetGroup(WorkspaceMenuCategoryRef);
 
 	InTabManager->RegisterTabSpawner(FPCGEditor_private::AttributesID[0], FOnSpawnTab::CreateSP(this, &FPCGEditor::SpawnTab_Attributes, 0))
-		.SetDisplayName(LOCTEXT("AttributesTab", "Attributes 1"))
+		.SetDisplayName(LOCTEXT("AttributesTab1", "Attributes 1"))
 		.SetGroup(AttributesGroup);
 
 	InTabManager->RegisterTabSpawner(FPCGEditor_private::AttributesID[1], FOnSpawnTab::CreateSP(this, &FPCGEditor::SpawnTab_Attributes, 1))
-		.SetDisplayName(LOCTEXT("AttributesTab", "Attributes 2"))
+		.SetDisplayName(LOCTEXT("AttributesTab2", "Attributes 2"))
 		.SetGroup(AttributesGroup);
 
 	InTabManager->RegisterTabSpawner(FPCGEditor_private::AttributesID[2], FOnSpawnTab::CreateSP(this, &FPCGEditor::SpawnTab_Attributes, 2))
-		.SetDisplayName(LOCTEXT("AttributesTab", "Attributes 3"))
+		.SetDisplayName(LOCTEXT("AttributesTab3", "Attributes 3"))
 		.SetGroup(AttributesGroup);
 
 	InTabManager->RegisterTabSpawner(FPCGEditor_private::AttributesID[3], FOnSpawnTab::CreateSP(this, &FPCGEditor::SpawnTab_Attributes, 3))
-		.SetDisplayName(LOCTEXT("AttributesTab", "Attributes 4"))
+		.SetDisplayName(LOCTEXT("AttributesTab4", "Attributes 4"))
 		.SetGroup(AttributesGroup);
 
 	InTabManager->RegisterTabSpawner(FPCGEditor_private::FindID, FOnSpawnTab::CreateSP(this, &FPCGEditor::SpawnTab_Find))
