[PCG] Fix AttributeFilter bug, where points could be accepted while they should not.

* Also add UnitTest

#rnx
#11177

diff --git a/Engine/Plugins/Experimental/PCG/Source/PCG/Private/Elements/PCGAttributeFilter.cpp b/Engine/Plugins/Experimental/PCG/Source/PCG/Private/Elements/PCGAttributeFilter.cpp
--- a/Engine/Plugins/Experimental/PCG/Source/PCG/Private/Elements/PCGAttributeFilter.cpp
+++ b/Engine/Plugins/Experimental/PCG/Source/PCG/Private/Elements/PCGAttributeFilter.cpp
@@ -630,7 +630,12 @@ bool FPCGAttributeFilterElementBase::DoFiltering(FPCGContext* Context, EPCGAttri
 			TargetValues.SetNum(PCGAttributeFilterConstants::ChunkSize);
 			FirstThresholdValues.SetNum(PCGAttributeFilterConstants::ChunkSize);
 			SecondThresholdValues.SetNum(PCGAttributeFilterConstants::ChunkSize);
-			SkipTests.SetNumZeroed(PCGAttributeFilterConstants::ChunkSize); // All initialized to false
+
+			const bool bShouldSample = FirstThresholdInfo.ThresholdPointData || SecondThresholdInfo.ThresholdPointData;
+			if (bShouldSample)
+			{
+				SkipTests.SetNumUninitialized(PCGAttributeFilterConstants::ChunkSize);
+			}
 
 			const int32 NumberOfIterations = (NumberOfEntries + PCGAttributeFilterConstants::ChunkSize - 1) / PCGAttributeFilterConstants::ChunkSize;
 
@@ -642,6 +647,12 @@ bool FPCGAttributeFilterElementBase::DoFiltering(FPCGContext* Context, EPCGAttri
 				TArrayView<Type> FirstThresholdView(FirstThresholdValues.GetData(), Range);
 				TArrayView<Type> SecondThresholdView(SecondThresholdValues.GetData(), Range);
 
+				// Need to reset the skip tests to False
+				if (bShouldSample)
+				{
+					FMemory::Memzero(SkipTests.GetData(), SkipTests.Num() * sizeof(bool));
+				}
+
 				// Sampling the points if needed
 				auto SamplePointData = [Range, StartIndex, &OperationData, &SkipTests](UPCGPointData* InPointData, const UPCGSpatialData* InSpatialData)
 				{
@@ -668,11 +679,14 @@ bool FPCGAttributeFilterElementBase::DoFiltering(FPCGContext* Context, EPCGAttri
 					}
 				};
 
-				SamplePointData(FirstThresholdInfo.ThresholdPointData, FirstThresholdInfo.ThresholdSpatialData);
-
-				if (FirstThresholdInfo.ThresholdPointData != SecondThresholdInfo.ThresholdPointData)
+				if (bShouldSample)
 				{
-					SamplePointData(SecondThresholdInfo.ThresholdPointData, SecondThresholdInfo.ThresholdSpatialData);
+					SamplePointData(FirstThresholdInfo.ThresholdPointData, FirstThresholdInfo.ThresholdSpatialData);
+
+					if (FirstThresholdInfo.ThresholdPointData != SecondThresholdInfo.ThresholdPointData)
+					{
+						SamplePointData(SecondThresholdInfo.ThresholdPointData, SecondThresholdInfo.ThresholdSpatialData);
+					}
 				}
 
 				// If ThresholdView point on ThresholdPointData points, there are only "ChunkSize" points in it.
@@ -703,7 +717,7 @@ bool FPCGAttributeFilterElementBase::DoFiltering(FPCGContext* Context, EPCGAttri
 
 				for (int32 j = 0; j < Range; ++j)
 				{
-					if (SkipTests[j])
+					if (bShouldSample && SkipTests[j])
 					{
 						AddElement(true, StartIndex + j);
 						continue;
diff --git a/Engine/Plugins/Experimental/PCG/Source/PCG/Private/Tests/Elements/PCGAttributeFilterTest.cpp b/Engine/Plugins/Experimental/PCG/Source/PCG/Private/Tests/Elements/PCGAttributeFilterTest.cpp
--- a/Engine/Plugins/Experimental/PCG/Source/PCG/Private/Tests/Elements/PCGAttributeFilterTest.cpp
+++ b/Engine/Plugins/Experimental/PCG/Source/PCG/Private/Tests/Elements/PCGAttributeFilterTest.cpp
@@ -18,11 +18,13 @@ IMPLEMENT_CUSTOM_SIMPLE_AUTOMATION_TEST(FPCGPointFilterDensity, FPCGTestBaseClas
 IMPLEMENT_CUSTOM_SIMPLE_AUTOMATION_TEST(FPCGPointFilterDensityRange, FPCGTestBaseClass, "Plugins.PCG.AttributeFilter.Points.DensityRange", PCGTestsCommon::TestFlags)
 IMPLEMENT_CUSTOM_SIMPLE_AUTOMATION_TEST(FPCGAttributeFilterInt, FPCGTestBaseClass, "Plugins.PCG.AttributeFilter.Params.Int", PCGTestsCommon::TestFlags)
 IMPLEMENT_CUSTOM_SIMPLE_AUTOMATION_TEST(FPCGAttributeFilterIntRange, FPCGTestBaseClass, "Plugins.PCG.AttributeFilter.Params.IntRange", PCGTestsCommon::TestFlags)
+IMPLEMENT_CUSTOM_SIMPLE_AUTOMATION_TEST(FPCGAttributeFilterSkipTestBug, FPCGTestBaseClass, "Plugins.PCG.AttributeFilter.SkipTestBug", PCGTestsCommon::TestFlags)
 
 namespace PCGPointFilterTest
 {
 	const FName InsideFilterLabel = TEXT("InsideFilter");
 	const FName OutsideFilterLabel = TEXT("OutsideFilter");
+	const FName FilterLabel = TEXT("Filter");
 
 	UPCGPointData* GeneratePointDataWithRandomDensity(int32 InNumPoints, int32 InRandomSeed)
 	{
@@ -295,5 +297,83 @@ bool FPCGAttributeFilterIntRange::RunTest(const FString& Parameters)
 		UTEST_TRUE(*FString::Printf(TEXT("Attribute has a value (%d) outside the range [%d, %d]"), Value, (int32)IntMinThreshold, (int32)IntMaxThreshold), Value < (int32)IntMinThreshold || Value > (int32)IntMaxThreshold);
 	}
 
+	return true;
+}
+
+// We do not reset SkipTests in the case of point sampling, resulting to accepting points that should have not been accepted.
+// This fails before fixed CL of UE-201595.
+bool FPCGAttributeFilterSkipTestBug::RunTest(const FString& Parameters)
+{
+	PCGTestsCommon::FTestData TestData;
+	UPCGAttributeFilteringSettings* Settings = PCGTestsCommon::GenerateSettings<UPCGAttributeFilteringSettings>(TestData);
+	check(Settings);
+
+	Settings->TargetAttribute.SetPointProperty(EPCGPointProperties::Density);
+	Settings->ThresholdAttribute.SetPointProperty(EPCGPointProperties::Density);
+	Settings->Operator = EPCGAttributeFilterOperator::Equal;
+	Settings->bUseSpatialQuery = true;
+
+	FPCGElementPtr TestElement = TestData.Settings->GetElement();
+	UPCGPointData* InputPointData = NewObject<UPCGPointData>();
+	UPCGPointData* ThresholdPointData = NewObject<UPCGPointData>();
+
+	TArray<FPCGPoint>& InputPoints = InputPointData->GetMutablePoints();
+	TArray<FPCGPoint>& ThresholdPoints = ThresholdPointData->GetMutablePoints();
+
+	// Take a big number to make sure we go over the 256 default chunk size
+	constexpr int32 NumPoints = 2048;
+	constexpr int32 HalfNumPoints = NumPoints / 2;
+	InputPoints.Reserve(NumPoints);
+	ThresholdPoints.Reserve(NumPoints);
+
+	for (int32 i = 0; i < NumPoints; ++i)
+	{
+		// First half are very different so the sampling should fail, second half are the same the sampling to succeed but the filtering to fail.
+		if (i >= HalfNumPoints)
+		{
+			FPCGPoint& Point = InputPoints.Emplace_GetRef();
+			Point.Transform.SetLocation(FVector(10 * i, 10 * i, 10 * i));
+			FPCGPoint& ThresholdPoint = ThresholdPoints.Add_GetRef(Point);
+			ThresholdPoint.Density = 0.5f;
+		}
+		else
+		{
+			FPCGPoint& Point = InputPoints.Emplace_GetRef();
+			Point.Transform.SetLocation(FVector(i, i, i));
+			FPCGPoint& ThresholdPoint = ThresholdPoints.Emplace_GetRef();
+			ThresholdPoint.Transform.SetLocation(FVector(-10 * i - 1000, -10 * i - 1000, -10 * i - 1000));
+		}
+	}
+
+	FPCGTaggedData& TaggedData = TestData.InputData.TaggedData.Emplace_GetRef();
+	TaggedData.Pin = PCGPinConstants::DefaultInputLabel;
+	TaggedData.Data = InputPointData;
+
+	FPCGTaggedData& SecondTaggedData = TestData.InputData.TaggedData.Emplace_GetRef();
+	SecondTaggedData.Pin = PCGPointFilterTest::FilterLabel;
+	SecondTaggedData.Data = ThresholdPointData;
+
+	TUniquePtr<FPCGContext> Context = TestData.InitializeTestContext();
+
+	while (!TestElement->Execute(Context.Get())) {}
+
+	// We should have outputs on both
+	TArray<FPCGTaggedData> InFilterOutput = Context->OutputData.GetInputsByPin(PCGPointFilterTest::InsideFilterLabel);
+	TArray<FPCGTaggedData> OutFilterOutput = Context->OutputData.GetInputsByPin(PCGPointFilterTest::OutsideFilterLabel);
+
+	UTEST_EQUAL(TEXT("InFilter pin has 1 output"), InFilterOutput.Num(), 1);
+	UTEST_EQUAL(TEXT("OutFilter pin has 1 output"), OutFilterOutput.Num(), 1);
+
+	// Making sure the output is a PointData
+	const UPCGPointData* InFilterPointData = Cast<UPCGPointData>(InFilterOutput[0].Data);
+	const UPCGPointData* OutFilterPointData = Cast<UPCGPointData>(OutFilterOutput[0].Data);
+
+	UTEST_NOT_NULL(TEXT("InFilter data is a point data"), InFilterPointData);
+	UTEST_NOT_NULL(TEXT("OutFilter data is a point data"), OutFilterPointData);
+
+	// We should have the inside filter with half the points and the outside filter to the rest
+	UTEST_EQUAL(TEXT("InFilter data has the right number of points"), InFilterPointData->GetPoints().Num(), HalfNumPoints);
+	UTEST_EQUAL(TEXT("OutFilter data has the right number of points"), OutFilterPointData->GetPoints().Num(), HalfNumPoints);
+
 	return true;
 }
\ No newline at end of file
