Improved tab summoning and added input chord (F4) to open a new details view if needed.
Added clickable node names in the details view & attribute list view to jump to node for easier navigation.

#rnx

diff --git a/Engine/Plugins/Experimental/PCG/Source/PCGEditor/Private/PCGEditor.cpp b/Engine/Plugins/Experimental/PCG/Source/PCGEditor/Private/PCGEditor.cpp
--- a/Engine/Plugins/Experimental/PCG/Source/PCGEditor/Private/PCGEditor.cpp
+++ b/Engine/Plugins/Experimental/PCG/Source/PCGEditor/Private/PCGEditor.cpp
@@ -128,6 +128,7 @@ void FPCGEditor::Initialize(const EToolkitMode::Type InMode, const TSharedPtr<cl
 	for (int PropertyDetailsIndex = 0; PropertyDetailsIndex < 4; ++PropertyDetailsIndex)
 	{
 		TSharedRef<SPCGEditorGraphDetailsView> PropertyDetailsWidget = SNew(SPCGEditorGraphDetailsView);
+		PropertyDetailsWidget->SetEditor(SharedThis(this));
 		PropertyDetailsWidget->SetObject(PCGGraphBeingEdited);
 		PropertyDetailsWidgets.Add(PropertyDetailsWidget);
 	}
@@ -595,6 +596,10 @@ void FPCGEditor::BindCommands()
 		PCGEditorCommands.Find,
 		FExecuteAction::CreateSP(this, &FPCGEditor::OnFind));
 
+	ToolkitCommands->MapAction(
+		PCGEditorCommands.ShowSelectedDetails,
+		FExecuteAction::CreateSP(this, &FPCGEditor::OpenDetailsView));
+
 	ToolkitCommands->MapAction(
 		PCGEditorCommands.PauseAutoRegeneration,
 		FExecuteAction::CreateSP(this, &FPCGEditor::OnPauseAutomaticRegeneration_Clicked),
@@ -696,6 +701,102 @@ void FPCGEditor::OnFind()
 	}
 }
 
+void FPCGEditor::OpenDetailsView()
+{
+	if (TabManager.IsValid())
+	{
+		auto InvokeFirstUnlockedTab = [this](bool bVisibleOnly) -> bool
+		{
+			for (int DetailsViewIndex = 0; DetailsViewIndex < PropertyDetailsWidgets.Num(); ++DetailsViewIndex)
+			{
+				TSharedPtr<SPCGEditorGraphDetailsView> DetailsView = PropertyDetailsWidgets[DetailsViewIndex];
+				if (DetailsView.IsValid() && !DetailsView->IsLocked())
+				{
+					if (!bVisibleOnly || TabManager->FindExistingLiveTab(FPCGEditor_private::PropertyDetailsID[DetailsViewIndex]))
+					{
+						TabManager->TryInvokeTab(FPCGEditor_private::PropertyDetailsID[DetailsViewIndex]);
+						return true;
+					}
+				}
+			}
+
+			return false;
+		};
+
+		if (InvokeFirstUnlockedTab(true) || InvokeFirstUnlockedTab(false))
+		{
+			return;
+		}
+
+		// Default to first if they are all locked
+		if (PropertyDetailsWidgets[0].IsValid())
+		{
+			TabManager->TryInvokeTab(FPCGEditor_private::PropertyDetailsID[0]);
+		}
+	}
+}
+
+void FPCGEditor::OnDetailsViewTabClosed(TSharedRef<SDockTab> DockTab, int Index)
+{
+	if (!PropertyDetailsWidgets.IsValidIndex(Index))
+	{
+		return;
+	}
+
+	TSharedPtr<SPCGEditorGraphDetailsView> DetailsView = PropertyDetailsWidgets[Index];
+	if (DetailsView.IsValid() && DetailsView->IsLocked())
+	{
+		DetailsView->SetIsLocked(false);
+	}
+}
+
+void FPCGEditor::OnAttributeListViewTabClosed(TSharedRef<SDockTab> DockTab, int Index)
+{
+	if (!AttributesWidgets.IsValidIndex(Index))
+	{
+		return;
+	}
+
+	TSharedPtr<SPCGEditorGraphAttributeListView> AttributeListView = AttributesWidgets[Index];
+	if (AttributeListView.IsValid())
+	{
+		if (AttributeListView->IsLocked())
+		{
+			AttributeListView->SetIsLocked(false);
+		}
+
+		UPCGEditorGraphNodeBase* NodeInspected = AttributeListView->GetNodeBeingInspected();
+		AttributeListView->SetNodeBeingInspected(nullptr);
+
+		if (NodeInspected)
+		{
+			bool bIsStillInspectedOnVisibleTabs = false;
+			for (int OtherTabIndex = 0; OtherTabIndex < AttributesWidgets.Num(); ++OtherTabIndex)
+			{
+				TSharedPtr<SPCGEditorGraphAttributeListView> ALV = AttributesWidgets[OtherTabIndex];
+				if (ALV.IsValid() && ALV->GetNodeBeingInspected() == NodeInspected && TabManager->FindExistingLiveTab(FPCGEditor_private::AttributesID[OtherTabIndex]))
+				{
+					bIsStillInspectedOnVisibleTabs = true;
+					break;
+				}
+			}
+
+			if (!bIsStillInspectedOnVisibleTabs)
+			{
+				NodeInspected->SetInspected(false);
+
+				for (TSharedPtr<SPCGEditorGraphAttributeListView> ALV : AttributesWidgets)
+				{
+					if (ALV.IsValid() && ALV->GetNodeBeingInspected() == NodeInspected)
+					{
+						ALV->SetNodeBeingInspected(nullptr);
+					}
+				}
+			}
+		}
+	}
+}
+
 void FPCGEditor::OnPauseAutomaticRegeneration_Clicked()
 {
 	if (!PCGGraphBeingEdited)
@@ -1331,7 +1432,33 @@ void FPCGEditor::OnToggleInspected()
 
 	if (bIsInspecting)
 	{
-		GetTabManager()->TryInvokeTab(FPCGEditor_private::AttributesID[0]);
+		// Summon the first attribute list view that is inspecting this node
+		auto InvokeFirstTab = [this, PCGGraphNodeBase](bool bVisibleOnly) -> bool
+		{
+			for (int AttributeListViewIndex = 0; AttributeListViewIndex < AttributesWidgets.Num(); ++AttributeListViewIndex)
+			{
+				TSharedPtr<SPCGEditorGraphAttributeListView> AttributeListView = AttributesWidgets[AttributeListViewIndex];
+				if (AttributeListView->GetNodeBeingInspected() == PCGGraphNodeBase)
+				{
+					if (!bVisibleOnly || TabManager->FindExistingLiveTab(FPCGEditor_private::AttributesID[AttributeListViewIndex]))
+					{
+						GetTabManager()->TryInvokeTab(FPCGEditor_private::AttributesID[AttributeListViewIndex]);
+						return true;
+					}
+				}
+			}
+
+			return false;
+		};
+
+		const bool bTabSummoned = (InvokeFirstTab(true) || InvokeFirstTab(false));
+
+		// Default to first if they are all locked
+		if (!bTabSummoned)
+		{
+			GetTabManager()->TryInvokeTab(FPCGEditor_private::AttributesID[0]);
+		}
+
 		DebugObjectTreeWidget->SetNodeBeingInspected(PCGNode);
 	}
 }
@@ -2260,21 +2387,7 @@ void FPCGEditor::OnSelectedNodesChanged(const TSet<UObject*>& NewSelection)
 	{
 		for (UObject* Object : NewSelection)
 		{
-			if (UPCGEditorGraphNodeBase* PCGGraphNode = Cast<UPCGEditorGraphNodeBase>(Object))
-			{
-				if (UPCGNode* PCGNode = PCGGraphNode->GetPCGNode())
-				{
-					if (PCGNode->IsInstance())
-					{
-						SelectedObjects.Add(Cast<UPCGSettingsInstance>(PCGNode->GetSettingsInterface()));
-					}
-					else
-					{
-						SelectedObjects.Add(PCGNode->GetSettings());
-					}
-				}
-			}
-			else if (UEdGraphNode* GraphNode = Cast<UEdGraphNode>(Object))
+			if (UEdGraphNode* GraphNode = Cast<UEdGraphNode>(Object))
 			{
 				SelectedObjects.Add(GraphNode);
 			}
@@ -2285,8 +2398,6 @@ void FPCGEditor::OnSelectedNodesChanged(const TSet<UObject*>& NewSelection)
 	{
 		PropertyDetailsWidget->SetObjects(SelectedObjects, /*bForceRefresh=*/true);
 	}
-
-	GetTabManager()->TryInvokeTab(FPCGEditor_private::PropertyDetailsID[0]);
 }
 
 void FPCGEditor::OnNodeTitleCommitted(const FText& NewText, ETextCommit::Type CommitInfo, UEdGraphNode* NodeBeingChanged)
@@ -2441,6 +2552,7 @@ TSharedRef<SDockTab> FPCGEditor::SpawnTab_PropertyDetails(const FSpawnTabArgs& A
 
 	return SNew(SDockTab)
 		.Label(Label)
+		.OnTabClosed_Raw(this, &FPCGEditor::OnDetailsViewTabClosed, PropertyDetailsIndex)
 		.TabColorScale(GetTabColorScale())
 		[
 			DetailsView.ToSharedRef()
@@ -2473,6 +2585,7 @@ TSharedRef<SDockTab> FPCGEditor::SpawnTab_Attributes(const FSpawnTabArgs& Args,
 
 	return SNew(SDockTab)
 		.Label(Label)
+		.OnTabClosed_Raw(this, &FPCGEditor::OnAttributeListViewTabClosed, AttributesIndex)
 		.TabColorScale(GetTabColorScale())
 		[
 			AttributesWidgets[AttributesIndex].ToSharedRef()
diff --git a/Engine/Plugins/Experimental/PCG/Source/PCGEditor/Private/PCGEditor.h b/Engine/Plugins/Experimental/PCG/Source/PCGEditor/Private/PCGEditor.h
--- a/Engine/Plugins/Experimental/PCG/Source/PCGEditor/Private/PCGEditor.h
+++ b/Engine/Plugins/Experimental/PCG/Source/PCGEditor/Private/PCGEditor.h
@@ -104,6 +104,15 @@ private:
 	/** Bring up the find tab */
 	void OnFind();
 
+	/** Bring up the first free details view, or if they are all locked, the first details view */
+	void OpenDetailsView();
+
+	/** Called when a details view tab is closed */
+	void OnDetailsViewTabClosed(TSharedRef<SDockTab> DockTab, int Index);
+
+	/** Called when an attribute list view tab is closed */
+	void OnAttributeListViewTabClosed(TSharedRef<SDockTab> DockTab, int Index);
+
 	/** Enable/Disable automatic PCG node generation */
 	void OnPauseAutomaticRegeneration_Clicked();
 	/** Has the user paused automatic regeneration in the Graph Editor */
diff --git a/Engine/Plugins/Experimental/PCG/Source/PCGEditor/Private/PCGEditorCommands.cpp b/Engine/Plugins/Experimental/PCG/Source/PCGEditor/Private/PCGEditorCommands.cpp
--- a/Engine/Plugins/Experimental/PCG/Source/PCGEditor/Private/PCGEditorCommands.cpp
+++ b/Engine/Plugins/Experimental/PCG/Source/PCGEditor/Private/PCGEditorCommands.cpp
@@ -20,7 +20,8 @@ void FPCGEditorCommands::RegisterCommands()
 	UI_COMMAND(CollapseNodes, "Collapse into Subgraph", "Collapse selected nodes into a separate PCGGraph asset.", EUserInterfaceActionType::Button, FInputChord(EModifierKey::Alt, EKeys::J));
 	UI_COMMAND(ExportNodes, "Export nodes to AssetData", "Exports selected nodes to separate and reusable PCGSettings assets.", EUserInterfaceActionType::Button, FInputChord());
 	UI_COMMAND(ConvertToStandaloneNodes, "Convert to standalone Nodes", "Converts instanced nodes to standalone nodes.", EUserInterfaceActionType::Button, FInputChord());
-	UI_COMMAND(Find, "Find", "Finds PCG nodes and comments in the current graph", EUserInterfaceActionType::Button, FInputChord(EModifierKey::Control, EKeys::F));
+	UI_COMMAND(Find, "Find", "Finds PCG nodes and comments in the current graph.", EUserInterfaceActionType::Button, FInputChord(EModifierKey::Control, EKeys::F));
+	UI_COMMAND(ShowSelectedDetails, "Show Node Details", "Opens a details panel for the selected nodes.", EUserInterfaceActionType::Button, FInputChord(EKeys::F4));
 	UI_COMMAND(PauseAutoRegeneration, "Pause Regen", "Pause automatic regeneration of the current graph.", EUserInterfaceActionType::ToggleButton, FInputChord(EModifierKey::Alt, EKeys::R));
 	UI_COMMAND(ForceGraphRegeneration, "Force Regen", "Manually force a regeneration of the current graph.\nCtrl-click will also perform a flush cache before the regeneration.", EUserInterfaceActionType::Button, FInputChord());
 	UI_COMMAND(OpenDebugObjectTreeTab, "Debug Object Tree", "Open the Debug Object Tree tab to display and select graph invocations to debug.", EUserInterfaceActionType::Button, FInputChord());
diff --git a/Engine/Plugins/Experimental/PCG/Source/PCGEditor/Private/PCGEditorCommands.h b/Engine/Plugins/Experimental/PCG/Source/PCGEditor/Private/PCGEditorCommands.h
--- a/Engine/Plugins/Experimental/PCG/Source/PCGEditor/Private/PCGEditorCommands.h
+++ b/Engine/Plugins/Experimental/PCG/Source/PCGEditor/Private/PCGEditorCommands.h
@@ -17,6 +17,7 @@ public:
 	TSharedPtr<FUICommandInfo> ExportNodes;
 	TSharedPtr<FUICommandInfo> ConvertToStandaloneNodes;
 	TSharedPtr<FUICommandInfo> Find;
+	TSharedPtr<FUICommandInfo> ShowSelectedDetails;
 	TSharedPtr<FUICommandInfo> PauseAutoRegeneration;
 	TSharedPtr<FUICommandInfo> ForceGraphRegeneration;
 	TSharedPtr<FUICommandInfo> OpenDebugObjectTreeTab;
diff --git a/Engine/Plugins/Experimental/PCG/Source/PCGEditor/Private/SPCGEditorGraphAttributeListView.cpp b/Engine/Plugins/Experimental/PCG/Source/PCGEditor/Private/SPCGEditorGraphAttributeListView.cpp
--- a/Engine/Plugins/Experimental/PCG/Source/PCGEditor/Private/SPCGEditorGraphAttributeListView.cpp
+++ b/Engine/Plugins/Experimental/PCG/Source/PCGEditor/Private/SPCGEditorGraphAttributeListView.cpp
@@ -453,9 +453,14 @@ void SPCGEditorGraphAttributeListView::Construct(const FArguments& InArgs, TShar
 			.VAlign(EVerticalAlignment::VAlign_Center)
 			.Padding(4.0f, 0.0f)
 			[
-				SAssignNew(NodeNameTextBlock, STextBlock)
-				.Text(PCGEditorGraphAttributeListView::NoNodeInspectedText)
-				.ToolTipText(PCGEditorGraphAttributeListView::NoNodeInspectedToolTip)
+				SNew(SButton)
+				.ButtonStyle(FAppStyle::Get(), "SimpleButton")
+				.OnClicked(this, &SPCGEditorGraphAttributeListView::OnNodeNameClicked)
+				[
+					SAssignNew(NodeNameTextBlock, STextBlock)
+					.Text(PCGEditorGraphAttributeListView::NoNodeInspectedText)
+					.ToolTipText(PCGEditorGraphAttributeListView::NoNodeInspectedToolTip)
+				]
 			]
 			+SHorizontalBox::Slot()
 			.FillWidth(1.0f)
@@ -1431,4 +1436,14 @@ FReply SPCGEditorGraphAttributeListView::OnLockClick()
 	return FReply::Handled();
 }
 
+FReply SPCGEditorGraphAttributeListView::OnNodeNameClicked()
+{
+	if (PCGEditorGraphNode.Get() && PCGEditorPtr.IsValid())
+	{
+		PCGEditorPtr.Pin()->JumpToNode(PCGEditorGraphNode.Get());
+	}
+
+	return FReply::Handled();
+}
+
 #undef LOCTEXT_NAMESPACE
diff --git a/Engine/Plugins/Experimental/PCG/Source/PCGEditor/Private/SPCGEditorGraphAttributeListView.h b/Engine/Plugins/Experimental/PCG/Source/PCGEditor/Private/SPCGEditorGraphAttributeListView.h
--- a/Engine/Plugins/Experimental/PCG/Source/PCGEditor/Private/SPCGEditorGraphAttributeListView.h
+++ b/Engine/Plugins/Experimental/PCG/Source/PCGEditor/Private/SPCGEditorGraphAttributeListView.h
@@ -143,6 +143,7 @@ public:
 	void SetNodeBeingInspected(UPCGEditorGraphNodeBase* InNode);
 
 	bool IsLocked() const { return bIsLocked; }
+	void SetIsLocked(bool bInIsLocked) { bIsLocked = bInIsLocked; }
 
 private:
 	TSharedRef<SHeaderRow> CreateHeaderRowWidget() const;
@@ -212,6 +213,7 @@ private:
 	const FSlateBrush* OnGetLockButtonImageResource() const;
 
 	FReply OnLockClick();
+	FReply OnNodeNameClicked();
 
 	/** Pointer back to the PCG editor that owns us */
 	TWeakPtr<FPCGEditor> PCGEditorPtr;
diff --git a/Engine/Plugins/Experimental/PCG/Source/PCGEditor/Private/SPCGEditorGraphDetailsView.cpp b/Engine/Plugins/Experimental/PCG/Source/PCGEditor/Private/SPCGEditorGraphDetailsView.cpp
--- a/Engine/Plugins/Experimental/PCG/Source/PCGEditor/Private/SPCGEditorGraphDetailsView.cpp
+++ b/Engine/Plugins/Experimental/PCG/Source/PCGEditor/Private/SPCGEditorGraphDetailsView.cpp
@@ -6,6 +6,9 @@
 #include "PCGNode.h"
 #include "PCGSettings.h"
 
+#include "PCGEditor.h"
+#include "PCGEditorGraphNode.h"
+
 #include "DetailsViewArgs.h"
 #include "IDetailsView.h"
 #include "PropertyEditorModule.h"
@@ -61,9 +64,14 @@ void SPCGEditorGraphDetailsView::Construct(const FArguments& InArgs)
 			.HAlign(HAlign_Center)
 			.VAlign(VAlign_Center)
 			[
-				SNew(STextBlock)
+				SNew(SButton)
+				.ButtonStyle(FAppStyle::Get(), "SimpleButton")
+				.OnClicked(this, &SPCGEditorGraphDetailsView::OnNameClicked)
 				.Visibility(this, &SPCGEditorGraphDetailsView::GetNameVisibility)
-				.Text(this, &SPCGEditorGraphDetailsView::GetName)
+				[
+					SNew(STextBlock)
+					.Text(this, &SPCGEditorGraphDetailsView::GetName)
+				]
 			]
 			+SHorizontalBox::Slot()
 			[
@@ -79,44 +87,61 @@ void SPCGEditorGraphDetailsView::Construct(const FArguments& InArgs)
 
 void SPCGEditorGraphDetailsView::SetObject(UObject* InObject, bool bForceRefresh)
 {
-	TArray<UObject*> InObjects;
+	TArray<TWeakObjectPtr<UObject>> InObjects;
 	InObjects.Add(InObject);
 
 	SetObjects(InObjects, bForceRefresh);
 }
 
-void SPCGEditorGraphDetailsView::SetObjects(const TArray<UObject*>& InObjects, bool bForceRefresh, bool bOverrideLock)
+void SPCGEditorGraphDetailsView::SetObjects(const TArray<TWeakObjectPtr<UObject>>& InObjects, bool bForceRefresh, bool bOverrideLock)
 {
 	if (bOverrideLock || !bIsLocked)
 	{
-		DetailsView->SetObjects(InObjects, bForceRefresh, bOverrideLock);
+		SelectedObjects = InObjects;
 
-		// Unlock the view automatically if we are viewing nothing
-		if (InObjects.IsEmpty())
+		// Filter only the types we're interested in, e.g. the settings and not the nodes themselves
+		TArray<TWeakObjectPtr<UObject>> ObjectsToView;
+		ObjectsToView.Reserve(InObjects.Num());
+
+		for (const TWeakObjectPtr<UObject>& InObject : InObjects)
 		{
-			bIsLocked = false;
+			if (!InObject.IsValid())
+			{
+				continue;
+			}
+
+			if (UPCGEditorGraphNodeBase* PCGGraphNode = Cast<UPCGEditorGraphNodeBase>(InObject.Get()))
+			{
+				if (UPCGNode* PCGNode = PCGGraphNode->GetPCGNode())
+				{
+					if (PCGNode->IsInstance())
+					{
+						ObjectsToView.Add(PCGNode->GetSettingsInterface());
+					}
+					else
+					{
+						ObjectsToView.Add(PCGNode->GetSettings());
+					}
+
+					continue;
+				}
+			}
+
+			ObjectsToView.Add(InObject);
 		}
-	}
-}
 
-void SPCGEditorGraphDetailsView::SetObjects(const TArray<TWeakObjectPtr<UObject>>& InObjects, bool bForceRefresh, bool bOverrideLock)
-{
-	TArray<UObject*> Objects;
-	Objects.Reserve(InObjects.Num());
-	for (const TWeakObjectPtr<UObject>& InObject : InObjects)
-	{
-		if (InObject.IsValid())
+		DetailsView->SetObjects(ObjectsToView);
+
+		if (InObjects.IsEmpty())
 		{
-			Objects.Add(InObject.Get());
+			bIsLocked = false;
 		}
 	}
-
-	SetObjects(Objects, bForceRefresh, bOverrideLock);
 }
 
 const TArray<TWeakObjectPtr<UObject>>& SPCGEditorGraphDetailsView::GetSelectedObjects() const
 {
-	return DetailsView->GetSelectedObjects();
+	return SelectedObjects;
 }
 
 FReply SPCGEditorGraphDetailsView::OnLockButtonClicked()
@@ -125,6 +150,24 @@ FReply SPCGEditorGraphDetailsView::OnLockButtonClicked()
 	return FReply::Handled();
 }
 
+FReply SPCGEditorGraphDetailsView::OnNameClicked()
+{
+	if (SelectedObjects.Num() != 1)
+	{
+		return FReply::Handled();
+	}
+
+	if (UEdGraphNode* GraphNode = Cast<UEdGraphNode>(SelectedObjects[0].Get()))
+	{
+		if (FPCGEditor* Editor = EditorPtr.Pin().Get())
+		{
+			Editor->JumpToNode(GraphNode);
+		}
+	}
+
+	return FReply::Handled();
+}
+
 const FSlateBrush* SPCGEditorGraphDetailsView::GetLockIcon() const
 {
 	return FAppStyle::GetBrush(bIsLocked ? TEXT("PropertyWindow.Locked") : TEXT("PropertyWindow.Unlocked"));
@@ -137,15 +180,15 @@ EVisibility SPCGEditorGraphDetailsView::GetNameVisibility() const
 
 FText SPCGEditorGraphDetailsView::GetName() const
 {
-	const TArray<TWeakObjectPtr<UObject>>& SelectedObjects = GetSelectedObjects();
-	if (SelectedObjects.IsEmpty())
+	const TArray<TWeakObjectPtr<UObject>>& DetailsSelectedObjects = DetailsView->GetSelectedObjects();
+	if (DetailsSelectedObjects.IsEmpty())
 	{
 		// Implementation note: this doesn't really happen since we reselect the graph settings when unselecting
 		return LOCTEXT("NoObjectsSelected", "Empty selection");
 	}
-	else if (SelectedObjects.Num() == 1)
+	else if (DetailsSelectedObjects.Num() == 1)
 	{
-		UObject* SelectedObject = SelectedObjects[0].Get();
+		UObject* SelectedObject = DetailsSelectedObjects[0].Get();
 		UPCGNode* OwnerNode = nullptr;
 
 		if (Cast<UPCGSettings>(SelectedObject) != nullptr || Cast<UPCGSettingsInstance>(SelectedObject) != nullptr)
@@ -168,7 +211,7 @@ FText SPCGEditorGraphDetailsView::GetName() const
 	}
 	else
 	{
-		return FText::Format(LOCTEXT("MultipleObjectsSelectedFmt", "{0} nodes"), FText::AsNumber(SelectedObjects.Num()));
+		return FText::Format(LOCTEXT("MultipleObjectsSelectedFmt", "{0} nodes"), FText::AsNumber(DetailsSelectedObjects.Num()));
 	}
 }
 
@@ -182,8 +225,8 @@ bool SPCGEditorGraphDetailsView::IsReadOnlyProperty(const FPropertyAndParent& In
 		return false;
 	}
 
-	const TArray<TWeakObjectPtr<UObject>>& SelectedObjects = DetailsView->GetSelectedObjects();
-	for (const TWeakObjectPtr<UObject>& SelectedObject : SelectedObjects)
+	const TArray<TWeakObjectPtr<UObject>>& DetailsSelectedObjects = DetailsView->GetSelectedObjects();
+	for (const TWeakObjectPtr<UObject>& SelectedObject : DetailsSelectedObjects)
 	{
 		if (!SelectedObject.IsValid())
 		{
diff --git a/Engine/Plugins/Experimental/PCG/Source/PCGEditor/Private/SPCGEditorGraphDetailsView.h b/Engine/Plugins/Experimental/PCG/Source/PCGEditor/Private/SPCGEditorGraphDetailsView.h
--- a/Engine/Plugins/Experimental/PCG/Source/PCGEditor/Private/SPCGEditorGraphDetailsView.h
+++ b/Engine/Plugins/Experimental/PCG/Source/PCGEditor/Private/SPCGEditorGraphDetailsView.h
@@ -5,6 +5,7 @@
 #include "Templates/SharedPointer.h"
 #include "Widgets/SCompoundWidget.h"
 
+class FPCGEditor;
 class IDetailsView;
 struct FPropertyAndParent;
 
@@ -19,6 +20,10 @@ public:
 	/** Constructs this widget with InArgs */
 	void Construct(const FArguments& InArgs);
 
+	/** Sets the editor associated to this details view, used for interaction from the details view to the editor */
+	void SetEditor(TWeakPtr<FPCGEditor> InEditorPtr) { EditorPtr = InEditorPtr; }
+
+	/** Returns the details view under this widget */
 	TSharedPtr<IDetailsView> GetDetailsView() const { return DetailsView; }
 
 	//~Begin IDetailsView-like interface
@@ -28,9 +33,6 @@ public:
 	/** Changes the selected objects in the details view. Will unlock automatically. */
 	void SetObject(UObject* InObject, bool bForceRefresh = false);
 
-	/** Changes the selected objects in the details view. Will unlock automatically. */
-	void SetObjects(const TArray<UObject*>& InObjects, bool bForceRefresh = false, bool bOverrideLock = false);
-
 	/** Changes the selected objects in the details view. Will unlock automatically. */
 	void SetObjects(const TArray<TWeakObjectPtr<UObject>>& InObjects, bool bForceRefresh = false, bool bOverrideLock = false);
 
@@ -38,6 +40,9 @@ public:
 	const TArray<TWeakObjectPtr<UObject>>& GetSelectedObjects() const;
 	//~End IDetailsView-like interface
 
+	/** Controls lock from the editor */
+	void SetIsLocked(bool bInIsLocked) { bIsLocked = bInIsLocked; }
+
 protected:
 	/** Returns whether a property should be readonly (used for instances) */
 	bool IsReadOnlyProperty(const FPropertyAndParent& InPropertyAndParent) const;
@@ -45,6 +50,7 @@ protected:
 	bool IsVisibleProperty(const FPropertyAndParent& InPropertyAndParent) const;
 
 	FReply OnLockButtonClicked();
+	FReply OnNameClicked();
 	const FSlateBrush* GetLockIcon() const;
 
 	/** Gets the current visibility state of the name section, which is related to whether the details view is locked. */
@@ -54,4 +60,7 @@ protected:
 
 	TSharedPtr<IDetailsView> DetailsView;
 	bool bIsLocked : 1 = false;
+
+	TArray<TWeakObjectPtr<UObject>> SelectedObjects;
+	TWeakPtr<FPCGEditor> EditorPtr = nullptr;
 };
\ No newline at end of file
