[PCG] Refactor tracking to remove actor tracking (costly) and support dynamic tracking on objects.

* Refactor CachedTrackedKeysToSettings to StaticallyTrackedKeysToSettings.
* Remove culling cache to unify the way we check culling for static and dynamic.
* Add DynamicTracking to GetObjectFromPropertyPath, Indirection, Texture and Subgraphs (including loop)
* Remove Actor tracking, replaced by key tracking in the actor and component mapping.
* OnActorChanged is now OnObjectChanged.

#rnx

diff --git a/Engine/Plugins/Experimental/PCG/Source/PCG/Private/Elements/IO/PCGDataTableElement.cpp b/Engine/Plugins/Experimental/PCG/Source/PCG/Private/Elements/IO/PCGDataTableElement.cpp
--- a/Engine/Plugins/Experimental/PCG/Source/PCG/Private/Elements/IO/PCGDataTableElement.cpp
+++ b/Engine/Plugins/Experimental/PCG/Source/PCG/Private/Elements/IO/PCGDataTableElement.cpp
@@ -7,13 +7,14 @@
 #include "PCGParamData.h"
 #include "Data/PCGPointData.h"
 #include "Elements/IO/PCGExternalDataContext.h"
+#include "Helpers/PCGDynamicTrackingHelpers.h"
 
 #include "Engine/DataTable.h"
 
 #define LOCTEXT_NAMESPACE "PCGLoadDataTable"
 
 #if WITH_EDITOR
-void UPCGLoadDataTableSettings::GetTrackedActorKeys(FPCGSelectionKeyToSettingsMap& OutKeysToSettings, TArray<TObjectPtr<const UPCGGraph>>& OutVisitedGraphs) const
+void UPCGLoadDataTableSettings::GetStaticTrackedKeys(FPCGSelectionKeyToSettingsMap& OutKeysToSettings, TArray<TObjectPtr<const UPCGGraph>>& OutVisitedGraphs) const
 {
 	if (DataTable.IsNull())
 	{
@@ -268,6 +269,14 @@ bool FPCGLoadDataTableElement::ExecuteLoad(FPCGExternalDataContext* Context) con
 	}
 
 	PointDataAccessorMapping.RowKeys = MakeUnique<FPCGAttributeAccessorKeysGenericPtrs>(AllRows);
+
+#if WITH_EDITOR
+	// If we have an override, register for dynamic tracking.
+	if (Context->IsValueOverriden(GET_MEMBER_NAME_CHECKED(UPCGLoadDataTableSettings, DataTable)))
+	{
+		FPCGDynamicTrackingHelper::AddSingleDynamicTrackingKey(Context, FPCGSelectionKey::CreateFromPath(DataTable), /*bIsCulled=*/false);
+	}
+#endif // WITH_EDITOR
 	
 	return FPCGExternalDataElement::ExecuteLoad(Context);
 }
diff --git a/Engine/Plugins/Experimental/PCG/Source/PCG/Private/Elements/PCGActorSelector.cpp b/Engine/Plugins/Experimental/PCG/Source/PCG/Private/Elements/PCGActorSelector.cpp
--- a/Engine/Plugins/Experimental/PCG/Source/PCG/Private/Elements/PCGActorSelector.cpp
+++ b/Engine/Plugins/Experimental/PCG/Source/PCG/Private/Elements/PCGActorSelector.cpp
@@ -239,20 +239,28 @@ FPCGSelectionKey FPCGSelectionKey::CreateFromPath(const FSoftObjectPath& InObjec
 	return Key;
 }
 
+FPCGSelectionKey FPCGSelectionKey::CreateFromPath(FSoftObjectPath&& InObjectPath)
+{
+	FPCGSelectionKey Key{};
+	Key.Selection = EPCGActorSelection::ByPath;
+	Key.ObjectPath = std::forward<FSoftObjectPath>(InObjectPath);
+	Key.ActorFilter = EPCGActorFilter::AllWorldActors;
+
+	return Key;
+}
+
 void FPCGSelectionKey::SetExtraDependency(const UClass* InExtraDependency)
 {
 	OptionalExtraDependency = InExtraDependency;
 }
 
-bool FPCGSelectionKey::IsMatching(const TSoftObjectPtr<UObject>& InObjectPtr, const UPCGComponent* InComponent) const
+bool FPCGSelectionKey::IsMatching(const UObject* InObject, const UPCGComponent* InComponent) const
 {
-	if (InObjectPtr.IsNull())
+	if (!InObject)
 	{
 		return false;
 	}
 
-	const UObject* InObject = InObjectPtr.Get();
-	
 	// If we filter something else than all world actors, matching depends on the component.
 	// Re-use the same mechanism than Get Actor Data, which should be cheap since we don't look for all actors in the world.
 	if (ActorFilter != EPCGActorFilter::AllWorldActors)
@@ -283,12 +291,95 @@ bool FPCGSelectionKey::IsMatching(const TSoftObjectPtr<UObject>& InObjectPtr, co
 	case EPCGActorSelection::ByClass:
 		return InObject && InObject->GetClass()->IsChildOf(SelectionClass);
 	case EPCGActorSelection::ByPath:
-		return InObjectPtr.ToSoftObjectPath() == ObjectPath;
+		return FSoftObjectPath(InObject) == ObjectPath;
 	default:
 		return false;
 	}
 }
 
+bool FPCGSelectionKey::IsMatching(const UObject* InObject, const TSet<FName>& InRemovedTags, const TSet<UPCGComponent*>& InComponents, TSet<UPCGComponent*>& MatchedComponents) const
+{
+	if (!InObject)
+	{
+		return false;
+	}
+
+	// If we filter something else than all world actors, matching depends on the component.
+	// Since we can have a lot of components in InComponents, we go the other way around (Actor to component)
+	if (ActorFilter != EPCGActorFilter::AllWorldActors)
+	{
+		bool bFoundMatch = false;
+
+		const AActor* InActor = Cast<const AActor>(InObject);
+
+		if (!InActor)
+		{
+			return false;
+		}
+
+		TArray<UActorComponent*, TInlineAllocator<64>> ActorComponents;
+
+		if (ActorFilter == EPCGActorFilter::Self || ActorFilter == EPCGActorFilter::Original)
+		{
+			InActor->GetComponents(UPCGComponent::StaticClass(), ActorComponents);
+		}
+		else if (ActorFilter == EPCGActorFilter::Parent || (ActorFilter == EPCGActorFilter::Root && !InActor->GetParentActor()))
+		{
+			TArray<AActor*> ActorsToCheck;
+			InActor->GetAllChildActors(ActorsToCheck, /*bIncludeDescendants=*/ActorFilter == EPCGActorFilter::Root);
+			ActorsToCheck.Add(const_cast<AActor*>(InActor));
+			TArray<UActorComponent*, TInlineAllocator<64>> TempActorComponents;
+			for (AActor* Current : ActorsToCheck)
+			{
+				// TempActorComponents is reset in GetComponents
+				InActor->GetComponents(UPCGComponent::StaticClass(), TempActorComponents);
+				ActorComponents.Append(TempActorComponents);
+			}
+		}
+
+		for (UActorComponent* Component : ActorComponents)
+		{
+			if (UPCGComponent* PCGComponent = Cast<UPCGComponent>(Component))
+			{
+				if (InComponents.Contains(PCGComponent))
+				{
+					MatchedComponents.Add(PCGComponent);
+					bFoundMatch = true;
+				}
+			}
+		}
+
+		return bFoundMatch;
+	}
+
+	bool bIsMatched = false;
+	switch (Selection)
+	{
+	case EPCGActorSelection::ByTag:
+	{
+		const AActor* InActor = Cast<const AActor>(InObject);
+		bIsMatched = InRemovedTags.Contains(Tag) || (InActor && InActor->ActorHasTag(Tag));
+		break;
+	}
+	case EPCGActorSelection::ByClass:
+		bIsMatched = InObject->IsA(SelectionClass);
+		break;
+	case EPCGActorSelection::ByPath:
+		bIsMatched = FSoftObjectPath(InObject) == ObjectPath;
+		break;
+	default:
+		bIsMatched = false;
+		break;
+	}
+
+	if (bIsMatched)
+	{
+		MatchedComponents.Append(InComponents);
+	}
+
+	return bIsMatched;
+}
+
 bool FPCGSelectionKey::operator==(const FPCGSelectionKey& InOther) const
 {
 	if (ActorFilter != InOther.ActorFilter || Selection != InOther.Selection || OptionalExtraDependency != InOther.OptionalExtraDependency)
@@ -315,6 +406,20 @@ bool FPCGSelectionKey::operator==(const FPCGSelectionKey& InOther) const
 	}
 }
 
+FArchive& operator<<(FArchive& Ar, FPCGSelectionKey& Key)
+{
+	// Serialize the normal UPROPERTY data
+	if (Ar.IsLoading() || Ar.IsSaving())
+	{
+		if (UScriptStruct* ThisStruct = FPCGSelectionKey::StaticStruct())
+		{
+			ThisStruct->SerializeTaggedProperties(Ar, reinterpret_cast<uint8*>(&Key), ThisStruct, nullptr);
+		}
+	}
+
+	return Ar;
+}
+
 uint32 GetTypeHash(const FPCGSelectionKey& In)
 {
 	uint32 HashResult = HashCombine(GetTypeHash(In.ActorFilter), GetTypeHash(In.Selection));
diff --git a/Engine/Plugins/Experimental/PCG/Source/PCG/Private/Elements/PCGDataFromActor.cpp b/Engine/Plugins/Experimental/PCG/Source/PCG/Private/Elements/PCGDataFromActor.cpp
--- a/Engine/Plugins/Experimental/PCG/Source/PCG/Private/Elements/PCGDataFromActor.cpp
+++ b/Engine/Plugins/Experimental/PCG/Source/PCG/Private/Elements/PCGDataFromActor.cpp
@@ -127,7 +127,7 @@ namespace PCGDataFromActorHelpers
 }
 
 #if WITH_EDITOR
-void UPCGDataFromActorSettings::GetTrackedActorKeys(FPCGSelectionKeyToSettingsMap& OutKeysToSettings, TArray<TObjectPtr<const UPCGGraph>>& OutVisitedGraphs) const
+void UPCGDataFromActorSettings::GetStaticTrackedKeys(FPCGSelectionKeyToSettingsMap& OutKeysToSettings, TArray<TObjectPtr<const UPCGGraph>>& OutVisitedGraphs) const
 {
 	FPCGSelectionKey Key = ActorSelector.GetAssociatedKey();
 	if (Mode == EPCGGetDataFromActorMode::GetDataFromPCGComponent || Mode == EPCGGetDataFromActorMode::GetDataFromPCGComponentOrParseComponents)
diff --git a/Engine/Plugins/Experimental/PCG/Source/PCG/Private/Elements/PCGDataTableRowToParamData.cpp b/Engine/Plugins/Experimental/PCG/Source/PCG/Private/Elements/PCGDataTableRowToParamData.cpp
--- a/Engine/Plugins/Experimental/PCG/Source/PCG/Private/Elements/PCGDataTableRowToParamData.cpp
+++ b/Engine/Plugins/Experimental/PCG/Source/PCG/Private/Elements/PCGDataTableRowToParamData.cpp
@@ -6,6 +6,7 @@
 #include "PCGParamData.h"
 #include "PCGPin.h"
 #include "PCGModule.h"
+#include "Helpers/PCGDynamicTrackingHelpers.h"
 #include "Metadata/PCGMetadata.h"
 
 #include "Engine/DataTable.h"
@@ -14,6 +15,29 @@
 
 #define LOCTEXT_NAMESPACE "PCGDataTableRowToParamDataElement"
 
+#if WITH_EDITOR
+void UPCGDataTableRowToParamDataSettings::GetStaticTrackedKeys(FPCGSelectionKeyToSettingsMap& OutKeysToSettings, TArray<TObjectPtr<const UPCGGraph>>& OutVisitedGraphs) const
+{
+	if (IsPropertyOverriddenByPin(GET_MEMBER_NAME_CHECKED(UPCGDataTableRowToParamDataSettings, DataTable)) || DataTable.IsNull())
+	{
+		// Dynamic tracking or null settings
+		return;
+	}
+
+	FPCGSelectionKey Key = FPCGSelectionKey::CreateFromPath(DataTable.ToSoftObjectPath());
+
+	OutKeysToSettings.FindOrAdd(Key).Emplace(this, /*bCulling=*/false);
+}
+
+void UPCGDataTableRowToParamDataSettings::ApplyDeprecationBeforeUpdatePins(UPCGNode* InOutNode, TArray<TObjectPtr<UPCGPin>>& InputPins, TArray<TObjectPtr<UPCGPin>>& OutputPins)
+{
+	Super::ApplyDeprecationBeforeUpdatePins(InOutNode, InputPins, OutputPins);
+
+	// Overridable properties have been renamed, rename all pins by their counterpart, to avoid breaking existing graphs.
+	InOutNode->RenameInputPin(TEXT("Path Override"), TEXT("Data Table"));
+}
+#endif // WITH_EDITOR
+
 FPCGElementPtr UPCGDataTableRowToParamDataSettings::CreateElement() const
 {
 	return MakeShared<FPCGDataTableRowToParamData>();
@@ -32,6 +56,19 @@ TArray<FPCGPinProperties> UPCGDataTableRowToParamDataSettings::OutputPinProperti
 	return PinProperties;
 }
 
+void UPCGDataTableRowToParamDataSettings::PostLoad()
+{
+	Super::PostLoad();
+
+#if WITH_EDITOR
+	if (!PathOverride_DEPRECATED.IsEmpty())
+	{
+		DataTable = TSoftObjectPtr<UDataTable>(PathOverride_DEPRECATED);
+		PathOverride_DEPRECATED.Empty();
+	}
+#endif // WITH_EDITOR
+}
+
 FString UPCGDataTableRowToParamDataSettings::GetAdditionalTitleInformation() const
 {
 	return FString::Printf(TEXT("%s[ %s ]"), DataTable ? *DataTable->GetFName().ToString() : TEXT("None"), *RowName.ToString());
@@ -44,16 +81,10 @@ bool FPCGDataTableRowToParamData::ExecuteInternal(FPCGContext* Context) const
 	const UPCGDataTableRowToParamDataSettings* Settings = Context->GetInputSettings<UPCGDataTableRowToParamDataSettings>();
 	check(Settings);
 
-	const FString& PathOverride = Settings->PathOverride;
 	const FName RowName = Settings->RowName;
 
 	TSoftObjectPtr<UDataTable> DataTablePtr = Settings->DataTable;
 
-	if (!PathOverride.IsEmpty())
-	{
-		DataTablePtr = FSoftObjectPath(PathOverride);
-	}
-
 	const UDataTable* DataTable = DataTablePtr.LoadSynchronous();
 	if (!DataTable)
 	{
@@ -90,6 +121,14 @@ bool FPCGDataTableRowToParamData::ExecuteInternal(FPCGContext* Context) const
 	FPCGTaggedData& Output = Outputs.Emplace_GetRef();
 	Output.Data = ParamData;
 
+#if WITH_EDITOR
+	// If we have an override, register for dynamic tracking.
+	if (Context->IsValueOverriden(GET_MEMBER_NAME_CHECKED(UPCGDataTableRowToParamDataSettings, DataTable)))
+	{
+		FPCGDynamicTrackingHelper::AddSingleDynamicTrackingKey(Context, FPCGSelectionKey::CreateFromPath(DataTable), /*bIsCulled=*/false);
+	}
+#endif // WITH_EDITOR
+
 	return true;
 }
 
diff --git a/Engine/Plugins/Experimental/PCG/Source/PCG/Private/Elements/PCGExecuteBlueprint.cpp b/Engine/Plugins/Experimental/PCG/Source/PCG/Private/Elements/PCGExecuteBlueprint.cpp
--- a/Engine/Plugins/Experimental/PCG/Source/PCG/Private/Elements/PCGExecuteBlueprint.cpp
+++ b/Engine/Plugins/Experimental/PCG/Source/PCG/Private/Elements/PCGExecuteBlueprint.cpp
@@ -500,7 +500,7 @@ EPCGSettingsType UPCGBlueprintSettings::GetType() const
 	}
 }
 
-void UPCGBlueprintSettings::GetTrackedActorKeys(FPCGSelectionKeyToSettingsMap& OutKeysToSettings, TArray<TObjectPtr<const UPCGGraph>>& OutVisitedGraphs) const
+void UPCGBlueprintSettings::GetStaticTrackedKeys(FPCGSelectionKeyToSettingsMap& OutKeysToSettings, TArray<TObjectPtr<const UPCGGraph>>& OutVisitedGraphs) const
 {
 	for (const FName& Tag : TrackedActorTags)
 	{
diff --git a/Engine/Plugins/Experimental/PCG/Source/PCG/Private/Elements/PCGGetActorProperty.cpp b/Engine/Plugins/Experimental/PCG/Source/PCG/Private/Elements/PCGGetActorProperty.cpp
--- a/Engine/Plugins/Experimental/PCG/Source/PCG/Private/Elements/PCGGetActorProperty.cpp
+++ b/Engine/Plugins/Experimental/PCG/Source/PCG/Private/Elements/PCGGetActorProperty.cpp
@@ -7,6 +7,7 @@
 #include "PCGComponent.h"
 #include "PCGParamData.h"
 #include "Helpers/PCGBlueprintHelpers.h"
+#include "Helpers/PCGDynamicTrackingHelpers.h"
 #include "Helpers/PCGPropertyHelpers.h"
 #include "Helpers/PCGHelpers.h"
 #include "Metadata/PCGAttributePropertySelector.h"
@@ -15,7 +16,7 @@
 #define LOCTEXT_NAMESPACE "PCGPropertyToParamDataElement"
 
 #if WITH_EDITOR
-void UPCGGetActorPropertySettings::GetTrackedActorKeys(FPCGSelectionKeyToSettingsMap& OutKeysToSettings, TArray<TObjectPtr<const UPCGGraph>>& OutVisitedGraphs) const
+void UPCGGetActorPropertySettings::GetStaticTrackedKeys(FPCGSelectionKeyToSettingsMap& OutKeysToSettings, TArray<TObjectPtr<const UPCGGraph>>& OutVisitedGraphs) const
 {
 	OutKeysToSettings.FindOrAdd(ActorSelector.GetAssociatedKey()).Emplace(this, bTrackActorsOnlyWithinBounds);
 }
@@ -162,7 +163,15 @@ bool FPCGGetActorPropertyElement::ExecuteInternal(FPCGContext* Context) const
 		const FPCGAttributePropertySelector Selector = FPCGAttributePropertySelector::CreateSelectorFromString(Settings->PropertyName.ToString());
 
 		PCGPropertyHelpers::FExtractorParameters Parameters{ ObjectToInspect, ObjectToInspect->GetClass(), Selector, Settings->OutputAttributeName, Settings->bForceObjectAndStructExtraction, /*bPropertyNeedsToBeVisible=*/true };
-		if (UPCGParamData* ParamData = PCGPropertyHelpers::ExtractPropertyAsAttributeSet(Parameters, Context))
+
+		// Don't care for object traversed in non-editor build, since it is only useful for tracking.
+		TSet<FSoftObjectPath>* ObjectTraversedPtr = nullptr;
+#if WITH_EDITOR
+		TSet<FSoftObjectPath> ObjectTraversed;
+		ObjectTraversedPtr = &ObjectTraversed;
+#endif // WITH_EDITOR
+
+		if (UPCGParamData* ParamData = PCGPropertyHelpers::ExtractPropertyAsAttributeSet(Parameters, Context, ObjectTraversedPtr))
 		{
 			TArray<FPCGTaggedData>& Outputs = Context->OutputData.TaggedData;
 			FPCGTaggedData& Output = Outputs.Emplace_GetRef();
@@ -179,6 +188,21 @@ bool FPCGGetActorPropertyElement::ExecuteInternal(FPCGContext* Context) const
 				PCGE_LOG(Error, GraphAndLog, FText::Format(LOCTEXT("FailedToExtract", "Fail to extract the property '{0}' on actor {1}."), Selector.GetDisplayText(), FText::FromString(FoundActor->GetName())));
 			}
 		}
+
+		// Register dynamic tracking
+#if WITH_EDITOR
+		if (!ObjectTraversed.IsEmpty())
+		{
+			FPCGDynamicTrackingHelper DynamicTracking;
+			DynamicTracking.EnableAndInitialize(Context, ObjectTraversed.Num());
+			for (FSoftObjectPath& Path : ObjectTraversed)
+			{
+				DynamicTracking.AddToTracking(FPCGSelectionKey::CreateFromPath(std::move(Path)), /*bCulled=*/false);
+			}
+
+			DynamicTracking.Finalize(Context);
+		}
+#endif // WITH_EDITOR
 	}
 
 	return true;
diff --git a/Engine/Plugins/Experimental/PCG/Source/PCG/Private/Elements/PCGGetPropertyFromObjectPath.cpp b/Engine/Plugins/Experimental/PCG/Source/PCG/Private/Elements/PCGGetPropertyFromObjectPath.cpp
--- a/Engine/Plugins/Experimental/PCG/Source/PCG/Private/Elements/PCGGetPropertyFromObjectPath.cpp
+++ b/Engine/Plugins/Experimental/PCG/Source/PCG/Private/Elements/PCGGetPropertyFromObjectPath.cpp
@@ -2,7 +2,9 @@
 
 #include "Elements/PCGGetPropertyFromObjectPath.h"
 
+#include "PCGComponent.h"
 #include "PCGParamData.h"
+#include "Helpers/PCGDynamicTrackingHelpers.h"
 #include "Helpers/PCGPropertyHelpers.h"
 #include "Metadata/Accessors/IPCGAttributeAccessor.h"
 #include "Metadata/Accessors/PCGAttributeAccessorKeys.h"
@@ -13,8 +15,15 @@
 #define LOCTEXT_NAMESPACE "PCGGetPropertyFromObjectPathElement"
 
 #if WITH_EDITOR
-void UPCGGetPropertyFromObjectPathSettings::GetTrackedActorKeys(FPCGSelectionKeyToSettingsMap& OutKeysToSettings, TArray<TObjectPtr<const UPCGGraph>>& OutVisitedGraphs) const
+void UPCGGetPropertyFromObjectPathSettings::GetStaticTrackedKeys(FPCGSelectionKeyToSettingsMap& OutKeysToSettings, TArray<TObjectPtr<const UPCGGraph>>& OutVisitedGraphs) const
 {
+	// If input pin is connected, tracking is dynamic.
+	const UPCGNode* Node = Cast<const UPCGNode>(GetOuter());
+	if (Node && Node->IsInputPinConnected(PCGPinConstants::DefaultInputLabel))
+	{
+		return;
+	}
+
 	for (const FSoftObjectPath& ObjectPath : ObjectPathsToExtract)
 	{
 		if (ObjectPath.IsNull())
@@ -288,6 +297,16 @@ bool FPCGGetPropertyFromObjectPathElement::ExecuteInternal(FPCGContext* Context)
 		}
 	};
 
+#if WITH_EDITOR
+	FPCGDynamicTrackingHelper DynamicTrackingHelper;
+	UPCGComponent* SourceComponent = Context->SourceComponent.Get();
+	const bool IsDynamicallyTracking = (SourceComponent && Context->Node) ? Context->Node->IsInputPinConnected(PCGPinConstants::DefaultInputLabel) : false;
+	if (IsDynamicallyTracking)
+	{
+		DynamicTrackingHelper.EnableAndInitialize(ThisContext, ThisContext->PathsToObjectsToExtractAndIncomingDataIndex.Num());
+	}
+#endif // WITH_EDITOR
+
 	for (const TTuple<FSoftObjectPath, int32>& SoftPathAndIndex : ThisContext->PathsToObjectsToExtractAndIncomingDataIndex)
 	{
 		const FSoftObjectPath& SoftPath = SoftPathAndIndex.Get<FSoftObjectPath>();
@@ -322,6 +341,9 @@ bool FPCGGetPropertyFromObjectPathElement::ExecuteInternal(FPCGContext* Context)
 		if (UPCGParamData* ParamData = PCGPropertyHelpers::ExtractPropertyAsAttributeSet(Parameters, Context))
 		{
 			AddToOutput(ParamData, Index);
+#if WITH_EDITOR
+			DynamicTrackingHelper.AddToTracking(FPCGSelectionKey::CreateFromPath(SoftPath), /*bIsCulled=*/ false);
+#endif // WITH_EDITOR
 		}
 		else
 		{
@@ -336,6 +358,10 @@ bool FPCGGetPropertyFromObjectPathElement::ExecuteInternal(FPCGContext* Context)
 		}
 	}
 
+#if WITH_EDITOR
+	DynamicTrackingHelper.Finalize(ThisContext);
+#endif // WITH_EDITOR
+
 	return true;
 }
 
diff --git a/Engine/Plugins/Experimental/PCG/Source/PCG/Private/Elements/PCGIndirectionElement.cpp b/Engine/Plugins/Experimental/PCG/Source/PCG/Private/Elements/PCGIndirectionElement.cpp
--- a/Engine/Plugins/Experimental/PCG/Source/PCG/Private/Elements/PCGIndirectionElement.cpp
+++ b/Engine/Plugins/Experimental/PCG/Source/PCG/Private/Elements/PCGIndirectionElement.cpp
@@ -5,6 +5,7 @@
 #include "PCGPin.h"
 #include "PCGSettings.h"
 #include "Elements/PCGExecuteBlueprint.h" // Blueprint element class
+#include "Helpers/PCGDynamicTrackingHelpers.h" // Dynamic tracking
 #include "Helpers/PCGSettingsHelpers.h"   // Graph and log errors
 
 #include UE_INLINE_GENERATED_CPP_BY_NAME(PCGIndirectionElement)
@@ -17,10 +18,11 @@ namespace PCGIndirectionSettings
 }
 
 #if WITH_EDITOR
-void UPCGIndirectionSettings::GetTrackedActorKeys(FPCGSelectionKeyToSettingsMap& OutKeysToSettings, TArray<TObjectPtr<const UPCGGraph>>& OutVisitedGraphs) const
+void UPCGIndirectionSettings::GetStaticTrackedKeys(FPCGSelectionKeyToSettingsMap& OutKeysToSettings, TArray<TObjectPtr<const UPCGGraph>>& OutVisitedGraphs) const
 {
-	if (Settings.IsNull())
+	if (IsPropertyOverriddenByPin(GET_MEMBER_NAME_CHECKED(UPCGIndirectionSettings, Settings)) || Settings.IsNull())
 	{
+		// Dynamic tracking or null settings
 		return;
 	}
 
@@ -262,6 +264,14 @@ bool FPCGIndirectionElement::PrepareDataInternal(FPCGContext* InContext) const
 		Context->InnerContext->DependenciesCrc = Context->DependenciesCrc;
 
 		Context->InnerContext->AsyncState = Context->AsyncState;
+
+#if WITH_EDITOR
+		// If we have an override, register for dynamic tracking.
+		if (Context->IsValueOverriden(GET_MEMBER_NAME_CHECKED(UPCGIndirectionSettings, Settings)))
+		{
+			FPCGDynamicTrackingHelper::AddSingleDynamicTrackingKey(Context, FPCGSelectionKey::CreateFromPath(Context->InnerSettings), /*bIsCulled=*/false);
+		}
+#endif // WITH_EDITOR
 	}
 
 	return true;
diff --git a/Engine/Plugins/Experimental/PCG/Source/PCG/Private/Elements/PCGLoopElement.cpp b/Engine/Plugins/Experimental/PCG/Source/PCG/Private/Elements/PCGLoopElement.cpp
--- a/Engine/Plugins/Experimental/PCG/Source/PCG/Private/Elements/PCGLoopElement.cpp
+++ b/Engine/Plugins/Experimental/PCG/Source/PCG/Private/Elements/PCGLoopElement.cpp
@@ -8,6 +8,7 @@
 #include "PCGPin.h"
 #include "PCGSubsystem.h"
 #include "Graph/PCGStackContext.h"
+#include "Helpers/PCGDynamicTrackingHelpers.h"
 #include "Helpers/PCGHelpers.h"
 
 #include "Algo/Copy.h"
@@ -180,6 +181,9 @@ bool FPCGLoopElement::ExecuteInternal(FPCGContext* InContext) const
 			// If OriginalSettings is null, then we ARE the original settings, and writing over the existing graph is incorrect (and potentially a race condition)
 			check(Settings->OriginalSettings);
 			Settings->SubgraphInstance->SetGraph(Settings->SubgraphOverride);
+#if WITH_EDITOR
+			FPCGDynamicTrackingHelper::AddSingleDynamicTrackingKey(Context, FPCGSelectionKey::CreateFromPath(Settings->SubgraphOverride), /*bIsCulled=*/false);
+#endif // WITH_EDITOR
 		}
 
 		UPCGGraph* Subgraph = Settings->GetSubgraph();
diff --git a/Engine/Plugins/Experimental/PCG/Source/PCG/Private/Elements/PCGTextureSampler.cpp b/Engine/Plugins/Experimental/PCG/Source/PCG/Private/Elements/PCGTextureSampler.cpp
--- a/Engine/Plugins/Experimental/PCG/Source/PCG/Private/Elements/PCGTextureSampler.cpp
+++ b/Engine/Plugins/Experimental/PCG/Source/PCG/Private/Elements/PCGTextureSampler.cpp
@@ -6,6 +6,7 @@
 #include "PCGContext.h"
 #include "PCGCrc.h"
 #include "Helpers/PCGBlueprintHelpers.h"
+#include "Helpers/PCGDynamicTrackingHelpers.h"
 #include "Helpers/PCGHelpers.h"
 #include "Helpers/PCGSettingsHelpers.h"
 
@@ -19,6 +20,19 @@
 #define LOCTEXT_NAMESPACE "PCGTextureSamplerElement"
 
 #if WITH_EDITOR
+void UPCGTextureSamplerSettings::GetStaticTrackedKeys(FPCGSelectionKeyToSettingsMap& OutKeysToSettings, TArray<TObjectPtr<const UPCGGraph>>& OutVisitedGraphs) const
+{
+	if (IsPropertyOverriddenByPin(GET_MEMBER_NAME_CHECKED(UPCGTextureSamplerSettings, Texture)) || Texture.IsNull())
+	{
+		// Dynamic tracking or null settings
+		return;
+	}
+
+	FPCGSelectionKey Key = FPCGSelectionKey::CreateFromPath(Texture.ToSoftObjectPath());
+
+	OutKeysToSettings.FindOrAdd(Key).Emplace(this, /*bCulling=*/false);
+}
+
 void UPCGTextureSamplerSettings::PostEditChangeProperty(FPropertyChangedEvent& PropertyChangedEvent) 
 {
 	if (PropertyChangedEvent.Property)
@@ -191,6 +205,14 @@ bool FPCGTextureSamplerElement::ExecuteInternal(FPCGContext* InContext) const
 	TextureData->bUseTileBounds = bUseTileBounds;
 	TextureData->TileBounds = FBox2D(TileBoundsMin, TileBoundsMax);
 
+#if WITH_EDITOR
+	// If we have an override, register for dynamic tracking.
+	if (Context->IsValueOverriden(GET_MEMBER_NAME_CHECKED(UPCGTextureSamplerSettings, Texture)))
+	{
+		FPCGDynamicTrackingHelper::AddSingleDynamicTrackingKey(Context, FPCGSelectionKey::CreateFromPath(Texture), /*bIsCulled=*/false);
+	}
+#endif // WITH_EDITOR
+
 	return false;
 }
 
diff --git a/Engine/Plugins/Experimental/PCG/Source/PCG/Private/Elements/PCGUserParameterGet.cpp b/Engine/Plugins/Experimental/PCG/Source/PCG/Private/Elements/PCGUserParameterGet.cpp
--- a/Engine/Plugins/Experimental/PCG/Source/PCG/Private/Elements/PCGUserParameterGet.cpp
+++ b/Engine/Plugins/Experimental/PCG/Source/PCG/Private/Elements/PCGUserParameterGet.cpp
@@ -9,6 +9,7 @@
 #include "PCGPin.h"
 #include "PCGSubgraph.h"
 #include "Data/PCGUserParametersData.h"
+#include "Helpers/PCGDynamicTrackingHelpers.h"
 #include "Helpers/PCGPropertyHelpers.h"
 #include "Metadata/PCGAttributePropertySelector.h"
 
@@ -163,11 +164,33 @@ bool FPCGUserParameterGetElement::ExecuteInternal(FPCGContext* Context) const
 
 	PCGPropertyHelpers::FExtractorParameters ExtractorParameters{ Parameters.GetMemory(), PropertyBag, FPCGAttributePropertySelector::CreateAttributeSelector(PropertyName), PropertyName, Settings->bForceObjectAndStructExtraction, /*bPropertyNeedsToBeVisible=*/false};
 
-	if (UPCGParamData* ParamData = PCGPropertyHelpers::ExtractPropertyAsAttributeSet(ExtractorParameters, Context))
+	// Don't care for object traversed in non-editor build, since it is only useful for tracking.
+	TSet<FSoftObjectPath>* ObjectTraversedPtr = nullptr;
+#if WITH_EDITOR
+	TSet<FSoftObjectPath> ObjectTraversed;
+	ObjectTraversedPtr = &ObjectTraversed;
+#endif // WITH_EDITOR
+
+	if (UPCGParamData* ParamData = PCGPropertyHelpers::ExtractPropertyAsAttributeSet(ExtractorParameters, Context, ObjectTraversedPtr))
 	{
 		Context->OutputData.TaggedData.Emplace_GetRef().Data = ParamData;
 	}
 
+	// Register dynamic tracking
+#if WITH_EDITOR
+	if (!ObjectTraversed.IsEmpty())
+	{
+		FPCGDynamicTrackingHelper DynamicTracking;
+		DynamicTracking.EnableAndInitialize(Context, ObjectTraversed.Num());
+		for (FSoftObjectPath& Path : ObjectTraversed)
+		{
+			DynamicTracking.AddToTracking(FPCGSelectionKey::CreateFromPath(std::move(Path)), /*bCulled=*/false);
+		}
+
+		DynamicTracking.Finalize(Context);
+	}
+#endif // WITH_EDITOR
+
 	return true;
 }
 
diff --git a/Engine/Plugins/Experimental/PCG/Source/PCG/Private/Helpers/PCGDynamicTrackingHelpers.cpp b/Engine/Plugins/Experimental/PCG/Source/PCG/Private/Helpers/PCGDynamicTrackingHelpers.cpp
new file mode 100644
--- /dev/null
+++ b/Engine/Plugins/Experimental/PCG/Source/PCG/Private/Helpers/PCGDynamicTrackingHelpers.cpp
@@ -0,0 +1,55 @@
+// Copyright Epic Games, Inc. All Rights Reserved.
+
+#include "Helpers/PCGDynamicTrackingHelpers.h"
+
+#if WITH_EDITOR
+
+#include "PCGComponent.h"
+#include "PCGContext.h"
+
+void FPCGDynamicTrackingHelper::EnableAndInitialize(const FPCGContext* InContext, int32 OptionalNumElements)
+{
+	if (InContext && InContext->SourceComponent.IsValid())
+	{
+		const UPCGSettings* Settings = InContext->GetOriginalSettings<UPCGSettings>();
+		if (Settings && Settings->CanDynamicalyTrackKeys())
+		{
+			CachedComponent = InContext->SourceComponent;
+			bDynamicallyTracked = true;
+			DynamicallyTrackedKeysAndCulling.Reserve(OptionalNumElements);
+		}
+	}
+}
+
+void FPCGDynamicTrackingHelper::AddToTracking(FPCGSelectionKey&& InKey, bool bIsCulled)
+{
+	if (bDynamicallyTracked)
+	{
+		DynamicallyTrackedKeysAndCulling.AddUnique(TPair<FPCGSelectionKey, bool>(std::forward<FPCGSelectionKey>(InKey), bIsCulled));
+	}
+}
+
+void FPCGDynamicTrackingHelper::Finalize(const FPCGContext* InContext)
+{
+	if (InContext && bDynamicallyTracked && CachedComponent == InContext->SourceComponent)
+	{
+		if (UPCGComponent* SourceComponent = CachedComponent.Get())
+		{
+			SourceComponent->RegisterDynamicTracking(InContext->GetOriginalSettings<UPCGSettings>(), DynamicallyTrackedKeysAndCulling);
+		}
+	}
+}
+
+void FPCGDynamicTrackingHelper::AddSingleDynamicTrackingKey(FPCGContext* InContext, FPCGSelectionKey&& InKey, bool bIsCulled)
+{
+	if (InContext)
+	{
+		if (UPCGComponent* SourceComponent = InContext->SourceComponent.Get())
+		{
+			TPair<FPCGSelectionKey, bool> NewPair(std::forward<FPCGSelectionKey>(InKey), bIsCulled);
+			SourceComponent->RegisterDynamicTracking(InContext->GetOriginalSettings<UPCGSettings>(), MakeArrayView(&NewPair, 1));
+		}
+	}
+}
+
+#endif // WITH_EDITOR
diff --git a/Engine/Plugins/Experimental/PCG/Source/PCG/Private/Helpers/PCGPropertyHelpers.cpp b/Engine/Plugins/Experimental/PCG/Source/PCG/Private/Helpers/PCGPropertyHelpers.cpp
--- a/Engine/Plugins/Experimental/PCG/Source/PCG/Private/Helpers/PCGPropertyHelpers.cpp
+++ b/Engine/Plugins/Experimental/PCG/Source/PCG/Private/Helpers/PCGPropertyHelpers.cpp
@@ -34,15 +34,16 @@ namespace PCGPropertyHelpers
 
 	/**
 	* Recursive function to go down the property chain to find the property and its container address.
-	* @param CurrentClass      Struct/Class for the current container
-	* @param CurrentName       Property name to look for in the container class.
-	* @param NextNames         List of property names to continue extracting at a deeper level.
-	* @param bNeedsToBeVisible Discard properties that are not visibile in Blueprint
-	* @param OutContainer      Raw address for the current container. Will be write to at each recursive call.
-	* @param OptionalContext   Optional context used for logging.
+	* @param CurrentClass            Struct/Class for the current container
+	* @param CurrentName             Property name to look for in the container class.
+	* @param NextNames               List of property names to continue extracting at a deeper level.
+	* @param bNeedsToBeVisible       Discard properties that are not visibile in Blueprint
+	* @param OutContainer            Raw address for the current container. Will be write to at each recursive call.
+	* @param OptionalContext         Optional context used for logging.
+	* @param OptionalObjectTraversed Optional set to store all object that we traversed, to be able to react to those objects changes.
 	* @returns                 The last property of the chain (and its container address is in OutContainer)
 	*/
-	const FProperty* ExtractPropertyChain(const UStruct* CurrentClass, const FName CurrentName, TArrayView<const FString> NextNames, const bool bNeedsToBeVisible, const void*& OutContainer, FPCGContext* OptionalContext)
+	const FProperty* ExtractPropertyChain(const UStruct* CurrentClass, const FName CurrentName, TArrayView<const FString> NextNames, const bool bNeedsToBeVisible, const void*& OutContainer, FPCGContext* OptionalContext, TSet<FSoftObjectPath>* OptionalObjectTraversed)
 	{
 		check(CurrentClass);
 
@@ -98,7 +99,7 @@ namespace PCGPropertyHelpers
 				return nullptr;
 			}
 
-			return ExtractPropertyChain(NextClass, FName(NextNames[0]), NextNames.RightChop(1), bNeedsToBeVisible, OutContainer, OptionalContext);
+			return ExtractPropertyChain(NextClass, FName(NextNames[0]), NextNames.RightChop(1), bNeedsToBeVisible, OutContainer, OptionalContext, OptionalObjectTraversed);
 		}
 		else
 		{
@@ -125,7 +126,7 @@ EPCGMetadataTypes PCGPropertyHelpers::GetMetadataTypeFromProperty(const FPropert
 	return PropertyAccessor.IsValid() ? EPCGMetadataTypes(PropertyAccessor->GetUnderlyingType()) : EPCGMetadataTypes::Unknown;
 }
 
-UPCGParamData* PCGPropertyHelpers::ExtractPropertyAsAttributeSet(const PCGPropertyHelpers::FExtractorParameters& Parameters, FPCGContext* InOptionalContext)
+UPCGParamData* PCGPropertyHelpers::ExtractPropertyAsAttributeSet(const PCGPropertyHelpers::FExtractorParameters& Parameters, FPCGContext* OptionalContext, TSet<FSoftObjectPath>* OptionalObjectTraversed)
 {
 	check(Parameters.Container && Parameters.Class);
 
@@ -136,7 +137,7 @@ UPCGParamData* PCGPropertyHelpers::ExtractPropertyAsAttributeSet(const PCGProper
 	// If Name is none, extract the container as-is, using Parameters.Class, otherwise, extract the chain.
 	if (!ExtractRoot)
 	{
-		Property = ExtractPropertyChain(Parameters.Class, PropertyName, Parameters.PropertySelector.GetExtraNames(), Parameters.bPropertyNeedsToBeVisible, Container, InOptionalContext);
+		Property = ExtractPropertyChain(Parameters.Class, PropertyName, Parameters.PropertySelector.GetExtraNames(), Parameters.bPropertyNeedsToBeVisible, Container, OptionalContext, OptionalObjectTraversed);
 		if (!Property)
 		{
 			return nullptr;
@@ -159,6 +160,9 @@ UPCGParamData* PCGPropertyHelpers::ExtractPropertyAsAttributeSet(const PCGProper
 	// Force extraction if the property is not supported by accessors.
 	const bool bShouldExtract = Parameters.bShouldExtract || !PCGAttributeAccessorHelpers::IsPropertyAccessorSupported(Property);
 
+	// Keep track if the extracted property is an object or not
+	const FObjectProperty* ObjectProperty = CastField<FObjectProperty>(Property);
+
 	// Special case where the property is a struct/object, that is not supported by our metadata, we will try to break it down to multiple attributes in the resulting param data, if asked.
 	if (ExtractRoot || ((Property->IsA<FStructProperty>() || Property->IsA<FObjectProperty>()) && bShouldExtract))
 	{
@@ -175,7 +179,7 @@ UPCGParamData* PCGPropertyHelpers::ExtractPropertyAsAttributeSet(const PCGProper
 			UnderlyingClass = StructProperty->Struct;
 			AddressFunc = [StructProperty](const void* InAddress) { return StructProperty->ContainerPtrToValuePtr<void>(InAddress); };
 		}
-		else if (const FObjectProperty* ObjectProperty = CastField<FObjectProperty>(Property))
+		else if (ObjectProperty)
 		{
 			UnderlyingClass = ObjectProperty->PropertyClass;
 			AddressFunc = [ObjectProperty](const void* InAddress) { return ObjectProperty->GetObjectPropertyValue_InContainer(InAddress); };
@@ -224,7 +228,7 @@ UPCGParamData* PCGPropertyHelpers::ExtractPropertyAsAttributeSet(const PCGProper
 
 	if (ExtractableProperties.IsEmpty())
 	{
-		LogError(LOCTEXT("NoPropertiesFound", "No properties found to extract"), InOptionalContext);
+		LogError(LOCTEXT("NoPropertiesFound", "No properties found to extract"), OptionalContext);
 		return nullptr;
 	}
 
@@ -261,21 +265,30 @@ UPCGParamData* PCGPropertyHelpers::ExtractPropertyAsAttributeSet(const PCGProper
 		// Add a new entry for all elements
 		PCGMetadataEntryKey EntryKey = Metadata->AddEntry();
 
+		// Offset the address if needed
+		const void* ContainerPtr = AddressFunc(ElementAddress);
+
 		for (ExtractablePropertyTuple& ExtractableProperty : ExtractableProperties)
 		{
 			const FName AttributeName = ExtractableProperty.Get<0>();
 			const FProperty* FinalProperty = ExtractableProperty.Get<1>();
 
-			// Offset the address if needed
-			const void* ContainerPtr = AddressFunc(ElementAddress);
-
 			if (!Metadata->SetAttributeFromDataProperty(AttributeName, EntryKey, ContainerPtr, FinalProperty, /*bCreate=*/ true))
 			{
-				LogError(FText::Format(LOCTEXT("ErrorCreatingAttribute", "Error while creating an attribute for property '{0}'. Either the property type is not supported by PCG or attribute creation failed."), FText::FromString(FinalProperty->GetName())), InOptionalContext);
+				LogError(FText::Format(LOCTEXT("ErrorCreatingAttribute", "Error while creating an attribute for property '{0}'. Either the property type is not supported by PCG or attribute creation failed."), FText::FromString(FinalProperty->GetName())), OptionalContext);
 				bValidOperation = false;
 				break;
 			}
 		}
+
+		if (bValidOperation && bShouldExtract && OptionalObjectTraversed && ObjectProperty)
+		{
+			const UObject* Object = ObjectProperty->GetPropertyValue_InContainer(ElementAddress);
+			if (IsValid(Object))
+			{
+				OptionalObjectTraversed->Add(Object);
+			}
+		}
 	}
 
 	return bValidOperation ? ParamData : nullptr;
diff --git a/Engine/Plugins/Experimental/PCG/Source/PCG/Private/PCGActorAndComponentMapping.cpp b/Engine/Plugins/Experimental/PCG/Source/PCG/Private/PCGActorAndComponentMapping.cpp
--- a/Engine/Plugins/Experimental/PCG/Source/PCG/Private/PCGActorAndComponentMapping.cpp
+++ b/Engine/Plugins/Experimental/PCG/Source/PCG/Private/PCGActorAndComponentMapping.cpp
@@ -30,6 +30,11 @@ namespace PCGActorAndComponentMapping
 {
 	FBox GetActorBounds(const AActor* InActor)
 	{
+		if (!InActor)
+		{
+			return FBox(EForceInit::ForceInit);
+		}
+
 		FBox ActorBounds = InActor->GetComponentsBoundingBox();
 		if (!ActorBounds.IsValid && InActor->GetRootComponent() != nullptr)
 		{
@@ -824,8 +829,7 @@ void FPCGActorAndComponentMapping::RegisterTracking(UPCGComponent* InComponent)
 	}
 
 	// Components owner needs to be always tracked
-	RegisterActor(ComponentOwner);
-	AlwaysTrackedActorsToComponentsMap.FindOrAdd(ComponentOwner).Add(InComponent);
+	AlwaysTrackedKeysToComponentsMap.FindOrAdd(FPCGSelectionKey(EPCGActorFilter::Self)).Add(InComponent);
 
 	UpdateTracking(InComponent, /*bInShouldDirtyActors=*/ false);
 
@@ -864,138 +868,52 @@ void FPCGActorAndComponentMapping::UpdateTracking(UPCGComponent* InComponent, bo
 	TArray<FPCGSelectionKey> AllKeys;
 	if (ChangedKeys == nullptr)
 	{
-		InComponent->CachedTrackedKeysToSettings.GenerateKeyArray(AllKeys);
+		AllKeys = InComponent->GatherTrackingKeys();
 		ChangedKeys = &AllKeys;
 	}
 
 	check(ChangedKeys);
 
-	if (ChangedKeys->IsEmpty())
-	{
-		// Nothing to do
-		return;
-	}
-
-	// And we also need to find all actors that should be tracked
-	TSet<TSoftObjectPtr<UObject>> CandidatesForTracking;
-	TSet<TSoftObjectPtr<UObject>> CandidatesForUntracking;
-
-	auto GatherObjects = [this, InComponent, &CandidatesForTracking, &CandidatesForUntracking](const FPCGSelectionKey& InKey, bool bInShouldUntrack)
+	auto RemoveFromMap = [InComponent](TMap<FPCGSelectionKey, TSet<UPCGComponent*>>& InMap, const FPCGSelectionKey& InKey)
 	{
-		if (InKey.Selection == EPCGActorSelection::ByPath)
-		{
-			// Don't track null objects
-			if (InKey.ObjectPath.IsNull())
-			{
-				return;
-			}
-
-			TSoftObjectPtr<UObject> Object(InKey.ObjectPath);
-			if (bInShouldUntrack && !CandidatesForTracking.Contains(Object))
-			{
-				CandidatesForUntracking.Add(Object);
-			}
-			else
-			{
-				CandidatesForTracking.Add(Object);
-				CandidatesForUntracking.Remove(Object);
-			}
-
-			return;
-		}
-
-		// InKey provide the info for selecting a given actor.
-		// We reconstruct the selector settings from this key, and we also force it to SelectMultiple, since
-		// we want to gather all the actors that matches this given key.
-		FPCGActorSelectorSettings SelectorSettings = FPCGActorSelectorSettings::ReconstructFromKey(InKey);
-		SelectorSettings.bSelectMultiple = true;
-		TArray<AActor*> AllActors = PCGActorSelector::FindActors(SelectorSettings, InComponent, [](const AActor*) { return true; }, [](const AActor*) { return true; });
-
-		for (AActor* Actor : AllActors)
+		if (TSet<UPCGComponent*>* Components = InMap.Find(InKey))
 		{
-			if (!Actor)
+			Components->Remove(InComponent);
+			if (Components->IsEmpty())
 			{
-				continue;
-			}
-
-			if (bInShouldUntrack && !CandidatesForTracking.Contains(Actor))
-			{
-				CandidatesForUntracking.Add(Actor);
-			}
-			else
-			{
-				CandidatesForTracking.Add(Actor);
-				CandidatesForUntracking.Remove(Actor);
+				InMap.Remove(InKey);
 			}
 		}
 	};
 
 	for (const FPCGSelectionKey& Key : *ChangedKeys)
 	{
-		const bool bShouldUntrack = !InComponent->CachedTrackedKeysToSettings.Contains(Key);
-		GatherObjects(Key, bShouldUntrack);
-	}
-
-	const bool bDisableDelayedActorRegistering = PCGActorAndComponentMapping::CVarDisableDelayedActorRegistering.GetValueOnAnyThread();
-	for (const TSoftObjectPtr<UObject>& ObjectPtr : CandidatesForTracking)
-	{
-		AActor* Actor = Cast<AActor>(ObjectPtr.Get());
-
-		if (Actor && !Actor->HasActorRegisteredAllComponents() && !bDisableDelayedActorRegistering)
+		// bShouldBeCulled is modified in IsKeyTrackedAndCulled
+		bool bShouldBeCulled = false;
+		if (!InComponent->IsKeyTrackedAndCulled(Key, bShouldBeCulled))
 		{
-			DelayedAddedActors.Emplace(Actor, { false, 0 });
-			continue;
+			// Untrack
+			RemoveFromMap(CulledTrackedKeysToComponentsMap, Key);
+			RemoveFromMap(AlwaysTrackedKeysToComponentsMap, Key);
 		}
-
-		bool bShouldCull = false;
-		if (!InComponent->IsObjectTracked(ObjectPtr, bShouldCull))
-		{
-			continue;
-		}
-
-		if (Actor)
+		else if (bShouldBeCulled)
 		{
-			// Making sure that we only have the component in one map.
-			auto RemoveFromMap = [InComponent, Actor](TMap<TObjectKey<AActor>, TSet<UPCGComponent*>>& InMap)
-			{
-				if (TSet<UPCGComponent*>* Components = InMap.Find(Actor))
-				{
-					Components->Remove(InComponent);
-					if (Components->IsEmpty())
-					{
-						InMap.Remove(Actor);
-					}
-				}
-			};
-
-			if (bShouldCull)
-			{
-				CulledTrackedActorsToComponentsMap.FindOrAdd(Actor).Add(InComponent);
-				RemoveFromMap(AlwaysTrackedActorsToComponentsMap);
-			}
-			else
-			{
-				AlwaysTrackedActorsToComponentsMap.FindOrAdd(Actor).Add(InComponent);
-				RemoveFromMap(CulledTrackedActorsToComponentsMap);
-			}
-
-			RegisterActor(Actor);
+			RemoveFromMap(AlwaysTrackedKeysToComponentsMap, Key);
+			CulledTrackedKeysToComponentsMap.FindOrAdd(Key).Add(InComponent);
 		}
 		else
 		{
-			TrackedObjectsToComponentsMap.FindOrAdd(ObjectPtr).Add(InComponent);
+			RemoveFromMap(CulledTrackedKeysToComponentsMap, Key);
+			AlwaysTrackedKeysToComponentsMap.FindOrAdd(Key).Add(InComponent);
 		}
 	}
-
-	// Also unregister all keys that are not tracked anymore.
-	UnregisterTracking(InComponent, &CandidatesForUntracking);
 }
 
 void FPCGActorAndComponentMapping::RemapTracking(const UPCGComponent* InOldComponent, UPCGComponent* InNewComponent)
 {
-	auto ReplaceInMap = [InOldComponent, InNewComponent]<typename Key>(TMap<Key, TSet<UPCGComponent*>>& InMap)
+	auto ReplaceInMap = [InOldComponent, InNewComponent](TMap<FPCGSelectionKey, TSet<UPCGComponent*>>& InMap)
 	{
-		for (TPair<Key, TSet<UPCGComponent*>>& It : InMap)
+		for (TPair<FPCGSelectionKey, TSet<UPCGComponent*>>& It : InMap)
 		{
 			if (It.Value.Remove(InOldComponent) > 0)
 			{
@@ -1004,9 +922,8 @@ void FPCGActorAndComponentMapping::RemapTracking(const UPCGComponent* InOldCompo
 		}
 	};
 
-	ReplaceInMap(CulledTrackedActorsToComponentsMap);
-	ReplaceInMap(AlwaysTrackedActorsToComponentsMap);
-	ReplaceInMap(TrackedObjectsToComponentsMap);
+	ReplaceInMap(CulledTrackedKeysToComponentsMap);
+	ReplaceInMap(AlwaysTrackedKeysToComponentsMap);
 
 	// Old component will probably die, but we'll force removing the delegates even if it is const.
 	if (UPCGComponent* MutableOldComponent = const_cast<UPCGComponent*>(InOldComponent))
@@ -1023,79 +940,56 @@ void FPCGActorAndComponentMapping::RemapTracking(const UPCGComponent* InOldCompo
 	}
 }
 
-void FPCGActorAndComponentMapping::UnregisterTracking(UPCGComponent* InComponent, const TSet<TSoftObjectPtr<UObject>>* OptionalObjectsToUntrack)
+void FPCGActorAndComponentMapping::UnregisterTracking(UPCGComponent* InComponent, const TSet<FPCGSelectionKey>* OptionalKeysToUntrack)
 {
-	if (!InComponent || (OptionalObjectsToUntrack && OptionalObjectsToUntrack->IsEmpty()))
+	if (!InComponent)
 	{
 		return;
 	}
 
-	TSet<AActor*> ActorCandidatesForUntrack;
-
-	if (OptionalObjectsToUntrack)
+	TSet<FPCGSelectionKey> KeysToRemove;
+	auto RemoveAllFromMap = [InComponent, &KeysToRemove](TMap<FPCGSelectionKey, TSet<UPCGComponent*>>& InMap)
 	{
-		auto RemoveFromMap = [InComponent](auto& InMap, auto& Key, auto OnEmpty)
+		for (auto& It : InMap)
 		{
-			if (TSet<UPCGComponent*>* It = InMap.Find(Key))
+			It.Value.Remove(InComponent);
+			if (It.Value.IsEmpty())
 			{
-				It->Remove(InComponent);
-				if (It->IsEmpty())
-				{
-					OnEmpty(Key);
-				}
-			}
-		};
-
-		for (const TSoftObjectPtr<UObject>& ObjectPtr : *OptionalObjectsToUntrack)
-		{
-			if (AActor* Actor = Cast<AActor>(ObjectPtr.Get()))
-			{
-				RemoveFromMap(CulledTrackedActorsToComponentsMap, Actor, [&ActorCandidatesForUntrack](AActor* Actor) { ActorCandidatesForUntrack.Add(Actor); });
-				RemoveFromMap(AlwaysTrackedActorsToComponentsMap, Actor, [&ActorCandidatesForUntrack](AActor* Actor) { ActorCandidatesForUntrack.Add(Actor); });
-			}
-			else
-			{
-				RemoveFromMap(TrackedObjectsToComponentsMap, ObjectPtr, [this](const TSoftObjectPtr<UObject>& ObjectPtr) { UnregisterObject(ObjectPtr); });
+				KeysToRemove.Add(It.Key);
 			}
 		}
-	}
-	else
+	};
+
+	auto RemoveKeysFromMap = [InComponent, &KeysToRemove](TMap<FPCGSelectionKey, TSet<UPCGComponent*>>& InMap, const TSet<FPCGSelectionKey>& SetToIterateOn)
 	{
-		auto RemoveAllFromMap = [InComponent](auto& InMap, auto OnEmpty)
+		for (const FPCGSelectionKey& KeyIt : SetToIterateOn)
 		{
-			for (auto& It : InMap)
+			if (TSet<UPCGComponent*>* ComponentSetPtr = InMap.Find(KeyIt))
 			{
-				It.Value.Remove(InComponent);
-				if (It.Value.IsEmpty())
+				ComponentSetPtr->Remove(InComponent);
+				if (ComponentSetPtr->IsEmpty())
 				{
-					OnEmpty(It.Key);
+					KeysToRemove.Add(KeyIt);
 				}
 			}
-		};
+		}
+	};
 
-		RemoveAllFromMap(CulledTrackedActorsToComponentsMap, [&ActorCandidatesForUntrack](const TObjectKey<AActor>& ActorKey) { ActorCandidatesForUntrack.Add(ActorKey.ResolveObjectPtr()); });
-		RemoveAllFromMap(AlwaysTrackedActorsToComponentsMap, [&ActorCandidatesForUntrack](const TObjectKey<AActor>& ActorKey) { ActorCandidatesForUntrack.Add(ActorKey.ResolveObjectPtr()); });
-		RemoveAllFromMap(TrackedObjectsToComponentsMap, [this](const TSoftObjectPtr<UObject>& ObjectPtr) { UnregisterObject(ObjectPtr); });
+	if (OptionalKeysToUntrack)
+	{
+		RemoveKeysFromMap(CulledTrackedKeysToComponentsMap, *OptionalKeysToUntrack);
+		RemoveKeysFromMap(AlwaysTrackedKeysToComponentsMap, *OptionalKeysToUntrack);
 	}
-
-	// We also need to untrack actors that doesn't have any component that tracks them.
-	auto ShouldBeRemoved = [](const AActor* InActor, TMap<TObjectKey<AActor>, TSet<UPCGComponent*>>& InMap)
+	else
 	{
-		TSet<UPCGComponent*>* RegisteredComponents = InMap.Find(InActor);
-		return !RegisteredComponents || RegisteredComponents->IsEmpty();
-	};
+		RemoveAllFromMap(CulledTrackedKeysToComponentsMap);
+		RemoveAllFromMap(AlwaysTrackedKeysToComponentsMap);
+	}
 
-	for (AActor* Candidate : ActorCandidatesForUntrack)
+	for (const FPCGSelectionKey& Key : KeysToRemove)
 	{
-		if (!Candidate)
-		{
-			continue;
-		}
-
-		if (Candidate && ShouldBeRemoved(Candidate, AlwaysTrackedActorsToComponentsMap) && ShouldBeRemoved(Candidate, CulledTrackedActorsToComponentsMap))
-		{
-			UnregisterActor(Candidate);
-		}
+		CulledTrackedKeysToComponentsMap.Remove(Key);
+		AlwaysTrackedKeysToComponentsMap.Remove(Key);
 	}
 }
 
@@ -1112,6 +1006,11 @@ void FPCGActorAndComponentMapping::UnregisterTracking(UPCGComponent* InComponent
 	InComponent->OnPCGGraphCleanedDelegate.RemoveAll(this);
 }
 
+bool FPCGActorAndComponentMapping::IsKeyTracked(const FPCGSelectionKey& InKey) const
+{
+	return CulledTrackedKeysToComponentsMap.Contains(InKey) || AlwaysTrackedKeysToComponentsMap.Contains(InKey);
+}
+
 void FPCGActorAndComponentMapping::ResetPartitionActorsMap()
 {
 	PartitionActorsMapLock.WriteLock();
@@ -1121,7 +1020,6 @@ void FPCGActorAndComponentMapping::ResetPartitionActorsMap()
 
 void FPCGActorAndComponentMapping::RegisterTrackingCallbacks()
 {
-	GEngine->OnActorMoved().AddRaw(this, &FPCGActorAndComponentMapping::OnActorMoved);
 	GEngine->OnLevelActorAdded().AddRaw(this, &FPCGActorAndComponentMapping::OnActorAdded);
 	GEngine->OnLevelActorDeleted().AddRaw(this, &FPCGActorAndComponentMapping::OnActorDeleted);
 	FCoreUObjectDelegates::OnObjectPropertyChanged.AddRaw(this, &FPCGActorAndComponentMapping::OnObjectPropertyChanged);
@@ -1139,7 +1037,6 @@ void FPCGActorAndComponentMapping::RegisterTrackingCallbacks()
 
 void FPCGActorAndComponentMapping::TeardownTrackingCallbacks()
 {
-	GEngine->OnActorMoved().RemoveAll(this);
 	GEngine->OnLevelActorAdded().RemoveAll(this);
 	GEngine->OnLevelActorDeleted().RemoveAll(this);
 	FCoreUObjectDelegates::OnObjectPropertyChanged.RemoveAll(this);
@@ -1239,6 +1136,20 @@ void FPCGActorAndComponentMapping::OnActorAdded_Internal(AActor* InActor, bool b
 	TRACE_CPUPROFILER_EVENT_SCOPE(FPCGActorAndComponentMapping::OnActorAdded);
 	check(InActor && !InActor->IsA<APCGWorldActor>() && PCGSubsystem && InActor->GetWorld() == PCGSubsystem->GetWorld());
 
+	if (ALandscapeProxy* LandscapeProxy = Cast<ALandscapeProxy>(InActor))
+	{
+		if (!LandscapeProxy->OnComponentDataChanged.IsBoundToObject(this))
+		{
+			LandscapeProxy->OnComponentDataChanged.AddRaw(this, &FPCGActorAndComponentMapping::OnLandscapeChanged);
+		}
+	}
+
+	if (!bShouldDirty)
+	{
+		// Nothing to do
+		return;
+	}
+
 	// A delayed actor should only be added from AddDelayedActors(), because that guarantees the new actor has waited at
 	// least one tick for its properties to be popualated.
 	if (!bForceAddDelayedActor && DelayedAddedActors.Contains(InActor))
@@ -1285,142 +1196,8 @@ void FPCGActorAndComponentMapping::OnActorAdded_Internal(AActor* InActor, bool b
 	}
 #endif // WITH_EDITOR
 
-	if (AddOrUpdateTrackedActor(InActor) && bShouldDirty)
-	{
-		// Finally notify them all
-		OnActorChanged(InActor, /*bInHasMoved=*/ false, nullptr, LevelInstanceDepth);
-	}
-}
-
-bool FPCGActorAndComponentMapping::AddOrUpdateTrackedActor(AActor* InActor)
-{
-	// We have to make sure to not create a infinite loop
-	if (!InActor || InActor->IsA<APCGWorldActor>() || !PCGSubsystem || InActor->GetWorld() != PCGSubsystem->GetWorld())
-	{
-		return false;
-	}
-
-#if WITH_EDITOR
-	if (InActor->bIsEditorPreviewActor)
-	{
-		return false;
-	}
-#endif
-
-	// Gather all components, and check if they want to track this one
-	TSet<UPCGComponent*> AllComponents = GetAllRegisteredComponents();
-
-	TSet<UPCGComponent*>* CulledTrackedComponents = nullptr;
-	TSet<UPCGComponent*>* AlwaysTrackedComponents = nullptr;
-	
-	for (UPCGComponent* PCGComponent : AllComponents)
-	{
-		// Making sure that they are not currently waiting to die
-		{
-			FScopeLock Lock(&DelayedComponentToUnregisterLock);
-			if (DelayedComponentToUnregister.Contains(PCGComponent))
-			{
-				continue;
-			}
-		}
-
-		bool bTrackingIsCulled = false;
-		if (PCGComponent && PCGComponent->IsObjectTracked(InActor, bTrackingIsCulled))
-		{
-			if (bTrackingIsCulled)
-			{
-				if (!CulledTrackedComponents)
-				{
-					CulledTrackedComponents = &CulledTrackedActorsToComponentsMap.FindOrAdd(InActor);
-				}
-
-				check(CulledTrackedComponents);
-				CulledTrackedComponents->Add(PCGComponent);
-			}
-			else
-			{
-				if (!AlwaysTrackedComponents)
-				{
-					AlwaysTrackedComponents = &AlwaysTrackedActorsToComponentsMap.FindOrAdd(InActor);
-				}
-
-				check(AlwaysTrackedComponents);
-				AlwaysTrackedComponents->Add(PCGComponent);
-			}
-		}
-	}
-
-	if (CulledTrackedComponents || AlwaysTrackedComponents)
-	{
-		RegisterActor(InActor);
-		return true;
-	}
-	else if (IsActorTracked(InActor))
-	{
-		// Do some cleanup if the actor was tracked. 
-		// We will force the refresh here, so return false to make sure we don't refresh it twice.
-		OnActorDeleted(InActor);
-		return false;
-	}
-	else
-	{
-		// If it is not tracked, and should not be tracked, just do nothing.
-		return false;
-	}
-}
-
-void FPCGActorAndComponentMapping::RegisterActor(AActor* InActor)
-{
-	if (!InActor)
-	{
-		return;
-	}
-
-	if (ALandscapeProxy* LandscapeProxy = Cast<ALandscapeProxy>(InActor))
-	{
-		// Only add it once.
-		if (!TrackedActorToPositionMap.Contains(InActor))
-		{
-			LandscapeProxy->OnComponentDataChanged.AddRaw(this, &FPCGActorAndComponentMapping::OnLandscapeChanged);
-		}
-	}
-
-	TrackedActorToPositionMap.FindOrAdd(InActor) = PCGActorAndComponentMapping::GetActorBounds(InActor);
-
-	// Also gather dependencies
-	UpdateActorDependencies(InActor);
-}
-
-bool FPCGActorAndComponentMapping::UnregisterActor(AActor* InActor)
-{
-	if (!InActor)
-	{
-		return false;
-	}
-
-	if (IsActorTracked(InActor))
-	{
-		TrackedActorToPositionMap.Remove(InActor);
-		CulledTrackedActorsToComponentsMap.Remove(InActor);
-		AlwaysTrackedActorsToComponentsMap.Remove(InActor);
-		TrackedActorsToDependenciesMap.Remove(InActor);
-
-		if (ALandscapeProxy* LandscapeProxy = Cast<ALandscapeProxy>(InActor))
-		{
-			LandscapeProxy->OnComponentDataChanged.RemoveAll(this);
-		}
-
-		return true;
-	}
-	else
-	{
-		return false;
-	}
-}
-
-bool FPCGActorAndComponentMapping::UnregisterObject(const TSoftObjectPtr<UObject>& InObject)
-{
-	return TrackedObjectsToComponentsMap.Remove(InObject) > 0;
+	// Finally notify them all
+	OnObjectChanged(InActor, FBox(EForceInit::ForceInit), nullptr, LevelInstanceDepth);
 }
 
 void FPCGActorAndComponentMapping::OnActorUnloaded(AActor& InActor)
@@ -1465,6 +1242,17 @@ void FPCGActorAndComponentMapping::OnActorDeleted_Internal(AActor* InActor, bool
 	TRACE_CPUPROFILER_EVENT_SCOPE(UPCGActorAndComponentMapping::OnActorDeleted);
 	check(InActor && PCGSubsystem && InActor->GetWorld() == PCGSubsystem->GetWorld());
 
+	if (ALandscapeProxy* LandscapeProxy = Cast<ALandscapeProxy>(InActor))
+	{
+		LandscapeProxy->OnComponentDataChanged.RemoveAll(this);
+	}
+
+	if (!bShouldDirty)
+	{
+		// Nothing to do
+		return;
+	}
+
 #if WITH_EDITOR
 	// When editing/deleting/replacing a level instance, this will allow us to first remove the previously tracked actors that were under the level instance
 	if (ALevelInstanceEditorInstanceActor* LevelInstanceEditorInstance = Cast<ALevelInstanceEditorInstanceActor>(InActor))
@@ -1496,73 +1284,33 @@ void FPCGActorAndComponentMapping::OnActorDeleted_Internal(AActor* InActor, bool
 	});
 #endif // WITH_EDITOR
 
-	if (!IsActorTracked(InActor))
-	{
-		return;
-	}
-
-	if (bShouldDirty)
-	{
-		// Notify all components that the actor has changed (was removed), but the Refresh will only happen AFTER the actor was actually removed from the world (because of delayed refresh).
-		OnActorChanged(InActor, /*bInHasMoved=*/ false, nullptr, LevelInstanceDepth, /*bNoRefreshOnOwner=*/true);
-	}
-
-	// And then delete everything
-	UnregisterActor(InActor);
+	// Notify all components that the actor has changed (was removed), but the Refresh will only happen AFTER the actor was actually removed from the world (because of delayed refresh).
+	OnObjectChanged(InActor, FBox(EForceInit::ForceInit), nullptr, LevelInstanceDepth, /*bNoRefreshOnOwner=*/true);
 }
 
-void FPCGActorAndComponentMapping::OnActorMoved(AActor* InActor)
-{
-	OnActorMoved_Internal(InActor, /*LevelInstanceDepth=*/0);
-}
-
-void FPCGActorAndComponentMapping::OnActorMoved_Internal(AActor* InActor, int32 LevelInstanceDepth)
+void FPCGActorAndComponentMapping::OnPreObjectPropertyChanged(UObject* InObject, const FEditPropertyChain& InEditPropertyChain)
 {
-	TRACE_CPUPROFILER_EVENT_SCOPE(FPCGActorAndComponentMapping::OnActorMoved);
-
-	if (!InActor || (PCGSubsystem && InActor->GetWorld() != PCGSubsystem->GetWorld()))
-	{
-		return;
-	}
-
-#if WITH_EDITOR
-	if (InActor->bIsEditorPreviewActor)
+	// Nothing to do if we track nothing
+	if (CulledTrackedKeysToComponentsMap.IsEmpty() && AlwaysTrackedKeysToComponentsMap.IsEmpty())
 	{
 		return;
 	}
 
-	// If we've moved a level instance actor, we instead "push" forward the notification while making sure this will
-	//  trigger changes at this "level" in the world hierarchy
-	PCGActorAndComponentMapping::PropagateToLevelInstanceActors(InActor, PCGSubsystem, [this, LevelInstanceDepth](AActor* LevelActor)
-	{
-		OnActorMoved_Internal(LevelActor, LevelInstanceDepth + 1);
-		return true;
-	});
-#endif
-
-	if (!IsActorTracked(InActor))
-	{
-		return;
-	}
-
-	// Notify all components
-	OnActorChanged(InActor, /*bInHasMoved=*/ true, nullptr, LevelInstanceDepth);
-
-	// Update Actor position
-	if (FBox* ActorBounds = TrackedActorToPositionMap.Find(InActor))
-	{
-		*ActorBounds = PCGActorAndComponentMapping::GetActorBounds(InActor);
-	}
-}
-
-void FPCGActorAndComponentMapping::OnPreObjectPropertyChanged(UObject* InObject, const FEditPropertyChain& InEditPropertyChain)
-{
 	// We want to track tags, to see if a tag was removed
 	TempTrackedActorTags.Empty();
 	FProperty* MemberProperty = InEditPropertyChain.GetActiveMemberNode() ? InEditPropertyChain.GetActiveMemberNode()->GetValue() : nullptr;
 	AActor* Actor = Cast<AActor>(InObject);
 
-	if (!Actor || (PCGSubsystem && Actor->GetWorld() != PCGSubsystem->GetWorld()) || !MemberProperty || MemberProperty->GetFName() != GET_MEMBER_NAME_CHECKED(AActor, Tags))
+	// Otherwise, if it's an actor component, track it as well
+	if (!Actor)
+	{
+		if (UActorComponent* ActorComponent = Cast<UActorComponent>(InObject))
+		{
+			Actor = ActorComponent->GetOwner();
+		}
+	}
+
+	if (!Actor || (PCGSubsystem && Actor->GetWorld() != PCGSubsystem->GetWorld()))
 	{
 		return;
 	}
@@ -1574,11 +1322,22 @@ void FPCGActorAndComponentMapping::OnPreObjectPropertyChanged(UObject* InObject,
 	}
 #endif
 
-	TempTrackedActorTags = TSet<FName>(Actor->Tags);
+	if (MemberProperty && MemberProperty->GetFName() == GET_MEMBER_NAME_CHECKED(AActor, Tags))
+	{
+		TempTrackedActorTags = TSet<FName>(Actor->Tags);
+	}
+
+	ActorToPreviousBoundsMap.FindOrAdd(Actor) = PCGActorAndComponentMapping::GetActorBounds(Actor);
 }
 
 void FPCGActorAndComponentMapping::OnObjectSaved(UObject* InObject, FObjectPreSaveContext InObjectSaveContext)
 {
+	// Nothing to do if we track nothing
+	if (CulledTrackedKeysToComponentsMap.IsEmpty() && AlwaysTrackedKeysToComponentsMap.IsEmpty())
+	{
+		return;
+	}
+
 	// Only trigger a refresh a new user data and if it is a data table
 	// We only track data table because we probably will catch other changes with OnObjectPropertyChanged.
 	// To avoid to force the check multiple times (on change + on save)
@@ -1593,6 +1352,12 @@ void FPCGActorAndComponentMapping::OnObjectPropertyChanged(UObject* InObject, FP
 {
 	TRACE_CPUPROFILER_EVENT_SCOPE(FPCGActorAndComponentMapping::OnObjectPropertyChanged);
 
+	// Nothing to do if we track nothing
+	if (CulledTrackedKeysToComponentsMap.IsEmpty() && AlwaysTrackedKeysToComponentsMap.IsEmpty())
+	{
+		return;
+	}
+
 	const bool bValueNotInteractive = (InEvent.ChangeType != EPropertyChangeType::Interactive);
 	// Special exception for actor tags, as we can't track otherwise an actor "losing" a tag
 	const bool bActorTagChange = (InEvent.Property && InEvent.Property->GetFName() == GET_MEMBER_NAME_CHECKED(AActor, Tags));
@@ -1605,11 +1370,6 @@ void FPCGActorAndComponentMapping::OnObjectPropertyChanged(UObject* InObject, FP
 	const bool bIsStaticMeshCompilationResult = bEventIsEmpty && InObject && InObject->IsA<UStaticMesh>()
 		&& FStaticMeshCompilingManager::Get().GetNumRemainingMeshes() > 0;
 
-	if ((!bValueNotInteractive && !bActorTagChange) || bIsTextureCompilationResult || bIsStaticMeshCompilationResult)
-	{
-		return;
-	}
-
 	// First check if it is an actor
 	AActor* Actor = Cast<AActor>(InObject);
 
@@ -1622,90 +1382,77 @@ void FPCGActorAndComponentMapping::OnObjectPropertyChanged(UObject* InObject, FP
 		}
 	}
 
-	// If we don't find any actor, try to see if it is a dependency
-	if (!Actor)
-	{
-		if (TSet<UPCGComponent*>* Components = TrackedObjectsToComponentsMap.Find(InObject))
-		{
-			for (UPCGComponent* Component : *Components)
-			{
-				if (ClearCache(InObject, Component, /*bIntersect=*/ false, {}, InObject))
-				{
-					Component->DirtyGenerated();
-					Component->Refresh();
-				}
-			}
-		}
-
-		for (const TPair<TObjectKey<AActor>, TSet<TObjectPtr<UObject>>>& TrackedActor : TrackedActorsToDependenciesMap)
-		{
-			if (TrackedActor.Value.Contains(InObject))
-			{
-				if (AActor* ActorToChange = TrackedActor.Key.ResolveObjectPtr())
-				{
-					// Ignore property changes on delayed actors. All their properties are still being set.
-					if (DelayedAddedActors.Contains(ActorToChange))
-					{
-						return;
-					}
-
-					OnActorChanged(ActorToChange, /*bInHasMoved=*/ false, /*InOriginatingChangeObject=*/ InObject);
-					UpdateActorDependencies(ActorToChange);
-				}
-			}
-		}
+	FBox PreviousBounds{ EForceInit::ForceInit };
+	ActorToPreviousBoundsMap.RemoveAndCopyValue(Actor, PreviousBounds);
 
+	if ((!bValueNotInteractive && !bActorTagChange) || bIsTextureCompilationResult || bIsStaticMeshCompilationResult)
+	{
 		return;
 	}
 
-	if (PCGSubsystem && Actor->GetWorld() != PCGSubsystem->GetWorld())
+	if (PCGSubsystem && Actor && Actor->GetWorld() != PCGSubsystem->GetWorld())
 	{
 		return;
 	}
 
 #if WITH_EDITOR
-	if (Actor->bIsEditorPreviewActor)
+	if (Actor && Actor->bIsEditorPreviewActor)
 	{
 		return;
 	}
 #endif
 
 	// Ignore property changes on delayed actors. All their properties are still being set.
-	if (DelayedAddedActors.Contains(Actor))
+	if (Actor && DelayedAddedActors.Contains(Actor))
 	{
 		return;
 	}
 
-	// Check if we are not tracking it or is a tag change.
-	bool bShouldChange = true;
-	if (!IsActorTracked(Actor) || bActorTagChange)
-	{
-		bShouldChange = AddOrUpdateTrackedActor(Actor);
-	}
+	OnObjectChanged(Actor ? static_cast<UObject*>(Actor) : InObject, PreviousBounds, /*InOriginatingChangeObject=*/ InObject);
+}
 
-	if (bShouldChange)
+void FPCGActorAndComponentMapping::OnObjectChanged(UObject* InObject, const FBox& InPreviousBounds, const UObject* InOriginatingChangeObject, int32 LevelInstanceDepth, bool bNoRefreshOwner)
+{
+	// Nothing to do if we track nothing
+	if (CulledTrackedKeysToComponentsMap.IsEmpty() && AlwaysTrackedKeysToComponentsMap.IsEmpty())
 	{
-		OnActorChanged(Actor, /*bInHasMoved=*/ false, /*InOriginatingChangeObject=*/ InObject);
+		return;
 	}
-	else
+
+	TRACE_CPUPROFILER_EVENT_SCOPE(FPCGActorAndComponentMapping::OnObjectChanged);
+
+	check(InObject);
+	AActor* Actor = Cast<AActor>(InObject);
+
+	ensure(!PCGSubsystem || !Actor || Actor->GetWorld() == PCGSubsystem->GetWorld());
+
+	/** First gather all the components that are tracking this actor */
+	TSet<UPCGComponent*> CulledTrackedComponents;
+	TSet<UPCGComponent*> AlwaysTrackedComponents;
+	TArray<FPCGSelectionKey> MatchedKeys;
+
+	TSet<FName> RemovedTags;
+	if (Actor && !TempTrackedActorTags.IsEmpty())
 	{
-		// Otherwise we are already tracking the actor, so update its dependencies
-		UpdateActorDependencies(Actor);
+		RemovedTags = TempTrackedActorTags.Difference(TSet<FName>(Actor->Tags));
 	}
-}
-
-void FPCGActorAndComponentMapping::OnActorChanged(AActor* InActor, bool bInHasMoved, const UObject* InOriginatingChangeObject, int32 LevelInstanceDepth, bool bNoRefreshOwner)
-{
-	TRACE_CPUPROFILER_EVENT_SCOPE(FPCGActorAndComponentMapping::OnActorChanged);
 
-	check(InActor);
-	ensure(!PCGSubsystem || InActor->GetWorld() == PCGSubsystem->GetWorld());
+	auto Gather = [InObject, &RemovedTags, &MatchedKeys](TMap<FPCGSelectionKey, TSet<UPCGComponent*>> InMap, TSet<UPCGComponent*>& OutSet)
+	{
+		for (auto& It : InMap)
+		{
+			if (It.Key.IsMatching(InObject, RemovedTags, It.Value, OutSet))
+			{
+				MatchedKeys.Add(It.Key);
+			}
+		}
+	};
 
-	TSet<UPCGComponent*>* CulledTrackedComponents = CulledTrackedActorsToComponentsMap.Find(InActor);
-	TSet<UPCGComponent*>* AlwaysTrackedComponents = AlwaysTrackedActorsToComponentsMap.Find(InActor);
+	Gather(CulledTrackedKeysToComponentsMap, CulledTrackedComponents);
+	Gather(AlwaysTrackedKeysToComponentsMap, AlwaysTrackedComponents);
 
 	// If this actor is not tracked, just early out
-	if (!CulledTrackedComponents && !AlwaysTrackedComponents)
+	if (CulledTrackedComponents.IsEmpty() && AlwaysTrackedComponents.IsEmpty())
 	{
 		return;
 	}
@@ -1713,24 +1460,46 @@ void FPCGActorAndComponentMapping::OnActorChanged(AActor* InActor, bool bInHasMo
 	TSet<UPCGComponent*> DirtyComponents;
 
 	EPCGComponentDirtyFlag DirtyFlag = EPCGComponentDirtyFlag::Actor;
-	if (InActor->IsA<ALandscapeProxy>())
+	if (InObject->IsA<ALandscapeProxy>())
 	{
 		DirtyFlag = DirtyFlag | EPCGComponentDirtyFlag::Landscape;
 	}
 
-	// Check if we have a change of tag too
-	TSet<FName> RemovedTags = TempTrackedActorTags.Difference(TSet<FName>(InActor->Tags));
-
 	// We need the actor bounds to know if we are intersecting, for the culled settings to be dirtied.
-	const FBox ActorBounds = PCGActorAndComponentMapping::GetActorBounds(InActor);
-	// We also need the OldActorBounds if it has moved, and is different than the new bounds.
-	FBox* OldActorBoundsPtr = bInHasMoved ? TrackedActorToPositionMap.Find(InActor) : nullptr;
-	if (OldActorBoundsPtr && OldActorBoundsPtr->Equals(ActorBounds))
+	const FBox ActorBounds = PCGActorAndComponentMapping::GetActorBounds(Actor);
+
 	{
-		OldActorBoundsPtr = nullptr;
+		TRACE_CPUPROFILER_EVENT_SCOPE(FPCGActorAndComponentMapping::OnActorChanged::AlwaysTrackedUpdate);
+
+		for (UPCGComponent* PCGComponent : AlwaysTrackedComponents)
+		{
+			if (!PCGComponent || PCGComponent == InOriginatingChangeObject || (bNoRefreshOwner && PCGComponent->GetOwner() == InObject))
+			{
+				continue;
+			}
+
+			// Don't mark "Owner changed" if the change originate from a PCG Component. It will be delegated to the ClearCacheForActor.
+			// It is necessary to avoid infine loops when there are multiple PCG components on one actor, and one component was generated.
+			const bool bOwnerChanged = (PCGComponent->GetOwner() == InObject) && (!InOriginatingChangeObject || !InOriginatingChangeObject->IsA<UPCGComponent>());
+			bool bShouldDirty = bOwnerChanged || (PCGComponent->ShouldTrackLandscape() && InObject->IsA<ALandscapeProxy>());
+
+			// Since when we clear the cache for a settings, we clear it all, it's only necessary to do it once on the original component. And we will dirty all the local components (with dirty dispatch)
+			if (!DirtyComponents.Contains(PCGComponent))
+			{
+				const FBox ComponentBounds = PCGComponent->GetGridBounds();
+				const bool bIntersect = bOwnerChanged || ActorBounds.Intersect(ComponentBounds);// || (OldActorBoundsPtr && OldActorBoundsPtr->Intersect(ComponentBounds));
+
+				bShouldDirty |= ClearCacheForKeys(MatchedKeys, PCGComponent, /*bIntersect=*/true, InOriginatingChangeObject);
+				if (bShouldDirty)
+				{
+					PCGComponent->DirtyGenerated(DirtyFlag, /*bDispatchToLocalComponents=*/true);
+					DirtyComponents.Add(PCGComponent);
+				}
+			}
+		}
 	}
 
-	if (CulledTrackedComponents)
+	if (!CulledTrackedComponents.IsEmpty())
 	{
 		// Keep a pointer to the current bounds we are looking at for the partition case. Set it to Actorbounds first, then will be set to old bounds if necessary.
 		const FBox* CurrentActorBoundsPtr = &ActorBounds;
@@ -1738,20 +1507,22 @@ void FPCGActorAndComponentMapping::OnActorChanged(AActor* InActor, bool bInHasMo
 		// Then do an octree find to get all components that intersect with this actor.
 		// If the actor has moved, we also need to find components that intersected with it before
 		// We first do it for non-partitioned, then we do it for partitioned
-		auto UpdateNonPartitioned = [this, &DirtyComponents, InActor, CulledTrackedComponents, &RemovedTags, DirtyFlag, InOriginatingChangeObject, bNoRefreshOwner](const FPCGComponentRef& ComponentRef) -> void
+		auto UpdateNonPartitioned = [this, &DirtyComponents, InObject, CulledTrackedComponents, &RemovedTags, DirtyFlag, InOriginatingChangeObject, bNoRefreshOwner, &MatchedKeys](const FPCGComponentRef& ComponentRef) -> void
 		{
 			TRACE_CPUPROFILER_EVENT_SCOPE(FPCGActorAndComponentMapping::OnActorChanged::UpdateNonPartitioned);
 
 			// Don't dirty if the component was already dirtied, not tracked, the origin of the change or its owner is the changed actor and we should not refresh.
-			if (DirtyComponents.Contains(ComponentRef.Component) || 
-				!CulledTrackedComponents->Contains(ComponentRef.Component) ||
+			if (DirtyComponents.Contains(ComponentRef.Component) ||
+				!CulledTrackedComponents.Contains(ComponentRef.Component) ||
 				(InOriginatingChangeObject == ComponentRef.Component) ||
-				(bNoRefreshOwner && ComponentRef.Component->GetOwner() == InActor))
+				(bNoRefreshOwner && ComponentRef.Component->GetOwner() == InObject))
 			{
 				return;
 			}
 
-			if (ClearCache(InActor, ComponentRef.Component, /*bIntersect=*/true, RemovedTags, InOriginatingChangeObject))
+			bool bShouldDirty = ComponentRef.Component->ShouldTrackLandscape() && InObject->IsA<ALandscapeProxy>();
+			bShouldDirty |= ClearCacheForKeys(MatchedKeys, ComponentRef.Component, /*bIntersect=*/true, InOriginatingChangeObject);
+			if (bShouldDirty)
 			{
 				ComponentRef.Component->DirtyGenerated(DirtyFlag);
 				DirtyComponents.Add(ComponentRef.Component);
@@ -1762,16 +1533,16 @@ void FPCGActorAndComponentMapping::OnActorChanged(AActor* InActor, bool bInHasMo
 
 		// For partitioned, we first need check if the original component intersect with the bounds, then forward the dirty call only to locals that intersect with the bounds.
 		// Note: CurrentActorBoundsPtr is passed by reference because it will be modified between lambda calls (cf comment above).
-		auto UpdatePartitioned = [this, &DirtyComponents, InActor, CulledTrackedComponents, &CurrentActorBoundsPtr, &RemovedTags, DirtyFlag, InOriginatingChangeObject, bNoRefreshOwner](const FPCGComponentRef& ComponentRef)  -> void
+		auto UpdatePartitioned = [this, &DirtyComponents, InObject, CulledTrackedComponents, &CurrentActorBoundsPtr, &RemovedTags, DirtyFlag, InOriginatingChangeObject, bNoRefreshOwner, &MatchedKeys](const FPCGComponentRef& ComponentRef)  -> void
 		{
 			TRACE_CPUPROFILER_EVENT_SCOPE(FPCGActorAndComponentMapping::OnActorChanged::UpdatePartitioned);
 
 			// Don't dirty if the component is not tracked, the origin of the change or its owner is the changed actor and we should not refresh.
 			// We can "re-dirty" it because changes can impact different local components, from the same
 			// original component.
-			if (!CulledTrackedComponents->Contains(ComponentRef.Component) ||
+			if (!CulledTrackedComponents.Contains(ComponentRef.Component) ||
 				(InOriginatingChangeObject == ComponentRef.Component) ||
-				(bNoRefreshOwner && ComponentRef.Component->GetOwner() == InActor))
+				(bNoRefreshOwner && ComponentRef.Component->GetOwner() == InObject))
 			{
 				return;
 			}
@@ -1785,12 +1556,14 @@ void FPCGActorAndComponentMapping::OnActorChanged(AActor* InActor, bool bInHasMo
 			}
 
 			const FBox Overlap = CurrentActorBoundsPtr->Overlap(ComponentBounds);
+			bool bShouldDirty = ComponentRef.Component->ShouldTrackLandscape() && InObject->IsA<ALandscapeProxy>();
+			bShouldDirty |= ClearCacheForKeys(MatchedKeys, ComponentRef.Component, /*bIntersect=*/true, InOriginatingChangeObject);
 			bool bWasDirtied = false;
 
 			// Since when we clear the cache for a settings, we clear it all, it's only necessary to do it once on the original component, then only dirty the local that intersects.
-			if (ClearCache(InActor, ComponentRef.Component, /*bIntersect=*/true, RemovedTags, InOriginatingChangeObject))
+			if (bShouldDirty)
 			{
-				ForAllIntersectingPartitionActors(Overlap, [InActor, Component = ComponentRef.Component, &bWasDirtied, DirtyFlag, InOriginatingChangeObject](APCGPartitionActor* InPartitionActor) -> void
+				ForAllIntersectingPartitionActors(Overlap, [InObject, Component = ComponentRef.Component, &bWasDirtied, DirtyFlag, InOriginatingChangeObject](APCGPartitionActor* InPartitionActor) -> void
 				{
 					if (UPCGComponent* LocalComponent = InPartitionActor->GetLocalComponent(Component))
 					{
@@ -1811,53 +1584,22 @@ void FPCGActorAndComponentMapping::OnActorChanged(AActor* InActor, bool bInHasMo
 		PartitionedOctree.FindElementsWithBoundsTest(ActorBounds, UpdatePartitioned);
 
 		// If it has moved, redo it with the old bounds.
-		if (OldActorBoundsPtr)
+		if (InPreviousBounds.IsValid)
 		{
 			TRACE_CPUPROFILER_EVENT_SCOPE(FPCGActorAndComponentMapping::OnActorChanged::SecondUpdateHasMoved);
 
 			// Set the actor bounds with the old one, to have the right Overlap in the Partition case.
-			CurrentActorBoundsPtr = OldActorBoundsPtr;
-			NonPartitionedOctree.FindElementsWithBoundsTest(*OldActorBoundsPtr, UpdateNonPartitioned);
-			PartitionedOctree.FindElementsWithBoundsTest(*OldActorBoundsPtr, UpdatePartitioned);
+			CurrentActorBoundsPtr = &InPreviousBounds;
+			NonPartitionedOctree.FindElementsWithBoundsTest(InPreviousBounds, UpdateNonPartitioned);
+			PartitionedOctree.FindElementsWithBoundsTest(InPreviousBounds, UpdatePartitioned);
 		}
 	}
 
 	// Finally, dirty all components that always track this actor that are not yet notified.
-	if (AlwaysTrackedComponents)
-	{
-		TRACE_CPUPROFILER_EVENT_SCOPE(FPCGActorAndComponentMapping::OnActorChanged::AlwaysTrackedUpdate);
-
-		for (UPCGComponent* PCGComponent : *AlwaysTrackedComponents)
-		{
-			if (!PCGComponent || PCGComponent == InOriginatingChangeObject || (bNoRefreshOwner && PCGComponent->GetOwner() == InActor))
-			{
-				continue;
-			}
-
-			// Don't mark "Owner changed" if the change originate from a PCG Component. It will be delegated to the ClearCacheForActor.
-			// It is necessary to avoid infine loops when there are multiple PCG components on one actor, and one component was generated.
-			const bool bOwnerChanged = (PCGComponent->GetOwner() == InActor) && (!InOriginatingChangeObject || !InOriginatingChangeObject->IsA<UPCGComponent>());
-			bool bWasDirtied = false;
-
-			// Since when we clear the cache for a settings, we clear it all, it's only necessary to do it once on the original component. And we will dirty all the local components (with dirty dispatch)
-			if (!DirtyComponents.Contains(PCGComponent))
-			{
-				const FBox ComponentBounds = PCGComponent->GetGridBounds();
-				const bool bIntersect = ActorBounds.Intersect(ComponentBounds) || (OldActorBoundsPtr && OldActorBoundsPtr->Intersect(ComponentBounds));
-				bWasDirtied = ClearCache(InActor, PCGComponent, /*bIntersect=*/bIntersect, RemovedTags, InOriginatingChangeObject);
-			}
-
-			if (bWasDirtied || bOwnerChanged)
-			{
-				PCGComponent->DirtyGenerated(DirtyFlag, /*bDispatchToLocalComponents=*/true);
-				DirtyComponents.Add(PCGComponent);
-			}
-		}
-	}
 
 	// If it is a landscape and we should discard the refresh, early out.
 	bool bIsInEditingMode = false;
-	ALandscapeProxy* Landscape = Cast<ALandscapeProxy>(InActor);
+	ALandscapeProxy* Landscape = Cast<ALandscapeProxy>(InObject);
 	if (PCGActorAndComponentMapping::ShouldDiscardLandscapeRefresh(Landscape, bIsInEditingMode, bIsCurrentlyExitingLandscapeEditMode))
 	{
 		// If we are in editing, keep track of all the dirtied landscape to refresh them when we exit.
@@ -1879,7 +1621,7 @@ void FPCGActorAndComponentMapping::OnActorChanged(AActor* InActor, bool bInHasMo
 			continue;
 		}
 
-		const bool bOwnerHasChanged = Component->GetOwner() == InActor;
+		const bool bOwnerHasChanged = Component->GetOwner() == InObject;
 
 		if ((!bNoRefreshOwner || !bOwnerHasChanged) && (!Component->bOnlyTrackItself || bOwnerHasChanged))
 		{
@@ -1887,9 +1629,9 @@ void FPCGActorAndComponentMapping::OnActorChanged(AActor* InActor, bool bInHasMo
 			// level hierarchy, otherwise we will end up generating in the Level Instance level, which is wrong.
 			// Note that in some instances (e.g. when something happens at the Level Instance level) we need to make sure that the
 			// PCG components higher-up are properly updated, hence the level instance depth
-			if (LevelInstanceSubsystem)
+			if (LevelInstanceSubsystem && Actor)
 			{
-				const ILevelInstanceInterface* ActorLevelInstance = LevelInstanceSubsystem->GetParentLevelInstance(InActor);
+				const ILevelInstanceInterface* ActorLevelInstance = LevelInstanceSubsystem->GetParentLevelInstance(Actor);
 				int LocalDepth = LevelInstanceDepth;
 				while (ActorLevelInstance && LocalDepth-- > 0)
 				{
@@ -1907,7 +1649,7 @@ void FPCGActorAndComponentMapping::OnActorChanged(AActor* InActor, bool bInHasMo
 					continue;
 				}
 			}
-			
+
 			Component->Refresh();
 		}
 	}
@@ -1938,35 +1680,7 @@ void FPCGActorAndComponentMapping::ApplyLandscapeChanges(ALandscapeProxy* InLand
 		return;
 	}
 	
-	// We don't know if the landscape moved, only that it has changed. Since `bHasMoved` is doing a bit more, always assume that the landscape has moved.
-	OnActorChanged(InLandscape, /*bHasMoved=*/true);
-
-	// Also update its dependencies
-	UpdateActorDependencies(InLandscape);
-}
-
-void FPCGActorAndComponentMapping::UpdateActorDependencies(AActor* InActor)
-{
-	if (!InActor)
-	{
-		return;
-	}
-
-	// Don't track what the PCG Component is already tracking (itself and graphs)
-	static const TArray<UClass*> ExcludedClasses =
-	{
-		UPCGComponent::StaticClass(),
-		UPCGGraphInterface::StaticClass(),
-		UMaterialInterface::StaticClass()
-	};
-
-
-	if (!PCGActorAndComponentMapping::CVarDisableObjectDependenciesTracking.GetValueOnAnyThread())
-	{
-		TSet<TObjectPtr<UObject>>& DependenciesSet = TrackedActorsToDependenciesMap.FindOrAdd(InActor);
-		DependenciesSet.Empty();
-		PCGHelpers::GatherDependencies(InActor, DependenciesSet, 1, ExcludedClasses);
-	}
+	OnObjectChanged(InLandscape, FBox(EForceInit::ForceInit), InLandscape);
 }
 
 void FPCGActorAndComponentMapping::OnPCGGraphGeneratedOrCleaned(UPCGComponent* InComponent)
@@ -1976,35 +1690,41 @@ void FPCGActorAndComponentMapping::OnPCGGraphGeneratedOrCleaned(UPCGComponent* I
 		return;
 	}
 
-	OnActorChanged(InComponent->GetOwner(), /*bInHasMoved=*/false, InComponent);
-}
-
-bool FPCGActorAndComponentMapping::IsActorTracked(const AActor* InActor) const
-{
-	return InActor && (CulledTrackedActorsToComponentsMap.Contains(InActor) || AlwaysTrackedActorsToComponentsMap.Contains(InActor));
+	OnObjectChanged(InComponent->GetOwner(), FBox(EForceInit::ForceInit), InComponent);
 }
 
-bool FPCGActorAndComponentMapping::ClearCache(const UObject* InObject, const UPCGComponent* InComponent, const bool bIntersect, const TSet<FName>& InRemovedTags, const UObject* InOriginatingChange) const
+bool FPCGActorAndComponentMapping::ClearCacheForKeys(const TArray<FPCGSelectionKey>& InKeys, const UPCGComponent* InComponent, const bool bIntersect, const UObject* InOriginatingChange) const
 {
-	check(InObject && InComponent);
+	check(InComponent);
 
 	if (!PCGSubsystem)
 	{
 		return false;
 	}
 
-	// Special case for the landscape. No settings associated, but we should dirty if the component is tracking the landscape.
-	bool bShouldDirty = InComponent->ShouldTrackLandscape() && InObject->IsA<ALandscapeProxy>();
-
-	for (const UPCGSettings* Settings : InComponent->GatherSettingsTracking(InObject, bIntersect, InRemovedTags, InOriginatingChange))
+	bool bShouldDirty = InComponent->ShouldTrackLandscape() && InOriginatingChange && InOriginatingChange->IsA<ALandscapeProxy>();
+	auto ClearCache = [this, InOriginatingChange, bIntersect, &bShouldDirty](const FPCGSelectionKey& Key, const FPCGSettingsAndCulling& SettingsAndCulling)
 	{
-		if (!Settings) 
+		if (!SettingsAndCulling.Key.IsValid() || (SettingsAndCulling.Value && !bIntersect))
 		{
-			continue;
+			return;
+		}
+
+		// Extra care if the change originates from a PCGComponent. Only dirty if we are tracking a PCG component.
+		if (InOriginatingChange && InOriginatingChange->IsA<UPCGComponent>()
+			&& (!Key.OptionalExtraDependency || !Key.OptionalExtraDependency->IsChildOf(UPCGComponent::StaticClass())))
+		{
+			return;
 		}
 
 		bShouldDirty = true;
+		const UPCGSettings* Settings = SettingsAndCulling.Key.Get();
 		PCGSubsystem->CleanFromCache(Settings->GetElement().Get(), Settings);
+	};
+
+	for (const FPCGSelectionKey& Key : InKeys)
+	{
+		InComponent->ApplyToEachSettings(Key, ClearCache);
 	}
 
 	return bShouldDirty;
diff --git a/Engine/Plugins/Experimental/PCG/Source/PCG/Private/PCGComponent.cpp b/Engine/Plugins/Experimental/PCG/Source/PCG/Private/PCGComponent.cpp
--- a/Engine/Plugins/Experimental/PCG/Source/PCG/Private/PCGComponent.cpp
+++ b/Engine/Plugins/Experimental/PCG/Source/PCG/Private/PCGComponent.cpp
@@ -3,6 +3,7 @@
 #include "PCGComponent.h"
 
 #include "PCGContext.h"
+#include "PCGCustomVersion.h"
 #include "PCGEngineSettings.h"
 #include "PCGGraph.h"
 #include "PCGInputOutputSettings.h"
@@ -369,6 +370,12 @@ FPCGTaskId UPCGComponent::CreateGenerateTask(bool bForce, const TArray<FPCGTaskI
 		return InvalidPCGTaskId;
 	}
 
+#if WITH_EDITOR
+	// No need for lock since it is not executed in parallel.
+	CurrentExecutionDynamicTracking.Empty();
+	CurrentExecutionDynamicTrackingSettings.Empty();
+#endif // WITH_EDITOR
+
 	return GetSubsystem()->ScheduleGraph(this, *AllDependencies);
 }
 
@@ -427,6 +434,10 @@ void UPCGComponent::PostProcessGraph(const FBox& InNewBounds, bool bInGenerated,
 
 			CallPostGenerateFunctions(Context);
 		}
+
+#if WITH_EDITOR
+		UpdateDynamicTracking();
+#endif // WITH_EDITOR
 	}
 
 	// Trigger notification - will be used by other tracking mechanisms
@@ -1215,8 +1226,24 @@ void UPCGComponent::Serialize(FArchive& Ar)
 	}
 #endif // WITH_EDITOR
 
+	Ar.UsingCustomVersion(FPCGCustomVersion::GUID);
+
 	Super::Serialize(Ar);
 
+#if WITH_EDITOR
+	int32 DataVersion = FPCGCustomVersion::LatestVersion;
+	if (Ar.IsLoading())
+	{
+		DataVersion = Ar.CustomVer(FPCGCustomVersion::GUID);
+	}
+
+	if (DataVersion >= FPCGCustomVersion::DynamicTrackingKeysSerializedInComponent)
+	{
+		Ar << DynamicallyTrackedKeysToSettings;
+	}
+#endif // WITH_EDITOR
+
+
 #if WITH_EDITOR
 	if (Ar.IsSaving() && CurrentEditingMode == EPCGEditorDirtyMode::Preview)
 	{
@@ -1691,41 +1718,37 @@ bool UPCGComponent::UpdateTrackingCache(TArray<FPCGSelectionKey>* OptionalChange
 		return false;
 	}
 
-	// Store in a temporary map to detect key changes.
-	TMap<FPCGSelectionKey, bool> NewTrackedKeysToCulling;
-
 	int32 FoundKeys = 0;
 
+	// Store in a temporary map to detect key changes.
+	FPCGSelectionKeyToSettingsMap NewTrackedKeysToSettings;
+
 	if (UPCGGraph* PCGGraph = GetGraph())
 	{
-		CachedTrackedKeysToSettings = PCGGraph->GetTrackedActorKeysToSettings();
+		NewTrackedKeysToSettings = PCGGraph->GetTrackedActorKeysToSettings();
 
 		// Also add a key for the landscape, with settings null and always culled, if we should track the landscape
 		if (ShouldTrackLandscape())
 		{
 			FPCGSelectionKey LandscapeKey = FPCGSelectionKey(ALandscapeProxy::StaticClass());
-			CachedTrackedKeysToSettings.FindOrAdd(LandscapeKey).Emplace(/*Settings*/nullptr, /*bIsCulled*/true);
+			NewTrackedKeysToSettings.FindOrAdd(LandscapeKey).Emplace(/*Settings*/nullptr, /*bIsCulled*/true);
 		}
 
 		// A tag should be culled, if only all the settings that track this tag should cull.
 		// Note that is only impact the fact that we track (or not) this tag.
 		// If a setting is marked as "should cull", it will only be dirtied (at least by default), if the actor with the
 		// given tag intersect with the component.
-		for (const TPair<FPCGSelectionKey, TArray<FPCGSettingsAndCulling>>& It : CachedTrackedKeysToSettings)
+		for (const TPair<FPCGSelectionKey, TArray<FPCGSettingsAndCulling>>& It : NewTrackedKeysToSettings)
 		{
 			const FPCGSelectionKey& Key = It.Key;
 
 			// Should cull only if all the settings requires a cull.
-			const bool bShouldCull = Algo::AllOf(It.Value, [](const FPCGSettingsAndCulling& SettingsAndCullingPair) { return SettingsAndCullingPair.Value; });
+			const bool bShouldCull = PCGSettings::IsKeyCulled(It.Value);
+			const TArray<FPCGSettingsAndCulling>* OldSettingsAndCulling = StaticallyTrackedKeysToSettings.Find(It.Key);
+			const bool OldCulling = OldSettingsAndCulling && PCGSettings::IsKeyCulled(*OldSettingsAndCulling);
+			const bool bNewKeyOrCullChanged = !OldSettingsAndCulling || (OldCulling != bShouldCull);
 
-			NewTrackedKeysToCulling.Emplace(Key, bShouldCull);
-
-			// Look for the key in the previous cached keys.
-			const bool* OldCulling = CachedTrackedKeysToCulling.Find(Key);
-			// It is a new key if we didn't find it. We also mark it changed if the key existed, but the culling changed.
-			const bool bNewKeyOrCullChanged = !OldCulling || (*OldCulling != bShouldCull);
-			// Then remove the key.
-			CachedTrackedKeysToCulling.Remove(Key);
+			StaticallyTrackedKeysToSettings.Remove(Key);
 
 			if (!bNewKeyOrCullChanged)
 			{
@@ -1740,18 +1763,18 @@ bool UPCGComponent::UpdateTrackingCache(TArray<FPCGSelectionKey>* OptionalChange
 		// At the end, we also have keys that were tracked but no more, so add them at the list of tracked keys
 		if (OptionalChangedKeys)
 		{
-			OptionalChangedKeys->Reserve(OptionalChangedKeys->Num() + CachedTrackedKeysToCulling.Num());
+			OptionalChangedKeys->Reserve(OptionalChangedKeys->Num() + StaticallyTrackedKeysToSettings.Num());
 
-			for (const TPair<FPCGSelectionKey, bool>& It : CachedTrackedKeysToCulling)
+			for (const TPair<FPCGSelectionKey, TArray<FPCGSettingsAndCulling>>& It : StaticallyTrackedKeysToSettings)
 			{
 				OptionalChangedKeys->Add(It.Key);
 			}
 		}
 	}
 
-	bool bHasChanged = CachedTrackedKeysToSettings.Num() != FoundKeys;
+	bool bHasChanged = NewTrackedKeysToSettings.Num() != FoundKeys;
 
-	CachedTrackedKeysToCulling = MoveTemp(NewTrackedKeysToCulling);
+	StaticallyTrackedKeysToSettings = MoveTemp(NewTrackedKeysToSettings);
 
 	return bHasChanged;
 }
@@ -2048,16 +2071,15 @@ bool UPCGComponent::ShouldGenerateBPPCGAddedToWorld() const
 	}
 }
 
-bool UPCGComponent::IsObjectTracked(const TSoftObjectPtr<UObject>& InObjectPtr, bool& bOutIsCulled) const
+bool UPCGComponent::IsObjectTracked(const UObject* InObject, bool& bOutIsCulled) const
 {
-	check(!InObjectPtr.IsNull());
+	check(InObject);
 
 	if (!GetOwner())
 	{
 		return false;
 	}
 
-	const UObject* InObject = InObjectPtr.Get();
 
 	// We should always track the owner of the component, without culling
 	if (GetOwner() == InObject)
@@ -2073,26 +2095,23 @@ bool UPCGComponent::IsObjectTracked(const TSoftObjectPtr<UObject>& InObjectPtr,
 		return true;
 	}
 
-	bool bFound = false;
-
-	for (const TPair<FPCGSelectionKey, bool>& It : CachedTrackedKeysToCulling)
+	auto CheckMap = [this, &InObject, &bOutIsCulled](const FPCGSelectionKeyToSettingsMap& InMap) -> bool
 	{
-		if (It.Key.IsMatching(InObjectPtr, this))
+		for (const auto& It : InMap)
 		{
-			bOutIsCulled = It.Value;
-			bFound = true;
+			bool bFound = It.Key.IsMatching(InObject, this);
 
-			// Check for other tags that might not be culled
-			if (bOutIsCulled)
+			if (bFound)
 			{
-				continue;
+				bOutIsCulled = PCGSettings::IsKeyCulled(It.Value);
+				return true;
 			}
-
-			return true;
 		}
-	}
 
-	return bFound;
+		return false;
+	};
+
+	return CheckMap(StaticallyTrackedKeysToSettings) || CheckMap(DynamicallyTrackedKeysToSettings);
 }
 
 void UPCGComponent::OnRefresh(bool bForceRefresh, bool bForceCleanup)
@@ -2683,48 +2702,66 @@ UPCGSubsystem* UPCGComponent::GetSubsystem() const
 }
 
 #if WITH_EDITOR
-TArray<const UPCGSettings*> UPCGComponent::GatherSettingsTracking(const UObject* InObject, bool bIntersect, const TSet<FName>& InRemovedTags, const UObject* InOriginatingChangeObject) const
+void UPCGComponent::ApplyToEachSettings(const FPCGSelectionKey& InKey, const TFunctionRef<void(const FPCGSelectionKey&, const FPCGSettingsAndCulling&)> InCallback) const
 {
-	TRACE_CPUPROFILER_EVENT_SCOPE(UPCGComponent::GatherSettingsTracking);
+	TRACE_CPUPROFILER_EVENT_SCOPE(UPCGComponent::ApplyToEachSettings);
 
-	if (!InObject)
+	auto FindAndApplyInMap = [&InKey, &InCallback](const FPCGSelectionKeyToSettingsMap& InMap)
 	{
-		return {};
-	}
+		if (const TArray<FPCGSettingsAndCulling>* StaticallyTrackedSettings = InMap.Find(InKey))
+		{
+			for (const FPCGSettingsAndCulling& SettingsAndCulling : *StaticallyTrackedSettings)
+			{
+				InCallback(InKey, SettingsAndCulling);
+			}
+		}
+	};
 
-	TArray<const UPCGSettings*> TrackedSettings;
+	FindAndApplyInMap(StaticallyTrackedKeysToSettings);
+	FindAndApplyInMap(DynamicallyTrackedKeysToSettings);
+}
 
-	for (const auto& It : CachedTrackedKeysToSettings)
+TArray<FPCGSelectionKey> UPCGComponent::GatherTrackingKeys() const
+{
+	TArray<FPCGSelectionKey> Keys;
+	Keys.Reserve(StaticallyTrackedKeysToSettings.Num() + DynamicallyTrackedKeysToSettings.Num());
+	for (const auto& It : StaticallyTrackedKeysToSettings)
 	{
-		const FPCGSelectionKey& Key = It.Key;
+		Keys.Add(It.Key);
+	}
 
-		const bool bRemovedTagIsTracked = (Key.Selection == EPCGActorSelection::ByTag) && InRemovedTags.Contains(Key.Tag);
+	for (const auto& It : DynamicallyTrackedKeysToSettings)
+	{
+		Keys.Add(It.Key);
+	}
 
-		if (It.Key.IsMatching(InObject, this) || bRemovedTagIsTracked)
-		{
-			// Extra care if the change originates from a PCGComponent. Only dirty if we are tracking a PCG component.
-			if (InOriginatingChangeObject && InOriginatingChangeObject->IsA<UPCGComponent>() 
-				&& (!It.Key.OptionalExtraDependency || !It.Key.OptionalExtraDependency->IsChildOf(UPCGComponent::StaticClass())))
-			{
-				continue;
-			}
+	return Keys;
+}
 
-			for (const FPCGSettingsAndCulling& SettingsAndCulling : It.Value)
-			{
-				if (SettingsAndCulling.Value && !bIntersect)
-				{
-					continue;
-				}
+bool UPCGComponent::IsKeyTrackedAndCulled(const FPCGSelectionKey& Key, bool& bOutIsCulled) const
+{
+	bool bIsTracked = false;
 
-				if (const UPCGSettings* Settings = SettingsAndCulling.Key.Get())
-				{
-					TrackedSettings.Add(Settings);
-				}
-			}
-		}
+	bool bStaticallyCulled = true;
+	bool bDynamicallyCulled = true;
+
+	if (auto* It = StaticallyTrackedKeysToSettings.Find(Key))
+	{
+		bIsTracked = true;
+		bStaticallyCulled = PCGSettings::IsKeyCulled(*It);
 	}
 
-	return TrackedSettings;
+	if (auto* It = DynamicallyTrackedKeysToSettings.Find(Key))
+	{
+		bIsTracked = true;
+		bDynamicallyCulled = PCGSettings::IsKeyCulled(*It);
+	}
+
+	// If it is tracked statically and dynamically, we will cull only and only if both are culling.
+	// Otherwise, it means that at least one key requires to always track, so bOutIsCulled needs to be False.
+	bOutIsCulled = bIsTracked && bStaticallyCulled && bDynamicallyCulled;
+
+	return bIsTracked;
 }
 
 bool UPCGComponent::ShouldTrackLandscape() const
@@ -2743,6 +2780,127 @@ bool UPCGComponent::ShouldTrackLandscape() const
 	return bUseLandscapePin || bHasLandscapeHasInput;
 }
 
+void UPCGComponent::RegisterDynamicTracking(const UPCGSettings* InSettings, const TArrayView<TPair<FPCGSelectionKey, bool>>& InDynamicKeysAndCulling)
+{
+	if (!InSettings)
+	{
+		return;
+	}
+
+	FScopeLock Lock(&CurrentExecutionDynamicTrackingLock);
+	CurrentExecutionDynamicTrackingSettings.Add(InSettings);
+
+	for (const TPair<FPCGSelectionKey, bool>& It : InDynamicKeysAndCulling)
+	{
+		CurrentExecutionDynamicTracking.FindOrAdd(It.Key).Emplace(InSettings, It.Value);
+	}
+}
+
+void UPCGComponent::UpdateDynamicTracking()
+{
+	TRACE_CPUPROFILER_EVENT_SCOPE(UPCGComponent::UpdateDynamicTracking);
+
+	UPCGSubsystem* Subsystem = GetSubsystem();
+	if (!Subsystem)
+	{
+		return;
+	}
+
+	// If the component is local, we defer the tracking to the original component.
+	// So move everything to the original (while making sure we are not duplicating keys/settings).
+	// Since it can happen in parallel, we need to lock.
+	if (IsLocalComponent())
+	{
+		TRACE_CPUPROFILER_EVENT_SCOPE(UPCGComponent::UpdateDynamicTracking::LocalComponent);
+
+		UPCGComponent* OriginalComponent = GetOriginalComponent();
+		FScopeLock Lock(&OriginalComponent->CurrentExecutionDynamicTrackingLock);
+		for (auto& It : CurrentExecutionDynamicTracking)
+		{
+			TArray<FPCGSettingsAndCulling>& OriginalSettingsAndCulling = OriginalComponent->CurrentExecutionDynamicTracking.FindOrAdd(It.Key);
+			for (FPCGSettingsAndCulling& SettingsAndCulling : It.Value)
+			{
+				OriginalSettingsAndCulling.AddUnique(std::move(SettingsAndCulling));
+			}
+		}
+
+		OriginalComponent->CurrentExecutionDynamicTrackingSettings.Append(CurrentExecutionDynamicTrackingSettings);
+
+		CurrentExecutionDynamicTracking.Empty();
+		CurrentExecutionDynamicTrackingSettings.Empty();
+		return;
+	}
+
+	TArray<FPCGSelectionKey> ChangedKeys;
+
+	// Locking to make sure we never hit this multiple times.
+	{
+		FScopeLock Lock(&CurrentExecutionDynamicTrackingLock);
+
+		// Go over all dynamic keys gathered during this execution.
+		// If they are not already tracked, we need to register this key.
+		// Otherwise, we need to gather all the settings that tracked this key that were not executed (because of caching).
+		for (auto& It : CurrentExecutionDynamicTracking)
+		{
+			if (TArray<FPCGSettingsAndCulling>* AllSettingsAndCulling = DynamicallyTrackedKeysToSettings.Find(It.Key))
+			{
+				for (FPCGSettingsAndCulling& SettingsAndCulling : *AllSettingsAndCulling)
+				{
+					if (SettingsAndCulling.Key.Get() && !CurrentExecutionDynamicTrackingSettings.Contains(SettingsAndCulling.Key.Get()))
+					{
+						It.Value.AddUnique(std::move(SettingsAndCulling));
+					}
+				}
+			}
+			else
+			{
+				ChangedKeys.Add(It.Key);
+			}
+		}
+
+		// Go over all already registered dynamic keys
+		// If they are not in the current execution gathered keys, we check if they are associated with settings that 
+		// were executed. If so, we re-add them to the current execution gathered keys.
+		// If not, it means that the key is no longer tracked and should be unregistered.
+		for (auto& It : DynamicallyTrackedKeysToSettings)
+		{
+			if (!CurrentExecutionDynamicTracking.Contains(It.Key))
+			{
+				TArray<FPCGSettingsAndCulling>* AllSettingsAndCulling = nullptr;
+
+				for (FPCGSettingsAndCulling& SettingsAndCulling : It.Value)
+				{
+					if (SettingsAndCulling.Key.Get() && !CurrentExecutionDynamicTrackingSettings.Contains(SettingsAndCulling.Key.Get()))
+					{
+						if (!AllSettingsAndCulling)
+						{
+							AllSettingsAndCulling = &CurrentExecutionDynamicTracking.Add(It.Key);
+						}
+
+						check(AllSettingsAndCulling);
+						// No need for Add Unique since they are already unique in the original map.
+						AllSettingsAndCulling->Add(std::move(SettingsAndCulling));
+					}
+				}
+
+				if (!AllSettingsAndCulling)
+				{
+					ChangedKeys.Add(It.Key);
+				}
+			}
+		}
+
+		DynamicallyTrackedKeysToSettings = std::move(CurrentExecutionDynamicTracking);
+		CurrentExecutionDynamicTracking.Empty();
+		CurrentExecutionDynamicTrackingSettings.Empty();
+	}
+
+	if (!ChangedKeys.IsEmpty())
+	{
+		Subsystem->UpdateComponentTracking(this, /*bShouldDirtyActors=*/false, &ChangedKeys);
+	}
+}
+
 #endif // WITH_EDITOR
 
 void UPCGComponent::SetManagedResources(const TArray<TObjectPtr<UPCGManagedResource>>& Resources)
diff --git a/Engine/Plugins/Experimental/PCG/Source/PCG/Private/PCGElement.cpp b/Engine/Plugins/Experimental/PCG/Source/PCG/Private/PCGElement.cpp
--- a/Engine/Plugins/Experimental/PCG/Source/PCG/Private/PCGElement.cpp
+++ b/Engine/Plugins/Experimental/PCG/Source/PCG/Private/PCGElement.cpp
@@ -174,6 +174,14 @@ void IPCGElement::PostExecute(FPCGContext* Context) const
 		Context->OutputData.Crc = Context->OutputData.ComputeCrc(bShouldComputeFullOutputDataCrc);
 	}
 
+#if WITH_EDITOR
+	// Register the element to the component indicating the element has run and can have dynamic tracked keys.
+	if (Settings && Settings->CanDynamicalyTrackKeys() && Context->SourceComponent.IsValid())
+	{
+		Context->SourceComponent->RegisterDynamicTracking(Settings, {});
+	}
+#endif // WITH_EDITOR
+
 	Context->CurrentPhase = EPCGExecutionPhase::Done;
 }
 
diff --git a/Engine/Plugins/Experimental/PCG/Source/PCG/Private/PCGGraph.cpp b/Engine/Plugins/Experimental/PCG/Source/PCG/Private/PCGGraph.cpp
--- a/Engine/Plugins/Experimental/PCG/Source/PCG/Private/PCGGraph.cpp
+++ b/Engine/Plugins/Experimental/PCG/Source/PCG/Private/PCGGraph.cpp
@@ -1027,7 +1027,7 @@ void UPCGGraph::GetTrackedActorKeysToSettings(FPCGSelectionKeyToSettingsMap& Out
 		// Don't track for disabled nodes.
 		if (Settings && Settings->bEnabled)
 		{
-			Settings->GetTrackedActorKeys(OutTagsToSettings, OutVisitedGraphs);
+			Settings->GetStaticTrackedKeys(OutTagsToSettings, OutVisitedGraphs);
 		}
 	}
 
diff --git a/Engine/Plugins/Experimental/PCG/Source/PCG/Private/PCGSubgraph.cpp b/Engine/Plugins/Experimental/PCG/Source/PCG/Private/PCGSubgraph.cpp
--- a/Engine/Plugins/Experimental/PCG/Source/PCG/Private/PCGSubgraph.cpp
+++ b/Engine/Plugins/Experimental/PCG/Source/PCG/Private/PCGSubgraph.cpp
@@ -7,6 +7,7 @@
 #include "PCGSubsystem.h"
 #include "Data/PCGUserParametersData.h"
 #include "Graph/PCGStackContext.h"
+#include "Helpers/PCGDynamicTrackingHelpers.h"
 #include "Helpers/PCGSettingsHelpers.h"
 
 #include "Algo/Find.h"
@@ -155,7 +156,7 @@ void UPCGBaseSubgraphSettings::PostEditChangeProperty(struct FPropertyChangedEve
 	}
 }
 
-void UPCGBaseSubgraphSettings::GetTrackedActorKeys(FPCGSelectionKeyToSettingsMap& OutKeysToSettings, TArray<TObjectPtr<const UPCGGraph>>& VisitedGraphs) const
+void UPCGBaseSubgraphSettings::GetStaticTrackedKeys(FPCGSelectionKeyToSettingsMap& OutKeysToSettings, TArray<TObjectPtr<const UPCGGraph>>& VisitedGraphs) const
 {
 	if (UPCGGraph* Subgraph = GetSubgraph())
 	{
@@ -528,6 +529,9 @@ bool FPCGSubgraphElement::ExecuteInternal(FPCGContext* InContext) const
 		{
 			// If OriginalSettings is null, then we ARE the original settings, and writing over the existing graph is incorrect (and potentially a race condition)
 			Settings->SubgraphInstance->SetGraph(Settings->SubgraphOverride);
+#if WITH_EDITOR
+			FPCGDynamicTrackingHelper::AddSingleDynamicTrackingKey(Context, FPCGSelectionKey::CreateFromPath(Settings->SubgraphOverride), /*bIsCulled=*/false);
+#endif // WITH_EDITOR
 		}
 	}
 
diff --git a/Engine/Plugins/Experimental/PCG/Source/PCG/Public/Elements/IO/PCGDataTableElement.h b/Engine/Plugins/Experimental/PCG/Source/PCG/Public/Elements/IO/PCGDataTableElement.h
--- a/Engine/Plugins/Experimental/PCG/Source/PCG/Public/Elements/IO/PCGDataTableElement.h
+++ b/Engine/Plugins/Experimental/PCG/Source/PCG/Public/Elements/IO/PCGDataTableElement.h
@@ -22,7 +22,8 @@ public:
 	virtual FText GetDefaultNodeTitle() const override;
 	virtual FText GetNodeTooltipText() const override;
 	virtual bool HasDynamicPins() const override { return true; }
-	virtual void GetTrackedActorKeys(FPCGSelectionKeyToSettingsMap& OutKeysToSettings, TArray<TObjectPtr<const UPCGGraph>>& OutVisitedGraphs) const override;
+	virtual void GetStaticTrackedKeys(FPCGSelectionKeyToSettingsMap& OutKeysToSettings, TArray<TObjectPtr<const UPCGGraph>>& OutVisitedGraphs) const override;
+	virtual bool CanDynamicalyTrackKeys() const override { return true; }
 #endif
 
 	virtual EPCGDataType GetCurrentPinTypes(const UPCGPin* InPin) const override;
diff --git a/Engine/Plugins/Experimental/PCG/Source/PCG/Public/Elements/PCGActorSelector.h b/Engine/Plugins/Experimental/PCG/Source/PCG/Public/Elements/PCGActorSelector.h
--- a/Engine/Plugins/Experimental/PCG/Source/PCG/Public/Elements/PCGActorSelector.h
+++ b/Engine/Plugins/Experimental/PCG/Source/PCG/Public/Elements/PCGActorSelector.h
@@ -42,8 +42,11 @@ enum class EPCGActorFilter : uint8
 * Object can be selected using the EPCGActorSelection::ByClass or EPCGActorSelection::ByPath
 * Actors have more options for selection with Self/Parent/Root/Original and also EPCGActorSelection::ByTag
 */
+USTRUCT()
 struct FPCGSelectionKey
 {
+	GENERATED_BODY()
+
 	FPCGSelectionKey() = default;
 
 	// For all filters others than AllWorldActor. For AllWorldActors Filter, use the other constructors.
@@ -53,26 +56,39 @@ struct FPCGSelectionKey
 	explicit FPCGSelectionKey(TSubclassOf<UObject> InSelectionClass);
 
 	static FPCGSelectionKey CreateFromPath(const FSoftObjectPath& InObjectPath);
+	static FPCGSelectionKey CreateFromPath(FSoftObjectPath&& InObjectPath);
 
 	bool operator==(const FPCGSelectionKey& InOther) const;
 
 	friend uint32 GetTypeHash(const FPCGSelectionKey& In);
-	bool IsMatching(const TSoftObjectPtr<UObject>& InObjectPtr, const UPCGComponent* InComponent) const;
+	bool IsMatching(const UObject* InObject, const UPCGComponent* InComponent) const;
+	bool IsMatching(const UObject* InObject, const TSet<FName>& InRemovedTags, const TSet<UPCGComponent*>& InComponents, TSet<UPCGComponent*>& MatchedComponents) const;
 
 	void SetExtraDependency(const UClass* InExtraDependency);
 
+	UPROPERTY()
 	EPCGActorFilter ActorFilter = EPCGActorFilter::AllWorldActors;
+
+	UPROPERTY()
 	EPCGActorSelection Selection = EPCGActorSelection::Unknown;
+
+	UPROPERTY()
 	FName Tag = NAME_None;
+
+	UPROPERTY()
 	TSubclassOf<UObject> SelectionClass = nullptr;
 
 	// If the Selection is ByPath, contain the path to select.
+	UPROPERTY()
 	FSoftObjectPath ObjectPath;
 
 	// If it should track a specific object dependency instead of an actor. For example, GetActorData with GetPCGComponent data.
-	const UClass* OptionalExtraDependency = nullptr;
+	UPROPERTY()
+	TObjectPtr<const UClass> OptionalExtraDependency = nullptr;
 };
 
+FArchive& operator<<(FArchive& Ar, FPCGSelectionKey& Key);
+
 /** Helper struct for organizing queries against the world to gather actors. */
 USTRUCT(BlueprintType)
 struct FPCGActorSelectorSettings
diff --git a/Engine/Plugins/Experimental/PCG/Source/PCG/Public/Elements/PCGDataFromActor.h b/Engine/Plugins/Experimental/PCG/Source/PCG/Public/Elements/PCGDataFromActor.h
--- a/Engine/Plugins/Experimental/PCG/Source/PCG/Public/Elements/PCGDataFromActor.h
+++ b/Engine/Plugins/Experimental/PCG/Source/PCG/Public/Elements/PCGDataFromActor.h
@@ -31,7 +31,7 @@ public:
 	virtual FText GetDefaultNodeTitle() const override { return NSLOCTEXT("PCGDataFromActorSettings", "NodeTitle", "Get Actor Data"); }
 	virtual FText GetNodeTooltipText() const override;
 	virtual EPCGSettingsType GetType() const override { return EPCGSettingsType::Spatial; }
-	virtual void GetTrackedActorKeys(FPCGSelectionKeyToSettingsMap& OutKeysToSettings, TArray<TObjectPtr<const UPCGGraph>>& OutVisitedGraphs) const override;
+	virtual void GetStaticTrackedKeys(FPCGSelectionKeyToSettingsMap& OutKeysToSettings, TArray<TObjectPtr<const UPCGGraph>>& OutVisitedGraphs) const override;
 	virtual bool HasDynamicPins() const override { return true; }
 	virtual void ApplyDeprecation(UPCGNode* InOutNode) override;
 #endif
diff --git a/Engine/Plugins/Experimental/PCG/Source/PCG/Public/Elements/PCGDataTableRowToParamData.h b/Engine/Plugins/Experimental/PCG/Source/PCG/Public/Elements/PCGDataTableRowToParamData.h
--- a/Engine/Plugins/Experimental/PCG/Source/PCG/Public/Elements/PCGDataTableRowToParamData.h
+++ b/Engine/Plugins/Experimental/PCG/Source/PCG/Public/Elements/PCGDataTableRowToParamData.h
@@ -14,11 +14,18 @@ class UPCGDataTableRowToParamDataSettings : public UPCGSettings
 	GENERATED_BODY()
 
 public:
+	//~Begin UObject interface
+	virtual void PostLoad() override;
+	//~End UObject interface
+
 #if WITH_EDITOR
 	//~Begin UPCGSettings interface
 	virtual FName GetDefaultNodeName() const override { return FName(TEXT("DataTableRowToAttributeSet")); }
 	virtual FText GetDefaultNodeTitle() const override { return NSLOCTEXT("PCGDataTableRowToParamDataSettings", "NodeTitle", "Data Table Row To Attribute Set"); }
 	virtual EPCGSettingsType GetType() const override { return EPCGSettingsType::Metadata; }
+	virtual void GetStaticTrackedKeys(FPCGSelectionKeyToSettingsMap& OutKeysToSettings, TArray<TObjectPtr<const UPCGGraph>>& OutVisitedGraphs) const override;
+	virtual bool CanDynamicalyTrackKeys() const override { return true; }
+	virtual void ApplyDeprecationBeforeUpdatePins(UPCGNode* InOutNode, TArray<TObjectPtr<UPCGPin>>& InputPins, TArray<TObjectPtr<UPCGPin>>& OutputPins);
 #endif
 
 	virtual FString GetAdditionalTitleInformation() const override;
@@ -30,14 +37,16 @@ public:
 	UPROPERTY(BlueprintReadWrite, EditAnywhere, Category = Settings, meta = (PCG_Overridable))
 	FName RowName = NAME_None;	
 
-	// Path override, hidden to be only presented as param pin
-	UPROPERTY(BlueprintReadWrite, Category = Settings, meta = (PCG_Overridable))
-	FString PathOverride = FString();
-
 	// the data table to copy from
-	UPROPERTY(BlueprintReadWrite, EditAnywhere, Category = Settings)
+	UPROPERTY(BlueprintReadWrite, EditAnywhere, Category = Settings, meta = (PCG_Overridable, PCG_OverrideAliases = "PathOverride"))
 	TSoftObjectPtr<UDataTable> DataTable;
 
+#if WITH_EDITORONLY_DATA
+	UPROPERTY()
+	FString PathOverride_DEPRECATED = FString();
+#endif // 
+
+
 protected:
 #if WITH_EDITOR
 	virtual EPCGChangeType GetChangeTypeForProperty(const FName& InPropertyName) const override { return Super::GetChangeTypeForProperty(InPropertyName) | EPCGChangeType::Cosmetic; }
diff --git a/Engine/Plugins/Experimental/PCG/Source/PCG/Public/Elements/PCGExecuteBlueprint.h b/Engine/Plugins/Experimental/PCG/Source/PCG/Public/Elements/PCGExecuteBlueprint.h
--- a/Engine/Plugins/Experimental/PCG/Source/PCG/Public/Elements/PCGExecuteBlueprint.h
+++ b/Engine/Plugins/Experimental/PCG/Source/PCG/Public/Elements/PCGExecuteBlueprint.h
@@ -301,7 +301,7 @@ public:
 	virtual FText GetDefaultNodeTitle() const override { return NSLOCTEXT("PCGBlueprintSettings", "NodeTitle", "Execute Blueprint"); }
 	virtual FLinearColor GetNodeTitleColor() const override;
 	virtual EPCGSettingsType GetType() const override;
-	virtual void GetTrackedActorKeys(FPCGSelectionKeyToSettingsMap& OutKeysToSettings, TArray<TObjectPtr<const UPCGGraph>>& OutVisitedGraphs) const override;
+	virtual void GetStaticTrackedKeys(FPCGSelectionKeyToSettingsMap& OutKeysToSettings, TArray<TObjectPtr<const UPCGGraph>>& OutVisitedGraphs) const override;
 	virtual UObject* GetJumpTargetForDoubleClick() const override;
 	virtual void ApplyDeprecationBeforeUpdatePins(UPCGNode* InOutNode, TArray<TObjectPtr<UPCGPin>>& InputPins, TArray<TObjectPtr<UPCGPin>>& OutputPins) override;
 	virtual TArray<FPCGPreConfiguredSettingsInfo> GetPreconfiguredInfo() const override;
diff --git a/Engine/Plugins/Experimental/PCG/Source/PCG/Public/Elements/PCGGetActorProperty.h b/Engine/Plugins/Experimental/PCG/Source/PCG/Public/Elements/PCGGetActorProperty.h
--- a/Engine/Plugins/Experimental/PCG/Source/PCG/Public/Elements/PCGGetActorProperty.h
+++ b/Engine/Plugins/Experimental/PCG/Source/PCG/Public/Elements/PCGGetActorProperty.h
@@ -31,7 +31,8 @@ public:
 	virtual FName GetDefaultNodeName() const override { return FName(TEXT("GetActorProperty")); }
 	virtual FText GetDefaultNodeTitle() const override { return NSLOCTEXT("PCGPropertyToParamDataSettings", "NodeTitle", "Get Actor Property"); }
 	virtual EPCGSettingsType GetType() const override { return EPCGSettingsType::Param; }
-	virtual void GetTrackedActorKeys(FPCGSelectionKeyToSettingsMap& OutKeysToSettings, TArray<TObjectPtr<const UPCGGraph>>& OutVisitedGraphs) const override;
+	virtual void GetStaticTrackedKeys(FPCGSelectionKeyToSettingsMap& OutKeysToSettings, TArray<TObjectPtr<const UPCGGraph>>& OutVisitedGraphs) const override;
+	virtual bool CanDynamicalyTrackKeys() const override { return true; }
 #endif
 
 	virtual FString GetAdditionalTitleInformation() const override;
diff --git a/Engine/Plugins/Experimental/PCG/Source/PCG/Public/Elements/PCGGetPropertyFromObjectPath.h b/Engine/Plugins/Experimental/PCG/Source/PCG/Public/Elements/PCGGetPropertyFromObjectPath.h
--- a/Engine/Plugins/Experimental/PCG/Source/PCG/Public/Elements/PCGGetPropertyFromObjectPath.h
+++ b/Engine/Plugins/Experimental/PCG/Source/PCG/Public/Elements/PCGGetPropertyFromObjectPath.h
@@ -25,7 +25,8 @@ public:
 	virtual EPCGSettingsType GetType() const override { return EPCGSettingsType::Param; }
 	virtual bool CanEditChange(const FProperty* InProperty) const override;
 	virtual bool IsPinUsedByNodeExecution(const UPCGPin* InPin) const override;
-	virtual void GetTrackedActorKeys(FPCGSelectionKeyToSettingsMap& OutKeysToSettings, TArray<TObjectPtr<const UPCGGraph>>& OutVisitedGraphs) const override;
+	virtual void GetStaticTrackedKeys(FPCGSelectionKeyToSettingsMap& OutKeysToSettings, TArray<TObjectPtr<const UPCGGraph>>& OutVisitedGraphs) const override;
+	virtual bool CanDynamicalyTrackKeys() const override { return true; }
 #endif
 	virtual FString GetAdditionalTitleInformation() const override;
 
@@ -81,8 +82,6 @@ struct FPCGGetPropertyFromObjectPathContext : public FPCGContext
 
 class FPCGGetPropertyFromObjectPathElement : public IPCGElement
 {
-public:
-	virtual bool IsCacheable(const UPCGSettings* InSettings) const { return false; }
 protected:
 	virtual FPCGContext* CreateContext() override;
 	virtual bool PrepareDataInternal(FPCGContext* Context) const override;
diff --git a/Engine/Plugins/Experimental/PCG/Source/PCG/Public/Elements/PCGIndirectionElement.h b/Engine/Plugins/Experimental/PCG/Source/PCG/Public/Elements/PCGIndirectionElement.h
--- a/Engine/Plugins/Experimental/PCG/Source/PCG/Public/Elements/PCGIndirectionElement.h
+++ b/Engine/Plugins/Experimental/PCG/Source/PCG/Public/Elements/PCGIndirectionElement.h
@@ -29,7 +29,8 @@ public:
 	virtual FText GetDefaultNodeTitle() const override;
 	virtual FText GetNodeTooltipText() const override;
 	virtual EPCGSettingsType GetType() const override { return EPCGSettingsType::Generic; }
-	virtual void GetTrackedActorKeys(FPCGSelectionKeyToSettingsMap& OutKeysToSettings, TArray<TObjectPtr<const UPCGGraph>>& OutVisitedGraphs) const override;
+	virtual void GetStaticTrackedKeys(FPCGSelectionKeyToSettingsMap& OutKeysToSettings, TArray<TObjectPtr<const UPCGGraph>>& OutVisitedGraphs) const override;
+	virtual bool CanDynamicalyTrackKeys() const override { return true; }
 #endif // WITH_EDITOR
 	virtual FString GetAdditionalTitleInformation() const override;
 	virtual bool HasFlippedTitleLines() const override { return true; }
diff --git a/Engine/Plugins/Experimental/PCG/Source/PCG/Public/Elements/PCGTextureSampler.h b/Engine/Plugins/Experimental/PCG/Source/PCG/Public/Elements/PCGTextureSampler.h
--- a/Engine/Plugins/Experimental/PCG/Source/PCG/Public/Elements/PCGTextureSampler.h
+++ b/Engine/Plugins/Experimental/PCG/Source/PCG/Public/Elements/PCGTextureSampler.h
@@ -29,6 +29,8 @@ public:
 	virtual FText GetDefaultNodeTitle() const override { return NSLOCTEXT("PCGTextureSamplerSettings", "NodeTitle", "Get Texture Data"); }
 	virtual FText GetNodeTooltipText() const override;
 	virtual EPCGSettingsType GetType() const override { return EPCGSettingsType::Spatial; }
+	virtual void GetStaticTrackedKeys(FPCGSelectionKeyToSettingsMap& OutKeysToSettings, TArray<TObjectPtr<const UPCGGraph>>& OutVisitedGraphs) const override;
+	virtual bool CanDynamicalyTrackKeys() const override { return true; }
 #endif
 
 protected:
diff --git a/Engine/Plugins/Experimental/PCG/Source/PCG/Public/Helpers/PCGDynamicTrackingHelpers.h b/Engine/Plugins/Experimental/PCG/Source/PCG/Public/Helpers/PCGDynamicTrackingHelpers.h
new file mode 100644
--- /dev/null
+++ b/Engine/Plugins/Experimental/PCG/Source/PCG/Public/Helpers/PCGDynamicTrackingHelpers.h
@@ -0,0 +1,38 @@
+// Copyright Epic Games, Inc. All Rights Reserved.
+
+#pragma once
+
+#if WITH_EDITOR
+
+#include "Elements/PCGActorSelector.h"
+#include "Containers/Array.h"
+
+struct FPCGContext;
+class UPCGComponent;
+
+/**
+* Simple helper class to factorize the logic for gathering dynamic tracking keys and pushing them to the component.
+* Only work for settings that override CanDynamicalyTrackKeys.
+*/
+struct PCG_API FPCGDynamicTrackingHelper
+{
+public:
+	/** Enable dynamic tracking, will cache the weak ptr of the component and optionally resize the array for keys. */
+	void EnableAndInitialize(const FPCGContext* InContext, int32 OptionalNumElements = 0);
+
+	/** Add the key to the tracking, will be uniquely added to the array. */
+	void AddToTracking(FPCGSelectionKey&& InKey, bool bIsCulled);
+
+	/** Push all the tracked keys to the ceched component if still valid and the same as the context. */
+	void Finalize(const FPCGContext* InContext);
+
+	/** Convenient function to push just a single tracking key to the component. */
+	static void AddSingleDynamicTrackingKey(FPCGContext* InContext, FPCGSelectionKey&& InKey, bool bIsCulled);
+
+private:
+	bool bDynamicallyTracked = false;
+	TWeakObjectPtr<UPCGComponent> CachedComponent;
+	TArray<TPair<FPCGSelectionKey, bool>, TInlineAllocator<16>> DynamicallyTrackedKeysAndCulling;
+};
+
+#endif // WITH_EDITOR
diff --git a/Engine/Plugins/Experimental/PCG/Source/PCG/Public/Helpers/PCGPropertyHelpers.h b/Engine/Plugins/Experimental/PCG/Source/PCG/Public/Helpers/PCGPropertyHelpers.h
--- a/Engine/Plugins/Experimental/PCG/Source/PCG/Public/Helpers/PCGPropertyHelpers.h
+++ b/Engine/Plugins/Experimental/PCG/Source/PCG/Public/Helpers/PCGPropertyHelpers.h
@@ -71,9 +71,10 @@ namespace PCGPropertyHelpers
 	/**
 	* Extract a given property in an Attribute Set.
 	* @param Parameters - Parameters for extaction, cf above.
-	* @prama InOptionalContext - Optional context if the extraction is done in a PCG Node, so errors are using the context to log.
+	* @param InOptionalContext - Optional context if the extraction is done in a PCG Node, so errors are using the context to log.
+	* @param OptionalObjectTraversed - Optional set to store all objects that we traversed, to be able to react to those objects changes.
 	*/
-	UPCGParamData* ExtractPropertyAsAttributeSet(const FExtractorParameters& Parameters, FPCGContext* InOptionalContext = nullptr);
+	UPCGParamData* ExtractPropertyAsAttributeSet(const FExtractorParameters& Parameters, FPCGContext* OptionalContext = nullptr, TSet<FSoftObjectPath>* OptionalObjectTraversed = nullptr);
 }
 
 //////
diff --git a/Engine/Plugins/Experimental/PCG/Source/PCG/Public/PCGActorAndComponentMapping.h b/Engine/Plugins/Experimental/PCG/Source/PCG/Public/PCGActorAndComponentMapping.h
--- a/Engine/Plugins/Experimental/PCG/Source/PCG/Public/PCGActorAndComponentMapping.h
+++ b/Engine/Plugins/Experimental/PCG/Source/PCG/Public/PCGActorAndComponentMapping.h
@@ -129,16 +129,8 @@ private:
 	bool AnyRuntimeGenComponentsExist() const;
 
 #if WITH_EDITOR
-	/* Return true if something is still tracked or was just untracked. */
-	bool AddOrUpdateTrackedActor(AActor* InActor);
-
-	void RegisterActor(AActor* InActor);
-
-	/* Return true if the actor was tracked. */
-	bool UnregisterActor(AActor* InActor);
-
-	/* Return true if the object was tracked. */
-	bool UnregisterObject(const TSoftObjectPtr<UObject>& InObject);
+	/** Return true if the key is tracked.*/
+	bool IsKeyTracked(const FPCGSelectionKey& InKey) const;
 
 	void OnActorAdded(AActor* InActor);
 	void OnActorLoaded(AActor& InActor);
@@ -146,8 +138,6 @@ private:
 	void OnActorDeleted(AActor* InActor);
 	void OnActorUnloaded(AActor& InActor);
 	void OnActorDeleted_Internal(AActor* InActor, bool bShouldDirty, int32 LevelInstanceDepth);
-	void OnActorMoved(AActor* InActor);
-	void OnActorMoved_Internal(AActor* InActor, int32 LevelInstanceDepth);
 	void OnLandscapeChanged(ALandscapeProxy* InLandscape, const FLandscapeProxyComponentDataChangedParams& InChangeParams);
 	void ApplyLandscapeChanges(ALandscapeProxy* InLandscape);
 	void OnPreObjectPropertyChanged(UObject* InObject, const FEditPropertyChain& InEditPropertyChain);
@@ -162,22 +152,17 @@ private:
 	void UnregisterTracking(UPCGComponent* InComponent);
 
 	/** Unregister tracking when a component is removed or keys are unregistered */
-	void UnregisterTracking(UPCGComponent* InComponent, const TSet<TSoftObjectPtr<UObject>>* OptionalObjectsToUntrack);
+	void UnregisterTracking(UPCGComponent* InComponent, const TSet<FPCGSelectionKey>* OptionalKeysToUntrack);
 
-	/** Trigger an update when the actor changed. 
-	* Can specify if the actor has moved to also update components that were at its previous position.
+	/** Trigger an update when an object/actor changed.
+	* Can specify previous actor bounds if it has moved to also update components that were at its previous position.
 	* Can also specify an optional object, originating the change, to avoid re-dirtying a component if it was the origin.
 	* Another option when an actor is deleted/unload, don't refresh their components.
 	*/
-	void OnActorChanged(AActor* InActor, bool bInHasMoved, const UObject* InOriginatingChangeObject = nullptr, int32 LevelInstanceDepth = 0, bool bNoRefreshOnOwner = false);
-
-	/** Update dependencies for a given tracked actor. */
-	void UpdateActorDependencies(AActor* InActor);
-
-	bool IsActorTracked(const AActor* InActor) const;
+	void OnObjectChanged(UObject* InObject, const FBox& InPreviousBounds = FBox(EForceInit::ForceInit), const UObject* InOriginatingChangeObject = nullptr, int32 LevelInstanceDepth = 0, bool bNoRefreshOnOwner = false);
 
-	/** Gather all settings from a given component that track the actor, and clear the cache for them. Returns true if we should dirty afterwards (aka at least one settings was cleared and/or landscape changed). */
-	bool ClearCache(const UObject* InObject, const UPCGComponent* InComponent, const bool bIntersect, const TSet<FName>& InRemovedTags, const UObject* InOriginatingChange) const;
+	/** Gather all settings from a given component that track the key, and clear the cache for them. Returns true if we should dirty afterwards (aka at least one settings was cleared and/or landscape changed). */
+	bool ClearCacheForKeys(const TArray<FPCGSelectionKey>& InKeys, const UPCGComponent* InComponent, const bool bIntersect, const UObject* InOriginatingChange) const;
 #endif // WITH_EDITOR
 
 private:
@@ -212,29 +197,23 @@ private:
 	/** Will hold all the components that are not partitioned (and not local) and are tracking something. Will be use to dispatch actor tracking updates. */
 	FPCGComponentOctreeAndMap NonPartitionedOctree;
 
-	/** Keep a mapping between tracked actors and the components that track them, and the tracking needs to be culled.*/
-	TMap<TObjectKey<AActor>, TSet<UPCGComponent*>> CulledTrackedActorsToComponentsMap;
+	/** Keep a mapping between tracked keys and the components that track them, and the tracking needs to be culled.*/
+	TMap<FPCGSelectionKey, TSet<UPCGComponent*>> CulledTrackedKeysToComponentsMap;
 
-	/** Same mapping but for always tracked actors */
-	TMap<TObjectKey<AActor>, TSet<UPCGComponent*>> AlwaysTrackedActorsToComponentsMap;
-
-	/** Object tracking. Use soft pointers, because tracked objects might not be yet loaded. */
-	TMap<TSoftObjectPtr<UObject>, TSet<UPCGComponent*>> TrackedObjectsToComponentsMap;
-
-	/** Finally keep a mapping between actors and their position to know if an actor move in/out of a component tracking bounds. Only kept for actors that need to be culled. */
-	TMap<TObjectKey<AActor>, FBox> TrackedActorToPositionMap;
+	/** Same mapping but for always tracked keys */
+	TMap<FPCGSelectionKey, TSet<UPCGComponent*>> AlwaysTrackedKeysToComponentsMap;
 
 	mutable FRWLock TrackedComponentsLock;
 
 	// Keep track of actors that aren't yet ready (or if the subsystem is not yet ready), whether we should dirty them and their instance level depth so we can add them in next tick.
 	TMap<TObjectKey<AActor>, TTuple<bool, int>> DelayedAddedActors;
 
-	/** Keep a mapping between tracked actors and their dependencies. */
-	TMap<TObjectKey<AActor>, TSet<TObjectPtr<UObject>>> TrackedActorsToDependenciesMap;
-
 	/** Transient list of tags, kept when there is a tag change on a tracked Actor. */
 	TSet<FName> TempTrackedActorTags;
 
+	/** Transient map of actors and their previous bounds, to know when they move and update the components that were touching the actor before, but not anymore. */
+	TMap<TObjectKey<AActor>, FBox> ActorToPreviousBoundsMap;
+
 #if WITH_EDITOR
 	// Part for the delayed landscape change update
 	double LastLandscapeDirtyTime = -1.0;
diff --git a/Engine/Plugins/Experimental/PCG/Source/PCG/Public/PCGComponent.h b/Engine/Plugins/Experimental/PCG/Source/PCG/Public/PCGComponent.h
--- a/Engine/Plugins/Experimental/PCG/Source/PCG/Public/PCGComponent.h
+++ b/Engine/Plugins/Experimental/PCG/Source/PCG/Public/PCGComponent.h
@@ -299,7 +299,7 @@ public:
 	/** Did the given node produce one or more data items in the given stack in a previous execution. */
 	bool HasNodeProducedData(const UPCGNode* InNode, const FPCGStack& Stack) const;
 
-	bool IsObjectTracked(const TSoftObjectPtr<UObject>& InObjectPtr, bool& bOutIsCulled) const;
+	bool IsObjectTracked(const UObject* InObject, bool& bOutIsCulled) const;
 
 	/** Know if we need to force a generation, in case of BP added to the world in editor */
 	bool ShouldGenerateBPPCGAddedToWorld() const;
@@ -309,6 +309,9 @@ public:
 
 	/** Get execution stack information. */
 	bool GetStackContext(FPCGStackContext& OutStackContext) const;
+
+	/** To be called by an element to notify the component that this settings have a dynamic dependency. */
+	void RegisterDynamicTracking(const UPCGSettings* InSettings, const TArrayView<TPair<FPCGSelectionKey, bool>>& InDynamicKeysAndCulling);
 #endif
 
 	/** Utility function (mostly for tests) to properly set the value of bIsComponentPartitioned.
@@ -437,14 +440,17 @@ private:
 	/** Returns true if something changed in the tracking. */
 	bool UpdateTrackingCache(TArray<FPCGSelectionKey>* OptionalChangedKeys = nullptr);
 
-	/** Gather all the settings that are tracking the given object. If it is an actor, the settings should cull and there is no intersection, it will not be added to the list.
-	* @param InObject                   The object the component is tracking. Can be an Actor (for culling).
-	* @param bIntersect                 If the actor is intersecting with the component. Useful for culling.
-	* @param InRemovedTags              List of tags that were removed, because the actor won't have those tags on it, but we still need to gather the settings that were tracking this tag. Can be empty.
-	* @param InOriginatingChangeObject: Optional pointer on the originating object that triggered the update. Useful to track PCGComponents generation/cleanup. Can be null.
-	* @return                           Array of all the settings that track the actor.
-	*/
-	TArray<const UPCGSettings*> GatherSettingsTracking(const UObject* InObject, const bool bIntersect, const TSet<FName>& InRemovedTags, const UObject* InOriginatingChangeObject) const;
+	/** Apply a function to all settings that track a given key. */
+	void ApplyToEachSettings(const FPCGSelectionKey& InKey, const TFunctionRef<void(const FPCGSelectionKey&, const FPCGSettingsAndCulling&)> InCallback) const;
+
+	/** Return all the keys tracked by the component (statically and dynamically). */
+	TArray<FPCGSelectionKey> GatherTrackingKeys() const;
+
+	/** Return true if the key is tracked, and if so, bOutIsCulled will contains if the key is culled or not. */
+	bool IsKeyTrackedAndCulled(const FPCGSelectionKey& Key, bool& bOutIsCulled) const;
+
+	/** Compare the temp map to the stored map for dynamic tracking and register/unregister accordingly. If it is a local component, it will push the info to the original. */
+	void UpdateDynamicTracking();
 
 	bool ShouldTrackLandscape() const;
 
@@ -515,12 +521,21 @@ private:
 #if WITH_EDITOR
 	int32 InspectionCounter = 0;
 	FBox LastGeneratedBoundsPriorToUndo = FBox(EForceInit::ForceInit);
-	FPCGSelectionKeyToSettingsMap CachedTrackedKeysToSettings;
-
-	TMap<FPCGSelectionKey, bool> CachedTrackedKeysToCulling;
 #endif
 
 #if WITH_EDITORONLY_DATA
+	FPCGSelectionKeyToSettingsMap StaticallyTrackedKeysToSettings;
+
+	// Temporary storage for dynamic tracking that will be filled during component execution.
+	FPCGSelectionKeyToSettingsMap CurrentExecutionDynamicTracking;
+	// Temporary storage for dynamic tracking that will keep all settings that could have dynamic tracking, in order to detect changes.
+	TSet<const UPCGSettings*> CurrentExecutionDynamicTrackingSettings;
+	mutable FCriticalSection CurrentExecutionDynamicTrackingLock;
+
+	// Need to keep a reference to all tracked settings to still react to changes after a map load (since the component won't have been executed).
+	// Serialization will be done in the Serialize function
+	FPCGSelectionKeyToSettingsMap DynamicallyTrackedKeysToSettings;
+
 	UPROPERTY(Transient)
 	TMap<FPCGStack, FPCGDataCollection> InspectionCache;
 #endif
diff --git a/Engine/Plugins/Experimental/PCG/Source/PCG/Public/PCGContext.h b/Engine/Plugins/Experimental/PCG/Source/PCG/Public/PCGContext.h
--- a/Engine/Plugins/Experimental/PCG/Source/PCG/Public/PCGContext.h
+++ b/Engine/Plugins/Experimental/PCG/Source/PCG/Public/PCGContext.h
@@ -49,7 +49,7 @@ struct PCG_API FPCGContext
 {
 	GENERATED_BODY()
 
-	virtual ~FPCGContext() {}
+	virtual ~FPCGContext() = default;
 
 	FPCGDataCollection InputData;
 	FPCGDataCollection OutputData;
@@ -116,18 +116,19 @@ struct PCG_API FPCGContext
 	//  in properties in any case. This will be called from the graph executor when needed and is implemented to look like normal reference traversal.
 	void AddStructReferencedObjects(FReferenceCollector& Collector);
 
-protected:
-	virtual UObject* GetExternalContainerForOverridableParam(const FPCGSettingsOverridableParam& InParam) { return nullptr; }
-	virtual void* GetUnsafeExternalContainerForOverridableParam(const FPCGSettingsOverridableParam& InParam) { return nullptr; }
-	virtual void AddExtraStructReferencedObjects(FReferenceCollector& Collector) {}
-
-private:
+	/** Caution: most use cases should use GetInputSettings, because they contain the overridden values. Use this one if you really need to get the original pointer. */
 	template<typename SettingsType>
 	const SettingsType* GetOriginalSettings() const 
 	{
 		return PCGContextHelpers::GetInputSettings<SettingsType>(Node, InputData);
 	}
 
+protected:
+	virtual UObject* GetExternalContainerForOverridableParam(const FPCGSettingsOverridableParam& InParam) { return nullptr; }
+	virtual void* GetUnsafeExternalContainerForOverridableParam(const FPCGSettingsOverridableParam& InParam) { return nullptr; }
+	virtual void AddExtraStructReferencedObjects(FReferenceCollector& Collector) {}
+
+private:
 	// Copy of the settings that will be used to apply overrides.
 	TObjectPtr<UPCGSettings> SettingsWithOverride = nullptr;
 
diff --git a/Engine/Plugins/Experimental/PCG/Source/PCG/Public/PCGCustomVersion.h b/Engine/Plugins/Experimental/PCG/Source/PCG/Public/PCGCustomVersion.h
--- a/Engine/Plugins/Experimental/PCG/Source/PCG/Public/PCGCustomVersion.h
+++ b/Engine/Plugins/Experimental/PCG/Source/PCG/Public/PCGCustomVersion.h
@@ -68,6 +68,9 @@ struct PCG_API FPCGCustomVersion
 		// Added 'bComponentsMustOverlapSelf' to GetActorData when the mode collects PCG component data
 		GetPCGComponentDataMustOverlapSourceComponentByDefault = 19,
 
+		// Added dynamic tracking to the PCG component serialization
+		DynamicTrackingKeysSerializedInComponent = 20,
+
 		// -----<new versions can be added above this line>-------------------------------------------------
 		VersionPlusOne,
 		LatestVersion = VersionPlusOne - 1
diff --git a/Engine/Plugins/Experimental/PCG/Source/PCG/Public/PCGSettings.h b/Engine/Plugins/Experimental/PCG/Source/PCG/Public/PCGSettings.h
--- a/Engine/Plugins/Experimental/PCG/Source/PCG/Public/PCGSettings.h
+++ b/Engine/Plugins/Experimental/PCG/Source/PCG/Public/PCGSettings.h
@@ -9,6 +9,8 @@
 #include "Elements/PCGActorSelector.h"
 #include "Tests/Determinism/PCGDeterminismSettings.h"
 
+#include "Algo/AllOf.h"
+
 #include "PCGSettings.generated.h"
 
 class UPCGComponent;
@@ -20,9 +22,18 @@ class UPCGGraph;
 class UPCGNode;
 class UPCGSettings;
 
-using FPCGSettingsAndCulling = TPair<TWeakObjectPtr<const UPCGSettings>, bool>;
+using FPCGSettingsAndCulling = TPair<TSoftObjectPtr<const UPCGSettings>, bool>;
 using FPCGSelectionKeyToSettingsMap = TMap<FPCGSelectionKey, TArray<FPCGSettingsAndCulling>>;
 
+namespace PCGSettings
+{
+	// A key is culled if and only if all the settings are culled.
+	inline bool IsKeyCulled(const TArray<FPCGSettingsAndCulling>& SettingsAndCulling)
+	{
+		return Algo::AllOf(SettingsAndCulling, [](const FPCGSettingsAndCulling& SettingsAndCullingPair) { return SettingsAndCullingPair.Value; });
+	}
+}
+
 UENUM()
 enum class EPCGSettingsExecutionMode : uint8
 {
@@ -278,8 +289,11 @@ public:
 	/** Can override to add a custom icon next to the pin label (and an optional tooltip). Return false if no override is available. */
 	virtual bool GetPinExtraIcon(const UPCGPin* InPin, FName& OutExtraIcon, FText& OutTooltip) const;
 
-	/** Derived classes must implement this to communicate dependencies on external actors */
-	virtual void GetTrackedActorKeys(FPCGSelectionKeyToSettingsMap& OutKeysToSettings, TArray<TObjectPtr<const UPCGGraph>>& OutVisitedGraphs) const {}
+	/** Derived classes must implement this to communicate dependencies that are known statically. */
+	virtual void GetStaticTrackedKeys(FPCGSelectionKeyToSettingsMap& OutKeysToSettings, TArray<TObjectPtr<const UPCGGraph>>& OutVisitedGraphs) const {}
+
+	/** Derived classes must implement this to communicate that they might have dynamic dependencies. */
+	virtual bool CanDynamicalyTrackKeys() const { return false; }
 
 	/** Override this class to provide an UObject to jump to in case of double click on node
 	 *  ie. returning a blueprint instance will open the given blueprint in its editor.
diff --git a/Engine/Plugins/Experimental/PCG/Source/PCG/Public/PCGSubgraph.h b/Engine/Plugins/Experimental/PCG/Source/PCG/Public/PCGSubgraph.h
--- a/Engine/Plugins/Experimental/PCG/Source/PCG/Public/PCGSubgraph.h
+++ b/Engine/Plugins/Experimental/PCG/Source/PCG/Public/PCGSubgraph.h
@@ -51,7 +51,8 @@ public:
 
 protected:
 #if WITH_EDITOR
-	virtual void GetTrackedActorKeys(FPCGSelectionKeyToSettingsMap& OutKeysToSettings, TArray<TObjectPtr<const UPCGGraph>>& OutVisitedGraphs) const override;
+	virtual void GetStaticTrackedKeys(FPCGSelectionKeyToSettingsMap& OutKeysToSettings, TArray<TObjectPtr<const UPCGGraph>>& OutVisitedGraphs) const override;
+	virtual bool CanDynamicalyTrackKeys() const override { return true; }
 	virtual EPCGChangeType GetChangeTypeForProperty(const FName& InPropertyName) const override;
 #endif
 
