[PCG] Fix issue when Switch node was triggering a generation when selected.

#rnx

diff --git a/Engine/Plugins/Experimental/PCG/Source/PCGEditor/Private/Details/EnumSelectorDetails.cpp b/Engine/Plugins/Experimental/PCG/Source/PCGEditor/Private/Details/EnumSelectorDetails.cpp
--- a/Engine/Plugins/Experimental/PCG/Source/PCGEditor/Private/Details/EnumSelectorDetails.cpp
+++ b/Engine/Plugins/Experimental/PCG/Source/PCGEditor/Private/Details/EnumSelectorDetails.cpp
@@ -240,6 +240,9 @@ void FEnumSelectorDetails::InitializeClassComboBox()
 void FEnumSelectorDetails::InitializeClassComboBoxSelection()
 {
 	bool bEnumFound = false;
+	ClassSelection.Reset();
+	ValueSelection = NAME_None;
+
 	// The selector is already defined, so use the struct data. Note: this is costly, but is needed to get the pointer to the correct widget option
 	if (UEnum* SelectedEnum = GetSelectedEnumClass())
 	{
@@ -248,23 +251,27 @@ void FEnumSelectorDetails::InitializeClassComboBoxSelection()
 			return SelectedEnum->GetName() == *InString;
 		});
 
+		// If an enum is selected, reload the options for this enum.
+		// If the selected value is valid (name is different of None), we're good.
 		if (SelectedOption)
 		{
 			ClassSelection = *SelectedOption;
 
 			ReloadValueOptions();
 			ValueSelection = SelectedEnum->GetNameByValue(GetSelectedEnumValue());
-			UpdateEnumValue(ValueSelection);
-			bEnumFound = true;
+			bEnumFound = ValueSelection != NAME_None;
 		}
 	}
 
 	// The selector is null or the enum is invalid, so default to the first option
 	if (!bEnumFound)
 	{
-		// Note: already verified this wasn't empty earlier
-		UpdateEnumClass(EnumClassSourceOptions[0]);
-		ClassSelection = EnumClassSourceOptions[0];
+		if (!ClassSelection)
+		{
+			// Note: already verified this wasn't empty earlier
+			UpdateEnumClass(EnumClassSourceOptions[0]);
+			ClassSelection = EnumClassSourceOptions[0];
+		}
 
 		// In the off chance that the enum has no options or only a MAX value
 		if (!EnumValueSourceOptions.IsEmpty())
