[PCG] Disallow changing partition actor labels. This is not a meaningful editing operation, and triggers a bug in runtime gen.

diff --git a/Engine/Plugins/Experimental/PCG/Source/PCG/Private/Grid/PCGPartitionActor.cpp b/Engine/Plugins/Experimental/PCG/Source/PCG/Private/Grid/PCGPartitionActor.cpp
--- a/Engine/Plugins/Experimental/PCG/Source/PCG/Private/Grid/PCGPartitionActor.cpp
+++ b/Engine/Plugins/Experimental/PCG/Source/PCG/Private/Grid/PCGPartitionActor.cpp
@@ -21,6 +21,10 @@ APCGPartitionActor::APCGPartitionActor(const FObjectInitializer& ObjectInitializ
 {
 	PCGGridSize = InvalidPCGGridSizeValue;
 
+#if WITH_EDITORONLY_DATA
+	bActorLabelEditable = false;
+#endif // WITH_EDITORONLY_DATA
+
 #if WITH_EDITOR
 	// Setup bounds component
 	BoundsComponent = CreateDefaultSubobject<UBoxComponent>(TEXT("BoundsComponent"));
