#UE Fix Crash in cooked builds due to serialization mismatch between cook and runtime

diff --git a/Engine/Plugins/Experimental/PCG/Source/PCG/Private/PCGComponent.cpp b/Engine/Plugins/Experimental/PCG/Source/PCG/Private/PCGComponent.cpp
--- a/Engine/Plugins/Experimental/PCG/Source/PCG/Private/PCGComponent.cpp
+++ b/Engine/Plugins/Experimental/PCG/Source/PCG/Private/PCGComponent.cpp
@@ -1230,18 +1230,21 @@ void UPCGComponent::Serialize(FArchive& Ar)
 
 	Super::Serialize(Ar);
 
-#if WITH_EDITOR
-	int32 DataVersion = FPCGCustomVersion::LatestVersion;
-	if (Ar.IsLoading())
+#if WITH_EDITORONLY_DATA
+	if (!Ar.IsCooking())
 	{
-		DataVersion = Ar.CustomVer(FPCGCustomVersion::GUID);
-	}
+		int32 DataVersion = FPCGCustomVersion::LatestVersion;
+		if (Ar.IsLoading())
+		{
+			DataVersion = Ar.CustomVer(FPCGCustomVersion::GUID);
+		}
 
-	if (DataVersion >= FPCGCustomVersion::DynamicTrackingKeysSerializedInComponent)
-	{
-		Ar << DynamicallyTrackedKeysToSettings;
+		if (DataVersion >= FPCGCustomVersion::DynamicTrackingKeysSerializedInComponent)
+		{
+			Ar << DynamicallyTrackedKeysToSettings;
+		}
 	}
-#endif // WITH_EDITOR
+#endif // WITH_EDITORONLY_DATA
 
 
 #if WITH_EDITOR
