Added match to nearest (& threshold istance) on the Match & Set.
Added comparision and distance traits on metadata attributes

diff --git a/Engine/Plugins/Experimental/PCG/Source/PCG/Private/Elements/PCGMatchAndSetAttributes.cpp b/Engine/Plugins/Experimental/PCG/Source/PCG/Private/Elements/PCGMatchAndSetAttributes.cpp
--- a/Engine/Plugins/Experimental/PCG/Source/PCG/Private/Elements/PCGMatchAndSetAttributes.cpp
+++ b/Engine/Plugins/Experimental/PCG/Source/PCG/Private/Elements/PCGMatchAndSetAttributes.cpp
@@ -11,8 +11,10 @@
 #include "Metadata/PCGMetadataAttributeTpl.h"
 #include "Metadata/PCGMetadataAttributeTraits.h"
 #include "Metadata/Accessors/IPCGAttributeAccessor.h"
+#include "Metadata/Accessors/IPCGAttributeAccessorTpl.h"
 #include "Metadata/Accessors/PCGAttributeAccessorHelpers.h"
 #include "Metadata/Accessors/PCGAttributeAccessorKeys.h"
+#include "Metadata/Accessors/PCGCustomAccessor.h"
 
 #include "Algo/AnyOf.h"
 
@@ -23,6 +25,7 @@
 namespace PCGMatchAndSetAttributesConstants
 {
 	const FName MatchDataLabel = TEXT("Match Data");
+	const FName MaxDistanceLabel = TEXT("Max Match Distance");
 }
 
 #if WITH_EDITOR
@@ -59,6 +62,16 @@ TArray<FPCGPinProperties> UPCGMatchAndSetAttributesSettings::InputPinProperties(
 		LOCTEXT("MatchDataTooltip", "Input containing the data to match to, then copy the accompanying attribute values")
 	);
 
+	if (bFindNearest && MaxDistanceMode == EPCGMatchMaxDistanceMode::AttributeMaxDistance)
+	{
+		PinProperties.Emplace(PCGMatchAndSetAttributesConstants::MaxDistanceLabel,
+			EPCGDataType::Point | EPCGDataType::Param,
+			/*bAllowMultipleConnections=*/false,
+			/*bAllowMultipleData=*/true,
+			LOCTEXT("MaxDistanceTooltip", "Input containing the maximum distance allowed for near search, selected by the Max Distance Attribute.")
+		);
+	}
+
 	return PinProperties;
 }
 
@@ -93,14 +106,17 @@ public:
 	};
 
 	FPCGAttributeSetPartition() = default;
-	FPCGAttributeSetPartition(FPCGContext* InContext, const UPCGParamData* InParamData, bool bPartitionByAttribute, FName AttributeName, bool bUseWeightAttribute, FName WeightAttributeName)
+
+	FPCGAttributeSetPartition(FPCGContext* InContext, const UPCGParamData* InParamData, bool bPartitionByAttribute, FName AttributeName, bool bUseWeightAttribute, FName WeightAttributeName, bool bInFindNearest, EPCGMatchMaxDistanceMode InMaxDistanceMode, const FPCGMetadataTypesConstantStruct* InMaxDistanceForNearestMatch)
 	{
-		Initialize(InContext, InParamData, bPartitionByAttribute, AttributeName, bUseWeightAttribute, WeightAttributeName);
+		Initialize(InContext, InParamData, bPartitionByAttribute, AttributeName, bUseWeightAttribute, WeightAttributeName, bInFindNearest, InMaxDistanceMode, InMaxDistanceForNearestMatch);
 	}
 
-	bool Initialize(FPCGContext* Context, const UPCGParamData* InParamData, bool bPartitionByAttribute, FName AttributeName, bool bUseWeightAttribute, FName WeightAttributeName)
+	bool Initialize(FPCGContext* Context, const UPCGParamData* InParamData, bool bPartitionByAttribute, FName AttributeName, bool bUseWeightAttribute, FName WeightAttributeName, bool bInFindNearest, EPCGMatchMaxDistanceMode InMaxDistanceMode, const FPCGMetadataTypesConstantStruct* InMaxDistanceForNearestMatch)
 	{
 		ParamData = InParamData;
+		bFindNearest = bInFindNearest;
+		MaxDistanceMode = InMaxDistanceMode;
 
 		if (!ParamData || !ParamData->Metadata)
 		{
@@ -138,6 +154,39 @@ public:
 			}
 		}
 
+		if (Attribute && MaxDistanceMode == EPCGMatchMaxDistanceMode::UseConstantMaxDistance)
+		{
+			auto ValidateAttributeSupportsDistance = [](auto AttributeDummyValue) -> bool
+			{
+				using AttributeType = decltype(AttributeDummyValue);
+				return PCG::Private::MetadataTraits<AttributeType>::CanComputeDistance;
+			};
+
+			if (!PCGMetadataAttribute::CallbackWithRightType(Attribute->GetTypeId(), ValidateAttributeSupportsDistance))
+			{
+				PCGE_LOG_C(Error, GraphAndLog, Context, FText::Format(LOCTEXT("AttributeDoesNotSupportDistance", "Attribute '{0}' does not support distance computation."), FText::FromName(Attribute->Name)));
+				return false;
+			}
+
+			if (InMaxDistanceForNearestMatch)
+			{
+				auto CreateConstantAttribute = [this](auto&& Value)
+				{
+					using ConstantType = std::decay_t<decltype(Value)>;
+					ConstantThreshold = MakeUnique<FPCGConstantValueAccessor<ConstantType>>(std::forward<decltype(Value)>(Value));
+					ConstantKey = MakeUnique<FPCGAttributeAccessorKeysSingleObjectPtr<void>>();
+				};
+
+				InMaxDistanceForNearestMatch->Dispatcher(CreateConstantAttribute);
+
+				if (!ConstantThreshold.IsValid() || !ConstantKey.IsValid())
+				{
+					PCGE_LOG_C(Error, GraphAndLog, Context, LOCTEXT("InvalidConstantThresholdAttribute", "Distance threshold is invalid."));
+					return false;
+				}
+			}
+		}
+
 		auto GetWeightFromAttribute = [WeightAttribute](PCGMetadataEntryKey EntryKey) -> double
 		{
 			double Weight = 1.0;
@@ -227,6 +276,59 @@ public:
 		return true;
 	}
 
+	bool InitializeForData(FPCGContext* Context, const UPCGData* InMaxDistanceData, const FPCGAttributePropertyInputSelector* InMaxDistanceSelector)
+	{
+		if (Attribute && MaxDistanceMode == EPCGMatchMaxDistanceMode::AttributeMaxDistance)
+		{
+			ConstantThreshold.Reset();
+			ConstantKey.Reset();
+
+			if (InMaxDistanceData && InMaxDistanceSelector)
+			{
+				ConstantThreshold = PCGAttributeAccessorHelpers::CreateConstAccessor(InMaxDistanceData, *InMaxDistanceSelector);
+				ConstantKey = PCGAttributeAccessorHelpers::CreateConstKeys(InMaxDistanceData, *InMaxDistanceSelector);
+
+				if (!ConstantThreshold.IsValid() || !ConstantKey.IsValid())
+				{
+					PCGE_LOG_C(Error, GraphAndLog, Context, FText::Format(LOCTEXT("InvalidThresholdAttribute", "Attribute '{0}' used for max distance is invalid."), FText::FromName(InMaxDistanceSelector->GetAttributeName())));
+					return false;
+				}
+			}
+			else
+			{
+				return false;
+			}
+		}
+
+		if (Attribute && ConstantThreshold.IsValid() && ConstantKey.IsValid())
+		{
+			auto ValidateCompatibleThresholdType = [this](auto AttributeDummyValue) -> bool
+			{
+				using AttributeType = decltype(AttributeDummyValue);
+
+				if constexpr (PCG::Private::MetadataTraits<AttributeType>::CanComputeDistance)
+				{
+					using DistanceType = PCG::Private::MetadataTraits<AttributeType>::DistanceType;
+					DistanceType ThresholdValue{};
+					check(ConstantThreshold && ConstantKey);
+					return ConstantThreshold->Get(ThresholdValue, *ConstantKey, EPCGAttributeAccessorFlags::AllowBroadcast | EPCGAttributeAccessorFlags::AllowConstructible);
+				}
+				else
+				{
+					return false;
+				}
+			};
+
+			if (!PCGMetadataAttribute::CallbackWithRightType(Attribute->GetTypeId(), ValidateCompatibleThresholdType))
+			{
+				PCGE_LOG_C(Error, GraphAndLog, Context, FText::Format(LOCTEXT("InvalidThresholdAttribute", "Distance threshold type is not compatible with attribute '{0}'."), FText::FromName(Attribute->Name)));
+				return false;
+			}
+		}
+
+		return true;
+	}
+
 	bool IsValid() const { return bIsValid; }
 
 	TArray<int32> GetMatchingPartitionDataIndices(const TUniquePtr<const IPCGAttributeAccessor>& InputAttribute, const TUniquePtr<const IPCGAttributeAccessorKeys>& InputKeys, int32 PointsNum) const
@@ -240,6 +342,27 @@ public:
 			{
 				using AttributeType = decltype(AttributeDummyValue);
 
+				// Get threshold value if we need it.
+				void* ThresholdValuesPtr = nullptr;
+				int ConstKeyCount = ConstantKey.IsValid() ? FMath::Max(1, ConstantKey->GetNum()) : 0;
+
+				if constexpr (PCG::Private::MetadataTraits<AttributeType>::CanComputeDistance)
+				{
+					if (ConstantThreshold.IsValid() && ConstantKey.IsValid())
+					{
+						using DistanceType = PCG::Private::MetadataTraits<AttributeType>::DistanceType;
+						DistanceType* TypedThresholdValues = new DistanceType[ConstKeyCount];
+						TArrayView<DistanceType> ThresholdValues(TypedThresholdValues, ConstKeyCount);
+						if (!ConstantThreshold->GetRange(ThresholdValues, 0, *ConstantKey, EPCGAttributeAccessorFlags::AllowBroadcast | EPCGAttributeAccessorFlags::AllowConstructible))
+						{
+							delete[] TypedThresholdValues;
+							return false;
+						}
+
+						ThresholdValuesPtr = TypedThresholdValues;
+					}
+				}
+
 				// Get the values to match against from the attribute
 				const FPCGMetadataAttribute<AttributeType>* TypedAttribute = static_cast<const FPCGMetadataAttribute<AttributeType>*>(Attribute);
 				TArray<AttributeType> AttributeValues;
@@ -250,20 +373,66 @@ public:
 					AttributeValues.Add(TypedAttribute->GetValue(PartitionEntry.Key));
 				}
 
-				return PCGMetadataElementCommon::ApplyOnAccessor<AttributeType>(*InputKeys, *InputAttribute, [&AttributeValues, &MatchingPartitionDataIndices](const AttributeType& InValue, int32)
+				bool bApplyOk = PCGMetadataElementCommon::ApplyOnAccessor<AttributeType>(*InputKeys, *InputAttribute, [this, &AttributeValues, &MatchingPartitionDataIndices, ThresholdValuesPtr, ConstKeyCount](const AttributeType& InValue, int32 InIndex)
 				{
 					int32 MatchingPartitionDataIndex = INDEX_NONE;
+					bool bFoundEqualMatch = false;
 					for (int32 AttributeValueIndex = 0; AttributeValueIndex < AttributeValues.Num(); ++AttributeValueIndex)
 					{
 						if (PCG::Private::MetadataTraits<AttributeType>::Equal(InValue, AttributeValues[AttributeValueIndex]))
 						{
 							MatchingPartitionDataIndex = AttributeValueIndex;
+							bFoundEqualMatch = true;
 							break;
 						}
+						else if (bFindNearest)
+						{
+							if constexpr (PCG::Private::MetadataTraits<AttributeType>::CanFindNearest)
+							{
+								if (MatchingPartitionDataIndex == INDEX_NONE || PCG::Private::MetadataTraits<AttributeType>::IsCloserTo(AttributeValues[AttributeValueIndex], AttributeValues[MatchingPartitionDataIndex], InValue))
+								{
+									MatchingPartitionDataIndex = AttributeValueIndex;
+								}
+							}
+						}
+					}
+
+					// Finally, if we haven't found an equal match, we should compare against the distance threshold.
+					if (!bFoundEqualMatch && ThresholdValuesPtr && ConstKeyCount > 0)
+					{
+						if constexpr (PCG::Private::MetadataTraits<AttributeType>::CanComputeDistance)
+						{
+							using DistanceType = PCG::Private::MetadataTraits<AttributeType>::DistanceType;
+
+							DistanceType Distance = PCG::Private::MetadataTraits<AttributeType>::Distance(AttributeValues[MatchingPartitionDataIndex], InValue);
+							const DistanceType& ThresholdValue = static_cast<DistanceType*>(ThresholdValuesPtr)[InIndex % ConstKeyCount];
+
+							if (Distance >= ThresholdValue)
+							{
+								MatchingPartitionDataIndex = INDEX_NONE;
+							}
+						}
+						else
+						{
+							MatchingPartitionDataIndex = INDEX_NONE;
+						}
 					}
 
 					MatchingPartitionDataIndices.Add(MatchingPartitionDataIndex);
 				}, EPCGAttributeAccessorFlags::AllowBroadcast | EPCGAttributeAccessorFlags::AllowConstructible);
+
+				// delete threshold value ptr
+				if constexpr (PCG::Private::MetadataTraits<AttributeType>::CanComputeDistance)
+				{
+					if (ThresholdValuesPtr)
+					{
+						using DistanceType = PCG::Private::MetadataTraits<AttributeType>::DistanceType;
+						delete[] static_cast<DistanceType*>(ThresholdValuesPtr);
+						ThresholdValuesPtr = nullptr;
+					}
+				}
+
+				return bApplyOk;
 			};
 
 			if (!PCGMetadataAttribute::CallbackWithRightType(Attribute->GetTypeId(), FindMatchingValueKeyIndex))
@@ -328,6 +497,11 @@ private:
 	const UPCGParamData* ParamData = nullptr;
 	const FPCGMetadataAttributeBase* Attribute = nullptr;
 	bool bIsValid = false;
+	bool bFindNearest = false;
+	EPCGMatchMaxDistanceMode MaxDistanceMode = EPCGMatchMaxDistanceMode::NoMaxDistance;
+
+	TUniquePtr<const IPCGAttributeAccessor> ConstantThreshold;
+	TUniquePtr<const IPCGAttributeAccessorKeys> ConstantKey;
 
 	TArray<TPair<PCGMetadataValueKey, AttributeSetPartitionEntry>> PartitionData;
 };
@@ -344,19 +518,24 @@ public:
 	{
 	}
 
-	bool Initialize()
+	bool Initialize(const TMap<const UPCGData*, const UPCGData*>& InDataToMaxDistanceMap)
 	{
 		if (!Settings)
 		{
 			return false;
 		}
 
+		DataToMaxDistanceMap = InDataToMaxDistanceMap;
+
 		AttributeSetPartition.Initialize(Context,
 			ParamData,
 			Settings->bMatchAttributes,
 			Settings->MatchAttribute,
 			Settings->bUseWeightAttribute,
-			Settings->WeightAttribute);
+			Settings->WeightAttribute,
+			Settings->bFindNearest,
+			Settings->MaxDistanceMode,
+			&Settings->MaxDistanceForNearestMatch);
 
 		return AttributeSetPartition.IsValid();
 	}
@@ -381,6 +560,14 @@ public:
 			}
 		}
 
+		const UPCGData** FoundMaxDistanceData = DataToMaxDistanceMap.Find(PointData);
+		const FPCGAttributePropertyInputSelector MaxDistanceSelector = Settings->MaxDistanceInputAttribute.CopyAndFixLast(FoundMaxDistanceData ? *FoundMaxDistanceData : nullptr);
+
+		if (!AttributeSetPartition.InitializeForData(Context, FoundMaxDistanceData ? *FoundMaxDistanceData : nullptr, FoundMaxDistanceData ? &MaxDistanceSelector : nullptr))
+		{
+			return false;
+		}
+
 		if (Settings->bUseInputWeightAttribute)
 		{
 			const FPCGAttributePropertyInputSelector InputWeightAttributeSource = Settings->InputWeightAttribute.CopyAndFixLast(PointData);
@@ -542,6 +729,7 @@ private:
 	TUniquePtr<const IPCGAttributeAccessor> InputWeightAccessor;
 	TUniquePtr<const IPCGAttributeAccessorKeys> InputAttributeKeys;
 	TArray<TPair<const FPCGMetadataAttributeBase*, FPCGMetadataAttributeBase*>> AttributesToSet;
+	TMap<const UPCGData*, const UPCGData*> DataToMaxDistanceMap;
 };
 
 FPCGMatchAndSetAttributesExecutionState::~FPCGMatchAndSetAttributesExecutionState()
@@ -566,7 +754,7 @@ bool FPCGMatchAndSetAttributesElement::PrepareDataInternal(FPCGContext* InContex
 	TArray<FPCGTaggedData>& Outputs = InContext->OutputData.TaggedData;
 	TArray<FPCGTaggedData> ParamDataInputs = InContext->InputData.GetInputsByPin(PCGMatchAndSetAttributesConstants::MatchDataLabel);
 
-	EPCGTimeSliceInitResult InitResult = TimeSlicedContext->InitializePerExecutionState([Settings, &ParamDataInputs](FPCGMatchAndSetAttributesElement::ContextType* Context, FPCGMatchAndSetAttributesExecutionState& OutState) -> EPCGTimeSliceInitResult
+	EPCGTimeSliceInitResult InitResult = TimeSlicedContext->InitializePerExecutionState([Settings, &ParamDataInputs, &Inputs](FPCGMatchAndSetAttributesElement::ContextType* Context, FPCGMatchAndSetAttributesExecutionState& OutState) -> EPCGTimeSliceInitResult
 	{
 		const UPCGParamData* ParamData = nullptr;
 
@@ -585,10 +773,37 @@ bool FPCGMatchAndSetAttributesElement::PrepareDataInternal(FPCGContext* InContex
 			return EPCGTimeSliceInitResult::NoOperation;
 		}
 
+		// If there are provided max distance entries, we should have either 1 or the same cardinality as the inputs
+		TMap<const UPCGData*, const UPCGData*> InputToMaxDistanceMapping;
+		if (Settings->MaxDistanceMode == EPCGMatchMaxDistanceMode::AttributeMaxDistance)
+		{
+			TArray<FPCGTaggedData> InputMaxDistanceData = Context->InputData.GetInputsByPin(PCGMatchAndSetAttributesConstants::MaxDistanceLabel);
+
+			if (InputMaxDistanceData.Num() == Inputs.Num())
+			{
+				for (int DataIndex = 0; DataIndex < Inputs.Num(); ++DataIndex)
+				{
+					InputToMaxDistanceMapping.Add(Inputs[DataIndex].Data, InputMaxDistanceData[DataIndex].Data);
+				}
+			}
+			else if (InputMaxDistanceData.Num() == 1)
+			{
+				for (int DataIndex = 0; DataIndex < Inputs.Num(); ++DataIndex)
+				{
+					InputToMaxDistanceMapping.Add(Inputs[DataIndex].Data, InputMaxDistanceData[0].Data);
+				}
+			}
+			else
+			{
+				PCGE_LOG_C(Error, GraphAndLog, Context, FText::Format(LOCTEXT("NoMatchingMaxDistanceData", "Invalid number of max distance providers for data; expected {0}, got {1}."), FText::AsNumber(Inputs.Num()), FText::AsNumber(InputMaxDistanceData.Num())));
+				return EPCGTimeSliceInitResult::AbortExecution;
+			}
+		}
+
 		OutState.Partition = new FPCGMatchAndSetPartition(Context, Settings, Context->SourceComponent.Get(), ParamData);
 		check(OutState.Partition);
 
-		if (OutState.Partition->Initialize())
+		if (OutState.Partition->Initialize(InputToMaxDistanceMapping))
 		{
 			return EPCGTimeSliceInitResult::Success;
 		}
diff --git a/Engine/Plugins/Experimental/PCG/Source/PCG/Public/Elements/PCGMatchAndSetAttributes.h b/Engine/Plugins/Experimental/PCG/Source/PCG/Public/Elements/PCGMatchAndSetAttributes.h
--- a/Engine/Plugins/Experimental/PCG/Source/PCG/Public/Elements/PCGMatchAndSetAttributes.h
+++ b/Engine/Plugins/Experimental/PCG/Source/PCG/Public/Elements/PCGMatchAndSetAttributes.h
@@ -5,6 +5,7 @@
 #include "PCGData.h"
 #include "PCGSettings.h"
 #include "Elements/PCGTimeSlicedElementBase.h"
+#include "Metadata/PCGMetadataTypesConstantStruct.h"
 
 #include "PCGMatchAndSetAttributes.generated.h"
 
@@ -26,6 +27,14 @@ struct FPCGMatchAndSetAttributesIterationState
 	UPCGPointData* OutPointData = nullptr;
 };
 
+UENUM()
+enum class EPCGMatchMaxDistanceMode
+{
+	NoMaxDistance UMETA(DisplayName="No maximum distance"),
+	UseConstantMaxDistance UMETA(DisplayName="Use constant maximum distance"),
+	AttributeMaxDistance
+};
+
 /** This class creates a PCG node that can match, select by weight or match & select by weight 
 * a 'matching' entry in a provided Attribute Set with multiple entries.
 * E.g. for a given point, if the point has the same specified attribute as the matching attribute in the attribute set,
@@ -54,38 +63,54 @@ protected:
 	// ~End UPCGSettings interface
 
 public:
-	/** Controls whether selection of the attribute set values to copy will be done by matching point-to-attribute set (true) or done randomly (false) */
+	/** Controls whether selection of the attribute set values to copy will be done by matching point-to-attribute set (true) or done randomly (false). */
 	UPROPERTY(BlueprintReadWrite, EditAnywhere, Category = Settings)
 	bool bMatchAttributes = false;
 
-	/** Attribute from the point data to select & maetch */
+	/** Attribute from the point data to select & match. */
 	UPROPERTY(BlueprintReadWrite, EditAnywhere, Category = Settings, meta = (EditCondition = "bMatchAttributes", PCG_Overridable))
 	FPCGAttributePropertyInputSelector InputAttribute;
 
-	/** Attribute from the attribute set to match against */
+	/** Attribute from the attribute set to match against. */
 	UPROPERTY(BlueprintReadWrite, EditAnywhere, Category = Settings, meta = (EditCondition = "bMatchAttributes"))
 	FName MatchAttribute = NAME_None;
 
-	/** Controls whether points that have no valid match in the attribute set are kept as is (default values) or removed from the output */
+	/** Controls whether points that have no valid match in the attribute set are kept as is (default values) or removed from the output. */
 	UPROPERTY(BlueprintReadWrite, EditAnywhere, Category = Settings, meta = (EditCondition = "bMatchAttributes", PCG_Overridable))
 	bool bKeepUnmatched = true;
 
+	/** Controls whether the match operation will return the nearest match and not only match on equality. */
+	UPROPERTY(BlueprintReadWrite, EditAnywhere, Category = Settings, meta = (EditCondition = "bMatchAttributes", PCG_Overridable))
+	bool bFindNearest = false;
+
+	/** Controls whether the match operation has a maximum distance on which to reject points that would be too far from the nearest value. */
+	UPROPERTY(BlueprintReadWrite, EditAnywhere, Category = Settings, meta = (EditCondition = "bFindNearest", EditConditionHides, PCG_Overridable))
+	EPCGMatchMaxDistanceMode MaxDistanceMode = EPCGMatchMaxDistanceMode::NoMaxDistance;
+
+	/** Constant value that establishes the maximum distance an entry can be from its nearest match to be selected */
+	UPROPERTY(BlueprintReadWrite, EditAnywhere, Category = Settings, meta = (EditCondition = "bFindNearest && MaxDistanceMode==EPCGMatchMaxDistanceMode::UseConstantMaxDistance", EditConditionHides))
+	FPCGMetadataTypesConstantStruct MaxDistanceForNearestMatch;
+
+	UPROPERTY(BlueprintReadWrite, EditAnywhere, Category = Settings, meta = (EditCondition = "bFindNearest && MaxDistanceMode==EPCGMatchMaxDistanceMode::AttributeMaxDistance", EditConditionHides, PCG_Overridable))
+	FPCGAttributePropertyInputSelector MaxDistanceInputAttribute;
+
+	/** Controls whether we will use the attribute provided in the Input Weight Attribute to perform entry selection. */
 	UPROPERTY(BlueprintReadWrite, EditAnywhere, Category = Settings)
 	bool bUseInputWeightAttribute = false;
 
-	/** Input weight from the points, assumed to be in the [0, 1] range */
+	/** Input weight from the points, assumed to be in the [0, 1] range. */
 	UPROPERTY(BlueprintReadWrite, EditAnywhere, Category = Settings, meta = (EditCondition = "bUseInputWeightAttribute", PCG_Overridable))
 	FPCGAttributePropertyInputSelector InputWeightAttribute;
 
-	/** Controls whether we will consider the weights, as determined by the Weight Attribute values on the attribute set */
+	/** Controls whether we will consider the weights, as determined by the Weight Attribute values on the attribute set. */
 	UPROPERTY(BlueprintReadWrite, EditAnywhere, Category = Settings, meta = (DisplayName = "Use Match Weight", InlineEditConditionToggle, PCG_Overridable))
 	bool bUseWeightAttribute = false;
 
-	/** Attribute to weight more or less some entries from the attribute set */
+	/** Attribute to weight more or less some entries from the attribute set. */
 	UPROPERTY(BlueprintReadWrite, EditAnywhere, Category = Settings, meta = (DisplayName="Match Weight Attribute", EditCondition = "bUseWeightAttribute", PCG_Overridable))
 	FName WeightAttribute = NAME_None;
 
-	/** Controls whether we will emit a warning and return nothing if there is no provided attribute set */
+	/** Controls whether we will emit a warning and return nothing if there is no provided attribute set. */
 	UPROPERTY(BlueprintReadWrite, EditAnywhere, Category = Settings)
 	bool bWarnIfNoMatchData = true;
 };
diff --git a/Engine/Plugins/Experimental/PCG/Source/PCG/Public/Metadata/PCGMetadataAttributeTraits.h b/Engine/Plugins/Experimental/PCG/Source/PCG/Public/Metadata/PCGMetadataAttributeTraits.h
--- a/Engine/Plugins/Experimental/PCG/Source/PCG/Public/Metadata/PCGMetadataAttributeTraits.h
+++ b/Engine/Plugins/Experimental/PCG/Source/PCG/Public/Metadata/PCGMetadataAttributeTraits.h
@@ -408,6 +408,33 @@ namespace PCG
 			}
 		};
 
+		/** "Find Nearest" trait support driven by less/less or equal operation, used for types that do not support distance computation */
+		template<typename Traits, typename T>
+		struct CompareBasedFindNearestTraits
+		{
+			enum { CanFindNearest = true};
+			enum { CanComputeDistance = false };
+
+			static bool IsCloserTo(const T& A, const T& B, const T& Target)
+			{
+				// If A is "between" B and Target, then it is closer, e.g. B < A <= Target or Target <= A < B
+				return (Traits::Less(B, A) && Traits::LessOrEqual(A, Target)) || (Traits::LessOrEqual(Target, A) && Traits::Less(A, B));
+			}
+		};
+
+		/** "Find Nearest" trait support driven by distance comparison, used for types that support distance computation */
+		template<typename Traits, typename T>
+		struct DistanceBasedFindNearestTraits
+		{
+			enum { CanFindNearest = true };
+			enum { CanComputeDistance = true };
+
+			static bool IsCloserTo(const T& A, const T& B, const T& Target)
+			{
+				return Traits::Distance(A, Target) < Traits::Distance(B, Target);
+			}
+		};
+
 		template<typename T>
 		struct DefaultWeightedSumTraits
 		{
@@ -440,9 +467,20 @@ namespace PCG
 			}
 		};
 
+		template <typename T>
+		struct DefaultScalarDistanceTraits : DistanceBasedFindNearestTraits<DefaultScalarDistanceTraits<T>, T>
+		{
+			using DistanceType = T;
+
+			static DistanceType Distance(const T& A, const T& B) 
+			{ 
+				return FMath::Abs(A - B);
+			}
+		};
+
 		// Common traits for int32, int64, float, double
 		template<typename T>
-		struct MetadataTraits : DefaultOperationTraits<T>, DefaultWeightedSumTraits<T>, DefaultMinMaxTraits<T>, DefaultCompareTraits<T>, LexToStringTraits<T>
+		struct MetadataTraits : DefaultOperationTraits<T>, DefaultWeightedSumTraits<T>, DefaultMinMaxTraits<T>, DefaultCompareTraits<T>, LexToStringTraits<T>, DefaultScalarDistanceTraits<T>
 		{
 			enum { CompressData = false };
 			enum { CanSearchString = false };
@@ -450,7 +488,7 @@ namespace PCG
 		};
 
 		template<>
-		struct MetadataTraits<float> : DefaultOperationTraits<float>, DefaultWeightedSumTraits<float>, DefaultMinMaxTraits<float>, DefaultCompareTraits<float>, LexToStringTraits<float>
+		struct MetadataTraits<float> : DefaultOperationTraits<float>, DefaultWeightedSumTraits<float>, DefaultMinMaxTraits<float>, DefaultCompareTraits<float>, LexToStringTraits<float>, DefaultScalarDistanceTraits<float>
 		{
 			enum { CompressData = false };
 			enum { CanSearchString = false };
@@ -463,7 +501,7 @@ namespace PCG
 		};
 
 		template<>
-		struct MetadataTraits<double> : DefaultOperationTraits<double>, DefaultWeightedSumTraits<double>, DefaultMinMaxTraits<double>, DefaultCompareTraits<double>, LexToStringTraits<double>
+		struct MetadataTraits<double> : DefaultOperationTraits<double>, DefaultWeightedSumTraits<double>, DefaultMinMaxTraits<double>, DefaultCompareTraits<double>, LexToStringTraits<double>, DefaultScalarDistanceTraits<double>
 		{
 			enum { CompressData = false };
 			enum { CanSearchString = false };
@@ -484,6 +522,8 @@ namespace PCG
 			enum { CanMulDiv = false };
 			enum { CanInterpolate = false };
 			enum { CanCompare = true };
+			enum { CanFindNearest = false };
+			enum { CanComputeDistance = false };
 			enum { CanSearchString = false };
 			enum { NeedsConstruction = false };
 
@@ -540,7 +580,7 @@ namespace PCG
 
 		// Vector types
 		template<typename T>
-		struct VectorTraits : DefaultOperationTraits<T>, DefaultWeightedSumTraits<T>, DefaultStringTraits<T>
+		struct VectorTraits : DefaultOperationTraits<T>, DefaultWeightedSumTraits<T>, DefaultStringTraits<T>, DistanceBasedFindNearestTraits<VectorTraits<T>, T>
 		{
 			enum { CompressData = false };
 			enum { CanMinMax = true };
@@ -549,10 +589,18 @@ namespace PCG
 			enum { CanSearchString = false };
 			enum { NeedsConstruction = false };
 
+			using DistanceType = typename T::FReal;
+
 			static T ZeroValue()
 			{
 				return T::Zero();
 			}
+
+			static DistanceType Distance(const T& A, const T& B) 
+			{
+				// Implementation note: we use SizeSquared here instead of SquareLength because FVector4 doesn't have it.
+				return (A - B).Size();
+			}
 		};
 
 		template<>
@@ -674,7 +722,7 @@ namespace PCG
 
 		// Quaternion
 		template<>
-		struct MetadataTraits<FQuat> : DefaultStringTraits<FQuat>
+		struct MetadataTraits<FQuat> : DefaultStringTraits<FQuat>, DistanceBasedFindNearestTraits<MetadataTraits<FQuat>, FQuat>
 		{
 			enum { CompressData = false };
 			enum { CanMinMax = false };
@@ -685,6 +733,8 @@ namespace PCG
 			enum { CanSearchString = false };
 			enum { NeedsConstruction = false };
 
+			using DistanceType = typename FQuat::FReal;
+
 			static bool Equal(const FQuat& A, const FQuat& B)
 			{
 				return A.Equals(B);
@@ -725,11 +775,16 @@ namespace PCG
 			{
 				return FQuat::Identity;
 			}
+
+			static FQuat::FReal Distance(const FQuat& A, const FQuat& B)
+			{
+				return A.AngularDistance(B);
+			}
 		};
 
 		// Rotator
 		template<>
-		struct MetadataTraits<FRotator> : DefaultStringTraits<FRotator>
+		struct MetadataTraits<FRotator> : DefaultStringTraits<FRotator>, DistanceBasedFindNearestTraits<MetadataTraits<FRotator>, FRotator>
 		{
 			enum { CompressData = false };
 			enum { CanMinMax = false };
@@ -740,6 +795,8 @@ namespace PCG
 			enum { CanSearchString = false };
 			enum { NeedsConstruction = false };
 
+			using DistanceType = typename FRotator::FReal;
+
 			static bool Equal(const FRotator& A, const FRotator& B)
 			{
 				return A.Equals(B);
@@ -775,6 +832,12 @@ namespace PCG
 			{
 				return FRotator::ZeroRotator;
 			}
+
+			static FRotator::FReal Distance(const FRotator& A, const FRotator& B)
+			{
+				// Use quaternions for distance calculation, but return degrees instead of radians here
+				return 180.0 * FQuat(A).AngularDistance(FQuat(B)) / UE_DOUBLE_PI;
+			}
 		};
 
 		// Transform
@@ -787,6 +850,8 @@ namespace PCG
 			enum { CanMulDiv = true };
 			enum { CanInterpolate = true };
 			enum { CanCompare = false };
+			enum { CanFindNearest = false };
+			enum { CanComputeDistance = false };
 			enum { CanSearchString = false };
 			enum { NeedsConstruction = false };
 
@@ -832,7 +897,7 @@ namespace PCG
 
 		// Strings
 		template<>
-		struct MetadataTraits<FString> : DefaultCompareTraits<FString>, DefaultStringTraits<FString>
+		struct MetadataTraits<FString> : DefaultCompareTraits<FString>, DefaultStringTraits<FString>, CompareBasedFindNearestTraits<MetadataTraits<FString>, FString>
 		{
 			enum { CompressData = true };
 			enum { CanMinMax = false };
@@ -864,7 +929,7 @@ namespace PCG
 		};
 
 		template<>
-		struct MetadataTraits<FName> : DefaultStringTraits<FName>
+		struct MetadataTraits<FName> : DefaultStringTraits<FName>, CompareBasedFindNearestTraits<MetadataTraits<FName>, FName>
 		{
 			enum { CompressData = false };
 			enum { CanMinMax = false };
@@ -917,7 +982,7 @@ namespace PCG
 		};
 
 		template<typename T>
-		struct SoftObjectPathTraits
+		struct SoftObjectPathTraits : CompareBasedFindNearestTraits<SoftObjectPathTraits<T>, T>
 		{
 			enum { CompressData = true };
 			enum { CanMinMax = false };
