Added missing typenames

diff --git a/Engine/Plugins/Experimental/PCG/Source/PCG/Private/Elements/PCGMatchAndSetAttributes.cpp b/Engine/Plugins/Experimental/PCG/Source/PCG/Private/Elements/PCGMatchAndSetAttributes.cpp
--- a/Engine/Plugins/Experimental/PCG/Source/PCG/Private/Elements/PCGMatchAndSetAttributes.cpp
+++ b/Engine/Plugins/Experimental/PCG/Source/PCG/Private/Elements/PCGMatchAndSetAttributes.cpp
@@ -308,7 +308,7 @@ public:
 
 				if constexpr (PCG::Private::MetadataTraits<AttributeType>::CanComputeDistance)
 				{
-					using DistanceType = PCG::Private::MetadataTraits<AttributeType>::DistanceType;
+					using DistanceType = typename PCG::Private::MetadataTraits<AttributeType>::DistanceType;
 					DistanceType ThresholdValue{};
 					check(ConstantThreshold && ConstantKey);
 					return ConstantThreshold->Get(ThresholdValue, *ConstantKey, EPCGAttributeAccessorFlags::AllowBroadcast | EPCGAttributeAccessorFlags::AllowConstructible);
@@ -350,7 +350,7 @@ public:
 				{
 					if (ConstantThreshold.IsValid() && ConstantKey.IsValid())
 					{
-						using DistanceType = PCG::Private::MetadataTraits<AttributeType>::DistanceType;
+						using DistanceType = typename PCG::Private::MetadataTraits<AttributeType>::DistanceType;
 						DistanceType* TypedThresholdValues = new DistanceType[ConstKeyCount];
 						TArrayView<DistanceType> ThresholdValues(TypedThresholdValues, ConstKeyCount);
 						if (!ConstantThreshold->GetRange(ThresholdValues, 0, *ConstantKey, EPCGAttributeAccessorFlags::AllowBroadcast | EPCGAttributeAccessorFlags::AllowConstructible))
@@ -402,7 +402,7 @@ public:
 					{
 						if constexpr (PCG::Private::MetadataTraits<AttributeType>::CanComputeDistance)
 						{
-							using DistanceType = PCG::Private::MetadataTraits<AttributeType>::DistanceType;
+							using DistanceType = typename PCG::Private::MetadataTraits<AttributeType>::DistanceType;
 
 							DistanceType Distance = PCG::Private::MetadataTraits<AttributeType>::Distance(AttributeValues[MatchingPartitionDataIndex], InValue);
 							const DistanceType& ThresholdValue = static_cast<DistanceType*>(ThresholdValuesPtr)[InIndex % ConstKeyCount];
@@ -426,7 +426,7 @@ public:
 				{
 					if (ThresholdValuesPtr)
 					{
-						using DistanceType = PCG::Private::MetadataTraits<AttributeType>::DistanceType;
+						using DistanceType = typename PCG::Private::MetadataTraits<AttributeType>::DistanceType;
 						delete[] static_cast<DistanceType*>(ThresholdValuesPtr);
 						ThresholdValuesPtr = nullptr;
 					}
