Adjusted some code in the match and set

diff --git a/Engine/Plugins/Experimental/PCG/Source/PCG/Private/Elements/PCGMatchAndSetAttributes.cpp b/Engine/Plugins/Experimental/PCG/Source/PCG/Private/Elements/PCGMatchAndSetAttributes.cpp
--- a/Engine/Plugins/Experimental/PCG/Source/PCG/Private/Elements/PCGMatchAndSetAttributes.cpp
+++ b/Engine/Plugins/Experimental/PCG/Source/PCG/Private/Elements/PCGMatchAndSetAttributes.cpp
@@ -68,7 +68,7 @@ TArray<FPCGPinProperties> UPCGMatchAndSetAttributesSettings::InputPinProperties(
 			EPCGDataType::Point | EPCGDataType::Param,
 			/*bAllowMultipleConnections=*/false,
 			/*bAllowMultipleData=*/true,
-			LOCTEXT("MaxDistanceTooltip", "Input containing the maximum distance allowed for near search, selected by the Max Distance Attribute.")
+			LOCTEXT("MaxDistanceTooltip", "Input containing the maximum distance allowed for nearest search, selected by the Max Distance Attribute.")
 		);
 	}
 
@@ -302,26 +302,21 @@ public:
 
 		if (Attribute && ConstantThreshold.IsValid() && ConstantKey.IsValid())
 		{
-			auto ValidateCompatibleThresholdType = [this](auto AttributeDummyValue) -> bool
+			if (!PCG::Private::IsBroadcastableOrConstructible(ConstantThreshold->GetUnderlyingType(), Attribute->GetTypeId()))
 			{
-				using AttributeType = decltype(AttributeDummyValue);
+				PCGE_LOG_C(Error, GraphAndLog, Context, FText::Format(LOCTEXT("InvalidThresholdAttribute", "Distance threshold type is not compatible with attribute '{0}'."), FText::FromName(Attribute->Name)));
+				return false;
+			}
 
-				if constexpr (PCG::Private::MetadataTraits<AttributeType>::CanComputeDistance)
-				{
-					using DistanceType = typename PCG::Private::MetadataTraits<AttributeType>::DistanceType;
-					DistanceType ThresholdValue{};
-					check(ConstantThreshold && ConstantKey);
-					return ConstantThreshold->Get(ThresholdValue, *ConstantKey, EPCGAttributeAccessorFlags::AllowBroadcast | EPCGAttributeAccessorFlags::AllowConstructible);
-				}
-				else
-				{
-					return false;
-				}
+			auto ValidateTypeCanComputeDistance = [](auto AttributeDummyValue) -> bool
+			{
+				using AttributeType = decltype(AttributeDummyValue);
+				return PCG::Private::MetadataTraits<AttributeType>::CanComputeDistance;
 			};
 
-			if (!PCGMetadataAttribute::CallbackWithRightType(Attribute->GetTypeId(), ValidateCompatibleThresholdType))
+			if (!PCGMetadataAttribute::CallbackWithRightType(Attribute->GetTypeId(), ValidateTypeCanComputeDistance))
 			{
-				PCGE_LOG_C(Error, GraphAndLog, Context, FText::Format(LOCTEXT("InvalidThresholdAttribute", "Distance threshold type is not compatible with attribute '{0}'."), FText::FromName(Attribute->Name)));
+				PCGE_LOG_C(Error, GraphAndLog, Context, FText::Format(LOCTEXT("AttributeDoesNotSupportMaxDistance", "The selected attribute '{0}' does not support computing distances."), FText::FromName(Attribute->Name)));
 				return false;
 			}
 		}
@@ -352,10 +347,18 @@ public:
 					{
 						using DistanceType = typename PCG::Private::MetadataTraits<AttributeType>::DistanceType;
 						DistanceType* TypedThresholdValues = new DistanceType[ConstKeyCount];
-						TArrayView<DistanceType> ThresholdValues(TypedThresholdValues, ConstKeyCount);
-						if (!ConstantThreshold->GetRange(ThresholdValues, 0, *ConstantKey, EPCGAttributeAccessorFlags::AllowBroadcast | EPCGAttributeAccessorFlags::AllowConstructible))
+						if (TypedThresholdValues)
+						{
+							TArrayView<DistanceType> ThresholdValues(TypedThresholdValues, ConstKeyCount);
+							if (!ConstantThreshold->GetRange(ThresholdValues, 0, *ConstantKey, EPCGAttributeAccessorFlags::AllowBroadcast | EPCGAttributeAccessorFlags::AllowConstructible))
+							{
+								delete[] TypedThresholdValues;
+								return false;
+							}
+						}
+						else
 						{
-							delete[] TypedThresholdValues;
+							// Couldn't allocate memory, exit early
 							return false;
 						}
 
diff --git a/Engine/Plugins/Experimental/PCG/Source/PCG/Public/Metadata/PCGMetadataAttributeTraits.h b/Engine/Plugins/Experimental/PCG/Source/PCG/Public/Metadata/PCGMetadataAttributeTraits.h
--- a/Engine/Plugins/Experimental/PCG/Source/PCG/Public/Metadata/PCGMetadataAttributeTraits.h
+++ b/Engine/Plugins/Experimental/PCG/Source/PCG/Public/Metadata/PCGMetadataAttributeTraits.h
@@ -415,6 +415,7 @@ namespace PCG
 			enum { CanFindNearest = true};
 			enum { CanComputeDistance = false };
 
+			/** Returns true if A is strictly closer to Target than B is. */
 			static bool IsCloserTo(const T& A, const T& B, const T& Target)
 			{
 				// If A is "between" B and Target, then it is closer, e.g. B < A <= Target or Target <= A < B
@@ -429,6 +430,7 @@ namespace PCG
 			enum { CanFindNearest = true };
 			enum { CanComputeDistance = true };
 
+			/** Returns true if A is strictly closer to Target than B is. */
 			static bool IsCloserTo(const T& A, const T& B, const T& Target)
 			{
 				return Traits::Distance(A, Target) < Traits::Distance(B, Target);
@@ -580,7 +582,7 @@ namespace PCG
 
 		// Vector types
 		template<typename T>
-		struct VectorTraits : DefaultOperationTraits<T>, DefaultWeightedSumTraits<T>, DefaultStringTraits<T>, DistanceBasedFindNearestTraits<VectorTraits<T>, T>
+		struct VectorTraits : DefaultOperationTraits<T>, DefaultWeightedSumTraits<T>, DefaultStringTraits<T>
 		{
 			enum { CompressData = false };
 			enum { CanMinMax = true };
@@ -588,6 +590,8 @@ namespace PCG
 			enum { CanCompare = true };
 			enum { CanSearchString = false };
 			enum { NeedsConstruction = false };
+			enum { CanFindNearest = true };
+			enum { CanComputeDistance = true };
 
 			using DistanceType = typename T::FReal;
 
@@ -598,9 +602,14 @@ namespace PCG
 
 			static DistanceType Distance(const T& A, const T& B) 
 			{
-				// Implementation note: we use SizeSquared here instead of SquareLength because FVector4 doesn't have it.
 				return (A - B).Size();
 			}
+
+			/** Returns true if A is strictly closer to Target than B is. Implementation note: we don't use the standard distance-based traits to have a more efficient check here. */
+			static bool IsCloserTo(const T& A, const T& B, const T& Target)
+			{
+				return (A - Target).SizeSquared() < (B - Target).SizeSquared();
+			}
 		};
 
 		template<>
