[PCG] Add support for instanced objects override

#rnx

diff --git a/Engine/Plugins/Experimental/PCG/Source/PCG/Private/Helpers/PCGSettingsHelpers.cpp b/Engine/Plugins/Experimental/PCG/Source/PCG/Private/Helpers/PCGSettingsHelpers.cpp
--- a/Engine/Plugins/Experimental/PCG/Source/PCG/Private/Helpers/PCGSettingsHelpers.cpp
+++ b/Engine/Plugins/Experimental/PCG/Source/PCG/Private/Helpers/PCGSettingsHelpers.cpp
@@ -273,7 +273,8 @@ namespace PCGSettingsHelpers
 		}
 	}
 
-	TArray<FPCGSettingsOverridableParam> GetAllOverridableParams(const UStruct* InClass, const FPCGGetAllOverridableParamsConfig& InConfig)
+	template <uint32 N>
+	TArray<FPCGSettingsOverridableParam> GetAllOverridableParams_Internal(const UStruct* InClass, const FPCGGetAllOverridableParamsConfig& InConfig, TArray<const FProperty*, TInlineAllocator<N>>& InAlreadySeenProperties)
 	{
 		TArray<FName> LabelCache;
 
@@ -374,7 +375,7 @@ namespace PCGSettingsHelpers
 				continue;
 			}
 
-			auto RecursiveExtraction = [&InConfig, &Res, &LabelCache, Property, InClass, &GatherAliases](const UStruct* NextClass)
+			auto RecursiveExtraction = [&InConfig, &Res, &LabelCache, Property, InClass, &GatherAliases, &InAlreadySeenProperties](const UStruct* NextClass)
 			{
 				// Reached max depth
 				if (InConfig.MaxStructDepth == 0)
@@ -384,6 +385,7 @@ namespace PCGSettingsHelpers
 
 				// Use the seed, and don't check metadata for PCG overridable.
 				FPCGGetAllOverridableParamsConfig RecurseConfig = InConfig;
+				InAlreadySeenProperties.Add(Property);
 				RecurseConfig.bUseSeed = true;
 #if WITH_EDITOR
 				RecurseConfig.IncludeMetadataValues.Remove(PCGObjectMetadata::Overridable);
@@ -394,7 +396,7 @@ namespace PCGSettingsHelpers
 					RecurseConfig.MaxStructDepth--;
 				}
 
-				for (FPCGSettingsOverridableParam& ChildParam : GetAllOverridableParams(NextClass, RecurseConfig))
+				for (FPCGSettingsOverridableParam& ChildParam : GetAllOverridableParams_Internal(NextClass, RecurseConfig, InAlreadySeenProperties))
 				{
 					FName Label = ChildParam.Label;
 					bool bHasNameClash = false;
@@ -439,6 +441,8 @@ namespace PCGSettingsHelpers
 					}
 #endif // WITH_EDITOR
 				}
+
+				InAlreadySeenProperties.Pop(/*bAllowShrinking=*/false);
 			};
 
 			const FProperty* PropertyToCheck = Property;
@@ -447,7 +451,8 @@ namespace PCGSettingsHelpers
 				PropertyToCheck = CastFieldChecked<FArrayProperty>(Property)->Inner;
 			}
 
-			if (PropertyToCheck->IsA<FObjectProperty>() && InConfig.bExtractObjects)
+			// Always extract instanced objects, but be careful with infinite recursion (class having an instanced property on itself)
+			if (PropertyToCheck->IsA<FObjectProperty>() && !InAlreadySeenProperties.Contains(PropertyToCheck) && (InConfig.bExtractObjects || PropertyToCheck->HasAllPropertyFlags(CPF_InstancedReference)))
 			{
 				RecursiveExtraction(CastFieldChecked<FObjectProperty>(PropertyToCheck)->PropertyClass);
 			}
@@ -489,4 +494,13 @@ namespace PCGSettingsHelpers
 
 		return Res;
 	}
+
+	TArray<FPCGSettingsOverridableParam> GetAllOverridableParams(const UStruct* InClass, const FPCGGetAllOverridableParamsConfig& InConfig)
+	{
+		constexpr uint32 InlineAllocator = 16;
+		// Array of already seen properties to catch infinite recursion
+		TArray<const FProperty*, TInlineAllocator<InlineAllocator>> AlreadySeenProperties;
+
+		return GetAllOverridableParams_Internal(InClass, InConfig, AlreadySeenProperties);
+	}
 }
diff --git a/Engine/Plugins/Experimental/PCG/Source/PCG/Private/PCGContext.cpp b/Engine/Plugins/Experimental/PCG/Source/PCG/Private/PCGContext.cpp
--- a/Engine/Plugins/Experimental/PCG/Source/PCG/Private/PCGContext.cpp
+++ b/Engine/Plugins/Experimental/PCG/Source/PCG/Private/PCGContext.cpp
@@ -210,7 +210,7 @@ void FPCGContext::OverrideSettings()
 			PCGE_LOG_C(Warning, GraphAndLog, this, FText::Format(LOCTEXT("OverrideWithAlias", "Attribute '{0}' was not found, but one of its deprecated aliases ('{1}') was. Please update the name to the new value."), FText::FromName(AttributeName), FText::FromName(AccessorResult.AliasUsed)));
 		}
 
-		TUniquePtr<IPCGAttributeAccessor> PropertyAccessor = PCGAttributeAccessorHelpers::CreatePropertyAccessor(Param.Properties.Last());
+		TUniquePtr<IPCGAttributeAccessor> PropertyAccessor = PCGAttributeAccessorHelpers::CreatePropertyChainAccessor(TArray<const FProperty*>(Param.Properties));
 		check(PropertyAccessor.IsValid());
 
 		const bool bParamOverridden = PCGMetadataAttribute::CallbackWithRightType(PropertyAccessor->GetUnderlyingType(), [this, &AttributeAccessor, &PropertyAccessor, &Param, &AttributeName, Container](auto Dummy) -> bool
@@ -226,15 +226,8 @@ void FPCGContext::OverrideSettings()
 				PCGE_LOG_C(Warning, GraphAndLog, this, FText::Format(LOCTEXT("ConversionFailed", "Parameter '{0}' cannot be converted from attribute '{1}'"), FText::FromName(Param.Label), FText::FromName(AttributeName)));
 				return false;
 			}
-			// TODO: Perhaps factorise this code in another property accessor.
-			void* PropertyObjectPtr = Container;
-			for (int32 i = 0; i < Param.Properties.Num() - 1; ++i)
-			{
-				PropertyObjectPtr = Param.Properties[i]->ContainerPtrToValuePtr<void>(PropertyObjectPtr);
-			}
-
-			FPCGAttributeAccessorKeysSingleObjectPtr PropertyObjectKey(PropertyObjectPtr);
 
+			FPCGAttributeAccessorKeysSingleObjectPtr PropertyObjectKey(Container);
 			PropertyAccessor->Set<PropertyType>(Value, PropertyObjectKey);
 
 			return true;
diff --git a/Engine/Plugins/Experimental/PCG/Source/PCG/Private/PCGSettings.cpp b/Engine/Plugins/Experimental/PCG/Source/PCG/Private/PCGSettings.cpp
--- a/Engine/Plugins/Experimental/PCG/Source/PCG/Private/PCGSettings.cpp
+++ b/Engine/Plugins/Experimental/PCG/Source/PCG/Private/PCGSettings.cpp
@@ -776,14 +776,28 @@ void UPCGSettings::InitializeCachedOverridableParams(bool bReset)
 
 		for (int32 j = 1; j < Param.PropertiesNames.Num(); ++j)
 		{
-			// If we have multiple depth properties, it should be Struct properties by construction
-			const FStructProperty* StructProperty = CastField<FStructProperty>(CurrentProperty);
-			if (ensure(StructProperty))
+			// If we have multiple depth properties, it should be Struct/Object properties by construction
+			const UStruct* UnderlyingStruct = nullptr;
+			if (const FStructProperty* StructProperty = CastField<FStructProperty>(CurrentProperty))
 			{
-				CurrentProperty = StructProperty->Struct->FindPropertyByName(Param.PropertiesNames[j]);
+				UnderlyingStruct = StructProperty->Struct;
+			}
+			else if (const FObjectProperty* ObjectProperty = CastField<FObjectProperty>(CurrentProperty))
+			{
+				UnderlyingStruct = ObjectProperty->PropertyClass;
+			}
+
+			if (ensure(UnderlyingStruct))
+			{
+				CurrentProperty = UnderlyingStruct->FindPropertyByName(Param.PropertiesNames[j]);
 				check(CurrentProperty);
 				Param.Properties.Add(CurrentProperty);
 			}
+			else
+			{
+				Param.Properties.Empty();
+				break;
+			}
 		}
 	}
 }
