[PCG] Tracking fixes for Level Instances

* Moving a level instance is now correctly updating the children
* Not dirtying the level instances on load

#rnx

diff --git a/Engine/Plugins/Experimental/PCG/Source/PCG/Private/PCGActorAndComponentMapping.cpp b/Engine/Plugins/Experimental/PCG/Source/PCG/Private/PCGActorAndComponentMapping.cpp
--- a/Engine/Plugins/Experimental/PCG/Source/PCG/Private/PCGActorAndComponentMapping.cpp
+++ b/Engine/Plugins/Experimental/PCG/Source/PCG/Private/PCGActorAndComponentMapping.cpp
@@ -21,9 +21,10 @@
 #include "Engine/Level.h"
 #include "Engine/StaticMesh.h"
 #include "Kismet/GameplayStatics.h"
+#include "LevelInstance/LevelInstanceActor.h"
+#include "LevelInstance/LevelInstanceEditorInstanceActor.h"
 #include "LevelInstance/LevelInstanceInterface.h"
 #include "LevelInstance/LevelInstanceSubsystem.h"
-#include "LevelInstance/LevelInstanceEditorInstanceActor.h"
 #include "Materials/MaterialInterface.h"
 
 namespace PCGActorAndComponentMapping
@@ -1113,10 +1114,43 @@ void FPCGActorAndComponentMapping::OnActorAdded(AActor* InActor)
 		return;
 	}
 
-	if (ALevelInstanceEditorInstanceActor* LevelInstanceEditorInstance = Cast<ALevelInstanceEditorInstanceActor>(InActor))
+	bool bShouldDirty = true;
+
+	// In case of Level Instances, we first need to catch LevelInstance actors. If we have a LevelInstance actor first,
+	// it means we are currently ADDING a new level instance to the level. It is then followed by the creation of a ALevelInstanceEditorInstanceActor.
+	// It also means it needs to be dirtied.
+	// On the other hand, if we have no LevelInstance actor added, it means the LevelInstance is LOADED, and so we should not dirty the level instance.
+	// To future people reading this code: If this doesn't work anymore, verify that the assumption is still valid.
+	if (ALevelInstance* LevelInstanceActor = Cast<ALevelInstance>(InActor))
 	{
+		TempAddedLevelInstances.Add(LevelInstanceActor);
+	}
+	else if (ALevelInstanceEditorInstanceActor* LevelInstanceEditorInstance = Cast<ALevelInstanceEditorInstanceActor>(InActor))
+	{
+		bShouldDirty = false;
+
 		const ULevelInstanceSubsystem* LevelInstanceSubsystem = PCGSubsystem->GetWorld() ? PCGSubsystem->GetWorld()->GetSubsystem<ULevelInstanceSubsystem>() : nullptr;
 		const ILevelInstanceInterface* ActorLevelInstance = LevelInstanceSubsystem ? LevelInstanceSubsystem->GetParentLevelInstance(InActor) : nullptr;
+
+		// Check if the Level instance actor was added, if so we need to dirty.
+		if (ActorLevelInstance && TempAddedLevelInstances.Contains(ActorLevelInstance))
+		{
+			TempAddedLevelInstances.Remove(ActorLevelInstance);
+			bShouldDirty = true;
+
+			// Also look for child level instances, because the level was not already loaded when the ALevelInstance was added.
+			LevelInstanceSubsystem->ForEachActorInLevelInstance(ActorLevelInstance, [this](AActor* LevelActor)
+			{
+				if (ALevelInstance* LevelInstanceActor = Cast<ALevelInstance>(LevelActor))
+				{
+					TempAddedLevelInstances.Add(LevelInstanceActor);
+				}
+
+				return true;
+			});
+		}
+
+		// We also need to compute the depth.
 		while (ActorLevelInstance)
 		{
 			++LevelInstanceDepth;
@@ -1128,7 +1162,7 @@ void FPCGActorAndComponentMapping::OnActorAdded(AActor* InActor)
 	// Implementation note: since this is called only for actors directly in the current level, the depth here is 0.
 	// Another implementation note: We delay adding because OnActorAdded fires before an actor's properties are set,
 	// so the actor is not ready for processing until the next tick.
-	DelayedAddedActors.Emplace(InActor, { true, LevelInstanceDepth });
+	DelayedAddedActors.Emplace(InActor, { bShouldDirty, LevelInstanceDepth });
 }
 
 void FPCGActorAndComponentMapping::OnActorAdded_Internal(AActor* InActor, bool bShouldDirty, int32 LevelInstanceDepth, bool bForceAddDelayedActor)
@@ -1197,7 +1231,7 @@ void FPCGActorAndComponentMapping::OnActorAdded_Internal(AActor* InActor, bool b
 #endif // WITH_EDITOR
 
 	// Finally notify them all
-	OnObjectChanged(InActor, FBox(EForceInit::ForceInit), nullptr, LevelInstanceDepth);
+	OnObjectChanged(InActor, /*InPreviousData=*/nullptr, /*InOriginatingObject=*/nullptr, LevelInstanceDepth);
 }
 
 void FPCGActorAndComponentMapping::OnActorUnloaded(AActor& InActor)
@@ -1285,7 +1319,7 @@ void FPCGActorAndComponentMapping::OnActorDeleted_Internal(AActor* InActor, bool
 #endif // WITH_EDITOR
 
 	// Notify all components that the actor has changed (was removed), but the Refresh will only happen AFTER the actor was actually removed from the world (because of delayed refresh).
-	OnObjectChanged(InActor, FBox(EForceInit::ForceInit), nullptr, LevelInstanceDepth, /*bNoRefreshOnOwner=*/true);
+	OnObjectChanged(InActor, /*InPreviousData=*/nullptr, /*InOriginatingObject=*/nullptr, LevelInstanceDepth, /*bNoRefreshOnOwner=*/true);
 }
 
 void FPCGActorAndComponentMapping::OnPreObjectPropertyChanged(UObject* InObject, const FEditPropertyChain& InEditPropertyChain)
@@ -1296,8 +1330,6 @@ void FPCGActorAndComponentMapping::OnPreObjectPropertyChanged(UObject* InObject,
 		return;
 	}
 
-	// We want to track tags, to see if a tag was removed
-	TempTrackedActorTags.Empty();
 	FProperty* MemberProperty = InEditPropertyChain.GetActiveMemberNode() ? InEditPropertyChain.GetActiveMemberNode()->GetValue() : nullptr;
 	AActor* Actor = Cast<AActor>(InObject);
 
@@ -1322,12 +1354,27 @@ void FPCGActorAndComponentMapping::OnPreObjectPropertyChanged(UObject* InObject,
 	}
 #endif
 
-	if (MemberProperty && MemberProperty->GetFName() == GET_MEMBER_NAME_CHECKED(AActor, Tags))
+	auto StorePreviousData = [this, MemberProperty](AActor* InActor, int32 LevelInstanceDepth, auto RecursiveCall) -> void
 	{
-		TempTrackedActorTags = TSet<FName>(Actor->Tags);
-	}
+		if (!ActorToPreviousDataMap.Contains(InActor))
+		{
+			FActorPreviousData& PreviousData = ActorToPreviousDataMap.Add(InActor);
+			PreviousData.Get<0>() = PCGActorAndComponentMapping::GetActorBounds(InActor);
+			if (MemberProperty && MemberProperty->GetFName() == GET_MEMBER_NAME_CHECKED(AActor, Tags))
+			{
+				PreviousData.Get<1>() = TSet<FName>(InActor->Tags);
+			}
 
-	ActorToPreviousBoundsMap.FindOrAdd(Actor) = PCGActorAndComponentMapping::GetActorBounds(Actor);
+			// Also propagate the pre-change to all child actors if it is within a level instance.
+			PCGActorAndComponentMapping::PropagateToLevelInstanceActors(InActor, PCGSubsystem, [this, LevelInstanceDepth, RecursiveCall](AActor* LevelActor)
+			{
+				RecursiveCall(LevelActor, LevelInstanceDepth + 1, RecursiveCall);
+				return true;
+			});
+		}
+	};
+
+	StorePreviousData(Actor, /*LevelInstanceDepth*/ 0, StorePreviousData);
 }
 
 void FPCGActorAndComponentMapping::OnObjectSaved(UObject* InObject, FObjectPreSaveContext InObjectSaveContext)
@@ -1382,14 +1429,6 @@ void FPCGActorAndComponentMapping::OnObjectPropertyChanged(UObject* InObject, FP
 		}
 	}
 
-	FBox PreviousBounds{ EForceInit::ForceInit };
-	ActorToPreviousBoundsMap.RemoveAndCopyValue(Actor, PreviousBounds);
-
-	if ((!bValueNotInteractive && !bActorTagChange) || bIsTextureCompilationResult || bIsStaticMeshCompilationResult)
-	{
-		return;
-	}
-
 	if (PCGSubsystem && Actor && Actor->GetWorld() != PCGSubsystem->GetWorld())
 	{
 		return;
@@ -1402,16 +1441,57 @@ void FPCGActorAndComponentMapping::OnObjectPropertyChanged(UObject* InObject, FP
 	}
 #endif
 
-	// Ignore property changes on delayed actors. All their properties are still being set.
-	if (Actor && DelayedAddedActors.Contains(Actor))
+	const bool bNoOperation = (!bValueNotInteractive && !bActorTagChange)
+		|| bIsTextureCompilationResult
+		|| bIsStaticMeshCompilationResult
+		|| (Actor && (DelayedAddedActors.Contains(Actor) || !ActorToPreviousDataMap.Contains(Actor)));
+
+	if (bNoOperation)
 	{
+		if (Actor)
+		{
+			auto RemoveAll = [this](AActor* InActor, auto RecursiveCall) -> void
+			{
+				ActorToPreviousDataMap.Remove(InActor);
+				PCGActorAndComponentMapping::PropagateToLevelInstanceActors(InActor, PCGSubsystem, [this, RecursiveCall](AActor* LevelActor)
+				{
+					RecursiveCall(LevelActor, RecursiveCall);
+					return true;
+				});
+			};
+
+			RemoveAll(Actor, RemoveAll);
+		}
+		
 		return;
 	}
 
-	OnObjectChanged(Actor ? static_cast<UObject*>(Actor) : InObject, PreviousBounds, /*InOriginatingChangeObject=*/ InObject);
+	if (Actor)
+	{
+		auto OnActorChanged = [this, InObject](AActor* InActor, int32 LevelInstanceDepth, auto RecursiveCall) -> void
+		{
+			if (InActor && ActorToPreviousDataMap.Contains(InActor))
+			{
+				FActorPreviousData PreviousData = ActorToPreviousDataMap.FindAndRemoveChecked(InActor);
+				OnObjectChanged(InActor, &PreviousData, /*InOriginatingChangeObject=*/ InObject, LevelInstanceDepth);
+			}
+
+			PCGActorAndComponentMapping::PropagateToLevelInstanceActors(InActor, PCGSubsystem, [this, LevelInstanceDepth, RecursiveCall](AActor* LevelActor)
+			{
+				RecursiveCall(LevelActor, LevelInstanceDepth + 1, RecursiveCall);
+				return true;
+			});
+		};
+
+		OnActorChanged(Actor, /*LevelInstanceDepth=*/0, OnActorChanged);
+	}
+	else
+	{
+		OnObjectChanged(InObject, /*InPreviousData=*/nullptr, /*InOriginatingChangeObject=*/ InObject);
+	}
 }
 
-void FPCGActorAndComponentMapping::OnObjectChanged(UObject* InObject, const FBox& InPreviousBounds, const UObject* InOriginatingChangeObject, int32 LevelInstanceDepth, bool bNoRefreshOwner)
+void FPCGActorAndComponentMapping::OnObjectChanged(UObject* InObject, const FActorPreviousData* InPreviousData, const UObject* InOriginatingChangeObject, int32 LevelInstanceDepth, bool bNoRefreshOwner)
 {
 	// Nothing to do if we track nothing
 	if (CulledTrackedKeysToComponentsMap.IsEmpty() && AlwaysTrackedKeysToComponentsMap.IsEmpty())
@@ -1432,9 +1512,9 @@ void FPCGActorAndComponentMapping::OnObjectChanged(UObject* InObject, const FBox
 	TArray<FPCGSelectionKey> MatchedKeys;
 
 	TSet<FName> RemovedTags;
-	if (Actor && !TempTrackedActorTags.IsEmpty())
+	if (Actor && InPreviousData && !InPreviousData->Get<1>().IsEmpty())
 	{
-		RemovedTags = TempTrackedActorTags.Difference(TSet<FName>(Actor->Tags));
+		RemovedTags = InPreviousData->Get<1>().Difference(TSet<FName>(Actor->Tags));
 	}
 
 	auto Gather = [InObject, &RemovedTags, &MatchedKeys](TMap<FPCGSelectionKey, TSet<UPCGComponent*>> InMap, TSet<UPCGComponent*>& OutSet)
@@ -1584,14 +1664,18 @@ void FPCGActorAndComponentMapping::OnObjectChanged(UObject* InObject, const FBox
 		PartitionedOctree.FindElementsWithBoundsTest(ActorBounds, UpdatePartitioned);
 
 		// If it has moved, redo it with the old bounds.
-		if (InPreviousBounds.IsValid)
+		if (InPreviousData)
 		{
-			TRACE_CPUPROFILER_EVENT_SCOPE(FPCGActorAndComponentMapping::OnActorChanged::SecondUpdateHasMoved);
+			const FBox& PreviousBounds = InPreviousData->Get<0>();
+			if (PreviousBounds.IsValid && !PreviousBounds.Equals(ActorBounds))
+			{
+				TRACE_CPUPROFILER_EVENT_SCOPE(FPCGActorAndComponentMapping::OnActorChanged::SecondUpdateHasMoved);
 
-			// Set the actor bounds with the old one, to have the right Overlap in the Partition case.
-			CurrentActorBoundsPtr = &InPreviousBounds;
-			NonPartitionedOctree.FindElementsWithBoundsTest(InPreviousBounds, UpdateNonPartitioned);
-			PartitionedOctree.FindElementsWithBoundsTest(InPreviousBounds, UpdatePartitioned);
+				// Set the actor bounds with the old one, to have the right Overlap in the Partition case.
+				CurrentActorBoundsPtr = &PreviousBounds;
+				NonPartitionedOctree.FindElementsWithBoundsTest(PreviousBounds, UpdateNonPartitioned);
+				PartitionedOctree.FindElementsWithBoundsTest(PreviousBounds, UpdatePartitioned);
+			}
 		}
 	}
 
@@ -1680,7 +1764,7 @@ void FPCGActorAndComponentMapping::ApplyLandscapeChanges(ALandscapeProxy* InLand
 		return;
 	}
 	
-	OnObjectChanged(InLandscape, FBox(EForceInit::ForceInit), InLandscape);
+	OnObjectChanged(InLandscape, /*InPreviousData=*/nullptr, InLandscape);
 }
 
 void FPCGActorAndComponentMapping::OnPCGGraphGeneratedOrCleaned(UPCGComponent* InComponent)
@@ -1690,7 +1774,7 @@ void FPCGActorAndComponentMapping::OnPCGGraphGeneratedOrCleaned(UPCGComponent* I
 		return;
 	}
 
-	OnObjectChanged(InComponent->GetOwner(), FBox(EForceInit::ForceInit), InComponent);
+	OnObjectChanged(InComponent->GetOwner(), /*InPreviousData=*/nullptr, InComponent);
 }
 
 bool FPCGActorAndComponentMapping::ClearCacheForKeys(const TArray<FPCGSelectionKey>& InKeys, const UPCGComponent* InComponent, const bool bIntersect, const UObject* InOriginatingChange) const
diff --git a/Engine/Plugins/Experimental/PCG/Source/PCG/Public/PCGActorAndComponentMapping.h b/Engine/Plugins/Experimental/PCG/Source/PCG/Public/PCGActorAndComponentMapping.h
--- a/Engine/Plugins/Experimental/PCG/Source/PCG/Public/PCGActorAndComponentMapping.h
+++ b/Engine/Plugins/Experimental/PCG/Source/PCG/Public/PCGActorAndComponentMapping.h
@@ -21,6 +21,7 @@ class AActor;
 class ALandscapeProxy;
 class APCGPartitionActor;
 class FLandscapeProxyComponentDataChangedParams;
+class ILevelInstanceInterface;
 class UObject;
 class UPCGComponent;
 class UPCGGraph;
@@ -128,6 +129,9 @@ private:
 	/** Return true if there are any Original or Non-Partitioned components set to GenerateAtRuntime. */
 	bool AnyRuntimeGenComponentsExist() const;
 
+	// Typedef to store the previous position and previous tags of changed actors.
+	using FActorPreviousData = TTuple<FBox, TSet<FName>>;
+
 #if WITH_EDITOR
 	/** Return true if the key is tracked.*/
 	bool IsKeyTracked(const FPCGSelectionKey& InKey) const;
@@ -155,11 +159,11 @@ private:
 	void UnregisterTracking(UPCGComponent* InComponent, const TSet<FPCGSelectionKey>* OptionalKeysToUntrack);
 
 	/** Trigger an update when an object/actor changed.
-	* Can specify previous actor bounds if it has moved to also update components that were at its previous position.
+	* Can specify previous actor data if it has moved or tags changed to also update components that were at its previous position or tracked previous tags.
 	* Can also specify an optional object, originating the change, to avoid re-dirtying a component if it was the origin.
 	* Another option when an actor is deleted/unload, don't refresh their components.
 	*/
-	void OnObjectChanged(UObject* InObject, const FBox& InPreviousBounds = FBox(EForceInit::ForceInit), const UObject* InOriginatingChangeObject = nullptr, int32 LevelInstanceDepth = 0, bool bNoRefreshOnOwner = false);
+	void OnObjectChanged(UObject* InObject, const FActorPreviousData* InPreviousData = nullptr, const UObject* InOriginatingChangeObject = nullptr, int32 LevelInstanceDepth = 0, bool bNoRefreshOnOwner = false);
 
 	/** Gather all settings from a given component that track the key, and clear the cache for them. Returns true if we should dirty afterwards (aka at least one settings was cleared and/or landscape changed). */
 	bool ClearCacheForKeys(const TArray<FPCGSelectionKey>& InKeys, const UPCGComponent* InComponent, const bool bIntersect, const UObject* InOriginatingChange) const;
@@ -208,11 +212,11 @@ private:
 	// Keep track of actors that aren't yet ready (or if the subsystem is not yet ready), whether we should dirty them and their instance level depth so we can add them in next tick.
 	TMap<TObjectKey<AActor>, TTuple<bool, int>> DelayedAddedActors;
 
-	/** Transient list of tags, kept when there is a tag change on a tracked Actor. */
-	TSet<FName> TempTrackedActorTags;
+	// Keep track of all Level Instance actors added, so we can detect when a level instance is added vs loaded
+	TSet<ILevelInstanceInterface*> TempAddedLevelInstances;
 
-	/** Transient map of actors and their previous bounds, to know when they move and update the components that were touching the actor before, but not anymore. */
-	TMap<TObjectKey<AActor>, FBox> ActorToPreviousBoundsMap;
+	/** Transient map of actors and their previous data, it's set in the pre object change to be able to track changes (such as tags or positions) */
+	TMap<TObjectKey<AActor>, FActorPreviousData> ActorToPreviousDataMap;
 
 #if WITH_EDITOR
 	// Part for the delayed landscape change update
