[PCG] Provide "DensityNoise" as an alias for "AttributeNoise"

#rnx

diff --git a/Engine/Plugins/Experimental/PCG/Source/PCG/Public/Elements/PCGAttributeNoise.h b/Engine/Plugins/Experimental/PCG/Source/PCG/Public/Elements/PCGAttributeNoise.h
--- a/Engine/Plugins/Experimental/PCG/Source/PCG/Public/Elements/PCGAttributeNoise.h
+++ b/Engine/Plugins/Experimental/PCG/Source/PCG/Public/Elements/PCGAttributeNoise.h
@@ -44,6 +44,7 @@ public:
 	virtual void ApplyDeprecationBeforeUpdatePins(UPCGNode* InOutNode, TArray<TObjectPtr<UPCGPin>>& InputPins, TArray<TObjectPtr<UPCGPin>>& OutputPins);
 	virtual FName GetDefaultNodeName() const override { return FName(TEXT("AttributeNoise")); }
 	virtual FText GetDefaultNodeTitle() const override { return NSLOCTEXT("PCGAttributeNoiseSettings", "NodeTitle", "Attribute Noise"); }
+	virtual TArray<FText> GetNodeTitleAliases() const { return { NSLOCTEXT("PCGAttributeNoiseSettings", "DensityNoiseAlias", "Density Noise") }; }
 	virtual EPCGSettingsType GetType() const override { return EPCGSettingsType::Metadata; }
 #endif
 
