[PCG] Improve ActorPropertyOverride logging when source attribute does not exist

#rnx

diff --git a/Engine/Plugins/Experimental/PCG/Source/PCG/Private/Metadata/PCGActorPropertyOverride.cpp b/Engine/Plugins/Experimental/PCG/Source/PCG/Private/Metadata/PCGActorPropertyOverride.cpp
--- a/Engine/Plugins/Experimental/PCG/Source/PCG/Private/Metadata/PCGActorPropertyOverride.cpp
+++ b/Engine/Plugins/Experimental/PCG/Source/PCG/Private/Metadata/PCGActorPropertyOverride.cpp
@@ -60,6 +60,13 @@ void FPCGActorSingleOverride::Initialize(const FPCGAttributePropertySelector& In
 {
 	InputKeys = PCGAttributeAccessorHelpers::CreateConstKeys(SourceData, InputSelector);
 	ActorOverrideInputAccessor = PCGAttributeAccessorHelpers::CreateConstAccessor(SourceData, InputSelector);
+
+	if (!ActorOverrideInputAccessor.IsValid())
+	{
+		PCGLog::LogWarningOnGraph(FText::Format(LOCTEXT("OverrideInputInvalid", "ActorOverride for input '{0}' is invalid or unsupported."), InputSelector.GetDisplayText()), Context);
+		return;
+	}
+
 	FPCGAttributePropertySelector OutputSelector = FPCGAttributePropertySelector::CreateSelectorFromString(OutputProperty);
 	const TArray<FString>& ExtraNames = OutputSelector.GetExtraNames();
 	if (ExtraNames.IsEmpty())
@@ -79,9 +86,9 @@ void FPCGActorSingleOverride::Initialize(const FPCGAttributePropertySelector& In
 		ActorOverrideOutputAccessor = PCGAttributeAccessorHelpers::CreatePropertyChainAccessor(PropertyNames, TemplateActor->GetClass());
 	}
 
-	if (!ActorOverrideInputAccessor.IsValid() || !ActorOverrideOutputAccessor.IsValid())
+	if (!ActorOverrideOutputAccessor.IsValid())
 	{
-		PCGLog::LogWarningOnGraph(FText::Format(LOCTEXT("OverrideInvalid", "ActorOverride from input '{0}' or output '{1}' is invalid or unsupported. Will be skipped."), InputSelector.GetDisplayText(), OutputSelector.GetDisplayText()), Context);
+		PCGLog::LogWarningOnGraph(FText::Format(LOCTEXT("OverrideOutputInvalid", "ActorOverride for actor property '{0}' is invalid or unsupported."), FText::FromString(OutputProperty)), Context);
 		return;
 	}
 
@@ -90,7 +97,7 @@ void FPCGActorSingleOverride::Initialize(const FPCGAttributePropertySelector& In
 		PCGLog::LogWarningOnGraph(
 			FText::Format(LOCTEXT("TypesIncompatible", "ActorOverride cannot set input '{0}' to output '{1}'. Cannot convert type '{2}' to type '{3}'. Will be skipped."),
 				InputSelector.GetDisplayText(),
-				OutputSelector.GetDisplayText(),
+				FText::FromString(OutputProperty),
 				PCG::Private::GetTypeNameText(ActorOverrideInputAccessor->GetUnderlyingType()),
 				PCG::Private::GetTypeNameText(ActorOverrideOutputAccessor->GetUnderlyingType())),
 			Context);
@@ -133,6 +140,19 @@ void FPCGActorOverrides::Initialize(const TArray<FPCGActorPropertyOverrideDescri
 	for (int32 i = 0; i < OverrideDescriptions.Num(); ++i)
 	{
 		FPCGAttributePropertyInputSelector InputSelector = OverrideDescriptions[i].InputSource.CopyAndFixLast(SourceData);
+
+		if (InputSelector.GetSelection() == EPCGAttributePropertySelection::Attribute)
+		{
+			const UPCGMetadata* Metadata = SourceData ? SourceData->ConstMetadata() : nullptr;
+			const FName AttributeName = InputSelector.GetAttributeName();
+
+			if (!Metadata || !Metadata->HasAttribute(AttributeName))
+			{
+				PCGLog::LogWarningOnGraph(FText::Format(LOCTEXT("InvalidAttribute", "Tried to initialize ActorOverride for input '{0}', but the attribute '{1}' does not exist."), InputSelector.GetDisplayText(), FText::FromName(AttributeName)), Context);
+				continue;
+			}
+		}
+
 		const FString& OutputProperty = OverrideDescriptions[i].PropertyTarget;
 
 		FPCGActorSingleOverride Override;
