Fixed loctext collision

#rnx

diff --git a/Engine/Plugins/Experimental/PCG/Source/PCG/Private/Elements/PCGMatchAndSetAttributes.cpp b/Engine/Plugins/Experimental/PCG/Source/PCG/Private/Elements/PCGMatchAndSetAttributes.cpp
--- a/Engine/Plugins/Experimental/PCG/Source/PCG/Private/Elements/PCGMatchAndSetAttributes.cpp
+++ b/Engine/Plugins/Experimental/PCG/Source/PCG/Private/Elements/PCGMatchAndSetAttributes.cpp
@@ -304,7 +304,7 @@ public:
 		{
 			if (!PCG::Private::IsBroadcastableOrConstructible(ConstantThreshold->GetUnderlyingType(), Attribute->GetTypeId()))
 			{
-				PCGE_LOG_C(Error, GraphAndLog, Context, FText::Format(LOCTEXT("InvalidThresholdAttribute", "Distance threshold type is not compatible with attribute '{0}'."), FText::FromName(Attribute->Name)));
+				PCGE_LOG_C(Error, GraphAndLog, Context, FText::Format(LOCTEXT("InvalidThresholdAttributeType", "Distance threshold type is not compatible with attribute '{0}'."), FText::FromName(Attribute->Name)));
 				return false;
 			}
 
