[PCG] Fix: Plugins.PCG.ReplaceTag.OneToNError test failure addressed

diff --git a/Engine/Plugins/Experimental/PCG/Source/PCG/Private/Tests/Elements/PCGReplaceTagTest.cpp b/Engine/Plugins/Experimental/PCG/Source/PCG/Private/Tests/Elements/PCGReplaceTagTest.cpp
--- a/Engine/Plugins/Experimental/PCG/Source/PCG/Private/Tests/Elements/PCGReplaceTagTest.cpp
+++ b/Engine/Plugins/Experimental/PCG/Source/PCG/Private/Tests/Elements/PCGReplaceTagTest.cpp
@@ -100,13 +100,15 @@ IMPLEMENT_CUSTOM_SIMPLE_AUTOMATION_TEST(FPCGReplaceTagTest_OneToNError, PCGRepla
 
 bool FPCGReplaceTagTest_OneToNError::RunTest(const FString& Parameters)
 {
-	AddExpectedError(TEXT("Input only supports 1:1, N:1, and N:N mappings from source to replaced tags, discarded."));
+	AddExpectedError(TEXT("Input only supports 1:1, N:1, and N:N mappings from source to replaced tags, Input data discarded."));
 
 	FTestParameters OneToNErrorParameters{};
 	OneToNErrorParameters.SelectedTags = TEXT("Awesome");
 	OneToNErrorParameters.ReplacedTags = TEXT("Neat, Great");
-	OneToNErrorParameters.ExpectedToHaveOneOutput = false;
-	OneToNErrorParameters.ExpectedOutputTags = TSet<FString>();
+
+	// Expects one output because input is forwarded
+	OneToNErrorParameters.ExpectedToHaveOneOutput = true;
+	OneToNErrorParameters.ExpectedOutputTags = TSet<FString>{ TEXT("Awesome"), TEXT("Fun"), TEXT("Cool") };
 
 	return GenerateTestDataRunAndValidate(OneToNErrorParameters);
 }
