[PCG] Make the Unbounded grid size enum value something that is clearly not a valid grid

diff --git a/Engine/Plugins/Experimental/PCG/Source/PCG/Private/PCGCommon.cpp b/Engine/Plugins/Experimental/PCG/Source/PCG/Private/PCGCommon.cpp
--- a/Engine/Plugins/Experimental/PCG/Source/PCG/Private/PCGCommon.cpp
+++ b/Engine/Plugins/Experimental/PCG/Source/PCG/Private/PCGCommon.cpp
@@ -60,8 +60,7 @@ namespace PCGHiGenGrid
 
 	uint32 UnboundedGridSize()
 	{
-		// TODO: support other units
-		return 100 * static_cast<uint32>(EPCGHiGenGrid::Unbounded);
+		return static_cast<uint32>(EPCGHiGenGrid::Unbounded);
 	}
 }
 
diff --git a/Engine/Plugins/Experimental/PCG/Source/PCG/Public/PCGCommon.h b/Engine/Plugins/Experimental/PCG/Source/PCG/Public/PCGCommon.h
--- a/Engine/Plugins/Experimental/PCG/Source/PCG/Public/PCGCommon.h
+++ b/Engine/Plugins/Experimental/PCG/Source/PCG/Public/PCGCommon.h
@@ -184,7 +184,7 @@ enum class EPCGHiGenGrid : uint32
 	GridMax = Grid2048 UMETA(Hidden),
 
 	// Should execute once rather than executing on any grid
-	Unbounded = 2 * GridMax,
+	Unbounded = 1u << 31,
 };
 ENUM_CLASS_FLAGS(EPCGHiGenGrid);
 
