[PCG] Fix for profiling not displaying results for partitioned components

The compiled tasks were erroneously compiled for a particular grid size which meant they could not be retrieved. This should only happen for higen graphs.

diff --git a/Engine/Plugins/Experimental/PCG/Source/PCG/Private/Graph/PCGGraphExecutor.cpp b/Engine/Plugins/Experimental/PCG/Source/PCG/Private/Graph/PCGGraphExecutor.cpp
--- a/Engine/Plugins/Experimental/PCG/Source/PCG/Private/Graph/PCGGraphExecutor.cpp
+++ b/Engine/Plugins/Experimental/PCG/Source/PCG/Private/Graph/PCGGraphExecutor.cpp
@@ -121,10 +121,12 @@ FPCGTaskId FPCGGraphExecutor::Schedule(
 	FPCGTaskId ScheduledId = InvalidPCGTaskId;
 
 	uint32 GenerationGridSize = PCGHiGenGrid::UninitializedGridSize();
-	const bool bComponentIsPartOfHierarchy = SourceComponent->IsLocalComponent() || SourceComponent->IsPartitioned();
-	if (bAllowHierarchicalGeneration && bComponentIsPartOfHierarchy)
+	if (bAllowHierarchicalGeneration && Graph->IsHierarchicalGenerationEnabled())
 	{
-		GenerationGridSize = SourceComponent->GetGenerationGridSize();
+		if (SourceComponent->IsLocalComponent() || SourceComponent->IsPartitioned())
+		{
+			GenerationGridSize = SourceComponent->GetGenerationGridSize();
+		}
 	}
 
 	// Get compiled tasks from compiler
