[PCG] Improved PCG builder logging and validation.

diff --git a/Engine/Plugins/Experimental/PCG/Source/PCG/Private/PCGComponent.cpp b/Engine/Plugins/Experimental/PCG/Source/PCG/Private/PCGComponent.cpp
--- a/Engine/Plugins/Experimental/PCG/Source/PCG/Private/PCGComponent.cpp
+++ b/Engine/Plugins/Experimental/PCG/Source/PCG/Private/PCGComponent.cpp
@@ -300,7 +300,7 @@ void UPCGComponent::Generate_Implementation(bool bForce)
 
 void UPCGComponent::GenerateLocal(bool bForce)
 {
-	GenerateInternal(bForce, EPCGHiGenGrid::Uninitialized, EPCGComponentGenerationTrigger::GenerateOnDemand, {});
+	GenerateLocalGetTaskId(bForce);
 }
 
 void UPCGComponent::GenerateLocal(EPCGComponentGenerationTrigger RequestedGenerationTrigger, bool bForce, EPCGHiGenGrid Grid, const TArray<FPCGTaskId>& Dependencies)
diff --git a/Engine/Plugins/Experimental/PCG/Source/PCGEditor/Private/WorldPartitionBuilder/PCGWorldPartitionBuilder.cpp b/Engine/Plugins/Experimental/PCG/Source/PCGEditor/Private/WorldPartitionBuilder/PCGWorldPartitionBuilder.cpp
--- a/Engine/Plugins/Experimental/PCG/Source/PCGEditor/Private/WorldPartitionBuilder/PCGWorldPartitionBuilder.cpp
+++ b/Engine/Plugins/Experimental/PCG/Source/PCGEditor/Private/WorldPartitionBuilder/PCGWorldPartitionBuilder.cpp
@@ -299,17 +299,45 @@ bool PCGWorldPartitionBuilder::GenerateComponents(
 	for (TWeakObjectPtr<UPCGComponent> ComponentWeakPtr : Components)
 	{
 		UPCGComponent* Component = ComponentWeakPtr.Get();
-		const UPCGGraph* Graph = Component ? Component->GetGraph() : nullptr;
-		if (!Component || !Graph)
+		if (!Component)
 		{
 			UE_LOG(LogPCGWorldPartitionBuilder, Warning, TEXT("Lost a PCG component weak pointer, component will not be generated."));
 			continue;
 		}
 
+		// Validate this before running the filter as the filtering can check the graph name etc.
+		const UPCGGraph* Graph = Component->GetGraph();
+		if (!Graph)
+		{
+			UE_LOG(LogPCGWorldPartitionBuilder, Warning, TEXT("PCG component on actor '%s' label '%s' has no graph assigned, skipping."),
+				*Component->GetOwner()->GetName(),
+				*Component->GetOwner()->GetActorNameOrLabel());
+			continue;
+		}
+
 		if (!ComponentFilter(Component))
 		{
 			continue;
 		}
+		
+		// Last minute validations, done here just prior to generation (after component has passed all previous filters) to minimize spam.
+		if (!Component->bActivated)
+		{
+			UE_LOG(LogPCGWorldPartitionBuilder, Warning, TEXT("'Activated' toggle was set false on PCG component on actor '%s' label '%s' graph '%s'. Component skipped."),
+				*Component->GetOwner()->GetName(),
+				*Component->GetOwner()->GetActorNameOrLabel(),
+				*Component->GetGraph()->GetName());
+			continue;
+		}
+
+		if (Component->IsManagedByRuntimeGenSystem())
+		{
+			UE_LOG(LogPCGWorldPartitionBuilder, Warning, TEXT("PCG component generation trigger is set to run-time generation on actor '%s' label '%s' graph '%s'. Component skipped."),
+				*Component->GetOwner()->GetName(),
+				*Component->GetOwner()->GetActorNameOrLabel(),
+				*Component->GetGraph()->GetName());
+			continue;
+		}
 
 		if (bOneComponentAtATime)
 		{
@@ -338,6 +366,10 @@ bool PCGWorldPartitionBuilder::GenerateComponents(
 			{
 				WaitForComponentGeneration(Component);
 			}
+			else
+			{
+				UE_LOG(LogPCGWorldPartitionBuilder, Warning, TEXT("GeneratedComponent reference lost, abandoning."));
+			}
 		}
 	}
 
@@ -346,6 +378,15 @@ bool PCGWorldPartitionBuilder::GenerateComponents(
 
 void PCGWorldPartitionBuilder::GenerateComponent(UPCGComponent* InComponent, UWorld* InWorld)
 {
+	// Separate ensures for maximum debug information.
+	if (!ensure(InComponent) || !ensure(InComponent->GetGraph()))
+	{
+		return;
+	}
+
+	ensure(InComponent->bActivated);
+	ensure(!InComponent->IsManagedByRuntimeGenSystem());
+
 	if (InComponent->GetSerializedEditingMode() == EPCGEditorDirtyMode::LoadAsPreview)
 	{
 		InComponent->SetEditingMode(EPCGEditorDirtyMode::LoadAsPreview, InComponent->GetSerializedEditingMode());
@@ -353,7 +394,24 @@ void PCGWorldPartitionBuilder::GenerateComponent(UPCGComponent* InComponent, UWo
 	}
 
 	// Force generate as components that are already generated may decline the request.
-	InComponent->GenerateLocal(/*bForce=*/true);
+	const FPCGTaskId GenerateTask = InComponent->GenerateLocalGetTaskId(/*bForce=*/true);
+
+	if (GenerateTask == InvalidPCGTaskId)
+	{
+		if (ensure(InComponent->GetOwner()))
+		{
+			UE_LOG(LogPCGWorldPartitionBuilder, Warning, TEXT("Scheduling generate task failed for PCG Component on actor '%s' label '%s' graph '%s'."),
+				*InComponent->GetOwner()->GetName(),
+				*InComponent->GetOwner()->GetActorNameOrLabel(),
+				*InComponent->GetGraph()->GetName());
+		}
+		else
+		{
+			UE_LOG(LogPCGWorldPartitionBuilder, Warning, TEXT("Scheduling generate task failed for PCG Component with no owner. Component '%s' graph '%s'."),
+				*InComponent->GetName(),
+				*InComponent->GetGraph()->GetName());
+		}
+	}
 }
 
 void PCGWorldPartitionBuilder::WaitForAllAsyncEditorProcesses(UWorld* InWorld)
