Added first iteration of named reroutes. Some satellite features are missing and there is limited functionality in collapse nodes and when deleting declarations.

diff --git a/Engine/Plugins/Experimental/PCG/Source/PCG/Private/Elements/PCGReroute.cpp b/Engine/Plugins/Experimental/PCG/Source/PCG/Private/Elements/PCGReroute.cpp
--- a/Engine/Plugins/Experimental/PCG/Source/PCG/Private/Elements/PCGReroute.cpp
+++ b/Engine/Plugins/Experimental/PCG/Source/PCG/Private/Elements/PCGReroute.cpp
@@ -44,6 +44,43 @@ FPCGElementPtr UPCGRerouteSettings::CreateElement() const
 	return MakeShared<FPCGRerouteElement>();
 }
 
+TArray<FPCGPinProperties> UPCGNamedRerouteDeclarationSettings::OutputPinProperties() const
+{
+	TArray<FPCGPinProperties> PinProperties;
+	// Default visible pin
+	PinProperties.Emplace(
+		PCGPinConstants::DefaultOutputLabel, 
+		EPCGDataType::Any, 
+		/*bAllowMultipleConnections=*/true, 
+		/*bAllowMultipleData=*/true);
+
+	FPCGPinProperties& InvisiblePin = PinProperties.Emplace_GetRef(
+		PCGNamedRerouteConstants::InvisiblePinLabel,
+		EPCGDataType::Any,
+		/*bAllowMultipleConnections=*/true,
+		/*bAllowMultipleData=*/true);
+	InvisiblePin.bInvisiblePin = true;
+	
+	return PinProperties;
+}
+
+TArray<FPCGPinProperties> UPCGNamedRerouteUsageSettings::InputPinProperties() const
+{
+	TArray<FPCGPinProperties> PinProperties = Super::InputPinProperties();
+	if (ensure(PinProperties.Num() == 1))
+	{
+		PinProperties[0].bInvisiblePin = true;
+	}
+
+	return PinProperties;
+}
+
+EPCGDataType UPCGNamedRerouteUsageSettings::GetCurrentPinTypes(const UPCGPin* InPin) const
+{
+	// Defer to declaration if possible
+	return Declaration ? Declaration->GetCurrentPinTypes(InPin) : Super::GetCurrentPinTypes(InPin);
+}
+
 bool FPCGRerouteElement::ExecuteInternal(FPCGContext* Context) const
 {
 	ensureMsgf(false, TEXT("Reroute elements are not supposed to execute - reroutes should be culled during graph compilation."));
diff --git a/Engine/Plugins/Experimental/PCG/Source/PCG/Private/PCGGraph.cpp b/Engine/Plugins/Experimental/PCG/Source/PCG/Private/PCGGraph.cpp
--- a/Engine/Plugins/Experimental/PCG/Source/PCG/Private/PCGGraph.cpp
+++ b/Engine/Plugins/Experimental/PCG/Source/PCG/Private/PCGGraph.cpp
@@ -390,6 +390,35 @@ bool UPCGGraph::Contains(const UPCGGraph* InGraph) const
 	return bContains;
 }
 
+UPCGNode* UPCGGraph::FindNodeWithSettings(const UPCGSettingsInterface* InSettings, bool bRecursive) const
+{
+	UPCGNode* NodeFound = nullptr;
+
+	auto FindNode = [&NodeFound, InSettings](UPCGNode* InNode)
+	{
+		if (InNode && InNode->GetSettingsInterface() == InSettings)
+		{
+			NodeFound = InNode;
+			return false; // stop execution
+		}
+		else
+		{
+			return true;
+		}
+	};
+
+	if (bRecursive)
+	{
+		ForEachNodeRecursively(FindNode);
+	}
+	else
+	{
+		ForEachNode(FindNode);
+	}
+
+	return NodeFound;
+}
+
 #if WITH_EDITOR
 void UPCGGraph::DeclareConstructClasses(TArray<FTopLevelAssetPath>& OutConstructClasses, const UClass* SpecificSubclass)
 {
diff --git a/Engine/Plugins/Experimental/PCG/Source/PCG/Private/PCGInputOutputSettings.cpp b/Engine/Plugins/Experimental/PCG/Source/PCG/Private/PCGInputOutputSettings.cpp
--- a/Engine/Plugins/Experimental/PCG/Source/PCG/Private/PCGInputOutputSettings.cpp
+++ b/Engine/Plugins/Experimental/PCG/Source/PCG/Private/PCGInputOutputSettings.cpp
@@ -213,6 +213,42 @@ void UPCGGraphInputOutputSettings::PostEditChangeProperty(struct FPropertyChange
 }
 #endif // WITH_EDITOR
 
+TArray<FPCGPinProperties> UPCGGraphInputOutputSettings::InputPinProperties() const
+{
+	if (bIsInput)
+	{
+		TArray<FPCGPinProperties> InvisiblePins = Pins;
+		for (FPCGPinProperties& Pin : InvisiblePins)
+		{
+			Pin.bInvisiblePin = true;
+		}
+
+		return InvisiblePins;
+	}
+	else
+	{
+		return Pins;
+	}
+}
+
+TArray<FPCGPinProperties> UPCGGraphInputOutputSettings::OutputPinProperties() const
+{
+	if (bIsInput)
+	{
+		return Pins;
+	}
+	else
+	{
+		TArray<FPCGPinProperties> InvisiblePins = Pins;
+		for (FPCGPinProperties& Pin : InvisiblePins)
+		{
+			Pin.bInvisiblePin = true;
+		}
+
+		return InvisiblePins;
+	}
+}
+
 void UPCGGraphInputOutputSettings::FixPinProperties()
 {
 	// No need to fix if we have no pins
diff --git a/Engine/Plugins/Experimental/PCG/Source/PCG/Private/PCGPin.cpp b/Engine/Plugins/Experimental/PCG/Source/PCG/Private/PCGPin.cpp
--- a/Engine/Plugins/Experimental/PCG/Source/PCG/Private/PCGPin.cpp
+++ b/Engine/Plugins/Experimental/PCG/Source/PCG/Private/PCGPin.cpp
@@ -90,7 +90,8 @@ bool FPCGPinProperties::operator==(const FPCGPinProperties& Other) const
 		bAllowMultipleConnections == Other.bAllowMultipleConnections &&
 		bAllowMultipleData == Other.bAllowMultipleData &&
 		Usage == Other.Usage &&
-		bAdvancedPin == Other.bAdvancedPin;
+		bAdvancedPin == Other.bAdvancedPin &&
+		bInvisiblePin == Other.bInvisiblePin;
 }
 
 UPCGPin::UPCGPin(const FObjectInitializer& ObjectInitializer)
diff --git a/Engine/Plugins/Experimental/PCG/Source/PCG/Private/PCGSubgraph.cpp b/Engine/Plugins/Experimental/PCG/Source/PCG/Private/PCGSubgraph.cpp
--- a/Engine/Plugins/Experimental/PCG/Source/PCG/Private/PCGSubgraph.cpp
+++ b/Engine/Plugins/Experimental/PCG/Source/PCG/Private/PCGSubgraph.cpp
@@ -18,7 +18,7 @@
 
 namespace PCGSubgraphSettings
 {
-	void RemoveAdvancedModeOnConnectedPins(const UPCGGraph* Subgraph, TArray<FPCGPinProperties>& InOutPinProperties, const bool bIsInput)
+	void RemoveAdvancedAndInvisibleOnConnectedPins(const UPCGGraph* Subgraph, TArray<FPCGPinProperties>& InOutPinProperties, const bool bIsInput)
 	{
 		const UPCGNode* SubgraphNode = bIsInput ? Subgraph->GetInputNode() : Subgraph->GetOutputNode();
 		check(SubgraphNode);
@@ -26,9 +26,14 @@ namespace PCGSubgraphSettings
 		for (FPCGPinProperties& PinProperties : InOutPinProperties)
 		{
 			const UPCGPin* Pin = bIsInput ? SubgraphNode->GetOutputPin(PinProperties.Label) : SubgraphNode->GetInputPin(PinProperties.Label);
-			if (ensure(Pin) && Pin->IsConnected())
+			if (ensure(Pin))
 			{
-				PinProperties.bAdvancedPin = false;
+				PinProperties.bInvisiblePin = false;
+
+				if (Pin->IsConnected())
+				{
+					PinProperties.bAdvancedPin = false;
+				}
 			}
 		}
 	}
@@ -193,7 +198,7 @@ TArray<FPCGPinProperties> UPCGBaseSubgraphSettings::InputPinProperties() const
 	if (UPCGGraph* Subgraph = GetSubgraph())
 	{
 		TArray<FPCGPinProperties> InputPins = Subgraph->GetInputNode()->InputPinProperties();
-		PCGSubgraphSettings::RemoveAdvancedModeOnConnectedPins(Subgraph, InputPins, /*bIsInput=*/true);
+		PCGSubgraphSettings::RemoveAdvancedAndInvisibleOnConnectedPins(Subgraph, InputPins, /*bIsInput=*/true);
 		return InputPins;
 	}
 	else
@@ -207,7 +212,7 @@ TArray<FPCGPinProperties> UPCGBaseSubgraphSettings::OutputPinProperties() const
 	if (UPCGGraph* Subgraph = GetSubgraph())
 	{
 		TArray<FPCGPinProperties> OutputPins = Subgraph->GetOutputNode()->OutputPinProperties();
-		PCGSubgraphSettings::RemoveAdvancedModeOnConnectedPins(Subgraph, OutputPins, /*bIsInput=*/false);
+		PCGSubgraphSettings::RemoveAdvancedAndInvisibleOnConnectedPins(Subgraph, OutputPins, /*bIsInput=*/false);
 		return OutputPins;
 	}
 	else
diff --git a/Engine/Plugins/Experimental/PCG/Source/PCG/Public/Elements/PCGReroute.h b/Engine/Plugins/Experimental/PCG/Source/PCG/Public/Elements/PCGReroute.h
--- a/Engine/Plugins/Experimental/PCG/Source/PCG/Public/Elements/PCGReroute.h
+++ b/Engine/Plugins/Experimental/PCG/Source/PCG/Public/Elements/PCGReroute.h
@@ -9,6 +9,11 @@
 
 #include "PCGReroute.generated.h"
 
+namespace PCGNamedRerouteConstants
+{
+	const FName InvisiblePinLabel = TEXT("InvisiblePin");
+}
+
 UCLASS(BlueprintType, ClassGroup = (Procedural))
 class PCG_API UPCGRerouteSettings : public UPCGSettings
 {
@@ -28,6 +33,8 @@ public:
 #if WITH_EDITOR
 	virtual FName GetDefaultNodeName() const override { return FName("Reroute"); }
 	virtual FText GetDefaultNodeTitle() const override { return NSLOCTEXT("PCGRerouteElement", "NodeTitle", "Reroute"); }
+	virtual EPCGSettingsType GetType() const override { return EPCGSettingsType::Reroute; }
+	virtual bool CanUserEditTitle() const override { return false; }
 #endif
 
 protected:
@@ -37,6 +44,43 @@ protected:
 	//~End UPCGSettings interface
 };
 
+/** Base class for both reroute declaration and usage to share implementation, but also because they use the same visual node representation in the editor. */
+UCLASS(ClassGroup = (Procedural))
+class PCG_API UPCGNamedRerouteBaseSettings : public UPCGRerouteSettings
+{
+	GENERATED_BODY()
+};
+
+UCLASS(BlueprintType, ClassGroup = (Procedural))
+class PCG_API UPCGNamedRerouteDeclarationSettings : public UPCGNamedRerouteBaseSettings
+{
+	GENERATED_BODY()
+
+protected:
+	virtual TArray<FPCGPinProperties> OutputPinProperties() const override;
+
+public:
+#if WITH_EDITOR
+	virtual bool CanUserEditTitle() const override { return true; }
+#endif
+};
+
+UCLASS(BlueprintType, ClassGroup = (Procedural))
+class PCG_API UPCGNamedRerouteUsageSettings : public UPCGNamedRerouteBaseSettings
+{
+	GENERATED_BODY()
+
+protected:
+	virtual TArray<FPCGPinProperties> InputPinProperties() const override;
+
+public:
+	EPCGDataType GetCurrentPinTypes(const UPCGPin* InPin) const override;
+
+public:
+	UPROPERTY(BlueprintReadWrite, Category = Settings)
+	TObjectPtr<const UPCGNamedRerouteDeclarationSettings> Declaration;
+};
+
 class PCG_API FPCGRerouteElement : public IPCGElement
 {
 protected:
diff --git a/Engine/Plugins/Experimental/PCG/Source/PCG/Public/Elements/PCGUserParameterGet.h b/Engine/Plugins/Experimental/PCG/Source/PCG/Public/Elements/PCGUserParameterGet.h
--- a/Engine/Plugins/Experimental/PCG/Source/PCG/Public/Elements/PCGUserParameterGet.h
+++ b/Engine/Plugins/Experimental/PCG/Source/PCG/Public/Elements/PCGUserParameterGet.h
@@ -35,6 +35,10 @@ public:
 
 	//~Begin UPCGSettings interface
 	virtual bool ShouldHookToPreTask() const override { return true; }
+#if WITH_EDITOR
+	virtual bool CanUserEditTitle() const override { return false; }
+#endif
+
 protected:
 	virtual TArray<FPCGPinProperties> InputPinProperties() const override { return {}; }
 	virtual TArray<FPCGPinProperties> OutputPinProperties() const override;
@@ -43,6 +47,7 @@ protected:
 #if WITH_EDITOR
 	virtual FName GetDefaultNodeName() const override { return FName("GetGraphParameter"); }
 	virtual FText GetDefaultNodeTitle() const override { return NSLOCTEXT("PCGUserParameterGetSettings", "NodeTitle", "Get Graph Parameter"); }
+	virtual EPCGSettingsType GetType() const override { return EPCGSettingsType::GraphParameters; }
 #endif
 	//~End UPCGSettings interface
 };
diff --git a/Engine/Plugins/Experimental/PCG/Source/PCG/Public/PCGGraph.h b/Engine/Plugins/Experimental/PCG/Source/PCG/Public/PCGGraph.h
--- a/Engine/Plugins/Experimental/PCG/Source/PCG/Public/PCGGraph.h
+++ b/Engine/Plugins/Experimental/PCG/Source/PCG/Public/PCGGraph.h
@@ -270,6 +270,9 @@ public:
 	/** Returns true if the current graph contains a subgraph node using statically the specified graph, recursively. */
 	bool Contains(const UPCGGraph* InGraph) const;
 
+	/** Returns the node with the given settings in the graph, if any */
+	UPCGNode* FindNodeWithSettings(const UPCGSettingsInterface* InSettings, bool bRecursive = false) const;
+
 	const TArray<UPCGNode*>& GetNodes() const { return Nodes; }
 	void AddNode(UPCGNode* InNode);
 	void AddNodes(TArray<UPCGNode*>& InNodes);
diff --git a/Engine/Plugins/Experimental/PCG/Source/PCG/Public/PCGInputOutputSettings.h b/Engine/Plugins/Experimental/PCG/Source/PCG/Public/PCGInputOutputSettings.h
--- a/Engine/Plugins/Experimental/PCG/Source/PCG/Public/PCGInputOutputSettings.h
+++ b/Engine/Plugins/Experimental/PCG/Source/PCG/Public/PCGInputOutputSettings.h
@@ -42,8 +42,8 @@ public:
 	virtual void ApplyDeprecationBeforeUpdatePins(UPCGNode* InOutNode, TArray<TObjectPtr<UPCGPin>>& InputPins, TArray<TObjectPtr<UPCGPin>>& OutputPins) override;
 #endif
 
-	TArray<FPCGPinProperties> InputPinProperties() const override { return Pins; }
-	TArray<FPCGPinProperties> OutputPinProperties() const override { return Pins; }
+	TArray<FPCGPinProperties> InputPinProperties() const override;
+	TArray<FPCGPinProperties> OutputPinProperties() const override;
 
 	virtual TArray<FPCGPinProperties> DefaultInputPinProperties() const override;
 	virtual TArray<FPCGPinProperties> DefaultOutputPinProperties() const override;
diff --git a/Engine/Plugins/Experimental/PCG/Source/PCG/Public/PCGPin.h b/Engine/Plugins/Experimental/PCG/Source/PCG/Public/PCGPin.h
--- a/Engine/Plugins/Experimental/PCG/Source/PCG/Public/PCGPin.h
+++ b/Engine/Plugins/Experimental/PCG/Source/PCG/Public/PCGPin.h
@@ -46,6 +46,9 @@ struct PCG_API FPCGPinProperties
 	UPROPERTY(BlueprintReadWrite, EditAnywhere, Category = Settings)
 	bool bAdvancedPin = false;
 
+	UPROPERTY(BlueprintReadWrite, Category = Settings)
+	bool bInvisiblePin = false;
+
 #if WITH_EDITORONLY_DATA
 	UPROPERTY(EditAnywhere, Category = Settings)
 	FText Tooltip;
diff --git a/Engine/Plugins/Experimental/PCG/Source/PCG/Public/PCGSettings.h b/Engine/Plugins/Experimental/PCG/Source/PCG/Public/PCGSettings.h
--- a/Engine/Plugins/Experimental/PCG/Source/PCG/Public/PCGSettings.h
+++ b/Engine/Plugins/Experimental/PCG/Source/PCG/Public/PCGSettings.h
@@ -60,7 +60,9 @@ enum class EPCGSettingsType : uint8
 	Param,
 	HierarchicalGeneration,
 	ControlFlow,
-	PointOps
+	PointOps,
+	GraphParameters,
+	Reroute
 };
 
 #if WITH_EDITOR
@@ -264,6 +266,8 @@ public:
 	virtual bool ShouldDrawNodeCompact() const { return false; }
 	/** Returns the icon to use instead of text in compact node form */
 	virtual bool GetCompactNodeIcon(FName& OutCompactNodeIcon) const { return false; }
+	/** Returns whether the user can directly interact with the node name */
+	virtual bool CanUserEditTitle() const { return true; }
 
 	/** UpdatePins will kick off invalid edges, so this is useful for moving edges around in case of pin changes. */
 	virtual void ApplyDeprecationBeforeUpdatePins(UPCGNode* InOutNode, TArray<TObjectPtr<UPCGPin>>& InputPins, TArray<TObjectPtr<UPCGPin>>& OutputPins);
diff --git a/Engine/Plugins/Experimental/PCG/Source/PCGEditor/Private/PCGEditor.cpp b/Engine/Plugins/Experimental/PCG/Source/PCGEditor/Private/PCGEditor.cpp
--- a/Engine/Plugins/Experimental/PCG/Source/PCGEditor/Private/PCGEditor.cpp
+++ b/Engine/Plugins/Experimental/PCG/Source/PCGEditor/Private/PCGEditor.cpp
@@ -23,6 +23,7 @@
 #include "PCGEditorGraphNode.h"
 #include "PCGEditorGraphNodeInput.h"
 #include "PCGEditorGraphNodeOutput.h"
+#include "PCGEditorGraphNodeReroute.h"
 #include "PCGEditorGraphSchema.h"
 #include "PCGEditorGraphSchemaActions.h"
 #include "PCGEditorMenuContext.h"
@@ -1125,7 +1126,14 @@ bool FPCGEditor::CanRenameNode() const
 		return false;
 	}
 
-	return (*SelectedNodes.CreateConstIterator())->IsA<UPCGEditorGraphNodeBase>();
+	if (const UPCGEditorGraphNodeBase* SelectedNode = Cast<UPCGEditorGraphNodeBase>(*SelectedNodes.CreateConstIterator()))
+	{
+		return !SelectedNode->GetPCGNode() || !SelectedNode->GetPCGNode()->GetSettings() || SelectedNode->GetPCGNode()->GetSettings()->CanUserEditTitle();
+	}
+	else
+	{
+		return false;
+	}
 }
 
 void FPCGEditor::OnCollapseNodesInSubgraph()
@@ -1229,6 +1237,12 @@ bool FPCGEditor::CanExportNodes() const
 			continue;
 		}
 
+		// Also exclude reroute nodes
+		if (Object->IsA<UPCGEditorGraphNodeReroute>() || Object->IsA<UPCGEditorGraphNodeNamedRerouteBase>())
+		{
+			continue;
+		}
+
 		if (Object->IsA<UPCGEditorGraphNodeBase>())
 		{
 			return true;
diff --git a/Engine/Plugins/Experimental/PCG/Source/PCGEditor/Private/PCGEditorGraph.cpp b/Engine/Plugins/Experimental/PCG/Source/PCGEditor/Private/PCGEditorGraph.cpp
--- a/Engine/Plugins/Experimental/PCG/Source/PCGEditor/Private/PCGEditorGraph.cpp
+++ b/Engine/Plugins/Experimental/PCG/Source/PCGEditor/Private/PCGEditorGraph.cpp
@@ -65,7 +65,24 @@ void UPCGEditorGraph::ReconstructGraph()
 			continue;
 		}
 
-		if (Cast<UPCGRerouteSettings>(PCGNode->GetSettings()))
+		// TODO: replace this with a templated lambda because this is all very similar
+		if (Cast<UPCGNamedRerouteDeclarationSettings>(PCGNode->GetSettings()))
+		{
+			FGraphNodeCreator<UPCGEditorGraphNodeNamedRerouteDeclaration> NodeCreator(*this);
+			UPCGEditorGraphNodeNamedRerouteDeclaration* RerouteGraphNode = NodeCreator.CreateNode(bSelectNewNode);
+			RerouteGraphNode->Construct(PCGNode);
+			NodeCreator.Finalize();
+			NodeLookup.Add(PCGNode, RerouteGraphNode);
+		}
+		else if (Cast<UPCGNamedRerouteUsageSettings>(PCGNode->GetSettings()))
+		{
+			FGraphNodeCreator<UPCGEditorGraphNodeNamedRerouteUsage> NodeCreator(*this);
+			UPCGEditorGraphNodeNamedRerouteUsage* RerouteGraphNode = NodeCreator.CreateNode(bSelectNewNode);
+			RerouteGraphNode->Construct(PCGNode);
+			NodeCreator.Finalize();
+			NodeLookup.Add(PCGNode, RerouteGraphNode);
+		}
+		else if (Cast<UPCGRerouteSettings>(PCGNode->GetSettings()))
 		{
 			FGraphNodeCreator<UPCGEditorGraphNodeReroute> NodeCreator(*this);
 			UPCGEditorGraphNodeReroute* RerouteGraphNode = NodeCreator.CreateNode(bSelectNewNode);
@@ -159,7 +176,7 @@ void UPCGEditorGraph::CreateLinks(UPCGEditorGraphNodeBase* GraphNode, bool bCrea
 	{
 		for (UPCGPin* InputPin : PCGNode->GetInputPins())
 		{
-			if (!InputPin)
+			if (!InputPin || InputPin->Properties.bInvisiblePin)
 			{
 				continue;
 			}
@@ -167,12 +184,8 @@ void UPCGEditorGraph::CreateLinks(UPCGEditorGraphNodeBase* GraphNode, bool bCrea
 			UEdGraphPin* InPin = GraphNode->FindPin(InputPin->Properties.Label, EEdGraphPinDirection::EGPD_Input);
 			if (!InPin)
 			{
-				if (!Cast<UPCGEditorGraphNodeInput>(GraphNode))
-				{
-					UE_LOG(LogPCGEditor, Error, TEXT("Invalid InputPin for %s"), *InputPin->Properties.Label.ToString());
-					ensure(false);
-				}
-
+				UE_LOG(LogPCGEditor, Error, TEXT("Invalid InputPin for %s"), *InputPin->Properties.Label.ToString());
+				ensure(false);
 				continue;
 			}
 
@@ -211,7 +224,7 @@ void UPCGEditorGraph::CreateLinks(UPCGEditorGraphNodeBase* GraphNode, bool bCrea
 	{
 		for (UPCGPin* OutputPin : PCGNode->GetOutputPins())
 		{
-			if (!OutputPin)
+			if (!OutputPin || OutputPin->Properties.bInvisiblePin)
 			{
 				continue;
 			}
@@ -219,11 +232,8 @@ void UPCGEditorGraph::CreateLinks(UPCGEditorGraphNodeBase* GraphNode, bool bCrea
 			UEdGraphPin* OutPin = GraphNode->FindPin(OutputPin->Properties.Label, EEdGraphPinDirection::EGPD_Output);
 			if (!OutPin)
 			{
-				if (!Cast<UPCGEditorGraphNodeOutput>(GraphNode))
-				{
-					UE_LOG(LogPCGEditor, Error, TEXT("Invalid OutputPin for %s"), *OutputPin->Properties.Label.ToString());
-				}
-
+				UE_LOG(LogPCGEditor, Error, TEXT("Invalid OutputPin for %s"), *OutputPin->Properties.Label.ToString());
+				ensure(false);
 				continue;
 			}
 
diff --git a/Engine/Plugins/Experimental/PCG/Source/PCGEditor/Private/PCGEditorGraphNodeBase.cpp b/Engine/Plugins/Experimental/PCG/Source/PCGEditor/Private/PCGEditorGraphNodeBase.cpp
--- a/Engine/Plugins/Experimental/PCG/Source/PCGEditor/Private/PCGEditorGraphNodeBase.cpp
+++ b/Engine/Plugins/Experimental/PCG/Source/PCGEditor/Private/PCGEditorGraphNodeBase.cpp
@@ -272,7 +272,20 @@ void UPCGEditorGraphNodeBase::RebuildEdgesFromPins()
 	{
 		PCGNode->GetGraph()->DisableNotificationsForEditor();
 	}
+
+	RebuildEdgesFromPins_Internal();
 	
+	if (PCGNode->GetGraph())
+	{
+		PCGNode->GetGraph()->EnableNotificationsForEditor();
+	}
+}
+
+void UPCGEditorGraphNodeBase::RebuildEdgesFromPins_Internal()
+{
+	check(PCGNode);
+	check(bDisableReconstructFromNode);
+
 	for (UEdGraphPin* Pin : Pins)
 	{
 		if (Pin->Direction == EEdGraphPinDirection::EGPD_Output)
@@ -289,11 +302,6 @@ void UPCGEditorGraphNodeBase::RebuildEdgesFromPins()
 			}
 		}
 	}
-
-	if (PCGNode->GetGraph())
-	{
-		PCGNode->GetGraph()->EnableNotificationsForEditor();
-	}
 }
 
 void UPCGEditorGraphNodeBase::OnNodeChanged(UPCGNode* InNode, EPCGChangeType ChangeType)
@@ -328,7 +336,7 @@ void UPCGEditorGraphNodeBase::OnNodeChanged(UPCGNode* InNode, EPCGChangeType Cha
 
 		if (!!(ChangeType & (EPCGChangeType::Structural | EPCGChangeType::Node | EPCGChangeType::Edge | EPCGChangeType::Cosmetic)))
 		{
-			ReconstructNode();
+			ReconstructNodeOnChange();
 		}
 	}
 }
@@ -877,7 +885,7 @@ void UPCGEditorGraphNodeBase::UpdatePosition()
 
 bool UPCGEditorGraphNodeBase::ShouldCreatePin(const UPCGPin* InPin) const
 {
-	return true;
+	return InPin && !InPin->Properties.bInvisiblePin;
 }
 
 void UPCGEditorGraphNodeBase::CreatePins(const TArray<UPCGPin*>& InInputPins, const TArray<UPCGPin*>& InOutputPins)
@@ -892,7 +900,7 @@ void UPCGEditorGraphNodeBase::CreatePins(const TArray<UPCGPin*>& InInputPins, co
 		}
 
 		UEdGraphPin* Pin = CreatePin(EEdGraphPinDirection::EGPD_Input, GetPinType(InputPin), InputPin->Properties.Label);
-		Pin->PinFriendlyName = FText::FromString(FName::NameToDisplayString(InputPin->Properties.Label.ToString(), false));
+		Pin->PinFriendlyName = GetPinFriendlyName(InputPin);
 		Pin->bAdvancedView = InputPin->Properties.bAdvancedPin;
 		bHasAdvancedPin |= Pin->bAdvancedView;
 	}
@@ -905,7 +913,7 @@ void UPCGEditorGraphNodeBase::CreatePins(const TArray<UPCGPin*>& InInputPins, co
 		}
 
 		UEdGraphPin* Pin = CreatePin(EEdGraphPinDirection::EGPD_Output, GetPinType(OutputPin), OutputPin->Properties.Label);
-		Pin->PinFriendlyName = FText::FromString(FName::NameToDisplayString(OutputPin->Properties.Label.ToString(), false));
+		Pin->PinFriendlyName = GetPinFriendlyName(OutputPin);
 		Pin->bAdvancedView = OutputPin->Properties.bAdvancedPin;
 		bHasAdvancedPin |= Pin->bAdvancedView;
 	}
@@ -920,4 +928,10 @@ void UPCGEditorGraphNodeBase::CreatePins(const TArray<UPCGPin*>& InInputPins, co
 	}
 }
 
+FText UPCGEditorGraphNodeBase::GetPinFriendlyName(const UPCGPin* InPin) const
+{
+	check(InPin);
+	return FText::FromString(FName::NameToDisplayString(InPin->Properties.Label.ToString(), false));
+}
+
 #undef LOCTEXT_NAMESPACE
diff --git a/Engine/Plugins/Experimental/PCG/Source/PCGEditor/Private/PCGEditorGraphNodeBase.h b/Engine/Plugins/Experimental/PCG/Source/PCGEditor/Private/PCGEditorGraphNodeBase.h
--- a/Engine/Plugins/Experimental/PCG/Source/PCGEditor/Private/PCGEditorGraphNodeBase.h
+++ b/Engine/Plugins/Experimental/PCG/Source/PCGEditor/Private/PCGEditorGraphNodeBase.h
@@ -110,6 +110,7 @@ protected:
 
 	/** Create PCG-side edges from editor pins/edges. */
 	void RebuildEdgesFromPins();
+	virtual void RebuildEdgesFromPins_Internal();
 
 	void OnNodeChanged(UPCGNode* InNode, EPCGChangeType ChangeType);
 	void OnPickColor();
@@ -119,10 +120,16 @@ protected:
 
 	void CreatePins(const TArray<UPCGPin*>& InInputPins, const TArray<UPCGPin*>& InOutputPins);
 
+	// Performs potentially custom logic when there's a change that would trigger a reconstruct (needed for linked nodes like the named reroutes)
+	virtual void ReconstructNodeOnChange() { ReconstructNode(); }
+
 	// Custom logic to hide some pins to the user (by not creating a UI pin, even if the model pin exists).
 	// Useful for deprecation
 	virtual bool ShouldCreatePin(const UPCGPin* InPin) const;
 
+	// Returns the appropriate pin name to allow for some flexibility
+	virtual FText GetPinFriendlyName(const UPCGPin* InPin) const;
+
 	UPROPERTY()
 	TObjectPtr<UPCGNode> PCGNode = nullptr;
 
diff --git a/Engine/Plugins/Experimental/PCG/Source/PCGEditor/Private/PCGEditorGraphNodeFactory.cpp b/Engine/Plugins/Experimental/PCG/Source/PCGEditor/Private/PCGEditorGraphNodeFactory.cpp
--- a/Engine/Plugins/Experimental/PCG/Source/PCGEditor/Private/PCGEditorGraphNodeFactory.cpp
+++ b/Engine/Plugins/Experimental/PCG/Source/PCGEditor/Private/PCGEditorGraphNodeFactory.cpp
@@ -39,7 +39,11 @@ TSharedPtr<SGraphNode> FPCGEditorGraphNodeFactory::CreateNode(UEdGraphNode* InNo
 		TSharedPtr<SGraphNode> VisualNode;
 		
 		const UPCGNode* PCGNode = GraphNode->GetPCGNode();
-		if (PCGNode && Cast<UPCGRerouteSettings>(PCGNode->GetSettings()))
+		if (PCGNode && Cast<UPCGNamedRerouteBaseSettings>(PCGNode->GetSettings()))
+		{
+			SAssignNew(VisualNode, SPCGEditorGraphVarNode, GraphNode);
+		}
+		else if (PCGNode && Cast<UPCGRerouteSettings>(PCGNode->GetSettings()))
 		{
 			SAssignNew(VisualNode, SPCGEditorGraphNodeKnot, GraphNode);
 		}
@@ -56,11 +60,6 @@ TSharedPtr<SGraphNode> FPCGEditorGraphNodeFactory::CreateNode(UEdGraphNode* InNo
 
 		return VisualNode.ToSharedRef();
 	}
-	else if (UPCGEditorGraphNodeReroute* RerouteNode = Cast<UPCGEditorGraphNodeReroute>(InNode))
-	{
-		TSharedRef<SGraphNode> VisualNode = SNew(SGraphNodeKnot, RerouteNode);
-		return VisualNode;
-	}
 
 	return nullptr;
 }
diff --git a/Engine/Plugins/Experimental/PCG/Source/PCGEditor/Private/PCGEditorGraphNodeReroute.cpp b/Engine/Plugins/Experimental/PCG/Source/PCGEditor/Private/PCGEditorGraphNodeReroute.cpp
--- a/Engine/Plugins/Experimental/PCG/Source/PCGEditor/Private/PCGEditorGraphNodeReroute.cpp
+++ b/Engine/Plugins/Experimental/PCG/Source/PCGEditor/Private/PCGEditorGraphNodeReroute.cpp
@@ -2,6 +2,11 @@
 
 #include "PCGEditorGraphNodeReroute.h"
 
+#include "PCGEditorGraph.h"
+
+#include "PCGNode.h"
+#include "Elements/PCGReroute.h"
+
 #include "EdGraph/EdGraphPin.h"
 
 FText UPCGEditorGraphNodeReroute::GetNodeTitle(ENodeTitleType::Type TitleType) const
@@ -73,3 +78,112 @@ UEdGraphPin* UPCGEditorGraphNodeReroute::GetOutputPin() const
 
 	return nullptr;
 }
+
+FText UPCGEditorGraphNodeNamedRerouteBase::GetNodeTitle(ENodeTitleType::Type TitleType) const
+{
+	// Never show full title, only list view
+	return Super::GetNodeTitle(ENodeTitleType::ListView);
+}
+
+FText UPCGEditorGraphNodeNamedRerouteUsage::GetPinFriendlyName(const UPCGPin* InPin) const
+{
+	return FText::FromString(" ");
+}
+
+void UPCGEditorGraphNodeNamedRerouteUsage::RebuildEdgesFromPins_Internal()
+{
+	Super::RebuildEdgesFromPins_Internal();
+
+	check(PCGNode);
+
+	if (PCGNode->HasInboundEdges())
+	{
+		return;
+	}
+
+	UPCGGraph* Graph = PCGNode->GetGraph();
+
+	if (!Graph)
+	{
+		return;
+	}
+
+	UPCGNamedRerouteUsageSettings* Usage = Cast<UPCGNamedRerouteUsageSettings>(PCGNode->GetSettings());
+
+	if (!Usage)
+	{
+		return;
+	}
+
+	// Make sure we're hooked to the declaration if it's not already the case
+	if (UPCGNode* DeclarationNode = Graph->FindNodeWithSettings(Usage->Declaration))
+	{
+		DeclarationNode->AddEdgeTo(PCGNamedRerouteConstants::InvisiblePinLabel, PCGNode, PCGPinConstants::DefaultInputLabel);
+	}
+}
+
+FText UPCGEditorGraphNodeNamedRerouteDeclaration::GetPinFriendlyName(const UPCGPin* InPin) const
+{
+	return FText::FromString(" ");
+}
+
+void UPCGEditorGraphNodeNamedRerouteDeclaration::ApplyToUsageNodes(TFunctionRef<void(UPCGEditorGraphNodeNamedRerouteUsage*)> Action)
+{
+	if (!GetPCGNode())
+	{
+		return;
+	}
+
+	const UPCGNamedRerouteDeclarationSettings* Declaration = Cast<UPCGNamedRerouteDeclarationSettings>(GetPCGNode()->GetSettings());
+
+	if (!Declaration)
+	{
+		return;
+	}
+
+	UPCGEditorGraph* EditorGraph = Cast<UPCGEditorGraph>(GetGraph());
+
+	if (!EditorGraph)
+	{
+		return;
+	}
+
+	for (const TObjectPtr<UEdGraphNode>& EdGraphNode : EditorGraph->Nodes)
+	{
+		if (UPCGEditorGraphNodeNamedRerouteUsage* RerouteNode = Cast<UPCGEditorGraphNodeNamedRerouteUsage>(EdGraphNode))
+		{
+			if (RerouteNode->GetPCGNode() &&
+				Cast<UPCGNamedRerouteUsageSettings>(RerouteNode->GetPCGNode()->GetSettings()) &&
+				Cast<UPCGNamedRerouteUsageSettings>(RerouteNode->GetPCGNode()->GetSettings())->Declaration == Declaration)
+			{
+				Action(RerouteNode);
+			}
+		}
+	}
+}
+
+void UPCGEditorGraphNodeNamedRerouteDeclaration::ReconstructNodeOnChange()
+{
+	Super::ReconstructNodeOnChange();
+
+	// We must make sure to trigger a notify node changed on all editor nodes that are usages of that declaration
+	ApplyToUsageNodes([](UPCGEditorGraphNodeNamedRerouteUsage* RerouteNode)
+	{
+		RerouteNode->ReconstructNodeOnChange();
+	});
+}
+
+void UPCGEditorGraphNodeNamedRerouteDeclaration::OnRenameNode(const FString& NewName)
+{
+	Super::OnRenameNode(NewName);
+
+	// Propagate the name change to downstream usage nodes
+	ApplyToUsageNodes([&NewName](UPCGEditorGraphNodeNamedRerouteUsage* RerouteNode)
+	{
+		RerouteNode->bCanRenameNode = true;
+		RerouteNode->OnRenameNode(NewName);
+		RerouteNode->bCanRenameNode = false;
+	});
+
+	ReconstructNodeOnChange();
+}
\ No newline at end of file
diff --git a/Engine/Plugins/Experimental/PCG/Source/PCGEditor/Private/PCGEditorGraphNodeReroute.h b/Engine/Plugins/Experimental/PCG/Source/PCGEditor/Private/PCGEditorGraphNodeReroute.h
--- a/Engine/Plugins/Experimental/PCG/Source/PCGEditor/Private/PCGEditorGraphNodeReroute.h
+++ b/Engine/Plugins/Experimental/PCG/Source/PCGEditor/Private/PCGEditorGraphNodeReroute.h
@@ -25,3 +25,40 @@ public:
 	UEdGraphPin* GetInputPin() const;
 	UEdGraphPin* GetOutputPin() const;
 };
+
+UCLASS()
+class UPCGEditorGraphNodeNamedRerouteBase : public UPCGEditorGraphNode
+{
+	GENERATED_BODY()
+
+public:
+	// ~Begin UEdGraphNode interface
+	virtual FText GetNodeTitle(ENodeTitleType::Type TitleType) const override;
+	// ~End UEdGraphNode interface
+};
+
+UCLASS()
+class UPCGEditorGraphNodeNamedRerouteUsage : public UPCGEditorGraphNodeNamedRerouteBase
+{
+	GENERATED_BODY()
+	friend class UPCGEditorGraphNodeNamedRerouteDeclaration;
+
+protected:
+	virtual void RebuildEdgesFromPins_Internal();
+	virtual FText GetPinFriendlyName(const UPCGPin* InPin) const override;
+};
+
+UCLASS()
+class UPCGEditorGraphNodeNamedRerouteDeclaration : public UPCGEditorGraphNodeNamedRerouteBase
+{
+	GENERATED_BODY()
+
+public:
+	virtual void OnRenameNode(const FString& NewName) override;
+
+protected:
+	virtual FText GetPinFriendlyName(const UPCGPin* InPin) const override;
+	virtual void ReconstructNodeOnChange() override;
+
+	void ApplyToUsageNodes(TFunctionRef<void(UPCGEditorGraphNodeNamedRerouteUsage*)> Action);
+};
\ No newline at end of file
diff --git a/Engine/Plugins/Experimental/PCG/Source/PCGEditor/Private/PCGEditorGraphSchema.cpp b/Engine/Plugins/Experimental/PCG/Source/PCGEditor/Private/PCGEditorGraphSchema.cpp
--- a/Engine/Plugins/Experimental/PCG/Source/PCGEditor/Private/PCGEditorGraphSchema.cpp
+++ b/Engine/Plugins/Experimental/PCG/Source/PCGEditor/Private/PCGEditorGraphSchema.cpp
@@ -9,6 +9,7 @@
 #include "Elements/PCGExecuteBlueprint.h"
 #include "Elements/PCGFilterByType.h"
 #include "Elements/PCGMakeConcreteElement.h"
+#include "Elements/PCGReroute.h"
 #include "Elements/PCGUserParameterGet.h"
 
 #include "PCGEditorCommon.h"
@@ -50,6 +51,7 @@ void UPCGEditorGraphSchema::GetPaletteActions(FGraphActionMenuBuilder& ActionMen
 	}
 	if (!!(InPCGElementTypeFilter & EPCGElementType::Other))
 	{
+		GetNamedRerouteUsageActions(ActionMenuBuilder);
 		GetExtraElementActions(ActionMenuBuilder);
 	}
 }
@@ -62,6 +64,7 @@ void UPCGEditorGraphSchema::GetGraphContextActions(FGraphContextMenuBuilder& Con
 	GetSubgraphElementActions(ContextMenuBuilder);
 	GetBlueprintElementActions(ContextMenuBuilder);
 	GetSettingsElementActions(ContextMenuBuilder, /*bIsContextual=*/true);
+	GetNamedRerouteUsageActions(ContextMenuBuilder, ContextMenuBuilder.CurrentGraph);
 	GetExtraElementActions(ContextMenuBuilder);
 }
 
@@ -518,21 +521,57 @@ void UPCGEditorGraphSchema::GetSubgraphElementActions(FGraphActionMenuBuilder& A
 
 void UPCGEditorGraphSchema::GetExtraElementActions(FGraphActionMenuBuilder& ActionMenuBuilder) const
 {
+	const FText NoCategory;
 	// Comment action
 	const FText CommentMenuDesc = LOCTEXT("PCGAddComment", "Add Comment...");
-	const FText CommentCategory;
 	const FText CommentDescription = LOCTEXT("PCGAddCommentTooltip", "Create a resizable comment box.");
 
-	const TSharedPtr<FPCGEditorGraphSchemaAction_NewComment> NewCommentAction(new FPCGEditorGraphSchemaAction_NewComment(CommentCategory, CommentMenuDesc, CommentDescription, 0));
+	const TSharedPtr<FPCGEditorGraphSchemaAction_NewComment> NewCommentAction(new FPCGEditorGraphSchemaAction_NewComment(NoCategory, CommentMenuDesc, CommentDescription, 0));
 	ActionMenuBuilder.AddAction(NewCommentAction);
 
 	// Reroute action
 	const FText RerouteMenuDesc = LOCTEXT("PCGAddRerouteNode", "Add Reroute Node");
-	const FText RerouteCategory;
 	const FText RerouteDescription = LOCTEXT("PCGAddRerouteNodeTooltip", "Add a reroute node, aka knot.");
 
-	const TSharedPtr<FPCGEditorGraphSchemaAction_NewReroute> NewRerouteAction(new FPCGEditorGraphSchemaAction_NewReroute(RerouteCategory, RerouteMenuDesc, RerouteDescription, 0));
+	const TSharedPtr<FPCGEditorGraphSchemaAction_NewReroute> NewRerouteAction(new FPCGEditorGraphSchemaAction_NewReroute(NoCategory, RerouteMenuDesc, RerouteDescription, 0));
 	ActionMenuBuilder.AddAction(NewRerouteAction);
+
+	// Named reroute declaration action
+	const FText NamedRerouteMenuDesc = LOCTEXT("PCGAddNamedRerouteDeclarationNode", "Add Named Reroute Declaration Node...");
+	const FText NamedRerouteDescription = LOCTEXT("PCGAddNamedRerouteDeclarationNodeTooltip", "Creates a new Named Reroute Declaration from the input.");
+
+	const TSharedPtr<FPCGEditorGraphSchemaAction_NewNamedRerouteDeclaration> NewNamedRerouteDeclarationAction(new FPCGEditorGraphSchemaAction_NewNamedRerouteDeclaration(NoCategory, NamedRerouteMenuDesc, NamedRerouteDescription, 0));
+	ActionMenuBuilder.AddAction(NewNamedRerouteDeclarationAction);
+}
+
+void UPCGEditorGraphSchema::GetNamedRerouteUsageActions(FGraphActionMenuBuilder& ActionMenuBuilder, const UEdGraph* CurrentGraph) const
+{
+	const UPCGEditorGraph* Graph = Cast<UPCGEditorGraph>(CurrentGraph);
+
+	if (!Graph)
+	{
+		return;
+	}
+
+	const UPCGGraph* PCGGraph = const_cast<UPCGEditorGraph*>(Graph)->GetPCGGraph();
+
+	if (!PCGGraph)
+	{
+		return;
+	}
+
+	for (const UPCGNode* Node : PCGGraph->GetNodes())
+	{
+		if (const UPCGNamedRerouteDeclarationSettings* RerouteDeclaration = Cast<UPCGNamedRerouteDeclarationSettings>(Node->GetSettings()))
+		{
+			static const FText Category = LOCTEXT("NamedRerouteCategory", "Named Reroutes");
+			const FText Name = FText::FromName(Node->NodeTitle);
+			const FText Tooltip = FText::Format(LOCTEXT("NamedRerouteTooltip", "Add a usage of '{0}' here."), Name);
+			TSharedPtr<FPCGEditorGraphSchemaAction_NewNamedRerouteUsage> NewRerouteAction(new FPCGEditorGraphSchemaAction_NewNamedRerouteUsage(Category, Name, Tooltip, 1 /* We want named reroutes to be on top */));
+			NewRerouteAction->DeclarationNode = Node;
+			ActionMenuBuilder.AddAction(NewRerouteAction);
+		}
+	}
 }
 
 void UPCGEditorGraphSchema::DroppedAssetsOnGraph(const TArray<FAssetData>& Assets, const FVector2D& GraphPosition, UEdGraph* Graph) const
diff --git a/Engine/Plugins/Experimental/PCG/Source/PCGEditor/Private/PCGEditorGraphSchema.h b/Engine/Plugins/Experimental/PCG/Source/PCGEditor/Private/PCGEditorGraphSchema.h
--- a/Engine/Plugins/Experimental/PCG/Source/PCGEditor/Private/PCGEditorGraphSchema.h
+++ b/Engine/Plugins/Experimental/PCG/Source/PCGEditor/Private/PCGEditorGraphSchema.h
@@ -38,6 +38,7 @@ private:
 	void GetSubgraphElementActions(FGraphActionMenuBuilder& ActionMenuBuilder) const;
 	void GetSettingsElementActions(FGraphActionMenuBuilder& ActionMenuBuilder, bool bIsContextual) const;
 	void GetExtraElementActions(FGraphActionMenuBuilder& ActionMenuBuilder) const;
+	void GetNamedRerouteUsageActions(FGraphActionMenuBuilder& ActionMenuBuilder, const UEdGraph* CurrentGraph = nullptr) const;
 
 	virtual bool TryCreateConnectionInternal(UEdGraphPin* A, UEdGraphPin* B, bool bAddConversionNodeIfNeeded) const;
 };
diff --git a/Engine/Plugins/Experimental/PCG/Source/PCGEditor/Private/PCGEditorGraphSchemaActions.cpp b/Engine/Plugins/Experimental/PCG/Source/PCGEditor/Private/PCGEditorGraphSchemaActions.cpp
--- a/Engine/Plugins/Experimental/PCG/Source/PCGEditor/Private/PCGEditorGraphSchemaActions.cpp
+++ b/Engine/Plugins/Experimental/PCG/Source/PCGEditor/Private/PCGEditorGraphSchemaActions.cpp
@@ -399,7 +399,7 @@ UEdGraphNode* FPCGEditorGraphSchemaAction_NewReroute::PerformAction(class UEdGra
 		return nullptr;
 	}
 
-	const FScopedTransaction Transaction(*FPCGEditorCommon::ContextIdentifier, LOCTEXT("PCGEditorNewReroute", "PCG Editor: New Reorute Node"), nullptr);
+	const FScopedTransaction Transaction(*FPCGEditorCommon::ContextIdentifier, LOCTEXT("PCGEditorNewReroute", "PCG Editor: New Reroute Node"), nullptr);
 	EditorGraph->Modify();
 
 	UPCGSettings* DefaultNodeSettings = nullptr;
@@ -423,4 +423,128 @@ UEdGraphNode* FPCGEditorGraphSchemaAction_NewReroute::PerformAction(class UEdGra
 	return NewNode;
 }
 
+UEdGraphNode* FPCGEditorGraphSchemaAction_NewNamedRerouteUsage::PerformAction(class UEdGraph* ParentGraph, UEdGraphPin* FromPin, const FVector2D Location, bool bSelectNewNode)
+{
+	check(DeclarationNode);
+
+	UPCGEditorGraph* EditorGraph = Cast<UPCGEditorGraph>(ParentGraph);
+	if (!EditorGraph)
+	{
+		UE_LOG(LogPCGEditor, Error, TEXT("Invalid EditorGraph"));
+		return nullptr;
+	}
+
+	UPCGGraph* PCGGraph = EditorGraph->GetPCGGraph();
+	if (!PCGGraph)
+	{
+		UE_LOG(LogPCGEditor, Error, TEXT("Invalid PCGGraph"));
+		return nullptr;
+	}
+
+	const FScopedTransaction Transaction(*FPCGEditorCommon::ContextIdentifier, LOCTEXT("PCGEditorNewNamedRerouteUsage", "PCG Editor: New Named Reroute Node Usage"), nullptr);
+	EditorGraph->Modify();
+
+	PCGGraph->DisableNotificationsForEditor();
+
+	ON_SCOPE_EXIT
+	{
+		PCGGraph->EnableNotificationsForEditor();
+	};
+
+	UPCGNamedRerouteDeclarationSettings* Declaration = CastChecked<UPCGNamedRerouteDeclarationSettings>(DeclarationNode->GetSettings());
+
+	UPCGSettings* DefaultNodeSettings = nullptr;
+	UPCGNode* NewPCGNode = PCGGraph->AddNodeOfType(UPCGNamedRerouteUsageSettings::StaticClass(), DefaultNodeSettings);
+
+	if (!NewPCGNode || !DefaultNodeSettings)
+	{
+		UE_LOG(LogPCGEditor, Error, TEXT("Failed creating named reroute node."));
+		return nullptr;
+	}
+
+	NewPCGNode->UpdateAfterSettingsChangeDuringCreation();
+
+	CastChecked<UPCGNamedRerouteUsageSettings>(DefaultNodeSettings)->Declaration = Declaration;
+	NewPCGNode->NodeTitle = DeclarationNode->NodeTitle;
+
+	// Create edge from the declaration to the new node
+	PCGGraph->AddEdge(const_cast<UPCGNode*>(DeclarationNode.Get()), PCGNamedRerouteConstants::InvisiblePinLabel, NewPCGNode, PCGPinConstants::DefaultInputLabel);
+	
+	FGraphNodeCreator<UPCGEditorGraphNodeNamedRerouteUsage> NodeCreator(*EditorGraph);
+	UPCGEditorGraphNodeNamedRerouteUsage* NewNode = NodeCreator.CreateUserInvokedNode(bSelectNewNode);
+	NewNode->Construct(NewPCGNode);
+	NewNode->NodePosX = Location.X;
+	NewNode->NodePosY = Location.Y;
+	NodeCreator.Finalize();
+
+	NewPCGNode->PositionX = Location.X;
+	NewPCGNode->PositionY = Location.Y;
+
+	return NewNode;
+}
+
+UEdGraphNode* FPCGEditorGraphSchemaAction_NewNamedRerouteDeclaration::PerformAction(class UEdGraph* ParentGraph, UEdGraphPin* FromPin, const FVector2D Location, bool bSelectNewNode)
+{
+	UPCGEditorGraph* EditorGraph = Cast<UPCGEditorGraph>(ParentGraph);
+	if (!EditorGraph)
+	{
+		UE_LOG(LogPCGEditor, Error, TEXT("Invalid EditorGraph"));
+		return nullptr;
+	}
+
+	UPCGGraph* PCGGraph = EditorGraph->GetPCGGraph();
+	if (!PCGGraph)
+	{
+		UE_LOG(LogPCGEditor, Error, TEXT("Invalid PCGGraph"));
+		return nullptr;
+	}
+
+	const FScopedTransaction Transaction(*FPCGEditorCommon::ContextIdentifier, LOCTEXT("PCGEditorNewNamedRerouteDeclaration", "PCG Editor: New Named Reroute Node Declaration"), nullptr);
+	EditorGraph->Modify();
+
+	UPCGSettings* DefaultNodeSettings = nullptr;
+	UPCGNode* NewPCGNode = PCGGraph->AddNodeOfType(UPCGNamedRerouteDeclarationSettings::StaticClass(), DefaultNodeSettings);
+
+	if (!NewPCGNode || !DefaultNodeSettings)
+	{
+		UE_LOG(LogPCGEditor, Error, TEXT("Failed creating named reroute node."));
+		return nullptr;
+	}
+
+	NewPCGNode->UpdateAfterSettingsChangeDuringCreation();
+
+	FGraphNodeCreator<UPCGEditorGraphNodeNamedRerouteDeclaration> NodeCreator(*EditorGraph);
+	UPCGEditorGraphNodeNamedRerouteDeclaration* NewNode = NodeCreator.CreateUserInvokedNode(bSelectNewNode);
+	NewNode->Construct(NewPCGNode);
+	NewNode->NodePosX = Location.X;
+	NewNode->NodePosY = Location.Y;
+	NodeCreator.Finalize();
+
+	NewPCGNode->PositionX = Location.X;
+	NewPCGNode->PositionY = Location.Y;
+
+	NewPCGNode->NodeTitle = TEXT("Reroute");
+
+	if (FromPin)
+	{
+		if (UPCGNamedRerouteDeclarationSettings* Declaration = Cast<UPCGNamedRerouteDeclarationSettings>(DefaultNodeSettings))
+		{
+			UPCGNode* FromNode = CastChecked<UPCGEditorGraphNodeBase>(FromPin->GetOwningNode())->GetPCGNode();
+
+			if (FromNode)
+			{
+				NewPCGNode->NodeTitle = FName(FromNode->GetNodeTitle(EPCGNodeTitleType::ListView).ToString() + " " + FromPin->PinName.ToString());
+			}
+			else
+			{
+				NewPCGNode->NodeTitle = FromPin->PinName;
+			}
+		}
+
+		NewNode->AutowireNewNode(FromPin);
+	}
+
+	return NewNode;
+}
+
 #undef LOCTEXT_NAMESPACE
diff --git a/Engine/Plugins/Experimental/PCG/Source/PCGEditor/Private/PCGEditorGraphSchemaActions.h b/Engine/Plugins/Experimental/PCG/Source/PCGEditor/Private/PCGEditorGraphSchemaActions.h
--- a/Engine/Plugins/Experimental/PCG/Source/PCGEditor/Private/PCGEditorGraphSchemaActions.h
+++ b/Engine/Plugins/Experimental/PCG/Source/PCGEditor/Private/PCGEditorGraphSchemaActions.h
@@ -195,3 +195,49 @@ struct FPCGEditorGraphSchemaAction_NewReroute : public FEdGraphSchemaAction
 	virtual UEdGraphNode* PerformAction(class UEdGraph* ParentGraph, UEdGraphPin* FromPin, const FVector2D Location, bool bSelectNewNode = true) override;
 	// End of FEdGraphSchemaAction interface
 };
+
+USTRUCT()
+struct FPCGEditorGraphSchemaAction_NewNamedRerouteUsage : public FEdGraphSchemaAction
+{
+	GENERATED_BODY()
+
+	// Inherit the base class's constructors
+	using FEdGraphSchemaAction::FEdGraphSchemaAction;
+
+	// Declaration that we want to add an usage of
+	UPROPERTY()
+	TObjectPtr<const UPCGNode> DeclarationNode = nullptr;
+
+	// Simple type info
+	static FName StaticGetTypeId()
+	{
+		static FName Type("FPCGEditorGraphSchemaAction_NewNamedRerouteUsage");
+		return Type;
+	}
+
+	// FEdGraphSchemaAction interface
+	virtual FName GetTypeId() const override { return StaticGetTypeId(); }
+	virtual UEdGraphNode* PerformAction(class UEdGraph* ParentGraph, UEdGraphPin* FromPin, const FVector2D Location, bool bSelectNewNode = true) override;
+	// End of FEdGraphSchemaAction interface
+};
+
+USTRUCT()
+struct FPCGEditorGraphSchemaAction_NewNamedRerouteDeclaration : public FEdGraphSchemaAction
+{
+	GENERATED_BODY()
+
+	// Inherit the base class's constructors
+	using FEdGraphSchemaAction::FEdGraphSchemaAction;
+
+	// Simple type info
+	static FName StaticGetTypeId()
+	{
+		static FName Type("FPCGEditorGraphSchemaAction_NewNamedRerouteDeclaration");
+		return Type;
+	}
+
+	// FEdGraphSchemaAction interface
+	virtual FName GetTypeId() const override { return StaticGetTypeId(); }
+	virtual UEdGraphNode* PerformAction(class UEdGraph* ParentGraph, UEdGraphPin* FromPin, const FVector2D Location, bool bSelectNewNode = true) override;
+	// End of FEdGraphSchemaAction interface
+};
\ No newline at end of file
diff --git a/Engine/Plugins/Experimental/PCG/Source/PCGEditor/Private/PCGEditorSettings.cpp b/Engine/Plugins/Experimental/PCG/Source/PCGEditor/Private/PCGEditorSettings.cpp
--- a/Engine/Plugins/Experimental/PCG/Source/PCGEditor/Private/PCGEditorSettings.cpp
+++ b/Engine/Plugins/Experimental/PCG/Source/PCGEditor/Private/PCGEditorSettings.cpp
@@ -26,6 +26,8 @@ UPCGEditorSettings::UPCGEditorSettings(const FObjectInitializer& ObjectInitializ
 	ControlFlowNodeColor = FLinearColor(0.0f, 1.0f, 0.0f);
 	PointOpsNodeColor = FLinearColor(0.0f, 0.04f, 0.23f);
 	HierarchicalGenerationNodeColor = FLinearColor(1.0f, 0.132868f, 0.0f);
+	GraphParametersNodeColor = FLinearColor::Yellow;
+	RerouteNodeColor = FLinearColor(0.5f, 1.0f, 0.83f);
 
 	DefaultPinColor = FLinearColor(0.29f, 0.29f, 0.29f);
 	SpatialDataPinColor = FLinearColor(1.0f, 1.0f, 1.0f);
@@ -89,6 +91,10 @@ FLinearColor UPCGEditorSettings::GetColor(UPCGSettings* Settings) const
 			return ControlFlowNodeColor;
 		case EPCGSettingsType::PointOps:
 			return PointOpsNodeColor;
+		case EPCGSettingsType::GraphParameters:
+			return GraphParametersNodeColor;
+		case EPCGSettingsType::Reroute:
+			return RerouteNodeColor;
 		case EPCGSettingsType::Generic: // falls through
 		default:
 			// Finally, we couldn't find any match, so return the default value
diff --git a/Engine/Plugins/Experimental/PCG/Source/PCGEditor/Private/SPCGEditorGraphVarNode.cpp b/Engine/Plugins/Experimental/PCG/Source/PCGEditor/Private/SPCGEditorGraphVarNode.cpp
--- a/Engine/Plugins/Experimental/PCG/Source/PCGEditor/Private/SPCGEditorGraphVarNode.cpp
+++ b/Engine/Plugins/Experimental/PCG/Source/PCGEditor/Private/SPCGEditorGraphVarNode.cpp
@@ -6,6 +6,8 @@
 
 #include "SPCGEditorGraphVarNode.h"
 
+#include "PCGEditorGraphNodeReroute.h"
+
 #include "GraphEditorSettings.h"
 #include "SCommentBubble.h"
 #include "SGraphNode.h"
@@ -26,7 +28,7 @@ void SPCGEditorGraphVarNode::Construct(const FArguments& InArgs, UPCGEditorGraph
 
 FSlateColor SPCGEditorGraphVarNode::GetVariableColor() const
 {
-	return FColor::Yellow;
+	return GraphNode->GetNodeTitleColor();
 }
 
 void SPCGEditorGraphVarNode::UpdateGraphNode()
@@ -42,6 +44,27 @@ void SPCGEditorGraphVarNode::UpdateGraphNode()
 
 	SetupErrorReporting();
 
+	bool bNeedsTitle = (Cast<UPCGEditorGraphNodeNamedRerouteBase>(GraphNode) != nullptr);
+
+	TSharedPtr<SWidget> TitleArea = SNullWidget::NullWidget.ToSharedPtr();
+
+	if (bNeedsTitle)
+	{
+		TSharedPtr<SNodeTitle> NodeTitle = SNew(SNodeTitle, GraphNode);
+		SAssignNew(TitleArea, SVerticalBox)
+			+ SVerticalBox::Slot()
+			.AutoHeight()
+			.VAlign(VAlign_Center)
+			[
+				CreateTitleWidget(NodeTitle)
+			]
+			+ SVerticalBox::Slot()
+			.AutoHeight()
+			[
+				NodeTitle.ToSharedRef()
+			];
+	}
+
 	// Setup a meta tag for this node
 	FGraphNodeMetaData TagMeta(TEXT("Graphnode"));
 	PopulateMetaTag(&TagMeta);
@@ -96,6 +119,13 @@ void SPCGEditorGraphVarNode::UpdateGraphNode()
 				]
 				+ SHorizontalBox::Slot()
 				.AutoWidth()
+				.HAlign(HAlign_Center)
+				.VAlign(VAlign_Center)
+				[
+					TitleArea.ToSharedRef()
+				]
+				+ SHorizontalBox::Slot()
+				.AutoWidth()
 				.HAlign(HAlign_Right)
 				.Padding( FMargin(2,0) )
 				[
diff --git a/Engine/Plugins/Experimental/PCG/Source/PCGEditor/Public/PCGEditorSettings.h b/Engine/Plugins/Experimental/PCG/Source/PCGEditor/Public/PCGEditorSettings.h
--- a/Engine/Plugins/Experimental/PCG/Source/PCGEditor/Public/PCGEditorSettings.h
+++ b/Engine/Plugins/Experimental/PCG/Source/PCGEditor/Public/PCGEditorSettings.h
@@ -89,6 +89,14 @@ public:
 	UPROPERTY(EditAnywhere, config, Category = Node, meta = (HideAlphaChannel))
 	FLinearColor HierarchicalGenerationNodeColor;
 
+	/** Color used for graph parameters operations */
+	UPROPERTY(EditAnywhere, config, Category = Node, meta = (HideAlphaChannel))
+	FLinearColor GraphParametersNodeColor;
+
+	/** Color used for reroute nodes */
+	UPROPERTY(EditAnywhere, config, Category = Node, meta = (HideAlphaChannel))
+	FLinearColor RerouteNodeColor;
+
 	/** Default pin color */
 	UPROPERTY(EditAnywhere, config, Category = Node, meta = (HideAlphaChannel))
 	FLinearColor DefaultPinColor;
