[PCG] Fix wrong EditCondition in FPCGWorldCommonQueryParams

#rnx

diff --git a/Engine/Plugins/Experimental/PCG/Source/PCG/Public/Data/PCGWorldData.h b/Engine/Plugins/Experimental/PCG/Source/PCG/Public/Data/PCGWorldData.h
--- a/Engine/Plugins/Experimental/PCG/Source/PCG/Public/Data/PCGWorldData.h
+++ b/Engine/Plugins/Experimental/PCG/Source/PCG/Public/Data/PCGWorldData.h
@@ -50,7 +50,7 @@ struct FPCGWorldCommonQueryParams
 	UPROPERTY(BlueprintReadWrite, EditAnywhere, Category = "Data|Filtering", meta = (PCG_Overridable))
 	EPCGWorldQueryFilterByTag ActorTagFilter = EPCGWorldQueryFilterByTag::NoTagFilter;
 
-	UPROPERTY(BlueprintReadWrite, EditAnywhere, Category = "Data|Filtering", meta = (PCG_Overridable, EditCondition = "TagFilter != EPCGWorldQueryFilterByTag::NoTagFilter"))
+	UPROPERTY(BlueprintReadWrite, EditAnywhere, Category = "Data|Filtering", meta = (PCG_Overridable, EditCondition = "ActorTagFilter != EPCGWorldQueryFilterByTag::NoTagFilter"))
 	FString ActorTagsList;
 
 	UPROPERTY(BlueprintReadWrite, EditAnywhere, Category = "Data|Filtering", meta = (PCG_Overridable))
