[PCG] Replace OnPreObjectPropertyChanged hook to OnObjectModified to catch all modifying flow for tracking

* Since it gather more actors than consumed, also added a cleaning mecanism, with timing controllable with CVars.

#rnx

diff --git a/Engine/Plugins/Experimental/PCG/Source/PCG/Private/Elements/PCGActorSelector.cpp b/Engine/Plugins/Experimental/PCG/Source/PCG/Private/Elements/PCGActorSelector.cpp
--- a/Engine/Plugins/Experimental/PCG/Source/PCG/Private/Elements/PCGActorSelector.cpp
+++ b/Engine/Plugins/Experimental/PCG/Source/PCG/Private/Elements/PCGActorSelector.cpp
@@ -297,7 +297,7 @@ bool FPCGSelectionKey::IsMatching(const UObject* InObject, const UPCGComponent*
 	}
 }
 
-bool FPCGSelectionKey::IsMatching(const UObject* InObject, const TSet<FName>& InRemovedTags, const TSet<UPCGComponent*>& InComponents, TSet<UPCGComponent*>& MatchedComponents) const
+bool FPCGSelectionKey::IsMatching(const UObject* InObject, const TSet<FName>& InRemovedTags, const TSet<UPCGComponent*>& InComponents, TSet<UPCGComponent*>* OptionalMatchedComponents) const
 {
 	if (!InObject)
 	{
@@ -343,8 +343,15 @@ bool FPCGSelectionKey::IsMatching(const UObject* InObject, const TSet<FName>& In
 			{
 				if (InComponents.Contains(PCGComponent))
 				{
-					MatchedComponents.Add(PCGComponent);
 					bFoundMatch = true;
+					if (OptionalMatchedComponents)
+					{
+						OptionalMatchedComponents->Add(PCGComponent);
+					}
+					else
+					{
+						break;
+					}
 				}
 			}
 		}
@@ -372,9 +379,9 @@ bool FPCGSelectionKey::IsMatching(const UObject* InObject, const TSet<FName>& In
 		break;
 	}
 
-	if (bIsMatched)
+	if (bIsMatched && OptionalMatchedComponents)
 	{
-		MatchedComponents.Append(InComponents);
+		OptionalMatchedComponents->Append(InComponents);
 	}
 
 	return bIsMatched;
diff --git a/Engine/Plugins/Experimental/PCG/Source/PCG/Private/PCGActorAndComponentMapping.cpp b/Engine/Plugins/Experimental/PCG/Source/PCG/Private/PCGActorAndComponentMapping.cpp
--- a/Engine/Plugins/Experimental/PCG/Source/PCG/Private/PCGActorAndComponentMapping.cpp
+++ b/Engine/Plugins/Experimental/PCG/Source/PCG/Private/PCGActorAndComponentMapping.cpp
@@ -71,6 +71,16 @@ namespace PCGActorAndComponentMapping
 		TEXT("pcg.LandscapeRefreshTimeDelayMS"),
 		1000,
 		TEXT("Time in MS between a landscape change and PCG refresh. Set it to 0 or negative value to disable the delay."));
+
+	static TAutoConsoleVariable<int> CVarActorModifiedPreviousDataCleanupDelayMS(
+		TEXT("pcg.ActorModifiedPreviousDataCleanupDelayMS"),
+		10000,
+		TEXT("Time in MS between a cached previous data from modified actor and the time we remove it from our cache."));
+
+	static TAutoConsoleVariable<int> CVarActorModifiedPreviousDataCheckDelayMS(
+		TEXT("pcg.ActorModifiedPreviousDataCheckDelayMS"),
+		1000,
+		TEXT("Delay in MS between cleanup checks on previous data from modified actors."));
 #endif // WITH_EDITOR
 
 	static TAutoConsoleVariable<bool> CVarDisableDelayedUnregister(
@@ -152,7 +162,9 @@ void FPCGActorAndComponentMapping::Tick()
 #if WITH_EDITOR
 	AddDelayedActors();
 
-	if (!DelayedModifiedLandscapes.IsEmpty() && LastLandscapeDirtyTime > 0.0 && ((FApp::GetCurrentTime() - LastLandscapeDirtyTime) * 1000.0) > PCGActorAndComponentMapping::CVarLandscapeRefreshTimeDelay.GetValueOnAnyThread())
+	const double CurrentTime = FApp::GetCurrentTime();
+
+	if (!DelayedModifiedLandscapes.IsEmpty() && LastLandscapeDirtyTime > 0.0 && ((CurrentTime - LastLandscapeDirtyTime) * 1000.0) > PCGActorAndComponentMapping::CVarLandscapeRefreshTimeDelay.GetValueOnAnyThread())
 	{
 		LastLandscapeDirtyTime = -1.0;
 		for (TObjectKey<ALandscapeProxy> Landscape : DelayedModifiedLandscapes)
@@ -162,6 +174,23 @@ void FPCGActorAndComponentMapping::Tick()
 
 		DelayedModifiedLandscapes.Empty();
 	}
+
+	// Cleaning up previous data gathered by the OnObjectModified function but not consumed.
+	if (LastPreviousActorDataCleanup < 0.0 || ((CurrentTime - LastPreviousActorDataCleanup) * 1000) > PCGActorAndComponentMapping::CVarActorModifiedPreviousDataCheckDelayMS.GetValueOnAnyThread())
+	{
+		TRACE_CPUPROFILER_EVENT_SCOPE(FPCGActorAndComponentMapping::Tick::PreviousActorDataCleanupCheck);
+
+		LastPreviousActorDataCleanup = CurrentTime;
+		TArray<TObjectKey<AActor>> Keys;
+		ActorToPreviousDataMap.GetKeys(Keys);
+		for (const TObjectKey<AActor>& Key : Keys)
+		{
+			if (((CurrentTime - ActorToPreviousDataMap[Key].Get<2>()) * 1000) > PCGActorAndComponentMapping::CVarActorModifiedPreviousDataCleanupDelayMS.GetValueOnAnyThread())
+			{
+				ActorToPreviousDataMap.Remove(Key);
+			}
+		}
+	}
 #endif // WITH_EDITOR
 }
 
@@ -1012,6 +1041,20 @@ bool FPCGActorAndComponentMapping::IsKeyTracked(const FPCGSelectionKey& InKey) c
 	return CulledTrackedKeysToComponentsMap.Contains(InKey) || AlwaysTrackedKeysToComponentsMap.Contains(InKey);
 }
 
+bool FPCGActorAndComponentMapping::IsActorTracked(const AActor* InActor) const
+{
+	TRACE_CPUPROFILER_EVENT_SCOPE(FPCGActorAndComponentMapping::IsActorTracked);
+
+	TSet<FName> EmptySet;
+
+	auto Matching = [InActor, &EmptySet](const TPair<FPCGSelectionKey, TSet<UPCGComponent*>>& It) -> bool
+	{
+		return It.Key.IsMatching(InActor, EmptySet, It.Value, nullptr);
+	};
+
+	return Algo::AnyOf(CulledTrackedKeysToComponentsMap, Matching) || Algo::AnyOf(AlwaysTrackedKeysToComponentsMap, Matching);
+}
+
 void FPCGActorAndComponentMapping::ResetPartitionActorsMap()
 {
 	PartitionActorsMapLock.WriteLock();
@@ -1024,7 +1067,7 @@ void FPCGActorAndComponentMapping::RegisterTrackingCallbacks()
 	GEngine->OnLevelActorAdded().AddRaw(this, &FPCGActorAndComponentMapping::OnActorAdded);
 	GEngine->OnLevelActorDeleted().AddRaw(this, &FPCGActorAndComponentMapping::OnActorDeleted);
 	FCoreUObjectDelegates::OnObjectPropertyChanged.AddRaw(this, &FPCGActorAndComponentMapping::OnObjectPropertyChanged);
-	FCoreUObjectDelegates::OnPreObjectPropertyChanged.AddRaw(this, &FPCGActorAndComponentMapping::OnPreObjectPropertyChanged);
+	FCoreUObjectDelegates::OnObjectModified.AddRaw(this, &FPCGActorAndComponentMapping::OnObjectModified);
 	FCoreUObjectDelegates::OnObjectPreSave.AddRaw(this, &FPCGActorAndComponentMapping::OnObjectSaved);
 
 	UWorld* World = PCGSubsystem ? PCGSubsystem->GetWorld() : nullptr;
@@ -1041,7 +1084,7 @@ void FPCGActorAndComponentMapping::TeardownTrackingCallbacks()
 	GEngine->OnLevelActorAdded().RemoveAll(this);
 	GEngine->OnLevelActorDeleted().RemoveAll(this);
 	FCoreUObjectDelegates::OnObjectPropertyChanged.RemoveAll(this);
-	FCoreUObjectDelegates::OnPreObjectPropertyChanged.RemoveAll(this);
+	FCoreUObjectDelegates::OnObjectModified.RemoveAll(this);
 	FCoreUObjectDelegates::OnObjectPreSave.RemoveAll(this);
 
 	UWorld* World = PCGSubsystem ? PCGSubsystem->GetWorld() : nullptr;
@@ -1322,15 +1365,16 @@ void FPCGActorAndComponentMapping::OnActorDeleted_Internal(AActor* InActor, bool
 	OnObjectChanged(InActor, /*InPreviousData=*/nullptr, /*InOriginatingObject=*/nullptr, LevelInstanceDepth, /*bNoRefreshOnOwner=*/true);
 }
 
-void FPCGActorAndComponentMapping::OnPreObjectPropertyChanged(UObject* InObject, const FEditPropertyChain& InEditPropertyChain)
+void FPCGActorAndComponentMapping::OnObjectModified(UObject* InObject)
 {
+	TRACE_CPUPROFILER_EVENT_SCOPE(FPCGActorAndComponentMapping::OnObjectModified);
+
 	// Nothing to do if we track nothing
 	if (CulledTrackedKeysToComponentsMap.IsEmpty() && AlwaysTrackedKeysToComponentsMap.IsEmpty())
 	{
 		return;
 	}
 
-	FProperty* MemberProperty = InEditPropertyChain.GetActiveMemberNode() ? InEditPropertyChain.GetActiveMemberNode()->GetValue() : nullptr;
 	AActor* Actor = Cast<AActor>(InObject);
 
 	// Otherwise, if it's an actor component, track it as well
@@ -1354,17 +1398,20 @@ void FPCGActorAndComponentMapping::OnPreObjectPropertyChanged(UObject* InObject,
 	}
 #endif
 
-	auto StorePreviousData = [this, MemberProperty](AActor* InActor, int32 LevelInstanceDepth, auto RecursiveCall) -> void
+	auto StorePreviousData = [this](AActor* InActor, int32 LevelInstanceDepth, auto RecursiveCall) -> void
 	{
 		if (!ActorToPreviousDataMap.Contains(InActor))
 		{
-			FActorPreviousData& PreviousData = ActorToPreviousDataMap.Add(InActor);
-			PreviousData.Get<0>() = PCGActorAndComponentMapping::GetActorBounds(InActor);
-			if (MemberProperty && MemberProperty->GetFName() == GET_MEMBER_NAME_CHECKED(AActor, Tags))
+			if (!IsActorTracked(InActor))
 			{
-				PreviousData.Get<1>() = TSet<FName>(InActor->Tags);
+				return;
 			}
 
+			FActorPreviousData& PreviousData = ActorToPreviousDataMap.Add(InActor);
+			PreviousData.Get<0>() = PCGActorAndComponentMapping::GetActorBounds(InActor);
+			PreviousData.Get<1>() = TSet<FName>(InActor->Tags);
+			PreviousData.Get<2>() = FApp::GetCurrentTime();
+
 			// Also propagate the pre-change to all child actors if it is within a level instance.
 			PCGActorAndComponentMapping::PropagateToLevelInstanceActors(InActor, PCGSubsystem, [this, LevelInstanceDepth, RecursiveCall](AActor* LevelActor)
 			{
@@ -1444,7 +1491,7 @@ void FPCGActorAndComponentMapping::OnObjectPropertyChanged(UObject* InObject, FP
 	const bool bNoOperation = (!bValueNotInteractive && !bActorTagChange)
 		|| bIsTextureCompilationResult
 		|| bIsStaticMeshCompilationResult
-		|| (Actor && (DelayedAddedActors.Contains(Actor) || !ActorToPreviousDataMap.Contains(Actor)));
+		|| (Actor && (DelayedAddedActors.Contains(Actor)));
 
 	if (bNoOperation)
 	{
@@ -1470,10 +1517,14 @@ void FPCGActorAndComponentMapping::OnObjectPropertyChanged(UObject* InObject, FP
 	{
 		auto OnActorChanged = [this, InObject](AActor* InActor, int32 LevelInstanceDepth, auto RecursiveCall) -> void
 		{
-			if (InActor && ActorToPreviousDataMap.Contains(InActor))
+			if (InActor)
 			{
-				FActorPreviousData PreviousData = ActorToPreviousDataMap.FindAndRemoveChecked(InActor);
-				OnObjectChanged(InActor, &PreviousData, /*InOriginatingChangeObject=*/ InObject, LevelInstanceDepth);
+				FActorPreviousData* PreviousData = ActorToPreviousDataMap.Find(InActor);
+				OnObjectChanged(InActor, PreviousData, /*InOriginatingChangeObject=*/ InObject, LevelInstanceDepth);
+				if (PreviousData)
+				{
+					ActorToPreviousDataMap.Remove(InActor);
+				}
 			}
 
 			PCGActorAndComponentMapping::PropagateToLevelInstanceActors(InActor, PCGSubsystem, [this, LevelInstanceDepth, RecursiveCall](AActor* LevelActor)
@@ -1521,7 +1572,7 @@ void FPCGActorAndComponentMapping::OnObjectChanged(UObject* InObject, const FAct
 	{
 		for (auto& It : InMap)
 		{
-			if (It.Key.IsMatching(InObject, RemovedTags, It.Value, OutSet))
+			if (It.Key.IsMatching(InObject, RemovedTags, It.Value, &OutSet))
 			{
 				MatchedKeys.Add(It.Key);
 			}
diff --git a/Engine/Plugins/Experimental/PCG/Source/PCG/Public/Elements/PCGActorSelector.h b/Engine/Plugins/Experimental/PCG/Source/PCG/Public/Elements/PCGActorSelector.h
--- a/Engine/Plugins/Experimental/PCG/Source/PCG/Public/Elements/PCGActorSelector.h
+++ b/Engine/Plugins/Experimental/PCG/Source/PCG/Public/Elements/PCGActorSelector.h
@@ -62,7 +62,7 @@ struct FPCGSelectionKey
 
 	friend uint32 GetTypeHash(const FPCGSelectionKey& In);
 	bool IsMatching(const UObject* InObject, const UPCGComponent* InComponent) const;
-	bool IsMatching(const UObject* InObject, const TSet<FName>& InRemovedTags, const TSet<UPCGComponent*>& InComponents, TSet<UPCGComponent*>& MatchedComponents) const;
+	bool IsMatching(const UObject* InObject, const TSet<FName>& InRemovedTags, const TSet<UPCGComponent*>& InComponents, TSet<UPCGComponent*>* OptionalMatchedComponents = nullptr) const;
 
 	void SetExtraDependency(const UClass* InExtraDependency);
 
diff --git a/Engine/Plugins/Experimental/PCG/Source/PCG/Public/PCGActorAndComponentMapping.h b/Engine/Plugins/Experimental/PCG/Source/PCG/Public/PCGActorAndComponentMapping.h
--- a/Engine/Plugins/Experimental/PCG/Source/PCG/Public/PCGActorAndComponentMapping.h
+++ b/Engine/Plugins/Experimental/PCG/Source/PCG/Public/PCGActorAndComponentMapping.h
@@ -130,12 +130,15 @@ private:
 	bool AnyRuntimeGenComponentsExist() const;
 
 	// Typedef to store the previous position and previous tags of changed actors.
-	using FActorPreviousData = TTuple<FBox, TSet<FName>>;
+	using FActorPreviousData = TTuple<FBox, TSet<FName>, double>;
 
 #if WITH_EDITOR
 	/** Return true if the key is tracked.*/
 	bool IsKeyTracked(const FPCGSelectionKey& InKey) const;
 
+	/** Return true if the actor is tracked.*/
+	bool IsActorTracked(const AActor* InActor) const;
+
 	void OnActorAdded(AActor* InActor);
 	void OnActorLoaded(AActor& InActor);
 	void OnActorAdded_Internal(AActor* InActor, bool bShouldDirty, int32 LevelInstanceDepth, bool bForceAddDelayedActor = false);
@@ -144,7 +147,7 @@ private:
 	void OnActorDeleted_Internal(AActor* InActor, bool bShouldDirty, int32 LevelInstanceDepth);
 	void OnLandscapeChanged(ALandscapeProxy* InLandscape, const FLandscapeProxyComponentDataChangedParams& InChangeParams);
 	void ApplyLandscapeChanges(ALandscapeProxy* InLandscape);
-	void OnPreObjectPropertyChanged(UObject* InObject, const FEditPropertyChain& InEditPropertyChain);
+	void OnObjectModified(UObject* InObject);
 	void OnObjectPropertyChanged(UObject* InObject, FPropertyChangedEvent& InEvent);
 	void OnObjectSaved(UObject* InObject, FObjectPreSaveContext InObjectSaveContext);
 	void OnPCGGraphGeneratedOrCleaned(UPCGComponent* InComponent);
@@ -228,5 +231,8 @@ private:
 
 	// Temp boolean to indicate we are currently exiting the landscape edit mode
 	bool bIsCurrentlyExitingLandscapeEditMode = false;
+
+	// Time keeper for cleaning up cached previous actor data
+	double LastPreviousActorDataCleanup = -1.0;
 #endif // WITH_EDITOR
 };
\ No newline at end of file
