Updated test to resolve deprecated alias on property

#rnx

diff --git a/Engine/Plugins/Experimental/PCG/Source/PCG/Private/Tests/Elements/PCGDataTableRowToParamDataTest.cpp b/Engine/Plugins/Experimental/PCG/Source/PCG/Private/Tests/Elements/PCGDataTableRowToParamDataTest.cpp
--- a/Engine/Plugins/Experimental/PCG/Source/PCG/Private/Tests/Elements/PCGDataTableRowToParamDataTest.cpp
+++ b/Engine/Plugins/Experimental/PCG/Source/PCG/Private/Tests/Elements/PCGDataTableRowToParamDataTest.cpp
@@ -88,7 +88,7 @@ CCC,"C Name","C String","333","3333","3.3","3.33","(X=3.0,Y=3.0)","(X=3.0,Y=3.0,
 
 	{
 		UPCGParamData* InputParamData = NewObject<UPCGParamData>();
-		InputParamData->MutableMetadata()->CreateStringAttribute(TEXT("PathOverride"), TestDataTable->GetPathName(), false);
+		InputParamData->MutableMetadata()->CreateSoftObjectPathAttribute(TEXT("DataTable"), FSoftObjectPath(TestDataTable), false);
 		InputParamData->MutableMetadata()->CreateStringAttribute(TEXT("RowName"), RowNameOverride.ToString(), false);
 		FPCGTaggedData& ParamInput = TestData.InputData.TaggedData.Emplace_GetRef();
 		ParamInput.Data = InputParamData;
