[PCG] RuntimeGen Improvements
* Pooling for WP Streaming GenerationSources
* Create EditorViewportGenSource on initialization rather than tick and drop support for multiple viewports
* Undo/Redo now always does a deep refresh of runtime components
* Defer shallow refresh to the next tick to avoid leaking resources when there are multiple refreshes in the same tick
* Fix bug where removing or changing GraphInstance on RuntimeGen component triggers ensures and leaks resources/locks grid cells from generating
* Fix bug where Unbounded partitioned components would not be generated if they have no grid sizes
* Fix bug where grid cells are scheduled for generation using an incorrect bounds test

#rnx

diff --git a/Engine/Plugins/Experimental/PCG/Source/PCG/Private/PCGActorAndComponentMapping.cpp b/Engine/Plugins/Experimental/PCG/Source/PCG/Private/PCGActorAndComponentMapping.cpp
--- a/Engine/Plugins/Experimental/PCG/Source/PCG/Private/PCGActorAndComponentMapping.cpp
+++ b/Engine/Plugins/Experimental/PCG/Source/PCG/Private/PCGActorAndComponentMapping.cpp
@@ -805,15 +805,7 @@ UPCGComponent* FPCGActorAndComponentMapping::GetLocalComponent(uint32 GridSize,
 		const TObjectPtr<APCGPartitionActor>* PartitionActor = PartitionActorsOnGrid->Find(CellCoords);
 		if (PartitionActor && *PartitionActor)
 		{
-			TArray<UPCGComponent*, TInlineAllocator<4>> PCGComponents;
-			(*PartitionActor)->GetComponents(PCGComponents);
-
-			UPCGComponent** MatchingComponent = PCGComponents.FindByPredicate([InOriginalComponent](UPCGComponent* Comp)
-			{
-				return Comp->GetOriginalComponent() == InOriginalComponent;
-			});
-
-			return MatchingComponent ? *MatchingComponent : nullptr;
+			return (*PartitionActor)->GetLocalComponent(InOriginalComponent);
 		}
 	}
 
diff --git a/Engine/Plugins/Experimental/PCG/Source/PCG/Private/PCGComponent.cpp b/Engine/Plugins/Experimental/PCG/Source/PCG/Private/PCGComponent.cpp
--- a/Engine/Plugins/Experimental/PCG/Source/PCG/Private/PCGComponent.cpp
+++ b/Engine/Plugins/Experimental/PCG/Source/PCG/Private/PCGComponent.cpp
@@ -1476,7 +1476,8 @@ void UPCGComponent::RefreshAfterGraphChanged(UPCGGraphInterface* InGraph, EPCGCh
 
 		DirtyGenerated(bDirtyInputs ? (EPCGComponentDirtyFlag::Actor | EPCGComponentDirtyFlag::Landscape) : EPCGComponentDirtyFlag::None);
 
-		if (bHasGraph)
+		// If there is no graph, we should still refresh if we are runtime-managed, since the RuntimeGenScheduler will need to flush its resources.
+		if (bHasGraph || IsManagedByRuntimeGenSystem())
 		{
 			Refresh(ChangeType);
 		}
@@ -1597,7 +1598,7 @@ void UPCGComponent::PostEditChangeProperty(FPropertyChangedEvent& PropertyChange
 	}
 	else if (PropName == GET_MEMBER_NAME_CHECKED(UPCGComponent, GraphInstance))
 	{
-		OnGraphChanged(GraphInstance, EPCGChangeType::Structural);
+		OnGraphChanged(GraphInstance, EPCGChangeType::Structural | EPCGChangeType::GenerationGrid);
 	}
 	else if (PropName == GET_MEMBER_NAME_CHECKED(UPCGComponent, InputType))
 	{
@@ -1679,6 +1680,15 @@ void UPCGComponent::PreEditUndo()
 	// so we can have a consistent state
 	LastGeneratedBoundsPriorToUndo = LastGeneratedBounds;
 
+	if (IsManagedByRuntimeGenSystem())
+	{
+		// Always flush the runtime gen state before undo/redo to avoid leaking resources or locking grid cells from future generation.
+		if (UPCGSubsystem* Subsystem = GetSubsystem())
+		{
+			Subsystem->RefreshRuntimeGenComponent(this, EPCGChangeType::GenerationGrid);
+		}
+	}
+
 	if (bGenerated)
 	{
 		// Cleanup so managed resources are cleaned in all cases
@@ -1706,7 +1716,7 @@ void UPCGComponent::PostEditUndo()
 		// operation removes the component, a valid refresh task ID is set but the refresh task itself will fail
 		// and leave the valid task ID hanging on the component. Forcing here means if we later retrieve this state
 		// from the undo/redo buffer, the refresh will be forced which will reset the state.
-		Refresh(EPCGChangeType::Structural, /*bCancelExistingRefresh=*/true);
+		Refresh(EPCGChangeType::Structural | EPCGChangeType::GenerationGrid, /*bCancelExistingRefresh=*/true);
 	}
 
 	Super::PostEditUndo();
diff --git a/Engine/Plugins/Experimental/PCG/Source/PCG/Private/PCGWorldActor.cpp b/Engine/Plugins/Experimental/PCG/Source/PCG/Private/PCGWorldActor.cpp
--- a/Engine/Plugins/Experimental/PCG/Source/PCG/Private/PCGWorldActor.cpp
+++ b/Engine/Plugins/Experimental/PCG/Source/PCG/Private/PCGWorldActor.cpp
@@ -133,11 +133,15 @@ void APCGWorldActor::EndPlay(const EEndPlayReason::Type EndPlayReason)
 
 void APCGWorldActor::CreateGridGuidsIfNecessary(const PCGHiGenGrid::FSizeArray& InGridSizes, bool bAreGridsSerialized)
 {
+	if (InGridSizes.IsEmpty())
+	{
+		return;
+	}
+
 	TMap<uint32, FGuid>& GuidsMap = bAreGridsSerialized ? GridGuids : TransientGridGuids;
 	FRWLock& GuidsLock = bAreGridsSerialized ? GridGuidsLock : TransientGridGuidsLock;
 
 	// Check if any need adding
-	ensure(!InGridSizes.IsEmpty());
 	PCGHiGenGrid::FSizeArray GridSizesToAdd;
 	{
 		FReadScopeLock ReadLock(GuidsLock);
diff --git a/Engine/Plugins/Experimental/PCG/Source/PCG/Private/RuntimeGen/PCGGenSourceManager.cpp b/Engine/Plugins/Experimental/PCG/Source/PCG/Private/RuntimeGen/PCGGenSourceManager.cpp
--- a/Engine/Plugins/Experimental/PCG/Source/PCG/Private/RuntimeGen/PCGGenSourceManager.cpp
+++ b/Engine/Plugins/Experimental/PCG/Source/PCG/Private/RuntimeGen/PCGGenSourceManager.cpp
@@ -26,55 +26,60 @@ FPCGGenSourceManager::FPCGGenSourceManager(const UWorld* InWorld)
 	// Collect PlayerControllers.
 	FGameModeEvents::GameModePostLoginEvent.AddRaw(this, &FPCGGenSourceManager::OnGameModePostLogin);
 	FGameModeEvents::GameModeLogoutEvent.AddRaw(this, &FPCGGenSourceManager::OnGameModePostLogout);
+
+#if WITH_EDITOR
+	EditorCameraGenSource = NewObject<UPCGGenSourceEditorCamera>();
+	EditorCameraGenSource->AddToRoot();
+#endif
 }
 
 FPCGGenSourceManager::~FPCGGenSourceManager()
 {
 	FGameModeEvents::GameModePostLoginEvent.RemoveAll(this);
 	FGameModeEvents::GameModeLogoutEvent.RemoveAll(this);
-}
 
-TSet<IPCGGenSourceBase*> FPCGGenSourceManager::GetGenSources(const APCGWorldActor* InPCGWorldActor) const
-{
-	// Some the the generation sources need to be queried each frame, so build the 'current' set now.
-	TSet<IPCGGenSourceBase*> CurrentGenSources = GenSources;
+#if WITH_EDITOR
+	if (EditorCameraGenSource && EditorCameraGenSource->IsRooted())
+	{
+		EditorCameraGenSource->RemoveFromRoot();
+	}
+#endif
 
-	if (const UWorld* InWorld = InPCGWorldActor->GetWorld())
+	for (UPCGGenSourceWPStreamingSource* WorldPartitionGenSource : WorldPartitionGenSources)
 	{
-		if (const UWorldPartition* WorldPartition = InWorld->GetWorldPartition())
+		if (WorldPartitionGenSource && WorldPartitionGenSource->IsRooted())
 		{
-			// TODO: Grab StreamingSourceProviders instead of StreamingSources? This could allow us to avoid calling NewObject every tick.
-			// Note: GetStreamingSources only works in GameWorld, so StreamingSources do not act as GenSources in editor.
-			for (const FWorldPartitionStreamingSource& Source : WorldPartition->GetStreamingSources())
-			{
-				// StreamingSources live only for this tick, so they don't need to be rooted.
-				UPCGGenSourceWPStreamingSource* GenSource = NewObject<UPCGGenSourceWPStreamingSource>();
-				GenSource->StreamingSource = &Source;
-
-				// TODO: Is it possible to avoid adding a StreamingSource for the Player, which we already capture in OnGameModePostLogin?
-				CurrentGenSources.Add(GenSource);
-			}
+			WorldPartitionGenSource->RemoveFromRoot();
 		}
 	}
+}
+
+TSet<IPCGGenSourceBase*> FPCGGenSourceManager::GetAllGenSources(const APCGWorldActor* InPCGWorldActor)
+{
+	if (bDirty)
+	{
+		UpdatePerTickGenSources(InPCGWorldActor);
+		bDirty = false;
+	}
 
-#if WITH_EDITOR
-	if (InPCGWorldActor->bTreatEditorViewportAsGenerationSource && !World->IsGameWorld())
+	TSet<IPCGGenSourceBase*> AllGenSources = RegisteredGenSources;
+
+	// Acquire a generation source for each active streaming source.
+	for (int32 I = 0; I < NumWorldPartitionStreamingSources; ++I)
 	{
-		for (FEditorViewportClient* EditorViewportClient : GEditor->GetAllViewportClients())
-		{
-			if (EditorViewportClient->IsVisible())
-			{
-				// TODO: grab all 4 of these when GenSourceManager is created, instead of creating a new UObject for each every frame?
-				UPCGGenSourceEditorCamera* GenSource = NewObject<UPCGGenSourceEditorCamera>();
-				GenSource->EditorViewportClient = EditorViewportClient;
+		check(WorldPartitionGenSources.IsValidIndex(I));
+		AllGenSources.Add(WorldPartitionGenSources[I]);
+	}
 
-				CurrentGenSources.Add(GenSource);
-			}
-		}
+#if WITH_EDITOR
+	// Acquire a generation source for the active editor viewport if one exists.
+	if (EditorCameraGenSource && EditorCameraGenSource->EditorViewportClient)
+	{
+		AllGenSources.Add(EditorCameraGenSource);
 	}
 #endif
 
-	return CurrentGenSources;
+	return AllGenSources;
 }
 
 bool FPCGGenSourceManager::RegisterGenSource(IPCGGenSourceBase* InGenSource)
@@ -87,7 +92,7 @@ bool FPCGGenSourceManager::RegisterGenSource(IPCGGenSourceBase* InGenSource)
 		}
 	}
 
-	return GenSources.Add(InGenSource).IsValidId();
+	return RegisteredGenSources.Add(InGenSource).IsValidId();
 }
 
 bool FPCGGenSourceManager::UnregisterGenSource(const IPCGGenSourceBase* InGenSource)
@@ -100,7 +105,7 @@ bool FPCGGenSourceManager::UnregisterGenSource(const IPCGGenSourceBase* InGenSou
 		}
 	}
 
-	return GenSources.Remove(InGenSource) > 0;
+	return RegisteredGenSources.Remove(InGenSource) > 0;
 }
 
 void FPCGGenSourceManager::OnGameModePostLogin(AGameModeBase* InGameMode, APlayerController* InPlayerController)
@@ -109,18 +114,16 @@ void FPCGGenSourceManager::OnGameModePostLogin(AGameModeBase* InGameMode, APlaye
 	{
 		return;
 	}
+
 	ensure(IsInGameThread());
 
-	if (InPlayerController->GetPawn())
-	{
-		UPCGGenSourcePlayer* GenSource = NewObject<UPCGGenSourcePlayer>();
-		GenSource->SetPlayerController(InPlayerController);
+	UPCGGenSourcePlayer* GenSource = NewObject<UPCGGenSourcePlayer>();
+	GenSource->SetPlayerController(InPlayerController);
 
-		FSetElementId ElementId = GenSources.Add(GenSource);
-		if (ElementId.IsValidId())
-		{
-			GenSource->AddToRoot();
-		}
+	FSetElementId ElementId = RegisteredGenSources.Add(GenSource);
+	if (ElementId.IsValidId())
+	{
+		GenSource->AddToRoot();
 	}
 }
 
@@ -130,18 +133,74 @@ void FPCGGenSourceManager::OnGameModePostLogout(AGameModeBase* InGameMode, ACont
 	{
 		return;
 	}
+
 	ensure(IsInGameThread());
 
-	for (IPCGGenSourceBase* GenSource : GenSources)
+	for (IPCGGenSourceBase* GenSource : RegisteredGenSources)
 	{
 		if (UPCGGenSourcePlayer* GenSourcePlayer = Cast<UPCGGenSourcePlayer>(GenSource))
 		{
 			if (GenSourcePlayer->GetPlayerController() == InController || !GenSourcePlayer->IsValid())
 			{
 				GenSourcePlayer->RemoveFromRoot();
-				GenSources.Remove(GenSourcePlayer);
+				RegisteredGenSources.Remove(GenSourcePlayer);
 				break;
 			}
 		}
 	}
 }
+
+void FPCGGenSourceManager::UpdatePerTickGenSources(const APCGWorldActor* InPCGWorldActor)
+{
+	if (const UWorld* InWorld = InPCGWorldActor->GetWorld())
+	{
+		if (const UWorldPartition* WorldPartition = InWorld->GetWorldPartition())
+		{
+			// TODO: Grab StreamingSourceProviders instead of StreamingSources?
+			// TODO: Is it possible to avoid adding a StreamingSource for the Player, which we already capture in OnGameModePostLogin?
+			// Note: GetStreamingSources only works in GameWorld, so StreamingSources do not act as generation sources in editor.
+			const TArray<FWorldPartitionStreamingSource>& StreamingSources = WorldPartition->GetStreamingSources();
+
+			NumWorldPartitionStreamingSources = StreamingSources.Num();
+
+			for (int32 I = 0; I < NumWorldPartitionStreamingSources; ++I)
+			{
+				const bool bCreateNewWorldPartitionGenSource = I >= WorldPartitionGenSources.Num();
+
+				// Allocate new WP generation sources as necessary.
+				UPCGGenSourceWPStreamingSource* GenSource = bCreateNewWorldPartitionGenSource ? WorldPartitionGenSources.Add_GetRef(NewObject<UPCGGenSourceWPStreamingSource>()) : WorldPartitionGenSources[I];
+
+				// Capture each streaming source in a single generation source.
+				if (GenSource)
+				{
+					if (bCreateNewWorldPartitionGenSource)
+					{
+						// New generation sources need to be rooted.
+						GenSource->AddToRoot();
+					}
+
+					GenSource->StreamingSource = &StreamingSources[I];
+				}
+			}
+		}
+	}
+
+#if WITH_EDITOR
+	EditorCameraGenSource->EditorViewportClient = nullptr;
+
+	// Update the active editor viewport client for the EditorCameraGenSource, only if requested by the world actor, in-editor, and the viewport is visible.
+	if (InPCGWorldActor->bTreatEditorViewportAsGenerationSource && !World->IsGameWorld())
+	{
+		if (FViewport* Viewport = GEditor->GetActiveViewport())
+		{
+			if (FEditorViewportClient* ViewportClient = static_cast<FEditorViewportClient*>(Viewport->GetClient()))
+			{
+				if (ViewportClient->IsVisible())
+				{
+					EditorCameraGenSource->EditorViewportClient = ViewportClient;
+				}
+			}
+		}
+	}
+#endif
+}
diff --git a/Engine/Plugins/Experimental/PCG/Source/PCG/Private/RuntimeGen/PCGRuntimeGenScheduler.cpp b/Engine/Plugins/Experimental/PCG/Source/PCG/Private/RuntimeGen/PCGRuntimeGenScheduler.cpp
--- a/Engine/Plugins/Experimental/PCG/Source/PCG/Private/RuntimeGen/PCGRuntimeGenScheduler.cpp
+++ b/Engine/Plugins/Experimental/PCG/Source/PCG/Private/RuntimeGen/PCGRuntimeGenScheduler.cpp
@@ -87,7 +87,13 @@ void FPCGRuntimeGenScheduler::Tick(APCGWorldActor* InPCGWorldActor)
 
 	TickCVars(InPCGWorldActor);
 
-	const TSet<IPCGGenSourceBase*> GenSources = bAnyRuntimeGenComponentsExist ? GenSourceManager->GetGenSources(InPCGWorldActor) : TSet<IPCGGenSourceBase*>();
+	TSet<IPCGGenSourceBase*> GenSources;
+
+	if (bAnyRuntimeGenComponentsExist)
+	{
+		GenSourceManager->Tick();
+		GenSources = GenSourceManager->GetAllGenSources(InPCGWorldActor);
+	}
 
 	// Initialize RuntimeGen PA pool if necessary. If PoolSize is 0, then we have not initialized the pool yet.
 	if (!GenSources.IsEmpty() || !GeneratedComponents.IsEmpty())
@@ -98,6 +104,10 @@ void FPCGRuntimeGenScheduler::Tick(APCGWorldActor* InPCGWorldActor)
 		}
 	}
 
+	// Remove any generation keys that have been registered for deferred removal.
+	GeneratedComponents = GeneratedComponents.Difference(GeneratedComponentsToRemove);
+	GeneratedComponentsToRemove.Empty();
+
 	// 1. Queue nearby components for generation.
 	
 	// Mapping of component + coordinates to priorities - needed to compute max priority over all gen sources.
@@ -238,7 +248,7 @@ void FPCGRuntimeGenScheduler::TickQueueComponentsForGeneration(
 	{
 		TRACE_CPUPROFILER_EVENT_SCOPE(FPCGRuntimeGenScheduler::CollectLocalComponents);
 
-		if (!ensure(OriginalComponent))
+		if (!ensure(OriginalComponent) || !OriginalComponent->GetGraph())
 		{
 			continue;
 		}
@@ -252,12 +262,13 @@ void FPCGRuntimeGenScheduler::TickQueueComponentsForGeneration(
 			PCGHiGenGrid::FSizeArray GridSizes;
 			ensure(PCGHelpers::GetGenerationGridSizes(OriginalComponent->GetGraph(), InPCGWorldActor, GridSizes, bHasUnbounded));
 
-			if (GridSizes.IsEmpty())
+			if (GridSizes.IsEmpty() && !bHasUnbounded)
 			{
 				continue;
 			}
 
-			const double MaxGenerationRadius = OriginalComponent->GetGenerationRadiusFromGrid(PCGHiGenGrid::GridSizeToGrid(GridSizes[0]));
+			const EPCGHiGenGrid MaxGrid = bHasUnbounded ? EPCGHiGenGrid::Unbounded : PCGHiGenGrid::GridSizeToGrid(GridSizes[0]);
+			const double MaxGenerationRadius = OriginalComponent->GetGenerationRadiusFromGrid(MaxGrid);
 
 			for (const IPCGGenSourceBase* GenSource : InGenSources)
 			{
@@ -351,12 +362,12 @@ void FPCGRuntimeGenScheduler::TickQueueComponentsForGeneration(
 
 								if (InPCGWorldActor->bUse2DGrid)
 								{
-									ModifiedGenSourcePosition.Z = CellBounds.Min.Z;
+									ModifiedGenSourcePosition.Z = IntersectedBounds.Min.Z;
 								}
 
 								// Verify the grid cell actually lies within the generation radius.
 								// TODO: this is no longer necessary if we rasterize the sphere instead.
-								const double LocalDistanceSquared = CellBounds.ComputeSquaredDistanceToPoint(ModifiedGenSourcePosition);
+								const double LocalDistanceSquared = IntersectedBounds.ComputeSquaredDistanceToPoint(ModifiedGenSourcePosition);
 								if (LocalDistanceSquared <= GenerationRadius * GenerationRadius)
 								{
 									AddComponentToGenerate(Key, GenSource, Policy, IntersectedBounds, InPCGWorldActor->bUse2DGrid);
@@ -374,7 +385,7 @@ void FPCGRuntimeGenScheduler::TickQueueComponentsForGeneration(
 	{
 		TRACE_CPUPROFILER_EVENT_SCOPE(FPCGRuntimeGenScheduler::CollectNonPartitionedComponents);
 
-		if (!ensure(OriginalComponent))
+		if (!ensure(OriginalComponent) || !OriginalComponent->GetGraph())
 		{
 			continue;
 		}
@@ -488,9 +499,10 @@ void FPCGRuntimeGenScheduler::TickCleanup(const TSet<IPCGGenSourceBase*>& InGenS
 		else
 		{
 			UPCGComponent* LocalComponent = ActorAndComponentMapping->GetLocalComponent(GridSize, GridCoords, OriginalComponent, /*bRuntimeGenerated=*/true);
-			if (!ensure(LocalComponent))
+			if (!LocalComponent)
 			{
-				UE_LOG(LogPCG, Error, TEXT("[RUNTIMEGEN] Generated local component could not be retrieved: %u, %d_%d_%d"), GridSize, GridCoords.X, GridCoords.Y, GridCoords.Z);
+				// Attempt to clean even in failure case to avoid leaking resources.
+				ComponentsToClean.Add({ GenerationKey, LocalComponent });
 				continue;
 			}
 
@@ -774,12 +786,6 @@ void FPCGRuntimeGenScheduler::CleanupLocalComponents(const APCGWorldActor* InPCG
 		if (Grid != EPCGHiGenGrid::Unbounded)
 		{
 			UPCGComponent* LocalComponent = ActorAndComponentMapping->GetLocalComponent(GridSize, GridCoords, OriginalComponent, /*bRuntimeGenerated=*/true);
-			if (!ensure(LocalComponent))
-			{
-				UE_LOG(LogPCG, Error, TEXT("[RUNTIMEGEN] Generated local component could not be retrieved: %u, %d_%d_%d"), GridSize, GridCoords.X, GridCoords.Y, GridCoords.Z);
-				continue;
-			}
-
 			ComponentsToClean.Add({ GenerationKey, LocalComponent });
 		}
 	}
@@ -801,9 +807,9 @@ void FPCGRuntimeGenScheduler::CleanupComponent(const FGridGenerationKey& Generat
 
 	APCGPartitionActor* PartitionActor = nullptr;
 
-	if (!ensure(GeneratedComponent))
+	if (!GeneratedComponent)
 	{
-		UE_LOG(LogPCG, Error, TEXT("[RUNTIMEGEN] GENERATED COMPONENT COULD NOT BE RECOVERED ON GRID %d AT (%d, %d, %d)"), GridSize, GridCoords.X, GridCoords.Y, GridCoords.Z);
+		UE_LOG(LogPCG, Warning, TEXT("Runtime generated component could not be recovered on grid %d at (%d, %d, %d). It has been lost or destroyed."), GridSize, GridCoords.X, GridCoords.Y, GridCoords.Z);
 
 		// If the GeneratedComponent has been lost for some reason, get the PA directly from the ActorAndComponentMapping.
 		PartitionActor = ActorAndComponentMapping->GetPartitionActor(GridSize, GridCoords, /*bRuntimeGenerated=*/true);
@@ -918,8 +924,8 @@ void FPCGRuntimeGenScheduler::RefreshComponent(UPCGComponent* InComponent, bool
 	{
 		// Refresh path - mark component dirty and removed generated keys which will cause it to be scheduled for regeneration.
 
-		// Remove from generated component set, component will be regenerated later (and in grid order so that e.g. unbounded
-		// is generated first).
+		// Register for deferred removal from generated components set, component will be regenerated later (and in grid order
+		// so that e.g. unbounded is generated first).
 		if (PartitionActor)
 		{
 			if (bLoggingEnabled)
@@ -927,21 +933,18 @@ void FPCGRuntimeGenScheduler::RefreshComponent(UPCGComponent* InComponent, bool
 				UE_LOG(LogPCG, Warning, TEXT("[RUNTIMEGEN] SHALLOW REFRESH LOCAL COMPONENT: '%s'"), *PartitionActor->GetActorNameOrLabel());
 			}
 
-			GeneratedComponents.Remove({ PartitionActor->GetPCGGridSize(), PartitionActor->GetGridCoord(), OriginalComponent });
-
+			GeneratedComponentsToRemove.Emplace({ PartitionActor->GetPCGGridSize(), PartitionActor->GetGridCoord(), OriginalComponent });
 			InComponent->CleanupLocalImmediate(/*bRemoveComponents=*/false);
 		}
 		else
 		{
-			// Remove original component
-			GeneratedComponents.Remove({ PCGHiGenGrid::UnboundedGridSize(), FIntVector(0), OriginalComponent });
-
-			// Remove local components
-			TSet<FGridGenerationKey> KeysToRemove;
+			// Register original component for deferred removal.
+			GeneratedComponentsToRemove.Emplace({ PCGHiGenGrid::UnboundedGridSize(), FIntVector(0), OriginalComponent });
 
+			// Register local components for deferred removal if they have not already registered themselves.
 			for (const FGridGenerationKey& Key : GeneratedComponents)
 			{
-				if (Key.GetOriginalComponent() == InComponent)
+				if (Key.GetOriginalComponent() == InComponent && !GeneratedComponentsToRemove.Contains(Key))
 				{
 					// TODO - clean up local immediate will have a flag in the future to clean up the local components on its own, so this call to CleanupLocalImmediate will not be required
 					UPCGComponent* LocalComponent = ActorAndComponentMapping->GetLocalComponent(
@@ -960,12 +963,10 @@ void FPCGRuntimeGenScheduler::RefreshComponent(UPCGComponent* InComponent, bool
 						LocalComponent->CleanupLocalImmediate(/*bRemoveComponents=*/false);
 					}
 
-					KeysToRemove.Add(Key);
+					GeneratedComponentsToRemove.Add(Key);
 				}
 			}
 
-			GeneratedComponents = GeneratedComponents.Difference(KeysToRemove);
-
 			if (bLoggingEnabled && OriginalComponent->GetOwner())
 			{
 				UE_LOG(LogPCG, Warning, TEXT("[RUNTIMEGEN] SHALLOW REFRESH COMPONENT: '%s' PARTITIONED: %d"),
@@ -1038,14 +1039,15 @@ void FPCGRuntimeGenScheduler::RefreshComponent(UPCGComponent* InComponent, bool
 				{
 					const FIntVector GridCoords = GenerationKey.GetGridCoords();
 
-					UPCGComponent* LocalComponent = ActorAndComponentMapping->GetLocalComponent(GridSize, GridCoords, OriginalComponent, /*bRuntimeGenerated=*/true);
-					if (!ensure(LocalComponent))
+					if (UPCGComponent* LocalComponent = ActorAndComponentMapping->GetLocalComponent(GridSize, GridCoords, OriginalComponent, /*bRuntimeGenerated=*/true))
 					{
-						UE_LOG(LogPCG, Error, TEXT("[RUNTIMEGEN] Generated local component could not be retrieved: %u, %d_%d_%d"), GridSize, GridCoords.X, GridCoords.Y, GridCoords.Z);
-						continue;
+						RefreshLocalComponent(LocalComponent);
+					}
+					else
+					{
+						// If the local component could not be recovered, cleanup its entry to avoid leaking resources/locking the grid cell.
+						CleanupComponent(GenerationKey, nullptr);
 					}
-
-					RefreshLocalComponent(LocalComponent);
 				}
 			}
 		}
diff --git a/Engine/Plugins/Experimental/PCG/Source/PCG/Public/RuntimeGen/PCGGenSourceManager.h b/Engine/Plugins/Experimental/PCG/Source/PCG/Public/RuntimeGen/PCGGenSourceManager.h
--- a/Engine/Plugins/Experimental/PCG/Source/PCG/Public/RuntimeGen/PCGGenSourceManager.h
+++ b/Engine/Plugins/Experimental/PCG/Source/PCG/Public/RuntimeGen/PCGGenSourceManager.h
@@ -6,6 +6,8 @@
 
 class APCGWorldActor;
 class UPCGGenSourceComponent;
+class UPCGGenSourceEditorCamera;
+class UPCGGenSourceWPStreamingSource;
 
 class AController;
 class AGameModeBase;
@@ -21,8 +23,11 @@ public:
 	FPCGGenSourceManager(const UWorld* InWorld);
 	~FPCGGenSourceManager();
 
-	/** Creates the set of all GenerationSources in the level based on the contents of the GenSourceManager's world. */
-	TSet<IPCGGenSourceBase*> GetGenSources(const APCGWorldActor* InPCGWorldActor) const;
+	/** Marks the GenSourceManager as dirty so that the next call to 'GetAllGenSources()' will update tracked generation sources. */
+	void Tick() { bDirty = true; }
+
+	/** Creates the set of all generation sources tracked by the manager. Call 'Tick()' to keep tracked generation sources up to date. */
+	TSet<IPCGGenSourceBase*> GetAllGenSources(const APCGWorldActor* InPCGWorldActor);
 
 	/** Adds a UPCGGenSource to be tracked by the GenSourceManager. */
 	bool RegisterGenSource(IPCGGenSourceBase* InGenSource);
@@ -34,8 +39,24 @@ protected:
 	void OnGameModePostLogin(AGameModeBase* InGameMode, APlayerController* InPlayerController);
 	void OnGameModePostLogout(AGameModeBase* InGameMode, AController* InController);
 
+	/** Updates tracked generation sources that should be refreshed per tick. */
+	void UpdatePerTickGenSources(const APCGWorldActor* InPCGWorldActor);
+
 protected:
-	TSet<IPCGGenSourceBase*> GenSources;
+	/** Tracks registered generation sources, such as UPCGGenSourceComponent and UPCGGenSourcePlayer. */
+	TSet<IPCGGenSourceBase*> RegisteredGenSources;
+
+#if WITH_EDITORONLY_DATA
+	/** Tracks the active/main editor viewport client. This is refreshed every tick to keep a handle to whichever viewport is active. */
+	UPCGGenSourceEditorCamera* EditorCameraGenSource = nullptr;
+#endif
+
+	/** Pool of GenSources dedicated to WorldPartition StreamingSources. This list grows as necessary and refreshes every tick. */
+	TArray<UPCGGenSourceWPStreamingSource*> WorldPartitionGenSources;
+
+	/** There may be more generation sources in 'WorldPartitionGenSources' than there are streaming sources in the world, so we only use the first N generation sources. */
+	int32 NumWorldPartitionStreamingSources = 0;
 
 	const UWorld* World = nullptr;
+	bool bDirty = false;
 };
diff --git a/Engine/Plugins/Experimental/PCG/Source/PCG/Public/RuntimeGen/PCGRuntimeGenScheduler.h b/Engine/Plugins/Experimental/PCG/Source/PCG/Public/RuntimeGen/PCGRuntimeGenScheduler.h
--- a/Engine/Plugins/Experimental/PCG/Source/PCG/Public/RuntimeGen/PCGRuntimeGenScheduler.h
+++ b/Engine/Plugins/Experimental/PCG/Source/PCG/Public/RuntimeGen/PCGRuntimeGenScheduler.h
@@ -102,6 +102,12 @@ private:
 	*/
 	TSet<FGridGenerationKey> GeneratedComponents;
 
+	/** Tracks the components which should be removed from the GeneratedComponents set on the next tick. This helps us defer removal in case we get multiple
+	* refreshes in a single tick. For example, a shallow refresh followed by a deep refresh would require the generated components to persist, otherwise we
+	* will leak Partition Actors.
+	*/
+	TSet<FGridGenerationKey> GeneratedComponentsToRemove;
+
 	/** Pool of RuntimeGen PartitionActors used for hierarchical generation. */
 	TArray<APCGPartitionActor*> PartitionActorPool;
 
