[PCG] Disable graph notifications while processing graph changes (like adding an edge), and defer type propagation change notifications until post step.

There was a cascade of change notifications that was freezing the editor for large graphs. Also added a couple of profile scopes that will make similar issues clear in future trace captures.

diff --git a/Engine/Plugins/Experimental/PCG/Source/PCG/Private/PCGGraph.cpp b/Engine/Plugins/Experimental/PCG/Source/PCG/Private/PCGGraph.cpp
--- a/Engine/Plugins/Experimental/PCG/Source/PCG/Private/PCGGraph.cpp
+++ b/Engine/Plugins/Experimental/PCG/Source/PCG/Private/PCGGraph.cpp
@@ -82,24 +82,32 @@ namespace PCGGraphUtils
 		InSourcePropertyDesc->CachedProperty->CopyCompleteValue(TargetValueAddress, SourceValueAddress);
 	}
 
-	EPCGChangeType NotifyTouchedNodes(const TSet<UPCGNode*>& InTouchedNodes, EPCGChangeType ChangeType)
+	EPCGChangeType NotifyTouchedNodes(const TSet<UPCGNode*>& InTouchedNodes, EPCGChangeType InChangeType)
 	{
 		EPCGChangeType FinalChangeType = EPCGChangeType::None;
 
+		// Build a final list of all touched nodes, so we can broadcast the change once below.
+		TSet<UPCGNode*> FinalTouchedNodes = InTouchedNodes;
+
 		for (UPCGNode* TouchedNode : InTouchedNodes)
 		{
 			if (TouchedNode)
 			{
-				const EPCGChangeType NodeChangeType = ChangeType | TouchedNode->PropagateDynamicPinTypes();
-
-#if WITH_EDITOR
-				TouchedNode->OnNodeChangedDelegate.Broadcast(TouchedNode, NodeChangeType | EPCGChangeType::Node | ChangeType);
-#endif
+				const EPCGChangeType NodeChangeType = InChangeType | TouchedNode->PropagateDynamicPinTypes(FinalTouchedNodes);
 
 				FinalChangeType |= NodeChangeType;
 			}
 		}
 
+		// Do change notifications for the final set.
+#if WITH_EDITOR
+		for (UPCGNode* TouchedNode : FinalTouchedNodes)
+		{
+			check(TouchedNode);
+			TouchedNode->OnNodeChangedDelegate.Broadcast(TouchedNode, EPCGChangeType::Node | InChangeType);
+		}
+#endif
+
 		return FinalChangeType;
 	}
 }
@@ -639,7 +647,7 @@ bool UPCGGraph::AddLabeledEdge(UPCGNode* From, const FName& FromPinLabel, UPCGNo
 
 	UPCGPin* FromPin = From->GetOutputPin(FromPinLabel);
 
-	if(!FromPin)
+	if (!FromPin)
 	{
 		UE_LOG(LogPCG, Error, TEXT("From node %s does not have the %s label"), *From->GetName(), *FromPinLabel.ToString());
 		return false;
@@ -647,12 +655,16 @@ bool UPCGGraph::AddLabeledEdge(UPCGNode* From, const FName& FromPinLabel, UPCGNo
 
 	UPCGPin* ToPin = To->GetInputPin(ToPinLabel);
 
-	if(!ToPin)
+	if (!ToPin)
 	{
 		UE_LOG(LogPCG, Error, TEXT("To node %s does not have the %s label"), *To->GetName(), *ToPinLabel.ToString());
 		return false;
 	}
 
+#if WITH_EDITOR
+	DisableNotificationsForEditor();
+#endif
+
 	TSet<UPCGNode*> TouchedNodes;
 
 	// Create edge
@@ -669,6 +681,9 @@ bool UPCGGraph::AddLabeledEdge(UPCGNode* From, const FName& FromPinLabel, UPCGNo
 	const EPCGChangeType ChangeType = PCGGraphUtils::NotifyTouchedNodes(TouchedNodes, EPCGChangeType::Structural);
 
 #if WITH_EDITOR
+	// After all nodes are notified, re-enable graph notifications and send graph change notification.
+	EnableNotificationsForEditor();
+
 	NotifyGraphChanged(ChangeType);
 #endif
 
@@ -750,6 +765,10 @@ void UPCGGraph::RemoveNodes_Internal(TArrayView<UPCGNode*> InNodes)
 
 	Modify();
 
+#if WITH_EDITOR
+	DisableNotificationsForEditor();
+#endif
+
 	TSet<UPCGNode*> TouchedNodes;
 
 	for (UPCGNode* Node : InNodes)
@@ -772,6 +791,10 @@ void UPCGGraph::RemoveNodes_Internal(TArrayView<UPCGNode*> InNodes)
 		Nodes.Remove(Node);
 	}
 
+#if WITH_EDITOR
+	EnableNotificationsForEditor();
+#endif
+
 	PCGGraphUtils::NotifyTouchedNodes(TouchedNodes, EPCGChangeType::Structural);
 
 	OnNodesRemoved(InNodes);
@@ -785,6 +808,10 @@ bool UPCGGraph::RemoveEdge(UPCGNode* From, const FName& FromLabel, UPCGNode* To,
 		return false;
 	}
 
+#if WITH_EDITOR
+	DisableNotificationsForEditor();
+#endif
+
 	UPCGPin* OutPin = From->GetOutputPin(FromLabel);
 	UPCGPin* InPin = To->GetInputPin(ToLabel);
 
@@ -796,12 +823,15 @@ bool UPCGGraph::RemoveEdge(UPCGNode* From, const FName& FromLabel, UPCGNode* To,
 
 	const EPCGChangeType ChangeType = PCGGraphUtils::NotifyTouchedNodes(TouchedNodes, EPCGChangeType::Structural);
 
+#if WITH_EDITOR
+	// After all nodes are notified, re-enable graph notifications and send graph change notification.
+	EnableNotificationsForEditor();
+
 	if (TouchedNodes.Num() > 0)
 	{
-#if WITH_EDITOR
 		NotifyGraphChanged(ChangeType);
-#endif
 	}
+#endif
 
 	return TouchedNodes.Num() > 0;
 }
@@ -865,6 +895,10 @@ bool UPCGGraph::RemoveInboundEdges(UPCGNode* InNode, const FName& InboundLabel)
 	check(InNode);
 	TSet<UPCGNode*> TouchedNodes;
 
+#if WITH_EDITOR
+	DisableNotificationsForEditor();
+#endif
+
 	if (UPCGPin* InputPin = InNode->GetInputPin(InboundLabel))
 	{
 		InputPin->BreakAllEdges(&TouchedNodes);
@@ -873,6 +907,9 @@ bool UPCGGraph::RemoveInboundEdges(UPCGNode* InNode, const FName& InboundLabel)
 	const EPCGChangeType ChangeType = PCGGraphUtils::NotifyTouchedNodes(TouchedNodes, EPCGChangeType::Structural);
 
 #if WITH_EDITOR
+	// After all nodes are notified, re-enable graph notifications and send graph change notification.
+	EnableNotificationsForEditor();
+
 	if (TouchedNodes.Num() > 0)
 	{
 		NotifyGraphChanged(ChangeType);
@@ -888,6 +925,10 @@ bool UPCGGraph::RemoveOutboundEdges(UPCGNode* InNode, const FName& OutboundLabel
 	// Make a list of downstream nodes which may need pin updates when the edges change
 	TSet<UPCGNode*> TouchedNodes;
 
+#if WITH_EDITOR
+	DisableNotificationsForEditor();
+#endif
+
 	if (UPCGPin* OutputPin = InNode->GetOutputPin(OutboundLabel))
 	{
 		OutputPin->BreakAllEdges(&TouchedNodes);
@@ -896,6 +937,9 @@ bool UPCGGraph::RemoveOutboundEdges(UPCGNode* InNode, const FName& OutboundLabel
 	const EPCGChangeType ChangeType = PCGGraphUtils::NotifyTouchedNodes(TouchedNodes, EPCGChangeType::Structural);
 
 #if WITH_EDITOR
+	// After all nodes are notified, re-enable graph notifications and send graph change notification.
+	EnableNotificationsForEditor();
+
 	if (TouchedNodes.Num() > 0)
 	{
 		NotifyGraphChanged(ChangeType);
@@ -1078,6 +1122,8 @@ void UPCGGraph::GetTrackedActorKeysToSettings(FPCGSelectionKeyToSettingsMap& Out
 
 void UPCGGraph::NotifyGraphChanged(EPCGChangeType ChangeType)
 {
+	TRACE_CPUPROFILER_EVENT_SCOPE(UPCGGraph::NotifyGraphChanged);
+
 	if (GraphChangeNotificationsDisableCounter > 0)
 	{
 		bDelayedChangeNotification = true;
@@ -1141,6 +1187,8 @@ void UPCGGraph::NotifyGraphParametersChanged(EPCGGraphParameterEvent InChangeTyp
 
 void UPCGGraph::OnNodeChanged(UPCGNode* InNode, EPCGChangeType ChangeType)
 {
+	TRACE_CPUPROFILER_EVENT_SCOPE(UPCGGraph::OnNodeChanged);
+
 	if (!!(ChangeType & EPCGChangeType::Structural))
 	{
 		// Update node to grid size map for grid size changes.
diff --git a/Engine/Plugins/Experimental/PCG/Source/PCG/Private/PCGNode.cpp b/Engine/Plugins/Experimental/PCG/Source/PCG/Private/PCGNode.cpp
--- a/Engine/Plugins/Experimental/PCG/Source/PCG/Private/PCGNode.cpp
+++ b/Engine/Plugins/Experimental/PCG/Source/PCG/Private/PCGNode.cpp
@@ -1051,8 +1051,13 @@ EPCGChangeType UPCGNode::UpdatePins(TFunctionRef<UPCGPin*(UPCGNode*)> PinAllocat
 	return ChangeType;
 }
 
-EPCGChangeType UPCGNode::PropagateDynamicPinTypes(const UPCGNode* FromNode /*= nullptr*/)
+EPCGChangeType UPCGNode::PropagateDynamicPinTypes(TSet<UPCGNode*>& TouchedNodes, const UPCGNode* FromNode /*= nullptr*/)
 {
+	// TODO - we have a performance issue - this function can recurse very deeply and take a long time. Perf can be tested by wiring
+	// nodes from reroutes in large graph. Functionality can be tested by chaining reroute nodes and checking type propagation.
+	// We could probably have a visited set - if we ensure that we update the most-upstream nodes first,
+	// so that each node receives the final type? 
+
 	EPCGChangeType ChangeType = EPCGChangeType::None;
 	const UPCGSettings* Settings = GetSettings();
 	if (!Settings || !Settings->HasDynamicPins())
@@ -1062,10 +1067,8 @@ EPCGChangeType UPCGNode::PropagateDynamicPinTypes(const UPCGNode* FromNode /*= n
 
 	ChangeType |= UpdatePins();
 
-#if WITH_EDITOR
-	// Reconstruct in case pin types need to change
-	OnNodeChangedDelegate.Broadcast(this, EPCGChangeType::Node);
-#endif
+	// ChangeType intentionally ignored here - using it breaks propagation of types across chained reroutes.
+	TouchedNodes.Add(this);
 
 	for (UPCGPin* OutputPin : OutputPins)
 	{
@@ -1086,7 +1089,7 @@ EPCGChangeType UPCGNode::PropagateDynamicPinTypes(const UPCGNode* FromNode /*= n
 			const UPCGSettings* OtherSettings = OtherNode->GetSettings();
 			if (Settings && OtherSettings->HasDynamicPins())
 			{
-				ChangeType |= OtherNode->PropagateDynamicPinTypes(this);
+				ChangeType |= OtherNode->PropagateDynamicPinTypes(TouchedNodes, this);
 			}
 		}
 	}
diff --git a/Engine/Plugins/Experimental/PCG/Source/PCG/Public/PCGNode.h b/Engine/Plugins/Experimental/PCG/Source/PCG/Public/PCGNode.h
--- a/Engine/Plugins/Experimental/PCG/Source/PCG/Public/PCGNode.h
+++ b/Engine/Plugins/Experimental/PCG/Source/PCG/Public/PCGNode.h
@@ -145,7 +145,7 @@ public:
 	const TArray<TObjectPtr<UPCGPin>>& GetOutputPins() const { return OutputPins; }
 
 	/** Recursively follow downstream edges and call UpdatePins on each node that has dynamic pins. */
-	EPCGChangeType PropagateDynamicPinTypes(const UPCGNode* FromNode = nullptr);
+	EPCGChangeType PropagateDynamicPinTypes(TSet<UPCGNode*>& TouchedNodes, const UPCGNode* FromNode = nullptr);
 
 #if WITH_EDITOR
 	/** Transfer all editor only properties to the other node */
