[PCG] Dynamic branch culling and required pins added

* Nodes can nominate required pins which must have active inputs for the node to execute (`IsInputPinRequiredByExecution()`).
* Added small marker to node pins to communicate this, and augmented pin tooltip too.
* Nodes like Branch and Switch deactivate output pins by setting a bitmask. Inactive pins are visualised in graph editor.
* The bitmask is 64 bits. Added ensure to limit number of output pins to 63 (64th bit is used as special "no associated pin" value for pin-less task inputs).
* Dynamic branch culling - compilation time - compiler creates a boolean expression of pins that must be active for node to execute.
* Dynamic branch culling - execution time - if a node deactivates one or more pins, boolean expression in downstream tasks updated and tasks recursively culled.
* `FPCGPinDependencyExpression` added which represents this conjunctive normal form boolean in a flat array using conjunction markers.
* Dynamic & static culling editor visualization unified - both query whether node was executed. Reliable and simple. Only active during inspection & after execution.
* Inactive/culled node visualisation fades the node and displays a little "cancelled" badge.
* A manual check is still needed for redirectors which never execute but look bad if not showing active/inactive state, but this is fairly minor and easily maintained.
* A few nodes should always execute even if unwired (`CanCullTaskIfUnwired()`), most will cull.
* Debug object tree view - don't allow selecting non-debuggable stacks which causes weird behaviours like showing all nodes as culled, also breaks Profiling window.
* Debug object tree view - only call SetStackBeingInspected if the stack has changed, prevents a bunch of useless graph executions.
* PCG engine setting added to disable the 'culled' visualisation.
* Branch node fix - was not doing disabled behaviour (first pin statically active) for Dynamic branches.
* Add execution dependencies from static subgraph subnodes to the static subgraph node so that tasks execute alongside other subgraph tasks and can be culled
* Log added for dynamic culling, and cleaned up execution log vs caching log so that latter is only required for caching investigations

Related node behaviour tweaks requested by users:

* Distance node - if Target pin unconnected, passthrough data from Source pin (rather than executing).
* Attribute Transfer node - if Source pin is unconnected, passthrough data from Target pin (rather than executing).
* Switch node - if node is disabled, first output pin is active, others are inactive (similar to Branch).

diff --git a/Engine/Plugins/Experimental/PCG/Source/PCG/Private/Elements/ControlFlow/PCGBranch.cpp b/Engine/Plugins/Experimental/PCG/Source/PCG/Private/Elements/ControlFlow/PCGBranch.cpp
--- a/Engine/Plugins/Experimental/PCG/Source/PCG/Private/Elements/ControlFlow/PCGBranch.cpp
+++ b/Engine/Plugins/Experimental/PCG/Source/PCG/Private/Elements/ControlFlow/PCGBranch.cpp
@@ -67,6 +67,28 @@ FPCGElementPtr UPCGBranchSettings::CreateElement() const
 
 bool UPCGBranchSettings::IsPinStaticallyActive(const FName& PinLabel) const
 {
+	if (!bEnabled)
+	{
+		// Disabled - everything passed through first pin.
+		return PinLabel == PCGBranchConstants::OutputLabelA;
+	}
+
+	// Dynamic branches are never known in advance - assume all branches are active prior to execution.
+	if (IsPropertyOverriddenByPin(GET_MEMBER_NAME_CHECKED(UPCGBranchSettings, bOutputToB)))
+	{
+		return true;
+	}
+
+	return PinLabel == (bOutputToB ? PCGBranchConstants::OutputLabelB : PCGBranchConstants::OutputLabelA);
+}
+
+bool UPCGBranchSettings::IsPinUsedByNodeExecution(const UPCGPin* InPin) const
+{
+	if (!InPin->IsOutputPin())
+	{
+		return Super::IsPinUsedByNodeExecution(InPin);
+	}
+
 	// Dynamic branches are never known in advance - assume all branches are active prior to execution.
 	if (IsPropertyOverriddenByPin(GET_MEMBER_NAME_CHECKED(UPCGBranchSettings, bOutputToB)))
 	{
@@ -76,7 +98,7 @@ bool UPCGBranchSettings::IsPinStaticallyActive(const FName& PinLabel) const
 	// Branch must be both enabled and set to B in order for output pin B to be active.
 	const FName ActiveOutputPinLabel = (bEnabled && bOutputToB) ? PCGBranchConstants::OutputLabelB : PCGBranchConstants::OutputLabelA;
 
-	return PinLabel == ActiveOutputPinLabel;
+	return InPin->Properties.Label == ActiveOutputPinLabel;
 }
 
 bool FPCGBranchElement::ExecuteInternal(FPCGContext* Context) const
@@ -91,6 +113,9 @@ bool FPCGBranchElement::ExecuteInternal(FPCGContext* Context) const
 	// Reuse the functionality of the Gather Node
 	Context->OutputData = PCGGather::GatherDataForPin(Context->InputData, PCGPinConstants::DefaultInputLabel, SelectedPinLabel);
 
+	// Output bitmask of deactivated pins.
+	Context->OutputData.InactiveOutputPinBitmask = Settings->bOutputToB ? 1 : 2;
+
 	return true;
 }
 
diff --git a/Engine/Plugins/Experimental/PCG/Source/PCG/Private/Elements/ControlFlow/PCGSwitch.cpp b/Engine/Plugins/Experimental/PCG/Source/PCG/Private/Elements/ControlFlow/PCGSwitch.cpp
--- a/Engine/Plugins/Experimental/PCG/Source/PCG/Private/Elements/ControlFlow/PCGSwitch.cpp
+++ b/Engine/Plugins/Experimental/PCG/Source/PCG/Private/Elements/ControlFlow/PCGSwitch.cpp
@@ -46,6 +46,12 @@ FText UPCGSwitchSettings::GetNodeTooltipText() const
 
 bool UPCGSwitchSettings::IsPinStaticallyActive(const FName& PinLabel) const
 {
+	if (!bEnabled)
+	{
+		// If node disabled, active pin is first pin - first option or otherwise default.
+		return PinLabel == (CachedPinLabels.IsEmpty() ? PCGPinConstants::DefaultInputLabel : CachedPinLabels[0]);
+	}
+
 	// Dynamic branches are never known in advance - assume all branches are active prior to execution.
 	if (IsSwitchDynamic())
 	{
@@ -226,6 +232,33 @@ bool UPCGSwitchSettings::IsValuePresent(const int64 Value) const
 	return Index != INDEX_NONE && Index < EnumSelection.Class->NumEnums() - 1;
 }
 
+int UPCGSwitchSettings::GetSelectedOutputPinIndex() const
+{
+	if (SelectionMode == EPCGControlFlowSelectionMode::Integer)
+	{
+		// Return selected pin index or if no selection matched return the index after the options which will be "Default" pin.
+		const int Index = IntOptions.IndexOfByKey(IntegerSelection);
+		return Index != INDEX_NONE ? Index : IntOptions.Num();
+	}
+	else if (SelectionMode == EPCGControlFlowSelectionMode::String)
+	{
+		// Return selected pin index or if no selection matched return the index after the options which will be "Default" pin.
+		const int Index = StringOptions.IndexOfByKey(StringSelection);
+		return Index != INDEX_NONE ? Index : StringOptions.Num();
+	}
+	else if (SelectionMode == EPCGControlFlowSelectionMode::Enum)
+	{
+		// A "hidden" value could be selected that wasn't cached, so do a name-wise comparison
+		const FName PinLabel(EnumSelection.Class->GetDisplayNameTextByValue(EnumSelection.Value).ToString());
+
+		// Return index if found, otherwise fallback to the index after the options which will be "Default" pin.
+		const int FoundIndex = CachedPinLabels.IndexOfByKey(PinLabel);
+		return FoundIndex != INDEX_NONE ? FoundIndex : StringOptions.Num();
+	}
+
+	return INDEX_NONE;
+}
+
 bool UPCGSwitchSettings::GetSelectedPinLabel(FName& OutSelectedPinLabel) const
 {
 	if (CachedPinLabels.IsEmpty())
@@ -281,6 +314,38 @@ void UPCGSwitchSettings::CachePinLabels()
 	});
 }
 
+bool UPCGSwitchSettings::IsPinUsedByNodeExecution(const UPCGPin* InPin) const
+{
+	if (!InPin->IsOutputPin())
+	{
+		return Super::IsPinUsedByNodeExecution(InPin);
+	}
+
+	// Dynamic control flow never known statically - assume all branches are active prior to execution.
+	if (IsSwitchDynamic())
+	{
+		return true;
+	}
+
+	if (CachedPinLabels.IsEmpty())
+	{
+		// Labels not ready yet. Assumed used.
+		return true;
+	}
+
+	const int SelectedPinIndex = GetSelectedOutputPinIndex();
+	if (SelectedPinIndex == INDEX_NONE || !CachedPinLabels.IsValidIndex(SelectedPinIndex))
+	{
+		ensure(false);
+		return true;
+	}
+
+	// TODO disabled state? Discussing with Ryan.
+	//const FName ActiveOutputPinLabel = (bEnabled && bOutputToB) ? PCGBranchConstants::OutputLabelB : PCGBranchConstants::OutputLabelA;
+
+	return InPin->Properties.Label == CachedPinLabels[SelectedPinIndex];
+}
+
 bool FPCGSwitchElement::ExecuteInternal(FPCGContext* Context) const
 {
 	TRACE_CPUPROFILER_EVENT_SCOPE(FPCGSwitchElement::ExecuteInternal);
@@ -298,7 +363,19 @@ bool FPCGSwitchElement::ExecuteInternal(FPCGContext* Context) const
 	// Reuse the functionality of the Gather node
 	Context->OutputData = PCGGather::GatherDataForPin(Context->InputData, PCGPinConstants::DefaultInputLabel, SelectedPinLabel);
 
+	// Output bitmask of deactivated pins.
+	const int NumOutputPins = Context->Node ? Context->Node->GetOutputPins().Num() : 0;
+	if (ensure(NumOutputPins > 0))
+	{
+		const int SelectedPinIndex = Settings->GetSelectedOutputPinIndex();
+		if (ensure(SelectedPinIndex != INDEX_NONE))
+		{
+			const int AllPinsMask = (1 << NumOutputPins) - 1;
+			Context->OutputData.InactiveOutputPinBitmask = ~(1 << SelectedPinIndex) & AllPinsMask;
+		}
+	}
+
 	return true;
 }
 
-#undef LOCTEXT_NAMESPACE
\ No newline at end of file
+#undef LOCTEXT_NAMESPACE
diff --git a/Engine/Plugins/Experimental/PCG/Source/PCG/Private/Elements/PCGAttributeTransferElement.cpp b/Engine/Plugins/Experimental/PCG/Source/PCG/Private/Elements/PCGAttributeTransferElement.cpp
--- a/Engine/Plugins/Experimental/PCG/Source/PCG/Private/Elements/PCGAttributeTransferElement.cpp
+++ b/Engine/Plugins/Experimental/PCG/Source/PCG/Private/Elements/PCGAttributeTransferElement.cpp
@@ -7,6 +7,7 @@
 #include "PCGPin.h"
 #include "Data/PCGPointData.h"
 #include "Data/PCGSpatialData.h"
+#include "Elements/PCGGather.h"
 #include "Elements/Metadata/PCGMetadataElementCommon.h"
 #include "Metadata/Accessors/PCGAttributeAccessorHelpers.h"
 
@@ -127,6 +128,13 @@ bool FPCGAttributeTransferElement::ExecuteInternal(FPCGContext* Context) const
 
 	check(Context);
 
+	if (Context->Node && !Context->Node->IsInputPinConnected(PCGAttributeTransferConstants::SourceLabel))
+	{
+		// If Source pin is unconnected then we no-op and pass through all data from Target pin.
+		Context->OutputData = PCGGather::GatherDataForPin(Context->InputData, PCGAttributeTransferConstants::TargetLabel);
+		return true;
+	}
+
 	const UPCGAttributeTransferSettings* Settings = Context->GetInputSettings<UPCGAttributeTransferSettings>();
 	check(Settings);
 
diff --git a/Engine/Plugins/Experimental/PCG/Source/PCG/Private/Elements/PCGDistance.cpp b/Engine/Plugins/Experimental/PCG/Source/PCG/Private/Elements/PCGDistance.cpp
--- a/Engine/Plugins/Experimental/PCG/Source/PCG/Private/Elements/PCGDistance.cpp
+++ b/Engine/Plugins/Experimental/PCG/Source/PCG/Private/Elements/PCGDistance.cpp
@@ -1,11 +1,13 @@
 // Copyright Epic Games, Inc. All Rights Reserved.
 
 #include "Elements/PCGDistance.h"
-#include "Data/PCGSpatialData.h"
-#include "Helpers/PCGAsync.h"
-#include "Data/PCGPointData.h"
+
 #include "PCGContext.h"
 #include "PCGPin.h"
+#include "Data/PCGSpatialData.h"
+#include "Data/PCGPointData.h"
+#include "Elements/PCGGather.h"
+#include "Helpers/PCGAsync.h"
 
 #define LOCTEXT_NAMESPACE "PCGDistanceElement"
 
@@ -83,6 +85,14 @@ FPCGElementPtr UPCGDistanceSettings::CreateElement() const
 bool FPCGDistanceElement::ExecuteInternal(FPCGContext* Context) const
 {
 	TRACE_CPUPROFILER_EVENT_SCOPE(FPCGDistanceElement::Execute);
+	check(Context);
+
+	if (Context->Node && !Context->Node->IsInputPinConnected(PCGDistance::TargetLabel))
+	{
+		// If Target pin is unconnected then we no-op and pass through all data from Target pin.
+		Context->OutputData = PCGGather::GatherDataForPin(Context->InputData, PCGDistance::SourceLabel);
+		return true;
+	}
 
 	const UPCGDistanceSettings* Settings = Context->GetInputSettings<UPCGDistanceSettings>();
 	check(Settings);
diff --git a/Engine/Plugins/Experimental/PCG/Source/PCG/Private/Elements/PCGOuterIntersectionElement.cpp b/Engine/Plugins/Experimental/PCG/Source/PCG/Private/Elements/PCGOuterIntersectionElement.cpp
--- a/Engine/Plugins/Experimental/PCG/Source/PCG/Private/Elements/PCGOuterIntersectionElement.cpp
+++ b/Engine/Plugins/Experimental/PCG/Source/PCG/Private/Elements/PCGOuterIntersectionElement.cpp
@@ -9,13 +9,6 @@
 
 #define LOCTEXT_NAMESPACE "PCGOuterIntersectionElement"
 
-namespace PCGIntersectionConstants
-{
-	const FName PrimaryLabel = TEXT("Primary Source");
-	const FName SecondaryLabel = TEXT("Source");
-	const FText SecondaryTooltip = LOCTEXT("SecondaryPinLabel", "Secondary pin inputs will be implicitly unioned together before being compared to the primary pin's input for calculation of the intersection operation. Empty data passed along on one of the secondary pins will result in an empty intersection output, unless the 'Ignore Empty Secondary Input' flag is enabled.");
-}
-
 TArray<FPCGPinProperties> UPCGOuterIntersectionSettings::StaticInputPinProperties() const
 {
 	TArray<FPCGPinProperties> StaticPinProperties;
diff --git a/Engine/Plugins/Experimental/PCG/Source/PCG/Private/Elements/PCGSplineSampler.cpp b/Engine/Plugins/Experimental/PCG/Source/PCG/Private/Elements/PCGSplineSampler.cpp
--- a/Engine/Plugins/Experimental/PCG/Source/PCG/Private/Elements/PCGSplineSampler.cpp
+++ b/Engine/Plugins/Experimental/PCG/Source/PCG/Private/Elements/PCGSplineSampler.cpp
@@ -22,12 +22,6 @@
 
 #define LOCTEXT_NAMESPACE "PCGSplineSamplerElement"
 
-namespace PCGSplineSamplerConstants
-{
-	const FName SplineLabel = TEXT("Spline");
-	const FName BoundingShapeLabel = TEXT("Bounding Shape");
-}
-
 namespace PCGSplineSamplerHelpers
 {
 	/**
diff --git a/Engine/Plugins/Experimental/PCG/Source/PCG/Private/Elements/PCGSurfaceSampler.cpp b/Engine/Plugins/Experimental/PCG/Source/PCG/Private/Elements/PCGSurfaceSampler.cpp
--- a/Engine/Plugins/Experimental/PCG/Source/PCG/Private/Elements/PCGSurfaceSampler.cpp
+++ b/Engine/Plugins/Experimental/PCG/Source/PCG/Private/Elements/PCGSurfaceSampler.cpp
@@ -22,12 +22,6 @@
 
 #define LOCTEXT_NAMESPACE "PCGSurfaceSamplerElement"
 
-namespace PCGSurfaceSamplerConstants
-{
-	const FName SurfaceLabel = TEXT("Surface");
-	const FName BoundingShapeLabel = TEXT("Bounding Shape");
-}
-
 namespace PCGSurfaceSampler
 {
 	bool FSurfaceSamplerParams::Initialize(const UPCGSurfaceSamplerSettings* Settings, const FPCGContext* Context, const FBox& InputBounds)
diff --git a/Engine/Plugins/Experimental/PCG/Source/PCG/Private/Elements/PCGVolumeSampler.cpp b/Engine/Plugins/Experimental/PCG/Source/PCG/Private/Elements/PCGVolumeSampler.cpp
--- a/Engine/Plugins/Experimental/PCG/Source/PCG/Private/Elements/PCGVolumeSampler.cpp
+++ b/Engine/Plugins/Experimental/PCG/Source/PCG/Private/Elements/PCGVolumeSampler.cpp
@@ -18,12 +18,6 @@
 
 #define LOCTEXT_NAMESPACE "PCGVolumeSamplerElement"
 
-namespace PCGVolumeSamplerConstants
-{
-	const FName VolumeLabel = TEXT("Volume");
-	const FName BoundingShapeLabel = TEXT("Bounding Shape");
-}
-
 namespace PCGVolumeSampler
 {
 	UPCGPointData* SampleVolume(FPCGContext* Context, const FVolumeSamplerParams& SamplerSettings, const UPCGSpatialData* Volume, const UPCGSpatialData* BoundingShape)
diff --git a/Engine/Plugins/Experimental/PCG/Source/PCG/Private/Graph/PCGGraphCompiler.cpp b/Engine/Plugins/Experimental/PCG/Source/PCG/Private/Graph/PCGGraphCompiler.cpp
--- a/Engine/Plugins/Experimental/PCG/Source/PCG/Private/Graph/PCGGraphCompiler.cpp
+++ b/Engine/Plugins/Experimental/PCG/Source/PCG/Private/Graph/PCGGraphCompiler.cpp
@@ -10,9 +10,19 @@
 #include "Elements/PCGHiGenGridSize.h"
 #include "Elements/PCGReroute.h"
 #include "Graph/PCGGraphExecutor.h"
+#include "Graph/PCGPinDependencyExpression.h"
 
+#include "HAL/IConsoleManager.h"
 #include "Misc/ScopeRWLock.h"
 
+namespace PCGGraphCompiler
+{
+	TAutoConsoleVariable<bool> CVarEnableTaskStaticCulling(
+		TEXT("pcg.GraphExecution.TaskStaticCulling"),
+		true,
+		TEXT("Enable static culling of tasks which considers static branches, generation grid size, trivial nodes and more."));
+}
+
 TArray<FPCGGraphTask> FPCGGraphCompiler::CompileGraph(UPCGGraph* InGraph, FPCGTaskId& NextId, FPCGStackContext& InOutStackContext)
 {
 	if (!InGraph)
@@ -131,10 +141,13 @@ TArray<FPCGGraphTask> FPCGGraphCompiler::CompileGraph(UPCGGraph* InGraph, FPCGTa
 					continue;
 				}
 
+				// We hook to pretask if the element requires data from the pretask, or otherwise if there are no inputs, to ensure
+				// that the element executes with the other subgraph tasks (and can be culled if the subgraph node is culled).
 				const UPCGSettings* Settings = Subtask.Node->GetSettings();
-				if (Settings && Settings->ShouldHookToPreTask())
+				const bool bRequiresDataFromPreTask = Settings && Settings->RequiresDataFromPreTask();
+				if (bRequiresDataFromPreTask || Subtask.Inputs.IsEmpty())
 				{
-					Subtask.Inputs.Emplace(PreId, nullptr, nullptr);
+					Subtask.Inputs.Emplace(PreId, /*InInboundPin=*/nullptr, /*InOutboundPin=*/nullptr, bRequiresDataFromPreTask);
 				}
 			}
 
@@ -151,10 +164,14 @@ TArray<FPCGGraphTask> FPCGGraphCompiler::CompileGraph(UPCGGraph* InGraph, FPCGTa
 			// especially since we cannot distinguish between the pre and post during execution so any data filtering related to pins is bound to fail.
 			PostTask.Element = MakeShared<FPCGTrivialElement>();
 
+			// Add execution-only dependency on pre-task, without this post task can be scheduled concurrently with pre-task, and concurrently
+			// with something that might become inactive and would then fail to dynamically cull this already-scheduled task.
+			PostTask.Inputs.Emplace(PreId, /*InInboundPin=*/nullptr, /*InOutboundPin=*/nullptr, /*bInProvideData=*/false);
+
 			// Add subgraph output node task as input to the post-task
 			if (OutputNodeTask)
 			{
-				PostTask.Inputs.Emplace(OutputNodeTask->NodeId, nullptr, nullptr);
+				PostTask.Inputs.Emplace(OutputNodeTask->NodeId, /*InInboundPin=*/nullptr, /*InOutboundPin=*/nullptr);
 			}
 
 			check(!IdMapping.Contains(Node));
@@ -467,21 +484,71 @@ EPCGHiGenGrid FPCGGraphCompiler::CalculateGridRecursive(
 	return Grid;
 }
 
-bool FPCGGraphCompiler::CalculateActiveRecursive(FPCGTaskId InTaskId, const TArray<FPCGGraphTask>& InCompiledTasks, TMap<int32, bool>& InTaskIdToActiveFlag)
+bool FPCGGraphCompiler::CalculateStaticallyActiveRecursive(FPCGTaskId InTaskId, const TArray<FPCGGraphTask>& InCompiledTasks, TMap<int32, bool>& InOutTaskIdToActiveFlag)
 {
-	if (const bool* bEntry = InTaskIdToActiveFlag.Find(InTaskId))
+	if (const bool* bEntry = InOutTaskIdToActiveFlag.Find(InTaskId))
 	{
 		return *bEntry;
 	}
 
-	bool bAnyInputActive = false;
+	// Nodes within subgraphs - if the subgraph node is inactive then all tasks within the subgraph are inactive.
+	if (InCompiledTasks[InTaskId].ParentId != InvalidPCGTaskId)
+	{
+		const bool bParentActive = CalculateStaticallyActiveRecursive(InCompiledTasks[InTaskId].ParentId, InCompiledTasks, InOutTaskIdToActiveFlag);
+		if (!bParentActive)
+		{
+			InOutTaskIdToActiveFlag.Add(InTaskId, bParentActive);
+
+			return bParentActive;
+		}
+	}
+
+	const UPCGNode* Node = InCompiledTasks[InTaskId].Node;
+	if (!Node)
+	{
+		InOutTaskIdToActiveFlag.Add(InTaskId, true);
+		return true;
+	}
+
+	TArray<FName, TInlineAllocator<8>> PinsRequiringActiveConnection;
 
-	for (FPCGGraphTaskInput Input : InCompiledTasks[InTaskId].Inputs)
+	// Three relevant types of input pins - required, non-advanced and advanced. This tracks if we encounter the second category.
+	bool bHasAnyNonAdvancedPins = false;
+
+	for (UPCGPin* InputPin : Node->GetInputPins())
+	{
+		if (!InputPin)
+		{
+			continue;
+		}
+
+		bHasAnyNonAdvancedPins |= !InputPin->Properties.bAdvancedPin;
+
+		if (Node->IsInputPinRequiredByExecution(InputPin))
+		{
+			PinsRequiringActiveConnection.AddUnique(InputPin->Properties.Label);
+		}
+	}
+	
+	bool bHasAnyNonAdvancedInputs = false;
+	bool bHasAnyActiveNonAdvancedInput = false;
+
+	for (const FPCGGraphTaskInput& Input : InCompiledTasks[InTaskId].Inputs)
 	{
+		const UPCGPin* InputPin = Input.OutPin;
+
+		// Only non-advanced input pins play a part in determining active/inactive state.
+		if (InputPin && InputPin->Properties.bAdvancedPin)
+		{
+			continue;
+		}
+
+		bHasAnyNonAdvancedInputs = true;
+
 		// Default to tasks being active unless proved otherwise.
 		bool bInputActive = true;
 
-		// If we are connected to an upstream branch node, evaluate if the output pin is active.
+		// If we are connected to an upstream node, evaluate if the output pin is active.
 		const UPCGNode* UpstreamNode = InCompiledTasks[Input.TaskId].Node;
 		if (const UPCGSettings* UpstreamSettings = UpstreamNode ? UpstreamNode->GetSettings() : nullptr)
 		{
@@ -490,22 +557,49 @@ bool FPCGGraphCompiler::CalculateActiveRecursive(FPCGTaskId InTaskId, const TArr
 
 		if (bInputActive)
 		{
-			bInputActive &= CalculateActiveRecursive(Input.TaskId, InCompiledTasks, InTaskIdToActiveFlag);
+			bInputActive &= CalculateStaticallyActiveRecursive(Input.TaskId, InCompiledTasks, InOutTaskIdToActiveFlag);
 		}
 
 		if (bInputActive)
 		{
-			bAnyInputActive = true;
-			break;
+			bHasAnyActiveNonAdvancedInput = true;
+
+			if (InputPin)
+			{
+				// Register received input on this pin.
+				PinsRequiringActiveConnection.Remove(InputPin->Properties.Label);
+			}
 		}
 	}
 
-	const bool bActive = bAnyInputActive || InCompiledTasks[InTaskId].Inputs.IsEmpty();
-	InTaskIdToActiveFlag.Add(InTaskId, bActive);
+	const UPCGSettings* Settings = Node->GetSettings();
+	const bool bCanCullIfUnwired = Settings && Settings->CanCullTaskIfUnwired();
+
+	bool bActive = true;
+
+	if (!PinsRequiringActiveConnection.IsEmpty())
+	{
+		// If PinsRequiringActiveConnection is not empty then we did not find an input for each required pin.
+		bActive = false;
+	}
+	else if (bCanCullIfUnwired)
+	{
+		// Cull if we have non-advanced pins but we don't have any active non-advanced inputs.
+		bActive = !(bHasAnyNonAdvancedPins && !bHasAnyActiveNonAdvancedInput);
+	}
+	else if (bHasAnyNonAdvancedInputs)
+	{
+		// This task is allowed to be unwired, so we have non-advanced inputs and they're all inactive - basically
+		// all upstream inputs are inactive which forces this task to be inactive.
+		bActive = bHasAnyActiveNonAdvancedInput || InCompiledTasks[InTaskId].Inputs.IsEmpty();
+	}
+
+	InOutTaskIdToActiveFlag.Add(InTaskId, bActive);
+
 	return bActive;
 }
 
-void FPCGGraphCompiler::CullTasksStaticBranchNodes(TArray<FPCGGraphTask>& InOutCompiledTasks)
+void FPCGGraphCompiler::CullTasksStaticInactive(TArray<FPCGGraphTask>& InOutCompiledTasks)
 {
 	if (InOutCompiledTasks.IsEmpty())
 	{
@@ -518,7 +612,8 @@ void FPCGGraphCompiler::CullTasksStaticBranchNodes(TArray<FPCGGraphTask>& InOutC
 
 	for (int32 i = 1; i < InOutCompiledTasks.Num(); ++i)
 	{
-		CalculateActiveRecursive(InOutCompiledTasks[i].NodeId, InOutCompiledTasks, NodeIdToActiveFlag);
+		// Results of each call memoized via NodeIdToActiveFlag.
+		CalculateStaticallyActiveRecursive(InOutCompiledTasks[i].NodeId, InOutCompiledTasks, NodeIdToActiveFlag);
 	}
 
 	auto ShouldCull = [&NodeIdToActiveFlag](const FPCGGraphTask& InTask)
@@ -679,6 +774,112 @@ void FPCGGraphCompiler::PostCullStackCleanup(TArray<FPCGGraphTask>& InCompiledTa
 	}
 }
 
+void FPCGGraphCompiler::CalculateDynamicActivePinDependencies(FPCGTaskId InTaskId, TArray<FPCGGraphTask>& InOutCompiledTasks)
+{
+	if (!ensure(InOutCompiledTasks.IsValidIndex(InTaskId)))
+	{
+		return;
+	}
+
+	const UPCGNode* Node = InOutCompiledTasks[InTaskId].Node;
+
+	// First if there are any required pins, build the pin dependencies from these. We have an array of inputs, each one can
+	// correspond to a pin. Use a map to compile an expression for each input pin label.
+	TMap<FName, FPCGPinDependencyExpression> InputPinLabelToPinDependency;
+
+	for (const FPCGGraphTaskInput& Input : InOutCompiledTasks[InTaskId].Inputs)
+	{
+		if (!Node || !Input.OutPin)
+		{
+			continue;
+		}
+
+		const UPCGPin* InputPin = Input.OutPin;
+
+		// Consider only primary input pins in this pass.
+		if (!Node->IsInputPinRequiredByExecution(InputPin))
+		{
+			continue;
+		}
+
+		const UPCGNode* UpstreamNode = InOutCompiledTasks[Input.TaskId].Node;
+		if (!UpstreamNode)
+		{
+			continue;
+		}
+
+		const int PinIndex = UpstreamNode->GetOutputPins().IndexOfByPredicate([&Input](const UPCGPin* InPin)
+		{
+			return InPin == Input.InPin;
+		});
+
+		if (PinIndex != INDEX_NONE)
+		{
+			check(PinIndex < PCGPinIdHelpers::MaxOutputPins);
+
+			FPCGPinDependencyExpression& PinDependency = InputPinLabelToPinDependency.FindOrAdd(Input.OutPin->Properties.Label);
+			PinDependency.AddPinDependency(PCGPinIdHelpers::NodeIdAndPinIndexToPinId(Input.TaskId, PinIndex));
+		}
+	}
+
+	// Result expression. Conjunction of disjunctions of pin IDs that are required to be active for this task to be active.
+	// Example - keep task if: UpstreamPin0Active && (UpstreamPin1Active || UpstreamPin2Active)
+	FPCGPinDependencyExpression PinDependency;
+
+	if (!InputPinLabelToPinDependency.IsEmpty())
+	{
+		// If we have registered pin dependencies from the first pass (required pins) then we have can compose these for our pin
+		// dependency expression. Build conjunction from the per-pin disjunctions.
+		for (TPair<FName, FPCGPinDependencyExpression>& PinExpression : InputPinLabelToPinDependency)
+		{
+			PinDependency.AppendUsingConjunction(PinExpression.Value);
+		}
+	}
+	else
+	{
+		// If we don't have any dependent pins logged by now then there are no required pins (note that a node that does not have task
+		// inputs for required pins will be statically culled in an earlier compilation step). In which case we'll be active if *any* input
+		// is active. We build a disjunction that expresses this.
+		for (const FPCGGraphTaskInput& Input : InOutCompiledTasks[InTaskId].Inputs)
+		{
+			if (Input.OutPin && Input.OutPin->Properties.bAdvancedPin)
+			{
+				// Advanced input pins never participate in keeping node active.
+				continue;
+			}
+
+			if (const UPCGPin* UpstreamOutputPin = Input.InPin)
+			{
+				// Input connection is via node pins.
+				const UPCGNode* UpstreamNode = InOutCompiledTasks[Input.TaskId].Node;
+				if (!UpstreamNode)
+				{
+					continue;
+				}
+
+				const int PinIndex = UpstreamNode->GetOutputPins().IndexOfByPredicate([UpstreamOutputPin](const UPCGPin* InPin)
+				{
+					return InPin == UpstreamOutputPin;
+				});
+
+				if (PinIndex != INDEX_NONE)
+				{
+					check(PinIndex < PCGPinIdHelpers::MaxOutputPins);
+
+					PinDependency.AddPinDependency(PCGPinIdHelpers::NodeIdAndPinIndexToPinId(Input.TaskId, PinIndex));
+				}
+			}
+			else
+			{
+				// No associated pin, use a special pin ID for a pin-less dependency.
+				PinDependency.AddPinDependency(PCGPinIdHelpers::NodeIdToPinId(Input.TaskId));
+			}
+		}
+	}
+
+	InOutCompiledTasks[InTaskId].PinDependency = MoveTemp(PinDependency);
+}
+
 TArray<FPCGGraphTask> FPCGGraphCompiler::GetCompiledTasks(UPCGGraph* InGraph, uint32 GenerationGridSize, FPCGStackContext& OutStackContext, bool bIsTopGraph)
 {
 	TArray<FPCGGraphTask> CompiledTasks;
@@ -741,6 +942,8 @@ void FPCGGraphCompiler::OffsetNodeIds(TArray<FPCGGraphTask>& Tasks, FPCGTaskId O
 		{
 			Input.TaskId += Offset;
 		}
+
+		Task.PinDependency.OffsetNodeIds(Offset);
 	}
 }
 
@@ -766,39 +969,55 @@ void FPCGGraphCompiler::CompileTopGraph(UPCGGraph* InGraph, uint32 GenerationGri
 		return;
 	}
 
-	// Remove reroute nodes before execution grid setup, as grid linkages need final nodes to connect from/to.
-	CullTasks(CompiledTasks, /*bAddPassthroughWires=*/true, [](const FPCGGraphTask& InTask) { return InTask.Node && Cast<UPCGRerouteSettings>(InTask.Node->GetSettings()); });
+	if (PCGGraphCompiler::CVarEnableTaskStaticCulling.GetValueOnAnyThread())
+	{
+		// Remove reroute nodes before execution grid setup, as grid linkages need final nodes to connect from/to.
+		CullTasks(CompiledTasks, /*bAddPassthroughWires=*/true, [](const FPCGGraphTask& InTask) { return InTask.Node && Cast<UPCGRerouteSettings>(InTask.Node->GetSettings()); });
 
-	// Cull inactive branches downstream of branch nodes with static selection values.
-	CullTasksStaticBranchNodes(CompiledTasks);
+		// Cull inactive branches downstream of branch nodes with static selection values.
+		CullTasksStaticInactive(CompiledTasks);
 
-	// For hierarchical generation resolve the execution grid for each task and cull any tasks that won't execute.
-	// TODO - we could add an else branch for higen disabled that culls the grid size nodes.
-	if (InGraph->IsHierarchicalGenerationEnabled() && GenerationGridSize != PCGHiGenGrid::UninitializedGridSize())
-	{
-		const EPCGHiGenGrid GenerationGrid = PCGHiGenGrid::GridSizeToGrid(GenerationGridSize);
-		const EPCGHiGenGrid DefaultGrid = PCGHiGenGrid::GridSizeToGrid(InGraph->GetDefaultGridSize());
+		// For hierarchical generation resolve the execution grid for each task and cull any tasks that won't execute.
+		// TODO - we could add an else branch for higen disabled that culls the grid size nodes.
+		if (InGraph->IsHierarchicalGenerationEnabled() && GenerationGridSize != PCGHiGenGrid::UninitializedGridSize())
+		{
+			const EPCGHiGenGrid GenerationGrid = PCGHiGenGrid::GridSizeToGrid(GenerationGridSize);
+			const EPCGHiGenGrid DefaultGrid = PCGHiGenGrid::GridSizeToGrid(InGraph->GetDefaultGridSize());
 
-		// Propagate grid size nodes through the graph to determine which grid size each task should execute on.
-		TArray<EPCGHiGenGrid> TaskGenerationGrid;
-		TaskGenerationGrid.SetNumZeroed(CompiledTasks.Num());
-		ResolveGridSizes(GenerationGrid, CompiledTasks, StackContext, DefaultGrid, TaskGenerationGrid);
+			// Propagate grid size nodes through the graph to determine which grid size each task should execute on.
+			TArray<EPCGHiGenGrid> TaskGenerationGrid;
+			TaskGenerationGrid.SetNumZeroed(CompiledTasks.Num());
+			ResolveGridSizes(GenerationGrid, CompiledTasks, StackContext, DefaultGrid, TaskGenerationGrid);
 
-		// Create linkage tasks for edges that cross from large grid to small grid tasks.
-		CreateGridLinkages(GenerationGrid, TaskGenerationGrid, CompiledTasks, StackContext);
+			// Create linkage tasks for edges that cross from large grid to small grid tasks.
+			CreateGridLinkages(GenerationGrid, TaskGenerationGrid, CompiledTasks, StackContext);
 
-		// Cull any task that should not execute on the current grid.
-		CullTasks(CompiledTasks, /*bAddPassthroughWires=*/false, [GenerationGrid, &TaskGenerationGrid](const FPCGGraphTask& InTask)
-		{
-			const EPCGHiGenGrid TaskGrid = TaskGenerationGrid[InTask.NodeId];
-			return TaskGrid != EPCGHiGenGrid::Uninitialized && !(TaskGrid & GenerationGrid);
-		});
+			// Cull any task that should not execute on the current grid.
+			CullTasks(CompiledTasks, /*bAddPassthroughWires=*/false, [GenerationGrid, &TaskGenerationGrid](const FPCGGraphTask& InTask)
+			{
+				const EPCGHiGenGrid TaskGrid = TaskGenerationGrid[InTask.NodeId];
+				return TaskGrid != EPCGHiGenGrid::Uninitialized && !(TaskGrid & GenerationGrid);
+			});
+		}
 	}
 
 	// Post culling - remove any stacks that are no longer part of execution. Besides being tidy this also helps
 	// debug tools discern which stacks were executed or not.
 	PostCullStackCleanup(CompiledTasks, StackContext);
 
+	// To feed dynamic culling at execution time, build list of upstream pins that we depend on (if all of these pins
+	// are determined to be inactive at execution time, then the task will be deactivated). An empty list means the
+	// task will never be deactivated and will always execute.
+	// TODO: Expand pin dependencies to pure branches that aren't directly downstream - nodes that have no side effects
+	// and feed into a branch can also be culled if the branch is culled.
+	// TODO: Pin dependencies should be transitive across nodes. If a node is dependent on a single upstream node, it could
+	// likely take the pin dependencies from the upstream node, which should save iterations in the dynamic culling code.
+	for (int TaskIndex = 0; TaskIndex < CompiledTasks.Num(); ++TaskIndex)
+	{
+		// Result is written directly to tasks.
+		CalculateDynamicActivePinDependencies(CompiledTasks[TaskIndex].NodeId, CompiledTasks);
+	}
+
 	const int TaskNum = CompiledTasks.Num();
 	const FPCGTaskId PreExecuteTaskId = FPCGTaskId(TaskNum);
 	const FPCGTaskId PostExecuteTaskId = PreExecuteTaskId + 1;
diff --git a/Engine/Plugins/Experimental/PCG/Source/PCG/Private/Graph/PCGGraphCompiler.h b/Engine/Plugins/Experimental/PCG/Source/PCG/Private/Graph/PCGGraphCompiler.h
--- a/Engine/Plugins/Experimental/PCG/Source/PCG/Private/Graph/PCGGraphCompiler.h
+++ b/Engine/Plugins/Experimental/PCG/Source/PCG/Private/Graph/PCGGraphCompiler.h
@@ -58,10 +58,10 @@ private:
 		const FPCGStackContext& InStackContext);
 
 	/** Discovers whether task is on a statically active branch (and needs to be passed to graph executor). */
-	static bool CalculateActiveRecursive(FPCGTaskId InTaskId, const TArray<FPCGGraphTask>& InCompiledTasks, TMap<int32, bool>& InTaskIdToActiveFlag);
+	static bool CalculateStaticallyActiveRecursive(FPCGTaskId InTaskId, const TArray<FPCGGraphTask>& InCompiledTasks, TMap<int32, bool>& InTaskIdToActiveFlag);
 
-	/** Culls nodes that are not on statically active branches. */
-	static void CullTasksStaticBranchNodes(TArray<FPCGGraphTask>& InOutCompiledTasks);
+	/** Culls nodes that are inactive for e.g. missing required inputs or on a statically inactive branch. */
+	static void CullTasksStaticInactive(TArray<FPCGGraphTask>& InOutCompiledTasks);
 
 	/** Culls tasks based on a given lambda. Never culls the first (input) task in the array. */
 	static void CullTasks(TArray<FPCGGraphTask>& InOutCompiledTasks, bool bAddPassthroughWires, TFunctionRef<bool(const FPCGGraphTask&)> CullTask);
@@ -69,6 +69,9 @@ private:
 	/** Remove any stack frames that are not used by any task. */
 	static void PostCullStackCleanup(TArray<FPCGGraphTask>& InCompiledTasks, FPCGStackContext& InOutStackContext);
 
+	/** Task will be active if *any* of the upstream pin IDs are active, or if the pin ID list is empty. */
+	static void CalculateDynamicActivePinDependencies(FPCGTaskId InTaskId, TArray<FPCGGraphTask>& InOutCompiledTasks);
+
 	mutable FRWLock GraphToTaskMapLock;
 	TMap<UPCGGraph*, TArray<FPCGGraphTask>> GraphToTaskMap;
 	TMap<UPCGGraph*, FPCGStackContext> GraphToStackContext;
diff --git a/Engine/Plugins/Experimental/PCG/Source/PCG/Private/Graph/PCGGraphExecutor.cpp b/Engine/Plugins/Experimental/PCG/Source/PCG/Private/Graph/PCGGraphExecutor.cpp
--- a/Engine/Plugins/Experimental/PCG/Source/PCG/Private/Graph/PCGGraphExecutor.cpp
+++ b/Engine/Plugins/Experimental/PCG/Source/PCG/Private/Graph/PCGGraphExecutor.cpp
@@ -66,6 +66,11 @@ namespace PCGGraphExecutor
 		1000.0f / 20.0f,
 		TEXT("Allocated time in ms per frame when running in editor (non pie)"));
 #endif
+
+	TAutoConsoleVariable<bool> CVarDynamicTaskCulling(
+		TEXT("pcg.Graph.DynamicTaskCulling"),
+		true,
+		TEXT("Controls whether tasks are culled at execution time, for example in response to an deactivated dynamic branch pin"));
 }
 
 FPCGGraphExecutor::FPCGGraphExecutor()
@@ -112,11 +117,7 @@ FPCGTaskId FPCGGraphExecutor::Schedule(
 {
 	check(SourceComponent);
 
-	if (IsGraphCacheDebuggingEnabled())
-	{
-		UE_LOG(LogPCG, Log, TEXT("[%s] --- SCHEDULE GRAPH ---"),
-			(SourceComponent && SourceComponent->GetOwner()) ? *SourceComponent->GetOwner()->GetName() : TEXT("MISSINGCOMPONENT"));
-	}
+	PCGGraphExecutionLogging::LogGraphSchedule(SourceComponent);
 	
 	FPCGTaskId ScheduledId = InvalidPCGTaskId;
 
@@ -559,6 +560,11 @@ void FPCGGraphExecutor::Execute()
 		}
 	}
 
+	if (!ScheduledTasks.IsEmpty())
+	{
+		PCGGraphExecutionLogging::LogGraphPostSchedule(Tasks, TaskSuccessors);
+	}
+
 	ScheduledTasks.Reset();
 
 	ScheduleLock.Unlock();
@@ -591,8 +597,8 @@ void FPCGGraphExecutor::Execute()
 	const float MaxPercentageOfThreadsToUse = FMath::Clamp(CVarMaxPercentageOfThreadsToUse.GetValueOnAnyThread(), 0.0f, 1.0f);
 	const int32 MaxNumThreads = FMath::Max(0, FMath::Min((int32)(FPlatformMisc::NumberOfCoresIncludingHyperthreads() * MaxPercentageOfThreadsToUse), CVarMaxNumTasks.GetValueOnAnyThread() - 1));
 	const bool bAllowMultiDispatch = PCGGraphExecutor::CVarGraphMultithreading.GetValueOnAnyThread();
-	const bool bGraphCacheDebuggingEnabled = IsGraphCacheDebuggingEnabled();
 	const bool bStripEmptyPointData = CVarStripEmptyPointData.GetValueOnAnyThread();
+	const bool bDynamicTaskCulling = PCGGraphExecutor::CVarDynamicTaskCulling.GetValueOnAnyThread();
 
 #if WITH_EDITOR
 	UpdateGenerationNotification();
@@ -705,6 +711,8 @@ void FPCGGraphExecutor::Execute()
 					continue;
 				}
 
+				PCGGraphExecutionLogging::LogTaskExecute(Task);
+
 				// If a task is cacheable and has been cached, then we don't need to create an active task for it unless
 				// there is an execution mode that would prevent us from doing so.
 				const UPCGSettingsInterface* TaskSettingsInterface = TaskInput.GetSettingsInterface(Task.Node ? Task.Node->GetSettingsInterface() : nullptr);
@@ -718,9 +726,9 @@ void FPCGGraphExecutor::Execute()
 					Task.Element->GetDependenciesCrc(TaskInput, TaskSettings, Task.SourceComponent.Get(), DependenciesCrc);
 				}
 
-				if (bGraphCacheDebuggingEnabled && !bCacheable && Task.SourceComponent.Get() && Task.Node)
+				if (!bCacheable)
 				{
-					UE_LOG(LogPCG, Warning, TEXT("[%s] %s\t\tCACHING DISABLED"), *Task.SourceComponent->GetOwner()->GetName(), *Task.Node->GetNodeTitle(EPCGNodeTitleType::ListView).ToString());
+					PCGGraphExecutionLogging::LogTaskExecuteCachingDisabled(Task);
 				}
 
 				FPCGDataCollection CachedOutput;
@@ -734,7 +742,7 @@ void FPCGGraphExecutor::Execute()
 				if (!bNeedsToCreateActiveTask)
 				{
 #if WITH_EDITOR
-					// doing this now since we're about to modify ReadyTasks potentially reallocating while Task is a reference. 
+					// Doing this now since we're about to modify ReadyTasks potentially reallocating while Task is a reference. 
 					if (UPCGComponent* SourceComponent = Task.SourceComponent.Get())
 					{
 						if (Task.StackIndex != INDEX_NONE)
@@ -745,6 +753,15 @@ void FPCGGraphExecutor::Execute()
 					}
 #endif
 
+					if (bDynamicTaskCulling && TaskSettings && TaskSettings->OutputPinsCanBeDeactivated() && CachedOutput.InactiveOutputPinBitmask != 0)
+					{
+						CullInactiveDownstreamNodes(Task.NodeId, CachedOutput.InactiveOutputPinBitmask);
+
+#if WITH_EDITOR
+						SendInactivePinNotification(Task.Node, Task.StackContext->GetStack(Task.StackIndex), CachedOutput.InactiveOutputPinBitmask);
+#endif
+					}
+
 					// Fast-forward cached result to stored results
 					FPCGTaskId SkippedTaskId = Task.NodeId;
 					StoreResults(SkippedTaskId, CachedOutput);
@@ -852,9 +869,28 @@ void FPCGGraphExecutor::Execute()
 			}
 		}
 
-		auto PostTaskExecute = [this, &bAnyTaskEnded, bGraphCacheDebuggingEnabled, bStripEmptyPointData](int32 TaskIndex)
+		auto PostTaskExecute = [this, &bAnyTaskEnded, bStripEmptyPointData](int32 TaskIndex)
 		{
 			FPCGGraphActiveTask& ActiveTask = ActiveTasks[TaskIndex];
+			check(ActiveTask.Context);
+
+			const UPCGSettingsInterface* ActiveTaskSettingsInterface = ActiveTask.Context->GetInputSettingsInterface();
+			const uint64 InactivePinMask = ActiveTask.Context->OutputData.InactiveOutputPinBitmask;
+
+			if (InactivePinMask != 0 && ActiveTaskSettingsInterface)
+			{
+				const UPCGSettings* ActiveTaskSettings = ActiveTaskSettingsInterface ? ActiveTaskSettingsInterface->GetSettings() : nullptr;
+
+				// If output pins may have been deactivated then perform culling and update information for editor visualization.
+				if (ActiveTaskSettings && ActiveTaskSettings->OutputPinsCanBeDeactivated())
+				{
+					CullInactiveDownstreamNodes(ActiveTask.NodeId, InactivePinMask);
+
+#if WITH_EDITOR
+					SendInactivePinNotification(ActiveTask.Context->Node, ActiveTask.StackContext->GetStack(ActiveTask.StackIndex), InactivePinMask);
+#endif
+				}
+			}
 
 #if WITH_EDITOR
 			if (!ActiveTask.bWasCancelled && !ActiveTask.bIsBypassed)
@@ -862,13 +898,9 @@ void FPCGGraphExecutor::Execute()
 			if (!ActiveTask.bWasCancelled)
 #endif
 			{
-				if (bGraphCacheDebuggingEnabled && ActiveTask.Context->SourceComponent.Get() && ActiveTask.Context->Node)
-				{
-					UE_LOG(LogPCG, Log, TEXT("         [%s] %s\t\tOUTPUT CRC %u"), *ActiveTask.Context->SourceComponent->GetOwner()->GetName(), *ActiveTask.Context->Node->GetNodeTitle(EPCGNodeTitleType::ListView).ToString(), ActiveTask.Context->OutputData.Crc.GetValue());
-				}
+				PCGGraphExecutionLogging::LogTaskExecuteOutputCRC(ActiveTask);
 
 				// Store result in cache as needed - done here because it needs to be done on the main thread
-				const UPCGSettingsInterface* ActiveTaskSettingsInterface = ActiveTask.Context->GetInputSettingsInterface();
 
 				// Don't store if errors or warnings present
 #if WITH_EDITOR
@@ -879,7 +911,6 @@ void FPCGGraphExecutor::Execute()
 
 				if (ActiveTaskSettingsInterface && !bHasErrorOrWarning && ActiveTask.Element->IsCacheableInstance(ActiveTaskSettingsInterface))
 				{
-					const UPCGSettings* ActiveTaskSettings = ActiveTaskSettingsInterface ? ActiveTaskSettingsInterface->GetSettings() : nullptr;
 					GraphCache.StoreInCache(ActiveTask.Element.Get(), ActiveTask.Context->DependenciesCrc, ActiveTask.Context->OutputData);
 				}
 			}
@@ -1023,10 +1054,7 @@ void FPCGGraphExecutor::Execute()
 			UpdateGenerationNotification();
 #endif
 
-			if (bGraphCacheDebuggingEnabled)
-			{
-				UE_LOG(LogPCG, Log, TEXT("--- FINISH FPCGGRAPHEXECUTOR::EXECUTE ---"));
-			}
+			PCGGraphExecutionLogging::LogGraphExecuteFrameFinished();
 		}
 
 		// Purge things from cache if memory usage is too high
@@ -1067,7 +1095,7 @@ void FPCGGraphExecutor::ClearAllTasks()
 	SleepingTasks.Reset();
 }
 
-void FPCGGraphExecutor::QueueNextTasks(FPCGTaskId FinishedTask)
+void FPCGGraphExecutor::QueueNextTasks(FPCGTaskId FinishedTask, bool bIgnoreMissingTasks)
 {
 	TRACE_CPUPROFILER_EVENT_SCOPE(FPCGGraphExecutor::QueueNextTasks);
 
@@ -1079,7 +1107,7 @@ void FPCGGraphExecutor::QueueNextTasks(FPCGTaskId FinishedTask)
 			FPCGGraphTask* SuccessorTaskPtr = Tasks.Find(Successor);
 
 			// This should never be null, but later recovery should be able to cleanup this properly
-			if (ensure(SuccessorTaskPtr))
+			if (SuccessorTaskPtr)
 			{
 				FPCGGraphTask& SuccessorTask = *SuccessorTaskPtr;
 
@@ -1094,6 +1122,10 @@ void FPCGGraphExecutor::QueueNextTasks(FPCGTaskId FinishedTask)
 					Tasks.Remove(Successor);
 				}
 			}
+			else
+			{
+				ensure(bIgnoreMissingTasks);
+			}
 		}
 
 		TaskSuccessors.Remove(FinishedTask);
@@ -1313,6 +1345,151 @@ void FPCGGraphExecutor::ClearResults()
 	ScheduleLock.Unlock();
 }
 
+void FPCGGraphExecutor::GetPinIdsToDeactivate(FPCGTaskId TaskId, uint64 InactiveOutputPinBitmask, TArray<FPCGPinId>& InOutPinIds)
+{
+	InOutPinIds.Reserve(InOutPinIds.Num() + FMath::CountBits(InactiveOutputPinBitmask));
+
+	int OutputPinIndex = 0;
+
+	while (InactiveOutputPinBitmask != 0)
+	{
+		if (InactiveOutputPinBitmask & 1)
+		{
+			InOutPinIds.AddUnique(PCGPinIdHelpers::NodeIdAndPinIndexToPinId(TaskId, OutputPinIndex));
+		}
+
+		InactiveOutputPinBitmask >>= 1;
+		++OutputPinIndex;
+	}
+}
+
+void FPCGGraphExecutor::CullInactiveDownstreamNodes(FPCGTaskId InCompletedTaskId, uint64 InInactiveOutputPinBitmask)
+{
+	TRACE_CPUPROFILER_EVENT_SCOPE(FPCGGraphExecutor::CullInactiveDownstreamNodes);
+
+	TArray<FPCGPinId> PinIdsToDeactivate;
+	GetPinIdsToDeactivate(InCompletedTaskId, InInactiveOutputPinBitmask, PinIdsToDeactivate);
+	check(!PinIdsToDeactivate.IsEmpty());
+
+	PCGGraphExecutionLogging::LogTaskCullingBegin(InCompletedTaskId, InInactiveOutputPinBitmask, PinIdsToDeactivate);
+
+	TSet<FPCGTaskId> AllRemovedTasks;
+
+	// Hoisted out of loop for performance reasons.
+	TArray<FPCGTaskId, TInlineAllocator<64>> TasksToRemove;
+
+	while (!PinIdsToDeactivate.IsEmpty())
+	{
+		const FPCGPinId PinId = PinIdsToDeactivate.Pop(/*bAllowShrinking=*/false);
+		const FPCGTaskId PinTaskId = PCGPinIdHelpers::GetNodeIdFromPinId(PinId);
+
+		PCGGraphExecutionLogging::LogTaskCullingBeginLoop(PinTaskId, PCGPinIdHelpers::GetPinIndexFromPinId(PinId), PinIdsToDeactivate);
+		LogTaskState();
+
+		const TSet<FPCGTaskId>* Successors = TaskSuccessors.Find(PinTaskId);
+		if (!Successors)
+		{
+			continue;
+		}
+
+		TasksToRemove.SetNum(0, EAllowShrinking::No);
+
+		// Build set of tasks that are candidates for culling when PinId is deactivated.
+		for (const FPCGTaskId SuccessorTaskId : *Successors)
+		{
+			// Successors are updated at the end of this function, which means it may
+			// contain task IDs that have been removed.
+			if (FPCGGraphTask* FoundTask = Tasks.Find(SuccessorTaskId))
+			{
+				bool bDependencyExpressionBecameFalse;
+				FoundTask->PinDependency.DeactivatePin(PinId, bDependencyExpressionBecameFalse);
+
+				if (bDependencyExpressionBecameFalse)
+				{
+					TasksToRemove.AddUnique(SuccessorTaskId);
+				}
+
+				PCGGraphExecutionLogging::LogTaskCullingUpdatedPinDeps(SuccessorTaskId, FoundTask->PinDependency, bDependencyExpressionBecameFalse);
+			}
+		}
+
+		// Now remove the tasks.
+		for (const FPCGTaskId RemovedTaskId : TasksToRemove)
+		{
+			// Scope in which RemovedTask reference is valid.
+			{
+				FPCGGraphTask& RemovedTask = Tasks[RemovedTaskId];
+
+				const UPCGNode* Node = RemovedTask.Node;
+				const int PinCount = Node ? Node->GetOutputPins().Num() : 0;
+
+				if (PinCount > 0)
+				{
+					// Deactivate all output pins.
+					const uint64 InactiveOutputPinBitmask = (1 << PinCount) - 1;
+
+					// Deactivate its pins - add to set of pins to deactivate.
+					GetPinIdsToDeactivate(RemovedTaskId, InactiveOutputPinBitmask, PinIdsToDeactivate);
+
+#if WITH_EDITOR
+					SendInactivePinNotification(RemovedTask.Node, RemovedTask.StackContext->GetStack(RemovedTask.StackIndex), InactiveOutputPinBitmask);
+#endif
+				}
+
+				// Also register a special pin-less pin ID for this node, for task dependencies that do not have a specific pin.
+				PinIdsToDeactivate.AddUnique(PCGPinIdHelpers::NodeIdToPinId(RemovedTaskId));
+
+				// Remove task as successor of upstream node.
+				RemoveTaskFromInputSuccessors(RemovedTaskId, RemovedTask.Inputs);
+
+				// Remove the deleted tasks from the inputs of downstream tasks.
+				if (const TSet<FPCGTaskId>* SuccessorsOfRemovedTask = TaskSuccessors.Find(RemovedTaskId))
+				{
+					for (const FPCGTaskId SuccessorTaskId : *SuccessorsOfRemovedTask)
+					{
+						if (FPCGGraphTask* SuccessorTask = Tasks.Find(SuccessorTaskId))
+						{
+							for (int InputIndex = SuccessorTask->Inputs.Num() - 1; InputIndex >= 0; --InputIndex)
+							{
+								if (SuccessorTask->Inputs[InputIndex].TaskId == RemovedTaskId)
+								{
+									SuccessorTask->Inputs.RemoveAtSwap(InputIndex);
+								}
+							}
+						}
+					}
+				}
+			}
+
+			// Remove from tasks. After this step all traces of RemovedTaskId should be erased from tasks, task inputs. Task successors will be
+			// updated below when queuing next tasks.
+			Tasks.Remove(RemovedTaskId);
+		}
+
+		AllRemovedTasks.Append(TasksToRemove);
+	}
+
+	// Ensure any downstream tasks are enqueued.
+	for (const FPCGTaskId TaskId : AllRemovedTasks)
+	{
+		// Queue downstream tasks in a similar manner to when a task draws from the cache and is skipped.
+		// Some downstream tasks will have been culled which we don't care about (hence the ignore flag),
+		// but some may not be queued and may be ready for queuing.
+		QueueNextTasks(TaskId, /*bIgnoreMissingTasks=*/true);
+	}
+}
+
+#if WITH_EDITOR
+void FPCGGraphExecutor::SendInactivePinNotification(const UPCGNode* InNode, const FPCGStack* InStack, uint64 InactiveOutputPinBitmask)
+{
+	const UPCGComponent* Component = InStack ? InStack->GetRootComponent() : nullptr;
+	if (Component && InNode)
+	{
+		Component->NotifyNodeDynamicInactivePins(InNode, InStack, InactiveOutputPinBitmask);
+	}
+}
+#endif
+
 void FPCGGraphExecutor::AddReferencedObjects(FReferenceCollector& Collector)
 {
 	TRACE_CPUPROFILER_EVENT_SCOPE(FPCGGraphExecutor::AddReferencedObjects);
@@ -1350,6 +1527,17 @@ FPCGElementPtr FPCGGraphExecutor::GetFetchInputElement()
 	return FetchInputElement;
 }
 
+void FPCGGraphExecutor::LogTaskState() const
+{
+#if WITH_EDITOR
+	if (PCGGraphExecutionLogging::CullingLogEnabled())
+	{
+		UE_LOG(LogPCG, Log, TEXT("\tDORMANT (FPCGGraphExecutor::Tasks):"));
+		PCGGraphExecutionLogging::LogGraphTasks(Tasks, &TaskSuccessors);
+	}
+#endif
+}
+
 #if WITH_EDITOR
 
 FPCGTaskId FPCGGraphExecutor::ScheduleDebugWithTaskCallback(UPCGComponent* InComponent, TFunction<void(FPCGTaskId/* TaskId*/, const UPCGNode*/* Node*/, const FPCGDataCollection&/* TaskOutput*/)> TaskCompleteCallback)
diff --git a/Engine/Plugins/Experimental/PCG/Source/PCG/Private/Graph/PCGGraphExecutor.h b/Engine/Plugins/Experimental/PCG/Source/PCG/Private/Graph/PCGGraphExecutor.h
--- a/Engine/Plugins/Experimental/PCG/Source/PCG/Private/Graph/PCGGraphExecutor.h
+++ b/Engine/Plugins/Experimental/PCG/Source/PCG/Private/Graph/PCGGraphExecutor.h
@@ -7,6 +7,7 @@
 #include "PCGSubsystem.h"
 #include "Graph/PCGGraphCache.h"
 #include "Graph/PCGGraphCompiler.h"
+#include "Graph/PCGPinDependencyExpression.h"
 #include "Graph/PCGStackContext.h"
 
 #include "UObject/GCObject.h"
@@ -60,7 +61,6 @@ struct FPCGGraphTaskInput
 struct FPCGGraphTask
 {
 	TArray<FPCGGraphTaskInput> Inputs;
-	//TArray<DataId> Outputs;
 	const UPCGNode* Node = nullptr;
 	TWeakObjectPtr<UPCGComponent> SourceComponent = nullptr;
 	FPCGElementPtr Element; // Added to have tasks that aren't node-bound
@@ -69,6 +69,11 @@ struct FPCGGraphTask
 	FPCGTaskId CompiledTaskId = InvalidPCGTaskId; // the task id as it exists when compiled
 	FPCGTaskId ParentId = InvalidPCGTaskId; // represents the parent sub object graph task, if we were called from one
 
+	/** Conjunction of disjunctions of pin IDs that are required to be active for this task to be active.
+	* Example - keep task if: UpstreamPin0Active && (UpstreamPin1Active || UpstreamPin2Active)
+	*/
+	FPCGPinDependencyExpression PinDependency;
+
 	int32 StackIndex = INDEX_NONE;
 	TSharedPtr<const FPCGStackContext> StackContext;
 };
@@ -190,7 +195,7 @@ public:
 private:
 	TSet<UPCGComponent*> Cancel(TFunctionRef<bool(TWeakObjectPtr<UPCGComponent>)> CancelFilter);
 	void ClearAllTasks();
-	void QueueNextTasks(FPCGTaskId FinishedTask);
+	void QueueNextTasks(FPCGTaskId FinishedTask, bool bIgnoreMissingTasks = false);
 	bool CancelNextTasks(FPCGTaskId CancelledTask, TSet<UPCGComponent*>& OutCancelledComponents);
 	void RemoveTaskFromInputSuccessors(FPCGTaskId CancelledTask, const TArray<FPCGGraphTaskInput>& CancelledTaskInputs);
 	void BuildTaskInput(const FPCGGraphTask& Task, FPCGDataCollection& TaskInput);
@@ -199,9 +204,20 @@ private:
 	void StoreResults(FPCGTaskId InTaskId, const FPCGDataCollection& InTaskOutput);
 	void ClearResults();
 
+	/** If the completed task has one or more deactivated pins, delete any downstream tasks that are inactive as a result. */
+	void CullInactiveDownstreamNodes(FPCGTaskId CompletedTaskId, uint64 InInactiveOutputPinBitmask);
+
+	/** Builds an array of all deactivated unique pin IDs. */
+	void GetPinIdsToDeactivate(FPCGTaskId TaskId, uint64 InactiveOutputPinBitmask, TArray<FPCGPinId>& InOutPinIds);
+
 	FPCGElementPtr GetFetchInputElement();
 
+	void LogTaskState() const;
+
 #if WITH_EDITOR
+	/** Notify the component that the given pins were deactivated during execution. */
+	void SendInactivePinNotification(const UPCGNode* InNode, const FPCGStack* InStack, uint64 InactiveOutputPinBitmask);
+
 	void SaveDirtyActors();
 	void ReleaseUnusedActors();
 
diff --git a/Engine/Plugins/Experimental/PCG/Source/PCG/Private/Graph/PCGPinDependencyExpression.cpp b/Engine/Plugins/Experimental/PCG/Source/PCG/Private/Graph/PCGPinDependencyExpression.cpp
new file mode 100644
--- /dev/null
+++ b/Engine/Plugins/Experimental/PCG/Source/PCG/Private/Graph/PCGPinDependencyExpression.cpp
@@ -0,0 +1,153 @@
+// Copyright Epic Games, Inc. All Rights Reserved.
+
+#include "Graph/PCGPinDependencyExpression.h"
+
+void FPCGPinDependencyExpression::AddPinDependency(FPCGPinId PinId)
+{
+	check(PinId != ConjunctionMarker && PinId != RemovedTermMarker);
+	Expression.Add(PinId);
+}
+
+void FPCGPinDependencyExpression::AddConjunction()
+{
+	Expression.Add(ConjunctionMarker);
+}
+
+void FPCGPinDependencyExpression::AppendUsingConjunction(const FPCGPinDependencyExpression& Other)
+{
+	if (Other.Expression.IsEmpty())
+	{
+		return;
+	}
+
+	// If this expression is not empty, add a conjunction to join the two expressions.
+	if (!Expression.IsEmpty())
+	{
+		AddConjunction();
+	}
+
+	Expression.Append(Other.Expression);
+}
+
+void FPCGPinDependencyExpression::DeactivatePin(FPCGPinId PinId, bool& bOutExpressionBecameFalse)
+{
+	bOutExpressionBecameFalse = false;
+
+	for (int I = 0; I < Expression.Num(); ++I)
+	{
+		if (Expression[I] == PinId)
+		{
+			// Remove pin by replacing with conjunction marker. Could have removed the element but we don't
+			// expect to see long expressions on average so this uses a simple approach.
+			Expression[I] = RemovedTermMarker;
+
+			// Check if this is was the last term in a disjunction (if we're surrounded by conjunction markers or removed pins).
+			bool bNoTermsOnLeft = true;
+			for (int J = I - 1; J >= 0; --J)
+			{
+				// End of disjunction
+				if (Expression[J] == ConjunctionMarker)
+				{
+					break;
+				}
+
+				if (Expression[J] != RemovedTermMarker)
+				{
+					bNoTermsOnLeft = false;
+					break;
+				}
+			}
+
+			bool bNoTermsOnRight = true;
+			for (int J = I + 1; J < Expression.Num(); ++J)
+			{
+				// End of disjunction
+				if (Expression[J] == ConjunctionMarker)
+				{
+					break;
+				}
+
+				if (Expression[J] != RemovedTermMarker)
+				{
+					bNoTermsOnRight = false;
+					break;
+				}
+			}
+
+			if (bNoTermsOnLeft && bNoTermsOnRight)
+			{
+				// Removing this term made the disjunction become false, and hence the whole expression.
+				// Don't early out however - make sure to update all disjunctions.
+				bOutExpressionBecameFalse = true;
+			}
+		}
+	}
+}
+
+void FPCGPinDependencyExpression::OffsetNodeIds(uint64 NodeIdOffset)
+{
+	for (FPCGPinId& PinId : Expression)
+	{
+		if (PinId != ConjunctionMarker)
+		{
+			PinId = PCGPinIdHelpers::OffsetNodeIdInPinId(PinId, NodeIdOffset);
+		}
+	}
+}
+
+#if WITH_EDITOR
+
+FString FPCGPinDependencyExpression::ToString() const
+{
+	FString ExpressionString;
+
+	bool InDisjunction = false;
+
+	for (int I = 0; I < Expression.Num(); ++I)
+	{
+		if (Expression[I] == ConjunctionMarker)
+		{
+			ExpressionString += TEXT(" && ");
+
+			if (InDisjunction)
+			{
+				ExpressionString += TEXT(")");
+				InDisjunction = false;
+			}
+		}
+		else
+		{
+			if (InDisjunction)
+			{
+				ExpressionString += TEXT(" || ");
+			}
+			else
+			{
+				ExpressionString += TEXT("(");
+			}
+
+			if (Expression[I] == RemovedTermMarker)
+			{
+				ExpressionString += TEXT(".");
+			}
+			else
+			{
+				// Decompose pin ID into task ID and pin index which is friendly to read.
+				const FPCGTaskId NodeId = PCGPinIdHelpers::GetNodeIdFromPinId(Expression[I]);
+				const uint64 PinIndex = PCGPinIdHelpers::GetPinIndexFromPinId(Expression[I]);
+				ExpressionString += FString::Printf(TEXT("%u_%u"), NodeId, PinIndex);
+			}
+
+			InDisjunction = true;
+		}
+	}
+
+	if (InDisjunction)
+	{
+		ExpressionString += TEXT(")");
+	}
+
+	return ExpressionString;
+}
+
+#endif // WITH_EDITOR
diff --git a/Engine/Plugins/Experimental/PCG/Source/PCG/Private/Graph/PCGPinDependencyExpression.h b/Engine/Plugins/Experimental/PCG/Source/PCG/Private/Graph/PCGPinDependencyExpression.h
new file mode 100644
--- /dev/null
+++ b/Engine/Plugins/Experimental/PCG/Source/PCG/Private/Graph/PCGPinDependencyExpression.h
@@ -0,0 +1,47 @@
+// Copyright Epic Games, Inc. All Rights Reserved.
+
+#pragma once
+
+#include "PCGCommon.h"
+
+/** Conjunction of disjunctions of pin IDs that are required to be active for this task to be active.
+* 
+* Example:
+*     Keep task if: UpstreamPin0Active && (UpstreamPin1Active || UpstreamPin2Active)
+*     This will be stored in an internal array as: [ PinID0, ConjunctionMarker, PinID1, PinID2 ]
+*     If Pin2 is deactivated during execution: [ PinID0, ConjunctionMarker, PinID1, RemovedMarker ]
+*     If Pin1 is deactivated during execution: [ PinID0, ConjunctionMarker, RemovedMarker, RemovedMarker ]
+*     Expression is now false and the task using this expression can be culled.
+*/
+struct FPCGPinDependencyExpression
+{
+	/** Expression building. Appends a pin ID to the current disjunction. */
+	void AddPinDependency(FPCGPinId PinId);
+
+	/** Expression building. Add a boolean && to the expression. */
+	void AddConjunction();
+
+	/** Expression building. Append the given expression using a conjunction, result will be Expression && Other.Expression. */
+	void AppendUsingConjunction(const FPCGPinDependencyExpression& Other);
+
+	/** Expression evaluation. Deactivating a pin is done by removing it from the expression. If any disjunction collapses
+	* then the expression becomes false and the task can be culled.
+	*/
+	void DeactivatePin(FPCGPinId PinId, bool& bOutExpressionBecameFalse);
+
+	/** Applies given node ID offset to all pin IDs in expression (pin ID is combination of node ID and pin index). */
+	void OffsetNodeIds(uint64 NodeIdOffset);
+
+#if WITH_EDITOR
+	FString ToString() const;
+#endif
+
+private:
+	/* These are used to demarcate conjunctions '&&' in the boolean expression. */
+	static const FPCGPinId ConjunctionMarker = std::numeric_limits<FPCGPinId>::max();
+
+	/* Signifies a pin removed from the expression, used to set terms false (pin inactive). */
+	static const FPCGPinId RemovedTermMarker = std::numeric_limits<FPCGPinId>::max() - 1;
+
+	TArray<FPCGPinId> Expression;
+};
diff --git a/Engine/Plugins/Experimental/PCG/Source/PCG/Private/PCGCommon.cpp b/Engine/Plugins/Experimental/PCG/Source/PCG/Private/PCGCommon.cpp
--- a/Engine/Plugins/Experimental/PCG/Source/PCG/Private/PCGCommon.cpp
+++ b/Engine/Plugins/Experimental/PCG/Source/PCG/Private/PCGCommon.cpp
@@ -97,3 +97,33 @@ namespace PCGDelegates
 #endif
 } // PCGDelegates
 
+namespace PCGPinIdHelpers
+{
+	FPCGPinId NodeIdAndPinIndexToPinId(FPCGTaskId NodeId, uint64 PinIndex)
+	{
+		// Construct a unique ID from node index and pin index.
+		ensure(PinIndex < PCGPinIdHelpers::MaxOutputPins);
+		return (NodeId * PCGPinIdHelpers::PinActiveBitmaskSize) + (PinIndex % PCGPinIdHelpers::PinActiveBitmaskSize);
+	}
+
+	FPCGPinId NodeIdToPinId(FPCGTaskId NodeId)
+	{
+		// Use (max pins - 1) to make a special pin ID that is not associated to a specific node pin.
+		return (NodeId * PCGPinIdHelpers::PinActiveBitmaskSize) + PCGPinIdHelpers::MaxOutputPins;
+	}
+
+	FPCGPinId OffsetNodeIdInPinId(FPCGPinId PinId, uint64 NodeIDOffset)
+	{
+		return NodeIDOffset * PCGPinIdHelpers::PinActiveBitmaskSize + PinId;
+	}
+
+	FPCGTaskId GetNodeIdFromPinId(FPCGPinId PinId)
+	{
+		return PinId / PCGPinIdHelpers::PinActiveBitmaskSize;
+	}
+
+	uint64 GetPinIndexFromPinId(FPCGPinId PinId)
+	{
+		return PinId % PCGPinIdHelpers::PinActiveBitmaskSize;
+	}
+}
diff --git a/Engine/Plugins/Experimental/PCG/Source/PCG/Private/PCGComponent.cpp b/Engine/Plugins/Experimental/PCG/Source/PCG/Private/PCGComponent.cpp
--- a/Engine/Plugins/Experimental/PCG/Source/PCG/Private/PCGComponent.cpp
+++ b/Engine/Plugins/Experimental/PCG/Source/PCG/Private/PCGComponent.cpp
@@ -1880,6 +1880,52 @@ void UPCGComponent::DisableInspection()
 	}
 };
 
+void UPCGComponent::NotifyNodeExecuted(const UPCGNode* InNode, const FPCGStack* InStack)
+{
+	if (!ensure(InStack && InNode))
+	{
+		return;
+	}
+
+	FWriteScopeLock Lock(NodeToStacksInWhichNodeExecutedLock);
+	NodeToStacksInWhichNodeExecuted.FindOrAdd(InNode).Add(*InStack);
+}
+
+uint64 UPCGComponent::GetNodeInactivePinMask(const UPCGNode* InNode, const FPCGStack& Stack) const
+{
+	FReadScopeLock Lock(NodeToStackToInactivePinMaskLock);
+
+	if (const TMap<const FPCGStack, uint64>* StackToMask = NodeToStackToInactivePinMask.Find(InNode))
+	{
+		if (const uint64* Mask = StackToMask->Find(Stack))
+		{
+			return *Mask;
+		}
+	}
+
+	return 0;
+}
+
+void UPCGComponent::NotifyNodeDynamicInactivePins(const UPCGNode* InNode, const FPCGStack* InStack, uint64 InactivePinBitmask) const
+{
+	if (!ensure(InStack && InNode))
+	{
+		return;
+	}
+
+	FWriteScopeLock Lock(NodeToStackToInactivePinMaskLock);
+	TMap<const FPCGStack, uint64>& StackToInactivePinMask = NodeToStackToInactivePinMask.FindOrAdd(InNode);
+	StackToInactivePinMask.FindOrAdd(*InStack) = InactivePinBitmask;
+}
+
+bool UPCGComponent::WasNodeExecuted(const UPCGNode* InNode, const FPCGStack& Stack) const
+{
+	FReadScopeLock Lock(NodeToStacksInWhichNodeExecutedLock);
+	const TSet<FPCGStack>* FoundStacks = NodeToStacksInWhichNodeExecuted.Find(InNode);
+
+	return FoundStacks && FoundStacks->Contains(Stack);
+}
+
 void UPCGComponent::StoreInspectionData(const FPCGStack* InStack, const UPCGNode* InNode, const FPCGDataCollection& InInputData, const FPCGDataCollection& InOutputData)
 {
 	if (!InNode || !ensure(InStack))
@@ -1887,6 +1933,9 @@ void UPCGComponent::StoreInspectionData(const FPCGStack* InStack, const UPCGNode
 		return;
 	}
 
+	// Notify component that this task executed. Useful for editor visualization.
+	NotifyNodeExecuted(InNode, InStack);
+
 	if (!InOutputData.TaggedData.IsEmpty())
 	{
 		FWriteScopeLock Lock(NodeToStacksThatProducedDataLock);
@@ -1958,6 +2007,16 @@ void UPCGComponent::ClearInspectionData()
 		FWriteScopeLock Lock(NodeToStacksThatProducedDataLock);
 		NodeToStacksThatProducedData.Reset();
 	}
+
+	{
+		FWriteScopeLock Lock(NodeToStacksInWhichNodeExecutedLock);
+		NodeToStacksInWhichNodeExecuted.Reset();
+	}
+
+	{
+		FWriteScopeLock Lock(NodeToStackToInactivePinMaskLock);
+		NodeToStackToInactivePinMask.Reset();
+	}
 }
 
 bool UPCGComponent::HasNodeProducedData(const UPCGNode* InNode, const FPCGStack& Stack) const
diff --git a/Engine/Plugins/Experimental/PCG/Source/PCG/Private/PCGGraph.cpp b/Engine/Plugins/Experimental/PCG/Source/PCG/Private/PCGGraph.cpp
--- a/Engine/Plugins/Experimental/PCG/Source/PCG/Private/PCGGraph.cpp
+++ b/Engine/Plugins/Experimental/PCG/Source/PCG/Private/PCGGraph.cpp
@@ -1160,8 +1160,6 @@ void UPCGGraph::NotifyGraphChanged(EPCGChangeType ChangeType)
 	// Graph settings, nodes, graph structure can all change the higen grid sizes.
 	if (ChangeType != EPCGChangeType::Cosmetic)
 	{
-		ResetNodeToOnActiveBranchMap();
-
 		FWriteScopeLock GridSizeLock(NodeToGridSizeLock);
 		NodeToGridSize.Reset();
 	}
@@ -1200,9 +1198,6 @@ void UPCGGraph::OnNodeChanged(UPCGNode* InNode, EPCGChangeType ChangeType)
 			NodeToGridSize.Reset();
 		}
 
-		// Any node/edge change can affect which branches are statically active/inactive.
-		ResetNodeToOnActiveBranchMap();
-
 		// Broadcast so that grid size visualization can be updated editor-side.
 		OnGraphStructureChangedDelegate.Broadcast(this);
 	}
@@ -1416,24 +1411,6 @@ FInstancedPropertyBag* UPCGGraph::GetMutableUserParametersStruct()
 	return &UserParameters;
 }
 
-#if WITH_EDITOR
-bool UPCGGraph::IsNodeOnActiveBranch(const UPCGNode* InNode) const
-{
-	{
-		FReadScopeLock Lock(NodeToOnActiveBranchLock);
-		if (const bool* CachedValue = NodeToOnActiveBranch.Find(InNode))
-		{
-			return *CachedValue;
-		}
-	}
-
-	{
-		FWriteScopeLock Lock(NodeToOnActiveBranchLock);
-		return CalculateNodeOnActiveBranchRecursive_Unsafe(InNode);
-	}
-}
-#endif // WITH_EDITOR
-
 uint32 UPCGGraph::GetNodeGenerationGridSize(const UPCGNode* InNode, uint32 InDefaultGridSize) const
 {
 	{
@@ -1498,69 +1475,6 @@ uint32 UPCGGraph::CalculateNodeGridSizeRecursive_Unsafe(const UPCGNode* InNode,
 	return GridSize;
 }
 
-#if WITH_EDITOR
-bool UPCGGraph::CalculateNodeOnActiveBranchRecursive_Unsafe(const UPCGNode* InNode) const
-{
-	if (const bool* CachedValue = NodeToOnActiveBranch.Find(InNode))
-	{
-		return *CachedValue;
-	}
-
-	bool bAnyInputActive = false;
-	bool bAnyEdgesPresent = false;
-
-	for (const UPCGPin* Pin : InNode->GetInputPins())
-	{
-		if (!Pin)
-		{
-			continue;
-		}
-
-		for (const UPCGEdge* Edge : Pin->Edges)
-		{
-			bAnyEdgesPresent = true;
-
-			bool bEdgeActive = true;
-
-			const UPCGPin* OtherPin = Edge ? Edge->InputPin : nullptr;
-			if (OtherPin && OtherPin->Node.Get())
-			{
-				if (const UPCGSettings* UpstreamSettings = OtherPin->Node->GetSettings())
-				{
-					bEdgeActive &= UpstreamSettings->IsPinStaticallyActive(OtherPin->Properties.Label);
-				}
-
-				bEdgeActive = bEdgeActive && CalculateNodeOnActiveBranchRecursive_Unsafe(OtherPin->Node);
-			}
-
-			if (bEdgeActive)
-			{
-				bAnyInputActive = true;
-				break;
-			}
-		}
-
-		if (bAnyInputActive)
-		{
-			break;
-		}
-	}
-
-	// Active if any input is active, or if there are no edges.
-	const bool bActive = bAnyInputActive || !bAnyEdgesPresent;
-
-	NodeToOnActiveBranch.Add(InNode, bActive);
-
-	return bActive;
-}
-
-void UPCGGraph::ResetNodeToOnActiveBranchMap()
-{
-	FWriteScopeLock Lock(NodeToOnActiveBranchLock);
-	NodeToOnActiveBranch.Reset();
-}
-#endif
-
 void UPCGGraph::AddUserParameters(const TArray<FPropertyBagPropertyDesc>& InDescs, const UPCGGraph* InOptionalOriginalGraph)
 {
 	UserParameters.AddProperties(InDescs);
diff --git a/Engine/Plugins/Experimental/PCG/Source/PCG/Private/PCGNode.cpp b/Engine/Plugins/Experimental/PCG/Source/PCG/Private/PCGNode.cpp
--- a/Engine/Plugins/Experimental/PCG/Source/PCG/Private/PCGNode.cpp
+++ b/Engine/Plugins/Experimental/PCG/Source/PCG/Private/PCGNode.cpp
@@ -637,6 +637,12 @@ const UPCGPin* UPCGNode::GetPassThroughOutputPin() const
 	return GetOutputPins().Num() > PrimaryOutputPinIndex ? GetOutputPins()[PrimaryOutputPinIndex] : nullptr;
 }
 
+bool UPCGNode::IsInputPinRequiredByExecution(const UPCGPin* InPin) const
+{
+	const UPCGSettings* Settings = GetSettings();
+	return !Settings || Settings->IsInputPinRequiredByExecution(InPin);
+}
+
 bool UPCGNode::IsPinUsedByNodeExecution(const UPCGPin* InPin) const
 {
 	check(InPin);
diff --git a/Engine/Plugins/Experimental/PCG/Source/PCG/Private/Utils/PCGGraphExecutionLogging.cpp b/Engine/Plugins/Experimental/PCG/Source/PCG/Private/Utils/PCGGraphExecutionLogging.cpp
--- a/Engine/Plugins/Experimental/PCG/Source/PCG/Private/Utils/PCGGraphExecutionLogging.cpp
+++ b/Engine/Plugins/Experimental/PCG/Source/PCG/Private/Utils/PCGGraphExecutionLogging.cpp
@@ -5,19 +5,25 @@
 #include "PCGComponent.h"
 #include "PCGContext.h"
 #include "PCGModule.h"
+#include "Graph/PCGGraphExecutor.h"
 
 #include "GameFramework/Actor.h"
 #include "HAL/IConsoleManager.h"
 
+namespace PCGGraphExecutionLogging
+{
 #if WITH_EDITOR
-static TAutoConsoleVariable<bool> CVarGraphExecutionLoggingEnable(
-	TEXT("pcg.GraphExecutionLogging"),
-	false,
-	TEXT("Enables fine grained log of graph execution"));
+	static TAutoConsoleVariable<bool> CVarGraphExecutionLoggingEnable(
+		TEXT("pcg.GraphExecution.EnableLogging"),
+		false,
+		TEXT("Enables fine grained log of graph execution"));
+
+	static TAutoConsoleVariable<bool> CVarGraphExecutionCullingLoggingEnable(
+		TEXT("pcg.GraphExecution.EnableCullingLogging"),
+		false,
+		TEXT("Enables fine grained log of dynamic task culling during graph execution"));
 #endif
 
-namespace PCGGraphExecutionLogging
-{
 	bool LogEnabled()
 	{
 #if WITH_EDITOR
@@ -27,6 +33,236 @@ namespace PCGGraphExecutionLogging
 #endif
 	}
 
+	bool CullingLogEnabled()
+	{
+#if WITH_EDITOR
+		return CVarGraphExecutionCullingLoggingEnable.GetValueOnAnyThread();
+#else
+		return false;
+#endif
+	}
+
+	void LogGraphTask(FPCGTaskId TaskId, const FPCGGraphTask& Task, const TSet<FPCGTaskId>* SuccessorIds)
+	{
+#if WITH_EDITOR
+		if (!LogEnabled())
+		{
+			return;
+		}
+
+		auto GenerateInputsString = [](const TArray<FPCGGraphTaskInput>& Inputs)
+		{
+			FString InputString;
+			bool bFirstInput = true;
+
+			for (const FPCGGraphTaskInput& Input : Inputs)
+			{
+				if (!bFirstInput)
+				{
+					InputString += TEXT(",");
+				}
+				bFirstInput = false;
+
+				InputString += FString::Printf(TEXT("%u->'%s'"), Input.TaskId, Input.OutPin ? *Input.OutPin->Properties.Label.ToString() : TEXT(""));
+			}
+
+			return InputString;
+		};
+
+		FString SuccessorsString;
+		if (SuccessorIds)
+		{
+			bool bFirstSuccessor = true;
+			for (const FPCGTaskId& SuccessorId : *SuccessorIds)
+			{
+				SuccessorsString += bFirstSuccessor ? FString::Printf(TEXT("%u"), SuccessorId) : FString::Printf(TEXT(",%u"), SuccessorId);
+				bFirstSuccessor = false;
+			}
+		}
+
+		UE_LOG(LogPCG, Log, TEXT("\t\tID: %u\tParent: %u\tNode: %s\tInputs: %s\tPinDeps: %s\tSuccessors: %s"),
+			TaskId,
+			Task.ParentId != InvalidPCGTaskId ? Task.ParentId : 0,
+			Task.Node ? (*Task.Node->GetNodeTitle(EPCGNodeTitleType::ListView).ToString()) : TEXT("NULL"),
+			*GenerateInputsString(Task.Inputs),
+			*Task.PinDependency.ToString(),
+			*SuccessorsString
+		);
+#endif
+	}
+
+	void LogGraphTasks(const TMap<FPCGTaskId, FPCGGraphTask>& Tasks, const TMap<FPCGTaskId, TSet<FPCGTaskId>>* TaskSuccessors)
+	{
+#if WITH_EDITOR
+		if (!LogEnabled())
+		{
+			return;
+		}
+
+		for (const TPair<FPCGTaskId, FPCGGraphTask>& TaskIdAndTask : Tasks)
+		{
+			PCGGraphExecutionLogging::LogGraphTask(TaskIdAndTask.Key, TaskIdAndTask.Value, TaskSuccessors ? TaskSuccessors->Find(TaskIdAndTask.Value.NodeId) : nullptr);
+		}
+#endif
+	}
+
+	void LogGraphTasks(const TArray<FPCGGraphTask>& Tasks)
+	{
+#if WITH_EDITOR
+		if (!LogEnabled())
+		{
+			return;
+		}
+
+		for (const FPCGGraphTask& Task : Tasks)
+		{
+			LogGraphTask(Task.NodeId, Task);
+		}
+#endif
+	}
+
+	void LogGraphSchedule(const UPCGComponent* SourceComponent)
+	{
+#if WITH_EDITOR
+		if (!LogEnabled())
+		{
+			return;
+		}
+
+		UE_LOG(LogPCG, Log, TEXT("[%s] --- SCHEDULE GRAPH ---"),
+			(SourceComponent && SourceComponent->GetOwner()) ? *SourceComponent->GetOwner()->GetName() : TEXT("MISSINGCOMPONENT"));
+#endif
+	}
+
+	void LogGraphPostSchedule(const TMap<FPCGTaskId, FPCGGraphTask>& Tasks, const TMap<FPCGTaskId, TSet<FPCGTaskId>>& TaskSuccessors)
+	{
+#if WITH_EDITOR
+		if (!CullingLogEnabled())
+		{
+			return;
+		}
+
+		UE_LOG(LogPCG, Log, TEXT("POST SCHEDULE:"));
+
+		LogGraphTasks(Tasks, &TaskSuccessors);
+#endif
+	}
+
+	void LogGraphExecuteFrameFinished()
+	{
+#if WITH_EDITOR
+		if (!LogEnabled())
+		{
+			return;
+		}
+
+		UE_LOG(LogPCG, Log, TEXT("--- FINISH FPCGGRAPHEXECUTOR::EXECUTE ---"));
+#endif
+	}
+
+#if WITH_EDITOR
+	FString GetPinsToDeactivateString(const TArray<FPCGPinId>& PinIdsToDeactivate)
+	{
+		FString PinIdsToDeactivateString;
+		bool bFirst = true;
+
+		for (const FPCGPinId& PinId : PinIdsToDeactivate)
+		{
+			const FPCGTaskId NodeId = PCGPinIdHelpers::GetNodeIdFromPinId(PinId);
+			const uint64 PinIndex = PCGPinIdHelpers::GetPinIndexFromPinId(PinId);
+			PinIdsToDeactivateString += bFirst ? FString::Printf(TEXT("%u_%u"), NodeId, PinIndex) : FString::Printf(TEXT(",%u_%u"), NodeId, PinIndex);
+			bFirst = false;
+		}
+
+		return PinIdsToDeactivateString;
+	}
+#endif
+
+	void LogTaskExecute(const FPCGGraphTask& Task)
+	{
+#if WITH_EDITOR
+		if (!LogEnabled() || !Task.SourceComponent.Get())
+		{
+			return;
+		}
+
+		UE_LOG(LogPCG, Log, TEXT("         [%s] %s\t\tEXECUTE"),
+			*Task.SourceComponent->GetOwner()->GetName(),
+			*FString::Printf(TEXT("%u'%s'"), Task.NodeId, Task.Node ? *Task.Node->GetNodeTitle(EPCGNodeTitleType::ListView).ToString() : TEXT("")));
+#endif
+	}
+
+	void LogTaskExecuteCachingDisabled(const FPCGGraphTask& Task)
+	{
+#if WITH_EDITOR
+		if (!LogEnabled() || !Task.SourceComponent.Get())
+		{
+			return;
+		}
+
+		UE_LOG(LogPCG, Warning, TEXT("[%s] %s\t\tCACHING DISABLED"),
+			*Task.SourceComponent->GetOwner()->GetName(),
+			*FString::Printf(TEXT("%u'%s'"), Task.NodeId, Task.Node ? *Task.Node->GetNodeTitle(EPCGNodeTitleType::ListView).ToString() : TEXT("")));
+#endif
+	}
+
+	void LogTaskExecuteOutputCRC(const FPCGGraphActiveTask& Task)
+	{
+#if WITH_EDITOR
+		if (!LogEnabled() || !Task.Context || !Task.Context->SourceComponent.Get() || !Task.Context->SourceComponent->GetOwner())
+		{
+			return;
+		}
+
+		UE_LOG(LogPCG, Log, TEXT("         [%s] %s\t\tOUTPUT CRC %u"),
+			*Task.Context->SourceComponent->GetOwner()->GetName(),
+			*FString::Printf(TEXT("%u'%s'"), Task.NodeId, Task.Context->Node ? *Task.Context->Node->GetNodeTitle(EPCGNodeTitleType::ListView).ToString() : TEXT("")),
+			Task.Context->OutputData.Crc.GetValue());
+#endif
+	}
+
+	void LogTaskCullingBegin(FPCGTaskId CompletedTaskId, uint64 InactiveOutputPinBitmask, const TArray<FPCGPinId>& PinIdsToDeactivate)
+	{
+#if WITH_EDITOR
+		if (!CullingLogEnabled())
+		{
+			return;
+		}
+
+		UE_LOG(LogPCG, Log, TEXT("BEGIN CullInactiveDownstreamNodes, CompletedTaskId: %u, InactiveOutputPinBitmask: %u, Deactivating pin IDs: %s"),
+			CompletedTaskId, InactiveOutputPinBitmask, *GetPinsToDeactivateString(PinIdsToDeactivate));
+#endif
+	}
+
+	void LogTaskCullingBeginLoop(FPCGTaskId PinTaskId, uint64 PinIndex, const TArray<FPCGPinId>& PinIdsToDeactivate)
+	{
+#if WITH_EDITOR
+		if (!CullingLogEnabled())
+		{
+			return;
+		}
+
+		UE_LOG(LogPCG, Log, TEXT("LOOP: DEACTIVATE %u_%u, remaining IDs: %s"), PinTaskId, PinIndex, *GetPinsToDeactivateString(PinIdsToDeactivate));
+#endif
+	}
+
+	void LogTaskCullingUpdatedPinDeps(FPCGTaskId TaskId, const FPCGPinDependencyExpression& PinDependency, bool bDependencyExpressionBecameFalse)
+	{
+#if WITH_EDITOR
+		if (!CullingLogEnabled())
+		{
+			return;
+		}
+
+		UE_LOG(LogPCG, Log, TEXT("UPDATED PIN DEP EXPRESSION (task ID %u): %s"), TaskId, *PinDependency.ToString());
+
+		if (bDependencyExpressionBecameFalse)
+		{
+			UE_LOG(LogPCG, Log, TEXT("CULL task ID %u"), TaskId);
+		}
+#endif
+	}
+
 	void LogGridLinkageTaskExecuteStore(const FPCGContext* InContext, EPCGHiGenGrid InGenerationGrid, int32 InFromGridSize, int32 InToGridSize, const FString& InResourcePath)
 	{
 #if WITH_EDITOR
diff --git a/Engine/Plugins/Experimental/PCG/Source/PCG/Private/Utils/PCGGraphExecutionLogging.h b/Engine/Plugins/Experimental/PCG/Source/PCG/Private/Utils/PCGGraphExecutionLogging.h
--- a/Engine/Plugins/Experimental/PCG/Source/PCG/Private/Utils/PCGGraphExecutionLogging.h
+++ b/Engine/Plugins/Experimental/PCG/Source/PCG/Private/Utils/PCGGraphExecutionLogging.h
@@ -2,15 +2,39 @@
 
 #pragma once
 
+#include "PCGCommon.h"
+
 #include "Containers/UnrealString.h"
 #include "HAL/Platform.h"
 
+class UPCGComponent;
 enum class EPCGHiGenGrid : uint32;
 struct FPCGContext;
-class UPCGComponent;
+struct FPCGGraphActiveTask;
+struct FPCGGraphTask;
+struct FPCGPinDependencyExpression;
 
 namespace PCGGraphExecutionLogging
 {
+	bool LogEnabled();
+	bool CullingLogEnabled();
+
+	void LogGraphTask(FPCGTaskId TaskId, const FPCGGraphTask& Task, const TSet<FPCGTaskId>* SuccessorIds = nullptr);
+	void LogGraphTasks(const TMap<FPCGTaskId, FPCGGraphTask>& Tasks, const TMap<FPCGTaskId, TSet<FPCGTaskId>>* TaskSuccessors = nullptr);
+	void LogGraphTasks(const TArray<FPCGGraphTask>& Tasks);
+
+	void LogGraphSchedule(const UPCGComponent* SourceComponent);
+	void LogGraphPostSchedule(const TMap<FPCGTaskId, FPCGGraphTask>& Tasks, const TMap<FPCGTaskId, TSet<FPCGTaskId>>& TaskSuccessors);
+	void LogGraphExecuteFrameFinished();
+
+	void LogTaskExecute(const FPCGGraphTask& Task);
+	void LogTaskExecuteCachingDisabled(const FPCGGraphTask& Task);
+	void LogTaskExecuteOutputCRC(const FPCGGraphActiveTask& Task);
+
+	void LogTaskCullingBegin(FPCGTaskId CompletedTaskId, uint64 InactiveOutputPinBitmask, const TArray<FPCGPinId>& PinIdsToDeactivate);
+	void LogTaskCullingBeginLoop(FPCGTaskId PinTaskId, uint64 PinIndex, const TArray<FPCGPinId>& PinIdsToDeactivate);
+	void LogTaskCullingUpdatedPinDeps(FPCGTaskId TaskId, const FPCGPinDependencyExpression& PinDependency, bool bDependencyExpressionBecameFalse);
+
 	void LogGridLinkageTaskExecuteStore(const FPCGContext* InContext, EPCGHiGenGrid InGenerationGrid, int32 InFromLOD, int32 InToLOD, const FString& InResourcePath);
 	void LogGridLinkageTaskExecuteRetrieve(const FPCGContext* InContext, EPCGHiGenGrid InGenerationGrid, int32 InFromLOD, int32 InToLOD, const FString& InResourcePath);
 	void LogGridLinkageTaskExecuteRetrieveSuccess(const FPCGContext* InContext, const UPCGComponent* InComponent, const FString& InResourcePath, int32 InDataItemCount);
diff --git a/Engine/Plugins/Experimental/PCG/Source/PCG/Public/Elements/ControlFlow/PCGBranch.h b/Engine/Plugins/Experimental/PCG/Source/PCG/Public/Elements/ControlFlow/PCGBranch.h
--- a/Engine/Plugins/Experimental/PCG/Source/PCG/Public/Elements/ControlFlow/PCGBranch.h
+++ b/Engine/Plugins/Experimental/PCG/Source/PCG/Public/Elements/ControlFlow/PCGBranch.h
@@ -20,10 +20,12 @@ public:
 	virtual FText GetDefaultNodeTitle() const override;
 	virtual FText GetNodeTooltipText() const override;
 	virtual EPCGSettingsType GetType() const override { return EPCGSettingsType::ControlFlow; }
-	virtual bool HasDynamicPins() const override { return true; }
 #endif
-
+	virtual bool HasDynamicPins() const override { return true; }
+	virtual bool OutputPinsCanBeDeactivated() const override { return true; }
 	virtual bool IsPinStaticallyActive(const FName& PinLabel) const override;
+	virtual bool IsInputPinRequiredByExecution(const UPCGPin* InPin) const override { return !InPin->Properties.bAdvancedPin; }
+	virtual bool IsPinUsedByNodeExecution(const UPCGPin* InPin) const override;
 
 protected:
 #if WITH_EDITOR
diff --git a/Engine/Plugins/Experimental/PCG/Source/PCG/Public/Elements/ControlFlow/PCGSwitch.h b/Engine/Plugins/Experimental/PCG/Source/PCG/Public/Elements/ControlFlow/PCGSwitch.h
--- a/Engine/Plugins/Experimental/PCG/Source/PCG/Public/Elements/ControlFlow/PCGSwitch.h
+++ b/Engine/Plugins/Experimental/PCG/Source/PCG/Public/Elements/ControlFlow/PCGSwitch.h
@@ -28,10 +28,12 @@ public:
 	virtual FText GetDefaultNodeTitle() const override;
 	virtual FText GetNodeTooltipText() const override;
 	virtual EPCGSettingsType GetType() const override { return EPCGSettingsType::ControlFlow; }
-	virtual bool HasDynamicPins() const override { return true; }
 #endif // WITH_EDITOR
-
+	virtual bool HasDynamicPins() const override { return true; }
+	virtual bool OutputPinsCanBeDeactivated() const override { return true; }
 	virtual bool IsPinStaticallyActive(const FName& PinLabel) const override;
+	virtual bool IsPinUsedByNodeExecution(const UPCGPin* InPin) const override;
+	virtual bool IsInputPinRequiredByExecution(const UPCGPin* InPin) const override { return !InPin->Properties.bAdvancedPin; }
 	virtual FString GetAdditionalTitleInformation() const override;
 	virtual bool HasFlippedTitleLines() const override { return true; }
 
@@ -84,6 +86,9 @@ public:
 	/** Returns true if the enum value exists within the selected enum class. */
 	bool IsValuePresent(int64 Value) const;
 
+	/** Returns the index of the currently selected pin. */
+	int GetSelectedOutputPinIndex() const;
+
 	/** Helper function to use the appropriate selection value to determine the current selection. Returns true if it succeeds or false if the value is invalid. */
 	bool GetSelectedPinLabel(FName& OutSelectedPinLabel) const;
 
diff --git a/Engine/Plugins/Experimental/PCG/Source/PCG/Public/Elements/Metadata/PCGMetadataBreakTransform.h b/Engine/Plugins/Experimental/PCG/Source/PCG/Public/Elements/Metadata/PCGMetadataBreakTransform.h
--- a/Engine/Plugins/Experimental/PCG/Source/PCG/Public/Elements/Metadata/PCGMetadataBreakTransform.h
+++ b/Engine/Plugins/Experimental/PCG/Source/PCG/Public/Elements/Metadata/PCGMetadataBreakTransform.h
@@ -22,6 +22,7 @@ public:
 	virtual FName GetDefaultNodeName() const override;
 	virtual FText GetDefaultNodeTitle() const override;
 #endif
+	virtual bool IsInputPinRequiredByExecution(const UPCGPin* InPin) const override { return !InPin->Properties.bAdvancedPin; }
 
 	virtual FPCGAttributePropertyInputSelector GetInputSource(uint32 Index) const override;
 
diff --git a/Engine/Plugins/Experimental/PCG/Source/PCG/Public/Elements/Metadata/PCGMetadataBreakVector.h b/Engine/Plugins/Experimental/PCG/Source/PCG/Public/Elements/Metadata/PCGMetadataBreakVector.h
--- a/Engine/Plugins/Experimental/PCG/Source/PCG/Public/Elements/Metadata/PCGMetadataBreakVector.h
+++ b/Engine/Plugins/Experimental/PCG/Source/PCG/Public/Elements/Metadata/PCGMetadataBreakVector.h
@@ -30,6 +30,7 @@ public:
 	virtual FName GetDefaultNodeName() const override;
 	virtual FText GetDefaultNodeTitle() const override;
 #endif
+	virtual bool IsInputPinRequiredByExecution(const UPCGPin* InPin) const override { return !InPin->Properties.bAdvancedPin; }
 
 	virtual FPCGAttributePropertyInputSelector GetInputSource(uint32 Index) const override;
 
diff --git a/Engine/Plugins/Experimental/PCG/Source/PCG/Public/Elements/Metadata/PCGMetadataElement.h b/Engine/Plugins/Experimental/PCG/Source/PCG/Public/Elements/Metadata/PCGMetadataElement.h
--- a/Engine/Plugins/Experimental/PCG/Source/PCG/Public/Elements/Metadata/PCGMetadataElement.h
+++ b/Engine/Plugins/Experimental/PCG/Source/PCG/Public/Elements/Metadata/PCGMetadataElement.h
@@ -33,6 +33,7 @@ public:
 	virtual FText GetDefaultNodeTitle() const override { return NSLOCTEXT("PCGMetadataOperationSettings", "NodeTitle", "Copy Attribute"); }
 	virtual FText GetNodeTooltipText() const;
 	virtual EPCGSettingsType GetType() const override { return EPCGSettingsType::Metadata; }
+	virtual bool IsInputPinRequiredByExecution(const UPCGPin* InPin) const override { return InPin->Properties.Label == PCGPinConstants::DefaultInputLabel; }
 	virtual bool IsPinUsedByNodeExecution(const UPCGPin* InPin) const override;
 	virtual void ApplyDeprecation(UPCGNode* InOutNode) override;
 #endif
diff --git a/Engine/Plugins/Experimental/PCG/Source/PCG/Public/Elements/Metadata/PCGMetadataPartition.h b/Engine/Plugins/Experimental/PCG/Source/PCG/Public/Elements/Metadata/PCGMetadataPartition.h
--- a/Engine/Plugins/Experimental/PCG/Source/PCG/Public/Elements/Metadata/PCGMetadataPartition.h
+++ b/Engine/Plugins/Experimental/PCG/Source/PCG/Public/Elements/Metadata/PCGMetadataPartition.h
@@ -24,6 +24,7 @@ public:
 	virtual EPCGSettingsType GetType() const override { return EPCGSettingsType::Metadata; }
 	virtual bool HasDynamicPins() const override { return true; }
 #endif
+	virtual bool IsInputPinRequiredByExecution(const UPCGPin* InPin) const override { return !InPin->Properties.bAdvancedPin; }
 	virtual FString GetAdditionalTitleInformation() const override;
 
 protected:
diff --git a/Engine/Plugins/Experimental/PCG/Source/PCG/Public/Elements/Metadata/PCGMetadataRenameElement.h b/Engine/Plugins/Experimental/PCG/Source/PCG/Public/Elements/Metadata/PCGMetadataRenameElement.h
--- a/Engine/Plugins/Experimental/PCG/Source/PCG/Public/Elements/Metadata/PCGMetadataRenameElement.h
+++ b/Engine/Plugins/Experimental/PCG/Source/PCG/Public/Elements/Metadata/PCGMetadataRenameElement.h
@@ -19,6 +19,7 @@ public:
 	virtual EPCGSettingsType GetType() const override { return EPCGSettingsType::Metadata; }
 	virtual bool HasDynamicPins() const override { return true; }
 #endif
+	virtual bool IsInputPinRequiredByExecution(const UPCGPin* InPin) const override { return !InPin->Properties.bAdvancedPin; }
 	virtual FString GetAdditionalTitleInformation() const override;
 
 protected:
diff --git a/Engine/Plugins/Experimental/PCG/Source/PCG/Public/Elements/PCGAddTag.h b/Engine/Plugins/Experimental/PCG/Source/PCG/Public/Elements/PCGAddTag.h
--- a/Engine/Plugins/Experimental/PCG/Source/PCG/Public/Elements/PCGAddTag.h
+++ b/Engine/Plugins/Experimental/PCG/Source/PCG/Public/Elements/PCGAddTag.h
@@ -23,6 +23,7 @@ public:
 	virtual EPCGSettingsType GetType() const override { return EPCGSettingsType::Generic; }
 	virtual bool HasDynamicPins() const override { return true; }
 #endif
+	virtual bool IsInputPinRequiredByExecution(const UPCGPin* InPin) const override { return !InPin->Properties.bAdvancedPin; }
 
 protected:
 	virtual TArray<FPCGPinProperties> InputPinProperties() const override;
diff --git a/Engine/Plugins/Experimental/PCG/Source/PCG/Public/Elements/PCGApplyScaleToBounds.h b/Engine/Plugins/Experimental/PCG/Source/PCG/Public/Elements/PCGApplyScaleToBounds.h
--- a/Engine/Plugins/Experimental/PCG/Source/PCG/Public/Elements/PCGApplyScaleToBounds.h
+++ b/Engine/Plugins/Experimental/PCG/Source/PCG/Public/Elements/PCGApplyScaleToBounds.h
@@ -22,6 +22,7 @@ public:
 	virtual FText GetNodeTooltipText() const override { return NSLOCTEXT("PCGApplyScaleToBoundsElement", "NodeTooltip", "Applies the scale of each point to its bounds and resets the scale."); }
 	virtual EPCGSettingsType GetType() const override { return EPCGSettingsType::PointOps; }
 #endif
+	virtual bool IsInputPinRequiredByExecution(const UPCGPin* InPin) const override { return !InPin->Properties.bAdvancedPin; }
 
 protected:
 	virtual TArray<FPCGPinProperties> InputPinProperties() const override { return Super::DefaultPointInputPinProperties(); }
diff --git a/Engine/Plugins/Experimental/PCG/Source/PCG/Public/Elements/PCGAttributeFilter.h b/Engine/Plugins/Experimental/PCG/Source/PCG/Public/Elements/PCGAttributeFilter.h
--- a/Engine/Plugins/Experimental/PCG/Source/PCG/Public/Elements/PCGAttributeFilter.h
+++ b/Engine/Plugins/Experimental/PCG/Source/PCG/Public/Elements/PCGAttributeFilter.h
@@ -81,6 +81,7 @@ public:
 	virtual EPCGSettingsType GetType() const override { return EPCGSettingsType::Filter; }
 	virtual bool HasDynamicPins() const override { return true; }
 #endif
+	virtual bool IsInputPinRequiredByExecution(const UPCGPin* InPin) const override { return InPin->Properties.Label == PCGPinConstants::DefaultInputLabel; }
 	virtual FString GetAdditionalTitleInformation() const override;
 	virtual EPCGDataType GetCurrentPinTypes(const UPCGPin* InPin) const override;
 
@@ -148,6 +149,7 @@ public:
 	virtual EPCGSettingsType GetType() const override { return EPCGSettingsType::Filter; }
 	virtual bool HasDynamicPins() const override { return true; }
 #endif
+	virtual bool IsInputPinRequiredByExecution(const UPCGPin* InPin) const override { return InPin->Properties.Label == PCGPinConstants::DefaultInputLabel; }
 	virtual FString GetAdditionalTitleInformation() const override;
 	virtual EPCGDataType GetCurrentPinTypes(const UPCGPin* InPin) const override;
 
diff --git a/Engine/Plugins/Experimental/PCG/Source/PCG/Public/Elements/PCGAttributeGetFromIndexElement.h b/Engine/Plugins/Experimental/PCG/Source/PCG/Public/Elements/PCGAttributeGetFromIndexElement.h
--- a/Engine/Plugins/Experimental/PCG/Source/PCG/Public/Elements/PCGAttributeGetFromIndexElement.h
+++ b/Engine/Plugins/Experimental/PCG/Source/PCG/Public/Elements/PCGAttributeGetFromIndexElement.h
@@ -19,6 +19,7 @@ public:
 	virtual FText GetNodeTooltipText() const override;
 	virtual EPCGSettingsType GetType() const override { return EPCGSettingsType::Param; }
 #endif
+	virtual bool IsInputPinRequiredByExecution(const UPCGPin* InPin) const override { return !InPin->Properties.bAdvancedPin; }
 
 protected:
 	virtual TArray<FPCGPinProperties> InputPinProperties() const override;
diff --git a/Engine/Plugins/Experimental/PCG/Source/PCG/Public/Elements/PCGAttributeGetFromPointIndexElement.h b/Engine/Plugins/Experimental/PCG/Source/PCG/Public/Elements/PCGAttributeGetFromPointIndexElement.h
--- a/Engine/Plugins/Experimental/PCG/Source/PCG/Public/Elements/PCGAttributeGetFromPointIndexElement.h
+++ b/Engine/Plugins/Experimental/PCG/Source/PCG/Public/Elements/PCGAttributeGetFromPointIndexElement.h
@@ -38,11 +38,15 @@ public:
 	virtual EPCGSettingsType GetType() const override { return EPCGSettingsType::Metadata; }
 	virtual void ApplyDeprecation(UPCGNode* InOutNode) override;
 #endif
+	virtual bool IsInputPinRequiredByExecution(const UPCGPin* InPin) const override { return !InPin->Properties.bAdvancedPin; }
 
+protected:
 	virtual TArray<FPCGPinProperties> InputPinProperties() const override;
 	virtual TArray<FPCGPinProperties> OutputPinProperties() const override;
+	virtual FPCGElementPtr CreateElement() const override;
 	//~End UPCGSettings interface
 
+public:
 	UPROPERTY(BlueprintReadWrite, EditAnywhere, Category = Settings, meta = (PCG_Overridable))
 	FPCGAttributePropertyInputSelector InputSource;
 
@@ -56,9 +60,6 @@ public:
 	UPROPERTY()
 	FName InputAttributeName_DEPRECATED = NAME_None;
 #endif // WITH_EDITORONLY_DATA
-
-protected:
-	virtual FPCGElementPtr CreateElement() const override;
 };
 
 
diff --git a/Engine/Plugins/Experimental/PCG/Source/PCG/Public/Elements/PCGAttributeNoise.h b/Engine/Plugins/Experimental/PCG/Source/PCG/Public/Elements/PCGAttributeNoise.h
--- a/Engine/Plugins/Experimental/PCG/Source/PCG/Public/Elements/PCGAttributeNoise.h
+++ b/Engine/Plugins/Experimental/PCG/Source/PCG/Public/Elements/PCGAttributeNoise.h
@@ -47,6 +47,7 @@ public:
 	virtual TArray<FText> GetNodeTitleAliases() const { return { NSLOCTEXT("PCGAttributeNoiseSettings", "DensityNoiseAlias", "Density Noise") }; }
 	virtual EPCGSettingsType GetType() const override { return EPCGSettingsType::Metadata; }
 #endif
+	virtual bool IsInputPinRequiredByExecution(const UPCGPin* InPin) const override { return !InPin->Properties.bAdvancedPin; }
 
 protected:
 	virtual TArray<FPCGPinProperties> InputPinProperties() const override { return Super::DefaultPointInputPinProperties(); }
diff --git a/Engine/Plugins/Experimental/PCG/Source/PCG/Public/Elements/PCGAttributeSelectElement.h b/Engine/Plugins/Experimental/PCG/Source/PCG/Public/Elements/PCGAttributeSelectElement.h
--- a/Engine/Plugins/Experimental/PCG/Source/PCG/Public/Elements/PCGAttributeSelectElement.h
+++ b/Engine/Plugins/Experimental/PCG/Source/PCG/Public/Elements/PCGAttributeSelectElement.h
@@ -62,6 +62,7 @@ public:
 	virtual EPCGSettingsType GetType() const override { return EPCGSettingsType::Metadata; }
 	virtual void ApplyDeprecation(UPCGNode* InOutNode) override;
 #endif
+	virtual bool IsInputPinRequiredByExecution(const UPCGPin* InPin) const override { return !InPin->Properties.bAdvancedPin; }
 	virtual FString GetAdditionalTitleInformation() const override;
 
 protected:
diff --git a/Engine/Plugins/Experimental/PCG/Source/PCG/Public/Elements/PCGAttributeTransferElement.h b/Engine/Plugins/Experimental/PCG/Source/PCG/Public/Elements/PCGAttributeTransferElement.h
--- a/Engine/Plugins/Experimental/PCG/Source/PCG/Public/Elements/PCGAttributeTransferElement.h
+++ b/Engine/Plugins/Experimental/PCG/Source/PCG/Public/Elements/PCGAttributeTransferElement.h
@@ -44,6 +44,7 @@ public:
 	virtual bool HasDynamicPins() const override { return true; }
 	virtual void ApplyDeprecation(UPCGNode* InOutNode) override;
 #endif
+	virtual bool IsInputPinRequiredByExecution(const UPCGPin* InPin) const override { return InPin->Properties.Label == PCGAttributeTransferConstants::TargetLabel; }
 	virtual EPCGDataType GetCurrentPinTypes(const UPCGPin* InPin) const override;
 	virtual FString GetAdditionalTitleInformation() const override;
 
diff --git a/Engine/Plugins/Experimental/PCG/Source/PCG/Public/Elements/PCGBoundsModifier.h b/Engine/Plugins/Experimental/PCG/Source/PCG/Public/Elements/PCGBoundsModifier.h
--- a/Engine/Plugins/Experimental/PCG/Source/PCG/Public/Elements/PCGBoundsModifier.h
+++ b/Engine/Plugins/Experimental/PCG/Source/PCG/Public/Elements/PCGBoundsModifier.h
@@ -34,6 +34,7 @@ public:
 	virtual FText GetNodeTooltipText() const override;
 	virtual EPCGSettingsType GetType() const override { return EPCGSettingsType::PointOps; }
 #endif
+	virtual bool IsInputPinRequiredByExecution(const UPCGPin* InPin) const override { return !InPin->Properties.bAdvancedPin; }
 
 protected:
 	virtual TArray<FPCGPinProperties> InputPinProperties() const override { return Super::DefaultPointInputPinProperties(); }
diff --git a/Engine/Plugins/Experimental/PCG/Source/PCG/Public/Elements/PCGCollapseElement.h b/Engine/Plugins/Experimental/PCG/Source/PCG/Public/Elements/PCGCollapseElement.h
--- a/Engine/Plugins/Experimental/PCG/Source/PCG/Public/Elements/PCGCollapseElement.h
+++ b/Engine/Plugins/Experimental/PCG/Source/PCG/Public/Elements/PCGCollapseElement.h
@@ -21,6 +21,7 @@ public:
 	virtual bool ShouldDrawNodeCompact() const override { return true; }
 	virtual bool GetCompactNodeIcon(FName& OutCompactNodeIcon) const override;
 #endif
+	virtual bool IsInputPinRequiredByExecution(const UPCGPin* InPin) const override { return !InPin->Properties.bAdvancedPin; }
 
 protected:
 	virtual TArray<FPCGPinProperties> InputPinProperties() const override;
diff --git a/Engine/Plugins/Experimental/PCG/Source/PCG/Public/Elements/PCGCombinePoints.h b/Engine/Plugins/Experimental/PCG/Source/PCG/Public/Elements/PCGCombinePoints.h
--- a/Engine/Plugins/Experimental/PCG/Source/PCG/Public/Elements/PCGCombinePoints.h
+++ b/Engine/Plugins/Experimental/PCG/Source/PCG/Public/Elements/PCGCombinePoints.h
@@ -23,6 +23,7 @@ public:
 	virtual FText GetNodeTooltipText() const override { return NSLOCTEXT("PCGCombinePointsElement", "NodeTooltip", "Combines each point to share a singular bound extent."); }
 	virtual EPCGSettingsType GetType() const override { return EPCGSettingsType::PointOps; }
 #endif
+	virtual bool IsInputPinRequiredByExecution(const UPCGPin* InPin) const override { return !InPin->Properties.bAdvancedPin; }
 
 protected:
 	virtual TArray<FPCGPinProperties> InputPinProperties() const override { return Super::DefaultPointInputPinProperties(); }
diff --git a/Engine/Plugins/Experimental/PCG/Source/PCG/Public/Elements/PCGConvertToAttributeSet.h b/Engine/Plugins/Experimental/PCG/Source/PCG/Public/Elements/PCGConvertToAttributeSet.h
--- a/Engine/Plugins/Experimental/PCG/Source/PCG/Public/Elements/PCGConvertToAttributeSet.h
+++ b/Engine/Plugins/Experimental/PCG/Source/PCG/Public/Elements/PCGConvertToAttributeSet.h
@@ -22,6 +22,7 @@ public:
 	virtual bool ShouldDrawNodeCompact() const override { return true; }
 	virtual bool GetCompactNodeIcon(FName& OutCompactNodeIcon) const override;
 #endif
+	virtual bool IsInputPinRequiredByExecution(const UPCGPin* InPin) const override { return !InPin->Properties.bAdvancedPin; }
 
 protected:
 	virtual TArray<FPCGPinProperties> InputPinProperties() const override;
diff --git a/Engine/Plugins/Experimental/PCG/Source/PCG/Public/Elements/PCGConvexHull2D.h b/Engine/Plugins/Experimental/PCG/Source/PCG/Public/Elements/PCGConvexHull2D.h
--- a/Engine/Plugins/Experimental/PCG/Source/PCG/Public/Elements/PCGConvexHull2D.h
+++ b/Engine/Plugins/Experimental/PCG/Source/PCG/Public/Elements/PCGConvexHull2D.h
@@ -22,6 +22,7 @@ public:
 	virtual FText GetNodeTooltipText() const override { return NSLOCTEXT("PCGConvexHullElement", "NodeTooltip", "Return the 2D convex hull of a set of points on the XY plane."); }
 	virtual EPCGSettingsType GetType() const override { return EPCGSettingsType::Spatial; }
 #endif
+	virtual bool IsInputPinRequiredByExecution(const UPCGPin* InPin) const override { return !InPin->Properties.bAdvancedPin; }
 
 protected:
 	virtual TArray<FPCGPinProperties> InputPinProperties() const override { return Super::DefaultPointInputPinProperties(); }
diff --git a/Engine/Plugins/Experimental/PCG/Source/PCG/Public/Elements/PCGCopyPoints.h b/Engine/Plugins/Experimental/PCG/Source/PCG/Public/Elements/PCGCopyPoints.h
--- a/Engine/Plugins/Experimental/PCG/Source/PCG/Public/Elements/PCGCopyPoints.h
+++ b/Engine/Plugins/Experimental/PCG/Source/PCG/Public/Elements/PCGCopyPoints.h
@@ -51,6 +51,7 @@ public:
 	virtual FText GetNodeTooltipText() const override;
 	virtual EPCGSettingsType GetType() const override { return EPCGSettingsType::Sampler; }
 #endif
+	virtual bool IsInputPinRequiredByExecution(const UPCGPin* InPin) const override { return !InPin->Properties.bAdvancedPin; }
 
 protected:
 	virtual TArray<FPCGPinProperties> InputPinProperties() const override;
diff --git a/Engine/Plugins/Experimental/PCG/Source/PCG/Public/Elements/PCGCreateAttribute.h b/Engine/Plugins/Experimental/PCG/Source/PCG/Public/Elements/PCGCreateAttribute.h
--- a/Engine/Plugins/Experimental/PCG/Source/PCG/Public/Elements/PCGCreateAttribute.h
+++ b/Engine/Plugins/Experimental/PCG/Source/PCG/Public/Elements/PCGCreateAttribute.h
@@ -39,7 +39,7 @@ public:
 	virtual bool IsPinUsedByNodeExecution(const UPCGPin* InPin) const override;
 	virtual bool CanEditChange(const FProperty* InProperty) const override;
 #endif // WITH_EDITOR
-
+	virtual bool IsInputPinRequiredByExecution(const UPCGPin* InPin) const override { return InPin->Properties.Label == PCGPinConstants::DefaultInputLabel; }
 	virtual EPCGDataType GetCurrentPinTypes(const UPCGPin* InPin) const override;
 	virtual FString GetAdditionalTitleInformation() const override;
 
diff --git a/Engine/Plugins/Experimental/PCG/Source/PCG/Public/Elements/PCGCreateSpline.h b/Engine/Plugins/Experimental/PCG/Source/PCG/Public/Elements/PCGCreateSpline.h
--- a/Engine/Plugins/Experimental/PCG/Source/PCG/Public/Elements/PCGCreateSpline.h
+++ b/Engine/Plugins/Experimental/PCG/Source/PCG/Public/Elements/PCGCreateSpline.h
@@ -32,6 +32,7 @@ public:
 	virtual FText GetNodeTooltipText() const override;
 	virtual EPCGSettingsType GetType() const override { return EPCGSettingsType::Spatial; }
 #endif
+	virtual bool IsInputPinRequiredByExecution(const UPCGPin* InPin) const override { return !InPin->Properties.bAdvancedPin; }
 
 	virtual TArray<FPCGPinProperties> InputPinProperties() const override { return Super::DefaultPointInputPinProperties(); }
 	virtual TArray<FPCGPinProperties> OutputPinProperties() const override;
diff --git a/Engine/Plugins/Experimental/PCG/Source/PCG/Public/Elements/PCGCullPointsOutsideActorBounds.h b/Engine/Plugins/Experimental/PCG/Source/PCG/Public/Elements/PCGCullPointsOutsideActorBounds.h
--- a/Engine/Plugins/Experimental/PCG/Source/PCG/Public/Elements/PCGCullPointsOutsideActorBounds.h
+++ b/Engine/Plugins/Experimental/PCG/Source/PCG/Public/Elements/PCGCullPointsOutsideActorBounds.h
@@ -22,6 +22,7 @@ public:
 	virtual FText GetNodeTooltipText() const override;
 	virtual EPCGSettingsType GetType() const override { return EPCGSettingsType::Spatial; }
 #endif
+	virtual bool IsInputPinRequiredByExecution(const UPCGPin* InPin) const override { return !InPin->Properties.bAdvancedPin; }
 
 protected:
 	virtual TArray<FPCGPinProperties> InputPinProperties() const override;
diff --git a/Engine/Plugins/Experimental/PCG/Source/PCG/Public/Elements/PCGDataNum.h b/Engine/Plugins/Experimental/PCG/Source/PCG/Public/Elements/PCGDataNum.h
--- a/Engine/Plugins/Experimental/PCG/Source/PCG/Public/Elements/PCGDataNum.h
+++ b/Engine/Plugins/Experimental/PCG/Source/PCG/Public/Elements/PCGDataNum.h
@@ -20,7 +20,8 @@ public:
 	virtual FText GetNodeTooltipText() const override;
 	virtual EPCGSettingsType GetType() const override { return EPCGSettingsType::Param; }
 #endif
-	
+	virtual bool IsInputPinRequiredByExecution(const UPCGPin* InPin) const override { return !InPin->Properties.bAdvancedPin; }
+
 protected:
 	virtual TArray<FPCGPinProperties> OutputPinProperties() const override;
 	virtual FPCGElementPtr CreateElement() const override;
diff --git a/Engine/Plugins/Experimental/PCG/Source/PCG/Public/Elements/PCGDebugElement.h b/Engine/Plugins/Experimental/PCG/Source/PCG/Public/Elements/PCGDebugElement.h
--- a/Engine/Plugins/Experimental/PCG/Source/PCG/Public/Elements/PCGDebugElement.h
+++ b/Engine/Plugins/Experimental/PCG/Source/PCG/Public/Elements/PCGDebugElement.h
@@ -25,6 +25,7 @@ public:
 	virtual FText GetDefaultNodeTitle() const override { return NSLOCTEXT("PCGDebugSettings", "NodeTitle", "Debug"); }
 	virtual EPCGSettingsType GetType() const override { return EPCGSettingsType::Debug; }
 #endif
+	virtual bool IsInputPinRequiredByExecution(const UPCGPin* InPin) const override { return !InPin->Properties.bAdvancedPin; }
 
 protected:
 	virtual TArray<FPCGPinProperties> InputPinProperties() const override;
diff --git a/Engine/Plugins/Experimental/PCG/Source/PCG/Public/Elements/PCGDeleteAttributesElement.h b/Engine/Plugins/Experimental/PCG/Source/PCG/Public/Elements/PCGDeleteAttributesElement.h
--- a/Engine/Plugins/Experimental/PCG/Source/PCG/Public/Elements/PCGDeleteAttributesElement.h
+++ b/Engine/Plugins/Experimental/PCG/Source/PCG/Public/Elements/PCGDeleteAttributesElement.h
@@ -39,6 +39,7 @@ public:
 	virtual EPCGSettingsType GetType() const override { return EPCGSettingsType::Metadata; }
 	virtual bool HasDynamicPins() const override { return true; }
 #endif
+	virtual bool IsInputPinRequiredByExecution(const UPCGPin* InPin) const override { return !InPin->Properties.bAdvancedPin; }
 	virtual FString GetAdditionalTitleInformation() const override;
 
 protected:
diff --git a/Engine/Plugins/Experimental/PCG/Source/PCG/Public/Elements/PCGDeleteTags.h b/Engine/Plugins/Experimental/PCG/Source/PCG/Public/Elements/PCGDeleteTags.h
--- a/Engine/Plugins/Experimental/PCG/Source/PCG/Public/Elements/PCGDeleteTags.h
+++ b/Engine/Plugins/Experimental/PCG/Source/PCG/Public/Elements/PCGDeleteTags.h
@@ -27,6 +27,7 @@ public:
 	virtual EPCGSettingsType GetType() const override { return EPCGSettingsType::Generic; }
 	virtual bool HasDynamicPins() const override { return true; }
 #endif
+	virtual bool IsInputPinRequiredByExecution(const UPCGPin* InPin) const override { return !InPin->Properties.bAdvancedPin; }
 	virtual FString GetAdditionalTitleInformation() const override;
 
 protected:
diff --git a/Engine/Plugins/Experimental/PCG/Source/PCG/Public/Elements/PCGDensityFilter.h b/Engine/Plugins/Experimental/PCG/Source/PCG/Public/Elements/PCGDensityFilter.h
--- a/Engine/Plugins/Experimental/PCG/Source/PCG/Public/Elements/PCGDensityFilter.h
+++ b/Engine/Plugins/Experimental/PCG/Source/PCG/Public/Elements/PCGDensityFilter.h
@@ -18,6 +18,7 @@ public:
 	virtual FText GetDefaultNodeTitle() const override { return NSLOCTEXT("PCGDensityFilterSettings", "NodeTitle", "Density Filter"); }
 	virtual EPCGSettingsType GetType() const override { return EPCGSettingsType::Filter; }
 #endif
+	virtual bool IsInputPinRequiredByExecution(const UPCGPin* InPin) const override { return !InPin->Properties.bAdvancedPin; }
 
 protected:
 	virtual TArray<FPCGPinProperties> InputPinProperties() const override { return Super::DefaultPointInputPinProperties(); }
diff --git a/Engine/Plugins/Experimental/PCG/Source/PCG/Public/Elements/PCGDensityRemapElement.h b/Engine/Plugins/Experimental/PCG/Source/PCG/Public/Elements/PCGDensityRemapElement.h
--- a/Engine/Plugins/Experimental/PCG/Source/PCG/Public/Elements/PCGDensityRemapElement.h
+++ b/Engine/Plugins/Experimental/PCG/Source/PCG/Public/Elements/PCGDensityRemapElement.h
@@ -20,6 +20,7 @@ public:
 	virtual FText GetDefaultNodeTitle() const override { return NSLOCTEXT("PCGDensityRemapSettings", "NodeTitle", "Density Remap"); }
 	virtual EPCGSettingsType GetType() const override { return EPCGSettingsType::Density; }
 #endif
+	virtual bool IsInputPinRequiredByExecution(const UPCGPin* InPin) const override { return !InPin->Properties.bAdvancedPin; }
 
 protected:
 	virtual TArray<FPCGPinProperties> InputPinProperties() const override { return Super::DefaultPointInputPinProperties(); }
diff --git a/Engine/Plugins/Experimental/PCG/Source/PCG/Public/Elements/PCGDifferenceElement.h b/Engine/Plugins/Experimental/PCG/Source/PCG/Public/Elements/PCGDifferenceElement.h
--- a/Engine/Plugins/Experimental/PCG/Source/PCG/Public/Elements/PCGDifferenceElement.h
+++ b/Engine/Plugins/Experimental/PCG/Source/PCG/Public/Elements/PCGDifferenceElement.h
@@ -27,6 +27,7 @@ public:
 	virtual EPCGSettingsType GetType() const override { return EPCGSettingsType::Spatial; }
 	virtual void ApplyStructuralDeprecation(UPCGNode* InOutNode) override;
 #endif
+	virtual bool IsInputPinRequiredByExecution(const UPCGPin* InPin) const override { return InPin->Properties.Label == PCGDifferenceConstants::SourceLabel; }
 	virtual FString GetAdditionalTitleInformation() const override;
 
 protected:
diff --git a/Engine/Plugins/Experimental/PCG/Source/PCG/Public/Elements/PCGDistance.h b/Engine/Plugins/Experimental/PCG/Source/PCG/Public/Elements/PCGDistance.h
--- a/Engine/Plugins/Experimental/PCG/Source/PCG/Public/Elements/PCGDistance.h
+++ b/Engine/Plugins/Experimental/PCG/Source/PCG/Public/Elements/PCGDistance.h
@@ -36,11 +36,11 @@ public:
 	virtual FText GetNodeTooltipText() const override;
 	virtual EPCGSettingsType GetType() const override { return EPCGSettingsType::Spatial; }
 #endif
+	virtual bool IsInputPinRequiredByExecution(const UPCGPin* InPin) const override { return InPin->Properties.Label == PCGDistance::SourceLabel; }
 
+protected:
 	virtual TArray<FPCGPinProperties> InputPinProperties() const override;
 	virtual TArray<FPCGPinProperties> OutputPinProperties() const override;
-
-protected:
 	virtual FPCGElementPtr CreateElement() const override;
 	//~End UPCGSettings interface
 
diff --git a/Engine/Plugins/Experimental/PCG/Source/PCG/Public/Elements/PCGDuplicatePoint.h b/Engine/Plugins/Experimental/PCG/Source/PCG/Public/Elements/PCGDuplicatePoint.h
--- a/Engine/Plugins/Experimental/PCG/Source/PCG/Public/Elements/PCGDuplicatePoint.h
+++ b/Engine/Plugins/Experimental/PCG/Source/PCG/Public/Elements/PCGDuplicatePoint.h
@@ -22,6 +22,7 @@ public:
 	virtual FText GetNodeTooltipText() const override { return NSLOCTEXT("PCGDuplicatePointElement", "NodeTooltip", "Creates duplicates of each point with optional transform offsets."); }
 	virtual EPCGSettingsType GetType() const override { return EPCGSettingsType::PointOps; }
 #endif
+	virtual bool IsInputPinRequiredByExecution(const UPCGPin* InPin) const override { return !InPin->Properties.bAdvancedPin; }
 
 protected:
 	virtual TArray<FPCGPinProperties> InputPinProperties() const override { return Super::DefaultPointInputPinProperties(); }
diff --git a/Engine/Plugins/Experimental/PCG/Source/PCG/Public/Elements/PCGExecuteBlueprint.h b/Engine/Plugins/Experimental/PCG/Source/PCG/Public/Elements/PCGExecuteBlueprint.h
--- a/Engine/Plugins/Experimental/PCG/Source/PCG/Public/Elements/PCGExecuteBlueprint.h
+++ b/Engine/Plugins/Experimental/PCG/Source/PCG/Public/Elements/PCGExecuteBlueprint.h
@@ -310,6 +310,8 @@ public:
 	virtual void ApplyPreconfiguredSettings(const FPCGPreConfiguredSettingsInfo& InPreconfiguredsInfo) override;
 	virtual FString GetAdditionalTitleInformation() const override;
 	virtual bool HasFlippedTitleLines() const override { return true; }
+	// This node may have side effects, don't assume we can cull even when unwired.
+	virtual bool CanCullTaskIfUnwired() const { return false; }
 	virtual TArray<FPCGPinProperties> InputPinProperties() const override;
 	virtual TArray<FPCGPinProperties> OutputPinProperties() const override;
 
diff --git a/Engine/Plugins/Experimental/PCG/Source/PCG/Public/Elements/PCGFilterByAttribute.h b/Engine/Plugins/Experimental/PCG/Source/PCG/Public/Elements/PCGFilterByAttribute.h
--- a/Engine/Plugins/Experimental/PCG/Source/PCG/Public/Elements/PCGFilterByAttribute.h
+++ b/Engine/Plugins/Experimental/PCG/Source/PCG/Public/Elements/PCGFilterByAttribute.h
@@ -19,7 +19,7 @@ public:
 	virtual FText GetDefaultNodeTitle() const override;
 	virtual FText GetNodeTooltipText() const override { return NSLOCTEXT("PCGFilterByAttributeElement", "NodeTooltip", "Separates input data by whether they have the specified attribute or not."); }
 #endif
-
+	virtual bool IsInputPinRequiredByExecution(const UPCGPin* InPin) const override { return !InPin->Properties.bAdvancedPin; }
 	virtual FString GetAdditionalTitleInformation() const override;
 
 protected:
diff --git a/Engine/Plugins/Experimental/PCG/Source/PCG/Public/Elements/PCGFilterByIndex.h b/Engine/Plugins/Experimental/PCG/Source/PCG/Public/Elements/PCGFilterByIndex.h
--- a/Engine/Plugins/Experimental/PCG/Source/PCG/Public/Elements/PCGFilterByIndex.h
+++ b/Engine/Plugins/Experimental/PCG/Source/PCG/Public/Elements/PCGFilterByIndex.h
@@ -19,6 +19,7 @@ public:
 	virtual FText GetDefaultNodeTitle() const override;
 	virtual FText GetNodeTooltipText() const override;
 #endif
+	virtual bool IsInputPinRequiredByExecution(const UPCGPin* InPin) const override { return !InPin->Properties.bAdvancedPin; }
 	virtual FString GetAdditionalTitleInformation() const override;
 
 protected:
diff --git a/Engine/Plugins/Experimental/PCG/Source/PCG/Public/Elements/PCGFilterByTag.h b/Engine/Plugins/Experimental/PCG/Source/PCG/Public/Elements/PCGFilterByTag.h
--- a/Engine/Plugins/Experimental/PCG/Source/PCG/Public/Elements/PCGFilterByTag.h
+++ b/Engine/Plugins/Experimental/PCG/Source/PCG/Public/Elements/PCGFilterByTag.h
@@ -26,6 +26,7 @@ public:
 	virtual FText GetDefaultNodeTitle() const override;
 	virtual FText GetNodeTooltipText() const override;
 #endif
+	virtual bool IsInputPinRequiredByExecution(const UPCGPin* InPin) const override { return !InPin->Properties.bAdvancedPin; }
 	virtual FString GetAdditionalTitleInformation() const override;
 
 protected:
diff --git a/Engine/Plugins/Experimental/PCG/Source/PCG/Public/Elements/PCGFilterByType.h b/Engine/Plugins/Experimental/PCG/Source/PCG/Public/Elements/PCGFilterByType.h
--- a/Engine/Plugins/Experimental/PCG/Source/PCG/Public/Elements/PCGFilterByType.h
+++ b/Engine/Plugins/Experimental/PCG/Source/PCG/Public/Elements/PCGFilterByType.h
@@ -21,6 +21,7 @@ public:
 	virtual bool ShouldDrawNodeCompact() const override { return true; }
 	virtual bool GetCompactNodeIcon(FName& OutCompactNodeIcon) const override;
 #endif
+	virtual bool IsInputPinRequiredByExecution(const UPCGPin* InPin) const override { return !InPin->Properties.bAdvancedPin; }
 	virtual EPCGDataType GetCurrentPinTypes(const UPCGPin* InPin) const override;
 
 protected:
diff --git a/Engine/Plugins/Experimental/PCG/Source/PCG/Public/Elements/PCGGather.h b/Engine/Plugins/Experimental/PCG/Source/PCG/Public/Elements/PCGGather.h
--- a/Engine/Plugins/Experimental/PCG/Source/PCG/Public/Elements/PCGGather.h
+++ b/Engine/Plugins/Experimental/PCG/Source/PCG/Public/Elements/PCGGather.h
@@ -22,6 +22,7 @@ public:
 	virtual EPCGSettingsType GetType() const override { return EPCGSettingsType::Spatial; }
 	virtual bool HasDynamicPins() const override { return true; }
 #endif
+	virtual bool IsInputPinRequiredByExecution(const UPCGPin* InPin) const override { return InPin->Properties.Label == PCGPinConstants::DefaultInputLabel; }
 
 protected:
 	virtual TArray<FPCGPinProperties> InputPinProperties() const override;
diff --git a/Engine/Plugins/Experimental/PCG/Source/PCG/Public/Elements/PCGGetBounds.h b/Engine/Plugins/Experimental/PCG/Source/PCG/Public/Elements/PCGGetBounds.h
--- a/Engine/Plugins/Experimental/PCG/Source/PCG/Public/Elements/PCGGetBounds.h
+++ b/Engine/Plugins/Experimental/PCG/Source/PCG/Public/Elements/PCGGetBounds.h
@@ -22,6 +22,7 @@ public:
 	virtual FText GetNodeTooltipText() const override;
 	virtual EPCGSettingsType GetType() const override { return EPCGSettingsType::Spatial; }
 #endif
+	virtual bool IsInputPinRequiredByExecution(const UPCGPin* InPin) const override { return !InPin->Properties.bAdvancedPin; }
 
 protected:
 	virtual TArray<FPCGPinProperties> InputPinProperties() const override;
diff --git a/Engine/Plugins/Experimental/PCG/Source/PCG/Public/Elements/PCGIndirectionElement.h b/Engine/Plugins/Experimental/PCG/Source/PCG/Public/Elements/PCGIndirectionElement.h
--- a/Engine/Plugins/Experimental/PCG/Source/PCG/Public/Elements/PCGIndirectionElement.h
+++ b/Engine/Plugins/Experimental/PCG/Source/PCG/Public/Elements/PCGIndirectionElement.h
@@ -34,6 +34,8 @@ public:
 #endif // WITH_EDITOR
 	virtual FString GetAdditionalTitleInformation() const override;
 	virtual bool HasFlippedTitleLines() const override { return true; }
+	// TODO: This is a safe default and could likely be refined based on the actual indirection target (especially if static).
+	virtual bool CanCullTaskIfUnwired() const { return false; }
 
 protected:
 #if WITH_EDITOR
diff --git a/Engine/Plugins/Experimental/PCG/Source/PCG/Public/Elements/PCGInnerIntersectionElement.h b/Engine/Plugins/Experimental/PCG/Source/PCG/Public/Elements/PCGInnerIntersectionElement.h
--- a/Engine/Plugins/Experimental/PCG/Source/PCG/Public/Elements/PCGInnerIntersectionElement.h
+++ b/Engine/Plugins/Experimental/PCG/Source/PCG/Public/Elements/PCGInnerIntersectionElement.h
@@ -20,6 +20,7 @@ public:
 	virtual FText GetNodeTooltipText() const override;
 	virtual EPCGSettingsType GetType() const override { return EPCGSettingsType::Spatial; }
 #endif
+	virtual bool IsInputPinRequiredByExecution(const UPCGPin* InPin) const override { return !InPin->Properties.bAdvancedPin; }
 
 	virtual TArray<FPCGPinProperties> InputPinProperties() const override;
 	virtual TArray<FPCGPinProperties> OutputPinProperties() const override;
diff --git a/Engine/Plugins/Experimental/PCG/Source/PCG/Public/Elements/PCGMakeConcreteElement.h b/Engine/Plugins/Experimental/PCG/Source/PCG/Public/Elements/PCGMakeConcreteElement.h
--- a/Engine/Plugins/Experimental/PCG/Source/PCG/Public/Elements/PCGMakeConcreteElement.h
+++ b/Engine/Plugins/Experimental/PCG/Source/PCG/Public/Elements/PCGMakeConcreteElement.h
@@ -21,6 +21,7 @@ public:
 	virtual EPCGSettingsType GetType() const override { return EPCGSettingsType::Spatial; }
 	virtual bool ShouldDrawNodeCompact() const override { return true; }
 #endif
+	virtual bool IsInputPinRequiredByExecution(const UPCGPin* InPin) const override { return !InPin->Properties.bAdvancedPin; }
 
 protected:
 	virtual TArray<FPCGPinProperties> OutputPinProperties() const override;
diff --git a/Engine/Plugins/Experimental/PCG/Source/PCG/Public/Elements/PCGMatchAndSetAttributes.h b/Engine/Plugins/Experimental/PCG/Source/PCG/Public/Elements/PCGMatchAndSetAttributes.h
--- a/Engine/Plugins/Experimental/PCG/Source/PCG/Public/Elements/PCGMatchAndSetAttributes.h
+++ b/Engine/Plugins/Experimental/PCG/Source/PCG/Public/Elements/PCGMatchAndSetAttributes.h
@@ -55,6 +55,7 @@ public:
 	virtual FText GetNodeTooltipText() const override;
 	virtual EPCGSettingsType GetType() const override { return EPCGSettingsType::Metadata; }
 #endif // WITH_EDITOR
+	virtual bool IsInputPinRequiredByExecution(const UPCGPin* InPin) const override { return InPin->Properties.Label == PCGPinConstants::DefaultInputLabel; }
 
 protected:
 	virtual TArray<FPCGPinProperties> InputPinProperties() const override;
@@ -120,4 +121,4 @@ class FPCGMatchAndSetAttributesElement : public TPCGTimeSlicedElementBase<FPCGMa
 protected:
 	virtual bool PrepareDataInternal(FPCGContext* InContext) const override;
 	virtual bool ExecuteInternal(FPCGContext* InContext) const override;
-};
\ No newline at end of file
+};
diff --git a/Engine/Plugins/Experimental/PCG/Source/PCG/Public/Elements/PCGMergeAttributes.h b/Engine/Plugins/Experimental/PCG/Source/PCG/Public/Elements/PCGMergeAttributes.h
--- a/Engine/Plugins/Experimental/PCG/Source/PCG/Public/Elements/PCGMergeAttributes.h
+++ b/Engine/Plugins/Experimental/PCG/Source/PCG/Public/Elements/PCGMergeAttributes.h
@@ -20,6 +20,7 @@ public:
 	virtual FText GetNodeTooltipText() const override;
 	virtual EPCGSettingsType GetType() const override { return EPCGSettingsType::Metadata; }
 #endif // WITH_EDITOR
+	virtual bool IsInputPinRequiredByExecution(const UPCGPin* InPin) const override { return !InPin->Properties.bAdvancedPin; }
 
 protected:
 	virtual TArray<FPCGPinProperties> InputPinProperties() const override;
diff --git a/Engine/Plugins/Experimental/PCG/Source/PCG/Public/Elements/PCGMergeElement.h b/Engine/Plugins/Experimental/PCG/Source/PCG/Public/Elements/PCGMergeElement.h
--- a/Engine/Plugins/Experimental/PCG/Source/PCG/Public/Elements/PCGMergeElement.h
+++ b/Engine/Plugins/Experimental/PCG/Source/PCG/Public/Elements/PCGMergeElement.h
@@ -20,6 +20,7 @@ public:
 	virtual FText GetNodeTooltipText() const override;
 	virtual EPCGSettingsType GetType() const override { return EPCGSettingsType::Spatial; }
 #endif
+	virtual bool IsInputPinRequiredByExecution(const UPCGPin* InPin) const override { return !InPin->Properties.bAdvancedPin; }
 
 	virtual TArray<FPCGPinProperties> InputPinProperties() const override;
 	virtual TArray<FPCGPinProperties> OutputPinProperties() const override;
diff --git a/Engine/Plugins/Experimental/PCG/Source/PCG/Public/Elements/PCGMutateSeed.h b/Engine/Plugins/Experimental/PCG/Source/PCG/Public/Elements/PCGMutateSeed.h
--- a/Engine/Plugins/Experimental/PCG/Source/PCG/Public/Elements/PCGMutateSeed.h
+++ b/Engine/Plugins/Experimental/PCG/Source/PCG/Public/Elements/PCGMutateSeed.h
@@ -25,6 +25,7 @@ public:
 	virtual FText GetNodeTooltipText() const override { return NSLOCTEXT("PCGMutateSeedElement", "NodeTooltip", "Applies a new random seed from point input."); }
 	virtual EPCGSettingsType GetType() const override { return EPCGSettingsType::Spatial; }
 #endif
+	virtual bool IsInputPinRequiredByExecution(const UPCGPin* InPin) const override { return !InPin->Properties.bAdvancedPin; }
 
 protected:
 	virtual TArray<FPCGPinProperties> InputPinProperties() const override { return Super::DefaultPointInputPinProperties(); }
diff --git a/Engine/Plugins/Experimental/PCG/Source/PCG/Public/Elements/PCGNormalToDensity.h b/Engine/Plugins/Experimental/PCG/Source/PCG/Public/Elements/PCGNormalToDensity.h
--- a/Engine/Plugins/Experimental/PCG/Source/PCG/Public/Elements/PCGNormalToDensity.h
+++ b/Engine/Plugins/Experimental/PCG/Source/PCG/Public/Elements/PCGNormalToDensity.h
@@ -33,6 +33,7 @@ public:
 	virtual FText GetDefaultNodeTitle() const override { return NSLOCTEXT("PCGNormalToDensitySettings", "NodeTitle", "Normal To Density"); }
 	virtual EPCGSettingsType GetType() const override { return EPCGSettingsType::Spatial; }
 #endif
+	virtual bool IsInputPinRequiredByExecution(const UPCGPin* InPin) const override { return !InPin->Properties.bAdvancedPin; }
 
 protected:
 	virtual TArray<FPCGPinProperties> InputPinProperties() const override { return Super::DefaultPointInputPinProperties(); }
diff --git a/Engine/Plugins/Experimental/PCG/Source/PCG/Public/Elements/PCGNumberOfElements.h b/Engine/Plugins/Experimental/PCG/Source/PCG/Public/Elements/PCGNumberOfElements.h
--- a/Engine/Plugins/Experimental/PCG/Source/PCG/Public/Elements/PCGNumberOfElements.h
+++ b/Engine/Plugins/Experimental/PCG/Source/PCG/Public/Elements/PCGNumberOfElements.h
@@ -20,8 +20,12 @@ class UPCGNumberOfElementsBaseSettings : public UPCGSettings
 {
 	GENERATED_BODY()
 
+	//~Begin UPCGSettings interface
+	virtual bool IsInputPinRequiredByExecution(const UPCGPin* InPin) const override { return !InPin->Properties.bAdvancedPin; }
+
 protected:
 	virtual TArray<FPCGPinProperties> OutputPinProperties() const override;
+	//~End UPCGSettings interface
 
 public:
 	UPROPERTY(BlueprintReadWrite, EditAnywhere, Category = Settings, meta = (PCG_Overridable))
diff --git a/Engine/Plugins/Experimental/PCG/Source/PCG/Public/Elements/PCGOuterIntersectionElement.h b/Engine/Plugins/Experimental/PCG/Source/PCG/Public/Elements/PCGOuterIntersectionElement.h
--- a/Engine/Plugins/Experimental/PCG/Source/PCG/Public/Elements/PCGOuterIntersectionElement.h
+++ b/Engine/Plugins/Experimental/PCG/Source/PCG/Public/Elements/PCGOuterIntersectionElement.h
@@ -7,6 +7,15 @@
 
 #include "PCGOuterIntersectionElement.generated.h"
 
+namespace PCGIntersectionConstants
+{
+	const FName PrimaryLabel = TEXT("Primary Source");
+	const FName SecondaryLabel = TEXT("Source");
+	const FText SecondaryTooltip = NSLOCTEXT("PCGOuterIntersectionElement", "SecondaryPinLabel",
+		"Secondary pin inputs will be implicitly unioned together before being compared to the primary pin's input for calculation of the intersection operation. "
+		"Empty data passed along on one of the secondary pins will result in an empty intersection output, unless the 'Ignore Empty Secondary Input' flag is enabled.");
+}
+
 UCLASS(BlueprintType, ClassGroup = (Procedural))
 class PCG_API UPCGOuterIntersectionSettings : public UPCGSettingsWithDynamicInputs
 {
@@ -21,9 +30,10 @@ public:
 	virtual EPCGSettingsType GetType() const override { return EPCGSettingsType::Spatial; }
 	virtual bool HasDynamicPins() const override { return true; }
 #endif // WITH_EDITOR
-
+	virtual bool IsInputPinRequiredByExecution(const UPCGPin* InPin) const override { return InPin->Properties.Label == PCGIntersectionConstants::PrimaryLabel; }
 	// If node disabled, don't intersect - pass through all primary edges
 	virtual bool OnlyPassThroughOneEdgeWhenDisabled() const override { return false; }
+
 protected:
 	virtual TArray<FPCGPinProperties> OutputPinProperties() const override;
 	virtual FPCGElementPtr CreateElement() const override;
diff --git a/Engine/Plugins/Experimental/PCG/Source/PCG/Public/Elements/PCGPointExtentsModifier.h b/Engine/Plugins/Experimental/PCG/Source/PCG/Public/Elements/PCGPointExtentsModifier.h
--- a/Engine/Plugins/Experimental/PCG/Source/PCG/Public/Elements/PCGPointExtentsModifier.h
+++ b/Engine/Plugins/Experimental/PCG/Source/PCG/Public/Elements/PCGPointExtentsModifier.h
@@ -31,6 +31,7 @@ public:
 	virtual FText GetDefaultNodeTitle() const override { return NSLOCTEXT("PCGPointExtentsModifierSettings", "NodeTitle", "Extents Modifier"); }
 	virtual EPCGSettingsType GetType() const override { return EPCGSettingsType::PointOps; }
 #endif
+	virtual bool IsInputPinRequiredByExecution(const UPCGPin* InPin) const override { return !InPin->Properties.bAdvancedPin; }
 
 protected:
 	virtual TArray<FPCGPinProperties> InputPinProperties() const override { return Super::DefaultPointInputPinProperties(); }
diff --git a/Engine/Plugins/Experimental/PCG/Source/PCG/Public/Elements/PCGPointMatchAndSet.h b/Engine/Plugins/Experimental/PCG/Source/PCG/Public/Elements/PCGPointMatchAndSet.h
--- a/Engine/Plugins/Experimental/PCG/Source/PCG/Public/Elements/PCGPointMatchAndSet.h
+++ b/Engine/Plugins/Experimental/PCG/Source/PCG/Public/Elements/PCGPointMatchAndSet.h
@@ -30,11 +30,11 @@ public:
 	virtual EPCGSettingsType GetType() const override { return EPCGSettingsType::Metadata; }
 	virtual void ApplyDeprecation(UPCGNode* InOutNode) override;
 #endif // WITH_EDITOR
+	virtual bool IsInputPinRequiredByExecution(const UPCGPin* InPin) const override { return !InPin->Properties.bAdvancedPin; }
 
+protected:
 	virtual TArray<FPCGPinProperties> InputPinProperties() const override;
 	virtual TArray<FPCGPinProperties> OutputPinProperties() const override;
-
-protected:
 	virtual FPCGElementPtr CreateElement() const override;
 
 public:
diff --git a/Engine/Plugins/Experimental/PCG/Source/PCG/Public/Elements/PCGPointNeighborhood.h b/Engine/Plugins/Experimental/PCG/Source/PCG/Public/Elements/PCGPointNeighborhood.h
--- a/Engine/Plugins/Experimental/PCG/Source/PCG/Public/Elements/PCGPointNeighborhood.h
+++ b/Engine/Plugins/Experimental/PCG/Source/PCG/Public/Elements/PCGPointNeighborhood.h
@@ -29,6 +29,7 @@ public:
 	virtual FText GetNodeTooltipText() const override { return NSLOCTEXT("PCGPointNeighborhoodElement", "NodeTooltip", "Computes quantities from nearby neighbor points, such as average density, color, and position."); }
 	virtual EPCGSettingsType GetType() const override { return EPCGSettingsType::Spatial; }
 #endif
+	virtual bool IsInputPinRequiredByExecution(const UPCGPin* InPin) const override { return !InPin->Properties.bAdvancedPin; }
 
 protected:
 	virtual TArray<FPCGPinProperties> InputPinProperties() const override { return Super::DefaultPointInputPinProperties(); }
diff --git a/Engine/Plugins/Experimental/PCG/Source/PCG/Public/Elements/PCGProjectionElement.h b/Engine/Plugins/Experimental/PCG/Source/PCG/Public/Elements/PCGProjectionElement.h
--- a/Engine/Plugins/Experimental/PCG/Source/PCG/Public/Elements/PCGProjectionElement.h
+++ b/Engine/Plugins/Experimental/PCG/Source/PCG/Public/Elements/PCGProjectionElement.h
@@ -31,9 +31,12 @@ public:
 	virtual bool HasDynamicPins() const override { return true; }
 #endif
 	virtual EPCGDataType GetCurrentPinTypes(const UPCGPin* InPin) const override;
+	virtual bool IsInputPinRequiredByExecution(const UPCGPin* InPin) const override { return !InPin->Properties.bAdvancedPin; }
 
+protected:
 	virtual TArray<FPCGPinProperties> InputPinProperties() const override;
 
+public:
 	UPROPERTY(BlueprintReadWrite, EditAnywhere, Category = Settings, meta = (ShowOnlyInnerProperties, PCG_Overridable))
 	FPCGProjectionParams ProjectionParams;
 
diff --git a/Engine/Plugins/Experimental/PCG/Source/PCG/Public/Elements/PCGReplaceTags.h b/Engine/Plugins/Experimental/PCG/Source/PCG/Public/Elements/PCGReplaceTags.h
--- a/Engine/Plugins/Experimental/PCG/Source/PCG/Public/Elements/PCGReplaceTags.h
+++ b/Engine/Plugins/Experimental/PCG/Source/PCG/Public/Elements/PCGReplaceTags.h
@@ -20,6 +20,7 @@ public:
 	virtual EPCGSettingsType GetType() const override { return EPCGSettingsType::Generic; }
 	virtual bool HasDynamicPins() const override { return true; }
 #endif
+	virtual bool IsInputPinRequiredByExecution(const UPCGPin* InPin) const override { return !InPin->Properties.bAdvancedPin; }
 	virtual FString GetAdditionalTitleInformation() const override;
 
 protected:
diff --git a/Engine/Plugins/Experimental/PCG/Source/PCG/Public/Elements/PCGResetPointCenter.h b/Engine/Plugins/Experimental/PCG/Source/PCG/Public/Elements/PCGResetPointCenter.h
--- a/Engine/Plugins/Experimental/PCG/Source/PCG/Public/Elements/PCGResetPointCenter.h
+++ b/Engine/Plugins/Experimental/PCG/Source/PCG/Public/Elements/PCGResetPointCenter.h
@@ -23,6 +23,7 @@ public:
 	virtual FText GetNodeTooltipText() const override { return NSLOCTEXT("PCGResetPointCenterElement", "NodeTooltip", "Modify the position of a point within its bounds, while keeping its bounds the same."); }
 	virtual EPCGSettingsType GetType() const override { return EPCGSettingsType::PointOps; }
 #endif
+	virtual bool IsInputPinRequiredByExecution(const UPCGPin* InPin) const override { return !InPin->Properties.bAdvancedPin; }
 
 protected:
 	virtual TArray<FPCGPinProperties> InputPinProperties() const override { return Super::DefaultPointInputPinProperties(); }
diff --git a/Engine/Plugins/Experimental/PCG/Source/PCG/Public/Elements/PCGSelectPoints.h b/Engine/Plugins/Experimental/PCG/Source/PCG/Public/Elements/PCGSelectPoints.h
--- a/Engine/Plugins/Experimental/PCG/Source/PCG/Public/Elements/PCGSelectPoints.h
+++ b/Engine/Plugins/Experimental/PCG/Source/PCG/Public/Elements/PCGSelectPoints.h
@@ -21,6 +21,7 @@ public:
 	virtual FText GetNodeTooltipText() const override;
 	virtual EPCGSettingsType GetType() const override { return EPCGSettingsType::Sampler; }
 #endif
+	virtual bool IsInputPinRequiredByExecution(const UPCGPin* InPin) const override { return !InPin->Properties.bAdvancedPin; }
 
 protected:
 	virtual TArray<FPCGPinProperties> InputPinProperties() const override { return Super::DefaultPointInputPinProperties(); }
diff --git a/Engine/Plugins/Experimental/PCG/Source/PCG/Public/Elements/PCGSelfPruning.h b/Engine/Plugins/Experimental/PCG/Source/PCG/Public/Elements/PCGSelfPruning.h
--- a/Engine/Plugins/Experimental/PCG/Source/PCG/Public/Elements/PCGSelfPruning.h
+++ b/Engine/Plugins/Experimental/PCG/Source/PCG/Public/Elements/PCGSelfPruning.h
@@ -98,6 +98,7 @@ public:
 	virtual FText GetDefaultNodeTitle() const override { return NSLOCTEXT("PCGSelfPruningSettings", "NodeTitle", "Self Pruning"); }
 	virtual EPCGSettingsType GetType() const override { return EPCGSettingsType::Filter; }
 #endif
+	virtual bool IsInputPinRequiredByExecution(const UPCGPin* InPin) const override { return !InPin->Properties.bAdvancedPin; }
 
 protected:
 	virtual TArray<FPCGPinProperties> InputPinProperties() const override { return Super::DefaultPointInputPinProperties(); }
diff --git a/Engine/Plugins/Experimental/PCG/Source/PCG/Public/Elements/PCGSortAttributes.h b/Engine/Plugins/Experimental/PCG/Source/PCG/Public/Elements/PCGSortAttributes.h
--- a/Engine/Plugins/Experimental/PCG/Source/PCG/Public/Elements/PCGSortAttributes.h
+++ b/Engine/Plugins/Experimental/PCG/Source/PCG/Public/Elements/PCGSortAttributes.h
@@ -32,6 +32,7 @@ public:
 	virtual EPCGSettingsType GetType() const override { return EPCGSettingsType::Generic; }
 	virtual bool HasDynamicPins() const override { return true; }
 #endif
+	virtual bool IsInputPinRequiredByExecution(const UPCGPin* InPin) const override { return !InPin->Properties.bAdvancedPin; }
 
 protected:
 	virtual TArray<FPCGPinProperties> OutputPinProperties() const override;
diff --git a/Engine/Plugins/Experimental/PCG/Source/PCG/Public/Elements/PCGSpatialNoise.h b/Engine/Plugins/Experimental/PCG/Source/PCG/Public/Elements/PCGSpatialNoise.h
--- a/Engine/Plugins/Experimental/PCG/Source/PCG/Public/Elements/PCGSpatialNoise.h
+++ b/Engine/Plugins/Experimental/PCG/Source/PCG/Public/Elements/PCGSpatialNoise.h
@@ -62,6 +62,7 @@ public:
 	virtual FText GetDefaultNodeTitle() const override { return NSLOCTEXT("PCGSpatialNoise", "NodeTitle", "Spatial Noise"); }
 	virtual EPCGSettingsType GetType() const override { return EPCGSettingsType::Spatial; }
 #endif
+	virtual bool IsInputPinRequiredByExecution(const UPCGPin* InPin) const override { return !InPin->Properties.bAdvancedPin; }
 
 protected:
 	virtual TArray<FPCGPinProperties> InputPinProperties() const override;
diff --git a/Engine/Plugins/Experimental/PCG/Source/PCG/Public/Elements/PCGSpawnActor.h b/Engine/Plugins/Experimental/PCG/Source/PCG/Public/Elements/PCGSpawnActor.h
--- a/Engine/Plugins/Experimental/PCG/Source/PCG/Public/Elements/PCGSpawnActor.h
+++ b/Engine/Plugins/Experimental/PCG/Source/PCG/Public/Elements/PCGSpawnActor.h
@@ -110,6 +110,7 @@ public:
 #endif // WITH_EDITOR
 
 	//~Begin UCPGSettings interface
+	virtual bool IsInputPinRequiredByExecution(const UPCGPin* InPin) const override { return !InPin->Properties.bAdvancedPin; }
 	virtual UPCGNode* CreateNode() const override;
 
 #if WITH_EDITOR	
diff --git a/Engine/Plugins/Experimental/PCG/Source/PCG/Public/Elements/PCGSplineSampler.h b/Engine/Plugins/Experimental/PCG/Source/PCG/Public/Elements/PCGSplineSampler.h
--- a/Engine/Plugins/Experimental/PCG/Source/PCG/Public/Elements/PCGSplineSampler.h
+++ b/Engine/Plugins/Experimental/PCG/Source/PCG/Public/Elements/PCGSplineSampler.h
@@ -13,6 +13,12 @@ class UPCGPolyLineData;
 class UPCGSpatialData;
 class UPCGPointData;
 
+namespace PCGSplineSamplerConstants
+{
+	const FName SplineLabel = TEXT("Spline");
+	const FName BoundingShapeLabel = TEXT("Bounding Shape");
+}
+
 UENUM()
 enum class EPCGSplineSamplingMode : uint8
 {
@@ -188,12 +194,10 @@ public:
 	virtual FText GetDefaultNodeTitle() const override { return NSLOCTEXT("PCGSplineSamplerSettings", "NodeTitle", "Spline Sampler"); }
 	virtual FText GetNodeTooltipText() const override;
 	virtual EPCGSettingsType GetType() const override { return EPCGSettingsType::Sampler; }
-#endif
-
-#if WITH_EDITOR
 	virtual void ApplyDeprecation(UPCGNode* InOutNode) override;
 	virtual void ApplyDeprecationBeforeUpdatePins(UPCGNode* InOutNode, TArray<TObjectPtr<UPCGPin>>& InputPins, TArray<TObjectPtr<UPCGPin>>& OutputPins) override;
 #endif
+	virtual bool IsInputPinRequiredByExecution(const UPCGPin* InPin) const override { return InPin->Properties.Label == PCGSplineSamplerConstants::SplineLabel; }
 
 protected:
 	virtual TArray<FPCGPinProperties> InputPinProperties() const override;
diff --git a/Engine/Plugins/Experimental/PCG/Source/PCG/Public/Elements/PCGSplitPoints.h b/Engine/Plugins/Experimental/PCG/Source/PCG/Public/Elements/PCGSplitPoints.h
--- a/Engine/Plugins/Experimental/PCG/Source/PCG/Public/Elements/PCGSplitPoints.h
+++ b/Engine/Plugins/Experimental/PCG/Source/PCG/Public/Elements/PCGSplitPoints.h
@@ -30,6 +30,7 @@ public:
 	virtual FText GetNodeTooltipText() const override { return NSLOCTEXT("PCGSplitPointsElement", "NodeTooltip", "Splits each input point into two separate points and sets bounds based on the position and axis of the cut."); }
 	virtual EPCGSettingsType GetType() const override { return EPCGSettingsType::PointOps; }
 #endif
+	virtual bool IsInputPinRequiredByExecution(const UPCGPin* InPin) const override { return !InPin->Properties.bAdvancedPin; }
 
 protected:
 	virtual TArray<FPCGPinProperties> InputPinProperties() const override { return Super::DefaultPointInputPinProperties();  }
diff --git a/Engine/Plugins/Experimental/PCG/Source/PCG/Public/Elements/PCGStaticMeshSpawner.h b/Engine/Plugins/Experimental/PCG/Source/PCG/Public/Elements/PCGStaticMeshSpawner.h
--- a/Engine/Plugins/Experimental/PCG/Source/PCG/Public/Elements/PCGStaticMeshSpawner.h
+++ b/Engine/Plugins/Experimental/PCG/Source/PCG/Public/Elements/PCGStaticMeshSpawner.h
@@ -50,6 +50,7 @@ public:
 	virtual EPCGSettingsType GetType() const override { return EPCGSettingsType::Spawner; }
 	virtual void ApplyDeprecation(UPCGNode* InOutNode) override;
 #endif
+	virtual bool IsInputPinRequiredByExecution(const UPCGPin* InPin) const override { return !InPin->Properties.bAdvancedPin; }
 
 protected:
 	virtual TArray<FPCGPinProperties> InputPinProperties() const override { return Super::DefaultPointInputPinProperties(); }
diff --git a/Engine/Plugins/Experimental/PCG/Source/PCG/Public/Elements/PCGSurfaceSampler.h b/Engine/Plugins/Experimental/PCG/Source/PCG/Public/Elements/PCGSurfaceSampler.h
--- a/Engine/Plugins/Experimental/PCG/Source/PCG/Public/Elements/PCGSurfaceSampler.h
+++ b/Engine/Plugins/Experimental/PCG/Source/PCG/Public/Elements/PCGSurfaceSampler.h
@@ -14,6 +14,12 @@ class UPCGPointData;
 class UPCGSpatialData;
 class UPCGSurfaceSamplerSettings;
 
+namespace PCGSurfaceSamplerConstants
+{
+	const FName SurfaceLabel = TEXT("Surface");
+	const FName BoundingShapeLabel = TEXT("Bounding Shape");
+}
+
 namespace PCGSurfaceSampler
 {
 	// TODO: Factor out "computed values" that shouldn't be exposed to the external user into a more appropriate state
@@ -87,17 +93,14 @@ public:
 	virtual FText GetNodeTooltipText() const override;
 	virtual EPCGSettingsType GetType() const override { return EPCGSettingsType::Sampler; }
 	virtual bool IsPinUsedByNodeExecution(const UPCGPin* InPin) const override;
-#endif
-
-	virtual TArray<FPCGPinProperties> InputPinProperties() const override;
-	virtual TArray<FPCGPinProperties> OutputPinProperties() const override;
-
-#if WITH_EDITOR
 	virtual void ApplyDeprecationBeforeUpdatePins(UPCGNode* InOutNode, TArray<TObjectPtr<UPCGPin>>& InputPins, TArray<TObjectPtr<UPCGPin>>& OutputPins) override;
 	virtual void ApplyDeprecation(UPCGNode* InOutNode) override;
 #endif
+	virtual bool IsInputPinRequiredByExecution(const UPCGPin* InPin) const override { return InPin->Properties.Label == PCGSurfaceSamplerConstants::SurfaceLabel; }
 
 protected:
+	virtual TArray<FPCGPinProperties> InputPinProperties() const override;
+	virtual TArray<FPCGPinProperties> OutputPinProperties() const override;
 	virtual FPCGElementPtr CreateElement() const override;
 	// ~End UPCGSettings interface
 
diff --git a/Engine/Plugins/Experimental/PCG/Source/PCG/Public/Elements/PCGTransformPoints.h b/Engine/Plugins/Experimental/PCG/Source/PCG/Public/Elements/PCGTransformPoints.h
--- a/Engine/Plugins/Experimental/PCG/Source/PCG/Public/Elements/PCGTransformPoints.h
+++ b/Engine/Plugins/Experimental/PCG/Source/PCG/Public/Elements/PCGTransformPoints.h
@@ -20,6 +20,7 @@ public:
 	virtual FText GetDefaultNodeTitle() const override { return NSLOCTEXT("PCGTransformPointsSettings", "NodeTitle", "Transform Points"); }
 	virtual EPCGSettingsType GetType() const override { return EPCGSettingsType::PointOps; }
 #endif
+	virtual bool IsInputPinRequiredByExecution(const UPCGPin* InPin) const override { return !InPin->Properties.bAdvancedPin; }
 
 protected:
 	virtual TArray<FPCGPinProperties> InputPinProperties() const override { return Super::DefaultPointInputPinProperties(); }
diff --git a/Engine/Plugins/Experimental/PCG/Source/PCG/Public/Elements/PCGUnionElement.h b/Engine/Plugins/Experimental/PCG/Source/PCG/Public/Elements/PCGUnionElement.h
--- a/Engine/Plugins/Experimental/PCG/Source/PCG/Public/Elements/PCGUnionElement.h
+++ b/Engine/Plugins/Experimental/PCG/Source/PCG/Public/Elements/PCGUnionElement.h
@@ -19,11 +19,11 @@ public:
 	virtual FText GetDefaultNodeTitle() const override { return NSLOCTEXT("PCGUnionSettings", "NodeTitle", "Union"); }
 	virtual EPCGSettingsType GetType() const override { return EPCGSettingsType::Spatial; }
 #endif
+	virtual bool IsInputPinRequiredByExecution(const UPCGPin* InPin) const override { return !InPin->Properties.bAdvancedPin; }
 
+protected:
 	virtual TArray<FPCGPinProperties> InputPinProperties() const override;
 	virtual TArray<FPCGPinProperties> OutputPinProperties() const override;
-
-protected:
 	virtual FPCGElementPtr CreateElement() const override;
 	//~End UPCGSettings interface
 
diff --git a/Engine/Plugins/Experimental/PCG/Source/PCG/Public/Elements/PCGUserParameterGet.h b/Engine/Plugins/Experimental/PCG/Source/PCG/Public/Elements/PCGUserParameterGet.h
--- a/Engine/Plugins/Experimental/PCG/Source/PCG/Public/Elements/PCGUserParameterGet.h
+++ b/Engine/Plugins/Experimental/PCG/Source/PCG/Public/Elements/PCGUserParameterGet.h
@@ -34,7 +34,7 @@ public:
 	void UpdatePropertyName(FName InNewName);
 
 	//~Begin UPCGSettings interface
-	virtual bool ShouldHookToPreTask() const override { return true; }
+	virtual bool RequiresDataFromPreTask() const override { return true; }
 #if WITH_EDITOR
 	virtual bool CanUserEditTitle() const override { return false; }
 #endif
diff --git a/Engine/Plugins/Experimental/PCG/Source/PCG/Public/Elements/PCGVolumeSampler.h b/Engine/Plugins/Experimental/PCG/Source/PCG/Public/Elements/PCGVolumeSampler.h
--- a/Engine/Plugins/Experimental/PCG/Source/PCG/Public/Elements/PCGVolumeSampler.h
+++ b/Engine/Plugins/Experimental/PCG/Source/PCG/Public/Elements/PCGVolumeSampler.h
@@ -12,6 +12,12 @@
 class UPCGPointData;
 class UPCGSpatialData;
 
+namespace PCGVolumeSamplerConstants
+{
+	const FName VolumeLabel = TEXT("Volume");
+	const FName BoundingShapeLabel = TEXT("Bounding Shape");
+}
+
 namespace PCGVolumeSampler
 {
 	const FVector DefaultVoxelSize = FVector(100.0, 100.0, 100.0);
@@ -75,11 +81,9 @@ public:
 	virtual FText GetDefaultNodeTitle() const override { return NSLOCTEXT("PCGVolumeSamplerSettings", "NodeTitle", "Volume Sampler"); }
 	virtual FText GetNodeTooltipText() const override;
 	virtual EPCGSettingsType GetType() const override { return EPCGSettingsType::Sampler; }
-#endif
-	
-#if WITH_EDITOR
 	virtual void ApplyDeprecationBeforeUpdatePins(UPCGNode* InOutNode, TArray<TObjectPtr<UPCGPin>>& InputPins, TArray<TObjectPtr<UPCGPin>>& OutputPins) override;
 #endif
+	virtual bool IsInputPinRequiredByExecution(const UPCGPin* InPin) const override { return InPin->Properties.Label == PCGVolumeSamplerConstants::VolumeLabel; }
 
 protected:
 	virtual TArray<FPCGPinProperties> InputPinProperties() const override;
diff --git a/Engine/Plugins/Experimental/PCG/Source/PCG/Public/PCGCommon.h b/Engine/Plugins/Experimental/PCG/Source/PCG/Public/PCGCommon.h
--- a/Engine/Plugins/Experimental/PCG/Source/PCG/Public/PCGCommon.h
+++ b/Engine/Plugins/Experimental/PCG/Source/PCG/Public/PCGCommon.h
@@ -9,6 +9,32 @@
 using FPCGTaskId = uint64;
 static const FPCGTaskId InvalidPCGTaskId = (uint64)-1;
 
+using FPCGPinId = uint64;
+
+namespace PCGPinIdHelpers
+{
+	/** Pin active bitmask stored in uint64, so 64 flags available. */
+	constexpr int PinActiveBitmaskSize = 64;
+
+	/** There are 64 pin flags available, however we use flag 63 as a special pin-less ID for task dependencies that don't have associated pins. */
+	constexpr int MaxOutputPins = PinActiveBitmaskSize - 1;
+
+	/** Convert node ID and pin index to a unique pin ID. */
+	FPCGPinId NodeIdAndPinIndexToPinId(FPCGTaskId NodeId, uint64 PinIndex);
+
+	/** Create a pin ID from a node ID alone. Used for task inputs that don't have associated pins. */
+	FPCGPinId NodeIdToPinId(FPCGTaskId NodeId);
+
+	/** Adjust the pin ID to incorporate the given node ID offset. */
+	FPCGPinId OffsetNodeIdInPinId(FPCGPinId PinId, uint64 NodeIDOffset);
+
+	/** Extract node ID from the given pin ID. */
+	FPCGTaskId GetNodeIdFromPinId(FPCGPinId PinId);
+
+	/** Extract pin index from unique pin ID. */
+	uint64 GetPinIndexFromPinId(FPCGPinId PinId);
+}
+
 UENUM(meta = (Bitflags))
 enum class EPCGChangeType : uint8
 {
diff --git a/Engine/Plugins/Experimental/PCG/Source/PCG/Public/PCGComponent.h b/Engine/Plugins/Experimental/PCG/Source/PCG/Public/PCGComponent.h
--- a/Engine/Plugins/Experimental/PCG/Source/PCG/Public/PCGComponent.h
+++ b/Engine/Plugins/Experimental/PCG/Source/PCG/Public/PCGComponent.h
@@ -280,6 +280,9 @@ public:
 	void StopGenerationInProgress();
 	bool IsGenerationInProgress();
 
+	/** Returns current refresh task ID. */
+	bool IsRefreshInProgress() const { return CurrentRefreshTask != InvalidPCGTaskId; }
+
 	/** Dirty generated data depending on the flag. By default the call is forwarded to the local components.
 	    We don't forward if the local component has callbacks that would dirty them too.
 		For example: When a tracked actor move, we only want to dirty the impacted local components.*/
@@ -296,7 +299,19 @@ public:
 	const FPCGDataCollection* GetInspectionData(const FPCGStack& InStack) const;
 	void ClearInspectionData();
 
-	/** Did the given node produce one or more data items in the given stack in a previous execution. */
+	/** Whether a task for the given node and stack was executed during the last execution. */
+	bool WasNodeExecuted(const UPCGNode* InNode, const FPCGStack& Stack) const;
+
+	/** Called at execution time each time a node has been executed. */
+	void NotifyNodeExecuted(const UPCGNode* InNode, const FPCGStack* InStack);
+
+	/** Retrieve the inactive pin bitmask for the given node and stack in the last execution. */
+	uint64 GetNodeInactivePinMask(const UPCGNode* InNode, const FPCGStack& Stack) const;
+
+	/** Whether the given node was culled by a dynamic branch in the given stack. */
+	void NotifyNodeDynamicInactivePins(const UPCGNode* InNode, const FPCGStack* InStack, uint64 InactivePinBitmask) const;
+
+	/** Did the given node produce one or more data items in the given stack during the last execution. */
 	bool HasNodeProducedData(const UPCGNode* InNode, const FPCGStack& Stack) const;
 
 	bool IsObjectTracked(const UObject* InObject, bool& bOutIsCulled) const;
@@ -546,6 +561,14 @@ private:
 	/** Map from nodes to all stacks for which the node produced at least one data item. */
 	TMap<TObjectKey<const UPCGNode>, TSet<FPCGStack>> NodeToStacksThatProducedData;
 	mutable FRWLock NodeToStacksThatProducedDataLock;
+
+	/** Map from nodes to all stacks for which a task for the node was executed. */
+	TMap<TObjectKey<const UPCGNode>, TSet<FPCGStack>> NodeToStacksInWhichNodeExecuted;
+	mutable FRWLock NodeToStacksInWhichNodeExecutedLock;
+
+	/** Map from nodes to stacks to mask of output pins that were deactivated during execution. */
+	mutable TMap<TObjectKey<const UPCGNode>, TMap<const FPCGStack, uint64>> NodeToStackToInactivePinMask;
+	mutable FRWLock NodeToStackToInactivePinMaskLock;
 #endif
 
 	mutable FCriticalSection GeneratedResourcesLock;
diff --git a/Engine/Plugins/Experimental/PCG/Source/PCG/Public/PCGData.h b/Engine/Plugins/Experimental/PCG/Source/PCG/Public/PCGData.h
--- a/Engine/Plugins/Experimental/PCG/Source/PCG/Public/PCGData.h
+++ b/Engine/Plugins/Experimental/PCG/Source/PCG/Public/PCGData.h
@@ -184,6 +184,9 @@ struct PCG_API FPCGDataCollection
 
 	/** A snapshot of the internal state of the data. If any dependency (setting, node input or external data) changes then this value should change. */
 	FPCGCrc Crc;
+
+	/** After the task is complete, bit j is set if output pin index j is deactivated. Stored here so that it can be retrieved from the cache. */
+	uint64 InactiveOutputPinBitmask = 0;
 };
 
 template<typename SettingsType>
diff --git a/Engine/Plugins/Experimental/PCG/Source/PCG/Public/PCGEdge.h b/Engine/Plugins/Experimental/PCG/Source/PCG/Public/PCGEdge.h
--- a/Engine/Plugins/Experimental/PCG/Source/PCG/Public/PCGEdge.h
+++ b/Engine/Plugins/Experimental/PCG/Source/PCG/Public/PCGEdge.h
@@ -31,9 +31,11 @@ public:
 	UPROPERTY()
 	TObjectPtr<UPCGNode> OutboundNode_DEPRECATED;
 
+	/** Pin at upstream end of edge. */
 	UPROPERTY()
 	TObjectPtr<UPCGPin> InputPin;
 
+	/** Pin at downstream end of edge. */
 	UPROPERTY()
 	TObjectPtr<UPCGPin> OutputPin;
 
diff --git a/Engine/Plugins/Experimental/PCG/Source/PCG/Public/PCGEngineSettings.h b/Engine/Plugins/Experimental/PCG/Source/PCG/Public/PCGEngineSettings.h
--- a/Engine/Plugins/Experimental/PCG/Source/PCG/Public/PCGEngineSettings.h
+++ b/Engine/Plugins/Experimental/PCG/Source/PCG/Public/PCGEngineSettings.h
@@ -21,6 +21,10 @@ public:
 	UPROPERTY(EditAnywhere, Config, Category = Workflow)
 	bool bGenerateOnDrop = true;
 
+	/** Display culling state in graph editor when a debug object is selected (requires regeneration to apply). */
+	UPROPERTY(EditAnywhere, Config, Category = Workflow)
+	bool bDisplayCullingStateWhenDebugging = true;
+
 #if WITH_EDITORONLY_DATA
 	// Console variables defined in PCGActorAndComponentMapping.cpp
 	// Use the variables to get the values, not GetDefault<UPCGEngineSettings>
diff --git a/Engine/Plugins/Experimental/PCG/Source/PCG/Public/PCGGraph.h b/Engine/Plugins/Experimental/PCG/Source/PCG/Public/PCGGraph.h
--- a/Engine/Plugins/Experimental/PCG/Source/PCG/Public/PCGGraph.h
+++ b/Engine/Plugins/Experimental/PCG/Source/PCG/Public/PCGGraph.h
@@ -309,9 +309,6 @@ public:
 	void DisableInspection() { bIsInspecting = false; }
 	bool DebugFlagAppliesToIndividualComponents() const { return bDebugFlagAppliesToIndividualComponents; }
 	void RemoveExtraEditorNode(const UObject* InNode);
-
-	/** Is this node on an active branch (and will therefore execute). */
-	bool IsNodeOnActiveBranch(const UPCGNode* InNode) const;
 #endif
 
 #if WITH_EDITOR
@@ -418,12 +415,6 @@ private:
 	/** Remove invalid edges and edges to nodes that are not present in the node array. */
 	void FixInvalidEdges();
 
-	/** Evaluates whether node is on active branch and will therefore execute. Not thread safe, must be called within write lock. */
-	bool CalculateNodeOnActiveBranchRecursive_Unsafe(const UPCGNode* InNode) const;
-
-	/** Clear all active branch flags. */
-	void ResetNodeToOnActiveBranchMap();
-
 	int32 GraphChangeNotificationsDisableCounter = 0;
 	bool bDelayedChangeNotification = false;
 	EPCGChangeType DelayedChangeType = EPCGChangeType::None;
@@ -431,10 +422,6 @@ private:
 	bool bUserPausedNotificationsInGraphEditor = false;
 	int32 NumberOfUserParametersPreEdit = 0;
 	bool bIsInspecting = false;
-
-	/** Cache for "on active branch" flag for nodes. */
-	mutable TMap<const UPCGNode*, bool> NodeToOnActiveBranch;
-	mutable FRWLock NodeToOnActiveBranchLock;
 #endif // WITH_EDITOR
 };
 
diff --git a/Engine/Plugins/Experimental/PCG/Source/PCG/Public/PCGInputOutputSettings.h b/Engine/Plugins/Experimental/PCG/Source/PCG/Public/PCGInputOutputSettings.h
--- a/Engine/Plugins/Experimental/PCG/Source/PCG/Public/PCGInputOutputSettings.h
+++ b/Engine/Plugins/Experimental/PCG/Source/PCG/Public/PCGInputOutputSettings.h
@@ -40,6 +40,8 @@ public:
 	virtual EPCGSettingsType GetType() const override { return EPCGSettingsType::InputOutput; }
 	virtual void PostEditChangeProperty(struct FPropertyChangedEvent& PropertyChangedEvent) override;
 	virtual void ApplyDeprecationBeforeUpdatePins(UPCGNode* InOutNode, TArray<TObjectPtr<UPCGPin>>& InputPins, TArray<TObjectPtr<UPCGPin>>& OutputPins) override;
+	// The output node does not emit a task (input node does however). This prevents it displaying as culled when debugging.
+	virtual bool EmitsTaskForExecution() const override { return bIsInput; }
 #endif
 
 	TArray<FPCGPinProperties> InputPinProperties() const override;
diff --git a/Engine/Plugins/Experimental/PCG/Source/PCG/Public/PCGNode.h b/Engine/Plugins/Experimental/PCG/Source/PCG/Public/PCGNode.h
--- a/Engine/Plugins/Experimental/PCG/Source/PCG/Public/PCGNode.h
+++ b/Engine/Plugins/Experimental/PCG/Source/PCG/Public/PCGNode.h
@@ -132,6 +132,9 @@ public:
 	/** Pin to which data is passed through when this node is disabled. */
 	virtual const UPCGPin* GetPassThroughOutputPin() const;
 
+	/** A node will be executed (not culled) if at least one required-pin is connected to at least one active upstream pin. */
+	virtual bool IsInputPinRequiredByExecution(const UPCGPin* InPin) const;
+
 	/** True if the pin is being used by the node. UI will gray out unused pins. */
 	virtual bool IsPinUsedByNodeExecution(const UPCGPin* InPin) const;
 
diff --git a/Engine/Plugins/Experimental/PCG/Source/PCG/Public/PCGSettings.h b/Engine/Plugins/Experimental/PCG/Source/PCG/Public/PCGSettings.h
--- a/Engine/Plugins/Experimental/PCG/Source/PCG/Public/PCGSettings.h
+++ b/Engine/Plugins/Experimental/PCG/Source/PCG/Public/PCGSettings.h
@@ -243,6 +243,9 @@ public:
 	/** If the node has any dynamic pins that can change based on input or settings */
 	virtual bool HasDynamicPins() const { return false; }
 
+	/** Whether this node can deactivate its output pins during execution, which can dynamically cull downstream nodes. */
+	virtual bool OutputPinsCanBeDeactivated() const { return false; }
+
 	/** Bitwise union of the allowed types of each incident edge on pin. Returns None if no common bits, or no edges. */
 	EPCGDataType GetTypeUnionOfIncidentEdges(const FName& PinLabel) const;
 
@@ -313,6 +316,9 @@ public:
 
 	/** Perform post-operations when an editor node is copied */
 	virtual void PostPaste();
+
+	/** Whether the node will emit one or more tasks for execution. */
+	virtual bool EmitsTaskForExecution() const { return true; }
 #endif // WITH_EDITOR
 
 	/** Derived classes can implement this to expose additional information or context, such as an asset in use by the node. */
@@ -335,12 +341,22 @@ public:
 	UE_DEPRECATED(5.4, "AdditionalTaskName is deprecated and is replaced with GetAdditionalTitleInformation.")
 	virtual FName AdditionalTaskName() const { return FName(GetAdditionalTitleInformation()); }
 
+	/** By default a node does not specify any pin requirements, and will execute if it has no non-advanced pins or if it has any active
+	* connection to any pin. Override this function to specify which pins *must* have active connections to avoid getting culled.
+	*/
+	virtual bool IsInputPinRequiredByExecution(const UPCGPin* InPin) const { return false; }
+
 	/** Returns true if InPin is in use by node (assuming node enabled). Can be used to communicate when a pin is not in use to user. */
 	virtual bool IsPinUsedByNodeExecution(const UPCGPin* InPin) const { return true; }
 
 	/** True if we know prior to execution that the given pin will be active (not on inactive branch). */
 	virtual bool IsPinStaticallyActive(const FName& OutputPinLabel) const { return true; }
 
+	/** True if we can safely cull this node & task if it has unwired non-advanced inputs. Counter-example: BP nodes may have side effects and always execute
+	* if not on an inactive branch.
+	*/
+	virtual bool CanCullTaskIfUnwired() const { return true; }
+
 	/** Returns true if only the first input edge is used from the primary pin when the node is disabled. */
 	virtual bool OnlyPassThroughOneEdgeWhenDisabled() const { return false; }
 
@@ -436,7 +452,7 @@ public:
 	virtual bool HasOverridableParams() const { return !CachedOverridableParams.IsEmpty(); }
 
 	/** Check if we need to hook the output of the pre-task to this. One use is to compute overrides in the subgraph element and pass the overrides as data, to all nodes that needs it. */
-	virtual bool ShouldHookToPreTask() const { return false; }
+	virtual bool RequiresDataFromPreTask() const { return false; }
 protected:
 	/** Iterate over OverridableParams to automatically add param pins to the list. */
 	void FillOverridableParamsPins(TArray<FPCGPinProperties>& OutPins) const;
diff --git a/Engine/Plugins/Experimental/PCG/Source/PCG/Public/PCGSubgraph.h b/Engine/Plugins/Experimental/PCG/Source/PCG/Public/PCGSubgraph.h
--- a/Engine/Plugins/Experimental/PCG/Source/PCG/Public/PCGSubgraph.h
+++ b/Engine/Plugins/Experimental/PCG/Source/PCG/Public/PCGSubgraph.h
@@ -48,6 +48,9 @@ protected:
 public:
 	//~Begin UPCGSettings interface
 	virtual bool HasFlippedTitleLines() const override { return true; }
+	// The graph may contain nodes that have side effects, don't assume we can cull even when unwired.
+	// TODO: For static SGs we could probably compute this value based on the subgraph nodes.
+	virtual bool CanCullTaskIfUnwired() const { return false; }
 
 protected:
 #if WITH_EDITOR
diff --git a/Engine/Plugins/Experimental/PCG/Source/PCGEditor/Private/PCGEditor.cpp b/Engine/Plugins/Experimental/PCG/Source/PCGEditor/Private/PCGEditor.cpp
--- a/Engine/Plugins/Experimental/PCG/Source/PCGEditor/Private/PCGEditor.cpp
+++ b/Engine/Plugins/Experimental/PCG/Source/PCGEditor/Private/PCGEditor.cpp
@@ -2501,17 +2501,24 @@ void FPCGEditor::OnComponentGenerationCompleteOrCancelled()
 {
 	DebugObjectTreeWidget->RequestRefresh();
 
-	// If we are debugging the graph cache then we need to refresh the cache count displayed in the title after every generation.
-	if (UPCGSubsystem* Subsystem = GetSubsystem())
+	const UPCGSubsystem* Subsystem = GetSubsystem();
+	const bool CacheDebuggingEnabled = Subsystem && Subsystem->IsGraphCacheDebuggingEnabled();
+
+	// Refresh nodes to report any errors/warnings, and to display culling state after execution.
+	check(PCGEditorGraph);
+	for (UEdGraphNode* Node : PCGEditorGraph->Nodes)
 	{
-		if (Subsystem->IsGraphCacheDebuggingEnabled() && PCGEditorGraph)
+		if (UPCGEditorGraphNodeBase* PCGEditorGraphNode = Cast<UPCGEditorGraphNodeBase>(Node))
 		{
-			for (UEdGraphNode* EdGraphNode : PCGEditorGraph->Nodes)
+			// If we are debugging the graph cache then we need to refresh the cache count displayed in the title after every generation.
+			EPCGChangeType ChangeType = CacheDebuggingEnabled ? EPCGChangeType::Cosmetic : EPCGChangeType::None;
+
+			ChangeType |= PCGEditorGraphNode->UpdateErrorsAndWarnings();
+			ChangeType |= PCGEditorGraphNode->UpdateStructuralVisualization(PCGComponentBeingInspected.Get(), &StackBeingInspected);
+
+			if (ChangeType != EPCGChangeType::None)
 			{
-				if (UPCGEditorGraphNodeBase* PCGEditorGraphNode = Cast<UPCGEditorGraphNodeBase>(EdGraphNode))
-				{
-					PCGEditorGraphNode->ReconstructNode();
-				}
+				PCGEditorGraphNode->ReconstructNode();
 			}
 		}
 	}
diff --git a/Engine/Plugins/Experimental/PCG/Source/PCGEditor/Private/PCGEditorGraph.cpp b/Engine/Plugins/Experimental/PCG/Source/PCGEditor/Private/PCGEditorGraph.cpp
--- a/Engine/Plugins/Experimental/PCG/Source/PCGEditor/Private/PCGEditorGraph.cpp
+++ b/Engine/Plugins/Experimental/PCG/Source/PCGEditor/Private/PCGEditorGraph.cpp
@@ -166,6 +166,25 @@ void UPCGEditorGraph::UpdateStructuralVisualization(UPCGComponent* PCGComponentB
 	}
 }
 
+const UPCGEditorGraphNodeBase* UPCGEditorGraph::GetEditorNodeFromPCGNode(const UPCGNode* InPCGNode) const
+{
+	if (ensure(InPCGNode))
+	{
+		for (const UEdGraphNode* EdGraphNode : Nodes)
+		{
+			if (const UPCGEditorGraphNodeBase* PCGEdGraphNode = Cast<UPCGEditorGraphNodeBase>(EdGraphNode))
+			{
+				if (PCGEdGraphNode->GetPCGNode() == InPCGNode)
+				{
+					return PCGEdGraphNode;
+				}
+			}
+		}
+	}
+
+	return nullptr;
+}
+
 void UPCGEditorGraph::CreateLinks(UPCGEditorGraphNodeBase* GraphNode, bool bCreateInbound, bool bCreateOutbound, const TMap<UPCGNode*, UPCGEditorGraphNodeBase*>& InPCGNodeToPCGEditorNodeMap)
 {
 	check(GraphNode);
diff --git a/Engine/Plugins/Experimental/PCG/Source/PCGEditor/Private/PCGEditorGraph.h b/Engine/Plugins/Experimental/PCG/Source/PCGEditor/Private/PCGEditorGraph.h
--- a/Engine/Plugins/Experimental/PCG/Source/PCGEditor/Private/PCGEditorGraph.h
+++ b/Engine/Plugins/Experimental/PCG/Source/PCGEditor/Private/PCGEditorGraph.h
@@ -18,7 +18,6 @@ class UPCGEditorGraph : public UEdGraph
 	GENERATED_BODY()
 
 public:
-
 	// ~Begin UObject interface
 	virtual void BeginDestroy() override;
 	// ~End UObject interface
@@ -43,6 +42,9 @@ public:
 	/** Updates the grid size visualization in the editor. */
 	void UpdateStructuralVisualization(UPCGComponent* PCGComponentBeingInspected, const FPCGStack* PCGStackBeingInspected);
 
+	/** Returns the PCG editor graph node corresponding to the given PCG node. */
+	const UPCGEditorGraphNodeBase* GetEditorNodeFromPCGNode(const UPCGNode* InPCGNode) const;
+
 protected:
 	void CreateLinks(UPCGEditorGraphNodeBase* InGraphNode, bool bCreateInbound, bool bCreateOutbound, const TMap<UPCGNode*, UPCGEditorGraphNodeBase*>& InPCGNodeToPCGEditorNodeMap);
 
diff --git a/Engine/Plugins/Experimental/PCG/Source/PCGEditor/Private/PCGEditorGraphNodeBase.cpp b/Engine/Plugins/Experimental/PCG/Source/PCGEditor/Private/PCGEditorGraphNodeBase.cpp
--- a/Engine/Plugins/Experimental/PCG/Source/PCGEditor/Private/PCGEditorGraphNodeBase.cpp
+++ b/Engine/Plugins/Experimental/PCG/Source/PCGEditor/Private/PCGEditorGraphNodeBase.cpp
@@ -3,11 +3,14 @@
 #include "PCGEditorGraphNodeBase.h"
 
 #include "PCGComponent.h"
+#include "PCGEdge.h"
+#include "PCGEngineSettings.h"
 #include "PCGGraph.h"
 #include "PCGPin.h"
 #include "PCGSettingsWithDynamicInputs.h"
 #include "PCGSubsystem.h"
 #include "PCGWorldActor.h"
+#include "Elements/PCGReroute.h"
 
 #include "PCGEditor.h"
 #include "PCGEditorCommands.h"
@@ -26,6 +29,83 @@
 
 #define LOCTEXT_NAMESPACE "PCGEditorGraphNodeBase"
 
+namespace PCGEditorGraphNodeBase
+{
+	/** Whether this node was culled during graph compilation or during graph execution. */
+	bool ShouldDisplayAsActive(const UPCGEditorGraphNodeBase* InNode, const UPCGComponent* InComponentBeingDebugged, const FPCGStack* InStackBeingInspected)
+	{
+		if (!InNode)
+		{
+			return true;
+		}
+
+		const UPCGNode* PCGNode = InNode->GetPCGNode();
+		if (!PCGNode)
+		{
+			return true;
+		}
+
+		// Don't display as culled while component is executing or about to refresh as nodes will flash to culled state and back
+		// which looks disturbing.
+		if (!InComponentBeingDebugged || InComponentBeingDebugged->IsGenerating() || InComponentBeingDebugged->IsRefreshInProgress())
+		{
+			return true;
+		}
+
+		const UPCGEngineSettings* EngineSettings = GetDefault<UPCGEngineSettings>();
+		const bool bActiveVisualizationEnabled = !ensure(EngineSettings) || EngineSettings->bDisplayCullingStateWhenDebugging;
+		const UPCGSettings* Settings = PCGNode ? PCGNode->GetSettings() : nullptr;
+
+		// Display whether node was culled dynamically or statically.
+		if (InStackBeingInspected && bActiveVisualizationEnabled)
+		{
+			if (!Settings || !Settings->IsA<UPCGRerouteSettings>())
+			{
+				// Task will be displayed as active if it was executed or if it does not produce tasks for execution.
+				return InComponentBeingDebugged->WasNodeExecuted(PCGNode, *InStackBeingInspected) || (Settings && !Settings->EmitsTaskForExecution());
+			}
+			else
+			{
+				// Named reroute usages mirror the enabled state of the upstream declaration.
+				if (const UPCGNamedRerouteUsageSettings* RerouteUsageSettings = Cast<UPCGNamedRerouteUsageSettings>(Settings))
+				{
+					const UPCGNode* DeclarationPCGNode = RerouteUsageSettings->Declaration ? Cast<UPCGNode>(RerouteUsageSettings->Declaration->GetOuter()) : nullptr;
+					const UPCGEditorGraph* EditorGraph = Cast<UPCGEditorGraph>(InNode->GetOuter());
+					const UPCGEditorGraphNodeBase* DeclarationNode = EditorGraph ? EditorGraph->GetEditorNodeFromPCGNode(DeclarationPCGNode) : nullptr;
+					return !DeclarationNode || ShouldDisplayAsActive(DeclarationNode, InComponentBeingDebugged, InStackBeingInspected);
+				}
+
+				// Special case - reroute culled state is evaluated here based on upstream connections. Reroutes are always culled/never executed, but still need
+				// to reflect the active/inactive state to not look wrong/confusing.
+				for (const UEdGraphPin* Pin : InNode->Pins)
+				{
+					if (Pin && Pin->Direction == EEdGraphPinDirection::EGPD_Input)
+					{
+						for (const UEdGraphPin* LinkedPin : Pin->LinkedTo)
+						{
+							if (const UPCGEditorGraphNodeBase* UpstreamNode = LinkedPin ? Cast<UPCGEditorGraphNodeBase>(LinkedPin->GetOwningNode()) : nullptr)
+							{
+								const bool bUpstreamNodeActive = ShouldDisplayAsActive(UpstreamNode, InComponentBeingDebugged, InStackBeingInspected);
+								const bool bUpstreamPinActive = UpstreamNode->IsOutputPinActive(LinkedPin);
+
+								if (bUpstreamNodeActive && bUpstreamPinActive)
+								{
+									// Active if any input is active.
+									return true;
+								}
+							}
+						}
+					}
+				}
+
+				return false;
+			}
+		}
+
+		return true;
+	}
+}
+
 void UPCGEditorGraphNodeBase::Construct(UPCGNode* InPCGNode)
 {
 	check(InPCGNode);
@@ -412,11 +492,20 @@ EPCGChangeType UPCGEditorGraphNodeBase::UpdateStructuralVisualization(UPCGCompon
 
 	EPCGChangeType ChangeType = EPCGChangeType::None;
 
-	SetOnActiveBranch(Graph->IsNodeOnActiveBranch(PCGNode));
+	const bool bIsCulled = !PCGEditorGraphNodeBase::ShouldDisplayAsActive(this, InComponentBeingDebugged, InStackBeingInspected);
+
+	SetIsCulledFromExecution(bIsCulled);
 
-	bool bShouldDisplayAsDisabled = !IsOnActiveBranch();
+	bool bShouldDisplayAsDisabled = bIsCulled;
 	bool bShouldDisplayAsHighlighted = false;
 
+	const uint64 NewInactiveMask = (InComponentBeingDebugged && InStackBeingInspected) ? InComponentBeingDebugged->GetNodeInactivePinMask(PCGNode, *InStackBeingInspected) : 0;
+	if (NewInactiveMask != InactiveOutputPinMask)
+	{
+		InactiveOutputPinMask = NewInactiveMask;
+		ChangeType |= EPCGChangeType::Cosmetic;
+	}
+
 	const bool HiGenEnabled = Graph->IsHierarchicalGenerationEnabled();
 	const uint32 InspectingGridSize = InComponentBeingDebugged ? InComponentBeingDebugged->GetGenerationGridSize() : PCGHiGenGrid::UninitializedGridSize();
 
@@ -491,6 +580,38 @@ FText UPCGEditorGraphNodeBase::GetGeneratedTitleLine() const
 	return PCGNode ? PCGNode->GetGeneratedTitleLine() : FText();
 }
 
+bool UPCGEditorGraphNodeBase::IsOutputPinActive(const UEdGraphPin* InOutputPin) const
+{
+	bool bPinActive = true;
+
+	if (InactiveOutputPinMask != 0)
+	{
+		bool bFoundPin = false;
+		int OutputPinIndex = 0;
+
+		for (const UEdGraphPin* NodePin : Pins)
+		{
+			if (NodePin == InOutputPin)
+			{
+				bFoundPin = true;
+				break;
+			}
+
+			if (NodePin->Direction == EEdGraphPinDirection::EGPD_Output)
+			{
+				++OutputPinIndex;
+			}
+		}
+
+		if (bFoundPin)
+		{
+			bPinActive = !((1ULL << OutputPinIndex) & InactiveOutputPinMask);
+		}
+	}
+
+	return bPinActive;
+}
+
 void UPCGEditorGraphNodeBase::EnterRenamingMode()
 {
 	bCanRenameNode = true;
@@ -765,6 +886,7 @@ void UPCGEditorGraphNodeBase::GetPinHoverText(const UEdGraphPin& Pin, FString& H
 		Description = MatchingPin->Properties.Tooltip.IsEmpty() ? FText::FromName(MatchingPin->Properties.Label) : MatchingPin->Properties.Tooltip;
 	}
 
+	FText Required;
 	FText MultiDataSupport;
 	FText MultiConnectionSupport;	
 
@@ -772,19 +894,26 @@ void UPCGEditorGraphNodeBase::GetPinHoverText(const UEdGraphPin& Pin, FString& H
 	{
 		if (bIsInputPin)
 		{
-			MultiDataSupport = MatchingPin->Properties.bAllowMultipleData ? FText(LOCTEXT("InputSupportsMultiData", "Supports multiple data in input(s). ")) : FText(LOCTEXT("InputSingleDataOnly", "Supports only single data in input(s). "));
-			MultiConnectionSupport = MatchingPin->Properties.AllowsMultipleConnections() ? FText(LOCTEXT("SupportsMultiInput", "Supports multiple inputs.")) : FText(LOCTEXT("SingleInputOnly", "Supports only one input."));
+			if (PCGNode && PCGNode->IsInputPinRequiredByExecution(MatchingPin))
+			{
+				Required = LOCTEXT("InputIsRequired", "Required input. ");
+			}
+
+			MultiDataSupport = MatchingPin->Properties.bAllowMultipleData ? LOCTEXT("InputSupportsMultiData", "Supports multiple data in input(s). ") : LOCTEXT("InputSingleDataOnly", "Supports only single data in input(s). ");
+
+			MultiConnectionSupport = MatchingPin->Properties.AllowsMultipleConnections() ? LOCTEXT("SupportsMultiInput", "Supports multiple inputs.") : LOCTEXT("SingleInputOnly", "Supports only one input.");
 		}
 		else
 		{
-			MultiDataSupport = MatchingPin->Properties.bAllowMultipleData ? FText(LOCTEXT("OutputSupportsMultiData", "Can generate multiple data.")) : FText(LOCTEXT("OutputSingleDataOnly", "Generates only single data."));
+			MultiDataSupport = MatchingPin->Properties.bAllowMultipleData ? LOCTEXT("OutputSupportsMultiData", "Can generate multiple data.") : LOCTEXT("OutputSingleDataOnly", "Generates only single data.");
 		}
 	}
 
-	HoverTextOut = FText::Format(LOCTEXT("PinHoverToolTipFull", "{0}\n\nType: {1}\nSubtype: {2}\nAdditional information: {3}{4}"),
+	HoverTextOut = FText::Format(LOCTEXT("PinHoverToolTipFull", "{0}\n\nType: {1}\nSubtype: {2}\nAdditional information: {3}{4}{5}"),
 		Description,
 		DataTypeText,
 		DataSubtypeText,
+		Required,
 		MultiDataSupport,
 		MultiConnectionSupport).ToString();
 }
diff --git a/Engine/Plugins/Experimental/PCG/Source/PCGEditor/Private/PCGEditorGraphNodeBase.h b/Engine/Plugins/Experimental/PCG/Source/PCGEditor/Private/PCGEditorGraphNodeBase.h
--- a/Engine/Plugins/Experimental/PCG/Source/PCGEditor/Private/PCGEditorGraphNodeBase.h
+++ b/Engine/Plugins/Experimental/PCG/Source/PCGEditor/Private/PCGEditorGraphNodeBase.h
@@ -9,11 +9,12 @@
 
 enum class EPCGChangeType : uint8;
 
-struct FPCGStack;
 class UPCGComponent;
+class UPCGGraph;
 class UPCGNode;
 class UPCGPin;
 class UToolMenu;
+struct FPCGStack;
 
 UCLASS()
 class UPCGEditorGraphNodeBase : public UEdGraphNode
@@ -58,8 +59,9 @@ public:
 	void SetInspected(bool InIsInspecting) { bIsInspected = InIsInspecting; }
 	bool GetInspected() const { return bIsInspected; }
 
-	void SetOnActiveBranch(bool bInIsOnActiveBranch) { bIsOnActiveBranch = bInIsOnActiveBranch; }
-	bool IsOnActiveBranch() const { return bIsOnActiveBranch; }
+	/** Whether node was culled either during compilation or at execution time. */
+	void SetIsCulledFromExecution(bool bInIsCulledFromExecution) { bIsCulledFromExecution = bInIsCulledFromExecution; }
+	bool IsCulledFromExecution() const { return bIsCulledFromExecution; }
 
 	/** Increase deferred reconstruct counter, calls to ReconstructNode will flag reconstruct to happen when count hits zero */
 	void EnableDeferredReconstruct();
@@ -93,6 +95,12 @@ public:
 	/** Generated part of node title, not user editable (like "X = 5.0"). */
 	FText GetGeneratedTitleLine() const;
 
+	/** Bitmask of inactive output pins. Bit N will be set if output pin index N is inactive. */
+	uint64 GetInactiveOutputPinMask() const { return InactiveOutputPinMask; }
+
+	/** Whether the given output pin was active in the previous execution. */
+	bool IsOutputPinActive(const UEdGraphPin* InOutputPin) const;
+
 	/** Marks the node as re-nameable and provokes the node to update, placing the user in rename mode. */
 	void EnterRenamingMode();
 
@@ -139,5 +147,10 @@ protected:
 	bool bIsInspected = false;
 	bool bIsHighlighted = false;
 	bool bCanUserAddRemoveSourcePins = false;
-	bool bIsOnActiveBranch = true;
+
+	/** Whether this node was culled in the last execution. */
+	bool bIsCulledFromExecution = false;
+
+	/** Bitmask of inactive output pins. Bit N will be set if output pin index N is inactive. */
+	uint64 InactiveOutputPinMask = 0;
 };
diff --git a/Engine/Plugins/Experimental/PCG/Source/PCGEditor/Private/PCGEditorStyle.cpp b/Engine/Plugins/Experimental/PCG/Source/PCGEditor/Private/PCGEditorStyle.cpp
--- a/Engine/Plugins/Experimental/PCG/Source/PCGEditor/Private/PCGEditorStyle.cpp
+++ b/Engine/Plugins/Experimental/PCG/Source/PCGEditor/Private/PCGEditorStyle.cpp
@@ -64,6 +64,8 @@ FPCGEditorStyle::FPCGEditorStyle() : FSlateStyleSet("PCGEditorStyle")
 	Set(PCGEditorStyleConstants::Pin_Composite_OUT_C, new IMAGE_BRUSH_SVG("Style/PCG_Graph_Composite_OUT", PinSize));
 	Set(PCGEditorStyleConstants::Pin_Composite_OUT_DC, new IMAGE_BRUSH_SVG("Style/PCG_Graph_Composite_OUT_Unplugged", PinSize));
 
+	Set(PCGEditorStyleConstants::Pin_Required, new IMAGE_BRUSH_SVG("Style/PCG_Graph_RequiredPin_IN", FVector2D(8.0f, 22.0f)));
+
 	Set(PCGEditorStyleConstants::Node_Overlay_Inactive, new IMAGE_BRUSH_SVG("Style/PCG_Node_Overlay_Inactive", Icon20x20));
 
 	Set(PCGNodeConstants::Icons::CompactNodeFilter, new IMAGE_BRUSH_SVG("Style/PCG_Graph_Filter", FVector2D(28.0f, 28.0f)));
diff --git a/Engine/Plugins/Experimental/PCG/Source/PCGEditor/Private/PCGEditorStyle.h b/Engine/Plugins/Experimental/PCG/Source/PCGEditor/Private/PCGEditorStyle.h
--- a/Engine/Plugins/Experimental/PCG/Source/PCGEditor/Private/PCGEditorStyle.h
+++ b/Engine/Plugins/Experimental/PCG/Source/PCGEditor/Private/PCGEditorStyle.h
@@ -48,5 +48,7 @@ namespace PCGEditorStyleConstants
 	const FName Pin_Composite_OUT_C = TEXT("PCG.Graph.Composite.Out.Connected");
 	const FName Pin_Composite_OUT_DC = TEXT("PCG.Graph.Composite.Out.Disconnected");
 
+	const FName Pin_Required = TEXT("PCG.Graph.Pin.Required");
+
 	const FName Node_Overlay_Inactive = TEXT("PCG.Node.Overlay.Inactive");
 }
\ No newline at end of file
diff --git a/Engine/Plugins/Experimental/PCG/Source/PCGEditor/Private/SPCGEditorGraphDebugObjectTree.cpp b/Engine/Plugins/Experimental/PCG/Source/PCGEditor/Private/SPCGEditorGraphDebugObjectTree.cpp
--- a/Engine/Plugins/Experimental/PCG/Source/PCGEditor/Private/SPCGEditorGraphDebugObjectTree.cpp
+++ b/Engine/Plugins/Experimental/PCG/Source/PCGEditor/Private/SPCGEditorGraphDebugObjectTree.cpp
@@ -880,7 +880,7 @@ void SPCGEditorGraphDebugObjectTree::OnGetChildren(FPCGEditorGraphDebugObjectIte
 
 void SPCGEditorGraphDebugObjectTree::OnSelectionChanged(FPCGEditorGraphDebugObjectItemPtr InItem, ESelectInfo::Type InSelectInfo)
 {
-	if (const FPCGStack* Stack = InItem ? InItem->GetPCGStack() : nullptr)
+	if (const FPCGStack* Stack = (InItem && InItem->IsDebuggable()) ? InItem->GetPCGStack() : nullptr)
 	{
 		SelectedStack = *Stack;
 	}
@@ -900,9 +900,10 @@ void SPCGEditorGraphDebugObjectTree::OnSelectionChanged(FPCGEditorGraphDebugObje
 		return;
 	}
 
-	if (const FPCGStack* PCGStack = InItem->GetPCGStack())
+	const FPCGStack* InspectedStack = PCGEditor.Pin()->GetStackBeingInspected();
+	if (!InspectedStack || SelectedStack != *InspectedStack)
 	{
-		PCGEditor.Pin()->SetStackBeingInspected(*PCGStack);
+		PCGEditor.Pin()->SetStackBeingInspected(SelectedStack);
 	}
 }
 
diff --git a/Engine/Plugins/Experimental/PCG/Source/PCGEditor/Private/SPCGEditorGraphNode.cpp b/Engine/Plugins/Experimental/PCG/Source/PCGEditor/Private/SPCGEditorGraphNode.cpp
--- a/Engine/Plugins/Experimental/PCG/Source/PCGEditor/Private/SPCGEditorGraphNode.cpp
+++ b/Engine/Plugins/Experimental/PCG/Source/PCGEditor/Private/SPCGEditorGraphNode.cpp
@@ -46,6 +46,9 @@ public:
 	FName GetLabelStyle(FName DefaultLabelStyle) const;
 	bool GetExtraIcon(FName& OutExtraIcon, FText& OutTooltip) const;
 
+	/** Whether pin is required to be connected for execution. */
+	bool IsRequiredForExecution() const;
+
 private:
 	void ApplyUnusedPinStyle(FSlateColor& InOutColor) const;
 	void GetPCGNodeAndPin(const UPCGNode*& OutNode, const UPCGPin*& OutPin) const;
@@ -53,10 +56,9 @@ private:
 
 void SPCGEditorGraphNodePin::Construct(const FArguments& InArgs, UEdGraphPin* InPin)
 {
-	// TODO: replace this with base class when we have sufficient controls to change the padding
-	// IMPLEMENTATION NOTE: this is the code from SGraphPin::Construct
-	// e.g. SGraphPin::Construct(SGraphPin::FArguments().SideToSideMargin(0.0f), InPin);
-	// with additional padding exposed
+	// IMPLEMENTATION NOTE: this is the code from SGraphPin::Construct with additional padding exposed,
+	// with an optional extra icon shown before the pin label, and with a marker icon to show pins
+	// that are required for execution.
 
 	bUsePinColorForText = InArgs._UsePinColorForText;
 	this->SetCursor(EMouseCursor::Default);
@@ -77,6 +79,22 @@ void SPCGEditorGraphNodePin::Construct(const FArguments& InArgs, UEdGraphPin* In
 
 	const bool bIsInput = (GetDirection() == EGPD_Input);
 
+	// A small marker to indicate the pin is required for the node to be executed.
+	TSharedPtr<SImage> RequiredPinIconWidget;
+	float RequiredPinMarkerWidth = 0.0f;
+	bool bDisplayPinMarker = false;
+	if (bIsInput)
+	{
+		const FSlateBrush* RequiredPinMarkerIcon = FPCGEditorStyle::Get().GetBrush(PCGEditorStyleConstants::Pin_Required);
+		RequiredPinMarkerWidth = RequiredPinMarkerIcon ? RequiredPinMarkerIcon->GetImageSize().X : 8.0f;
+		bDisplayPinMarker = IsRequiredForExecution();
+
+		RequiredPinIconWidget =
+			SNew(SImage)
+			.Image(bDisplayPinMarker ? RequiredPinMarkerIcon : FAppStyle::GetNoBrush())
+			.ColorAndOpacity(this, &SPCGEditorGraphNodePin::GetPinColor);
+	}
+
 	// Create the pin icon widget
 	TSharedRef<SWidget> PinWidgetRef = SPinTypeSelector::ConstructPinTypeImage(
 		MakeAttributeSP(this, &SPCGEditorGraphNodePin::GetPinIcon ),
@@ -206,7 +224,13 @@ void SPCGEditorGraphNodePin::Construct(const FArguments& InArgs, UEdGraphPin* In
 			+SHorizontalBox::Slot()
 			.AutoWidth()
 			.VAlign(VAlign_Center)
-			.Padding(0, 0, InArgs._SideToSideMargin, 0)
+			[
+				RequiredPinIconWidget.ToSharedRef()
+			]
+			+SHorizontalBox::Slot()
+			.AutoWidth()
+			.VAlign(VAlign_Center)
+			.Padding(bDisplayPinMarker ? 0.0f : RequiredPinMarkerWidth, 0, InArgs._SideToSideMargin, 0)
 			[
 				PinWidgetRef
 			]
@@ -277,8 +301,14 @@ void SPCGEditorGraphNodePin::GetPCGNodeAndPin(const UPCGNode*& OutNode, const UP
 	if (GraphPin && !GraphPin->IsPendingKill())
 	{
 		const UPCGEditorGraphNodeBase* EditorNode = CastChecked<const UPCGEditorGraphNodeBase>(GraphPinObj->GetOwningNode());
+		
 		OutNode = EditorNode ? EditorNode->GetPCGNode() : nullptr;
+
 		OutPin = OutNode ? OutNode->GetInputPin(GraphPin->GetFName()) : nullptr;
+		if (!OutPin)
+		{
+			OutPin = OutNode ? OutNode->GetOutputPin(GraphPin->GetFName()) : nullptr;
+		}
 	}
 	else
 	{
@@ -291,11 +321,27 @@ void SPCGEditorGraphNodePin::ApplyUnusedPinStyle(FSlateColor& InOutColor) const
 {
 	const UPCGPin* PCGPin = nullptr;
 	const UPCGNode* PCGNode = nullptr;
-
 	GetPCGNodeAndPin(PCGNode, PCGPin);
 
+	bool bPinDisabled = false;
+
+	// Check if the pin was deactivated in the previous execution.
+	const UEdGraphPin* Pin = GetPinObj();
+	if (Pin && Pin->Direction == EEdGraphPinDirection::EGPD_Output)
+	{
+		if (const UPCGEditorGraphNodeBase* Node = Cast<UPCGEditorGraphNodeBase>(Pin ? Pin->GetOwningNode() : nullptr))
+		{
+			// If node is already disabled, don't bother disabling pin on top of that, does not look nice to disable both and may
+			// not be meaningful to do so in any case.
+			if (!Node->IsDisplayAsDisabledForced())
+			{
+				bPinDisabled = !Node->IsOutputPinActive(Pin);
+			}
+		}
+	}
+
 	// Halve opacity if pin is unused - intended to happen whether disabled or not
-	if (PCGPin && PCGNode && !PCGNode->IsPinUsedByNodeExecution(PCGPin))
+	if (bPinDisabled || (PCGPin && PCGNode && !PCGNode->IsPinUsedByNodeExecution(PCGPin)))
 	{
 		FLinearColor Color = InOutColor.GetSpecifiedColor();
 		Color.A *= 0.5;
@@ -351,6 +397,15 @@ bool SPCGEditorGraphNodePin::GetExtraIcon(FName& OutExtraIcon, FText& OutTooltip
 	return (PCGPin && Settings) ? Settings->GetPinExtraIcon(PCGPin, OutExtraIcon, OutTooltip) : false;
 }
 
+bool SPCGEditorGraphNodePin::IsRequiredForExecution() const
+{
+	const UPCGPin* PCGPin = nullptr;
+	const UPCGNode* PCGNode = nullptr;
+	GetPCGNodeAndPin(PCGNode, PCGPin);
+
+	return PCGPin && PCGNode->IsInputPinRequiredByExecution(PCGPin);
+}
+
 void SPCGEditorGraphNode::Construct(const FArguments& InArgs, UPCGEditorGraphNodeBase* InNode)
 {
 	GraphNode = InNode;
@@ -523,52 +578,60 @@ void SPCGEditorGraphNode::AddPin(const TSharedRef<SGraphPin>& PinToAdd)
 	}
 
 	SGraphNode::AddPin(PinToAdd);
+
+	// The base class does not give an override to change the padding of the pin widgets, so do it here. Our input pins widgets include
+	// a small marker to indicate the pin is required, which need to display at the left edge of the node, so remove left padding.
+	if (PinToAdd->GetDirection() == EEdGraphPinDirection::EGPD_Input)
+	{
+		const int LastIndex = LeftNodeBox->GetChildren()->Num() - 1;
+		check(LastIndex >= 0);
+
+		SVerticalBox::FSlot& PinSlot = LeftNodeBox->GetSlot(LastIndex);
+
+		FMargin Margin = Settings->GetInputPinPadding();
+		Margin.Left = 0;
+		PinSlot.SetPadding(Margin);
+	}
 }
 
 void SPCGEditorGraphNode::GetOverlayBrushes(bool bSelected, const FVector2D WidgetSize, TArray<FOverlayBrushInfo>& Brushes) const
 {
 	check(PCGEditorGraphNode);
 	
-	if (!PCGEditorGraphNode->IsOnActiveBranch())
-	{
-		const FSlateBrush* InactiveBranchBrush = FPCGEditorStyle::Get().GetBrush(PCGEditorStyleConstants::Node_Overlay_Inactive);
+	FVector2D OverlayOffset(0.0, 0.0);
 
-		FOverlayBrushInfo BrushInfo;
-		BrushInfo.Brush = InactiveBranchBrush;
-		BrushInfo.OverlayOffset = -InactiveBranchBrush->GetImageSize() / 2.0;
-		Brushes.Add(BrushInfo);
-	}
-	else
+	auto AddOverlayBrush = [&OverlayOffset, &Brushes](const FName& BrushName)
 	{
-		const FSlateBrush* DebugBrush = FPCGEditorStyle::Get().GetBrush(TEXT("PCG.NodeOverlay.Debug"));
-		const FSlateBrush* InspectBrush = FPCGEditorStyle::Get().GetBrush(TEXT("PCG.NodeOverlay.Inspect"));
-
-		const FVector2D HalfDebugBrushSize = DebugBrush->GetImageSize() / 2.0;
-		const FVector2D HalfInspectBrushSize = InspectBrush->GetImageSize() / 2.0;
-
-		FVector2D OverlayOffset(0.0, 0.0);
+		const FSlateBrush* Brush = FPCGEditorStyle::Get().GetBrush(BrushName);
 
-		if (const UPCGNode* PCGNode = PCGEditorGraphNode->GetPCGNode())
+		if (Brush)
 		{
-			if (PCGNode->GetSettingsInterface() && PCGNode->GetSettingsInterface()->bDebug)
-			{
-				FOverlayBrushInfo BrushInfo;
-				BrushInfo.Brush = DebugBrush;
-				BrushInfo.OverlayOffset = OverlayOffset - HalfDebugBrushSize;
-				Brushes.Add(BrushInfo);
+			FOverlayBrushInfo BrushInfo;
+			BrushInfo.Brush = Brush;
+			BrushInfo.OverlayOffset = OverlayOffset - Brush->GetImageSize() / 2.0;
+			Brushes.Add(BrushInfo);
 
-				OverlayOffset.Y += HalfDebugBrushSize.Y + HalfInspectBrushSize.Y;
-			}
+			OverlayOffset.Y += Brush->GetImageSize().Y;
 		}
+	};
+
+	if (PCGEditorGraphNode->IsCulledFromExecution())
+	{
+		AddOverlayBrush(PCGEditorStyleConstants::Node_Overlay_Inactive);
+	}
 
-		if (PCGEditorGraphNode->GetInspected())
+	if (const UPCGNode* PCGNode = PCGEditorGraphNode->GetPCGNode())
+	{
+		if (PCGNode->GetSettingsInterface() && PCGNode->GetSettingsInterface()->bDebug)
 		{
-			FOverlayBrushInfo BrushInfo;
-			BrushInfo.Brush = InspectBrush;
-			BrushInfo.OverlayOffset = OverlayOffset - HalfInspectBrushSize;
-			Brushes.Add(BrushInfo);
+			AddOverlayBrush(TEXT("PCG.NodeOverlay.Debug"));
 		}
 	}
+
+	if (PCGEditorGraphNode->GetInspected())
+	{
+		AddOverlayBrush(TEXT("PCG.NodeOverlay.Inspect"));
+	}
 }
 
 void SPCGEditorGraphNode::OnNodeChanged()
