[PCG] Attempt fix clang link error

diff --git a/Engine/Plugins/Experimental/PCG/Source/PCG/Private/Graph/PCGPinDependencyExpression.h b/Engine/Plugins/Experimental/PCG/Source/PCG/Private/Graph/PCGPinDependencyExpression.h
--- a/Engine/Plugins/Experimental/PCG/Source/PCG/Private/Graph/PCGPinDependencyExpression.h
+++ b/Engine/Plugins/Experimental/PCG/Source/PCG/Private/Graph/PCGPinDependencyExpression.h
@@ -38,10 +38,10 @@ struct FPCGPinDependencyExpression
 
 private:
 	/* These are used to demarcate conjunctions '&&' in the boolean expression. */
-	static const FPCGPinId ConjunctionMarker = std::numeric_limits<FPCGPinId>::max();
+	inline static const FPCGPinId ConjunctionMarker = std::numeric_limits<FPCGPinId>::max();
 
 	/* Signifies a pin removed from the expression, used to set terms false (pin inactive). */
-	static const FPCGPinId RemovedTermMarker = std::numeric_limits<FPCGPinId>::max() - 1;
+	inline static const FPCGPinId RemovedTermMarker = std::numeric_limits<FPCGPinId>::max() - 1;
 
 	TArray<FPCGPinId> Expression;
 };
