[PCG] Fix static analysis issues

diff --git a/Engine/Plugins/Experimental/PCG/Source/PCG/Private/Elements/ControlFlow/PCGSwitch.cpp b/Engine/Plugins/Experimental/PCG/Source/PCG/Private/Elements/ControlFlow/PCGSwitch.cpp
--- a/Engine/Plugins/Experimental/PCG/Source/PCG/Private/Elements/ControlFlow/PCGSwitch.cpp
+++ b/Engine/Plugins/Experimental/PCG/Source/PCG/Private/Elements/ControlFlow/PCGSwitch.cpp
@@ -371,7 +371,7 @@ bool FPCGSwitchElement::ExecuteInternal(FPCGContext* Context) const
 		if (ensure(SelectedPinIndex != INDEX_NONE))
 		{
 			const int AllPinsMask = (1 << NumOutputPins) - 1;
-			Context->OutputData.InactiveOutputPinBitmask = ~(1 << SelectedPinIndex) & AllPinsMask;
+			Context->OutputData.InactiveOutputPinBitmask = ~(1ULL << SelectedPinIndex) & AllPinsMask;
 		}
 	}
 
diff --git a/Engine/Plugins/Experimental/PCG/Source/PCG/Private/Graph/PCGGraphExecutor.cpp b/Engine/Plugins/Experimental/PCG/Source/PCG/Private/Graph/PCGGraphExecutor.cpp
--- a/Engine/Plugins/Experimental/PCG/Source/PCG/Private/Graph/PCGGraphExecutor.cpp
+++ b/Engine/Plugins/Experimental/PCG/Source/PCG/Private/Graph/PCGGraphExecutor.cpp
@@ -1426,7 +1426,7 @@ void FPCGGraphExecutor::CullInactiveDownstreamNodes(FPCGTaskId InCompletedTaskId
 				if (PinCount > 0)
 				{
 					// Deactivate all output pins.
-					const uint64 InactiveOutputPinBitmask = (1 << PinCount) - 1;
+					const uint64 InactiveOutputPinBitmask = (1ULL << PinCount) - 1;
 
 					// Deactivate its pins - add to set of pins to deactivate.
 					GetPinIdsToDeactivate(RemovedTaskId, InactiveOutputPinBitmask, PinIdsToDeactivate);
