[PCG] Optimization for Get Actor Data when gathering PCG Data

* Leverage the fact that all PCG Components are tracked by the PCG Subsystem
* Octree allow for a very efficient spatial search
* Replacing a few TFunction by TFunctionRef in Subsystem API calls.

#rnx

diff --git a/Engine/Plugins/Experimental/PCG/Source/PCG/Private/Elements/PCGDataFromActor.cpp b/Engine/Plugins/Experimental/PCG/Source/PCG/Private/Elements/PCGDataFromActor.cpp
--- a/Engine/Plugins/Experimental/PCG/Source/PCG/Private/Elements/PCGDataFromActor.cpp
+++ b/Engine/Plugins/Experimental/PCG/Source/PCG/Private/Elements/PCGDataFromActor.cpp
@@ -45,6 +45,8 @@ namespace PCGDataFromActorHelpers
 		bool bMustOverlap = false,
 		const FBox& OverlappingBounds = FBox())
 	{
+		TRACE_CPUPROFILER_EVENT_SCOPE(FPCGDataFromActorElement::GetPCGComponentsFromActor);
+
 		TInlineComponentArray<UPCGComponent*, 1> PCGComponents;
 
 		if (!Actor || !Subsystem)
@@ -75,31 +77,27 @@ namespace PCGDataFromActorHelpers
 
 		if (bGetLocalComponents)
 		{
+			auto AddComponent = [&LocalComponents, bGetAllGrids, AllowedGrids](UPCGComponent* LocalComponent)
+			{
+				if (bGetAllGrids || (AllowedGrids & (int32)LocalComponent->GetGenerationGrid()))
+				{
+					LocalComponents.Add(LocalComponent);
+				}
+			};
+
 			// Collect the local components for each actor PCG component.
 			for (UPCGComponent* Component : PCGComponents)
 			{
 				if (Component && Component->IsPartitioned())
 				{
-					Subsystem->ForAllRegisteredLocalComponents(Component, [&LocalComponents, bGetAllGrids, AllowedGrids, bMustOverlap, &OverlappingBounds](UPCGComponent* LocalComponent)
+					if (bMustOverlap)
 					{
-						if (bGetAllGrids || (AllowedGrids & (int32)LocalComponent->GetGenerationGrid()))
-						{
-							if (bMustOverlap)
-							{
-								const FBox LocalBounds = LocalComponent->GetGridBounds();
-
-								// We reject overlaps with zero volume instead of simply checking Intersect(...) to avoid bounds which touch but do not overlap.
-								if (OverlappingBounds.Overlap(LocalBounds).GetVolume() > 0)
-								{
-									LocalComponents.Add(LocalComponent);
-								}
-							}
-							else
-							{
-								LocalComponents.Add(LocalComponent);
-							}
-						}
-					});
+						Subsystem->ForAllRegisteredIntersectingLocalComponents(Component, OverlappingBounds, AddComponent);
+					}
+					else
+					{
+						Subsystem->ForAllRegisteredLocalComponents(Component, AddComponent);
+					}
 				}
 			}
 		}
@@ -309,8 +307,56 @@ bool FPCGDataFromActorElement::ExecuteInternal(FPCGContext* InContext) const
 			};
 		}
 
-		Context->FoundActors = PCGActorSelector::FindActors(Settings->ActorSelector, Context->SourceComponent.Get(), BoundsCheck, SelfIgnoreCheck);
-		Context->bPerformedQuery = true;
+		// When gathering PCG data on any world actor, we can leverage the octree kept by the Tracking system, and get all intersecting components if we need to overlap self
+		// or just gather all registered components (which is way faster than going through all actors in the world).
+		if (Settings->Mode == EPCGGetDataFromActorMode::GetDataFromPCGComponent && Settings->ActorSelector.ActorFilter == EPCGActorFilter::AllWorldActors)
+		{
+			UPCGSubsystem* Subsystem = Context->SourceComponent.IsValid() ? Context->SourceComponent->GetSubsystem() : nullptr;
+			if (Subsystem)
+			{
+				TRACE_CPUPROFILER_EVENT_SCOPE(FPCGDataFromActorElement::Execute::FindPCGComponents);
+
+				const FPCGSelectionKey Key = Settings->ActorSelector.GetAssociatedKey();
+
+				// TODO: Perhaps move the logic into the selector.
+				if (Settings->ActorSelector.bMustOverlapSelf)
+				{
+					FBox ActorBounds = PCGHelpers::GetGridBounds(Self, PCGComponent);
+					for (UPCGComponent* Component : Subsystem->GetAllIntersectingComponents(ActorBounds))
+					{
+						if (AActor* Actor = Component->GetOwner())
+						{
+							if (Key.IsMatching(Actor, Component))
+							{
+								Context->FoundActors.Add(Actor);
+							}
+						}
+					}
+				}
+				else
+				{
+					for (UPCGComponent* Component : Subsystem->GetAllRegisteredComponents())
+					{
+						if (AActor* Actor = Component->GetOwner())
+						{
+							if (Key.IsMatching(Actor, Component))
+							{
+								Context->FoundActors.Add(Actor);
+							}
+						}
+					}
+				}
+
+				Context->bPerformedQuery = true;
+			}
+		}
+		
+		if (!Context->bPerformedQuery)
+		{
+			TRACE_CPUPROFILER_EVENT_SCOPE(FPCGDataFromActorElement::Execute::FindActors);
+			Context->FoundActors = PCGActorSelector::FindActors(Settings->ActorSelector, Context->SourceComponent.Get(), BoundsCheck, SelfIgnoreCheck);
+			Context->bPerformedQuery = true;
+		}
 
 		if (Context->FoundActors.IsEmpty())
 		{
diff --git a/Engine/Plugins/Experimental/PCG/Source/PCG/Private/PCGActorAndComponentMapping.cpp b/Engine/Plugins/Experimental/PCG/Source/PCG/Private/PCGActorAndComponentMapping.cpp
--- a/Engine/Plugins/Experimental/PCG/Source/PCG/Private/PCGActorAndComponentMapping.cpp
+++ b/Engine/Plugins/Experimental/PCG/Source/PCG/Private/PCGActorAndComponentMapping.cpp
@@ -194,7 +194,7 @@ void FPCGActorAndComponentMapping::Tick()
 #endif // WITH_EDITOR
 }
 
-TArray<FPCGTaskId> FPCGActorAndComponentMapping::DispatchToRegisteredLocalComponents(UPCGComponent* OriginalComponent, const TFunction<FPCGTaskId(UPCGComponent*)>& InFunc) const
+TArray<FPCGTaskId> FPCGActorAndComponentMapping::DispatchToRegisteredLocalComponents(UPCGComponent* OriginalComponent, const TFunctionRef<FPCGTaskId(UPCGComponent*)>& InFunc) const
 {
 	TRACE_CPUPROFILER_EVENT_SCOPE(FPCGActorAndComponentMapping::DispatchToRegisteredLocalComponents);
 	if (!ensure(OriginalComponent))
@@ -218,7 +218,7 @@ TArray<FPCGTaskId> FPCGActorAndComponentMapping::DispatchToRegisteredLocalCompon
 	return DispatchToLocalComponents(OriginalComponent, *PartitionActorsPtr, InFunc);
 }
 
-TArray<FPCGTaskId> FPCGActorAndComponentMapping::DispatchToLocalComponents(UPCGComponent* OriginalComponent, const TSet<TObjectPtr<APCGPartitionActor>>& PartitionActors, const TFunction<FPCGTaskId(UPCGComponent*)>& InFunc) const
+TArray<FPCGTaskId> FPCGActorAndComponentMapping::DispatchToLocalComponents(UPCGComponent* OriginalComponent, const TSet<TObjectPtr<APCGPartitionActor>>& PartitionActors, const TFunctionRef<FPCGTaskId(UPCGComponent*)>& InFunc) const
 {
 	TArray<FPCGTaskId> TaskIds;
 	for (APCGPartitionActor* PartitionActor : PartitionActors)
@@ -501,7 +501,7 @@ void FPCGActorAndComponentMapping::UnregisterNonPartitionedPCGComponent(UPCGComp
 	NonPartitionedOctree.RemoveComponent(InComponent);
 }
 
-void FPCGActorAndComponentMapping::ForAllIntersectingComponents(const FBoxCenterAndExtent& InBounds, TFunction<void(UPCGComponent*)> InFunc) const
+void FPCGActorAndComponentMapping::ForAllIntersectingPartitionedComponents(const FBoxCenterAndExtent& InBounds, TFunctionRef<void(UPCGComponent*)> InFunc) const
 {
 	PartitionedOctree.FindElementsWithBoundsTest(InBounds, [&InFunc](const FPCGComponentRef& ComponentRef)
 	{
@@ -509,6 +509,23 @@ void FPCGActorAndComponentMapping::ForAllIntersectingComponents(const FBoxCenter
 	});
 }
 
+TArray<UPCGComponent*> FPCGActorAndComponentMapping::GetAllIntersectingComponents(const FBoxCenterAndExtent& InBounds) const
+{
+	TArray<UPCGComponent*> Result;
+	auto AddToResult = [&Result](const FPCGComponentRef& ComponentRef)
+	{
+		if (IsValid(ComponentRef.Component))
+		{
+			Result.Add(ComponentRef.Component);
+		}
+	};
+
+	PartitionedOctree.FindElementsWithBoundsTest(InBounds, AddToResult);
+	NonPartitionedOctree.FindElementsWithBoundsTest(InBounds, AddToResult);
+
+	return Result;
+}
+
 void FPCGActorAndComponentMapping::RegisterPartitionActor(APCGPartitionActor* InActor, bool bDoComponentMapping)
 {
 	check(InActor);
@@ -547,7 +564,7 @@ void FPCGActorAndComponentMapping::RegisterPartitionActor(APCGPartitionActor* In
 	if (!bIsRuntimeGenerated)
 	{
 		FWriteScopeLock WriteLock(ComponentToPartitionActorsMapLock);
-		ForAllIntersectingComponents(FBoxCenterAndExtent(InActor->GetFixedBounds()), [this, InActor, bDoComponentMapping](UPCGComponent* Component)
+		ForAllIntersectingPartitionedComponents(FBoxCenterAndExtent(InActor->GetFixedBounds()), [this, InActor, bDoComponentMapping](UPCGComponent* Component)
 		{
 			// For each component, do the mapping if we ask it explicitly, or if the component is generated
 			if (bDoComponentMapping || Component->bGenerated)
@@ -589,7 +606,7 @@ void FPCGActorAndComponentMapping::UnregisterPartitionActor(APCGPartitionActor*
 	if (!bIsRuntimeGenerated)
 	{
 		FWriteScopeLock WriteLock(ComponentToPartitionActorsMapLock);
-		ForAllIntersectingComponents(FBoxCenterAndExtent(Actor->GetFixedBounds()), [this, Actor](UPCGComponent* Component)
+		ForAllIntersectingPartitionedComponents(FBoxCenterAndExtent(Actor->GetFixedBounds()), [this, Actor](UPCGComponent* Component)
 		{
 			TSet<TObjectPtr<APCGPartitionActor>>* PartitionActorsPtr = ComponentToPartitionActorsMap.Find(Component);
 			if (PartitionActorsPtr)
@@ -600,7 +617,7 @@ void FPCGActorAndComponentMapping::UnregisterPartitionActor(APCGPartitionActor*
 	}
 }
 
-void FPCGActorAndComponentMapping::ForAllIntersectingPartitionActors(const FBox& InBounds, TFunction<void(APCGPartitionActor*)> InFunc) const
+void FPCGActorAndComponentMapping::ForAllIntersectingPartitionActors(const FBox& InBounds, TFunctionRef<void(APCGPartitionActor*)> InFunc) const
 {
 	// No PCGWorldActor just early out. Same for invalid bounds.
 	APCGWorldActor* PCGWorldActor = PCGSubsystem->GetPCGWorldActor();
diff --git a/Engine/Plugins/Experimental/PCG/Source/PCG/Private/PCGSubsystem.cpp b/Engine/Plugins/Experimental/PCG/Source/PCG/Private/PCGSubsystem.cpp
--- a/Engine/Plugins/Experimental/PCG/Source/PCG/Private/PCGSubsystem.cpp
+++ b/Engine/Plugins/Experimental/PCG/Source/PCG/Private/PCGSubsystem.cpp
@@ -851,7 +851,7 @@ bool UPCGSubsystem::IsGraphCurrentlyExecuting(UPCGGraph* Graph)
 	return GraphExecutor->IsGraphCurrentlyExecuting(Graph);
 }
 
-void UPCGSubsystem::ForAllRegisteredLocalComponents(UPCGComponent* OriginalComponent, const TFunction<void(UPCGComponent*)>& InFunc) const
+void UPCGSubsystem::ForAllRegisteredLocalComponents(UPCGComponent* InOriginalComponent, const TFunctionRef<void(UPCGComponent*)>& InFunc) const
 {
 	auto WrapperFunc = [&InFunc](UPCGComponent* Component) -> FPCGTaskId
 	{
@@ -859,10 +859,30 @@ void UPCGSubsystem::ForAllRegisteredLocalComponents(UPCGComponent* OriginalCompo
 		return InvalidPCGTaskId;
 	};
 
-	ActorAndComponentMapping.DispatchToRegisteredLocalComponents(OriginalComponent, WrapperFunc);
+	ActorAndComponentMapping.DispatchToRegisteredLocalComponents(InOriginalComponent, WrapperFunc);
 }
 
-void UPCGSubsystem::ForAllOverlappingComponentsInHierarchy(UPCGComponent* InComponent, const TFunction<void(UPCGComponent*)>& InFunc) const
+void UPCGSubsystem::ForAllRegisteredIntersectingLocalComponents(UPCGComponent* InOriginalComponent, const FBoxCenterAndExtent& InBounds, const TFunctionRef<void(UPCGComponent*)>& InFunc) const
+{
+	check(InOriginalComponent);
+	const FBox Overlap = InOriginalComponent->GetGridBounds().Overlap(InBounds.GetBox());
+
+	// We reject overlaps with zero volume instead of simply checking Intersect(...) to avoid bounds which touch but do not overlap.
+	if (Overlap.GetVolume() <= 0)
+	{
+		return;
+	}
+
+	ActorAndComponentMapping.ForAllIntersectingPartitionActors(Overlap, [InOriginalComponent, &InFunc](APCGPartitionActor* Actor)
+	{
+		if (UPCGComponent* LocalComponent = Actor->GetLocalComponent(InOriginalComponent))
+		{
+			InFunc(LocalComponent);
+		}
+	});
+}
+
+void UPCGSubsystem::ForAllOverlappingComponentsInHierarchy(UPCGComponent* InComponent, const TFunctionRef<void(UPCGComponent*)>& InFunc) const
 {
 	UPCGComponent* OriginalComponent = InComponent->GetOriginalComponent();
 
@@ -1408,7 +1428,7 @@ void UPCGSubsystem::DeletePartitionActors(bool bOnlyDeleteUnused, bool bOnlyChil
 
 					auto FoundComponents = [&bIntersectWithOneComponent](UPCGComponent*) { bIntersectWithOneComponent = true; };
 
-					ActorAndComponentMapping.ForAllIntersectingComponents(PartitionActor->GetFixedBounds(), FoundComponents);
+					ActorAndComponentMapping.ForAllIntersectingPartitionedComponents(PartitionActor->GetFixedBounds(), FoundComponents);
 
 					if (!bIntersectWithOneComponent)
 					{
diff --git a/Engine/Plugins/Experimental/PCG/Source/PCG/Public/PCGActorAndComponentMapping.h b/Engine/Plugins/Experimental/PCG/Source/PCG/Public/PCGActorAndComponentMapping.h
--- a/Engine/Plugins/Experimental/PCG/Source/PCG/Public/PCGActorAndComponentMapping.h
+++ b/Engine/Plugins/Experimental/PCG/Source/PCG/Public/PCGActorAndComponentMapping.h
@@ -106,16 +106,19 @@ private:
 	void UnregisterNonPartitionedPCGComponent(UPCGComponent* InComponent);
 
 	/* Call the InFunc function to all local component registered to the original component. Return the list of all the tasks scheduled. Thread safe*/
-	TArray<FPCGTaskId> DispatchToRegisteredLocalComponents(UPCGComponent* OriginalComponent, const TFunction<FPCGTaskId(UPCGComponent*)>& InFunc) const;
+	TArray<FPCGTaskId> DispatchToRegisteredLocalComponents(UPCGComponent* OriginalComponent, const TFunctionRef<FPCGTaskId(UPCGComponent*)>& InFunc) const;
 
 	/* Call the InFunc function to all local component from the set of partition actors. Return the list of all the tasks scheduled. */
-	TArray<FPCGTaskId> DispatchToLocalComponents(UPCGComponent* OriginalComponent, const TSet<TObjectPtr<APCGPartitionActor>>& PartitionActors, const TFunction<FPCGTaskId(UPCGComponent*)>& InFunc) const;
+	TArray<FPCGTaskId> DispatchToLocalComponents(UPCGComponent* OriginalComponent, const TSet<TObjectPtr<APCGPartitionActor>>& PartitionActors, const TFunctionRef<FPCGTaskId(UPCGComponent*)>& InFunc) const;
 
-	/** Iterate other all the components which bounds intersect the box in param and call a callback. Thread safe */
-	void ForAllIntersectingComponents(const FBoxCenterAndExtent& InBounds, TFunction<void(UPCGComponent*)> InFunc) const;
+	/** Call the InFunc function to all partitioned components which bounds intersect 'InBounds'. */
+	void ForAllIntersectingPartitionedComponents(const FBoxCenterAndExtent& InBounds, TFunctionRef<void(UPCGComponent*)> InFunc) const;
+
+	/** Gather all the PCG components within some bounds. */
+	TArray<UPCGComponent*> GetAllIntersectingComponents(const FBoxCenterAndExtent& InBounds) const;
 
 	/** Iterate other all the int coordinates given a box and call a callback. Thread safe */
-	void ForAllIntersectingPartitionActors(const FBox& InBounds, TFunction<void(APCGPartitionActor*)> InFunc) const;
+	void ForAllIntersectingPartitionActors(const FBox& InBounds, TFunctionRef<void(APCGPartitionActor*)> InFunc) const;
 
 	/** Update the current mapping between a PCG component and its PCG Partition actors */
 	void UpdateMappingPCGComponentPartitionActor(UPCGComponent* InComponent);
diff --git a/Engine/Plugins/Experimental/PCG/Source/PCG/Public/PCGSubsystem.h b/Engine/Plugins/Experimental/PCG/Source/PCG/Public/PCGSubsystem.h
--- a/Engine/Plugins/Experimental/PCG/Source/PCG/Public/PCGSubsystem.h
+++ b/Engine/Plugins/Experimental/PCG/Source/PCG/Public/PCGSubsystem.h
@@ -182,15 +182,25 @@ public:
 	void UnregisterPartitionActor(APCGPartitionActor* InActor) { ActorAndComponentMapping.UnregisterPartitionActor(InActor); }
 
 	TSet<UPCGComponent*> GetAllRegisteredPartitionedComponents() const { return ActorAndComponentMapping.GetAllRegisteredPartitionedComponents(); }
+	TSet<UPCGComponent*> GetAllRegisteredComponents() const { return ActorAndComponentMapping.GetAllRegisteredComponents(); }
 
 	/** Flushes the graph cache completely, use only for debugging */
 	void FlushCache();
 
-	/* Call the InFunc function to all local component registered to the original component. Thread safe*/
-	void ForAllRegisteredLocalComponents(UPCGComponent* OriginalComponent, const TFunction<void(UPCGComponent*)>& InFunc) const;
+	/** Call the InFunc function to all local component registered to the original component. Thread safe*/
+	void ForAllRegisteredLocalComponents(UPCGComponent* InOriginalComponent, const TFunctionRef<void(UPCGComponent*)>& InFunc) const;
+
+	/** Call the InFunc function to all local component registered to the original component within some bounds. Thread safe*/
+	void ForAllRegisteredIntersectingLocalComponents(UPCGComponent* InOriginalComponent, const FBoxCenterAndExtent& InBounds, const TFunctionRef<void(UPCGComponent*)>& InFunc) const;
+
+	/** Get all components in specified bounds. */
+	void ForAllIntersectingPartitionedComponents(const FBoxCenterAndExtent& InBounds, TFunctionRef<void(UPCGComponent*)> InFunc) const { return ActorAndComponentMapping.ForAllIntersectingPartitionedComponents(InBounds, std::move(InFunc)); }
+
+	/** Gather all the PCG components within some bounds. */
+	TArray<UPCGComponent*> GetAllIntersectingComponents(const FBoxCenterAndExtent& InBounds) const { return ActorAndComponentMapping.GetAllIntersectingComponents(InBounds); }
 
 	/** Traverses the hierarchy associated with the given component and calls InFunc for each overlapping component. */
-	void ForAllOverlappingComponentsInHierarchy(UPCGComponent* InComponent, const TFunction<void(UPCGComponent*)>& InFunc) const;
+	void ForAllOverlappingComponentsInHierarchy(UPCGComponent* InComponent, const TFunctionRef<void(UPCGComponent*)>& InFunc) const;
 
 	/** Retrieves a local component using grid size and grid coordinates, returns nullptr if no such component found. */
 	UPCGComponent* GetLocalComponent(uint32 GridSize, const FIntVector& CellCoords, const UPCGComponent* InOriginalComponent, bool bTransient = false);
