[PCG] Fix higen grid visualisation when inspecting higen graphs. The culling visualisation was overriding it.

diff --git a/Engine/Plugins/Experimental/PCG/Source/PCGEditor/Private/PCGEditorGraphNodeBase.cpp b/Engine/Plugins/Experimental/PCG/Source/PCGEditor/Private/PCGEditorGraphNodeBase.cpp
--- a/Engine/Plugins/Experimental/PCG/Source/PCGEditor/Private/PCGEditorGraphNodeBase.cpp
+++ b/Engine/Plugins/Experimental/PCG/Source/PCGEditor/Private/PCGEditorGraphNodeBase.cpp
@@ -492,11 +492,7 @@ EPCGChangeType UPCGEditorGraphNodeBase::UpdateStructuralVisualization(UPCGCompon
 
 	EPCGChangeType ChangeType = EPCGChangeType::None;
 
-	const bool bIsCulled = !PCGEditorGraphNodeBase::ShouldDisplayAsActive(this, InComponentBeingDebugged, InStackBeingInspected);
-
-	SetIsCulledFromExecution(bIsCulled);
-
-	bool bShouldDisplayAsDisabled = bIsCulled;
+	bool bShouldDisplayAsDisabled = false;
 	bool bShouldDisplayAsHighlighted = false;
 
 	const uint64 NewInactiveMask = (InComponentBeingDebugged && InStackBeingInspected) ? InComponentBeingDebugged->GetNodeInactivePinMask(PCGNode, *InStackBeingInspected) : 0;
@@ -506,6 +502,8 @@ EPCGChangeType UPCGEditorGraphNodeBase::UpdateStructuralVisualization(UPCGCompon
 		ChangeType |= EPCGChangeType::Cosmetic;
 	}
 
+	bool bIsCulled = !PCGEditorGraphNodeBase::ShouldDisplayAsActive(this, InComponentBeingDebugged, InStackBeingInspected);
+
 	const bool HiGenEnabled = Graph->IsHierarchicalGenerationEnabled();
 	const uint32 InspectingGridSize = InComponentBeingDebugged ? InComponentBeingDebugged->GetGenerationGridSize() : PCGHiGenGrid::UninitializedGridSize();
 
@@ -531,11 +529,30 @@ EPCGChangeType UPCGEditorGraphNodeBase::UpdateStructuralVisualization(UPCGCompon
 		}
 		const uint32 NodeGridSize = Graph->GetNodeGenerationGridSize(PCGNode, DefaultGridSize);
 
-		// Disable nodes that are on a smaller grid
-		bShouldDisplayAsDisabled |= NodeGridSize < InspectingGridSize;
+		if (NodeGridSize < InspectingGridSize)
+		{
+			// Disable nodes that are on a smaller grid
+			bShouldDisplayAsDisabled |= NodeGridSize < InspectingGridSize;
+
+			// We don't know if the node was culled or not on that grid, disable visualization.
+			bIsCulled = false;
+		}
+		else if (NodeGridSize > InspectingGridSize)
+		{
+			// If node is on larger grid than current, highlight it to indicate that its data is available for use
+			bShouldDisplayAsHighlighted |= true;
+
+			// We don't know if the node was culled or not on that grid, disable visualization.
+			bIsCulled = false;
+		}
+	}
 
-		// If node is on larger grid than current, highlight it to indicate that its data is available for use
-		bShouldDisplayAsHighlighted |= NodeGridSize > InspectingGridSize;
+	SetIsCulledFromExecution(bIsCulled);
+
+	if (bIsCulled)
+	{
+		bShouldDisplayAsDisabled = true;
+		bShouldDisplayAsHighlighted = false;
 	}
 
 	if (IsDisplayAsDisabledForced() != bShouldDisplayAsDisabled)
