[PCG] Fix crash when dragging out and creating new node from a redirector.

We could not repro this again so did not get to the root cause (the from-pin in the graph action object had no owner node, so some problem upstream). Putting this in as a short term fix to prevent a hard crash, with an ensure to make sure it doesn't fall off the radar.

diff --git a/Engine/Plugins/Experimental/PCG/Source/PCGEditor/Private/PCGEditorGraphSchemaActions.cpp b/Engine/Plugins/Experimental/PCG/Source/PCGEditor/Private/PCGEditorGraphSchemaActions.cpp
--- a/Engine/Plugins/Experimental/PCG/Source/PCGEditor/Private/PCGEditorGraphSchemaActions.cpp
+++ b/Engine/Plugins/Experimental/PCG/Source/PCGEditor/Private/PCGEditorGraphSchemaActions.cpp
@@ -69,7 +69,7 @@ UEdGraphNode* FPCGEditorGraphSchemaAction_NewNativeElement::PerformAction(UEdGra
 	NewPCGNode->PositionX = Location.X;
 	NewPCGNode->PositionY = Location.Y;
 
-	if (FromPin)
+	if (FromPin && ensure(FromPin->GetOwningNode()))
 	{
 		NewNode->AutowireNewNode(FromPin);
 	}
