[PCG] Add barriers for pcg generation dependencies

* PCG Component will now register themselves as dependencies when they start generating.
* When the refresh trigger happens, it will check if all dependencies are done
* It's behind a CVar to disable it quickly if side effects are detected

#rnx

diff --git a/Engine/Plugins/Experimental/PCG/Source/PCG/Private/Elements/PCGDataFromActor.cpp b/Engine/Plugins/Experimental/PCG/Source/PCG/Private/Elements/PCGDataFromActor.cpp
--- a/Engine/Plugins/Experimental/PCG/Source/PCG/Private/Elements/PCGDataFromActor.cpp
+++ b/Engine/Plugins/Experimental/PCG/Source/PCG/Private/Elements/PCGDataFromActor.cpp
@@ -610,9 +610,12 @@ void FPCGDataFromActorElement::ProcessActor(FPCGContext* Context, const UPCGData
 			// this not to be a problem, as these intersections should be garbage collected, but this requires more investigation
 			for (const FPCGTaggedData& TaggedData : Component->GetGeneratedGraphOutput().TaggedData)
 			{
-				FPCGTaggedData& DuplicatedTaggedData = Outputs.Add_GetRef(TaggedData);
-				DuplicatedTaggedData.Data = Cast<UPCGData>(StaticDuplicateObject(TaggedData.Data, GetTransientPackage()));
-				DuplicatedTaggedData.Tags.Add(PCGDataFromActorConstants::PCGComponentDataGridSizeTagPrefix + FString::FromInt(PCGHiGenGrid::GridToGridSize(Component->GetGenerationGrid())));
+				if (ensure(TaggedData.Data))
+				{
+					FPCGTaggedData& DuplicatedTaggedData = Outputs.Add_GetRef(TaggedData);
+					DuplicatedTaggedData.Data = Cast<UPCGData>(StaticDuplicateObject(TaggedData.Data, GetTransientPackage()));
+					DuplicatedTaggedData.Tags.Add(PCGDataFromActorConstants::PCGComponentDataGridSizeTagPrefix + FString::FromInt(PCGHiGenGrid::GridToGridSize(Component->GetGenerationGrid())));
+				}
 			}
 			//Outputs.Append(Component->GetGeneratedGraphOutput().TaggedData);
 		}
diff --git a/Engine/Plugins/Experimental/PCG/Source/PCG/Private/PCGActorAndComponentMapping.cpp b/Engine/Plugins/Experimental/PCG/Source/PCG/Private/PCGActorAndComponentMapping.cpp
--- a/Engine/Plugins/Experimental/PCG/Source/PCG/Private/PCGActorAndComponentMapping.cpp
+++ b/Engine/Plugins/Experimental/PCG/Source/PCG/Private/PCGActorAndComponentMapping.cpp
@@ -81,6 +81,11 @@ namespace PCGActorAndComponentMapping
 		TEXT("pcg.ActorModifiedPreviousDataCheckDelayMS"),
 		1000,
 		TEXT("Delay in MS between cleanup checks on previous data from modified actors."));
+
+	static TAutoConsoleVariable<bool> CVarDisablePCGDataInterdependencyOptimization(
+		TEXT("pcg.DisablePCGDataInterdependencyOptimization"),
+		false,
+		TEXT("Disable the optimization that keep track of components depending on others, as a safety measure."));
 #endif // WITH_EDITOR
 
 	static TAutoConsoleVariable<bool> CVarDisableDelayedUnregister(
@@ -191,6 +196,11 @@ void FPCGActorAndComponentMapping::Tick()
 			}
 		}
 	}
+
+	if (PCGActorAndComponentMapping::CVarDisablePCGDataInterdependencyOptimization.GetValueOnAnyThread() && !ComponentsToDependencyMap.IsEmpty())
+	{
+		ComponentsToDependencyMap.Empty();
+	}
 #endif // WITH_EDITOR
 }
 
@@ -877,6 +887,8 @@ void FPCGActorAndComponentMapping::RegisterTracking(UPCGComponent* InComponent)
 	{
 		InComponent->OnPCGGraphGeneratedDelegate.AddRaw(this, &FPCGActorAndComponentMapping::OnPCGGraphGeneratedOrCleaned);
 		InComponent->OnPCGGraphCleanedDelegate.AddRaw(this, &FPCGActorAndComponentMapping::OnPCGGraphGeneratedOrCleaned);
+		InComponent->OnPCGGraphStartGeneratingDelegate.AddRaw(this, &FPCGActorAndComponentMapping::OnPCGGraphStartsGenerating);
+		InComponent->OnPCGGraphCancelledDelegate.AddRaw(this, &FPCGActorAndComponentMapping::OnPCGGraphCancelled);
 	}
 }
 
@@ -964,11 +976,27 @@ void FPCGActorAndComponentMapping::RemapTracking(const UPCGComponent* InOldCompo
 	ReplaceInMap(CulledTrackedKeysToComponentsMap);
 	ReplaceInMap(AlwaysTrackedKeysToComponentsMap);
 
+	// If this component has dependencies, we transfer them.
+	if (ComponentsToDependencyMap.Contains(InOldComponent))
+	{
+		TArray<TObjectKey<UPCGComponent>> Temp;
+		ComponentsToDependencyMap.RemoveAndCopyValue(InOldComponent, Temp);
+		ComponentsToDependencyMap.Emplace(InNewComponent, std::move(Temp));
+	}
+
+	// If this component was a dependency to any other component, just remove it. New one will register itself when generating if needed.
+	for (TPair<TObjectKey<UPCGComponent>, TArray<TObjectKey<UPCGComponent>>>& It : ComponentsToDependencyMap)
+	{
+		It.Value.RemoveSwap(InOldComponent);
+	}
+
 	// Old component will probably die, but we'll force removing the delegates even if it is const.
 	if (UPCGComponent* MutableOldComponent = const_cast<UPCGComponent*>(InOldComponent))
 	{
 		MutableOldComponent->OnPCGGraphGeneratedDelegate.RemoveAll(this);
 		MutableOldComponent->OnPCGGraphCleanedDelegate.RemoveAll(this);
+		MutableOldComponent->OnPCGGraphStartGeneratingDelegate.RemoveAll(this);
+		MutableOldComponent->OnPCGGraphCancelledDelegate.RemoveAll(this);
 	}
 
 	// And just making sure we are not registering multiple times
@@ -976,6 +1004,8 @@ void FPCGActorAndComponentMapping::RemapTracking(const UPCGComponent* InOldCompo
 	{
 		InNewComponent->OnPCGGraphGeneratedDelegate.AddRaw(this, &FPCGActorAndComponentMapping::OnPCGGraphGeneratedOrCleaned);
 		InNewComponent->OnPCGGraphCleanedDelegate.AddRaw(this, &FPCGActorAndComponentMapping::OnPCGGraphGeneratedOrCleaned);
+		InNewComponent->OnPCGGraphStartGeneratingDelegate.AddRaw(this, &FPCGActorAndComponentMapping::OnPCGGraphStartsGenerating);
+		InNewComponent->OnPCGGraphCancelledDelegate.AddRaw(this, &FPCGActorAndComponentMapping::OnPCGGraphCancelled);
 	}
 }
 
@@ -1043,6 +1073,8 @@ void FPCGActorAndComponentMapping::UnregisterTracking(UPCGComponent* InComponent
 
 	InComponent->OnPCGGraphGeneratedDelegate.RemoveAll(this);
 	InComponent->OnPCGGraphCleanedDelegate.RemoveAll(this);
+	InComponent->OnPCGGraphStartGeneratingDelegate.RemoveAll(this);
+	InComponent->OnPCGGraphCancelledDelegate.RemoveAll(this);
 }
 
 bool FPCGActorAndComponentMapping::IsKeyTracked(const FPCGSelectionKey& InKey) const
@@ -1756,6 +1788,7 @@ void FPCGActorAndComponentMapping::OnObjectChanged(UObject* InObject, const FAct
 	}
 
 	ULevelInstanceSubsystem* LevelInstanceSubsystem = (PCGSubsystem && PCGSubsystem->GetWorld()) ? PCGSubsystem->GetWorld()->GetSubsystem<ULevelInstanceSubsystem>() : nullptr;
+	const UPCGComponent* OriginatingComponent = Cast<const UPCGComponent>(InOriginatingChangeObject);
 
 	// And refresh all dirtied components
 	for (UPCGComponent* Component : DirtyComponents)
@@ -1765,6 +1798,24 @@ void FPCGActorAndComponentMapping::OnObjectChanged(UObject* InObject, const FAct
 			continue;
 		}
 
+		// This part checks if the change originates from a PCG Component. If so, we check the dirty component has no more other PCG dependencies.
+		// In that case we can proceed for the refresh, otherwise early out, this will be woken up by another dependency change.
+		if (OriginatingComponent)
+		{
+			if (TArray<TObjectKey<UPCGComponent>>* Dependencies = ComponentsToDependencyMap.Find(Component))
+			{
+				Dependencies->RemoveSwap(OriginatingComponent);
+				if (Dependencies->IsEmpty())
+				{
+					ComponentsToDependencyMap.Remove(Component);
+				}
+				else
+				{
+					continue;
+				}
+			}
+		}
+
 		const bool bOwnerHasChanged = Component->GetOwner() == InObject;
 
 		if ((!bNoRefreshOwner || !bOwnerHasChanged) && (!Component->bOnlyTrackItself || bOwnerHasChanged))
@@ -1827,6 +1878,75 @@ void FPCGActorAndComponentMapping::ApplyLandscapeChanges(ALandscapeProxy* InLand
 	OnObjectChanged(InLandscape, /*InPreviousData=*/nullptr, InLandscape);
 }
 
+void FPCGActorAndComponentMapping::OnPCGGraphStartsGenerating(UPCGComponent* InComponent)
+{
+	TRACE_CPUPROFILER_EVENT_SCOPE(FPCGActorAndComponentMapping::OnPCGGraphStartsGenerating);
+
+	if (!InComponent || !InComponent->GetOwner() || PCGActorAndComponentMapping::CVarDisablePCGDataInterdependencyOptimization.GetValueOnAnyThread())
+	{
+		return;
+	}
+
+	// When a graph starts generating, look for component that depends on it and keep a count.
+	// When this graph will be done generating, we'll trigger a dependency generation. But if multiple graphs are generated at the same time and 
+	// they all contribute to the same dependency, we'll trigger the dependency only when all the graphs are done generating
+	TSet<UPCGComponent*> TrackedComponents;
+	TSet<FName> RemovedTags;
+
+	const FBox ComponentBounds = InComponent->GetGridBounds();
+
+	for (auto& It : CulledTrackedKeysToComponentsMap)
+	{
+		TSet<UPCGComponent*> TempTrackedComponents;
+		It.Key.IsMatching(InComponent->GetOwner(), RemovedTags, It.Value, &TempTrackedComponents);
+
+		// Removing all components that aren't intersecting with it, since it won't contribute to refresh
+		for (UPCGComponent* TrackedComponent : TempTrackedComponents)
+		{
+			if (ensure(TrackedComponent) && TrackedComponent->GetOwner() && PCGActorAndComponentMapping::GetActorBounds(TrackedComponent->GetOwner()).Intersect(ComponentBounds))
+			{
+				TrackedComponents.Add(TrackedComponent);
+			}
+		}
+	}
+
+	for (auto& It : AlwaysTrackedKeysToComponentsMap)
+	{
+		It.Key.IsMatching(InComponent->GetOwner(), RemovedTags, It.Value, &TrackedComponents);
+	}
+
+	for (UPCGComponent* Component : TrackedComponents)
+	{
+		if (Component == InComponent)
+		{
+			continue;
+		}
+
+		TArray<TObjectKey<UPCGComponent>>& CurrentDependencies = ComponentsToDependencyMap.FindOrAdd(Component);
+		CurrentDependencies.AddUnique(InComponent);
+	}
+}
+
+void FPCGActorAndComponentMapping::OnPCGGraphCancelled(UPCGComponent* InComponent)
+{
+	if (!InComponent)
+	{
+		return;
+	}
+
+	TArray<TObjectKey<UPCGComponent>> Keys;
+	ComponentsToDependencyMap.GetKeys(Keys);
+	for (const TObjectKey<UPCGComponent>& Key : Keys)
+	{
+		TArray<TObjectKey<UPCGComponent>>& Value = ComponentsToDependencyMap[Key];
+		Value.Remove(InComponent);
+		if (Value.IsEmpty())
+		{
+			ComponentsToDependencyMap.Remove(Key);
+		}
+	}
+}
+
 void FPCGActorAndComponentMapping::OnPCGGraphGeneratedOrCleaned(UPCGComponent* InComponent)
 {
 	if (!InComponent || !InComponent->GetOwner())
diff --git a/Engine/Plugins/Experimental/PCG/Source/PCG/Private/PCGComponent.cpp b/Engine/Plugins/Experimental/PCG/Source/PCG/Private/PCGComponent.cpp
--- a/Engine/Plugins/Experimental/PCG/Source/PCG/Private/PCGComponent.cpp
+++ b/Engine/Plugins/Experimental/PCG/Source/PCG/Private/PCGComponent.cpp
@@ -329,6 +329,13 @@ FPCGTaskId UPCGComponent::GenerateInternal(bool bForce, EPCGHiGenGrid Grid, EPCG
 
 	CurrentGenerationTask = GetSubsystem()->ScheduleComponent(this, Grid, /*bSave=*/bForce, Dependencies);
 
+#if WITH_EDITOR
+	if (CurrentGenerationTask != InvalidPCGTaskId)
+	{
+		OnPCGGraphStartGeneratingDelegate.Broadcast(this);
+	}
+#endif // WITH_EDITOR
+
 	return CurrentGenerationTask;
 }
 
@@ -397,13 +404,6 @@ void UPCGComponent::PostProcessGraph(const FBox& InNewBounds, bool bInGenerated,
 
 		CurrentGenerationTask = InvalidPCGTaskId;
 
-#if WITH_EDITOR
-		// Reset this flag to avoid re-generating on further refreshes.
-		bForceGenerateOnBPAddedToWorld = false;
-
-		bDirtyGenerated = false;
-		OnPCGGraphGeneratedDelegate.Broadcast(this);
-#endif
 		// After a successful generation, we also want to call PostGenerateFunctions
 		// if we have any. We also need a context.
 
@@ -415,14 +415,17 @@ void UPCGComponent::PostProcessGraph(const FBox& InNewBounds, bool bInGenerated,
 			//GeneratedGraphOutput = Context->InputData;
 			for (const FPCGTaggedData& TaggedData : Context->InputData.TaggedData)
 			{
-				// TODO: outering the first layer might not be sufficient here - might need to expose
-				// some methods in the data to traverse all the data to outer everything for serialization
-				if (UPCGData* DuplicatedData = TaggedData.Data->DuplicateData())
+				if (ensure(TaggedData.Data))
 				{
-					FPCGTaggedData& DuplicatedTaggedData = GeneratedGraphOutput.TaggedData.Add_GetRef(TaggedData);
-					DuplicatedTaggedData.Data = DuplicatedData;
-					DuplicatedData->Rename(nullptr, this);
-					DuplicatedData->Flatten();
+					// TODO: outering the first layer might not be sufficient here - might need to expose
+					// some methods in the data to traverse all the data to outer everything for serialization
+					if (UPCGData* DuplicatedData = TaggedData.Data->DuplicateData())
+					{
+						FPCGTaggedData& DuplicatedTaggedData = GeneratedGraphOutput.TaggedData.Add_GetRef(TaggedData);
+						DuplicatedTaggedData.Data = DuplicatedData;
+						DuplicatedData->Rename(nullptr, this);
+						DuplicatedData->Flatten();
+					}
 				}
 			}
 
@@ -436,6 +439,12 @@ void UPCGComponent::PostProcessGraph(const FBox& InNewBounds, bool bInGenerated,
 		}
 
 #if WITH_EDITOR
+		// Reset this flag to avoid re-generating on further refreshes.
+		bForceGenerateOnBPAddedToWorld = false;
+
+		bDirtyGenerated = false;
+		OnPCGGraphGeneratedDelegate.Broadcast(this);
+
 		UpdateDynamicTracking();
 #endif // WITH_EDITOR
 	}
@@ -552,6 +561,8 @@ void UPCGComponent::OnProcessGraphAborted(bool bQuiet)
 
 	StopGenerationInProgress();
 
+	OnPCGGraphCancelledDelegate.Broadcast(this);
+
 	GetSubsystem()->OnComponentGenerationCompleteOrCancelled.Broadcast();
 #endif
 }
diff --git a/Engine/Plugins/Experimental/PCG/Source/PCG/Public/PCGActorAndComponentMapping.h b/Engine/Plugins/Experimental/PCG/Source/PCG/Public/PCGActorAndComponentMapping.h
--- a/Engine/Plugins/Experimental/PCG/Source/PCG/Public/PCGActorAndComponentMapping.h
+++ b/Engine/Plugins/Experimental/PCG/Source/PCG/Public/PCGActorAndComponentMapping.h
@@ -153,7 +153,9 @@ private:
 	void OnObjectModified(UObject* InObject);
 	void OnObjectPropertyChanged(UObject* InObject, FPropertyChangedEvent& InEvent);
 	void OnObjectSaved(UObject* InObject, FObjectPreSaveContext InObjectSaveContext);
+	void OnPCGGraphStartsGenerating(UPCGComponent* InComponent);
 	void OnPCGGraphGeneratedOrCleaned(UPCGComponent* InComponent);
+	void OnPCGGraphCancelled(UPCGComponent* InComponent);
 
 	/** Remap the tracking in case of BP components. */
 	void RemapTracking(const UPCGComponent* InOldComponent, UPCGComponent* InNewComponent);
@@ -224,6 +226,9 @@ private:
 	/** Transient map of actors and their previous data, it's set in the pre object change to be able to track changes (such as tags or positions) */
 	TMap<TObjectKey<AActor>, FActorPreviousData> ActorToPreviousDataMap;
 
+	/** Transient map that keep track of all components that depends on another component currently generating, to trigger the refresh only once, when all are done. */
+	TMap<TObjectKey<UPCGComponent>, TArray<TObjectKey<UPCGComponent>>> ComponentsToDependencyMap;
+
 #if WITH_EDITOR
 	// Part for the delayed landscape change update
 	double LastLandscapeDirtyTime = -1.0;
diff --git a/Engine/Plugins/Experimental/PCG/Source/PCG/Public/PCGComponent.h b/Engine/Plugins/Experimental/PCG/Source/PCG/Public/PCGComponent.h
--- a/Engine/Plugins/Experimental/PCG/Source/PCG/Public/PCGComponent.h
+++ b/Engine/Plugins/Experimental/PCG/Source/PCG/Public/PCGComponent.h
@@ -35,6 +35,8 @@ class UClass;
 struct FPCGContext;
 
 #if WITH_EDITOR
+	DECLARE_MULTICAST_DELEGATE_OneParam(FOnPCGGraphStartGenerating, UPCGComponent*);
+	DECLARE_MULTICAST_DELEGATE_OneParam(FOnPCGGraphCancelled, UPCGComponent*);
 	DECLARE_MULTICAST_DELEGATE_OneParam(FOnPCGGraphGenerated, UPCGComponent*);
 	DECLARE_MULTICAST_DELEGATE_OneParam(FOnPCGGraphCleaned, UPCGComponent*);
 #endif
@@ -247,7 +249,9 @@ public:
 	bool bForceGenerateOnBPAddedToWorld = false;
 
 	FOnPCGGraphGenerated OnPCGGraphGeneratedDelegate;
+	FOnPCGGraphStartGenerating OnPCGGraphStartGeneratingDelegate;
 	FOnPCGGraphCleaned OnPCGGraphCleanedDelegate;
+	FOnPCGGraphCancelled OnPCGGraphCancelledDelegate;
 #endif
 
 	/** Flag to indicate whether this component has run in the editor. Note that for partitionable actors, this will always be false. */
