[PCG] Async loading context interface

* Add interface to implement for contexts of elements that need async loading
* Will hold a reference on loaded assets until the context is deleted
* On cancellation, loading jobs are cleaned up.

#rnx

diff --git a/Engine/Plugins/Experimental/PCG/Source/PCG/Private/Async/PCGAsyncLoadingContext.cpp b/Engine/Plugins/Experimental/PCG/Source/PCG/Private/Async/PCGAsyncLoadingContext.cpp
new file mode 100644
--- /dev/null
+++ b/Engine/Plugins/Experimental/PCG/Source/PCG/Private/Async/PCGAsyncLoadingContext.cpp
@@ -0,0 +1,58 @@
+// Copyright Epic Games, Inc. All Rights Reserved.
+
+#include "Async/PCGAsyncLoadingContext.h"
+
+#include "PCGContext.h"
+
+#include "Engine/AssetManager.h"
+
+IPCGAsyncLoadingContext::~IPCGAsyncLoadingContext()
+{
+	// We need to make sure to cancel any loading requested.
+	CancelLoading();
+}
+
+void IPCGAsyncLoadingContext::CancelLoading()
+{
+	if (LoadHandle.IsValid() && LoadHandle->IsActive())
+	{
+		LoadHandle->CancelHandle();
+	}
+
+	LoadHandle.Reset();
+}
+
+bool IPCGAsyncLoadingContext::RequestResourceLoad(FPCGContext* ThisContext, TArray<FSoftObjectPath>&& ObjectsToLoad, bool bAsynchronous)
+{
+	if (!ObjectsToLoad.IsEmpty() && !bLoadRequested)
+	{
+		if (!bAsynchronous)
+		{
+			LoadHandle = UAssetManager::GetStreamableManager().RequestSyncLoad(std::forward<TArray<FSoftObjectPath>>(ObjectsToLoad));
+			bLoadRequested = true;
+
+			return true;
+		}
+		else
+		{
+			ThisContext->bIsPaused = true;
+
+			// It is a bit unsafe to pass this to a delegate lambda. But if the context dies before the completion of the loading, the context will cancel the loading in its dtor.
+			LoadHandle = UAssetManager::GetStreamableManager().RequestAsyncLoad(std::forward<TArray<FSoftObjectPath>>(ObjectsToLoad), [ThisContext]() { ThisContext->bIsPaused = false; });
+			bLoadRequested = true;
+
+			// If the load handle is not active it means objects were invalid
+			if (!LoadHandle->IsActive())
+			{
+				ThisContext->bIsPaused = false;
+				return true;
+			}
+			else
+			{
+				return false;
+			}
+		}
+	}
+
+	return true;
+}
\ No newline at end of file
diff --git a/Engine/Plugins/Experimental/PCG/Source/PCG/Private/Elements/PCGGetPropertyFromObjectPath.cpp b/Engine/Plugins/Experimental/PCG/Source/PCG/Private/Elements/PCGGetPropertyFromObjectPath.cpp
--- a/Engine/Plugins/Experimental/PCG/Source/PCG/Private/Elements/PCGGetPropertyFromObjectPath.cpp
+++ b/Engine/Plugins/Experimental/PCG/Source/PCG/Private/Elements/PCGGetPropertyFromObjectPath.cpp
@@ -10,8 +10,6 @@
 #include "Metadata/Accessors/PCGAttributeAccessorKeys.h"
 #include "Metadata/Accessors/PCGAttributeAccessorHelpers.h"
 
-#include "Engine/AssetManager.h"
-
 #define LOCTEXT_NAMESPACE "PCGGetPropertyFromObjectPathElement"
 
 #if WITH_EDITOR
@@ -141,10 +139,11 @@ bool FPCGGetPropertyFromObjectPathElement::PrepareDataInternal(FPCGContext* Cont
 	check(Settings);
 
 	FPCGGetPropertyFromObjectPathContext* ThisContext = static_cast<FPCGGetPropertyFromObjectPathContext*>(Context);
+	IPCGAsyncLoadingContext* AsyncLoadContext = static_cast<IPCGAsyncLoadingContext*>(ThisContext);
 
 	bool bIsDone = true;
 
-	if (!ThisContext->bRequestSent)
+	if (!AsyncLoadContext->WasLoadRequested())
 	{
 		const UPCGNode* Node = Context->Node;
 		const bool InPinIsConnected = Node ? Node->IsInputPinConnected(PCGPinConstants::DefaultInputLabel) : false;
@@ -252,19 +251,8 @@ bool FPCGGetPropertyFromObjectPathElement::PrepareDataInternal(FPCGContext* Cont
 
 		if (!ObjectsToLoad.IsEmpty())
 		{
-			if (Settings->bSynchronousLoad)
-			{
-				ThisContext->LoadHandle = UAssetManager::GetStreamableManager().RequestSyncLoad(std::move(ObjectsToLoad));
-			}
-			else
-			{
-				bIsDone = false;
-				ThisContext->bIsPaused = true;
-				ThisContext->LoadHandle = UAssetManager::GetStreamableManager().RequestAsyncLoad(std::move(ObjectsToLoad), [Context]() { Context->bIsPaused = false; });
-			}
+			bIsDone = AsyncLoadContext->RequestResourceLoad(ThisContext, std::move(ObjectsToLoad), !Settings->bSynchronousLoad);
 		}
-
-		ThisContext->bRequestSent = true;
 	}
 
 	return bIsDone;
diff --git a/Engine/Plugins/Experimental/PCG/Source/PCG/Public/Async/PCGAsyncLoadingContext.h b/Engine/Plugins/Experimental/PCG/Source/PCG/Public/Async/PCGAsyncLoadingContext.h
new file mode 100644
--- /dev/null
+++ b/Engine/Plugins/Experimental/PCG/Source/PCG/Public/Async/PCGAsyncLoadingContext.h
@@ -0,0 +1,37 @@
+// Copyright Epic Games, Inc. All Rights Reserved.
+
+#pragma once
+
+#include "Containers/Array.h"
+#include "Engine/StreamableManager.h"
+#include "UObject/SoftObjectPath.h"
+
+struct FPCGContext;
+
+/**
+* Extension interface to add Async Loading support to any Context.
+* Just inherit IPCGAsyncLoadingContext on the element context.
+*/
+struct IPCGAsyncLoadingContext
+{
+public:
+	virtual ~IPCGAsyncLoadingContext();
+
+	/** Request a load. If load was already requested, do nothing. LoadHandle will be set in the context, meaning that assets will stay alive while context is loaded.
+	* Request can be synchronous or asynchronous. If loading is asynchronous, the current task is paused and will be woken up when the loading is done.
+	* WARNING: Make sure to call this function with soft paths that are NOT null.
+	* Returns true if the execution can continue (objects are loaded or invalid), or false if we need to wait for loading
+	*/
+	bool RequestResourceLoad(FPCGContext* ThisContext, TArray<FSoftObjectPath>&& ObjectsToLoad, bool bAsynchronous = true);
+
+	void CancelLoading();
+
+	bool WasLoadRequested() const { return bLoadRequested; }
+
+private:
+	/** If the load was already requested */
+	bool bLoadRequested = false;
+
+	/** Handle holder for any loaded resources */
+	TSharedPtr<FStreamableHandle> LoadHandle;
+};
\ No newline at end of file
diff --git a/Engine/Plugins/Experimental/PCG/Source/PCG/Public/Elements/PCGGetPropertyFromObjectPath.h b/Engine/Plugins/Experimental/PCG/Source/PCG/Public/Elements/PCGGetPropertyFromObjectPath.h
--- a/Engine/Plugins/Experimental/PCG/Source/PCG/Public/Elements/PCGGetPropertyFromObjectPath.h
+++ b/Engine/Plugins/Experimental/PCG/Source/PCG/Public/Elements/PCGGetPropertyFromObjectPath.h
@@ -4,8 +4,7 @@
 
 #include "PCGContext.h"
 #include "PCGSettings.h"
-
-#include "Engine/StreamableManager.h"
+#include "Async/PCGAsyncLoadingContext.h"
 
 #include "PCGGetPropertyFromObjectPath.generated.h"
 
@@ -73,11 +72,9 @@ public:
 	bool bSilenceErrorOnEmptyObjectPath = false;
 };
 
-struct FPCGGetPropertyFromObjectPathContext : public FPCGContext
+struct FPCGGetPropertyFromObjectPathContext : public FPCGContext, public IPCGAsyncLoadingContext
 {
-	TSharedPtr<FStreamableHandle> LoadHandle;
 	TArray<TTuple<FSoftObjectPath, int32>> PathsToObjectsToExtractAndIncomingDataIndex;
-	bool bRequestSent = false;
 };
 
 class FPCGGetPropertyFromObjectPathElement : public IPCGElement
