[PCG] Move the plugin out of experimental

#virtualized

diff --git a/Engine/Plugins/PCG/Source/PCG/Public/Tests/Determinism/PCGDifferenceDeterminismTest.h b/Engine/Plugins/PCG/Source/PCG/Public/Tests/Determinism/PCGDifferenceDeterminismTest.h
new file mode 100644
--- /dev/null
+++ b/Engine/Plugins/PCG/Source/PCG/Public/Tests/Determinism/PCGDifferenceDeterminismTest.h
@@ -0,0 +1,19 @@
+// Copyright Epic Games, Inc. All Rights Reserved.
+#pragma once
+#if WITH_EDITOR
+
+#include "HAL/Platform.h"
+
+namespace PCGDeterminismTests
+{
+	namespace DifferenceElement
+	{
+		PCG_API bool RunTestSuite();
+	}
+}
+
+#endif
+
+#if UE_ENABLE_INCLUDE_ORDER_DEPRECATED_IN_5_2
+#include "CoreMinimal.h"
+#endif
