[PCG] Remove child component iteration (since it is not used elsewhere in the engine) and can provide wrong results

#rnx

diff --git a/Engine/Plugins/PCG/Source/PCG/Private/Helpers/PCGHelpers.cpp b/Engine/Plugins/PCG/Source/PCG/Private/Helpers/PCGHelpers.cpp
--- a/Engine/Plugins/PCG/Source/PCG/Private/Helpers/PCGHelpers.cpp
+++ b/Engine/Plugins/PCG/Source/PCG/Private/Helpers/PCGHelpers.cpp
@@ -111,7 +111,8 @@ namespace PCGHelpers
 		FBox Box(EForceInit::ForceInit);
 
 		const bool bNonColliding = true;
-		const bool bIncludeFromChildActors = true;
+		// FIXME: This option was disabled because of wrong bounds with components attached to other components. (UE-205314)
+		const bool bIncludeFromChildActors = false;
 
 		if (InActor)
 		{
@@ -151,7 +152,8 @@ namespace PCGHelpers
 			else
 			{
 				const bool bNonColliding = true;
-				const bool bIncludeFromChildActors = true;
+				// FIXME: This option was disabled because of wrong bounds with components attached to other components. (UE-205314)
+				const bool bIncludeFromChildActors = false;
 
 				const FTransform& ActorToWorld = InActor->GetTransform();
 				const FTransform WorldToActor = ActorToWorld.Inverse();
