[PCG] Expose a few functions in Metadata to allow Attribute Set Multi Entry manipulation/iteration + fix partial copy of metadata

#rnx

diff --git a/Engine/Plugins/PCG/Source/PCG/Private/Metadata/PCGMetadata.cpp b/Engine/Plugins/PCG/Source/PCG/Private/Metadata/PCGMetadata.cpp
--- a/Engine/Plugins/PCG/Source/PCG/Private/Metadata/PCGMetadata.cpp
+++ b/Engine/Plugins/PCG/Source/PCG/Private/Metadata/PCGMetadata.cpp
@@ -129,13 +129,10 @@ void UPCGMetadata::InitializeAsCopyWithAttributeFilter(const UPCGMetadata* InMet
 		return;
 	}
 
-	Parent = InMetadataToCopy->Parent;
-	OtherParents = InMetadataToCopy->OtherParents;
-	ParentKeys = InMetadataToCopy->ParentKeys;
-	ItemKeyOffset = InMetadataToCopy->ItemKeyOffset;
-
 	const bool bSkipAttributesInFilterList = (InFilterMode == EPCGMetadataFilterMode::ExcludeAttributes);
 
+	// If we have a partial copy, it will flatten the metadata, so we don't need a parent.
+	// Otherwise, we keep the parent hierarchy.
 	const bool bPartialCopy = EntriesToCopy && EntriesToCopy->Num() <= InMetadataToCopy->GetItemCountForChild();
 	TArray<PCGMetadataEntryKey> NewEntryKeys;
 	TArray<PCGMetadataValueKey> NewValueKeys;
@@ -144,10 +141,21 @@ void UPCGMetadata::InitializeAsCopyWithAttributeFilter(const UPCGMetadata* InMet
 		const int32 Count = EntriesToCopy->Num();
 		NewEntryKeys.SetNumUninitialized(Count);
 		NewValueKeys.SetNumUninitialized(Count);
+		ParentKeys.SetNumUninitialized(Count);
 		for (int32 j = 0; j < Count; ++j)
 		{
 			NewEntryKeys[j] = PCGMetadataEntryKey(j);
+			ParentKeys[j] = -1;
 		}
+
+		ItemKeyOffset = 0;
+	}
+	else
+	{
+		ParentKeys = InMetadataToCopy->ParentKeys;
+		ItemKeyOffset = InMetadataToCopy->ItemKeyOffset;
+		Parent = InMetadataToCopy->Parent;
+		OtherParents = InMetadataToCopy->OtherParents;
 	}
 
 	// Copy attributes
diff --git a/Engine/Plugins/PCG/Source/PCG/Public/Metadata/PCGMetadata.h b/Engine/Plugins/PCG/Source/PCG/Public/Metadata/PCGMetadata.h
--- a/Engine/Plugins/PCG/Source/PCG/Public/Metadata/PCGMetadata.h
+++ b/Engine/Plugins/PCG/Source/PCG/Public/Metadata/PCGMetadata.h
@@ -189,6 +189,8 @@ public:
 	UFUNCTION(BlueprintCallable, Category = "PCG|Metadata")
 	bool HasCommonAttributes(const UPCGMetadata* InMetadata) const;
 
+	/** Return the number of attributes in this metadata. */
+	UFUNCTION(BlueprintCallable, Category = "PCG|Metadata")
 	int32 GetAttributeCount() const;
 
 	template <typename T>
@@ -291,9 +293,12 @@ public:
 	void ComputeWeightedAttribute(PCGMetadataEntryKey& OutKey, const TArrayView<TPair<PCGMetadataEntryKey, float>>& InWeightedKeys, const UPCGMetadata* InMetadata);
 
 	int64 GetItemKeyCountForParent() const;
-	int64 GetItemCountForChild() const;
 	int64 GetLocalItemCount() const;
 
+	/** Return the number of entries in metadata including the parent entries. */
+	UFUNCTION(BlueprintCallable, Category = "PCG|Metadata", meta = (DisplayName = "Get Number of Entries"))
+	int64 GetItemCountForChild() const;
+
 	/**
 	* Create a new attribute. If the attribute already exists, it will raise a warning (use FindOrCreateAttribute if this usecase can arise)
 	* If the attribute already exists but is of the wrong type, it will fail and return nullptr. Same if the name is invalid.
