[PCG] Update Static Mesh Spawner to support by default Async loading

* Static Mesh Spawner context is now extended by the AsyncLoadingContext
* Meshes could be synchronously loaded in the Mesh Selectors to apply bounds to points.
   This logic was moved inside the Execute of the MeshSpawner.
* Edit InstanceData packer to support Vec2 and use a switch that might be faster than series of if else

#rnx

diff --git a/Engine/Plugins/PCG/Source/PCG/Private/Elements/PCGStaticMeshSpawner.cpp b/Engine/Plugins/PCG/Source/PCG/Private/Elements/PCGStaticMeshSpawner.cpp
--- a/Engine/Plugins/PCG/Source/PCG/Private/Elements/PCGStaticMeshSpawner.cpp
+++ b/Engine/Plugins/PCG/Source/PCG/Private/Elements/PCGStaticMeshSpawner.cpp
@@ -254,6 +254,33 @@ bool FPCGStaticMeshSpawnerElement::PrepareDataInternal(FPCGContext* InContext) c
 		}
 	}
 
+	IPCGAsyncLoadingContext* AsyncLoadingContext = static_cast<IPCGAsyncLoadingContext*>(Context);
+
+	if (Context->CurrentInputIndex == Inputs.Num() && !AsyncLoadingContext->WasLoadRequested() && !Context->MeshInstancesData.IsEmpty() && !Settings->bSynchronousLoad)
+	{
+		TArray<FSoftObjectPath> ObjectsToLoad;
+		for (const FPCGStaticMeshSpawnerContext::FPackedInstanceListData& InstanceData : Context->MeshInstancesData)
+		{
+			for (const FPCGMeshInstanceList& MeshInstanceList : InstanceData.MeshInstances)
+			{
+				if (!MeshInstanceList.Descriptor.StaticMesh.IsNull())
+				{
+					ObjectsToLoad.AddUnique(MeshInstanceList.Descriptor.StaticMesh.ToSoftObjectPath());
+				}
+
+				for (const TSoftObjectPtr<UMaterialInterface>& OverrideMaterial : MeshInstanceList.Descriptor.OverrideMaterials)
+				{
+					if (!OverrideMaterial.IsNull())
+					{
+						ObjectsToLoad.AddUnique(OverrideMaterial.ToSoftObjectPath());
+					}
+				}
+			}
+		}
+
+		return AsyncLoadingContext->RequestResourceLoad(Context, std::move(ObjectsToLoad), /*bAsynchronous=*/true);
+	}
+
 	return true;
 }
 
@@ -275,7 +302,29 @@ bool FPCGStaticMeshSpawnerElement::ExecuteInternal(FPCGContext* InContext) const
 		{
 			while (Context->CurrentDataIndex < InstanceList.MeshInstances.Num())
 			{
-				SpawnStaticMeshInstances(Context, InstanceList.MeshInstances[Context->CurrentDataIndex], InstanceList.TargetActor, InstanceList.PackedCustomData[Context->CurrentDataIndex]);
+				const FPCGMeshInstanceList& MeshInstance = InstanceList.MeshInstances[Context->CurrentDataIndex];
+				SpawnStaticMeshInstances(Context, MeshInstance, InstanceList.TargetActor, InstanceList.PackedCustomData[Context->CurrentDataIndex]);
+
+				// Now that the mesh is loaded/spawned, set the bounds to out points if requested.
+				if (MeshInstance.Descriptor.StaticMesh && Settings->bApplyMeshBoundsToPoints)
+				{
+					if (TMap<UPCGPointData*, TArray<int32>>* OutPointDataToPointIndex = Context->MeshToOutPoints.Find(MeshInstance.Descriptor.StaticMesh))
+					{
+						const FBox Bounds = MeshInstance.Descriptor.StaticMesh->GetBoundingBox();
+						for (TPair<UPCGPointData*, TArray<int32>>& It : *OutPointDataToPointIndex)
+						{
+							check(It.Key);
+							TArray<FPCGPoint>& OutPoints = It.Key->GetMutablePoints();
+							for (int32 Index : It.Value)
+							{
+								FPCGPoint& Point = OutPoints[Index];
+								Point.BoundsMin = Bounds.Min;
+								Point.BoundsMax = Bounds.Max;
+							}
+						}
+					}
+				}
+
 				++Context->CurrentDataIndex;
 
 				if (Context->AsyncState.ShouldStop())
@@ -320,7 +369,7 @@ bool FPCGStaticMeshSpawnerElement::CanExecuteOnlyOnMainThread(FPCGContext* Conte
 	return Context->CurrentPhase == EPCGExecutionPhase::Execute || Context->CurrentPhase == EPCGExecutionPhase::PrepareData;
 }
 
-void FPCGStaticMeshSpawnerElement::SpawnStaticMeshInstances(FPCGContext* Context, const FPCGMeshInstanceList& InstanceList, AActor* TargetActor, const FPCGPackedCustomData& PackedCustomData) const
+void FPCGStaticMeshSpawnerElement::SpawnStaticMeshInstances(FPCGStaticMeshSpawnerContext* Context, const FPCGMeshInstanceList& InstanceList, AActor* TargetActor, const FPCGPackedCustomData& PackedCustomData) const
 {
 	// Populate the (H)ISM from the previously prepared entries
 	TRACE_CPUPROFILER_EVENT_SCOPE(FPCGStaticMeshSpawnerElement::Execute::PopulateISMs);
@@ -330,7 +379,7 @@ void FPCGStaticMeshSpawnerElement::SpawnStaticMeshInstances(FPCGContext* Context
 		return;
 	}
 
-	// TODO: we could likely pre-load these meshes asynchronously in the settings
+	// Will be synchronously loaded if not loaded. But by default it should already have been loaded asynchronously in PrepareData, so this is free.
 	UStaticMesh* LoadedMesh = InstanceList.Descriptor.StaticMesh.LoadSynchronous();
 
 	if (!LoadedMesh)
@@ -344,9 +393,9 @@ void FPCGStaticMeshSpawnerElement::SpawnStaticMeshInstances(FPCGContext* Context
 		return;
 	}
 
-	// TODO: we could likely pre-load these materials asynchronously in the settings
 	for (TSoftObjectPtr<UMaterialInterface> OverrideMaterial : InstanceList.Descriptor.OverrideMaterials)
 	{
+		// Will be synchronously loaded if not loaded. But by default it should already have been loaded asynchronously in PrepareData, so this is free.
 		if (OverrideMaterial.IsValid() && !OverrideMaterial.LoadSynchronous())
 		{
 			PCGE_LOG(Error, GraphAndLog, FText::Format(LOCTEXT("OverrideMaterialLoadFailed", "Unable to load override material '{0}'"), FText::FromString(OverrideMaterial.ToString())));
@@ -366,6 +415,7 @@ void FPCGStaticMeshSpawnerElement::SpawnStaticMeshInstances(FPCGContext* Context
 
 	const UPCGStaticMeshSpawnerSettings* Settings = Context->GetInputSettings<UPCGStaticMeshSpawnerSettings>();
 	check(Settings);
+
 	UPCGManagedISMComponent* MISMC = UPCGActorHelpers::GetOrCreateManagedISMC(TargetActor, Context->SourceComponent.Get(), Settings->UID, Params);
 	
 	check(MISMC);
diff --git a/Engine/Plugins/PCG/Source/PCG/Private/InstanceDataPackers/PCGInstanceDataPackerBase.cpp b/Engine/Plugins/PCG/Source/PCG/Private/InstanceDataPackers/PCGInstanceDataPackerBase.cpp
--- a/Engine/Plugins/PCG/Source/PCG/Private/InstanceDataPackers/PCGInstanceDataPackerBase.cpp
+++ b/Engine/Plugins/PCG/Source/PCG/Private/InstanceDataPackers/PCGInstanceDataPackerBase.cpp
@@ -14,36 +14,25 @@ void UPCGInstanceDataPackerBase::PackInstances_Implementation(FPCGContext& Conte
 
 bool UPCGInstanceDataPackerBase::AddTypeToPacking(int TypeId, FPCGPackedCustomData& OutPackedCustomData) const
 {
-	if (TypeId == PCG::Private::MetadataTypes<float>::Id)
+	switch (TypeId)
 	{
+	case PCG::Private::MetadataTypes<float>::Id: // fall-through
+	case PCG::Private::MetadataTypes<double>::Id: // fall-through
+	case PCG::Private::MetadataTypes<int32>::Id: // fall-through
+	case PCG::Private::MetadataTypes<int64>::Id:
 		OutPackedCustomData.NumCustomDataFloats += 1;
-	}
-	else if (TypeId == PCG::Private::MetadataTypes<double>::Id)
-	{
-		OutPackedCustomData.NumCustomDataFloats += 1;
-	}
-	else if (TypeId == PCG::Private::MetadataTypes<int32>::Id)
-	{
-		OutPackedCustomData.NumCustomDataFloats += 1;
-	}
-	else if (TypeId == PCG::Private::MetadataTypes<int64>::Id)
-	{
-		OutPackedCustomData.NumCustomDataFloats += 1;
-	}
-	else if (TypeId == PCG::Private::MetadataTypes<FRotator>::Id)
-	{
-		OutPackedCustomData.NumCustomDataFloats += 3;
-	}
-	else if (TypeId == PCG::Private::MetadataTypes<FVector>::Id)
-	{
+		break;
+	case PCG::Private::MetadataTypes<FVector2D>::Id:
+		OutPackedCustomData.NumCustomDataFloats += 2;
+		break;
+	case PCG::Private::MetadataTypes<FRotator>::Id: // fall-through
+	case PCG::Private::MetadataTypes<FVector>::Id:
 		OutPackedCustomData.NumCustomDataFloats += 3;
-	}
-	else if (TypeId == PCG::Private::MetadataTypes<FVector4>::Id)
-	{
+		break;
+	case PCG::Private::MetadataTypes<FVector4>::Id:
 		OutPackedCustomData.NumCustomDataFloats += 4;
-	}
-	else
-	{
+		break;
+	default:
 		return false;
 	}
 
@@ -75,39 +64,45 @@ void UPCGInstanceDataPackerBase::PackCustomDataFromAttributes(const FPCGMeshInst
 		{
 			check(AttributeBase);
 
-			if (AttributeBase->GetTypeId() == PCG::Private::MetadataTypes<float>::Id)
+			switch (AttributeBase->GetTypeId())
+			{
+			case PCG::Private::MetadataTypes<float>::Id:
 			{
 				const FPCGMetadataAttribute<float>* Attribute = static_cast<const FPCGMetadataAttribute<float>*>(AttributeBase);
 				check(Attribute);
 
 				const float Value = Attribute->GetValueFromItemKey(PointMetadataEntry);
 				OutPackedCustomData.CustomData.Add(Value);
+				break;
 			}
-			else if (AttributeBase->GetTypeId() == PCG::Private::MetadataTypes<double>::Id)
+			case PCG::Private::MetadataTypes<double>::Id:
 			{
 				const FPCGMetadataAttribute<double>* Attribute = static_cast<const FPCGMetadataAttribute<double>*>(AttributeBase);
 				check(Attribute);
 
 				const double Value = Attribute->GetValueFromItemKey(PointMetadataEntry);
 				OutPackedCustomData.CustomData.Add(Value);
+				break;
 			}
-			else if (AttributeBase->GetTypeId() == PCG::Private::MetadataTypes<int32>::Id)
+			case PCG::Private::MetadataTypes<int32>::Id:
 			{
 				const FPCGMetadataAttribute<int32>* Attribute = static_cast<const FPCGMetadataAttribute<int32>*>(AttributeBase);
 				check(Attribute);
 
 				const float Value = static_cast<float>(Attribute->GetValueFromItemKey(PointMetadataEntry));
 				OutPackedCustomData.CustomData.Add(Value);
+				break;
 			}
-			else if (AttributeBase->GetTypeId() == PCG::Private::MetadataTypes<int64>::Id)
+			case PCG::Private::MetadataTypes<int64>::Id:
 			{
 				const FPCGMetadataAttribute<int64>* Attribute = static_cast<const FPCGMetadataAttribute<int64>*>(AttributeBase);
 				check(Attribute);
 
 				const float Value = static_cast<float>(Attribute->GetValueFromItemKey(PointMetadataEntry));
 				OutPackedCustomData.CustomData.Add(Value);
+				break;
 			}
-			else if (AttributeBase->GetTypeId() == PCG::Private::MetadataTypes<FRotator>::Id)
+			case PCG::Private::MetadataTypes<FRotator>::Id:
 			{
 				const FPCGMetadataAttribute<FRotator>* Attribute = static_cast<const FPCGMetadataAttribute<FRotator>*>(AttributeBase);
 				check(Attribute);
@@ -116,8 +111,19 @@ void UPCGInstanceDataPackerBase::PackCustomDataFromAttributes(const FPCGMeshInst
 				OutPackedCustomData.CustomData.Add(Value.Roll);
 				OutPackedCustomData.CustomData.Add(Value.Pitch);
 				OutPackedCustomData.CustomData.Add(Value.Yaw);
+				break;
 			}
-			else if (AttributeBase->GetTypeId() == PCG::Private::MetadataTypes<FVector>::Id)
+			case PCG::Private::MetadataTypes<FVector2D>::Id:
+			{
+				const FPCGMetadataAttribute<FVector2D>* Attribute = static_cast<const FPCGMetadataAttribute<FVector2D>*>(AttributeBase);
+				check(Attribute);
+
+				const FVector2D Value = Attribute->GetValueFromItemKey(PointMetadataEntry);
+				OutPackedCustomData.CustomData.Add(Value.X);
+				OutPackedCustomData.CustomData.Add(Value.Y);
+				break;
+			}
+			case PCG::Private::MetadataTypes<FVector>::Id:
 			{
 				const FPCGMetadataAttribute<FVector>* Attribute = static_cast<const FPCGMetadataAttribute<FVector>*>(AttributeBase);
 				check(Attribute);
@@ -126,8 +132,9 @@ void UPCGInstanceDataPackerBase::PackCustomDataFromAttributes(const FPCGMeshInst
 				OutPackedCustomData.CustomData.Add(Value.X);
 				OutPackedCustomData.CustomData.Add(Value.Y);
 				OutPackedCustomData.CustomData.Add(Value.Z);
+				break;
 			}
-			else if (AttributeBase->GetTypeId() == PCG::Private::MetadataTypes<FVector4>::Id)
+			case PCG::Private::MetadataTypes<FVector4>::Id:
 			{
 				const FPCGMetadataAttribute<FVector4>* Attribute = static_cast<const FPCGMetadataAttribute<FVector4>*>(AttributeBase);
 				check(Attribute);
@@ -137,6 +144,10 @@ void UPCGInstanceDataPackerBase::PackCustomDataFromAttributes(const FPCGMeshInst
 				OutPackedCustomData.CustomData.Add(Value.Y);
 				OutPackedCustomData.CustomData.Add(Value.Z);
 				OutPackedCustomData.CustomData.Add(Value.W);
+				break;
+			}
+			default:
+				break;
 			}
 		}
 	}
diff --git a/Engine/Plugins/PCG/Source/PCG/Private/MeshSelectors/PCGMeshSelectorByAttribute.cpp b/Engine/Plugins/PCG/Source/PCG/Private/MeshSelectors/PCGMeshSelectorByAttribute.cpp
--- a/Engine/Plugins/PCG/Source/PCG/Private/MeshSelectors/PCGMeshSelectorByAttribute.cpp
+++ b/Engine/Plugins/PCG/Source/PCG/Private/MeshSelectors/PCGMeshSelectorByAttribute.cpp
@@ -180,10 +180,8 @@ bool UPCGMeshSelectorByAttribute::SelectInstances(
 	int32 LastCheckpointIndex = CurrentPointIndex;
 	constexpr int32 TimeSlicingCheckFrequency = 1024;
 	TMap<PCGMetadataValueKey, TSoftObjectPtr<UStaticMesh>>& ValueKeyToMesh = Context.ValueKeyToMesh;
-	TMap<TSoftObjectPtr<UStaticMesh>, FBox>& MeshToBoundingBox = Context.MeshToBoundingBox;
 
 	const TArray<FPCGPoint>& Points = InPointData->GetPoints();
-	TArray<FPCGPoint>* OutPoints = OutPointData ? &OutPointData->GetMutablePoints() : nullptr;
 
 	while (CurrentPointIndex < Points.Num())
 	{
@@ -230,22 +228,11 @@ bool UPCGMeshSelectorByAttribute::SelectInstances(
 		InstanceList.Instances.Emplace(Point.Transform);
 		InstanceList.InstancesMetadataEntry.Emplace(Point.MetadataEntry);
 
-		// Update point bounds to reflect StaticMesh bounds
-		if (OutPoints && Settings->bApplyMeshBoundsToPoints)
+		if (OutPointData && Settings->bApplyMeshBoundsToPoints)
 		{
-			if (!MeshToBoundingBox.Contains(Mesh) && Mesh.LoadSynchronous())
-			{
-				MeshToBoundingBox.Add(Mesh, Mesh->GetBoundingBox());
-			}
-
-			if (MeshToBoundingBox.Contains(Mesh))
-			{
-				const FBox MeshBounds = MeshToBoundingBox[Mesh];
-
-				// CurrentPointIndex - 1, because CurrentPointIndex is incremented at the beginning of the function
-				(*OutPoints)[CurrentPointIndex - 1].BoundsMin = MeshBounds.Min;
-				(*OutPoints)[CurrentPointIndex - 1].BoundsMax = MeshBounds.Max;
-			}
+			TArray<int32>& PointIndices = Context.MeshToOutPoints.FindOrAdd(Mesh).FindOrAdd(OutPointData);
+			// CurrentPointIndex - 1, because CurrentPointIndex is incremented at the beginning of the loop
+			PointIndices.Emplace(CurrentPointIndex - 1);
 		}
 
 		// Check if we should stop here and continue in a subsequent call
diff --git a/Engine/Plugins/PCG/Source/PCG/Private/MeshSelectors/PCGMeshSelectorWeighted.cpp b/Engine/Plugins/PCG/Source/PCG/Private/MeshSelectors/PCGMeshSelectorWeighted.cpp
--- a/Engine/Plugins/PCG/Source/PCG/Private/MeshSelectors/PCGMeshSelectorWeighted.cpp
+++ b/Engine/Plugins/PCG/Source/PCG/Private/MeshSelectors/PCGMeshSelectorWeighted.cpp
@@ -196,7 +196,6 @@ bool UPCGMeshSelectorWeighted::SelectInstances(
 		int32 LastCheckpointIndex = CurrentPointIndex;
 		constexpr int32 TimeSlicingCheckFrequency = 1024;
 		TMap<TSoftObjectPtr<UStaticMesh>, PCGMetadataValueKey>& MeshToValueKey = Context.MeshToValueKey;
-		TMap<TSoftObjectPtr<UStaticMesh>, FBox>& MeshToBoundingBox = Context.MeshToBoundingBox;
 		const int32 TotalWeight = CumulativeWeights.Last();
 
 		const TArray<FPCGPoint>& Points = InPointData->GetPoints();
@@ -225,23 +224,6 @@ bool UPCGMeshSelectorWeighted::SelectInstances(
 				if (OutPointData && OutAttribute)
 				{
 					FPCGPoint& OutPoint = OutPoints->Add_GetRef(Point);
-
-					// Update point bounds to reflect StaticMesh bounds
-					if (Settings->bApplyMeshBoundsToPoints)
-					{
-						if (!MeshToBoundingBox.Contains(Mesh) && Mesh.LoadSynchronous())
-						{
-							MeshToBoundingBox.Add(Mesh, Mesh->GetBoundingBox());
-						}
-
-						if (MeshToBoundingBox.Contains(Mesh))
-						{
-							const FBox MeshBounds = MeshToBoundingBox[Mesh];
-							OutPoint.BoundsMin = MeshBounds.Min;
-							OutPoint.BoundsMax = MeshBounds.Max;
-						}
-					}
-
 					PCGMetadataValueKey* OutValueKey = MeshToValueKey.Find(Mesh);
 					if (!OutValueKey)
 					{
@@ -253,6 +235,13 @@ bool UPCGMeshSelectorWeighted::SelectInstances(
 
 					OutPointData->Metadata->InitializeOnSet(OutPoint.MetadataEntry);
 					OutAttribute->SetValueFromValueKey(OutPoint.MetadataEntry, *OutValueKey);
+
+					if (Settings->bApplyMeshBoundsToPoints)
+					{
+						TArray<int32>& PointIndices = Context.MeshToOutPoints.FindOrAdd(Mesh).FindOrAdd(OutPointData);
+						// CurrentPointIndex - 1, because CurrentPointIndex is incremented at the beginning of the loop
+						PointIndices.Emplace(CurrentPointIndex - 1);
+					}
 				}
 			}
 
diff --git a/Engine/Plugins/PCG/Source/PCG/Private/MeshSelectors/PCGMeshSelectorWeightedByCategory.cpp b/Engine/Plugins/PCG/Source/PCG/Private/MeshSelectors/PCGMeshSelectorWeightedByCategory.cpp
--- a/Engine/Plugins/PCG/Source/PCG/Private/MeshSelectors/PCGMeshSelectorWeightedByCategory.cpp
+++ b/Engine/Plugins/PCG/Source/PCG/Private/MeshSelectors/PCGMeshSelectorWeightedByCategory.cpp
@@ -183,7 +183,6 @@ bool UPCGMeshSelectorWeightedByCategory::SelectInstances(
 	int32 LastCheckpointIndex = CurrentPointIndex;
 	constexpr int32 TimeSlicingCheckFrequency = 1024;
 	TMap<TSoftObjectPtr<UStaticMesh>, PCGMetadataValueKey>& MeshToValueKey = Context.MeshToValueKey;
-	TMap<TSoftObjectPtr<UStaticMesh>, FBox>& MeshToBoundingBox = Context.MeshToBoundingBox;
 
 	// Assign points to entries
 	const TArray<FPCGPoint>& Points = InPointData->GetPoints();
@@ -232,22 +231,6 @@ bool UPCGMeshSelectorWeightedByCategory::SelectInstances(
 			{
 				FPCGPoint& OutPoint = OutPoints->Add_GetRef(Point);
 
-				// Update point bounds to reflect StaticMesh bounds
-				if (Settings->bApplyMeshBoundsToPoints)
-				{
-					if (!MeshToBoundingBox.Contains(Mesh) && Mesh.LoadSynchronous())
-					{
-						MeshToBoundingBox.Add(Mesh, Mesh->GetBoundingBox());
-					}
-
-					if (MeshToBoundingBox.Contains(Mesh))
-					{
-						const FBox MeshBounds = MeshToBoundingBox[Mesh];
-						OutPoint.BoundsMin = MeshBounds.Min;
-						OutPoint.BoundsMax = MeshBounds.Max;
-					}
-				}
-
 				PCGMetadataValueKey* OutValueKey = MeshToValueKey.Find(Mesh);
 				if (!OutValueKey)
 				{
@@ -259,6 +242,13 @@ bool UPCGMeshSelectorWeightedByCategory::SelectInstances(
 
 				OutPointData->Metadata->InitializeOnSet(OutPoint.MetadataEntry);
 				OutAttribute->SetValueFromValueKey(OutPoint.MetadataEntry, *OutValueKey);
+
+				if (Settings->bApplyMeshBoundsToPoints)
+				{
+					TArray<int32>& PointIndices = Context.MeshToOutPoints.FindOrAdd(Mesh).FindOrAdd(OutPointData);
+					// CurrentPointIndex - 1, because CurrentPointIndex is incremented at the beginning of the loop
+					PointIndices.Emplace(CurrentPointIndex - 1);
+				}
 			}
 		}
 
diff --git a/Engine/Plugins/PCG/Source/PCG/Public/Elements/PCGStaticMeshSpawner.h b/Engine/Plugins/PCG/Source/PCG/Public/Elements/PCGStaticMeshSpawner.h
--- a/Engine/Plugins/PCG/Source/PCG/Public/Elements/PCGStaticMeshSpawner.h
+++ b/Engine/Plugins/PCG/Source/PCG/Public/Elements/PCGStaticMeshSpawner.h
@@ -14,6 +14,7 @@ class UPCGSpatialData;
 struct FPCGContext;
 struct FPCGMeshInstanceList;
 struct FPCGPackedCustomData;
+struct FPCGStaticMeshSpawnerContext;
 
 class UStaticMesh;
 
@@ -105,6 +106,10 @@ public:
 	UPROPERTY(meta=(DeprecatedProperty, DeprecationMessage="Use MeshSelectorType and MeshSelectorParameters instead."))
 	TArray<FPCGStaticMeshSpawnerEntry> Meshes_DEPRECATED;
 
+	/** Meshes/Materials will be synchronously loaded before spawning instead of asynchronously. */
+	UPROPERTY(BlueprintReadWrite, EditAnywhere, Category = "Settings|Debug")
+	bool bSynchronousLoad = false;
+
 protected:
 	void RefreshMeshSelector();
 	void RefreshInstancePacker();
@@ -120,7 +125,7 @@ protected:
 	virtual FPCGContext* CreateContext() override;
 	virtual bool PrepareDataInternal(FPCGContext* Context) const override;
 	virtual bool ExecuteInternal(FPCGContext* Context) const override;	
-	void SpawnStaticMeshInstances(FPCGContext* Context, const FPCGMeshInstanceList& InstanceList, AActor* TargetActor, const FPCGPackedCustomData& PackedCustomData) const;
+	void SpawnStaticMeshInstances(FPCGStaticMeshSpawnerContext* Context, const FPCGMeshInstanceList& InstanceList, AActor* TargetActor, const FPCGPackedCustomData& PackedCustomData) const;
 };
 
 
diff --git a/Engine/Plugins/PCG/Source/PCG/Public/Elements/PCGStaticMeshSpawnerContext.h b/Engine/Plugins/PCG/Source/PCG/Public/Elements/PCGStaticMeshSpawnerContext.h
--- a/Engine/Plugins/PCG/Source/PCG/Public/Elements/PCGStaticMeshSpawnerContext.h
+++ b/Engine/Plugins/PCG/Source/PCG/Public/Elements/PCGStaticMeshSpawnerContext.h
@@ -3,6 +3,7 @@
 #pragma once
 
 #include "PCGContext.h"
+#include "Async/PCGAsyncLoadingContext.h"
 #include "MeshSelectors/PCGMeshSelectorBase.h"
 #include "InstanceDataPackers/PCGInstanceDataPackerBase.h"
 
@@ -18,7 +19,7 @@ struct FPCGInstancesAndWeights
 };
 
 USTRUCT(BlueprintType)
-struct FPCGStaticMeshSpawnerContext : public FPCGContext
+struct FPCGStaticMeshSpawnerContext : public FPCGContext, public IPCGAsyncLoadingContext
 {
 	GENERATED_BODY()
 
@@ -51,8 +52,10 @@ struct FPCGStaticMeshSpawnerContext : public FPCGContext
 	FPCGMeshMaterialOverrideHelper MaterialOverrideHelper;
 	int32 CurrentPointIndex = 0;
 
-	// Used in all selectors
-	TMap<TSoftObjectPtr<UStaticMesh>, FBox> MeshToBoundingBox;
+	// Used in all selectors if we have to change the out points bounds by the mesh bounds. Will be empty otherwise.
+	// We need to keep all point indices that will spawn this mesh, in all output point data.
+
+	TMap<TSoftObjectPtr<UStaticMesh>, TMap<UPCGPointData*, TArray<int32>>> MeshToOutPoints;
 
 	// Used in by-attribute selector
 	TMap<PCGMetadataValueKey, TSoftObjectPtr<UStaticMesh>> ValueKeyToMesh;
