[PCG] Update few nodes to use async loading

* Data Table Row To Param Data
* Point From Mesh
* Texture Sampler
* Load Data Table
* Mesh Sampler

#rnx

diff --git a/Engine/Plugins/PCG/Source/PCG/Private/Elements/IO/PCGDataTableElement.cpp b/Engine/Plugins/PCG/Source/PCG/Private/Elements/IO/PCGDataTableElement.cpp
--- a/Engine/Plugins/PCG/Source/PCG/Private/Elements/IO/PCGDataTableElement.cpp
+++ b/Engine/Plugins/PCG/Source/PCG/Private/Elements/IO/PCGDataTableElement.cpp
@@ -84,8 +84,22 @@ bool FPCGLoadDataTableElement::PrepareLoad(FPCGExternalDataContext* Context) con
 	const UPCGLoadDataTableSettings* Settings = Context->GetInputSettings<UPCGLoadDataTableSettings>();
 	check(Settings);
 
-	// 1. Load data table - if we do, we should root it (?)
-	UDataTable* DataTable = Settings->DataTable.LoadSynchronous();
+	// Nothing to do if we have no data table.
+	if (Settings->DataTable.IsNull())
+	{
+		return true;
+	}
+
+	// 1. Request load data table. Return false if we need to wait, otherwise continue.
+	if (!Context->WasLoadRequested())
+	{
+		if (!Context->RequestResourceLoad(Context, { Settings->DataTable.ToSoftObjectPath() }, !Settings->bSynchronousLoad))
+		{
+			return false;
+		}
+	}
+
+	UDataTable* DataTable = Settings->DataTable.Get();
 	if (!DataTable)
 	{
 		PCGE_LOG(Error, GraphAndLog, FText::Format(LOCTEXT("DataTableNotLoaded", "Provided Data Table ('{0}') could not be loaded."), FText::FromString(Settings->DataTable.ToString())));
diff --git a/Engine/Plugins/PCG/Source/PCG/Private/Elements/PCGDataTableRowToParamData.cpp b/Engine/Plugins/PCG/Source/PCG/Private/Elements/PCGDataTableRowToParamData.cpp
--- a/Engine/Plugins/PCG/Source/PCG/Private/Elements/PCGDataTableRowToParamData.cpp
+++ b/Engine/Plugins/PCG/Source/PCG/Private/Elements/PCGDataTableRowToParamData.cpp
@@ -2,7 +2,6 @@
 
 #include "Elements/PCGDataTableRowToParamData.h"
 
-#include "PCGContext.h"
 #include "PCGParamData.h"
 #include "PCGPin.h"
 #include "PCGModule.h"
@@ -74,18 +73,44 @@ FString UPCGDataTableRowToParamDataSettings::GetAdditionalTitleInformation() con
 	return FString::Printf(TEXT("%s[ %s ]"), DataTable ? *DataTable->GetFName().ToString() : TEXT("None"), *RowName.ToString());
 }
 
+bool FPCGDataTableRowToParamData::PrepareDataInternal(FPCGContext* Context) const
+{
+	TRACE_CPUPROFILER_EVENT_SCOPE(FPCGDataTableRowToParamData::PrepareDataInternal);
+
+	check(Context);
+
+	const UPCGDataTableRowToParamDataSettings* Settings = Context->GetInputSettings<UPCGDataTableRowToParamDataSettings>();
+	check(Settings);
+
+	if (Settings->DataTable.IsNull())
+	{
+		return true;
+	}
+
+	FPCGDataTableRowToParamDataContext* ThisContext = static_cast<FPCGDataTableRowToParamDataContext*>(Context);
+
+	if (!ThisContext->WasLoadRequested())
+	{
+		return ThisContext->RequestResourceLoad(ThisContext, { Settings->DataTable.ToSoftObjectPath() }, !Settings->bSynchronousLoad);
+	}
+
+	return true;
+}
+
 bool FPCGDataTableRowToParamData::ExecuteInternal(FPCGContext* Context) const
 {
 	TRACE_CPUPROFILER_EVENT_SCOPE(FPCGDataTableRowToParamData::Execute);
 
+	check(Context);
+
 	const UPCGDataTableRowToParamDataSettings* Settings = Context->GetInputSettings<UPCGDataTableRowToParamDataSettings>();
 	check(Settings);
 
 	const FName RowName = Settings->RowName;
 
-	TSoftObjectPtr<UDataTable> DataTablePtr = Settings->DataTable;
+	const TSoftObjectPtr<UDataTable> DataTablePtr = Settings->DataTable;
 
-	const UDataTable* DataTable = DataTablePtr.LoadSynchronous();
+	const UDataTable* DataTable = DataTablePtr.Get();
 	if (!DataTable)
 	{
 		PCGE_LOG(Warning, GraphAndLog, LOCTEXT("DataTableInvalid", "FPCGDataTableRowToParamData: Data table is invalid"));
diff --git a/Engine/Plugins/PCG/Source/PCG/Private/Elements/PCGPointFromMeshElement.cpp b/Engine/Plugins/PCG/Source/PCG/Private/Elements/PCGPointFromMeshElement.cpp
--- a/Engine/Plugins/PCG/Source/PCG/Private/Elements/PCGPointFromMeshElement.cpp
+++ b/Engine/Plugins/PCG/Source/PCG/Private/Elements/PCGPointFromMeshElement.cpp
@@ -3,7 +3,6 @@
 #include "Elements/PCGPointFromMeshElement.h"
 
 #include "PCGComponent.h"
-#include "PCGContext.h"
 #include "PCGEdge.h"
 #include "Data/PCGPointData.h"
 #include "Data/PCGSpatialData.h"
@@ -31,6 +30,30 @@ FPCGElementPtr UPCGPointFromMeshSettings::CreateElement() const
 	return MakeShared<FPCGPointFromMeshElement>();
 }
 
+bool FPCGPointFromMeshElement::PrepareDataInternal(FPCGContext* Context) const
+{
+	TRACE_CPUPROFILER_EVENT_SCOPE(FPCGPointFromMeshElement::Execute);
+
+	check(Context);
+
+	const UPCGPointFromMeshSettings* Settings = Context->GetInputSettings<UPCGPointFromMeshSettings>();
+	check(Settings);
+
+	if (Settings->StaticMesh.IsNull())
+	{
+		return true;
+	}
+
+	FPCGPointFromMeshContext* ThisContext = static_cast<FPCGPointFromMeshContext*>(Context);
+
+	if (!ThisContext->WasLoadRequested())
+	{
+		return ThisContext->RequestResourceLoad(ThisContext, { Settings->StaticMesh.ToSoftObjectPath() }, !Settings->bSynchronousLoad);
+	}
+
+	return true;
+}
+
 bool FPCGPointFromMeshElement::ExecuteInternal(FPCGContext* Context) const
 {
 	TRACE_CPUPROFILER_EVENT_SCOPE(FPCGPointFromMeshElement::Execute);
@@ -45,7 +68,7 @@ bool FPCGPointFromMeshElement::ExecuteInternal(FPCGContext* Context) const
 		return true;
 	}
 
-	if (!Settings->StaticMesh.LoadSynchronous())
+	if (!Settings->StaticMesh.Get())
 	{
 		PCGE_LOG(Error, GraphAndLog, LOCTEXT("LoadStaticMeshFailed", "Failed to load StaticMesh"));
 		return true;
diff --git a/Engine/Plugins/PCG/Source/PCG/Private/Elements/PCGTextureSampler.cpp b/Engine/Plugins/PCG/Source/PCG/Private/Elements/PCGTextureSampler.cpp
--- a/Engine/Plugins/PCG/Source/PCG/Private/Elements/PCGTextureSampler.cpp
+++ b/Engine/Plugins/PCG/Source/PCG/Private/Elements/PCGTextureSampler.cpp
@@ -86,11 +86,35 @@ void UPCGTextureSamplerSettings::UpdateDisplayTextureArrayIndex()
 }
 #endif
 
+bool FPCGTextureSamplerElement::PrepareDataInternal(FPCGContext* InContext) const
+{
+	TRACE_CPUPROFILER_EVENT_SCOPE(FPCGTextureSamplerElement::Execute);
+
+	FPCGTextureSamplerContext* Context = static_cast<FPCGTextureSamplerContext*>(InContext);
+	check(Context);
+
+	const UPCGTextureSamplerSettings* Settings = Context->GetInputSettings<UPCGTextureSamplerSettings>();
+	check(Settings);
+
+	if (Settings->Texture.IsNull())
+	{
+		return true;
+	}
+
+	if (!Context->WasLoadRequested())
+	{
+		return Context->RequestResourceLoad(Context, { Settings->Texture.ToSoftObjectPath() }, !Settings->bSynchronousLoad);
+	}
+
+	return true;
+}
+
 bool FPCGTextureSamplerElement::ExecuteInternal(FPCGContext* InContext) const
 {
 	TRACE_CPUPROFILER_EVENT_SCOPE(FPCGTextureSamplerElement::Execute);
 
 	FPCGTextureSamplerContext* Context = static_cast<FPCGTextureSamplerContext*>(InContext);
+	check(Context);
 
 	if (Context->bIsPaused)
 	{
@@ -110,7 +134,7 @@ bool FPCGTextureSamplerElement::ExecuteInternal(FPCGContext* InContext) const
 		return true;
 	}
 
-	UTexture* Texture = Settings->Texture.LoadSynchronous();
+	UTexture* Texture = Settings->Texture.Get();
 
 	if (!Texture)
 	{
diff --git a/Engine/Plugins/PCG/Source/PCG/Public/Async/PCGAsyncLoadingContext.h b/Engine/Plugins/PCG/Source/PCG/Public/Async/PCGAsyncLoadingContext.h
--- a/Engine/Plugins/PCG/Source/PCG/Public/Async/PCGAsyncLoadingContext.h
+++ b/Engine/Plugins/PCG/Source/PCG/Public/Async/PCGAsyncLoadingContext.h
@@ -12,7 +12,7 @@ struct FPCGContext;
 * Extension interface to add Async Loading support to any Context.
 * Just inherit IPCGAsyncLoadingContext on the element context.
 */
-struct IPCGAsyncLoadingContext
+struct PCG_API IPCGAsyncLoadingContext
 {
 public:
 	virtual ~IPCGAsyncLoadingContext();
diff --git a/Engine/Plugins/PCG/Source/PCG/Public/Elements/IO/PCGDataTableElement.h b/Engine/Plugins/PCG/Source/PCG/Public/Elements/IO/PCGDataTableElement.h
--- a/Engine/Plugins/PCG/Source/PCG/Public/Elements/IO/PCGDataTableElement.h
+++ b/Engine/Plugins/PCG/Source/PCG/Public/Elements/IO/PCGDataTableElement.h
@@ -39,6 +39,10 @@ public:
 
 	UPROPERTY(BlueprintReadWrite, EditAnywhere, Category = Settings, meta = (ValidEnumValues = "Point, Param"))
 	EPCGExclusiveDataType OutputType = EPCGExclusiveDataType::Point;
+
+	/** By default, data table loading is asynchronous, can force it synchronous if needed. */
+	UPROPERTY(BlueprintReadWrite, EditAnywhere, Category = "Settings|Debug")
+	bool bSynchronousLoad = false;
 };
 
 class FPCGLoadDataTableElement : public FPCGExternalDataElement
diff --git a/Engine/Plugins/PCG/Source/PCG/Public/Elements/IO/PCGExternalDataContext.h b/Engine/Plugins/PCG/Source/PCG/Public/Elements/IO/PCGExternalDataContext.h
--- a/Engine/Plugins/PCG/Source/PCG/Public/Elements/IO/PCGExternalDataContext.h
+++ b/Engine/Plugins/PCG/Source/PCG/Public/Elements/IO/PCGExternalDataContext.h
@@ -3,13 +3,14 @@
 #pragma once
 
 #include "PCGContext.h"
+#include "Async/PCGAsyncLoadingContext.h"
 #include "Data/PCGPointData.h"
 #include "Metadata/PCGAttributePropertySelector.h"
 #include "Metadata/Accessors/IPCGAttributeAccessor.h"
 #include "Metadata/Accessors/PCGAttributeAccessorHelpers.h"
 #include "Metadata/Accessors/PCGAttributeAccessorKeys.h"
 
-struct FPCGExternalDataContext : public FPCGContext
+struct FPCGExternalDataContext : public FPCGContext, public IPCGAsyncLoadingContext
 {
 	// Used for data table parsing
 	struct FRowToPointAccessors
diff --git a/Engine/Plugins/PCG/Source/PCG/Public/Elements/PCGDataTableRowToParamData.h b/Engine/Plugins/PCG/Source/PCG/Public/Elements/PCGDataTableRowToParamData.h
--- a/Engine/Plugins/PCG/Source/PCG/Public/Elements/PCGDataTableRowToParamData.h
+++ b/Engine/Plugins/PCG/Source/PCG/Public/Elements/PCGDataTableRowToParamData.h
@@ -2,7 +2,9 @@
 
 #pragma once
 
+#include "PCGContext.h"
 #include "PCGSettings.h"
+#include "Async/PCGAsyncLoadingContext.h"
 
 #include "PCGDataTableRowToParamData.generated.h"
 
@@ -41,6 +43,10 @@ public:
 	UPROPERTY(BlueprintReadWrite, EditAnywhere, Category = Settings, meta = (PCG_Overridable, PCG_OverrideAliases = "PathOverride"))
 	TSoftObjectPtr<UDataTable> DataTable;
 
+	/** By default, data table loading is asynchronous, can force it synchronous if needed. */
+	UPROPERTY(BlueprintReadWrite, EditAnywhere, Category = "Settings|Debug")
+	bool bSynchronousLoad = false;
+
 #if WITH_EDITORONLY_DATA
 	UPROPERTY()
 	FString PathOverride_DEPRECATED = FString();
@@ -55,13 +61,16 @@ protected:
 	// ~End UPCGSettings interface
 };
 
-class FPCGDataTableRowToParamData : public IPCGElement
+struct FPCGDataTableRowToParamDataContext : public FPCGContext, public IPCGAsyncLoadingContext {};
+
+class FPCGDataTableRowToParamData : public IPCGElementWithCustomContext<FPCGDataTableRowToParamDataContext>
 {
 public:
 	virtual bool CanExecuteOnlyOnMainThread(FPCGContext* Context) const override { return true; }
 	virtual bool IsCacheable(const UPCGSettings* InSettings) const override { return false; }
 
 protected:
+	virtual bool PrepareDataInternal(FPCGContext* Context) const override;
 	virtual bool ExecuteInternal(FPCGContext* Context) const override;
 };
 
diff --git a/Engine/Plugins/PCG/Source/PCG/Public/Elements/PCGPointFromMeshElement.h b/Engine/Plugins/PCG/Source/PCG/Public/Elements/PCGPointFromMeshElement.h
--- a/Engine/Plugins/PCG/Source/PCG/Public/Elements/PCGPointFromMeshElement.h
+++ b/Engine/Plugins/PCG/Source/PCG/Public/Elements/PCGPointFromMeshElement.h
@@ -2,7 +2,9 @@
 
 #pragma once
 
+#include "PCGContext.h"
 #include "PCGSettings.h"
+#include "Async/PCGAsyncLoadingContext.h"
 
 #include "UObject/SoftObjectPtr.h"
 
@@ -38,10 +40,17 @@ public:
 	/** Name of the string attribute to be created and hold a SoftObjectPath to the StaticMesh */
 	UPROPERTY(BlueprintReadWrite, EditAnywhere, Category = Settings)
 	FName MeshPathAttributeName = NAME_None;
+
+	/** By default, mesh loading is asynchronous, can force it synchronous if needed. */
+	UPROPERTY(BlueprintReadWrite, EditAnywhere, Category = "Settings|Debug")
+	bool bSynchronousLoad = false;
 };
 
-class FPCGPointFromMeshElement : public IPCGElement
+struct FPCGPointFromMeshContext : public FPCGContext, public IPCGAsyncLoadingContext {};
+
+class FPCGPointFromMeshElement : public IPCGElementWithCustomContext<FPCGPointFromMeshContext>
 {
 protected:
+	virtual bool PrepareDataInternal(FPCGContext* Context) const override;
 	virtual bool ExecuteInternal(FPCGContext* Context) const override;
 };
diff --git a/Engine/Plugins/PCG/Source/PCG/Public/Elements/PCGTextureSampler.h b/Engine/Plugins/PCG/Source/PCG/Public/Elements/PCGTextureSampler.h
--- a/Engine/Plugins/PCG/Source/PCG/Public/Elements/PCGTextureSampler.h
+++ b/Engine/Plugins/PCG/Source/PCG/Public/Elements/PCGTextureSampler.h
@@ -4,6 +4,7 @@
 
 #include "PCGContext.h"
 #include "PCGSettings.h"
+#include "Async/PCGAsyncLoadingContext.h"
 #include "Data/PCGTextureData.h"
 
 #include "PCGTextureSampler.generated.h"
@@ -102,6 +103,10 @@ public:
 	bool bForceEditorOnlyCPUSampling = false;
 #endif
 
+	/** By default, texture loading is asynchronous, can force it synchronous if needed. */
+	UPROPERTY(BlueprintReadWrite, EditAnywhere, Category = "Settings|Debug")
+	bool bSynchronousLoad = false;
+
 protected:
 #if WITH_EDITORONLY_DATA
 	// Used to hide the 'TextureArrayIndex' property.
@@ -110,7 +115,7 @@ protected:
 #endif
 };
 
-struct PCG_API FPCGTextureSamplerContext : public FPCGContext
+struct PCG_API FPCGTextureSamplerContext : public FPCGContext, public IPCGAsyncLoadingContext
 {
 	bool bTextureReadbackDone = false;
 };
@@ -122,5 +127,6 @@ public:
 
 protected:
 	virtual FPCGContext* CreateContext() override;
+	virtual bool PrepareDataInternal(FPCGContext* InContext) const override;
 	virtual bool ExecuteInternal(FPCGContext* InContext) const override;
 };
diff --git a/Engine/Plugins/PCG/Source/PCG/Public/PCGElement.h b/Engine/Plugins/PCG/Source/PCG/Public/PCGElement.h
--- a/Engine/Plugins/PCG/Source/PCG/Public/PCGElement.h
+++ b/Engine/Plugins/PCG/Source/PCG/Public/PCGElement.h
@@ -135,6 +135,16 @@ private:
 	void CleanupAndValidateOutput(FPCGContext* Context) const;
 };
 
+/**
+* Convenience class for element with custom context that can be default constructed (need no initialization)
+*/
+template <typename ContextType, typename = typename std::enable_if_t<std::is_base_of_v<FPCGContext, ContextType> && std::is_default_constructible_v<ContextType>>>
+class PCG_API IPCGElementWithCustomContext : public IPCGElement
+{
+public:
+	virtual FPCGContext* CreateContext() { return new ContextType(); }
+};
+
 class UE_DEPRECATED(5.4, "This class has been deprecated. Please inherit from IPCGElement directly.") PCG_API FSimplePCGElement : public IPCGElement
 {};
 
diff --git a/Engine/Plugins/PCGGeometryScriptInterop/Source/PCGGeometryScriptInterop/Private/Elements/PCGMeshSampler.cpp b/Engine/Plugins/PCGGeometryScriptInterop/Source/PCGGeometryScriptInterop/Private/Elements/PCGMeshSampler.cpp
--- a/Engine/Plugins/PCGGeometryScriptInterop/Source/PCGGeometryScriptInterop/Private/Elements/PCGMeshSampler.cpp
+++ b/Engine/Plugins/PCGGeometryScriptInterop/Source/PCGGeometryScriptInterop/Private/Elements/PCGMeshSampler.cpp
@@ -116,9 +116,22 @@ bool FPCGMeshSamplerElement::PrepareDataInternal(FPCGContext* InContext) const
 	const UPCGMeshSamplerSettings* Settings = Context->GetInputSettings<UPCGMeshSamplerSettings>();
 	check(Settings);
 
-	// TODO: Could be async
-	TSoftObjectPtr<UStaticMesh> StaticMeshPtr = Settings->StaticMesh;
-	UStaticMesh* StaticMesh = StaticMeshPtr.LoadSynchronous();
+	const TSoftObjectPtr<UStaticMesh> StaticMeshPtr = Settings->StaticMesh;
+	if (StaticMeshPtr.IsNull())
+	{
+		return true;
+	}
+
+	// 1. Request load for mesh. Return false if we need to wait, otherwise continue.
+	if (!Context->WasLoadRequested())
+	{
+		if (!Context->RequestResourceLoad(Context, { StaticMeshPtr.ToSoftObjectPath() }, !Settings->bSynchronousLoad))
+		{
+			return false;
+		}
+	}
+
+	UStaticMesh* StaticMesh = StaticMeshPtr.Get();
 
 	if (!StaticMesh)
 	{
diff --git a/Engine/Plugins/PCGGeometryScriptInterop/Source/PCGGeometryScriptInterop/Public/Elements/PCGMeshSampler.h b/Engine/Plugins/PCGGeometryScriptInterop/Source/PCGGeometryScriptInterop/Public/Elements/PCGMeshSampler.h
--- a/Engine/Plugins/PCGGeometryScriptInterop/Source/PCGGeometryScriptInterop/Public/Elements/PCGMeshSampler.h
+++ b/Engine/Plugins/PCGGeometryScriptInterop/Source/PCGGeometryScriptInterop/Public/Elements/PCGMeshSampler.h
@@ -5,6 +5,7 @@
 #include "PCGContext.h"
 #include "PCGElement.h"
 #include "PCGSettings.h"
+#include "Async/PCGAsyncLoadingContext.h"
 
 #include "Async/Future.h"
 #include "GeometryScript/GeometryScriptTypes.h"
@@ -140,6 +141,9 @@ public:
 	 */
 	UPROPERTY(BlueprintReadWrite, EditAnywhere, Category = "Settings|Points", meta=(ClampMin="0", ClampMax="1", PCG_Overridable))
 	float PointSteepness = 0.5f;
+
+	UPROPERTY(BlueprintReadWrite, EditAnywhere, Category = "Settings|Debug")
+	bool bSynchronousLoad = false;
 	
 protected:
 #if WITH_EDITORONLY_DATA
@@ -155,7 +159,7 @@ protected:
 /**
 * Extra context to store all the data that need to be kept between multiple executions (time slicing)
 */
-struct FPCGMeshSamplerContext : public FPCGContext
+struct FPCGMeshSamplerContext : public FPCGContext, public IPCGAsyncLoadingContext
 {
 public:
 	~FPCGMeshSamplerContext();
