[PCG] Update C++ API in PCGMetadataPartitionCommon.h to support partitioning on multiple attributes

+ AttributePartition and AttributeGenericPartition C++ API Extended in PCGMetadataPartitionCommon.h to support a TArrayView of selectors
+ PCGMetadataPartitionCommon now supports partitioning on multiple attributes
* Modified MetadataPartition element to support multiple partition attributes

TODO
-------
Unit Test Coverage to follow
Performance optimizations for the multi partition

diff --git a/Engine/Plugins/PCG/Source/PCG/Private/Elements/Metadata/PCGMetadataPartition.cpp b/Engine/Plugins/PCG/Source/PCG/Private/Elements/Metadata/PCGMetadataPartition.cpp
--- a/Engine/Plugins/PCG/Source/PCG/Private/Elements/Metadata/PCGMetadataPartition.cpp
+++ b/Engine/Plugins/PCG/Source/PCG/Private/Elements/Metadata/PCGMetadataPartition.cpp
@@ -3,6 +3,9 @@
 #include "Elements/Metadata/PCGMetadataPartition.h"
 
 #include "PCGContext.h"
+#include "PCGModule.h"
+#include "Helpers/PCGHelpers.h"
+#include "Metadata/PCGMetadataAttribute.h"
 #include "Metadata/PCGMetadataPartitionCommon.h"
 
 #include UE_INLINE_GENERATED_CPP_BY_NAME(PCGMetadataPartition)
@@ -37,23 +40,42 @@ void UPCGMetadataPartitionSettings::PostLoad()
 #if WITH_EDITOR
 	if (PartitionAttribute_DEPRECATED != NAME_None)
 	{
-		PartitionAttributeSource.SetAttributeName(PartitionAttribute_DEPRECATED);
+		PartitionAttributeSelectors.Empty();
+		PartitionAttributeSelectors.Emplace_GetRef().SetAttributeName(PartitionAttribute_DEPRECATED);
 		PartitionAttribute_DEPRECATED = NAME_None;
 	}
+
+	if (PartitionAttributeSource_DEPRECATED.GetAttributeName() != PCGMetadataAttributeConstants::LastAttributeName)
+	{
+		PartitionAttributeSelectors.Empty();
+		PartitionAttributeSelectors.Emplace(PartitionAttributeSource_DEPRECATED);
+		PartitionAttributeSource_DEPRECATED.SetAttributeName(PCGMetadataAttributeConstants::LastAttributeName);
+	}
 #endif // WITH_EDITOR
 }
 
 FString UPCGMetadataPartitionSettings::GetAdditionalTitleInformation() const
 {
 #if WITH_EDITOR
-	if (IsPropertyOverriddenByPin(GET_MEMBER_NAME_CHECKED(UPCGMetadataPartitionSettings, PartitionAttributeSource)))
+	if (IsPropertyOverriddenByPin(GET_MEMBER_NAME_CHECKED(UPCGMetadataPartitionSettings, PartitionAttributeNames)))
 	{
 		return FString();
 	}
 	else
 #endif
 	{
-		return PartitionAttributeSource.GetDisplayText().ToString();
+		if (PartitionAttributeSelectors.IsEmpty())
+		{
+			return FString();
+		}
+
+		FString OutString = PartitionAttributeSelectors[0].GetDisplayText().ToString();
+		for (int I = 1; I < PartitionAttributeSelectors.Num(); ++I)
+		{
+			OutString += ", " + PartitionAttributeSelectors[I].GetDisplayText().ToString();
+		}
+
+		return OutString;
 	}
 }
 
@@ -65,15 +87,43 @@ bool FPCGMetadataPartitionElement::ExecuteInternal(FPCGContext* Context) const
 	const UPCGMetadataPartitionSettings* Settings = Context->GetInputSettings<UPCGMetadataPartitionSettings>();
 	check(Settings);
 
+	// TODO: This is a temporary solution for overrides until arrays are supported
+	TArray<FPCGAttributePropertyInputSelector> OverriddenSelectors;
+	const TArray<FString> AttributeNames = PCGHelpers::GetStringArrayFromCommaSeparatedString(Settings->PartitionAttributeNames);
+	// If the names are overridden by the user, generate the selectors with them
+	OverriddenSelectors.SetNum(AttributeNames.Num());
+	for (const FString& AttributeName : AttributeNames)
+	{
+		OverriddenSelectors.Emplace_GetRef().SetAttributeName(FName(AttributeName));
+	}
+
+	const TArray<FPCGAttributePropertyInputSelector>& ActiveSelectors = OverriddenSelectors.IsEmpty() ? Settings->PartitionAttributeSelectors : OverriddenSelectors;
+
 	TArray<FPCGTaggedData> Inputs = Context->InputData.GetInputsByPin(PCGPinConstants::DefaultInputLabel);
 	TArray<FPCGTaggedData>& Outputs = Context->OutputData.TaggedData;
 
+	TArray<FPCGAttributePropertySelector> PartitionAttributeSources;
+	PartitionAttributeSources.SetNum(ActiveSelectors.Num());
+
+	if (PartitionAttributeSources.IsEmpty())
+	{
+		Outputs = Inputs;
+		return true;
+	}
+
 	for (const FPCGTaggedData& Input : Inputs)
 	{
 		const UPCGData* InData = Input.Data;
-		const FPCGAttributePropertyInputSelector PartitionAttributeSelector = Settings->PartitionAttributeSource.CopyAndFixLast(InData);
 
-		for (UPCGData* PartitionData : PCGMetadataPartitionCommon::AttributePartition(InData, PartitionAttributeSelector, Context))
+		for (int I = 0; I < PartitionAttributeSources.Num(); ++I)
+		{
+			PartitionAttributeSources[I] = static_cast<FPCGAttributePropertySelector>(ActiveSelectors[I].CopyAndFixLast(InData));
+		}
+
+		TArray<UPCGData*> PartitionDataArray;
+		PartitionDataArray = PCGMetadataPartitionCommon::AttributePartition(InData, PartitionAttributeSources, Context);
+
+		for (UPCGData* PartitionData : PartitionDataArray)
 		{
 			if (PartitionData)
 			{
diff --git a/Engine/Plugins/PCG/Source/PCG/Private/Metadata/PCGMetadataPartitionCommon.cpp b/Engine/Plugins/PCG/Source/PCG/Private/Metadata/PCGMetadataPartitionCommon.cpp
--- a/Engine/Plugins/PCG/Source/PCG/Private/Metadata/PCGMetadataPartitionCommon.cpp
+++ b/Engine/Plugins/PCG/Source/PCG/Private/Metadata/PCGMetadataPartitionCommon.cpp
@@ -2,8 +2,6 @@
 
 #include "Metadata/PCGMetadataPartitionCommon.h"
 
-#include "PCGContext.h"
-#include "PCGElement.h"
 #include "PCGModule.h"
 #include "PCGParamData.h"
 #include "Data/PCGPointData.h"
@@ -27,6 +25,7 @@ namespace PCGMetadataPartitionCommon
 	*/
 	TArray<TArray<int32>> AttributePartition(const FPCGMetadataAttributeBase* InAttribute, const IPCGAttributeAccessorKeys& InKeys, FPCGContext* InOptionalContext)
 	{
+		TRACE_CPUPROFILER_EVENT_SCOPE(PCGMetadataPartitionCommon::AttributePartition);
 		check(InAttribute);
 
 		const int32 NumberOfEntries = InKeys.GetNum();
@@ -61,7 +60,7 @@ namespace PCGMetadataPartitionCommon
 					continue;
 				}
 
-				// TODO: Might want to upgrade to something better wince it can be quadractic and grow quickly.
+				// TODO: Might want to upgrade to something better since it can be quadratic and grow quickly.
 				int32 UniqueValueKeyIndex = Algo::IndexOfByPredicate(UniqueValueKeys, [ValueKey, InAttribute](const PCGMetadataValueKey& Key)
 				{
 					return InAttribute->AreValuesEqual(ValueKey, Key);
@@ -114,7 +113,7 @@ namespace PCGMetadataPartitionCommon
 			}
 		}
 
-		// Since we partition on the value array, it is not guarenteed that the values appears in the same order than the entries.
+		// Since we partition on the value array, it is not guaranteed that the values appears in the same order than the entries.
 		// So sort the final array using the first index as a sort criteria. Empty partitions will be at the beginning too.
 		PartitionedData.Sort([](const TArray<int32>& LHS, const TArray<int32>& RHS) -> bool
 		{ 
@@ -142,12 +141,13 @@ namespace PCGMetadataPartitionCommon
 	template <typename T>
 	TArray<TArray<int32>> ValuePartition(const IPCGAttributeAccessor& InAccessor, const IPCGAttributeAccessorKeys& InKeys, FPCGContext* InOptionalContext)
 	{
+		TRACE_CPUPROFILER_EVENT_SCOPE(PCGMetadataPartitionCommon::ValuePartition);
 		TArray<T> UniqueValues;
 		TArray<TArray<int32>> PartitionedData;
 
 		PCGMetadataElementCommon::ApplyOnAccessor<T>(InKeys, InAccessor, [&PartitionedData, &UniqueValues](const T& InValue, int32 InIndex)
 		{
-			// TODO: Might want to upgrade to something better wince it can be quadractic and grow quickly.
+			// TODO: Might want to upgrade to something better since it can be quadratic and grow quickly.
 			int32 UniqueValueIndex = UniqueValues.IndexOfByPredicate([&InValue](const T& OtherValue)
 			{
 				// For consistency with the attribute part, use MetadataTraits::Equal
@@ -171,6 +171,7 @@ namespace PCGMetadataPartitionCommon
 	*/
 	TArray<TArray<int32>> AttributeGenericPartition(const UPCGData* InData, const FPCGAttributePropertySelector& InSelector, FPCGContext* InOptionalContext)
 	{
+		TRACE_CPUPROFILER_EVENT_SCOPE(PCGMetadataPartitionCommon::AttributeGenericPartition::SingleSelector);
 		if (!InData)
 		{
 			return {};
@@ -218,7 +219,7 @@ namespace PCGMetadataPartitionCommon
 				// Can't partition on a transform.
 				if constexpr (std::is_same_v<AttributeType, FTransform>)
 				{
-					PCGLog::LogErrorOnGraph(FText::Format(LOCTEXT("InvalidType", "Attribute {0} is a transform, partition transforms is not supported"), InSelector.GetDisplayText()), InOptionalContext);
+					PCGLog::LogErrorOnGraph(FText::Format(LOCTEXT("InvalidType", "Attribute {0} is a transform, partition on transforms is not supported"), InSelector.GetDisplayText()), InOptionalContext);
 					return {};
 				}
 				else
@@ -231,12 +232,138 @@ namespace PCGMetadataPartitionCommon
 		}
 	}
 
+	/**
+	 * Partition on multiple attributes by first partitioning on the attributes independently. Then take the resultant
+	 * partition and convert them to a BitArray representation of each element's partition. Once in BitArray form,
+	 * combine the results with a logical AND operation to filter them into final partition groupings.
+	 *
+	 * Multi-Partition Example:
+	 * Pt  A  B  C                         Partition on A->[0,1],[2,3,4]
+	 *  0  a  a  a                         Partition on B->[0],[1,2],[3,4]
+	 *  1  a  b  a                         Partition on C->[0,1],[2],[3,4]
+	 *  2  b  b  b                         Partition on A&B->[0],[1],[2],[3,4]
+	 *  3  b  c  c                         Final Partition (A&B&C)->[0],[1],[2],[3,4]
+	 *  4  b  c  c
+	 */
+	TArray<TArray<int32>> AttributeGenericPartition(const UPCGData* InData, const TArrayView<const FPCGAttributePropertySelector>& InSelectorArrayView, FPCGContext* InOptionalContext)
+	{
+		TRACE_CPUPROFILER_EVENT_SCOPE(PCGMetadataPartitionCommon::AttributeGenericPartition::MultiSelector);
+
+		// Small optimization to partition on a single attribute
+		if (InSelectorArrayView.Num() == 1)
+		{
+			return AttributeGenericPartition(InData, InSelectorArrayView[0], InOptionalContext);
+		}
+
+		if (!InData || InSelectorArrayView.IsEmpty() || !InData->ConstMetadata())
+		{
+			return {};
+		}
+
+		// Get the element count from the number of keys which should work for spatial points and attribute sets
+		const TUniquePtr<const IPCGAttributeAccessorKeys> Keys = PCGAttributeAccessorHelpers::CreateConstKeys(InData, InSelectorArrayView[0]);
+		if (!Keys.IsValid())
+		{
+			PCGLog::LogErrorOnGraph(FText::Format(LOCTEXT("InvalidKeys", "Could not create keys for the input data with selector {0}"), InSelectorArrayView[0].GetDisplayText()), InOptionalContext);
+			return {};
+		}
+
+		const int64 NumElements = Keys->GetNum();
+		const int64 NumAttributes = InSelectorArrayView.Num();
+
+		using IndexPartition = TArray<TArray<int32>>;
+		using BitPartition = TArray<TBitArray<>>;
+
+		TArray<IndexPartition> IndexPartitions;
+		IndexPartitions.SetNum(NumAttributes);
+		TArray<BitPartition> BitPartitions;
+		BitPartitions.SetNum(NumAttributes);
+
+		{
+			TRACE_CPUPROFILER_EVENT_SCOPE(PCGMetadataPartitionCommon::AttributeGenericPartition::MultiSelector::ConversionToBitArray);
+			/* TODO: Can be executed in parallel, threadsafe. There is a follow up task to evaluate between
+			 * option A.) partitioning on all attributes, and then merging and B.) Partitioning on each attribute
+			 * in succession, further partitioning the grouping results of the previous iteration.
+			 */
+			// Calculate each partition and convert it to a bitfield for simple/efficient intersection processing
+			for (int32 I = 0; I < InSelectorArrayView.Num(); ++I)
+			{
+				IndexPartition& CurrentIndexPartition = IndexPartitions[I];
+				BitPartition& CurrentBitPartition = BitPartitions[I];
+
+				// TODO: Ideally, refactor AttributeGenericPartition to return directly into BitArray format instead to avoid conversion
+				// Partition once for each attribute
+				CurrentIndexPartition = AttributeGenericPartition(InData, InSelectorArrayView[I], InOptionalContext);
+				// The bit partitions will match the index partitions
+				CurrentBitPartition.SetNum(CurrentIndexPartition.Num());
+
+				auto ConvertIndexGroupingToBitGrouping = [NumElements](const IndexPartition& InIndexPartition, BitPartition& OutBitPartition)
+				{
+					for (int32 GroupIndex = 0; GroupIndex < InIndexPartition.Num(); ++GroupIndex)
+					{
+						OutBitPartition[GroupIndex].SetNum(NumElements, false);
+						for (const int32 Index : InIndexPartition[GroupIndex])
+						{
+							OutBitPartition[GroupIndex].Insert(true, Index);
+						}
+					}
+				};
+
+				ConvertIndexGroupingToBitGrouping(CurrentIndexPartition, CurrentBitPartition);
+			}
+		}
+
+		BitPartition IterativePartition = BitPartitions[0]; // TODO: This can be optimized to filter down in pairs in parallel - O(log N) - instead of serial
+
+		{
+			TRACE_CPUPROFILER_EVENT_SCOPE(PCGMetadataPartitionCommon::AttributeGenericPartition::Intersection);
+			// Intersect all the BitArray partitions
+			for (int32 PartitionIndex = 1; PartitionIndex < BitPartitions.Num(); ++PartitionIndex)
+			{
+				BitPartition CurrentBitPartition = IterativePartition;
+				BitPartition& NextBitPartition = BitPartitions[PartitionIndex];
+
+				IterativePartition.Empty();
+
+				for (const TBitArray<>& FirstBitArray : CurrentBitPartition)
+				{
+					for (const TBitArray<>& SecondBitArray : NextBitPartition)
+					{
+						TBitArray<> Result = TBitArray<>::BitwiseAND(FirstBitArray, SecondBitArray, EBitwiseOperatorFlags::MaxSize);
+						// Only capture if non-empty. Discard empty BitArrays.
+						if (TConstSetBitIterator(Result))
+						{
+							IterativePartition.Emplace(std::move(Result));
+						}
+					}
+				}
+			}
+		}
+
+		IndexPartition FinalPartition;
+
+		{
+			TRACE_CPUPROFILER_EVENT_SCOPE(PCGMetadataPartitionCommon::AttributeGenericPartition::MultiSelector::ConversionToIndices);
+			// Convert back into indices
+			for (const TBitArray<>& BitArray : IterativePartition)
+			{
+				TArray<int>& Indices = FinalPartition.Emplace_GetRef();
+				for (TConstSetBitIterator It(BitArray); It; ++It)
+				{
+					Indices.Emplace(It.GetIndex());
+				}
+			}
+		}
+
+		return FinalPartition;
+	}
+
 	/**
 	* Do a partition on the given point data for the selector
 	*/
-	TArray<UPCGData*> AttributePointPartition(const UPCGPointData* InData, const FPCGAttributePropertySelector& InSelector, FPCGContext* InOptionalContext)
+	TArray<UPCGData*> AttributePointPartition(const UPCGPointData* InData, const TArrayView<const FPCGAttributePropertySelector>& InSelectorArrayView, FPCGContext* InOptionalContext)
 	{
-		TArray<TArray<int32>> Partition = AttributeGenericPartition(InData, InSelector, InOptionalContext);
+		TArray<TArray<int32>> Partition = AttributeGenericPartition(InData, InSelectorArrayView, InOptionalContext);
 		if (Partition.IsEmpty())
 		{
 			return {};
@@ -268,23 +395,23 @@ namespace PCGMetadataPartitionCommon
 		return PartitionedData;
 	}
 
-	TArray<UPCGData*> AttributeParamSpatialPartition(const UPCGData* InData, const FPCGAttributePropertySelector& InSelector, FPCGContext* InOptionalContext)
+	TArray<UPCGData*> AttributeParamSpatialPartition(const UPCGData* InData, const TArrayView<const FPCGAttributePropertySelector>& InSelectorArray, FPCGContext* InOptionalContext)
 	{
-		const UPCGSpatialData* InSpatialData = Cast<const UPCGSpatialData>(InData);
-		const UPCGParamData* InParamData = Cast<const UPCGParamData>(InData);
-
-		if (!InSpatialData && !InParamData)
+		if (!InData->IsA<UPCGSpatialData>() && !InData->IsA<UPCGParamData>())
 		{
 			PCGLog::LogErrorOnGraph(LOCTEXT("InvalidDataType", "Input data is not an attribute set nor a spatial data. Operation not supported."), InOptionalContext);
 			return {};
 		}
 
-		TArray<TArray<int32>> Partition = AttributeGenericPartition(InData, InSelector, InOptionalContext);
+		const TArray<TArray<int32>> Partition = AttributeGenericPartition(InData, InSelectorArray, InOptionalContext);
+
 		if (Partition.IsEmpty())
 		{
 			return {};
 		}
 
+		const UPCGSpatialData* InSpatialData = Cast<const UPCGSpatialData>(InData);
+
 		TArray<UPCGData*> PartitionedData;
 		PartitionedData.Reserve(Partition.Num());
 
@@ -293,7 +420,7 @@ namespace PCGMetadataPartitionCommon
 		const UPCGMetadata* OriginalMetadata = InData->ConstMetadata();
 		OriginalMetadata->GetAttributes(AttributeNames, AttributeTypes);
 
-		for (TArray<int32>& Indices : Partition)
+		for (const TArray<int32>& Indices : Partition)
 		{
 			if (Indices.IsEmpty())
 			{
@@ -346,14 +473,27 @@ namespace PCGMetadataPartitionCommon
 	}
 
 	TArray<UPCGData*> AttributePartition(const UPCGData* InData, const FPCGAttributePropertySelector& InSelector, FPCGContext* InOptionalContext)
+	{
+		const TArrayView<const FPCGAttributePropertySelector> ArrayView(&InSelector, 1);
+		if (const UPCGPointData* InPointData = Cast<UPCGPointData>(InData))
+		{
+			return AttributePointPartition(InPointData, ArrayView, InOptionalContext);
+		}
+		else
+		{
+			return AttributeParamSpatialPartition(InData, ArrayView, InOptionalContext);
+		}
+	}
+
+	TArray<UPCGData*> AttributePartition(const UPCGData* InData, const TArrayView<const FPCGAttributePropertySelector>& InSelectorArrayView, FPCGContext* InOptionalContext)
 	{
 		if (const UPCGPointData* InPointData = Cast<UPCGPointData>(InData))
 		{
-			return AttributePointPartition(InPointData, InSelector, InOptionalContext);
+			return AttributePointPartition(InPointData, InSelectorArrayView, InOptionalContext);
 		}
 		else
 		{
-			return AttributeParamSpatialPartition(InData, InSelector, InOptionalContext);
+			return AttributeParamSpatialPartition(InData, InSelectorArrayView, InOptionalContext);
 		}
 	}
 }
diff --git a/Engine/Plugins/PCG/Source/PCG/Private/Tests/Elements/PCGAttributePartitionTest.cpp b/Engine/Plugins/PCG/Source/PCG/Private/Tests/Elements/PCGAttributePartitionTest.cpp
--- a/Engine/Plugins/PCG/Source/PCG/Private/Tests/Elements/PCGAttributePartitionTest.cpp
+++ b/Engine/Plugins/PCG/Source/PCG/Private/Tests/Elements/PCGAttributePartitionTest.cpp
@@ -19,7 +19,7 @@ bool FPCGAttributePartition_Points::RunTest(const FString& Parameters)
 	UPCGMetadataPartitionSettings* Settings = PCGTestsCommon::GenerateSettings<UPCGMetadataPartitionSettings>(TestData);
 	check(Settings);
 
-	Settings->PartitionAttributeSource.SetPointProperty(EPCGPointProperties::Density);
+	Settings->PartitionAttributeSelectors.Emplace_GetRef().SetPointProperty(EPCGPointProperties::Density);
 
 	UPCGPointData* InputPointData = NewObject<UPCGPointData>();
 	TArray<FPCGPoint>& Points = InputPointData->GetMutablePoints();
@@ -86,7 +86,7 @@ bool FPCGAttributePartition_Order::RunTest(const FString& Parameters)
 	check(Settings);
 
 	const FName AttributeName = TEXT("Attr");
-	Settings->PartitionAttributeSource.SetAttributeName(AttributeName);
+	Settings->PartitionAttributeSelectors.Emplace_GetRef().SetAttributeName(AttributeName);
 
 	UPCGPointData* InputPointData = NewObject<UPCGPointData>();
 	FPCGMetadataAttribute<int>* Attribute = InputPointData->Metadata->CreateAttribute<int>(AttributeName, 0, false, false);
@@ -146,7 +146,6 @@ bool FPCGAttributePartition_Order::RunTest(const FString& Parameters)
 	return true;
 }
 
-
 bool FPCGAttributePartition_AttributeSet::RunTest(const FString& Parameters)
 {
 	PCGTestsCommon::FTestData TestData;
@@ -154,7 +153,7 @@ bool FPCGAttributePartition_AttributeSet::RunTest(const FString& Parameters)
 	check(Settings);
 
 	const FName InputAttributeName = TEXT("Double");
-	Settings->PartitionAttributeSource.SetAttributeName(InputAttributeName);
+	Settings->PartitionAttributeSelectors.Emplace_GetRef().SetAttributeName(InputAttributeName);
 
 	UPCGParamData* InputParam = NewObject<UPCGParamData>();
 	FPCGMetadataAttribute<double>* Attribute = InputParam->Metadata->CreateAttribute<double>(InputAttributeName, 0.0, true, false);
diff --git a/Engine/Plugins/PCG/Source/PCG/Public/Elements/Metadata/PCGMetadataPartition.h b/Engine/Plugins/PCG/Source/PCG/Public/Elements/Metadata/PCGMetadataPartition.h
--- a/Engine/Plugins/PCG/Source/PCG/Public/Elements/Metadata/PCGMetadataPartition.h
+++ b/Engine/Plugins/PCG/Source/PCG/Public/Elements/Metadata/PCGMetadataPartition.h
@@ -34,12 +34,21 @@ protected:
 	//~End UPCGSettings interface
 
 public:
-	UPROPERTY(BlueprintReadWrite, EditAnywhere, Category = Settings, meta = (PCG_Overridable, PCG_OverrideAliases = "PartitionAttribute"))
-	FPCGAttributePropertyInputSelector PartitionAttributeSource;
+	// TODO: Should be overridable once array override is supported
+	/** The data will be partitioned on these selected attributes. */
+	UPROPERTY(BlueprintReadWrite, EditAnywhere, Category = Settings)
+	TArray<FPCGAttributePropertyInputSelector> PartitionAttributeSelectors;
+
+	// TODO: Should be deprecated once array override is supported
+	UPROPERTY(meta = (PCG_Overridable))
+	FString PartitionAttributeNames;
 
 #if WITH_EDITORONLY_DATA
 	UPROPERTY()
 	FName PartitionAttribute_DEPRECATED = NAME_None;
+
+	UPROPERTY()
+	FPCGAttributePropertyInputSelector PartitionAttributeSource_DEPRECATED;
 #endif // WITH_EDITORONLY_DATA
 };
 
diff --git a/Engine/Plugins/PCG/Source/PCG/Public/Metadata/PCGMetadataPartitionCommon.h b/Engine/Plugins/PCG/Source/PCG/Public/Metadata/PCGMetadataPartitionCommon.h
--- a/Engine/Plugins/PCG/Source/PCG/Public/Metadata/PCGMetadataPartitionCommon.h
+++ b/Engine/Plugins/PCG/Source/PCG/Public/Metadata/PCGMetadataPartitionCommon.h
@@ -21,6 +21,16 @@ namespace PCGMetadataPartitionCommon
 	*/
 	PCG_API TArray<UPCGData*> AttributePartition(const UPCGData* InData, const FPCGAttributePropertySelector& InSelector, FPCGContext* InOptionalContext = nullptr);
 
+	/**
+	* Partition the incoming data on the given attribute.
+	* Will return an array of data. Each data will contain the entries (in stable order) for a given partition.
+	* @param InData - Data to partition.
+	* @param InSelectorArrayView - ArrayView of selectors on the attribute to partition.
+	* @param InOptionalContext - Optional context for logging.
+	* @returns Array of data.
+	*/
+	PCG_API TArray<UPCGData*> AttributePartition(const UPCGData* InData, const TArrayView<const FPCGAttributePropertySelector>& InSelectorArrayView, FPCGContext* InOptionalContext = nullptr);
+
 	/**
 	* Generic partition for the incoming data on the given attribute.
 	* Will return an array of bucket indices. Each bucket will contain the indices (in stable order) for a given partition.
@@ -30,4 +40,14 @@ namespace PCGMetadataPartitionCommon
 	* @returns Array of bucket indices.
 	*/
 	PCG_API TArray<TArray<int32>> AttributeGenericPartition(const UPCGData* InData, const FPCGAttributePropertySelector& InSelector, FPCGContext* InOptionalContext = nullptr);
+
+	/**
+	* Generic partition for the incoming data on the given array of attributes.
+	* Will return an array of bucket indices. Each bucket will contain the indices (in stable order) for a given partition.
+	* @param InData - Data to partition, need to support attributes (spatial data or attribute set).
+	* @param InSelectorArrayView - ArrayView of selectors on the attribute to partition.
+	* @param InOptionalContext - Optional context for logging.
+	* @returns Array of bucket indices.
+	*/
+	PCG_API TArray<TArray<int32>> AttributeGenericPartition(const UPCGData* InData, const TArrayView<const FPCGAttributePropertySelector>& InSelectorArrayView, FPCGContext* InOptionalContext = nullptr);
 }
\ No newline at end of file
