[PCG] Soften a hard assert to an ensure to fix a crash that occurs occasionally when connecting nodes

diff --git a/Engine/Plugins/PCG/Source/PCGEditor/Private/PCGEditorGraphSchema.cpp b/Engine/Plugins/PCG/Source/PCGEditor/Private/PCGEditorGraphSchema.cpp
--- a/Engine/Plugins/PCG/Source/PCGEditor/Private/PCGEditorGraphSchema.cpp
+++ b/Engine/Plugins/PCG/Source/PCGEditor/Private/PCGEditorGraphSchema.cpp
@@ -162,9 +162,13 @@ bool UPCGEditorGraphSchema::TryCreateConnectionInternal(UEdGraphPin* InA, UEdGra
 	UEdGraphPin* B = (InA->Direction == EGPD_Input) ? InA : InB;
 	check(A->Direction == EGPD_Output && B->Direction == EGPD_Input);
 
-	UEdGraphNode* NodeA = A->GetOwningNode();
-	UEdGraphNode* NodeB = B->GetOwningNode();
-	check(NodeA && NodeB);
+	UEdGraphNode* NodeA = A->GetOwningNodeUnchecked();
+	UEdGraphNode* NodeB = B->GetOwningNodeUnchecked();
+	if (!ensure(NodeA && NodeB))
+	{
+		// TODO: We've had crashes where one of these nodes was nullptr, we need to figure out why this can happen.
+		return false;
+	}
 
 	UPCGEditorGraphNodeBase* PCGEdGraphNodeA = CastChecked<UPCGEditorGraphNodeBase>(NodeA);
 	UPCGEditorGraphNodeBase* PCGEdGraphNodeB = CastChecked<UPCGEditorGraphNodeBase>(NodeB);
