[PCG] Add CreateSurfaceFromSpline node for creating an implicit surface from a closed spline
* Add UPCGSplineInteriorSurfaceData which is an implicit surface subtype that represents the interior of a spline as as surface
* Automatic compact conversion node when dragging spline pin to surface pin

diff --git a/Engine/Plugins/PCG/Source/PCG/Private/Data/PCGLandscapeSplineData.cpp b/Engine/Plugins/PCG/Source/PCG/Private/Data/PCGLandscapeSplineData.cpp
--- a/Engine/Plugins/PCG/Source/PCG/Private/Data/PCGLandscapeSplineData.cpp
+++ b/Engine/Plugins/PCG/Source/PCG/Private/Data/PCGLandscapeSplineData.cpp
@@ -244,7 +244,7 @@ const UPCGPointData* UPCGLandscapeSplineData::CreatePointData(FPCGContext* Conte
 	SamplerParams.Mode = EPCGSplineSamplingMode::Distance;
 	SamplerParams.Dimension = EPCGSplineSamplingDimension::OnHorizontal;
 
-	PCGSplineSampler::SampleLineData(/*LineData=*/this, /*InBoundingShapeData=*/nullptr, /*InProjectionTarget=*/nullptr, /*InProjectionParams=*/{}, SamplerParams, Data);
+	PCGSplineSamplerHelpers::SampleLineData(/*LineData=*/this, /*InBoundingShapeData=*/nullptr, /*InProjectionTarget=*/nullptr, /*InProjectionParams=*/{}, SamplerParams, Data);
 
 	UE_LOG(LogPCG, Verbose, TEXT("Landscape spline %s generated %d points"), *Spline->GetFName().ToString(), Points.Num());
 
diff --git a/Engine/Plugins/PCG/Source/PCG/Private/Data/PCGSplineData.cpp b/Engine/Plugins/PCG/Source/PCG/Private/Data/PCGSplineData.cpp
--- a/Engine/Plugins/PCG/Source/PCG/Private/Data/PCGSplineData.cpp
+++ b/Engine/Plugins/PCG/Source/PCG/Private/Data/PCGSplineData.cpp
@@ -52,6 +52,24 @@ void UPCGSplineData::Initialize(const TArray<FSplinePoint>& InSplinePoints, bool
 	CachedBounds = CachedBounds.TransformBy(InTransform);
 }
 
+void UPCGSplineData::Initialize(const FPCGSplineStruct& InSplineStruct)
+{
+	SplineStruct = InSplineStruct;
+	CachedBounds = SplineStruct.GetBounds();
+
+	// Expand bounds by the radius of points, otherwise sections of the curve that are close
+	// to the bounds will report an invalid density.
+	FVector SplinePointsRadius = FVector::ZeroVector;
+	const FInterpCurveVector& SplineScales = SplineStruct.GetSplinePointsScale();
+	for (const FInterpCurvePoint<FVector>& SplineScale : SplineScales.Points)
+	{
+		SplinePointsRadius = FVector::Max(SplinePointsRadius, SplineScale.OutVal.GetAbs());
+	}
+
+	CachedBounds = CachedBounds.ExpandBy(SplinePointsRadius, SplinePointsRadius);
+	CachedBounds = CachedBounds.TransformBy(SplineStruct.Transform);
+}
+
 void UPCGSplineData::ApplyTo(USplineComponent* InSplineComponent)
 {
 	SplineStruct.ApplyTo(InSplineComponent);
@@ -132,7 +150,7 @@ const UPCGPointData* UPCGSplineData::CreatePointData(FPCGContext* Context) const
 	FPCGSplineSamplerParams SamplerParams;
 	SamplerParams.Mode = EPCGSplineSamplingMode::Distance;
 
-	PCGSplineSampler::SampleLineData(this, /*InBoundingShape=*/nullptr, /*InProjectionTarget=*/nullptr, /*InProjectionParams=*/{}, SamplerParams, Data);
+	PCGSplineSamplerHelpers::SampleLineData(this, /*InBoundingShape=*/nullptr, /*InProjectionTarget=*/nullptr, /*InProjectionParams=*/{}, SamplerParams, Data);
 	UE_LOG(LogPCG, Verbose, TEXT("Spline generated %d points"), Data->GetPoints().Num());
 
 	return Data;
diff --git a/Engine/Plugins/PCG/Source/PCG/Private/Data/PCGSplineInteriorSurfaceData.cpp b/Engine/Plugins/PCG/Source/PCG/Private/Data/PCGSplineInteriorSurfaceData.cpp
new file mode 100644
--- /dev/null
+++ b/Engine/Plugins/PCG/Source/PCG/Private/Data/PCGSplineInteriorSurfaceData.cpp
@@ -0,0 +1,196 @@
+// Copyright Epic Games, Inc. All Rights Reserved.
+
+#include "Data/PCGSplineInteriorSurfaceData.h"
+
+#include "Data/PCGPointData.h"
+#include "Data/PCGSplineData.h"
+#include "Elements/PCGSplineSampler.h"
+
+#include UE_INLINE_GENERATED_CPP_BY_NAME(PCGSplineInteriorSurfaceData)
+
+void UPCGSplineInteriorSurfaceData::Initialize(const UPCGSplineData* InSplineData)
+{
+	check(InSplineData);
+
+	InitializeFromData(InSplineData);
+	SplineStruct = InSplineData->SplineStruct;
+
+	CacheData();
+}
+
+void UPCGSplineInteriorSurfaceData::PostLoad()
+{
+	Super::PostLoad();
+	CacheData();
+}
+
+void UPCGSplineInteriorSurfaceData::AddToCrc(FArchiveCrc32& Ar, bool bFullDataCrc) const
+{
+	Super::AddToCrc(Ar, bFullDataCrc);
+
+	// This data does not have a bespoke CRC implementation so just use a global unique data CRC.
+	AddUIDToCrc(Ar);
+}
+
+bool UPCGSplineInteriorSurfaceData::SamplePoint(const FTransform& InTransform, const FBox& InBounds, FPCGPoint& OutPoint, UPCGMetadata* OutMetadata) const
+{
+	if (!PointInsidePolygon(InTransform, InBounds))
+	{
+		return false;
+	}
+
+	const FVector InLocation = InTransform.GetLocation();
+
+	// Project the sampled point onto the approximate spline surface.
+	const FVector::FReal ProjectionHeight = PCGSplineSamplerHelpers::ProjectOntoSplineInteriorSurface(CachedSplinePoints, InLocation);
+	const FVector::FReal SampleMinHeight = InLocation.Z + InBounds.Min.Z;
+	const FVector::FReal SampleMaxHeight = InLocation.Z + InBounds.Max.Z;
+
+	// Discard if the sample point lies above/below the spline surface.
+	if (SampleMinHeight > ProjectionHeight || SampleMaxHeight < ProjectionHeight)
+	{
+		return false;
+	}
+
+	new(&OutPoint) FPCGPoint(InTransform, /*Density=*/1.0f, /*Seed=*/0);
+	OutPoint.SetLocalBounds(InBounds);
+
+	return true;
+}
+
+bool UPCGSplineInteriorSurfaceData::ProjectPoint(const FTransform& InTransform, const FBox& InBounds, const FPCGProjectionParams& InParams, FPCGPoint& OutPoint, UPCGMetadata* OutMetadata) const
+{
+	if (!PointInsidePolygon(InTransform, InBounds))
+	{
+		return false;
+	}
+
+	FTransform ProjectedTransform = InTransform;
+	const FVector InLocation = InTransform.GetLocation();
+	const FVector::FReal ProjectionHeight = (InParams.bProjectPositions || InParams.bProjectRotations) ? PCGSplineSamplerHelpers::ProjectOntoSplineInteriorSurface(CachedSplinePoints, InLocation) : 0.0f;
+
+	if (InParams.bProjectPositions)
+	{
+		// Project the sampled point onto the approximate spline surface.
+		ProjectedTransform.SetLocation(FVector(InLocation.X, InLocation.Y, ProjectionHeight));
+	}
+
+	if (InParams.bProjectRotations)
+	{
+		// Project rotation by sampling the rotation of the nearest point on the spline.
+		float OutDistanceSquared, OutSegment;
+		const float NearestSplineKey = SplineStruct.SplineCurves.Position.FindNearest(FVector(InLocation.X, InLocation.Y, ProjectionHeight), OutDistanceSquared, OutSegment);
+		ProjectedTransform.SetRotation(SplineStruct.GetQuaternionAtSplineInputKey(NearestSplineKey, ESplineCoordinateSpace::Local));
+	}
+
+	new(&OutPoint) FPCGPoint(ProjectedTransform, /*Density=*/1.0f, /*Seed=*/0);
+	OutPoint.SetLocalBounds(InBounds);
+
+	return true;
+}
+
+const UPCGPointData* UPCGSplineInteriorSurfaceData::CreatePointData(FPCGContext* Context) const
+{
+	TRACE_CPUPROFILER_EVENT_SCOPE(UPCGSplineData::CreatePointData);
+
+	UPCGPointData* Data = NewObject<UPCGPointData>();
+	Data->InitializeFromData(this);
+
+	FPCGSplineSamplerParams SamplerParams;
+	SamplerParams.Dimension = EPCGSplineSamplingDimension::OnInterior;
+	SamplerParams.bProjectOntoSurface = true;
+
+	UPCGSplineData* SplineData = NewObject<UPCGSplineData>();
+	SplineData->Initialize(SplineStruct);
+
+	// TODO: Ideally SampleInteriorData could just consume a UPCGSplineInteriorSurfaceData or FPCGSplineStruct, to avoid the extra copy on the spline struct.
+	PCGSplineSamplerHelpers::SampleInteriorData(Context, SplineData, /*InBoundingShape=*/nullptr, /*InProjectionTarget=*/nullptr, /*InProjectionParams=*/{}, SamplerParams, Data);
+	UE_LOG(LogPCG, Verbose, TEXT("SplineInteriorSurface generated %d points."), Data->GetPoints().Num());
+
+	SplineData->MarkAsGarbage();
+
+	return Data;
+}
+
+UPCGSpatialData* UPCGSplineInteriorSurfaceData::CopyInternal() const
+{
+	UPCGSplineInteriorSurfaceData* NewData = NewObject<UPCGSplineInteriorSurfaceData>();
+
+	CopyBaseSurfaceData(NewData);
+	NewData->SplineStruct = SplineStruct;
+	NewData->CachedBounds = CachedBounds;
+	NewData->CachedSplinePoints = CachedSplinePoints;
+	NewData->CachedSplinePoints2D = CachedSplinePoints2D;
+#if WITH_EDITOR
+	NewData->bNeedsToCache = false;
+#endif
+
+	return NewData;
+}
+
+void UPCGSplineInteriorSurfaceData::CacheData()
+{
+#if WITH_EDITOR
+	bNeedsToCache = false;
+#endif
+
+	CachedBounds = SplineStruct.GetBounds();
+
+	// Expand bounds by the radius of points, otherwise sections of the curve that are close
+	// to the bounds will report an invalid density.
+	FVector SplinePointsRadius = FVector::ZeroVector;
+	const FInterpCurveVector& SplineScales = SplineStruct.GetSplinePointsScale();
+	for (const FInterpCurvePoint<FVector>& SplineScale : SplineScales.Points)
+	{
+		SplinePointsRadius = FVector::Max(SplinePointsRadius, SplineScale.OutVal.GetAbs());
+	}
+
+	CachedBounds = CachedBounds.ExpandBy(SplinePointsRadius, SplinePointsRadius);
+	CachedBounds = CachedBounds.TransformBy(SplineStruct.Transform);
+
+	FPCGSplineSamplerParams SamplerParams;
+	SamplerParams.Mode = EPCGSplineSamplingMode::Subdivision;
+	SamplerParams.SubdivisionsPerSegment = 5;
+
+	UPCGPointData* PointData = NewObject<UPCGPointData>();
+	UPCGSplineData* SplineData = NewObject<UPCGSplineData>();
+	SplineData->Initialize(SplineStruct);
+
+	// TODO: Ideally SampleLineData could just consume a UPCGSplineInteriorSurfaceData or FPCGSplineStruct, to avoid the extra copy on the spline struct.
+	// TODO: It might be preferable to directly sample the spline than use the spline sampler, since it does a lot of work we don't care about.
+	PCGSplineSamplerHelpers::SampleLineData(SplineData, /*InBoundingShape=*/nullptr, /*InProjectionTarget=*/nullptr, /*InProjectionParams=*/{}, SamplerParams, PointData);
+
+	// Cache the points which describe the polygon of our spline.
+	for (const FPCGPoint& Point : PointData->GetPoints())
+	{
+		CachedSplinePoints.Add(Point.Transform.GetLocation());
+		CachedSplinePoints2D.Add(FVector2D(Point.Transform.GetLocation()));
+	}
+
+	SplineData->MarkAsGarbage();
+	PointData->MarkAsGarbage();
+}
+
+bool UPCGSplineInteriorSurfaceData::PointInsidePolygon(const FTransform& InTransform, const FBox& InBounds) const
+{
+#if WITH_EDITOR
+	check(!bNeedsToCache);
+#endif
+
+	const FBox TransformedBounds = InBounds.TransformBy(InTransform);
+
+	// Test point bounds against the spline bounds.
+	if (CachedBounds.ComputeSquaredDistanceToBox(TransformedBounds) > 0.0)
+	{
+		return false;
+	}
+
+	const FVector PointLocation = InTransform.GetLocation();
+
+	// Maximum distance a ray needs to travel to guarantee exiting the polygon at its widest point from the sample location.
+	const FVector::FReal MaxDistance = CachedBounds.Max.X - PointLocation.X + UE_KINDA_SMALL_NUMBER;
+
+	// Test sample location against spline interior.
+	// TODO: This could also consume a transform + bounds to accept rotated bounds that overlap the polygon.
+	return PCGSplineSamplerHelpers::PointInsidePolygon2D(CachedSplinePoints2D, FVector2D(PointLocation), MaxDistance);
+}
diff --git a/Engine/Plugins/PCG/Source/PCG/Private/Elements/PCGCreateSurfaceFromSpline.cpp b/Engine/Plugins/PCG/Source/PCG/Private/Elements/PCGCreateSurfaceFromSpline.cpp
new file mode 100644
--- /dev/null
+++ b/Engine/Plugins/PCG/Source/PCG/Private/Elements/PCGCreateSurfaceFromSpline.cpp
@@ -0,0 +1,76 @@
+// Copyright Epic Games, Inc. All Rights Reserved.
+
+#include "Elements/PCGCreateSurfaceFromSpline.h"
+
+#include "PCGContext.h"
+#include "Data/PCGSplineData.h"
+#include "Data/PCGSplineInteriorSurfaceData.h"
+
+#include UE_INLINE_GENERATED_CPP_BY_NAME(PCGCreateSurfaceFromSpline)
+
+#define LOCTEXT_NAMESPACE "PCGCreateSurfaceFromSplineElement"
+
+#if WITH_EDITOR
+FText UPCGCreateSurfaceFromSplineSettings::GetNodeTooltipText() const
+{
+	// TODO: We should refactor interior sampling to allow for local-space projection to yield more correct sampling for rotated splines.
+	return LOCTEXT("NodeTooltip", "Create an implicit surface for each given spline. The surface is given by the top-down 2D projection of the spline. Each spline must be closed.");
+}
+#endif
+
+TArray<FPCGPinProperties> UPCGCreateSurfaceFromSplineSettings::InputPinProperties() const
+{
+	TArray<FPCGPinProperties> PinProperties;
+	PinProperties.Emplace(PCGPinConstants::DefaultInputLabel, EPCGDataType::Spline, /*bAllowMultipleConnections=*/true, /*bAllowMultipleData=*/true);
+	return PinProperties;
+}
+
+TArray<FPCGPinProperties> UPCGCreateSurfaceFromSplineSettings::OutputPinProperties() const
+{
+	TArray<FPCGPinProperties> PinProperties;
+	PinProperties.Emplace(PCGPinConstants::DefaultOutputLabel, EPCGDataType::Surface, /*bInAllowMultipleConnections=*/true, /*bAllowMultipleData=*/true);
+	return PinProperties;
+}
+
+FPCGElementPtr UPCGCreateSurfaceFromSplineSettings::CreateElement() const
+{
+	return MakeShared<FPCGCreateSurfaceFromSplineElement>();
+}
+
+bool FPCGCreateSurfaceFromSplineElement::ExecuteInternal(FPCGContext* Context) const
+{
+	TRACE_CPUPROFILER_EVENT_SCOPE(FPCGCreateSurfaceFromSplineElement::Execute);
+
+	const UPCGCreateSurfaceFromSplineSettings* Settings = Context->GetInputSettings<UPCGCreateSurfaceFromSplineSettings>();
+	check(Settings);
+
+	const TArray<FPCGTaggedData> Inputs = Context->InputData.GetInputsByPin(PCGPinConstants::DefaultInputLabel);
+	TArray<FPCGTaggedData>& Outputs = Context->OutputData.TaggedData;
+
+	for (const FPCGTaggedData& Input : Inputs)
+	{
+		const UPCGSplineData* SplineData = Cast<UPCGSplineData>(Input.Data);
+
+		if (!SplineData)
+		{
+			PCGE_LOG(Error, GraphAndLog, LOCTEXT("InvalidInputData", "Invalid input. Must be a spline data."));
+			continue;
+		}
+
+		if (!SplineData->IsClosed())
+		{
+			PCGE_LOG(Error, GraphAndLog, LOCTEXT("InvalidOpenSpline", "Unable to create an implicit surface from an open spline. Make sure your splines are closed."));
+			continue;
+		}
+
+		UPCGSplineInteriorSurfaceData* SurfaceData = NewObject<UPCGSplineInteriorSurfaceData>();
+		SurfaceData->Initialize(SplineData);
+
+		FPCGTaggedData& Output = Outputs.Add_GetRef(Input);
+		Output.Data = SurfaceData;
+	}
+
+	return true;
+}
+
+#undef LOCTEXT_NAMESPACE
diff --git a/Engine/Plugins/PCG/Source/PCG/Private/Elements/PCGSplineSampler.cpp b/Engine/Plugins/PCG/Source/PCG/Private/Elements/PCGSplineSampler.cpp
--- a/Engine/Plugins/PCG/Source/PCG/Private/Elements/PCGSplineSampler.cpp
+++ b/Engine/Plugins/PCG/Source/PCG/Private/Elements/PCGSplineSampler.cpp
@@ -302,10 +302,39 @@ namespace PCGSplineSamplerHelpers
 		// Odd number of intersections means we are inside the polygon
 		return (IntersectionPoints.Num() % 2) == 1;
 	}
-}
 
-namespace PCGSplineSampler
-{
+	/** Projects a point in space onto the approximated surface defined by a closed spline. */
+	FVector::FReal ProjectOntoSplineInteriorSurface(const TArray<FVector>& SplinePoints, const FVector& PointToProject)
+	{
+		// Compute average Z value weighted by 1 / Distance^2
+		FVector::FReal SumZ = 0.f;
+		FVector::FReal SumWeights = 0.f;
+
+		for (int32 PointIndex = 0; PointIndex < SplinePoints.Num(); ++PointIndex)
+		{
+			const FVector& Point = SplinePoints[PointIndex];
+			// TODO: It would be more accurate to use distance to the polyline instead of distance to the polyline points,
+			// however it would also be much more expensive. Perhaps worth investigating when Params.bTreatSplineAsPolyline is true.
+			const FVector::FReal DistanceSquared = FVector::DistSquaredXY(PointToProject, Point); 
+
+			// If sample point overlaps exactly with a border point, then that must be the height.
+			if (FMath::IsNearlyZero(DistanceSquared))
+			{
+				return Point.Z;
+			}
+
+			// TODO[UE-205462]: It would be ideal to find a better constraint than inverse squared distance, which produces boundary artifacts for undersampled splines.
+			const FVector::FReal Weight = 1.f / DistanceSquared;
+
+			SumWeights += Weight;
+			SumZ += Point.Z * Weight;
+		}
+
+		// If the weights sum to zero then there were no points to sample, and therefore nothing to project to.
+		// TODO: This is not robust for interior points which are enormously far away from the spline and could collapse to 0 incorrectly.
+		return FMath::IsNearlyZero(SumWeights) ? 0.0 : SumZ / SumWeights;
+	}
+
 	struct FSamplerResult
 	{
 		FTransform LocalTransform;
@@ -754,7 +783,6 @@ namespace PCGSplineSampler
 		int NumHeightSteps;
 	};
 
-	/** Samples on spline or within volume around it. */
 	void SampleLineData(const UPCGPolyLineData* LineData, const UPCGSpatialData* InBoundingShapeData, const UPCGSpatialData* InProjectionTarget, const FPCGProjectionParams& InProjectionParams, const FPCGSplineSamplerParams& Params, UPCGPointData* OutPointData)
 	{
 		check(LineData && OutPointData);
@@ -851,7 +879,6 @@ namespace PCGSplineSampler
 		}
 	}
  
-	/** Samples 2D region bounded by spline. */
 	void SampleInteriorData(FPCGContext* Context, const UPCGPolyLineData* LineData, const UPCGSpatialData* InBoundingShape, const UPCGSpatialData* InProjectionTarget, const FPCGProjectionParams& InProjectionParams, const FPCGSplineSamplerParams& Params, UPCGPointData* OutPointData)
 	{
 		check(Context && LineData && OutPointData);
@@ -1052,8 +1079,6 @@ namespace PCGSplineSampler
 			const FVector::FReal LocalMinY = MinY + StartIterationIndex * Params.InteriorSampleSpacing;
 			const FVector::FReal LocalMaxY = bIsLastIteration ? (MaxY + UE_KINDA_SMALL_NUMBER) : (MinY + EndIterationIndex * Params.InteriorSampleSpacing);
 
-			TArray<FVector::FReal> DistancesSquaredOnXY;
-
 			// Point sampling
 			for(FVector::FReal Y = LocalMinY; Y < LocalMaxY; Y += Params.InteriorSampleSpacing)
 			{
@@ -1093,38 +1118,7 @@ namespace PCGSplineSampler
 
 						if (bProjectOntoSurface)
 						{
-							// Precompute 2D distance to every spline point
-							DistancesSquaredOnXY.Reset(SplineSamplePoints.Num());
-							for (const FVector& SplinePoint : SplineSamplePoints)
-							{
-								FVector::FReal DistanceSquared = FVector::DistSquaredXY(SurfaceLocation, SplinePoint);
-								DistancesSquaredOnXY.Add(DistanceSquared);
-							}
-
-							// Compute average Z value weighted by 1 / Distance^2
-							FVector::FReal SumZ = 0.f;
-							FVector::FReal SumWeights = 0.f;
-							for (int32 PointIndex = 0; PointIndex < SplineSamplePoints.Num(); ++PointIndex)
-							{
-								const FVector::FReal DistanceSquared = DistancesSquaredOnXY[PointIndex];
-
-								// If sample point overlaps exactly with a border point, then that must be the height
-								if (FMath::IsNearlyZero(DistanceSquared))
-								{
-									SumZ = SplineSamplePoints[PointIndex].Z;
-									SumWeights = 1.f;
-									break;
-								}
-
-								// TODO: it would be more accurate to use distance to the polyline instead of distance to the polyline points,
-								// however it would also be much more expensive. Perhaps worth investigating when Params.bTreatSplineAsPolyline is true
-								const FVector::FReal Weight = 1.f / DistanceSquared;
-
-								SumWeights += Weight;
-								SumZ += SplineSamplePoints[PointIndex].Z * Weight;
-							}
-
-							SurfaceLocation.Z = SumZ / SumWeights;
+							SurfaceLocation.Z = ProjectOntoSplineInteriorSurface(SplineSamplePoints, SurfaceLocation);
 						}
 
 						// if bTreatAsPolyline, then we shouldnt use this, we should use nearest point on the polygon line segments
@@ -1362,7 +1356,7 @@ bool FPCGSplineSamplerElement::ExecuteInternal(FPCGContext* Context) const
 		}
 
 		// TODO: do something for point data approximations
-		const UPCGPolyLineData* LineData = PCGSplineSampler::GetPolyLineData(SpatialData);
+		const UPCGPolyLineData* LineData = PCGSplineSamplerHelpers::GetPolyLineData(SpatialData);
 		if (!LineData)
 		{
 			continue;
@@ -1385,11 +1379,11 @@ bool FPCGSplineSamplerElement::ExecuteInternal(FPCGContext* Context) const
 
 		if (SamplerParams.Dimension == EPCGSplineSamplingDimension::OnInterior)
 		{
-			PCGSplineSampler::SampleInteriorData(Context, LineData, BoundingShape, ProjectionTarget, ProjectionParams, SamplerParams, SampledPointData);
+			PCGSplineSamplerHelpers::SampleInteriorData(Context, LineData, BoundingShape, ProjectionTarget, ProjectionParams, SamplerParams, SampledPointData);
 		}
 		else
 		{
-			PCGSplineSampler::SampleLineData(LineData, BoundingShape, ProjectionTarget, ProjectionParams, SamplerParams, SampledPointData);
+			PCGSplineSamplerHelpers::SampleLineData(LineData, BoundingShape, ProjectionTarget, ProjectionParams, SamplerParams, SampledPointData);
 		}
 	}
 
diff --git a/Engine/Plugins/PCG/Source/PCG/Private/PCGPin.cpp b/Engine/Plugins/PCG/Source/PCG/Private/PCGPin.cpp
--- a/Engine/Plugins/PCG/Source/PCG/Private/PCGPin.cpp
+++ b/Engine/Plugins/PCG/Source/PCG/Private/PCGPin.cpp
@@ -356,6 +356,12 @@ bool UPCGPin::IsCompatible(const UPCGPin* OtherPin) const
 		return true;
 	}
 
+	// Spline -> Surface conversion
+	if (UpstreamTypes == EPCGDataType::Spline && DownstreamTypes == EPCGDataType::Surface)
+	{
+		return true;
+	}
+
 	// Otherwise allow if there is overlap. Don't detect wide -> narrow issues - conversion nodes deal with that
 	return !!(UpstreamTypes & DownstreamTypes);
 }
@@ -412,6 +418,12 @@ EPCGTypeConversion UPCGPin::GetRequiredTypeConversion(const UPCGPin* InOtherPin)
 		return EPCGTypeConversion::CollapseToPoint;
 	}
 
+	// Spline -> Surface conversion
+	if (UpstreamTypes == EPCGDataType::Spline && DownstreamTypes == EPCGDataType::Surface)
+	{
+		return EPCGTypeConversion::SplineToSurface;
+	}
+
 	// Any or Spatial -> Concrete - "Make Concrete" conversion. We decided to support Any as it is a superset of Concrete,
 	// and some pins are Any until they dynamically narrow, and other pins may stay Any but still want to be supported.
 	bool bUpstreamIsSpatialOrAny = (UpstreamTypes == EPCGDataType::Any) || (UpstreamTypes == EPCGDataType::Spatial);
diff --git a/Engine/Plugins/PCG/Source/PCG/Public/Data/PCGSplineData.h b/Engine/Plugins/PCG/Source/PCG/Public/Data/PCGSplineData.h
--- a/Engine/Plugins/PCG/Source/PCG/Public/Data/PCGSplineData.h
+++ b/Engine/Plugins/PCG/Source/PCG/Public/Data/PCGSplineData.h
@@ -23,6 +23,7 @@ class PCG_API UPCGSplineData : public UPCGPolyLineData
 public:
 	void Initialize(USplineComponent* InSpline);
 	void Initialize(const TArray<FSplinePoint>& InSplinePoints, bool bInClosedLoop, const FTransform& InTransform);
+	void Initialize(const FPCGSplineStruct& InSplineStruct);
 	void ApplyTo(USplineComponent* InSpline);
 
 	// ~Begin UPCGData interface
diff --git a/Engine/Plugins/PCG/Source/PCG/Public/Data/PCGSplineInteriorSurfaceData.h b/Engine/Plugins/PCG/Source/PCG/Public/Data/PCGSplineInteriorSurfaceData.h
new file mode 100644
--- /dev/null
+++ b/Engine/Plugins/PCG/Source/PCG/Public/Data/PCGSplineInteriorSurfaceData.h
@@ -0,0 +1,77 @@
+// Copyright Epic Games, Inc. All Rights Reserved.
+
+#pragma once
+
+#include "Data/PCGSurfaceData.h"
+#include "Data/PCGSplineStruct.h"
+
+#include "PCGSplineInteriorSurfaceData.generated.h"
+
+class UPCGSpatialData;
+class UPCGSplineData;
+struct FPCGProjectionParams;
+
+/**
+* Represents a surface implicitly using the top-down 2D projection of a closed spline.
+*/
+UCLASS(BlueprintType, ClassGroup = (Procedural))
+class PCG_API UPCGSplineInteriorSurfaceData : public UPCGSurfaceData
+{
+	GENERATED_BODY()
+
+public:
+	void Initialize(const UPCGSplineData* InSplineData);
+
+	// ~Begin UObject interface
+	virtual void PostLoad() override;
+	// ~End UObject interface
+
+	// ~Begin UPCGData interface
+	virtual void AddToCrc(FArchiveCrc32& Ar, bool bFullDataCrc) const override;
+	// ~End UPCGData interface
+
+	//~Begin UPCGSpatialData interface
+	virtual FBox GetBounds() const override { return CachedBounds; }
+	virtual FBox GetStrictBounds() const override { return CachedBounds; }
+	virtual bool IsBounded() const override { return !!CachedBounds.IsValid; }
+	virtual bool SamplePoint(const FTransform& Transform, const FBox& Bounds, FPCGPoint& OutPoint, UPCGMetadata* OutMetadata) const override;
+	virtual bool ProjectPoint(const FTransform& InTransform, const FBox& InBounds, const FPCGProjectionParams& InParams, FPCGPoint& OutPoint, UPCGMetadata* OutMetadata) const override;
+	virtual bool HasNonTrivialTransform() const override { return true; }
+protected:
+	virtual UPCGSpatialData* CopyInternal() const override;
+	//~End UPCGSpatialData interface
+
+public:
+	// ~Begin UPCGSpatialDataWithPointCache interface
+	virtual const UPCGPointData* CreatePointData(FPCGContext* Context) const override;
+	// ~End UPCGConcreteDataWithPointCache interface
+
+protected:
+	/** Recompute cached data used for sampling. */
+	void CacheData();
+
+	/** True if the given location falls in the top-down projection of the polygon given by our cached spline points. */
+	bool PointInsidePolygon(const FTransform& InTransform, const FBox& InBounds) const;
+
+protected:
+	/** Minimal data needed to replicate the behavior from USplineComponent. */
+	UPROPERTY()
+	FPCGSplineStruct SplineStruct;
+
+	UPROPERTY(Transient)
+	FBox CachedBounds = FBox(EForceInit::ForceInit);
+
+	/** Cached list of points to describe the polygon given by the spline. */
+	UPROPERTY(Transient)
+	TArray<FVector> CachedSplinePoints;
+
+	/** Cached list of 2D points to describe the polygon given by the spline. */
+	UPROPERTY(Transient)
+	TArray<FVector2D> CachedSplinePoints2D;
+
+#if WITH_EDITORONLY_DATA
+	/** Flag that indicates when we need to recompute our cached information. Preferable to serializing data that can be recreated on demand. */
+	UPROPERTY(Transient)
+	bool bNeedsToCache = true;
+#endif
+};
diff --git a/Engine/Plugins/PCG/Source/PCG/Public/Elements/PCGCreateSurfaceFromSpline.h b/Engine/Plugins/PCG/Source/PCG/Public/Elements/PCGCreateSurfaceFromSpline.h
new file mode 100644
--- /dev/null
+++ b/Engine/Plugins/PCG/Source/PCG/Public/Elements/PCGCreateSurfaceFromSpline.h
@@ -0,0 +1,47 @@
+// Copyright Epic Games, Inc. All Rights Reserved.
+
+#pragma once
+
+#include "PCGSettings.h"
+
+#include "PCGCreateSurfaceFromSpline.generated.h"
+
+/**
+ * Create an implicit surface for each given spline. The surface is given by the top-down 2D projection of the spline. Each spline must be closed.
+ * 
+ * Note that by default a low resolution polyline is used to represent the spline. If you observe that your interior samples are not sufficiently accurate,
+ * you can resample the spline with a SplineSampler at your desired resolution, and convert the points back to a spline with the CreateSpline node.
+ */
+UCLASS(BlueprintType, ClassGroup = (Procedural))
+class PCG_API UPCGCreateSurfaceFromSplineSettings : public UPCGSettings
+{
+	GENERATED_BODY()
+
+public:
+	//~Begin UPCGSettings interface
+#if WITH_EDITOR
+	virtual FName GetDefaultNodeName() const override { return FName(TEXT("CreateSurfaceFromSpline")); }
+	virtual FText GetDefaultNodeTitle() const override { return NSLOCTEXT("PCGCreateSurfaceFromSplineElement", "NodeTitle", "Create Surface From Spline"); }
+	virtual FText GetNodeTooltipText() const override;
+	virtual EPCGSettingsType GetType() const override { return EPCGSettingsType::Spatial; }
+	virtual bool ShouldDrawNodeCompact() const override { return bShouldDrawNodeCompact; }
+#endif
+
+protected:
+	virtual TArray<FPCGPinProperties> InputPinProperties() const override;
+	virtual TArray<FPCGPinProperties> OutputPinProperties() const override;
+	virtual FPCGElementPtr CreateElement() const override;
+	//~End UPCGSettings interface
+
+#if WITH_EDITORONLY_DATA
+public:
+	UPROPERTY()
+	bool bShouldDrawNodeCompact = false;
+#endif
+};
+
+class FPCGCreateSurfaceFromSplineElement : public IPCGElement
+{
+protected:
+	virtual bool ExecuteInternal(FPCGContext* Context) const override;
+};
diff --git a/Engine/Plugins/PCG/Source/PCG/Public/Elements/PCGSplineSampler.h b/Engine/Plugins/PCG/Source/PCG/Public/Elements/PCGSplineSampler.h
--- a/Engine/Plugins/PCG/Source/PCG/Public/Elements/PCGSplineSampler.h
+++ b/Engine/Plugins/PCG/Source/PCG/Public/Elements/PCGSplineSampler.h
@@ -8,6 +8,7 @@
 
 #include "PCGSplineSampler.generated.h"
 
+struct FPCGContext;
 struct FPCGProjectionParams;
 class UPCGPolyLineData;
 class UPCGSpatialData;
@@ -170,16 +171,19 @@ struct PCG_API FPCGSplineSamplerParams
 
 namespace PCGSplineSamplerHelpers
 {
-	/** Tests if a point lies inside the given polygon by casting a ray to MaxDistance and counting the intersections */
-	bool PointInsidePolygon2D(const TArray<FVector2D>& PolygonPoints, const FVector2D& Point, FVector::FReal MaxDistance);
-}
-
-namespace PCGSplineSampler
-{
+	/** Samples on spline or within volume around it. */
 	void SampleLineData(const UPCGPolyLineData* LineData, const UPCGSpatialData* InBoundingShape, const UPCGSpatialData* InProjectionTarget, const FPCGProjectionParams& InProjectionParams, const FPCGSplineSamplerParams& Params, UPCGPointData* OutPointData);
 
-	void SampleInteriorData(const UPCGPolyLineData* LineData, const UPCGSpatialData* InBoundingShape, const UPCGSpatialData* InProjectionTarget, const FPCGProjectionParams& InProjectionParams, const FPCGSplineSamplerParams& Params, UPCGPointData* OutPointData);
+	/** Samples 2D region bounded by spline. */
+	void SampleInteriorData(FPCGContext* Context, const UPCGPolyLineData* LineData, const UPCGSpatialData* InBoundingShape, const UPCGSpatialData* InProjectionTarget, const FPCGProjectionParams& InProjectionParams, const FPCGSplineSamplerParams& Params, UPCGPointData* OutPointData);
+
 	const UPCGPolyLineData* GetPolyLineData(const UPCGSpatialData* InSpatialData);
+
+	/** Tests if a point lies inside the given polygon by casting a ray to MaxDistance and counting the intersections. */
+	bool PointInsidePolygon2D(const TArray<FVector2D>& PolygonPoints, const FVector2D& Point, FVector::FReal MaxDistance);
+
+	/** Projects a point in space onto the approximated surface defined by a closed spline. Returns the height of the point after projection. */
+	FVector::FReal ProjectOntoSplineInteriorSurface(const TArray<FVector>& SplinePoints, const FVector& PointToProject);
 }
 
 UCLASS(BlueprintType, ClassGroup = (Procedural))
diff --git a/Engine/Plugins/PCG/Source/PCG/Public/PCGPin.h b/Engine/Plugins/PCG/Source/PCG/Public/PCGPin.h
--- a/Engine/Plugins/PCG/Source/PCG/Public/PCGPin.h
+++ b/Engine/Plugins/PCG/Source/PCG/Public/PCGPin.h
@@ -72,6 +72,7 @@ enum class EPCGTypeConversion : uint8
 	CollapseToPoint,
 	Filter,
 	MakeConcrete,
+	SplineToSurface,
 	Failed,
 };
 
diff --git a/Engine/Plugins/PCG/Source/PCGEditor/Private/PCGEditorGraphSchema.cpp b/Engine/Plugins/PCG/Source/PCGEditor/Private/PCGEditorGraphSchema.cpp
--- a/Engine/Plugins/PCG/Source/PCGEditor/Private/PCGEditorGraphSchema.cpp
+++ b/Engine/Plugins/PCG/Source/PCGEditor/Private/PCGEditorGraphSchema.cpp
@@ -6,6 +6,7 @@
 #include "PCGGraph.h"
 #include "PCGPin.h"
 #include "Elements/PCGCollapseElement.h"
+#include "Elements/PCGCreateSurfaceFromSpline.h"
 #include "Elements/PCGExecuteBlueprint.h"
 #include "Elements/PCGFilterByType.h"
 #include "Elements/PCGMakeConcreteElement.h"
@@ -278,6 +279,16 @@ bool UPCGEditorGraphSchema::TryCreateConnectionInternal(UEdGraphPin* InA, UEdGra
 
 		return ConnectViaIntermediate(ConversionPCGNode);
 	}
+	else if (Conversion == EPCGTypeConversion::SplineToSurface)
+	{
+		UPCGSettings* NodeSettings = nullptr;
+		UPCGNode* ConversionPCGNode = PCGGraph->AddNodeOfType(UPCGCreateSurfaceFromSplineSettings::StaticClass(), NodeSettings);
+
+		UPCGCreateSurfaceFromSplineSettings* Settings = CastChecked<UPCGCreateSurfaceFromSplineSettings>(NodeSettings);
+		Settings->bShouldDrawNodeCompact = true;
+
+		return ConnectViaIntermediate(ConversionPCGNode);
+	}
 	else
 	{
 		const bool bModified = Super::TryCreateConnection(InA, InB);
