[PCG] Hotfix for deprecation bug in Attribute Partition

#rnx

diff --git a/Engine/Plugins/PCG/Source/PCG/Private/Elements/Metadata/PCGMetadataPartition.cpp b/Engine/Plugins/PCG/Source/PCG/Private/Elements/Metadata/PCGMetadataPartition.cpp
--- a/Engine/Plugins/PCG/Source/PCG/Private/Elements/Metadata/PCGMetadataPartition.cpp
+++ b/Engine/Plugins/PCG/Source/PCG/Private/Elements/Metadata/PCGMetadataPartition.cpp
@@ -45,11 +45,11 @@ void UPCGMetadataPartitionSettings::PostLoad()
 		PartitionAttribute_DEPRECATED = NAME_None;
 	}
 
-	if (PartitionAttributeSource_DEPRECATED.GetAttributeName() != PCGMetadataAttributeConstants::LastAttributeName)
+	if (PartitionAttributeSource_DEPRECATED != FPCGAttributePropertyInputSelector())
 	{
 		PartitionAttributeSelectors.Empty();
 		PartitionAttributeSelectors.Emplace(PartitionAttributeSource_DEPRECATED);
-		PartitionAttributeSource_DEPRECATED.SetAttributeName(PCGMetadataAttributeConstants::LastAttributeName);
+		PartitionAttributeSource_DEPRECATED = FPCGAttributePropertyInputSelector();
 	}
 #endif // WITH_EDITOR
 }
