[PCG] Display grid size label on nodes when inspecting higen graphs/components.

Cleanout old "highlight" visualization for parent grid nodes.

diff --git a/Engine/Plugins/PCG/Source/PCGEditor/Private/PCGEditorGraphNodeBase.cpp b/Engine/Plugins/PCG/Source/PCGEditor/Private/PCGEditorGraphNodeBase.cpp
--- a/Engine/Plugins/PCG/Source/PCGEditor/Private/PCGEditorGraphNodeBase.cpp
+++ b/Engine/Plugins/PCG/Source/PCGEditor/Private/PCGEditorGraphNodeBase.cpp
@@ -492,9 +492,6 @@ EPCGChangeType UPCGEditorGraphNodeBase::UpdateStructuralVisualization(UPCGCompon
 
 	EPCGChangeType ChangeType = EPCGChangeType::None;
 
-	bool bShouldDisplayAsDisabled = false;
-	bool bShouldDisplayAsHighlighted = false;
-
 	const uint64 NewInactiveMask = (InComponentBeingDebugged && InStackBeingInspected) ? InComponentBeingDebugged->GetNodeInactivePinMask(PCGNode, *InStackBeingInspected) : 0;
 	if (NewInactiveMask != InactiveOutputPinMask)
 	{
@@ -502,31 +499,33 @@ EPCGChangeType UPCGEditorGraphNodeBase::UpdateStructuralVisualization(UPCGCompon
 		ChangeType |= EPCGChangeType::Cosmetic;
 	}
 
-	bool bIsCulled = !PCGEditorGraphNodeBase::ShouldDisplayAsActive(this, InComponentBeingDebugged, InStackBeingInspected);
+	// Check top graph for higen enable - subgraphs always inherit higen state from the top graph.
+	const UPCGGraph* TopGraph = InComponentBeingDebugged ? InComponentBeingDebugged->GetGraph() : nullptr;
+	const bool HiGenEnabled = TopGraph && TopGraph->IsHierarchicalGenerationEnabled();
+
+	// Set the inspected grid size - this is used for grid size visualization.
+	uint32 InspectingGridSize = PCGHiGenGrid::UninitializedGridSize();
+	EPCGHiGenGrid InspectingGrid = EPCGHiGenGrid::Uninitialized;
+	if (TopGraph && TopGraph->IsHierarchicalGenerationEnabled())
+	{
+		InspectingGridSize = InComponentBeingDebugged->GetGenerationGridSize();
+		InspectingGrid = InComponentBeingDebugged->GetGenerationGrid();
+	}
+
+	if (InspectedGenerationGrid != InspectingGrid)
+	{
+		InspectedGenerationGrid = InspectingGrid;
+		ChangeType |= EPCGChangeType::Cosmetic;
+	}
 
-	const bool HiGenEnabled = Graph->IsHierarchicalGenerationEnabled();
-	const uint32 InspectingGridSize = InComponentBeingDebugged ? InComponentBeingDebugged->GetGenerationGridSize() : PCGHiGenGrid::UninitializedGridSize();
+	bool bShouldDisplayAsDisabled = false;
+	bool bIsCulled = !PCGEditorGraphNodeBase::ShouldDisplayAsActive(this, InComponentBeingDebugged, InStackBeingInspected);
 
 	// Show grid size visualization if higen is enabled and if we're inspecting a specific grid, and we're inspecting a subgraph since subgraphs
-	// execute at the invoked grid level.
+	// execute at the invoked grid level. 
 	if (HiGenEnabled && InStackBeingInspected && InStackBeingInspected->IsCurrentFrameInRootGraph() && InspectingGridSize != PCGHiGenGrid::UninitializedGridSize())
 	{
-		uint32 DefaultGridSize;
-		UPCGSubsystem* Subsystem = UPCGSubsystem::GetActiveEditorInstance();
-		if (Graph->IsHierarchicalGenerationEnabled())
-		{
-			DefaultGridSize = Graph->GetDefaultGridSize();
-		}
-		else if (Subsystem && Subsystem->GetPCGWorldActor())
-		{
-			DefaultGridSize = Subsystem->GetPCGWorldActor()->PartitionGridSize;
-		}
-		else
-		{
-			// Fallback, should not be hit
-			ensure(false);
-			DefaultGridSize = APCGWorldActor::DefaultPartitionGridSize;
-		}
+		const uint32 DefaultGridSize = TopGraph->GetDefaultGridSize();
 		const uint32 NodeGridSize = Graph->GetNodeGenerationGridSize(PCGNode, DefaultGridSize);
 
 		if (NodeGridSize < InspectingGridSize)
@@ -539,12 +538,27 @@ EPCGChangeType UPCGEditorGraphNodeBase::UpdateStructuralVisualization(UPCGCompon
 		}
 		else if (NodeGridSize > InspectingGridSize)
 		{
-			// If node is on larger grid than current, highlight it to indicate that its data is available for use
-			bShouldDisplayAsHighlighted |= true;
-
 			// We don't know if the node was culled or not on that grid, disable visualization.
 			bIsCulled = false;
 		}
+
+		const EPCGHiGenGrid Grid = PCGHiGenGrid::GridSizeToGrid(NodeGridSize);
+		if (GenerationGrid != Grid)
+		{
+			GenerationGrid = Grid;
+			ChangeType |= EPCGChangeType::Cosmetic;
+		}
+	}
+	else
+	{
+		// If higen is enabled then we are inspecting an invoked subgraph. Display the inspected grid size so that the user still
+		// gets the execution grid information.
+		const EPCGHiGenGrid Grid = HiGenEnabled ? PCGHiGenGrid::GridSizeToGrid(InspectingGridSize) : EPCGHiGenGrid::Uninitialized;
+		if (GenerationGrid != Grid)
+		{
+			GenerationGrid = Grid;
+			ChangeType |= EPCGChangeType::Cosmetic;
+		}
 	}
 
 	SetIsCulledFromExecution(bIsCulled);
@@ -552,7 +566,6 @@ EPCGChangeType UPCGEditorGraphNodeBase::UpdateStructuralVisualization(UPCGCompon
 	if (bIsCulled)
 	{
 		bShouldDisplayAsDisabled = true;
-		bShouldDisplayAsHighlighted = false;
 	}
 
 	if (IsDisplayAsDisabledForced() != bShouldDisplayAsDisabled)
@@ -561,12 +574,6 @@ EPCGChangeType UPCGEditorGraphNodeBase::UpdateStructuralVisualization(UPCGCompon
 		ChangeType |= EPCGChangeType::Cosmetic;
 	}
 
-	if (IsHighlighted() != bShouldDisplayAsHighlighted)
-	{
-		SetIsHighlighted(bShouldDisplayAsHighlighted);
-		ChangeType |= EPCGChangeType::Cosmetic;
-	}
-
 	return ChangeType;
 }
 
@@ -765,7 +772,7 @@ FLinearColor UPCGEditorGraphNodeBase::GetNodeBodyTintColor() const
 		}
 	}
 
-	return IsHighlighted() ? GetDefault<UPCGEditorSettings>()->HighlightedNodeBodyTintColor : Super::GetNodeBodyTintColor();
+	return Super::GetNodeBodyTintColor();
 }
 
 FEdGraphPinType UPCGEditorGraphNodeBase::GetPinType(const UPCGPin* InPin)
diff --git a/Engine/Plugins/PCG/Source/PCGEditor/Private/PCGEditorGraphNodeBase.h b/Engine/Plugins/PCG/Source/PCGEditor/Private/PCGEditorGraphNodeBase.h
--- a/Engine/Plugins/PCG/Source/PCGEditor/Private/PCGEditorGraphNodeBase.h
+++ b/Engine/Plugins/PCG/Source/PCGEditor/Private/PCGEditorGraphNodeBase.h
@@ -4,6 +4,7 @@
 
 #include "EdGraph/EdGraphNode.h"
 
+#include "PCGCommon.h"
 
 #include "PCGEditorGraphNodeBase.generated.h"
 
@@ -80,10 +81,6 @@ public:
 	/** Returns custom compact node icon if available */
 	bool GetCompactNodeIcon(FName& OutCompactNodeIcon) const;
 
-	// Highlighted nodes draw with a light tint. Currently used to indicate nodes from hi-gen grids that are available for use.
-	void SetIsHighlighted(bool bInIsHighlighted) { bIsHighlighted = bInIsHighlighted; }
-	bool IsHighlighted() const { return bIsHighlighted; }
-
 	bool CanUserAddRemoveDynamicInputPins() const { return bCanUserAddRemoveSourcePins; }
 
 	/** Whether to flip the order of the title lines - display generated title first and authored second. */
@@ -98,6 +95,12 @@ public:
 	/** Bitmask of inactive output pins. Bit N will be set if output pin index N is inactive. */
 	uint64 GetInactiveOutputPinMask() const { return InactiveOutputPinMask; }
 
+	/** The grid that this node executes on if higen is enabled, otherwise Unitialized. */
+	EPCGHiGenGrid GetGenerationGrid() const { return GenerationGrid; }
+
+	/** The higen grid currently being inspected if any, otherwise Uninitialized. */
+	EPCGHiGenGrid GetInspectedGenerationGrid() const { return InspectedGenerationGrid; }
+
 	/** Whether the given output pin was active in the previous execution. */
 	bool IsOutputPinActive(const UEdGraphPin* InOutputPin) const;
 
@@ -145,7 +148,6 @@ protected:
 	bool bDeferredReconstruct = false;
 	bool bDisableReconstructFromNode = false;
 	bool bIsInspected = false;
-	bool bIsHighlighted = false;
 	bool bCanUserAddRemoveSourcePins = false;
 
 	/** Whether this node was culled in the last execution. */
@@ -153,4 +155,10 @@ protected:
 
 	/** Bitmask of inactive output pins. Bit N will be set if output pin index N is inactive. */
 	uint64 InactiveOutputPinMask = 0;
+
+	/** The grid that this node executes on if higen is enabled, otherwise Unitialized. */
+	EPCGHiGenGrid GenerationGrid = EPCGHiGenGrid::Uninitialized;
+
+	/** The higen grid currently being inspected if any, otherwise Uninitialized. */
+	EPCGHiGenGrid InspectedGenerationGrid = EPCGHiGenGrid::Uninitialized;
 };
diff --git a/Engine/Plugins/PCG/Source/PCGEditor/Private/PCGEditorGraphSchema.cpp b/Engine/Plugins/PCG/Source/PCGEditor/Private/PCGEditorGraphSchema.cpp
--- a/Engine/Plugins/PCG/Source/PCGEditor/Private/PCGEditorGraphSchema.cpp
+++ b/Engine/Plugins/PCG/Source/PCGEditor/Private/PCGEditorGraphSchema.cpp
@@ -733,7 +733,7 @@ void FPCGEditorConnectionDrawingPolicy::DetermineWiringStyle(UEdGraphPin* Output
 
 		if (PCGPin && UpstreamEditorNode && DownstreamEditorNode)
 		{
-			const bool bDownstreamNodeForceDisabled = DownstreamEditorNode->IsDisplayAsDisabledForced() && !DownstreamEditorNode->IsHighlighted();
+			const bool bDownstreamNodeForceDisabled = DownstreamEditorNode->IsDisplayAsDisabledForced();
 
 			// Look for the PCG edge that correlates with passed in (OutputPin, InputPin) edge
 			const TObjectPtr<UPCGEdge>* PCGEdge = PCGPin->Edges.FindByPredicate([UpstreamEditorNode, OutputPin](const UPCGEdge* ConnectedPCGEdge)
diff --git a/Engine/Plugins/PCG/Source/PCGEditor/Private/PCGEditorSettings.cpp b/Engine/Plugins/PCG/Source/PCGEditor/Private/PCGEditorSettings.cpp
--- a/Engine/Plugins/PCG/Source/PCGEditor/Private/PCGEditorSettings.cpp
+++ b/Engine/Plugins/PCG/Source/PCGEditor/Private/PCGEditorSettings.cpp
@@ -11,7 +11,6 @@ UPCGEditorSettings::UPCGEditorSettings(const FObjectInitializer& ObjectInitializ
 {
 	DefaultNodeColor = FLinearColor(0.4f, 0.62f, 1.0f);
 	InstancedNodeBodyTintColor = FLinearColor(0.5f, 0.5f, 0.5f);
-	HighlightedNodeBodyTintColor = FLinearColor(0.25f, 0.25f, 0.5f, 0.33f);
 	InputOutputNodeColor = FLinearColor(1.0f, 0.0f, 0.0f);
 	SetOperationNodeColor = FLinearColor(1.0f, 0.2f, 1.0f);
 	DensityOperationNodeColor = FLinearColor(0.2f, 0.59f, 0.99f);
diff --git a/Engine/Plugins/PCG/Source/PCGEditor/Private/PCGEditorStyle.cpp b/Engine/Plugins/PCG/Source/PCGEditor/Private/PCGEditorStyle.cpp
--- a/Engine/Plugins/PCG/Source/PCGEditor/Private/PCGEditorStyle.cpp
+++ b/Engine/Plugins/PCG/Source/PCGEditor/Private/PCGEditorStyle.cpp
@@ -78,6 +78,13 @@ FPCGEditorStyle::FPCGEditorStyle() : FSlateStyleSet("PCGEditorStyle")
 	Set("PCG.Node.NodeTitleInlineEditableText", NodeTitleStyle);
 	Set("PCG.Node.InstancedNodeTitleInlineEditableText", FInlineEditableTextBlockStyle(NodeTitleStyle)
 		.SetTextStyle(GraphNodeItalicTitle));
+
+	// Styles for higen grid size label. Hand-tweaked values to match UI target mockups.
+	Set(PCGEditorStyleConstants::Node_Overlay_GridSizeLabel_Active_Border, new FSlateRoundedBoxBrush(
+		FLinearColor::White,
+		PCGEditorStyleConstants::Node_Overlay_GridSizeLabel_BorderRadius,
+		FLinearColor::Black,
+		PCGEditorStyleConstants::Node_Overlay_GridSizeLabel_BorderStroke));
 }
 
 const FPCGEditorStyle& FPCGEditorStyle::Get()
diff --git a/Engine/Plugins/PCG/Source/PCGEditor/Private/PCGEditorStyle.h b/Engine/Plugins/PCG/Source/PCGEditor/Private/PCGEditorStyle.h
--- a/Engine/Plugins/PCG/Source/PCGEditor/Private/PCGEditorStyle.h
+++ b/Engine/Plugins/PCG/Source/PCGEditor/Private/PCGEditorStyle.h
@@ -51,4 +51,9 @@ namespace PCGEditorStyleConstants
 	const FName Pin_Required = TEXT("PCG.Graph.Pin.Required");
 
 	const FName Node_Overlay_Inactive = TEXT("PCG.Node.Overlay.Inactive");
+
+	const FName Node_Overlay_GridSizeLabel_Active_Border = TEXT("PCG.Node.Overlay.ThisGridSizeLabel.Active.Border");
+
+	const float Node_Overlay_GridSizeLabel_BorderRadius = 11.0f;
+	const float Node_Overlay_GridSizeLabel_BorderStroke = 1.5f;
 }
\ No newline at end of file
diff --git a/Engine/Plugins/PCG/Source/PCGEditor/Private/SPCGEditorGraphNode.cpp b/Engine/Plugins/PCG/Source/PCGEditor/Private/SPCGEditorGraphNode.cpp
--- a/Engine/Plugins/PCG/Source/PCGEditor/Private/SPCGEditorGraphNode.cpp
+++ b/Engine/Plugins/PCG/Source/PCGEditor/Private/SPCGEditorGraphNode.cpp
@@ -2,13 +2,14 @@
 
 #include "SPCGEditorGraphNode.h"
 
-#include "PCGEditorGraphNodeBase.h"
-#include "PCGEditorStyle.h"
 #include "PCGNode.h"
 #include "PCGPin.h"
 #include "PCGSettings.h"
 #include "PCGSettingsWithDynamicInputs.h"
 
+#include "PCGEditorGraphNodeBase.h"
+#include "PCGEditorStyle.h"
+
 #include "GraphEditorSettings.h"
 #include "IDocumentation.h"
 #include "SCommentBubble.h"
@@ -460,8 +461,7 @@ void SPCGEditorGraphNode::UpdateGraphNode()
 
 const FSlateBrush* SPCGEditorGraphNode::GetNodeBodyBrush() const
 {
-	const bool bNeedsTint = PCGEditorGraphNode &&
-		((PCGEditorGraphNode->GetPCGNode() && PCGEditorGraphNode->GetPCGNode()->IsInstance()) || PCGEditorGraphNode->IsHighlighted());
+	const bool bNeedsTint = PCGEditorGraphNode && PCGEditorGraphNode->GetPCGNode() && PCGEditorGraphNode->GetPCGNode()->IsInstance();
 	if (bNeedsTint)
 	{
 		return FAppStyle::GetBrush("Graph.Node.TintedBody");
@@ -594,6 +594,84 @@ void SPCGEditorGraphNode::AddPin(const TSharedRef<SGraphPin>& PinToAdd)
 	}
 }
 
+TArray<FOverlayWidgetInfo> SPCGEditorGraphNode::GetOverlayWidgets(bool bSelected, const FVector2D& WidgetSize) const
+{
+	TArray<FOverlayWidgetInfo> OverlayWidgets = SGraphNode::GetOverlayWidgets(bSelected, WidgetSize);
+
+	if (PCGEditorGraphNode && !PCGEditorGraphNode->ShouldDrawCompact())
+	{
+		// Higen grid size overlay widget. All magic numbers below hand tweaked to match UI mockup.
+		const EPCGHiGenGrid InspectedGrid = PCGEditorGraphNode->GetInspectedGenerationGrid();
+		UPCGNode* PCGNode = PCGEditorGraphNode->GetPCGNode();
+
+		//const bool bHigenEnabled = PCGEditorGraphNode->GetPCGNode() && PCGEditorGraphNode->getpcn
+		const bool bInspectingHigen = InspectedGrid != EPCGHiGenGrid::Uninitialized;
+		if (bInspectingHigen && PCGEditorGraphNode->IsNodeEnabled())
+		{
+			FText GenerationGridText;
+			const EPCGHiGenGrid Grid = PCGEditorGraphNode->GetGenerationGrid();
+
+			if (Grid == EPCGHiGenGrid::Unbounded)
+			{
+				GenerationGridText = FText::FromString(TEXT("UB"));
+			}
+			else
+			{
+				// Meters are easier on the eyes.
+				const uint32 GridSize = PCGHiGenGrid::GridToGridSize(Grid) / 100;
+				GenerationGridText = FText::AsNumber(GridSize, &FNumberFormattingOptions::DefaultNoGrouping());
+			}
+
+			FLinearColor Tint = FLinearColor::White;
+			if (Grid != EPCGHiGenGrid::Uninitialized)
+			{
+				Tint = GetGridLabelColor(Grid);
+			}
+			else if (PCGEditorGraphNode->IsDisplayAsDisabledForced())
+			{
+				Tint.A *= 0.35f;
+			}
+
+			// Create a border brush for each combination of grids, to workaround issue where the tint does not apply
+			// to the border element.
+			const FSlateBrush* BorderBrush = GetBorderBrush(InspectedGrid, Grid);
+
+			FLinearColor TextColor = FColor::White;
+			FLinearColor BackgroundColor = FColor::Black;
+			if (InspectedGrid == Grid)
+			{
+				// Flip colors for active grid to highlight them.
+				Swap(TextColor, BackgroundColor);
+			}
+
+			TSharedPtr<SWidget> GridSizeLabel =
+				SNew(SHorizontalBox)
+				.Visibility(EVisibility::Visible)
+				+SHorizontalBox::Slot()
+				[
+					SNew(SBorder)
+					.BorderImage(BorderBrush)
+					.Padding(FMargin(12, 3))
+					.ColorAndOpacity(Tint)
+					[
+						SNew(STextBlock)
+						.TextStyle(FAppStyle::Get(), "Graph.Node.NodeTitle")
+						.Text(GenerationGridText)
+						.Justification(ETextJustify::Center)
+						.ColorAndOpacity(TextColor)
+					]
+				];
+
+			FOverlayWidgetInfo GridSizeLabelInfo(GridSizeLabel);
+			GridSizeLabelInfo.OverlayOffset = FVector2D(GetDesiredSize().X - 30.0f, -9.0f);
+
+			OverlayWidgets.Add(GridSizeLabelInfo);
+		}
+	}
+
+	return OverlayWidgets;
+}
+
 void SPCGEditorGraphNode::GetOverlayBrushes(bool bSelected, const FVector2D WidgetSize, TArray<FOverlayBrushInfo>& Brushes) const
 {
 	check(PCGEditorGraphNode);
@@ -866,4 +944,54 @@ void SPCGEditorGraphNode::UpdateCompactNode()
 	CreateOutputSideAddButton(RightNodeBox);
 }
 
+FLinearColor SPCGEditorGraphNode::GetGridLabelColor(EPCGHiGenGrid NodeGrid)
+{
+	// All colours hand tweaked to give a kind of "temperature scale" for the hierarchy.
+	switch (NodeGrid)
+	{
+	case EPCGHiGenGrid::Unbounded:
+		return FColor(255, 255, 255, 255);
+	case EPCGHiGenGrid::Grid2048:
+		return FColor(53, 60, 171, 255);
+	case EPCGHiGenGrid::Grid1024:
+		return FColor(31, 82, 210, 255);
+	case EPCGHiGenGrid::Grid512:
+		return FColor(16, 120, 217, 255);
+	case EPCGHiGenGrid::Grid256:
+		return FColor(8, 151, 208, 255);
+	case EPCGHiGenGrid::Grid128:
+		return FColor(9, 170, 188, 255);
+	case EPCGHiGenGrid::Grid64:
+		return FColor(64, 185, 150, 255);
+	case EPCGHiGenGrid::Grid32:
+		return FColor(144, 189, 114, 255);
+	case EPCGHiGenGrid::Grid16:
+		return FColor(207, 185, 89, 255);
+	case EPCGHiGenGrid::Grid8:
+		return FColor(252, 189, 61, 255);
+	case EPCGHiGenGrid::Grid4:
+		return FColor(243, 227, 28, 255);
+	default:
+		ensure(false);
+		return FLinearColor::White;
+	}
+}
+
+const FSlateBrush* SPCGEditorGraphNode::GetBorderBrush(EPCGHiGenGrid InspectedGrid, EPCGHiGenGrid NodeGrid) const
+{
+	if (InspectedGrid == NodeGrid)
+	{
+		return FPCGEditorStyle::Get().GetBrush(PCGEditorStyleConstants::Node_Overlay_GridSizeLabel_Active_Border);
+	}
+
+	// Hand tweaked multiplier to fade child node grid size labels.
+	const float Opacity = (InspectedGrid < NodeGrid) ? 1.0f : 0.5f;
+
+	return new FSlateRoundedBoxBrush(
+		FLinearColor::Black * Opacity,
+		PCGEditorStyleConstants::Node_Overlay_GridSizeLabel_BorderRadius,
+		GetGridLabelColor(NodeGrid) * Opacity,
+		PCGEditorStyleConstants::Node_Overlay_GridSizeLabel_BorderStroke);
+}
+
 #undef LOCTEXT_NAMESPACE
diff --git a/Engine/Plugins/PCG/Source/PCGEditor/Private/SPCGEditorGraphNode.h b/Engine/Plugins/PCG/Source/PCGEditor/Private/SPCGEditorGraphNode.h
--- a/Engine/Plugins/PCG/Source/PCGEditor/Private/SPCGEditorGraphNode.h
+++ b/Engine/Plugins/PCG/Source/PCGEditor/Private/SPCGEditorGraphNode.h
@@ -4,6 +4,8 @@
 
 #include "SGraphNode.h"
 
+#include "PCGCommon.h"
+
 struct FOverlayBrushInfo;
 
 class UPCGEditorGraphNodeBase;
@@ -29,6 +31,7 @@ public:
 	//~ End SGraphNode Interface
 
 	//~ Begin SNodePanel::SNode Interface
+	virtual TArray<FOverlayWidgetInfo> GetOverlayWidgets(bool bSelected, const FVector2D& WidgetSize) const override;
 	virtual void GetOverlayBrushes(bool bSelected, const FVector2D WidgetSize, TArray<FOverlayBrushInfo>& Brushes) const override;
 	//~ End SNodePanel::SNode Interface
 
@@ -40,5 +43,13 @@ protected:
 	void UpdateCompactNode();
 
 private:
+	static FLinearColor GetGridLabelColor(EPCGHiGenGrid NodeGrid);
+
+	/**
+	* Get the border brush for the given combination of grid sizes and enabled state. All a big workaround for FSlateRoundedBoxBrush not respecting
+	* the tint colour.
+	*/
+	const FSlateBrush* GetBorderBrush(EPCGHiGenGrid InspectedGrid, EPCGHiGenGrid NodeGrid) const;
+
 	UPCGEditorGraphNodeBase* PCGEditorGraphNode = nullptr;
 };
diff --git a/Engine/Plugins/PCG/Source/PCGEditor/Public/PCGEditorSettings.h b/Engine/Plugins/PCG/Source/PCGEditor/Public/PCGEditorSettings.h
--- a/Engine/Plugins/PCG/Source/PCGEditor/Public/PCGEditorSettings.h
+++ b/Engine/Plugins/PCG/Source/PCGEditor/Public/PCGEditorSettings.h
@@ -29,10 +29,6 @@ public:
 	UPROPERTY(EditAnywhere, config, Category = Node, meta = (HideAlphaChannel))
 	FLinearColor InstancedNodeBodyTintColor;
 
-	/** Highlighted node body tint color, used to indicate available data sources in hierarchical graphs */
-	UPROPERTY(EditAnywhere, config, Category = Node)
-	FLinearColor HighlightedNodeBodyTintColor;
-
 	/** Color used for input & output nodes */
 	UPROPERTY(EditAnywhere, config, Category = Node, meta = (HideAlphaChannel))
 	FLinearColor InputOutputNodeColor;
