[PCG] Fix dynamic tracking not working on map load if the component is on a BP

* Map was not copied over since it is not a normal UPROPERTY

#rnx

diff --git a/Engine/Plugins/PCG/Source/PCG/Private/PCGComponent.cpp b/Engine/Plugins/PCG/Source/PCG/Private/PCGComponent.cpp
--- a/Engine/Plugins/PCG/Source/PCG/Private/PCGComponent.cpp
+++ b/Engine/Plugins/PCG/Source/PCG/Private/PCGComponent.cpp
@@ -3203,6 +3203,7 @@ void FPCGComponentInstanceData::ApplyToComponent(UActorComponent* Component, con
 
 			PCGComponent->CurrentEditingMode = SourceComponent->CurrentEditingMode;
 			PCGComponent->PreviousEditingMode = SourceComponent->PreviousEditingMode;
+			PCGComponent->DynamicallyTrackedKeysToSettings = SourceComponent->DynamicallyTrackedKeysToSettings;
 #endif // WITH_EDITOR
 
 			// Non-critical but should be done: transient data, tracked actors cache, landscape tracking
