[PCG] Add a log/warning/error emitting node (Print String)

diff --git a/Engine/Plugins/PCG/Source/PCG/Private/Elements/PCGPrintElement.cpp b/Engine/Plugins/PCG/Source/PCG/Private/Elements/PCGPrintElement.cpp
new file mode 100644
--- /dev/null
+++ b/Engine/Plugins/PCG/Source/PCG/Private/Elements/PCGPrintElement.cpp
@@ -0,0 +1,130 @@
+// Copyright Epic Games, Inc. All Rights Reserved.
+
+#include "Elements/PCGPrintElement.h"
+
+#include "PCGComponent.h"
+#include "PCGContext.h"
+#include "PCGGraph.h"
+#include "PCGModule.h"
+
+#include "Components/ActorComponent.h"
+#include "GameFramework/Actor.h"
+
+#if WITH_EDITOR
+#include "EngineGlobals.h"
+#include "Editor/EditorEngine.h"
+#endif // WITH_EDITOR
+
+#define LOCTEXT_NAMESPACE "PCGPrintElement"
+
+namespace PCGPrintElementConstants
+{
+	FText Delimiter = LOCTEXT("Delimiter", "::");
+}
+
+TArray<FPCGPinProperties> UPCGPrintElementSettings::InputPinProperties() const
+{
+	TArray<FPCGPinProperties> PinProperties;
+	PinProperties.Emplace(PCGPinConstants::DefaultInputLabel, EPCGDataType::Any);
+	return PinProperties;
+}
+
+TArray<FPCGPinProperties> UPCGPrintElementSettings::OutputPinProperties() const
+{
+	TArray<FPCGPinProperties> PinProperties;
+	PinProperties.Emplace(PCGPinConstants::DefaultOutputLabel, EPCGDataType::Any);
+	return PinProperties;
+}
+
+FPCGElementPtr UPCGPrintElementSettings::CreateElement() const
+{
+	return MakeShared<FPCGPrintElement>();
+}
+
+bool FPCGPrintElement::ExecuteInternal(FPCGContext* Context) const
+{
+	TRACE_CPUPROFILER_EVENT_SCOPE(FPCGPrintElement::Execute);
+
+	Context->OutputData.TaggedData = Context->InputData.GetInputsByPin(PCGPinConstants::DefaultInputLabel);
+
+#if !(UE_BUILD_SHIPPING || UE_BUILD_TEST) || USE_LOGGING_IN_SHIPPING
+	const UPCGPrintElementSettings* Settings = Context->GetInputSettings<UPCGPrintElementSettings>();
+	check(Settings);
+
+	if (!Settings->bEnablePrint)
+	{
+		return true;
+	}
+
+	TArray<FString> Prefixes;
+
+	// Build the message prefixes
+	const UPCGComponent* Component = Context->SourceComponent->GetOriginalComponent();
+	if (Component)
+	{
+		if (Settings->bPrefixWithOwner)
+		{
+			Prefixes.Emplace(Component->GetOwner() ? Component->GetOwner()->GetName() : LOCTEXT("MissingOwner", "Missing Owner").ToString());
+		}
+
+		if (Settings->bPrefixWithComponent)
+		{
+			Prefixes.Emplace(Component->GetName());
+		}
+
+		if (Settings->bPrefixWithGraph)
+		{
+			Prefixes.Emplace(Component->GetGraph() ? Component->GetGraph()->GetName() : LOCTEXT("MissingGraph", "Missing Graph").ToString());
+		}
+	}
+
+	if (Settings->bPrefixWithNode)
+	{
+		Prefixes.Emplace(Context->Node ? Context->Node->GetName() : LOCTEXT("MissingNode", "Missing Node").ToString());
+	}
+
+	FString ObjectPrefix;
+	if (!Prefixes.IsEmpty())
+	{
+		ObjectPrefix = FString::Printf(TEXT("[%s]: "), *FString::Join(Prefixes, *PCGPrintElementConstants::Delimiter.ToString()));
+	}
+
+	const FString FinalString = Settings->CustomPrefix + ObjectPrefix + Settings->PrintString;
+
+	switch (Settings->Verbosity)
+	{
+		case EPCGPrintVerbosity::Error:
+			PCGLog::LogErrorOnGraph(FText::FromString(FinalString), Settings->bDisplayOnNode ? Context : nullptr);
+			break;
+		case EPCGPrintVerbosity::Warning:
+			PCGLog::LogWarningOnGraph(FText::FromString(FinalString), Settings->bDisplayOnNode ? Context : nullptr);
+			break;
+		case EPCGPrintVerbosity::Log:
+			PCGE_LOG_C(Log, LogOnly, Context, FText::FromString(FinalString));
+			break;
+		default:
+			checkNoEntry();
+	}
+
+#if WITH_EDITOR
+	if (Settings->bPrintToScreen && GEngine)
+	{
+		if (UEditorEngine* Editor = static_cast<UEditorEngine*>(GEngine))
+		{
+			uint32 HashKey = HashCombine(GetTypeHash(Context->Node->GetFName()), Context->Node->GetUniqueID());
+
+			if (Settings->bPrintPerComponent)
+			{
+				HashKey = HashCombine(HashKey, Component->GetUniqueID());
+			}
+
+			Editor->AddOnScreenDebugMessage(static_cast<uint64>(HashKey), Settings->PrintToScreenDuration, Settings->PrintToScreenColor, FinalString);
+		}
+	}
+#endif // WITH_EDITOR
+#endif // !(UE_BUILD_SHIPPING || UE_BUILD_TEST) || USE_LOGGING_IN_SHIPPING
+
+	return true;
+}
+
+#undef LOCTEXT_NAMESPACE
diff --git a/Engine/Plugins/PCG/Source/PCG/Public/Elements/PCGPrintElement.h b/Engine/Plugins/PCG/Source/PCG/Public/Elements/PCGPrintElement.h
new file mode 100644
--- /dev/null
+++ b/Engine/Plugins/PCG/Source/PCG/Public/Elements/PCGPrintElement.h
@@ -0,0 +1,108 @@
+// Copyright Epic Games, Inc. All Rights Reserved.
+
+#pragma once
+
+#include "PCGElement.h"
+#include "PCGSettings.h"
+
+#include "Logging/LogVerbosity.h"
+
+#include "PCGPrintElement.generated.h"
+
+UENUM()
+enum class EPCGPrintVerbosity : uint8
+{
+	Log = ELogVerbosity::Log,
+	Warning = ELogVerbosity::Warning,
+	Error = ELogVerbosity::Error
+};
+
+/**
+ * Issues a specified message to the log, and optionally to the graph and/or screen.
+ * Note: This node will not function in shipping builds.
+ */
+UCLASS(BlueprintType, ClassGroup = (Procedural), meta = (Keywords = "debug print string"))
+class UPCGPrintElementSettings : public UPCGSettings
+{
+	GENERATED_BODY()
+
+public:
+	//~Begin UPCGSettings interface
+#if WITH_EDITOR
+	virtual FName GetDefaultNodeName() const override { return FName(TEXT("Print String")); }
+	virtual FText GetDefaultNodeTitle() const override { return NSLOCTEXT("PCGPrintElement", "NodeTitle", "Print String"); }
+	virtual FText GetNodeTooltipText() const override { return NSLOCTEXT("PCGPrintElement", "NodeTooltip", "Issues a specified message to the log, and optionally to the graph and/or screen."); }
+	virtual EPCGSettingsType GetType() const override { return EPCGSettingsType::Debug; }
+	virtual bool HasDynamicPins() const override { return true; }
+#endif
+
+protected:
+	virtual TArray<FPCGPinProperties> InputPinProperties() const override;
+	virtual TArray<FPCGPinProperties> OutputPinProperties() const override;
+	virtual FPCGElementPtr CreateElement() const override;
+	//~End UPCGSettings interface
+
+public:
+	/** The core message to print to the logger, graph, and/or screen. */
+	UPROPERTY(BlueprintReadWrite, EditAnywhere, Category = Settings, meta = (PCG_Overridable))
+	FString PrintString;
+
+	/** The verbosity level of the printed message. */
+	UPROPERTY(BlueprintReadWrite, EditAnywhere, Category = Settings, meta = (PCG_Overridable))
+	EPCGPrintVerbosity Verbosity = EPCGPrintVerbosity::Log;
+
+	/** A prefix to which the core message will be appended. */
+	UPROPERTY(BlueprintReadWrite, EditAnywhere, Category = Settings, meta = (PCG_Overridable))
+	FString CustomPrefix;
+
+	/** Display warnings or errors on this node. */
+	UPROPERTY(BlueprintReadWrite, EditAnywhere, Category = Settings, meta = (PCG_Overridable))
+	bool bDisplayOnNode = false;
+
+	/** Use the component as part of the key hash and print a message for each component with this node. */
+	UPROPERTY(BlueprintReadWrite, EditAnywhere, Category = Settings, meta = (EditCondition = "bPrintToScreen", EditConditionHides))
+	bool bPrintPerComponent = true;
+
+#if WITH_EDITORONLY_DATA
+	/** Print the message to the editor viewport. */
+	UPROPERTY(BlueprintReadWrite, EditAnywhere, Category = Settings, meta = (PCG_Overridable))
+	bool bPrintToScreen = false;
+
+	/** The duration (in seconds) of the on screen message. */
+	UPROPERTY(BlueprintReadWrite, EditAnywhere, Category = Settings, meta = (PCG_Overridable, EditCondition = "bPrintToScreen", EditConditionHides))
+	double PrintToScreenDuration = 15.0;
+
+	/** The color of the on screen message. */
+	UPROPERTY(BlueprintReadWrite, EditAnywhere, Category = Settings, meta = (PCG_Overridable, EditCondition = "bPrintToScreen", EditConditionHides))
+	FColor PrintToScreenColor = FColor::Cyan;
+#endif // WITH_EDITORONLY_DATA
+
+	/** Prefix the message with the name of the component's owner. */
+	UPROPERTY(BlueprintReadWrite, EditAnywhere, Category = Settings, meta = (EditCondition = "bPrintPerComponent", EditConditionHides))
+	bool bPrefixWithOwner = false;
+
+	/** Prefix the message with the name of the component. */
+	UPROPERTY(BlueprintReadWrite, EditAnywhere, Category = Settings, meta = (EditCondition = "bPrintPerComponent", EditConditionHides))
+	bool bPrefixWithComponent = false;
+
+	/** Prefix the message with the name of the graph. */
+    UPROPERTY(BlueprintReadWrite, EditAnywhere, Category = Settings, meta = (EditCondition = "bPrintPerComponent", EditConditionHides))
+    bool bPrefixWithGraph = true;
+
+	/** Prefix the message with the name of the node. */
+	UPROPERTY(BlueprintReadWrite, EditAnywhere, Category = Settings, meta = (EditCondition = "bPrintPerComponent", EditConditionHides))
+	bool bPrefixWithNode = true;
+
+	/** Enable the functionality of this node. Disable to bypass printing. */
+	UPROPERTY(BlueprintReadWrite, EditAnywhere, Category = Settings, meta = (PCG_Overridable))
+	bool bEnablePrint = true;
+};
+
+class FPCGPrintElement : public IPCGElement
+{
+public:
+	virtual bool IsCacheable(const UPCGSettings* InSettings) const override { return false; }
+
+protected:
+	virtual bool ExecuteInternal(FPCGContext* Context) const override;
+};
\ No newline at end of file
