[PCG] Support Attribute set for Match and Set Attributes

* Refactor FPCGPointPartitionBase to FPCGPartitionBase, working on any data
* Modify the calls in SpawnActor that was using FPCGPartitionBase

#rnx

diff --git a/Engine/Plugins/PCG/Source/PCG/Private/Elements/PCGMatchAndSetAttributes.cpp b/Engine/Plugins/PCG/Source/PCG/Private/Elements/PCGMatchAndSetAttributes.cpp
--- a/Engine/Plugins/PCG/Source/PCG/Private/Elements/PCGMatchAndSetAttributes.cpp
+++ b/Engine/Plugins/PCG/Source/PCG/Private/Elements/PCGMatchAndSetAttributes.cpp
@@ -5,6 +5,7 @@
 #include "PCGContext.h"
 #include "PCGParamData.h"
 #include "PCGPin.h"
+#include "Data/PCGPointData.h"
 #include "Elements/Metadata/PCGMetadataElementCommon.h"
 #include "Helpers/PCGBlueprintHelpers.h"
 #include "Helpers/PCGPointDataPartition.h"
@@ -41,7 +42,7 @@ FText UPCGMatchAndSetAttributesSettings::GetDefaultNodeTitle() const
 
 FText UPCGMatchAndSetAttributesSettings::GetNodeTooltipText() const
 {
-	return LOCTEXT("NodeTooltip", "Matches or randomly assigns values from the Attribute Set to the input Point Data");
+	return LOCTEXT("NodeTooltip", "Matches or randomly assigns values from the Attribute Set to the input data");
 }
 #endif // WITH_EDITOR
 
@@ -54,7 +55,7 @@ UPCGMatchAndSetAttributesSettings::UPCGMatchAndSetAttributesSettings()
 TArray<FPCGPinProperties> UPCGMatchAndSetAttributesSettings::InputPinProperties() const
 {
 	TArray<FPCGPinProperties> PinProperties;
-	PinProperties.Emplace(PCGPinConstants::DefaultInputLabel, EPCGDataType::Point);
+	PinProperties.Emplace(PCGPinConstants::DefaultInputLabel, EPCGDataType::Point | EPCGDataType::Param);
 	PinProperties.Emplace(PCGMatchAndSetAttributesConstants::MatchDataLabel, 
 		EPCGDataType::Param, 
 		/*bAllowMultipleConnection=*/false, 
@@ -77,7 +78,10 @@ TArray<FPCGPinProperties> UPCGMatchAndSetAttributesSettings::InputPinProperties(
 
 TArray<FPCGPinProperties> UPCGMatchAndSetAttributesSettings::OutputPinProperties() const
 {
-	return Super::DefaultPointOutputPinProperties();
+	TArray<FPCGPinProperties> PinProperties;
+	PinProperties.Emplace(PCGPinConstants::DefaultOutputLabel, EPCGDataType::Point | EPCGDataType::Param);
+
+	return PinProperties;
 }
 
 FPCGElementPtr UPCGMatchAndSetAttributesSettings::CreateElement() const
@@ -509,11 +513,11 @@ private:
 	TArray<TPair<PCGMetadataValueKey, AttributeSetPartitionEntry>> PartitionData;
 };
 
-class FPCGMatchAndSetPartition : public FPCGPointDataPartitionBase<FPCGMatchAndSetPartition, PCGMetadataValueKey>
+class FPCGMatchAndSetPartition : public FPCGDataPartitionBase<FPCGMatchAndSetPartition, PCGMetadataValueKey>
 {
 public:
 	FPCGMatchAndSetPartition(FPCGContext* InContext, const UPCGMatchAndSetAttributesSettings* InSettings, const UPCGComponent* InSourceComponent, const UPCGParamData* InParamData) 
-		: FPCGPointDataPartitionBase<FPCGMatchAndSetPartition, PCGMetadataValueKey>()
+		: FPCGDataPartitionBase<FPCGMatchAndSetPartition, PCGMetadataValueKey>()
 		, Context(InContext)
 		, Settings(InSettings)
 		, ParamData(InParamData)
@@ -543,27 +547,27 @@ public:
 		return AttributeSetPartition.IsValid();
 	}
 
-	bool InitializeForPointData(const UPCGPointData* PointData, UPCGPointData* OutPointData)
+	bool InitializeForData(const UPCGData* InData, UPCGData* OutData)
 	{
-		if (!PointData || !PointData->ConstMetadata() || !OutPointData)
+		if (!InData || !InData->ConstMetadata() || !OutData || !OutData->MutableMetadata())
 		{
 			return false;
 		}
 
 		if (Settings->bMatchAttributes)
 		{
-			const FPCGAttributePropertyInputSelector InputAttributeSource = Settings->InputAttribute.CopyAndFixLast(PointData);
-			InputAttributeAccessor = PCGAttributeAccessorHelpers::CreateConstAccessor(PointData, InputAttributeSource);
-			InputAttributeKeys = PCGAttributeAccessorHelpers::CreateConstKeys(PointData, InputAttributeSource);
+			const FPCGAttributePropertyInputSelector InputAttributeSource = Settings->InputAttribute.CopyAndFixLast(InData);
+			InputAttributeAccessor = PCGAttributeAccessorHelpers::CreateConstAccessor(InData, InputAttributeSource);
+			InputAttributeKeys = PCGAttributeAccessorHelpers::CreateConstKeys(InData, InputAttributeSource);
 
 			if (!InputAttributeAccessor.IsValid() || !InputAttributeKeys.IsValid())
 			{
-				PCGE_LOG_C(Error, GraphAndLog, Context, FText::Format(LOCTEXT("MissingAttribute", "Point data does not have the input attribute '{0}'."), InputAttributeSource.GetDisplayText()));
+				PCGE_LOG_C(Error, GraphAndLog, Context, FText::Format(LOCTEXT("MissingAttribute", "Input data does not have the input attribute '{0}'."), InputAttributeSource.GetDisplayText()));
 				return false;
 			}
 		}
 
-		const UPCGData** FoundMaxDistanceData = DataToMaxDistanceMap.Find(PointData);
+		const UPCGData** FoundMaxDistanceData = DataToMaxDistanceMap.Find(InData);
 		const FPCGAttributePropertyInputSelector MaxDistanceSelector = Settings->MaxDistanceInputAttribute.CopyAndFixLast(FoundMaxDistanceData ? *FoundMaxDistanceData : nullptr);
 
 		if (!AttributeSetPartition.InitializeForData(Context, FoundMaxDistanceData ? *FoundMaxDistanceData : nullptr, FoundMaxDistanceData ? &MaxDistanceSelector : nullptr))
@@ -573,17 +577,17 @@ public:
 
 		if (Settings->bUseInputWeightAttribute)
 		{
-			const FPCGAttributePropertyInputSelector InputWeightAttributeSource = Settings->InputWeightAttribute.CopyAndFixLast(PointData);
-			InputWeightAccessor = PCGAttributeAccessorHelpers::CreateConstAccessor(PointData, InputWeightAttributeSource);
+			const FPCGAttributePropertyInputSelector InputWeightAttributeSource = Settings->InputWeightAttribute.CopyAndFixLast(InData);
+			InputWeightAccessor = PCGAttributeAccessorHelpers::CreateConstAccessor(InData, InputWeightAttributeSource);
 
 			if (!InputAttributeKeys)
 			{
-				InputAttributeKeys = PCGAttributeAccessorHelpers::CreateConstKeys(PointData, InputWeightAttributeSource);
+				InputAttributeKeys = PCGAttributeAccessorHelpers::CreateConstKeys(InData, InputWeightAttributeSource);
 			}
 
 			if (!InputWeightAccessor.IsValid() || !InputAttributeKeys.IsValid())
 			{
-				PCGE_LOG_C(Error, GraphAndLog, Context, FText::Format(LOCTEXT("MissingWeightAttribute", "Point data does not have the input weight attribute '{0}'."), InputWeightAttributeSource.GetDisplayText()));
+				PCGE_LOG_C(Error, GraphAndLog, Context, FText::Format(LOCTEXT("MissingWeightAttribute", "Input data does not have the input weight attribute '{0}'."), InputWeightAttributeSource.GetDisplayText()));
 				return false;
 			}
 
@@ -602,6 +606,9 @@ public:
 		TArray<EPCGMetadataTypes> ParamAttributeTypes;
 		ParamData->ConstMetadata()->GetAttributes(ParamAttributeNames, ParamAttributeTypes);
 
+		UPCGMetadata* OutMetadata = OutData->MutableMetadata();
+		check(OutMetadata);
+
 		for (int32 AttributeIndex = 0; AttributeIndex < ParamAttributeNames.Num(); ++AttributeIndex)
 		{
 			const FName AttributeName = ParamAttributeNames[AttributeIndex];
@@ -613,107 +620,160 @@ public:
 			}
 
 			const FPCGMetadataAttributeBase* ParamAttribute = ParamData->ConstMetadata()->GetConstAttribute(AttributeName);
-			
-			FPCGMetadataAttributeBase* PointAttribute = OutPointData->Metadata->GetMutableAttribute(AttributeName);
+
+			FPCGMetadataAttributeBase* PointAttribute = OutMetadata->GetMutableAttribute(AttributeName);
 			if (PointAttribute && PointAttribute->GetTypeId() != ParamAttribute->GetTypeId())
 			{
-				OutPointData->Metadata->DeleteAttribute(AttributeName);
+				OutMetadata->DeleteAttribute(AttributeName);
 				PointAttribute = nullptr;
 			}
 
 			if (!PointAttribute)
 			{
-				PointAttribute = OutPointData->Metadata->CopyAttribute(ParamAttribute, AttributeName, /*bKeepParent=*/false, /*bCopyEntries=*/false, /*bCopyValues=*/false);
+				PointAttribute = OutMetadata->CopyAttribute(ParamAttribute, AttributeName, /*bKeepParent=*/false, /*bCopyEntries=*/false, /*bCopyValues=*/false);
 			}
 
 			if (!PointAttribute) // Failed to create attribute
 			{
-				PCGE_LOG_C(Warning, GraphAndLog, Context, FText::Format(LOCTEXT("UnableToCreateAttribute", "Unable to create attribute '{0}' on point data."), FText::FromName(AttributeName)));
+				PCGE_LOG_C(Warning, GraphAndLog, Context, FText::Format(LOCTEXT("UnableToCreateAttribute", "Unable to create attribute '{0}' on output data."), FText::FromName(AttributeName)));
 				return false;
 			}
 
 			AttributesToSet.Emplace(ParamAttribute, PointAttribute);
 		}
 
-		return true;
+		const int32 NumElements = GetNumElements(InData);
+		PartitionDataIndices = AttributeSetPartition.GetMatchingPartitionDataIndices(InputAttributeAccessor, InputAttributeKeys, NumElements);
+
+		if (PartitionDataIndices.Num() != NumElements)
+		{
+			return false;
+		}
+
+		const UPCGPointData* InPointData = Cast<const UPCGPointData>(InData);
+		SetupWeights(NumElements, InPointData ? &InPointData->GetPoints() : nullptr);
+
+		return Weights.Num() == NumElements;
 	}
 
-	TArray<double> GetWeights(const TArray<FPCGPoint>& Points) 
+	void SetupWeights(const int32 NumElements, const TArray<FPCGPoint>* Points) 
 	{
-		TArray<double> Weights;
+		check(!Points || NumElements == Points->Num());
+
+		Weights.Reset(NumElements);
 
 		if (InputWeightAccessor.IsValid() && InputAttributeKeys.IsValid())
 		{
-			Weights.SetNumUninitialized(Points.Num());
+			Weights.SetNumUninitialized(NumElements);
 			InputWeightAccessor->GetRange<double>(Weights, 0, *InputAttributeKeys, EPCGAttributeAccessorFlags::AllowConstructible);
 		}
-		else
+		else if (Points)
 		{
-			Weights.Reserve(Points.Num());
-
 			// Generate a random value from the seed
-			for (const FPCGPoint& Point : Points)
+			for (const FPCGPoint& Point : *Points)
 			{
 				Weights.Add(UPCGBlueprintHelpers::GetRandomStreamFromPoint(Point, Settings, SourceComponent).FRand());
 			}
 		}
-
-		return Weights;
+		else
+		{
+			FRandomStream RandomStream = UPCGBlueprintHelpers::GetRandomStreamFromPoint(FPCGPoint(), Settings, SourceComponent);
+			for (int32 i = 0; i < NumElements; ++i)
+			{
+				Weights.Add(RandomStream.FRand());
+			}
+		}
 	}
 
-	FPCGPointDataPartitionBase::Element* SelectPoint(const FPCGPoint& Point, int32 PointIndex)
+	FPCGDataPartitionBase::Element* Select(int32 Index)
 	{
 		return nullptr;
 	}
 
-	void Finalize(const UPCGPointData* InPointData, UPCGPointData* OutPointData)
+	int32 GetNumElements(const UPCGData* InData)
 	{
-		check(InPointData && OutPointData);
-		const TArray<FPCGPoint>& Points = InPointData->GetPoints();
-		TArray<FPCGPoint>& OutPoints = OutPointData->GetMutablePoints();
-
-		// In this case, we haven't written any points yet to the OutPointData,
-		// as we need to find the matching index for each entry first.
-		TArray<int32> PartitionDataIndices = AttributeSetPartition.GetMatchingPartitionDataIndices(InputAttributeAccessor, InputAttributeKeys, Points.Num());
+		check(InData);
 
-		if (PartitionDataIndices.Num() != Points.Num())
+		if (const UPCGPointData* InPointData = Cast<const UPCGPointData>(InData))
 		{
-			return;
+			return InPointData->GetPoints().Num();
 		}
+		else if (const UPCGMetadata* Metadata = InData->ConstMetadata())
+		{
+			return Metadata->GetItemCountForChild();
+		}
+		else
+		{
+			return 0;
+		}
+	}
 
-		const TArray<double> PointWeights = GetWeights(Points);
-
-		if (PointWeights.Num() != Points.Num())
+	void Finalize(const UPCGData* InData, UPCGData* OutData)
+	{
+		check(InData && OutData);
+		if (const UPCGPointData* InPointData = Cast<const UPCGPointData>(InData))
+		{
+			FinalizeInternal(InPointData, static_cast<UPCGPointData*>(OutData), InPointData->GetPoints().Num());
+		}
+		else if (const UPCGMetadata* Metadata = InData->ConstMetadata())
 		{
-			return;
+			FinalizeInternal(InData, OutData, Metadata->GetItemCountForChild());
 		}
+		else
+		{
+			checkNoEntry();
+		}
+	}
 
-		for (int32 PointIndex = 0; PointIndex < Points.Num(); ++PointIndex)
+	template <typename T>
+	void FinalizeInternal(const T* InData, T* OutData, int32 NumElements)
+	{
+		const UPCGMetadata* InMetadata = InData->ConstMetadata();
+		UPCGMetadata* OutMetadata = OutData->MutableMetadata();
+		check(InMetadata && OutMetadata);
+
+		for (int32 Index = 0; Index < NumElements; ++Index)
 		{
-			const FPCGPoint& Point = Points[PointIndex];
-			int32 PartitionDataIndex = PartitionDataIndices[PointIndex];
+			int32 PartitionDataIndex = PartitionDataIndices[Index];
 			PCGMetadataEntryKey AttributeSetKey = PCGInvalidEntryKey;
 
 			if (PartitionDataIndex != INDEX_NONE)
 			{
-				AttributeSetKey = AttributeSetPartition.GetWeightedEntry(PartitionDataIndex, PointWeights[PointIndex]);
+				AttributeSetKey = AttributeSetPartition.GetWeightedEntry(PartitionDataIndex, Weights[Index]);
 			}
 
 			if (Settings->bKeepUnmatched || AttributeSetKey != PCGInvalidEntryKey)
 			{
-				FPCGPoint& OutPoint = OutPoints.Add_GetRef(Point);
+				PCGMetadataEntryKey PreviousKey = PCGMetadataEntryKey(Index);
+				PCGMetadataEntryKey NewEntry;
+
+				if constexpr (std::is_same_v<T, UPCGPointData>)
+				{
+					FPCGPoint& OutPoint = static_cast<UPCGPointData*>(OutData)->GetMutablePoints().Add_GetRef(static_cast<const UPCGPointData*>(InData)->GetPoints()[Index]);
+					PreviousKey = OutPoint.MetadataEntry;
+
+					if (AttributeSetKey != PCGInvalidEntryKey)
+					{
+						OutPoint.MetadataEntry = OutMetadata->AddEntry(OutPoint.MetadataEntry);
+					}
+
+					NewEntry = OutPoint.MetadataEntry;
+				}
+				else
+				{
+					NewEntry = OutMetadata->AddEntry();
+					OutMetadata->SetAttributes(NewEntry, InMetadata, PreviousKey);
+				}
 
 				if (AttributeSetKey != PCGInvalidEntryKey)
 				{
 					// This is similar to UPCGMetadata::SetAttributes but for a subset of attributes
-					OutPoint.MetadataEntry = OutPointData->MutableMetadata()->AddEntry(Point.MetadataEntry);
-
 					for (const TPair<const FPCGMetadataAttributeBase*, FPCGMetadataAttributeBase*>& AttributePair : AttributesToSet)
 					{
 						const FPCGMetadataAttributeBase* ParamAttribute = AttributePair.Key;
 						FPCGMetadataAttributeBase* PointAttribute = AttributePair.Value;
 
-						PointAttribute->SetValue(OutPoint.MetadataEntry, ParamAttribute, AttributeSetKey);
+						PointAttribute->SetValue(NewEntry, ParamAttribute, AttributeSetKey);
 					}
 				}
 			}
@@ -733,6 +793,8 @@ private:
 	TUniquePtr<const IPCGAttributeAccessorKeys> InputAttributeKeys;
 	TArray<TPair<const FPCGMetadataAttributeBase*, FPCGMetadataAttributeBase*>> AttributesToSet;
 	TMap<const UPCGData*, const UPCGData*> DataToMaxDistanceMap;
+	TArray<int32> PartitionDataIndices;
+	TArray<double> Weights;
 };
 
 FPCGMatchAndSetAttributesExecutionState::~FPCGMatchAndSetAttributesExecutionState()
@@ -832,18 +894,33 @@ bool FPCGMatchAndSetAttributesElement::PrepareDataInternal(FPCGContext* InContex
 	{
 		FPCGTaggedData& Output = Outputs.Add_GetRef(Inputs[IterationIndex]);
 
-		OutState.InPointData = Cast<UPCGPointData>(Inputs[IterationIndex].Data);
-		if (!OutState.InPointData)
+		OutState.InData = Inputs[IterationIndex].Data;
+		if (!OutState.InData || !OutState.InData->ConstMetadata())
 		{
-			PCGE_LOG_C(Error, GraphAndLog, InContext, FText::Format(LOCTEXT("InvalidInputDataType", "Input {0}: Input data must be of type Point"), FText::AsNumber(IterationIndex)));
+			PCGE_LOG_C(Error, GraphAndLog, InContext, FText::Format(LOCTEXT("InvalidInputDataType", "Input {0}: Input data must be of type Point or Param"), FText::AsNumber(IterationIndex)));
 			return EPCGTimeSliceInitResult::NoOperation;
 		}
 
-		OutState.OutPointData = NewObject<UPCGPointData>();
-		OutState.OutPointData->InitializeFromData(OutState.InPointData);
-		OutState.OutPointData->GetMutablePoints().Reserve(OutState.InPointData->GetPoints().Num());
+		if (const UPCGPointData* InPointData = Cast<const UPCGPointData>(OutState.InData))
+		{
+			UPCGPointData* OutPointData = NewObject<UPCGPointData>();
+			OutPointData->InitializeFromData(InPointData);
+			OutPointData->GetMutablePoints().Reserve(InPointData->GetPoints().Num());
+			OutState.OutData = OutPointData;
+		}
+		else if (const UPCGParamData* InParamData = Cast<const UPCGParamData>(OutState.InData))
+		{
+			UPCGParamData* OutParamData = NewObject<UPCGParamData>();
+			OutParamData->Metadata->AddAttributes(InParamData->Metadata);
+
+			OutState.OutData = OutParamData;
+		}
+		else
+		{
+			return EPCGTimeSliceInitResult::AbortExecution;
+		}
 
-		Output.Data = OutState.OutPointData;
+		Output.Data = OutState.OutData;
 
 		return EPCGTimeSliceInitResult::Success;
 	});
@@ -887,8 +964,8 @@ bool FPCGMatchAndSetAttributesElement::ExecuteInternal(FPCGContext* InContext) c
 		// It should be guaranteed to be a success at this point
 		check(InitResult == EPCGTimeSliceInitResult::Success);
 
-		// Run the execution until the time slice is finished
-		return ExecState.Partition->SelectPoints(*Context, IterState.InPointData, IterState.CurrentPointIndex, IterState.OutPointData);
+		// Run the execution until the time slice is finished. We actually don't care about the max index since we won't select anything.
+		return ExecState.Partition->SelectMultiple(*Context, IterState.InData, IterState.CurrentIndex, /*MaxIndex*/ 0, IterState.OutData);
 	});
 }
 
diff --git a/Engine/Plugins/PCG/Source/PCG/Private/Elements/PCGSpawnActor.cpp b/Engine/Plugins/PCG/Source/PCG/Private/Elements/PCGSpawnActor.cpp
--- a/Engine/Plugins/PCG/Source/PCG/Private/Elements/PCGSpawnActor.cpp
+++ b/Engine/Plugins/PCG/Source/PCG/Private/Elements/PCGSpawnActor.cpp
@@ -33,57 +33,61 @@ static TAutoConsoleVariable<bool> CVarAllowActorReuse(
 	true,
 	TEXT("Controls whether PCG spawned actors can be reused and skipped when re-executing"));
 
-class FPCGSpawnActorPartitionByAttribute : public FPCGPointDataPartitionBase<FPCGSpawnActorPartitionByAttribute, TSubclassOf<AActor>>
+class FPCGSpawnActorPartitionByAttribute : public FPCGDataPartitionBase<FPCGSpawnActorPartitionByAttribute, TSubclassOf<AActor>>
 {
 public:
 	FPCGSpawnActorPartitionByAttribute(FName InSpawnAttribute)
-		: FPCGPointDataPartitionBase<FPCGSpawnActorPartitionByAttribute, TSubclassOf<AActor>>()
+		: FPCGDataPartitionBase<FPCGSpawnActorPartitionByAttribute, TSubclassOf<AActor>>()
 		, SpawnAttribute(InSpawnAttribute)
 	{
 	}
 
-	bool InitializeForPointData(const UPCGPointData* PointData, UPCGPointData* OutPointData)
+	bool InitializeForData(const UPCGData* InData, UPCGData* OutData)
 	{
-		if (!PointData || !PointData->ConstMetadata())
+		if (!InData || !InData->IsA<UPCGPointData>() || !OutData || !OutData->IsA<UPCGPointData>())
 		{
 			return false;
 		}
 
 		FPCGAttributePropertyInputSelector InputSource;
 		InputSource.SetAttributeName(SpawnAttribute);
-		InputSource = InputSource.CopyAndFixLast(PointData);
-		SpawnAttributeAccessor = PCGAttributeAccessorHelpers::CreateConstAccessor(PointData, InputSource);
-		SpawnAttributeKeys = PCGAttributeAccessorHelpers::CreateConstKeys(PointData, InputSource);
+		InputSource = InputSource.CopyAndFixLast(InData);
+		SpawnAttributeAccessor = PCGAttributeAccessorHelpers::CreateConstAccessor(InData, InputSource);
+		SpawnAttributeKeys = PCGAttributeAccessorHelpers::CreateConstKeys(InData, InputSource);
 
 		return SpawnAttributeAccessor.IsValid() && SpawnAttributeKeys.IsValid();
 	}
 
-	void AddToPartitionData(FPCGPointDataPartitionBase::Element* SelectedElement, const UPCGPointData* ParentPointData, const FPCGPoint& Point)
+	void AddToPartitionData(FPCGDataPartitionBase::Element* SelectedElement, const UPCGData* ParentData, int32 Index)
 	{
+		// Already checked in initialize that this is a valid cast
+		const UPCGPointData* ParentPointData = static_cast<const UPCGPointData*>(ParentData);
+
 		check(SelectedElement);
 		if (!SelectedElement->PartitionData)
 		{
-			SelectedElement->PartitionData = NewObject<UPCGPointData>();
-			SelectedElement->PartitionData->InitializeFromData(ParentPointData);
+			UPCGPointData* PartitionData = NewObject<UPCGPointData>();
+			PartitionData->InitializeFromData(ParentPointData);
+			SelectedElement->PartitionData = PartitionData;
 		}
 
 		check(SelectedElement->PartitionData);
-		SelectedElement->PartitionData->GetMutablePoints().Add(Point);
+		static_cast<UPCGPointData*>(SelectedElement->PartitionData)->GetMutablePoints().Add(ParentPointData->GetPoints()[Index]);
 	}
 
-	FPCGPointDataPartitionBase::Element* SelectPoint(const FPCGPoint& Point, int32 PointIndex)
+	FPCGDataPartitionBase::Element* Select(int32 Index)
 	{
 		FSoftClassPath ActorPath;
 		TSoftClassPtr<AActor> ActorClassSoftPtr;
 
-		if (SpawnAttributeAccessor->Get<FSoftClassPath>(ActorPath, PointIndex, *SpawnAttributeKeys))
+		if (SpawnAttributeAccessor->Get<FSoftClassPath>(ActorPath, Index, *SpawnAttributeKeys))
 		{
 			ActorClassSoftPtr = TSoftClassPtr<AActor>(ActorPath);
 		}
 		else
 		{
 			FString ActorPathString;
-			if (SpawnAttributeAccessor->Get<FString>(ActorPathString, PointIndex, *SpawnAttributeKeys))
+			if (SpawnAttributeAccessor->Get<FString>(ActorPathString, Index, *SpawnAttributeKeys))
 			{
 				ActorPath = FSoftClassPath(ActorPathString);
 				ActorClassSoftPtr = TSoftClassPtr<AActor>(ActorPath);
@@ -522,7 +526,7 @@ bool FPCGSpawnActorElement::SpawnAndPrepareSubgraphs(FPCGSubgraphContext* Contex
 			int32 CurrentPointIndex = 0;
 
 			// Selection is still needed if are fully skipped in order to write to the OutPointData.
-			Selector.SelectPoints(*Context, PointData, CurrentPointIndex, OutPointData);
+			Selector.SelectMultiple(*Context, PointData, CurrentPointIndex, PointData->GetPoints().Num(), OutPointData);
 
 			if (!bFullySkippedDueToReuse)
 			{
@@ -531,7 +535,7 @@ bool FPCGSpawnActorElement::SpawnAndPrepareSubgraphs(FPCGSubgraphContext* Contex
 					FPCGTaggedData PartialInput = Input;
 					PartialInput.Data = Element.Value.PartitionData;
 
-					SpawnOrCollapse(Element.Key, nullptr, PartialInput, Element.Value.PartitionData, OutPointData);
+					SpawnOrCollapse(Element.Key, nullptr, PartialInput, static_cast<UPCGPointData*>(Element.Value.PartitionData), OutPointData);
 
 					// Exception case here: if we've spawned actors but are merging the PCG inputs,
 					// normally this node is taken as a subgraph node (e.g. no need to do anything more than forwarding the inputs)
diff --git a/Engine/Plugins/PCG/Source/PCG/Public/Elements/PCGMatchAndSetAttributes.h b/Engine/Plugins/PCG/Source/PCG/Public/Elements/PCGMatchAndSetAttributes.h
--- a/Engine/Plugins/PCG/Source/PCG/Public/Elements/PCGMatchAndSetAttributes.h
+++ b/Engine/Plugins/PCG/Source/PCG/Public/Elements/PCGMatchAndSetAttributes.h
@@ -11,7 +11,6 @@
 
 // Defined in the cpp file
 class FPCGMatchAndSetPartition;
-class UPCGPointData;
 
 struct FPCGMatchAndSetAttributesExecutionState
 {
@@ -22,9 +21,9 @@ struct FPCGMatchAndSetAttributesExecutionState
 
 struct FPCGMatchAndSetAttributesIterationState
 {
-	int CurrentPointIndex = 0;
-	const UPCGPointData* InPointData = nullptr;
-	UPCGPointData* OutPointData = nullptr;
+	int CurrentIndex = 0;
+	const UPCGData* InData = nullptr;
+	UPCGData* OutData = nullptr;
 };
 
 UENUM()
@@ -54,6 +53,7 @@ public:
 	virtual FText GetDefaultNodeTitle() const override;
 	virtual FText GetNodeTooltipText() const override;
 	virtual EPCGSettingsType GetType() const override { return EPCGSettingsType::Metadata; }
+	virtual bool HasDynamicPins() const { return true; }
 #endif // WITH_EDITOR
 	virtual bool IsInputPinRequiredByExecution(const UPCGPin* InPin) const override { return InPin->Properties.Label == PCGPinConstants::DefaultInputLabel; }
 
diff --git a/Engine/Plugins/PCG/Source/PCG/Public/Helpers/PCGPointDataPartition.h b/Engine/Plugins/PCG/Source/PCG/Public/Helpers/PCGPointDataPartition.h
--- a/Engine/Plugins/PCG/Source/PCG/Public/Helpers/PCGPointDataPartition.h
+++ b/Engine/Plugins/PCG/Source/PCG/Public/Helpers/PCGPointDataPartition.h
@@ -5,10 +5,9 @@
 #include "PCGContext.h"
 #include "PCGElement.h"
 #include "PCGModule.h"
-#include "Data/PCGPointData.h"
 
 /** 
-* Convenience class to do point data partitioning, i.e. building a TMap<Value, TArray<point/point indices>>
+* Convenience class to do data partitioning, i.e. building a TMap<Value, TArray<indices>>
 * Uses CRTP to remove the need for virtual calls.
 * Has some behavior options depending on the derived class.
 * Currently supports time-slicing (caveat for rooting data, not done yet)
@@ -21,17 +20,17 @@
 * Methods that are not re-implemented/hidden will retain their default implementation, which is mostly trivial.
 */
 template<typename Derived, typename KeyType>
-class FPCGPointDataPartitionBase
+class FPCGDataPartitionBase
 {
 public:
 	struct Element
 	{
-		TArray<int32> PointIndices;
-		UPCGPointData* PartitionData = nullptr;
+		TArray<int32> Indices;
+		UPCGData* PartitionData = nullptr;
 	};
 
-	FPCGPointDataPartitionBase() = default;
-	FPCGPointDataPartitionBase(const TArrayView<KeyType>& InKeys)
+	FPCGDataPartitionBase() = default;
+	FPCGDataPartitionBase(const TArrayView<KeyType>& InKeys)
 	{
 		for (const KeyType& Key : InKeys)
 		{
@@ -44,15 +43,15 @@ public:
 
 	/** Overridable behavior methods */
 	bool Initialize() { return true; }
-	bool InitializeForPointData(const UPCGPointData* PointData) { return true; }
-	void AddToPartitionData(Element* SelectedElement, const UPCGPointData* ParentPointData, const FPCGPoint& Point) {}
-	void WriteToOutputData(UPCGPointData* OutPointData, Element* SelectedElement, const FPCGPoint& Point, int32 PointIndex) {}
-	void Finalize(const UPCGPointData* InPointData, UPCGPointData* OutPointData) {}
-	Element* SelectPoint(const FPCGPoint& Point, int32 PointIndex) { return nullptr; }
+	bool InitializeForData(const UPCGData* Data) { return true; }
+	void AddToPartitionData(Element* SelectedElement, const UPCGData* ParentData, int32 Index) {}
+	void WriteToOutputData(const UPCGData* ParentData, UPCGData* OutData, Element* SelectedElement, int32 Index) {}
+	void Finalize(const UPCGData* InData, UPCGData* OutData) {}
+	Element* Select(int32 Index) { return nullptr; }
 	int32 TimeSlicingCheckFrequency() const { return 1024; }
 
 	/** API */
-	bool SelectPoints(FPCGContext& Context, const UPCGPointData* PointData, int32& InCurrentPointIndex, UPCGPointData* OutPointData);
+	bool SelectMultiple(FPCGContext& Context, const UPCGData* InData, int32& InCurrentIndex, int32 InMaxIndex, UPCGData* OutData);
 	void Reset() { ElementMap.Reset(); }
 
 	// Data access
@@ -60,44 +59,40 @@ public:
 };
 
 template<typename Derived, typename KeyType>
-bool FPCGPointDataPartitionBase<Derived, KeyType>::SelectPoints(FPCGContext& Context, const UPCGPointData* PointData, int32& InCurrentPointIndex, UPCGPointData* OutPointData)
+bool FPCGDataPartitionBase<Derived, KeyType>::SelectMultiple(FPCGContext& Context, const UPCGData* InData, int32& InCurrentIndex, int32 InMaxIndex, UPCGData* OutData)
 {
-	if (!PointData)
+	if (!InData)
 	{
-		PCGE_LOG_C(Error, GraphAndLog, &Context, NSLOCTEXT("PCGPointDataPartition", "InputMissingData", "Missing input data"));
+		PCGE_LOG_C(Error, GraphAndLog, &Context, NSLOCTEXT("PCGDataPartition", "InputMissingData", "Missing input data"));
 		return true;
 	}
 
-	if (InCurrentPointIndex == 0)
+	if (InCurrentIndex == 0)
 	{
-		if (!This()->InitializeForPointData(PointData, OutPointData))
+		if (!This()->InitializeForData(InData, OutData))
 		{
 			return true;
 		}
 	}
 
-	int32 CurrentPointIndex = InCurrentPointIndex;
-	int32 LastCheckpointIndex = InCurrentPointIndex;
+	int32 CurrentIndex = InCurrentIndex;
+	int32 LastCheckpointIndex = InCurrentIndex;
 	const int32 TimeSlicingCheckFrequency = This()->TimeSlicingCheckFrequency();
 
-	const TArray<FPCGPoint>& Points = PointData->GetPoints();
-
-	while (CurrentPointIndex < Points.Num())
+	while (CurrentIndex < InMaxIndex)
 	{
-		const FPCGPoint& Point = Points[CurrentPointIndex];
-
-		Element* SelectedElement = This()->SelectPoint(Point, CurrentPointIndex);
+		Element* SelectedElement = This()->Select(CurrentIndex);
 		if (SelectedElement)
 		{
-			SelectedElement->PointIndices.Add(CurrentPointIndex);
-			This()->AddToPartitionData(SelectedElement, PointData, Point);
+			SelectedElement->Indices.Add(CurrentIndex);
+			This()->AddToPartitionData(SelectedElement, InData, CurrentIndex);
 		}
 
-		This()->WriteToOutputData(OutPointData, SelectedElement, Point, CurrentPointIndex);
+		This()->WriteToOutputData(InData, OutData, SelectedElement, CurrentIndex);
 
-		++CurrentPointIndex;
+		++CurrentIndex;
 
-		if (CurrentPointIndex - LastCheckpointIndex >= TimeSlicingCheckFrequency)
+		if (CurrentIndex - LastCheckpointIndex >= TimeSlicingCheckFrequency)
 		{
 			if (Context.ShouldStop())
 			{
@@ -105,17 +100,17 @@ bool FPCGPointDataPartitionBase<Derived, KeyType>::SelectPoints(FPCGContext& Con
 			}
 			else
 			{
-				LastCheckpointIndex = CurrentPointIndex;
+				LastCheckpointIndex = CurrentIndex;
 			}
 		}
 	}
 
-	InCurrentPointIndex = CurrentPointIndex;
+	InCurrentIndex = CurrentIndex;
 
-	if (CurrentPointIndex == Points.Num())
+	if (CurrentIndex == InMaxIndex)
 	{
-		This()->Finalize(PointData, OutPointData);
+		This()->Finalize(InData, OutData);
 	}
 
-	return (CurrentPointIndex == Points.Num());
+	return (CurrentIndex == InMaxIndex);
 }
\ No newline at end of file
