[PCG] Format Subgraph node names the same way that we format them in the palette view

* This will add spaces between uppercase to the asset names

#rnx

diff --git a/Engine/Plugins/PCG/Source/PCG/Private/PCGSubgraph.cpp b/Engine/Plugins/PCG/Source/PCG/Private/PCGSubgraph.cpp
--- a/Engine/Plugins/PCG/Source/PCG/Private/PCGSubgraph.cpp
+++ b/Engine/Plugins/PCG/Source/PCG/Private/PCGSubgraph.cpp
@@ -295,7 +295,8 @@ FString UPCGSubgraphSettings::GetAdditionalTitleInformation() const
 
 	if (UPCGGraph* TargetSubgraph = GetSubgraph())
 	{
-		return TargetSubgraph->GetName();
+		// Use the same transformation than in the palette view to add spaces between uppercase characters
+		return FName::NameToDisplayString(TargetSubgraph->GetName(), /*bIsBool=*/false);
 	}
 	else
 	{
