[PCG] Subgraph changes notify the tracking to detect dynamically tracked subgraphs

* Also added a condition to avoid trigering a refresh that would have been picked up by a subgraph node or a pcg component

#rnx

diff --git a/Engine/Plugins/PCG/Source/PCG/Private/PCGGraph.cpp b/Engine/Plugins/PCG/Source/PCG/Private/PCGGraph.cpp
--- a/Engine/Plugins/PCG/Source/PCG/Private/PCGGraph.cpp
+++ b/Engine/Plugins/PCG/Source/PCG/Private/PCGGraph.cpp
@@ -1157,11 +1157,22 @@ void UPCGGraph::NotifyGraphChanged(EPCGChangeType ChangeType)
 		}
 	}
 
-	// Graph settings, nodes, graph structure can all change the higen grid sizes.
 	if (ChangeType != EPCGChangeType::Cosmetic)
 	{
-		FWriteScopeLock GridSizeLock(NodeToGridSizeLock);
-		NodeToGridSize.Reset();
+		{
+			// Graph settings, nodes, graph structure can all change the higen grid sizes.
+			FWriteScopeLock GridSizeLock(NodeToGridSizeLock);
+			NodeToGridSize.Reset();
+		}
+
+		// Also notify other systems that this graph changed, only if the owner is not a PCG Component nor PCG Subgraph.
+		// They already have their own system to trigger a refresh.
+		const UObject* Outer = GetOuter();
+		if (!Outer || !(Outer->IsA<UPCGComponent>() || Outer->IsA<UPCGSubgraphSettings>()))
+		{
+			FPropertyChangedEvent EmptyEvent{ nullptr };
+			FCoreUObjectDelegates::OnObjectPropertyChanged.Broadcast(this, EmptyEvent);
+		}
 	}
 
 	OnGraphChangedDelegate.Broadcast(this, ChangeType);
@@ -1598,6 +1609,15 @@ void UPCGGraphInstance::OnGraphChanged(UPCGGraphInterface* InGraph, EPCGChangeTy
 {
 	if (InGraph == Graph)
 	{
+		// Also notify other systems that this graph changed, only if the owner is not a PCG Component nor PCG Subgraph.
+		// They already have their own system to trigger a refresh.
+		const UObject* Outer = GetOuter();
+		if (!Outer || !(Outer->IsA<UPCGComponent>() || Outer->IsA<UPCGSubgraphSettings>()))
+		{
+			FPropertyChangedEvent EmptyEvent{ nullptr };
+			FCoreUObjectDelegates::OnObjectPropertyChanged.Broadcast(this, EmptyEvent);
+		}
+
 		OnGraphChangedDelegate.Broadcast(this, ChangeType);
 	}
 }
