[PCG] Add another pass of override for graph parameters if a subgraph is overridden

* This will also fix problems when graphs instances were not respected as overrides when the default graph was set.

#rnx

diff --git a/Engine/Plugins/PCG/Source/PCG/Private/Elements/PCGLoopElement.cpp b/Engine/Plugins/PCG/Source/PCG/Private/Elements/PCGLoopElement.cpp
--- a/Engine/Plugins/PCG/Source/PCG/Private/Elements/PCGLoopElement.cpp
+++ b/Engine/Plugins/PCG/Source/PCG/Private/Elements/PCGLoopElement.cpp
@@ -184,6 +184,8 @@ bool FPCGLoopElement::ExecuteInternal(FPCGContext* InContext) const
 #if WITH_EDITOR
 			FPCGDynamicTrackingHelper::AddSingleDynamicTrackingKey(Context, FPCGSelectionKey::CreateFromPath(Settings->SubgraphOverride), /*bIsCulled=*/false);
 #endif // WITH_EDITOR
+
+			Context->UpdateOverridesWithOverriddenGraph();
 		}
 
 		UPCGGraph* Subgraph = Settings->GetSubgraph();
diff --git a/Engine/Plugins/PCG/Source/PCG/Private/PCGSubgraph.cpp b/Engine/Plugins/PCG/Source/PCG/Private/PCGSubgraph.cpp
--- a/Engine/Plugins/PCG/Source/PCG/Private/PCGSubgraph.cpp
+++ b/Engine/Plugins/PCG/Source/PCG/Private/PCGSubgraph.cpp
@@ -424,16 +424,20 @@ void* FPCGSubgraphContext::GetUnsafeExternalContainerForOverridableParam(const F
 	}
 }
 
-FPCGContext* FPCGSubgraphElement::Initialize(const FPCGDataCollection& InputData, TWeakObjectPtr<UPCGComponent> SourceComponent, const UPCGNode* Node)
+void FPCGSubgraphContext::InitializeUserParametersStruct()
 {
-	FPCGSubgraphContext* Context = new FPCGSubgraphContext();
-	Context->InputData = InputData;
-	Context->SourceComponent = SourceComponent;
-	Context->Node = Node;
+	// Only duplicate the UserParameters if we have overridable params and we have at least one param pin connected.
+	GraphInstanceParametersOverride.Reset();
 
-	// Only duplicate the UserParameters if we have overriable params and we have at least one param pin connected.
-	const UPCGBaseSubgraphSettings* Settings = Context->GetInputSettings<UPCGBaseSubgraphSettings>();
+	// Will return the OG settings the first time this is called (when initializing the element/context)
+	// and will contain the "hardcoded" graph in the subgraph node.
+	// If subgraph is overridden, this be called a second time with the SubgraphInstance containing the updated overridden subgraph.
+	const UPCGBaseSubgraphSettings* Settings = GetInputSettings<UPCGBaseSubgraphSettings>();
 	check(Settings);
+
+	const UClass* SettingsClass = Settings->GetClass();
+	check(SettingsClass);
+
 	const TArray<FPCGSettingsOverridableParam>& OverridableParams = Settings->OverridableParams();
 
 	const UPCGGraphInterface* Graph = Settings->GetSubgraphInterface();
@@ -447,14 +451,49 @@ FPCGContext* FPCGSubgraphElement::Initialize(const FPCGDataCollection& InputData
 		int32 Index = 0;
 		while (!bHasParamConnected && Index < OverridableParams.Num())
 		{
-			bHasParamConnected |= !InputData.GetParamsByPin(OverridableParams[Index++].Label).IsEmpty();
+			// Discard any override that is a property of the settings (we are looking for overrides for the graph instance)
+			// We use the first property name, since it will be the one related to this settings properties.
+			const FName PropertyName = OverridableParams[Index].PropertiesNames.IsEmpty() ? NAME_None : OverridableParams[Index].PropertiesNames[0];
+			const FName Label = OverridableParams[Index].Label;
+			++Index;
+
+			if (PropertyName != NAME_None && SettingsClass->FindPropertyByName(PropertyName))
+			{
+				continue;
+			}
+
+			bHasParamConnected |= !InputData.GetParamsByPin(Label).IsEmpty();
 		}
 
 		if (bHasParamConnected)
 		{
-			Context->GraphInstanceParametersOverride = FInstancedStruct(UserParametersView);
+			GraphInstanceParametersOverride = FInstancedStruct(UserParametersView);
 		}
 	}
+}
+
+void FPCGSubgraphContext::UpdateOverridesWithOverriddenGraph()
+{
+	// We have a "catch-22" kind of problem here. When we initialize the subgraph element, we look for the graph "hardcoded"
+	// in the settings to duplicate its user parameters and override it with the override pins. But Subgraph override is also coming from
+	// the override pins, meaning that without adding some kind of "read order" on the override, we have to read them twice. It's less efficient
+	// but makes things simpler to understand. We have a first override read to get the Subgraph Override, then if it is set, we have a second
+	// read with the user parameters of the subgraph override.
+	InitializeUserParametersStruct();
+	if (GraphInstanceParametersOverride.IsValid())
+	{
+		OverrideSettings();
+	}
+}
+
+FPCGContext* FPCGSubgraphElement::Initialize(const FPCGDataCollection& InputData, TWeakObjectPtr<UPCGComponent> SourceComponent, const UPCGNode* Node)
+{
+	FPCGSubgraphContext* Context = new FPCGSubgraphContext();
+	Context->InputData = InputData;
+	Context->SourceComponent = SourceComponent;
+	Context->Node = Node;
+
+	Context->InitializeUserParametersStruct();
 
 	return Context;
 }
@@ -538,6 +577,8 @@ bool FPCGSubgraphElement::ExecuteInternal(FPCGContext* InContext) const
 #if WITH_EDITOR
 			FPCGDynamicTrackingHelper::AddSingleDynamicTrackingKey(Context, FPCGSelectionKey::CreateFromPath(Settings->SubgraphOverride), /*bIsCulled=*/false);
 #endif // WITH_EDITOR
+
+			Context->UpdateOverridesWithOverriddenGraph();
 		}
 	}
 
diff --git a/Engine/Plugins/PCG/Source/PCG/Public/PCGSubgraph.h b/Engine/Plugins/PCG/Source/PCG/Public/PCGSubgraph.h
--- a/Engine/Plugins/PCG/Source/PCG/Public/PCGSubgraph.h
+++ b/Engine/Plugins/PCG/Source/PCG/Public/PCGSubgraph.h
@@ -165,6 +165,12 @@ struct PCG_API FPCGSubgraphContext : public FPCGContext
 	bool bScheduledSubgraph = false;
 	FInstancedStruct GraphInstanceParametersOverride;
 
+	// Analyze input data to detect if there is any override for the user parameters. If so will duplicate it to gather overrides.
+	void InitializeUserParametersStruct();
+
+	// If we have a subgraph override, update the underlying duplicated parameters with the overrides from the subgraph.
+	void UpdateOverridesWithOverriddenGraph();
+
 protected:
 	virtual void* GetUnsafeExternalContainerForOverridableParam(const FPCGSettingsOverridableParam& InParam) override;
 };
