[PCG] Hotfix: Attribute Filter and Attribute Filter Range now default to Any for Output Pin rather than Points

#rnx

diff --git a/Engine/Plugins/PCG/Source/PCG/Private/Elements/PCGAttributeFilter.cpp b/Engine/Plugins/PCG/Source/PCG/Private/Elements/PCGAttributeFilter.cpp
--- a/Engine/Plugins/PCG/Source/PCG/Private/Elements/PCGAttributeFilter.cpp
+++ b/Engine/Plugins/PCG/Source/PCG/Private/Elements/PCGAttributeFilter.cpp
@@ -261,7 +261,7 @@ EPCGDataType UPCGAttributeFilteringSettings::GetCurrentPinTypes(const UPCGPin* I
 	const EPCGDataType InputTypeUnion = GetTypeUnionOfIncidentEdges(PCGPinConstants::DefaultInputLabel);
 
 	// Spatial is collapsed into points
-	if ((InputTypeUnion & EPCGDataType::Spatial) == InputTypeUnion)
+	if (InputTypeUnion != EPCGDataType::None && (InputTypeUnion & EPCGDataType::Spatial) == InputTypeUnion)
 	{
 		return EPCGDataType::Point;
 	}
@@ -390,7 +390,7 @@ EPCGDataType UPCGAttributeFilteringRangeSettings::GetCurrentPinTypes(const UPCGP
 	const EPCGDataType InputTypeUnion = GetTypeUnionOfIncidentEdges(PCGPinConstants::DefaultInputLabel);
 
 	// Spatial is collapsed into points
-	if ((InputTypeUnion & EPCGDataType::Spatial) == InputTypeUnion)
+	if (InputTypeUnion != EPCGDataType::None && (InputTypeUnion & EPCGDataType::Spatial) == InputTypeUnion)
 	{
 		return EPCGDataType::Point;
 	}
diff --git a/Engine/Plugins/PCG/Source/PCG/Public/Elements/PCGAttributeFilter.h b/Engine/Plugins/PCG/Source/PCG/Public/Elements/PCGAttributeFilter.h
--- a/Engine/Plugins/PCG/Source/PCG/Public/Elements/PCGAttributeFilter.h
+++ b/Engine/Plugins/PCG/Source/PCG/Public/Elements/PCGAttributeFilter.h
@@ -77,7 +77,7 @@ public:
 #if WITH_EDITOR
 	virtual FName GetDefaultNodeName() const override { return FName(TEXT("AttributeFilter")); }
 	virtual FText GetDefaultNodeTitle() const override { return NSLOCTEXT("PCGAttributeFilteringElement", "NodeTitle", "Attribute Filter"); }
-	virtual TArray<FText> GetNodeTitleAliases() const { return { NSLOCTEXT("PCGAttributeFilteringElement", "AliasNodeTitle", "Point Filter") }; }
+	virtual TArray<FText> GetNodeTitleAliases() const override { return { NSLOCTEXT("PCGAttributeFilteringElement", "AliasNodeTitle", "Point Filter") }; }
 	virtual EPCGSettingsType GetType() const override { return EPCGSettingsType::Filter; }
 	virtual bool HasDynamicPins() const override { return true; }
 #endif
@@ -145,7 +145,7 @@ public:
 #if WITH_EDITOR
 	virtual FName GetDefaultNodeName() const override { return FName(TEXT("AttributeFilterRange")); }
 	virtual FText GetDefaultNodeTitle() const override { return NSLOCTEXT("PCGAttributeFilteringElement", "NodeTitleRange", "Attribute Filter Range"); }
-	virtual TArray<FText> GetNodeTitleAliases() const { return { NSLOCTEXT("PCGAttributeFilteringElement", "AliasNodeTitleRange", "Point Filter Range") }; }
+	virtual TArray<FText> GetNodeTitleAliases() const override { return { NSLOCTEXT("PCGAttributeFilteringElement", "AliasNodeTitleRange", "Point Filter Range") }; }
 	virtual EPCGSettingsType GetType() const override { return EPCGSettingsType::Filter; }
 	virtual bool HasDynamicPins() const override { return true; }
 #endif
