[PCG] Create Sample Texture node
--Samples color of texture at each point

diff --git a/Engine/Plugins/PCG/Source/PCG/Private/Data/PCGTextureData.cpp b/Engine/Plugins/PCG/Source/PCG/Private/Data/PCGTextureData.cpp
--- a/Engine/Plugins/PCG/Source/PCG/Private/Data/PCGTextureData.cpp
+++ b/Engine/Plugins/PCG/Source/PCG/Private/Data/PCGTextureData.cpp
@@ -41,6 +41,29 @@ namespace PCGTextureSamplingHelpers
 		return PlatformData->Mips.Num() == 1 && PlatformData->PixelFormat == PF_B8G8R8A8;
 	}
 
+	template<typename ValueType>
+	ValueType SampleInternal(FVector2D PositionLocalSpace,
+		int32 Width,
+		int32 Height,
+		TFunctionRef<ValueType(int32 Index)> SamplingFunction)
+	{
+		// Accounts for texel values being at texel centers
+		const double TexelX = (PositionLocalSpace.X * Width - 0.5);
+		const double TexelY = (PositionLocalSpace.Y * Height - 0.5);
+
+		const int32 X0 = FMath::Clamp(FMath::FloorToInt(TexelX), 0, Width - 1);
+		const int32 X1 = FMath::Min(X0 + 1, Width - 1);
+		const int32 Y0 = FMath::Clamp(FMath::FloorToInt(TexelY), 0, Height - 1);
+		const int32 Y1 = FMath::Min(Y0 + 1, Height - 1);
+
+		const ValueType SampleX0Y0 = SamplingFunction(X0 + Y0 * Width);
+		const ValueType SampleX1Y0 = SamplingFunction(X1 + Y0 * Width);
+		const ValueType SampleX0Y1 = SamplingFunction(X0 + Y1 * Width);
+		const ValueType SampleX1Y1 = SamplingFunction(X1 + Y1 * Width);
+
+		return FMath::BiLerp(SampleX0Y0, SampleX1Y0, SampleX0Y1, SampleX1Y1, TexelX - X0, TexelY - Y0);
+	}
+
 	template<typename ValueType>
 	bool Sample(const FVector2D& InPosition,
 		const FBox2D& InSurface,
@@ -62,8 +85,8 @@ namespace PCGTextureSamplingHelpers
 		FVector2D Pos = FVector2D::ZeroVector;
 		if (!InTextureData->bUseAdvancedTiling)
 		{
-			// TODO: There seems to be a bias issue here, as the bounds size are not in the same space as the texels.
-			Pos = LocalSpacePos * FVector2D(Width, Height);
+			Pos.X = FMath::Clamp(LocalSpacePos.X, 0.0, 1.0);
+			Pos.Y = FMath::Clamp(LocalSpacePos.Y, 0.0, 1.0);
 		}
 		else
 		{
@@ -91,43 +114,10 @@ namespace PCGTextureSamplingHelpers
 			FVector::FReal X = FMath::Frac(SamplePosition.X + 0.5);
 			FVector::FReal Y = FMath::Frac(SamplePosition.Y + 0.5);
 
-			X *= Width;
-			Y *= Height;
-
 			Pos = FVector2D(X, Y);
 		}
 
-		// TODO: this isn't super robust, if that becomes an issue
-		int32 X0 = FMath::FloorToInt(Pos.X);
-		if (X0 < 0 || X0 >= Width)
-		{
-			X0 = 0;
-		}
-
-		int32 X1 = FMath::CeilToInt(Pos.X);
-		if (X1 < 0 || X1 >= Width)
-		{
-			X1 = 0;
-		}
-
-		int32 Y0 = FMath::FloorToInt(Pos.Y);
-		if (Y0 < 0 || Y0 >= Height)
-		{
-			Y0 = 0;
-		}
-
-		int32 Y1 = FMath::CeilToInt(Pos.Y);
-		if (Y1 < 0 || Y1 >= Height)
-		{
-			Y1 = 0;
-		}
-
-		ValueType SampleX0Y0 = SamplingFunction(X0 + Y0 * Width);
-		ValueType SampleX1Y0 = SamplingFunction(X1 + Y0 * Width);
-		ValueType SampleX0Y1 = SamplingFunction(X0 + Y1 * Width);
-		ValueType SampleX1Y1 = SamplingFunction(X1 + Y1 * Width);
-
-		SampledValue = FMath::BiLerp(SampleX0Y0, SampleX1Y0, SampleX0Y1, SampleX1Y1, Pos.X - X0, Pos.Y - Y0);
+		SampledValue = SampleInternal(Pos, Width, Height, SamplingFunction);
 		return true;
 	}
 
@@ -177,7 +167,6 @@ bool UPCGBaseTextureData::SamplePoint(const FTransform& InTransform, const FBox&
 	// TODO: embed local bounds center offset at this time?
 	OutPoint.Transform = InTransform;
 	FVector PointPositionInLocalSpace = Transform.InverseTransformPosition(InTransform.GetLocation());
-	PointPositionInLocalSpace.Z = 0;
 	OutPoint.Transform.SetLocation(Transform.TransformPosition(PointPositionInLocalSpace));
 	OutPoint.SetLocalBounds(InBounds); // TODO: should set Min.Z = Max.Z = 0;
 
@@ -191,7 +180,7 @@ bool UPCGBaseTextureData::SamplePoint(const FTransform& InTransform, const FBox&
 	{
 		OutPoint.Color = Color;
 		OutPoint.Density = ((DensityFunction == EPCGTextureDensityFunction::Ignore) ? 1.0f : PCGTextureSamplingHelpers::SampleFloatChannel(Color, ColorChannel));
-		return OutPoint.Density > 0;
+		return OutPoint.Density > 0 || bKeepZeroDensityPoints;
 	}
 	else
 	{
@@ -272,6 +261,25 @@ bool UPCGBaseTextureData::IsValid() const
 	return Height > 0 && Width > 0;
 }
 
+bool UPCGBaseTextureData::SamplePointLocal(const FVector2D& LocalPosition, FVector4& OutColor, float& OutDensity) const
+{
+	if (!ensure(Width > 0 && Height > 0))
+	{
+		return false;
+	}
+
+	FVector2D Pos;
+	Pos.X = FMath::Frac(LocalPosition.X);
+	Pos.Y = FMath::Frac(LocalPosition.Y);
+
+	const FLinearColor OutSample = PCGTextureSamplingHelpers::SampleInternal<FLinearColor>(Pos, Width, Height, [this](int32 Index) { return ColorData[Index]; });
+
+	OutColor = OutSample;
+	OutDensity = (DensityFunction == EPCGTextureDensityFunction::Ignore) ? 1.0f : PCGTextureSamplingHelpers::SampleFloatChannel(OutSample, ColorChannel);
+	
+	return OutDensity > 0.0 || bKeepZeroDensityPoints;
+}
+
 void UPCGBaseTextureData::CopyBaseTextureData(UPCGBaseTextureData* NewTextureData) const
 {
 	CopyBaseSurfaceData(NewTextureData);
@@ -589,12 +597,27 @@ bool UPCGTextureData::InitializeFromGPUTexture(const TFunction<void()>& PostInit
 #if WITH_EDITOR
 bool UPCGTextureData::InitializeGPUTextureFromCPU()
 {
-	if (!DuplicateTexture || !PCGTextureSamplingHelpers::CanGPUTextureBeCPUAccessed(DuplicateTexture))
+	// There's a bit of a mix of texture types in this class currently, due to some functionality for readback being 2D-only.
+	UTexture2D* TextureAs2D = Cast<UTexture2D>(Texture.Get());
+	if (!TextureAs2D)
 	{
-		return false;
+		if (UTexture2DArray* Texture2DArray = Cast<UTexture2DArray>(Texture))
+		{
+			TextureAs2D = Texture2DArray->SourceTextures.IsValidIndex(TextureIndex) ? Texture2DArray->SourceTextures[TextureIndex] : nullptr;
+		}
+	}
+
+	UTexture2D* TextureForReadback = nullptr;
+	if (TextureAs2D && PCGTextureSamplingHelpers::CanGPUTextureBeCPUAccessed(TextureAs2D))
+	{
+		TextureForReadback = TextureAs2D;
+	}
+	else if (PCGTextureSamplingHelpers::CanGPUTextureBeCPUAccessed(DuplicateTexture))
+	{
+		TextureForReadback = DuplicateTexture;
 	}
 
-	FTexturePlatformData* PlatformData = DuplicateTexture->GetPlatformData();
+	FTexturePlatformData* PlatformData = TextureForReadback ? TextureForReadback->GetPlatformData() : nullptr;
 	if (!PlatformData)
 	{
 		return false;
diff --git a/Engine/Plugins/PCG/Source/PCG/Private/Elements/PCGSampleTexture.cpp b/Engine/Plugins/PCG/Source/PCG/Private/Elements/PCGSampleTexture.cpp
new file mode 100644
--- /dev/null
+++ b/Engine/Plugins/PCG/Source/PCG/Private/Elements/PCGSampleTexture.cpp
@@ -0,0 +1,129 @@
+// Copyright Epic Games, Inc. All Rights Reserved.
+
+#include "Elements/PCGSampleTexture.h"
+
+#include "PCGContext.h"
+#include "PCGPin.h"
+#include "Data/PCGPointData.h"
+#include "Data/PCGTextureData.h"
+#include "Helpers/PCGAsync.h"
+#include "Metadata/Accessors/PCGAttributeAccessorHelpers.h"
+
+#define LOCTEXT_NAMESPACE "PCGSampleTextureElement"
+
+TArray<FPCGPinProperties> UPCGSampleTextureSettings::InputPinProperties() const
+{
+	TArray<FPCGPinProperties> PinProperties;
+
+	PinProperties.Emplace(PCGSampleTextureConstants::InputPointLabel,
+		EPCGDataType::Point,
+		/*bAllowMultipleConnections=*/true,
+		/*bAllowMultipleData=*/true);
+
+	PinProperties.Emplace(PCGSampleTextureConstants::InputTextureLabel,
+		EPCGDataType::Texture,
+		/*bAllowMultipleConnections=*/false,
+		/*bAllowMultipleData=*/false);
+
+	return PinProperties;
+}
+
+FPCGElementPtr UPCGSampleTextureSettings::CreateElement() const
+{
+	return MakeShared<FPCGSampleTextureElement>();
+}
+
+bool FPCGSampleTextureElement::ExecuteInternal(FPCGContext* Context) const
+{
+	TRACE_CPUPROFILER_EVENT_SCOPE(FPCGSampleTextureElement::Execute);
+	check(Context);
+
+	const UPCGSampleTextureSettings* Settings = Context->GetInputSettings<UPCGSampleTextureSettings>();
+	check(Settings);
+
+	const TArray<FPCGTaggedData> PointInputs = Context->InputData.GetInputsByPin(PCGSampleTextureConstants::InputPointLabel);
+	const TArray<FPCGTaggedData> BaseTextureInput = Context->InputData.GetInputsByPin(PCGSampleTextureConstants::InputTextureLabel);
+	TArray<FPCGTaggedData>& Outputs = Context->OutputData.TaggedData;
+
+	if (BaseTextureInput.Num() > 1)
+	{
+		PCGE_LOG(Warning, GraphAndLog, LOCTEXT("InvalidNumberOfTextureData", "Only 1 texture input is allowed."));
+	}
+
+	const UPCGBaseTextureData* BaseTextureData = !BaseTextureInput.IsEmpty() ? Cast<UPCGBaseTextureData>(BaseTextureInput[0].Data) : nullptr;
+	if (!BaseTextureData)
+	{
+		return true;
+	}
+
+	for (int i = 0; i < PointInputs.Num(); ++i)
+	{
+		const UPCGPointData* PointData = Cast<UPCGPointData>(PointInputs[i].Data);
+		if (!PointData)
+		{
+			PCGE_LOG(Warning, GraphAndLog, FText::Format(LOCTEXT("InvalidPointData", "Point Input {0} is not point data."), i));
+			continue;
+		}
+
+		const TArray<FPCGPoint>& InputPoints = PointData->GetPoints();
+
+		TUniquePtr<const IPCGAttributeAccessor> InputAccessor = nullptr;
+		TUniquePtr<const IPCGAttributeAccessorKeys> InputKeys = nullptr;
+
+		if (Settings->TextureMappingMethod == EPCGTextureMappingMethod::UVCoordinates)
+		{
+			const FPCGAttributePropertyInputSelector UVSource = Settings->UVCoordinatesAttribute.CopyAndFixLast(PointData);
+	
+			InputAccessor = PCGAttributeAccessorHelpers::CreateConstAccessor(PointData, UVSource);
+			InputKeys = PCGAttributeAccessorHelpers::CreateConstKeys(PointData, UVSource);
+			if (!InputAccessor || !InputKeys)
+			{
+				PCGE_LOG(Warning, GraphAndLog, FText::Format(LOCTEXT("InvalidUVAccessor", "Could not create coordinate accessor {0} for Point Input {1}."), FText::FromName(UVSource.GetName()), i));
+				continue;
+			}
+
+			if (!PCG::Private::IsOfTypes<FVector, FVector2D>(InputAccessor->GetUnderlyingType()))
+			{
+				PCGE_LOG(Warning, GraphAndLog, FText::Format(LOCTEXT("InvalidAccessorType", "Accessor {0} must be of type Vector2 or Vector3"), FText::FromName(UVSource.GetName())));
+				continue;
+			}
+		}
+
+		UPCGPointData* OutPointData = NewObject<UPCGPointData>();
+		OutPointData->InitializeFromData(PointData);
+
+		auto ProcessPoint = [Settings, BaseTextureData, &InputAccessor, &InputKeys, &InputPoints, OutPointData](int32 Index, FPCGPoint& OutPoint)
+		{
+			OutPoint = InputPoints[Index];
+
+			if (Settings->TextureMappingMethod == EPCGTextureMappingMethod::UVCoordinates)
+			{
+				FVector OutSamplePosition;
+				InputAccessor->Get(OutSamplePosition, Index, *InputKeys, EPCGAttributeAccessorFlags::AllowBroadcast | EPCGAttributeAccessorFlags::AllowConstructible);
+
+				if (Settings->TilingMode == EPCGTextureAddressMode::Clamp)
+				{
+					OutSamplePosition.X = FMath::Clamp(OutSamplePosition.X, 0.0, 1.0);
+					OutSamplePosition.Y = FMath::Clamp(OutSamplePosition.Y, 0.0, 1.0);
+				}
+
+				return BaseTextureData->SamplePointLocal(FVector2D(OutSamplePosition), OutPoint.Color, OutPoint.Density);
+			}
+			else
+			{
+				return BaseTextureData->SamplePoint(OutPoint.Transform, OutPoint.GetLocalBounds(), OutPoint, OutPointData->Metadata);
+			}
+		};
+
+		TArray<FPCGPoint>& OutPoints = OutPointData->GetMutablePoints();
+		OutPoints.SetNumUninitialized(InputPoints.Num());
+		FPCGTaggedData& Output = Outputs.Add_GetRef(PointInputs[i]);
+		Output.Data = OutPointData;
+
+		FPCGAsync::AsyncPointProcessing(Context, InputPoints.Num(), OutPoints, ProcessPoint);
+	}
+
+	return true;
+}
+
+#undef LOCTEXT_NAMESPACE
diff --git a/Engine/Plugins/PCG/Source/PCG/Private/Tests/Elements/PCGSampleTextureTest.cpp b/Engine/Plugins/PCG/Source/PCG/Private/Tests/Elements/PCGSampleTextureTest.cpp
new file mode 100644
--- /dev/null
+++ b/Engine/Plugins/PCG/Source/PCG/Private/Tests/Elements/PCGSampleTextureTest.cpp
@@ -0,0 +1,187 @@
+// Copyright Epic Games, Inc. All Rights Reserved.
+
+#include "Tests/PCGTestsCommon.h"
+
+#include "PCGContext.h"
+#include "Data/PCGPointData.h"
+#include "Data/PCGTextureData.h"
+#include "Engine/Texture2D.h"
+#include "Metadata/Accessors/PCGAttributeAccessorHelpers.h"
+#include "UObject/Package.h"
+
+#include "Elements/PCGSampleTexture.h"
+
+#if WITH_EDITOR
+class PCGSampleTextureTestBase : public FPCGTestBaseClass
+{
+public:
+	using FPCGTestBaseClass::FPCGTestBaseClass;
+
+	struct FTestParameters
+	{
+		FName UVAttribute = TEXT("UV");
+		EPCGTextureMappingMethod TextureMappingMethod = EPCGTextureMappingMethod::Planar;
+		FVector2D UVPosition;
+		FVector4 ExpectedColor = FVector4();
+		float ExpectedDensity = 0.0f;
+	};
+protected:
+	UPCGTextureData* GenerateTextureData()
+	{
+		UTexture2D* Texture = NewObject<UTexture2D>(GetTransientPackage(), NAME_None, RF_Transient);
+		
+		constexpr int32 SizeX = 4;
+		constexpr int32 SizeY = 4;
+		Texture->Source.Init(
+			SizeX,
+			SizeY,
+			1,
+			1,
+			TSF_BGRA8
+		);
+
+		Texture->SRGB = false;
+		Texture->CompressionNone = true;
+		Texture->MipGenSettings = TMGS_NoMipmaps;
+		Texture->AddressX = TA_Clamp;
+		Texture->AddressY = TA_Clamp;
+
+		const uint32 BPP = Texture->Source.GetBytesPerPixel();
+
+		uint8* TexData = Texture->Source.LockMip(0);
+		FMemory::Memzero(TexData, SizeX * SizeY * sizeof(uint8) * BPP);
+
+		for (int32 Y = 0; Y < SizeY; ++Y)
+		{
+			for (int32 X = 0; X < SizeX; ++X)
+			{
+				// 4 channels
+				TexData[(X + Y * SizeX) * BPP] = static_cast<uint8>((X / 8.0f) * 255); // Blue
+				TexData[(X + Y * SizeX) * BPP + 1] = static_cast<uint8>((Y / 8.0f) * 255); // Green
+				TexData[(X + Y * SizeX) * BPP + 2] = static_cast<uint8>(1.0f * 255); // Red
+				TexData[(X + Y * SizeX) * BPP + 3] = static_cast<uint8>(0.5f * 255); // Alpha
+			}
+		}
+		
+		Texture->Source.UnlockMip(0);
+		Texture->UpdateResource();
+
+		UPCGTextureData* TextureData = NewObject<UPCGTextureData>();
+		TextureData->Initialize(Texture, /*InTextureIndex*/0, FTransform(), /*PostInitializeCallback*/[]() {});
+
+		return TextureData;
+	}
+	
+	bool GenerateTestDataRunAndValidate(const FTestParameters& Parameters)
+	{
+		PCGTestsCommon::FTestData TestData;
+		PCGTestsCommon::GenerateSettings<UPCGSampleTextureSettings>(TestData);
+		UPCGSampleTextureSettings* Settings = CastChecked<UPCGSampleTextureSettings>(TestData.Settings);
+		Settings->UVCoordinatesAttribute.SetAttributeName(Parameters.UVAttribute);
+		Settings->TextureMappingMethod = Parameters.TextureMappingMethod;
+	
+		FPCGTaggedData& PointInputs = TestData.InputData.TaggedData.Emplace_GetRef();
+		UPCGPointData* InData = PCGTestsCommon::CreatePointData();
+		PointInputs.Data = InData;
+		PointInputs.Pin = PCGSampleTextureConstants::InputPointLabel;
+
+		FPCGMetadataAttribute<FVector2D>* InputUV = InData->Metadata->CreateAttribute<FVector2D>(Parameters.UVAttribute, FVector2D::ZeroVector, /*bAllowsInterpolation=*/true, /*bOverrideParent=*/true);
+		TArray<FPCGPoint>& Points = InData->GetMutablePoints();
+		Points.SetNum(16);
+		for (int i = 0; i < 4; ++i)
+		{
+			for (int j = 0; j < 4; ++j)
+			{
+				int CurrentPoint = i * 4 + j;
+
+				// Place each point at the texel center for easy UV access
+				Points[CurrentPoint].Transform.SetLocation(FVector((j + 0.5) * 0.25, (i + 0.5) * 0.25, 0.0));
+
+				// Set the UV metadata for each point
+				InData->Metadata->InitializeOnSet(Points[CurrentPoint].MetadataEntry);
+				InputUV->SetValue(Points[CurrentPoint].MetadataEntry, FVector2D(Points[CurrentPoint].Transform.GetLocation()));
+			}
+		}
+
+		FPCGTaggedData& TextureInputs = TestData.InputData.TaggedData.Emplace_GetRef();
+		TextureInputs.Data = GenerateTextureData();
+		TextureInputs.Pin = PCGSampleTextureConstants::InputTextureLabel;
+
+		FPCGElementPtr TestElement = TestData.Settings->GetElement();
+
+		TUniquePtr<FPCGContext> Context = TestData.InitializeTestContext();
+
+		while (!TestElement->Execute(Context.Get())) {}
+
+		const TArray<FPCGTaggedData>& Outputs = Context->OutputData.TaggedData;
+
+		UTEST_EQUAL("Output count", Outputs.Num(), 1);
+
+		const UPCGPointData* OutPointData = Cast<UPCGPointData>(Outputs[0].Data);
+
+		UTEST_NOT_NULL("Output point data", OutPointData);
+
+		const TArray<FPCGPoint>& OutPoints = OutPointData->GetPoints();
+
+		UTEST_EQUAL("Output point count", OutPoints.Num(), 16);
+
+		const FPCGMetadataAttribute<FVector2D>* OutputUV = OutPointData->Metadata->GetConstTypedAttribute<FVector2D>(Parameters.UVAttribute);
+		UTEST_NOT_NULL("UV Attribute exists", OutputUV);
+
+		PCGMetadataValueKey ValueKey = OutputUV->FindValue(Parameters.UVPosition);
+
+		UTEST_TRUE("Output Value Key is not default value key.", ValueKey != PCGDefaultValueKey)
+
+		for (const FPCGPoint& OutPoint : OutPoints)
+		{
+			if (OutputUV->GetValueKey(OutPoint.MetadataEntry) == ValueKey)
+			{
+				UTEST_EQUAL("Point Color", OutPoint.Color, Parameters.ExpectedColor);
+				UTEST_EQUAL("Point Density", OutPoint.Density, Parameters.ExpectedDensity);
+			}
+		}
+
+		return true;
+	}
+};
+
+
+IMPLEMENT_CUSTOM_SIMPLE_AUTOMATION_TEST(FPCGSampleTextureTest_LocalCorner, PCGSampleTextureTestBase, "Plugins.PCG.SampleTexture.LocalCorner", PCGTestsCommon::TestFlags)
+
+bool FPCGSampleTextureTest_LocalCorner::RunTest(const FString& Parameters)
+{
+	FTestParameters LocalCornerParameters{};
+	LocalCornerParameters.TextureMappingMethod = EPCGTextureMappingMethod::UVCoordinates;
+	LocalCornerParameters.UVPosition = FVector2D(0.0, 0.0);
+	LocalCornerParameters.ExpectedColor = FVector4(0.0, 0.0, 255.0, 128.0);
+	LocalCornerParameters.ExpectedDensity = 0.0f;
+
+	return GenerateTestDataRunAndValidate(LocalCornerParameters);
+}
+
+IMPLEMENT_CUSTOM_SIMPLE_AUTOMATION_TEST(FPCGSampleTextureTest_LocalTexelCenter, PCGSampleTextureTestBase, "Plugins.PCG.SampleTexture.LocalTexelCenter", PCGTestsCommon::TestFlags)
+
+bool FPCGSampleTextureTest_LocalTexelCenter::RunTest(const FString& Parameters)
+{
+	FTestParameters LocalTexelCenterParameters{};
+	LocalTexelCenterParameters.TextureMappingMethod = EPCGTextureMappingMethod::UVCoordinates;
+	LocalTexelCenterParameters.UVPosition = FVector2D(0.125, 0.125);
+	LocalTexelCenterParameters.ExpectedColor = FVector4(0.0, 0.0, 255.0, 128.0);
+	LocalTexelCenterParameters.ExpectedDensity = 0.0f;
+	
+	return GenerateTestDataRunAndValidate(LocalTexelCenterParameters);
+}
+
+IMPLEMENT_CUSTOM_SIMPLE_AUTOMATION_TEST(FPCGSampleTextureTest_LocalCenter, PCGSampleTextureTestBase, "Plugins.PCG.SampleTexture.LocalCenter", PCGTestsCommon::TestFlags)
+
+bool FPCGSampleTextureTest_LocalCenter::RunTest(const FString& Parameters)
+{
+	FTestParameters LocalCenterParameters{};
+	LocalCenterParameters.TextureMappingMethod = EPCGTextureMappingMethod::UVCoordinates;
+	LocalCenterParameters.UVPosition = FVector2D(0.5, 0.5);
+	LocalCenterParameters.ExpectedColor = FVector4((3.0/16.0 * 255.0), (3.0/16.0 * 255.0), 255.0, 128.0); // Average color across points 5, 6, 9, and 10
+	LocalCenterParameters.ExpectedDensity = 0.0f;
+
+	return GenerateTestDataRunAndValidate(LocalCenterParameters);
+}
+#endif
\ No newline at end of file
diff --git a/Engine/Plugins/PCG/Source/PCG/Public/Data/PCGTextureData.h b/Engine/Plugins/PCG/Source/PCG/Public/Data/PCGTextureData.h
--- a/Engine/Plugins/PCG/Source/PCG/Public/Data/PCGTextureData.h
+++ b/Engine/Plugins/PCG/Source/PCG/Public/Data/PCGTextureData.h
@@ -26,6 +26,13 @@ enum class EPCGTextureDensityFunction : uint8
 	Multiply
 };
 
+UENUM()
+enum class EPCGTextureAddressMode : uint8
+{
+	Clamp UMETA(ToolTip = "Clamps UV to 0-1."),
+	Wrap UMETA(ToolTip = "Tiles the texture to fit.")
+};
+
 namespace PCGTextureSamplingHelpers
 {
 	/** Returns true if a texture is CPU-accessible. */
@@ -55,6 +62,9 @@ public:
 	virtual const UPCGPointData* CreatePointData(FPCGContext* Context) const override;
 	//~End UPCGSpatialDataWithPointCache interface
 
+	/** Sample using a local space 'UV' position. */
+	bool SamplePointLocal(const FVector2D& LocalPosition, FVector4& OutColor, float& OutDensity) const;
+
 	virtual bool IsValid() const;
 
 public:
diff --git a/Engine/Plugins/PCG/Source/PCG/Public/Elements/PCGSampleTexture.h b/Engine/Plugins/PCG/Source/PCG/Public/Elements/PCGSampleTexture.h
new file mode 100644
--- /dev/null
+++ b/Engine/Plugins/PCG/Source/PCG/Public/Elements/PCGSampleTexture.h
@@ -0,0 +1,62 @@
+// Copyright Epic Games, Inc. All Rights Reserved.
+
+#pragma once
+
+#include "PCGSettings.h"
+#include "Data/PCGTextureData.h"
+
+#include "PCGSampleTexture.generated.h"
+
+namespace PCGSampleTextureConstants
+{
+	const FName InputPointLabel = TEXT("Point");
+	const FName InputTextureLabel = TEXT("BaseTexture");
+}
+
+UENUM()
+enum class EPCGTextureMappingMethod : uint8
+{
+	Planar UMETA(DisplayName = "Planar From Texture Data"),
+	UVCoordinates UMETA(DisplayName = "Use Explicit Points UV Coordinates")
+};
+
+/** Samples color of texture at each point. */
+UCLASS(BlueprintType, ClassGroup = (Procedural))
+class UPCGSampleTextureSettings : public UPCGSettings
+{
+	GENERATED_BODY()
+
+public:
+	//~Begin UPCGSettings interface
+#if WITH_EDITOR
+	virtual FName GetDefaultNodeName() const override { return FName(TEXT("SampleTexture")); }
+	virtual FText GetDefaultNodeTitle() const override { return NSLOCTEXT("PCGSampleTextureElement", "NodeTitle", "Sample Texture"); }
+	virtual FText GetNodeTooltipText() const override { return NSLOCTEXT("PCGSampleTextureElement", "NodeTooltip", "Samples color of texture at each point."); }
+	virtual EPCGSettingsType GetType() const override { return EPCGSettingsType::Sampler; }
+#endif
+
+protected:
+	virtual TArray<FPCGPinProperties> InputPinProperties() const override;
+	virtual TArray<FPCGPinProperties> OutputPinProperties() const override { return Super::DefaultPointOutputPinProperties(); }
+	virtual FPCGElementPtr CreateElement() const override;
+	//~End UPCGSettings interface
+
+public:
+	/** Whether to treat the sample positions as being in 0-1 UV space. If method is Planar then the coordinates will be transformed according to the texture settings. */
+	UPROPERTY(BlueprintReadWrite, EditAnywhere, Category = Settings, meta = (PCG_Overridable))
+	EPCGTextureMappingMethod TextureMappingMethod = EPCGTextureMappingMethod::Planar;
+	
+	/** The attribute that provides sample positions for sampling the texture. */
+	UPROPERTY(BlueprintReadWrite, EditAnywhere, Category = Settings, meta = (PCG_Overridable, DisplayName = "UV Coordinates Attribute", EditCondition = "TextureMappingMethod == EPCGTextureMappingMethod::UVCoordinates", EditConditionHides))
+	FPCGAttributePropertyInputSelector UVCoordinatesAttribute;
+
+	/** Overrides the texture's tiling to wrap or clamp its UVs. */
+	UPROPERTY(BlueprintReadWrite, EditAnywhere, Category = Settings, meta = (PCG_Overridable, EditCondition = "TextureMappingMethod == EPCGTextureMappingMethod::UVCoordinates", EditConditionHides))
+	EPCGTextureAddressMode TilingMode = EPCGTextureAddressMode::Wrap;
+};
+
+class FPCGSampleTextureElement : public IPCGElement
+{
+protected:
+	virtual bool ExecuteInternal(FPCGContext* Context) const override;
+};
