[PCG] Fix ensure in PCG.Switch.Enum test

Ensure was covering output pins in a nodeless context where the element was driven by tests

#rnx

diff --git a/Engine/Plugins/PCG/Source/PCG/Private/Elements/ControlFlow/PCGSwitch.cpp b/Engine/Plugins/PCG/Source/PCG/Private/Elements/ControlFlow/PCGSwitch.cpp
--- a/Engine/Plugins/PCG/Source/PCG/Private/Elements/ControlFlow/PCGSwitch.cpp
+++ b/Engine/Plugins/PCG/Source/PCG/Private/Elements/ControlFlow/PCGSwitch.cpp
@@ -364,14 +364,17 @@ bool FPCGSwitchElement::ExecuteInternal(FPCGContext* Context) const
 	Context->OutputData = PCGGather::GatherDataForPin(Context->InputData, PCGPinConstants::DefaultInputLabel, SelectedPinLabel);
 
 	// Output bitmask of deactivated pins.
-	const int NumOutputPins = Context->Node ? Context->Node->GetOutputPins().Num() : 0;
-	if (ensure(NumOutputPins > 0))
+	if (Context->Node)
 	{
-		const int SelectedPinIndex = Settings->GetSelectedOutputPinIndex();
-		if (ensure(SelectedPinIndex != INDEX_NONE))
+		const int NumOutputPins = Context->Node->GetOutputPins().Num();
+		if (ensure(NumOutputPins > 0))
 		{
-			const int AllPinsMask = (1 << NumOutputPins) - 1;
-			Context->OutputData.InactiveOutputPinBitmask = ~(1ULL << SelectedPinIndex) & AllPinsMask;
+			const int SelectedPinIndex = Settings->GetSelectedOutputPinIndex();
+			if (ensure(SelectedPinIndex != INDEX_NONE))
+			{
+				const int AllPinsMask = (1 << NumOutputPins) - 1;
+				Context->OutputData.InactiveOutputPinBitmask = ~(1ULL << SelectedPinIndex) & AllPinsMask;
+			}
 		}
 	}
 
