[PCG] Store Crc per data item rather than on data collection. The dependency Crc now takes Crcs only from consumed data items, which gives much finer change propagation. If data on only one output pin changes, only the nodes connected to that output pin will be affected.

diff --git a/Engine/Plugins/PCG/Source/PCG/Private/Elements/PCGSpawnActor.cpp b/Engine/Plugins/PCG/Source/PCG/Private/Elements/PCGSpawnActor.cpp
--- a/Engine/Plugins/PCG/Source/PCG/Private/Elements/PCGSpawnActor.cpp
+++ b/Engine/Plugins/PCG/Source/PCG/Private/Elements/PCGSpawnActor.cpp
@@ -783,8 +783,8 @@ void FPCGSpawnActorElement::SpawnActors(FPCGSubgraphContext* Context, AActor* Ta
 	{
 		FPCGDataCollection SingleInputCollection;
 		SingleInputCollection.TaggedData.Add(Output);
-		// TODO: review this, it might make more sense to do a full crc here
-		SingleInputCollection.Crc = SingleInputCollection.ComputeCrc(/*bFullDataCrc=*/false);
+		// TODO: review this, it might make more sense to do a full data crc here
+		SingleInputCollection.ComputeCrcs(/*bFullDataCrc=*/false);
 
 		GetDependenciesCrc(SingleInputCollection, Settings, Context->SourceComponent.Get(), InputDependenciesCrc);
 
diff --git a/Engine/Plugins/PCG/Source/PCG/Private/Graph/PCGGraphExecutor.cpp b/Engine/Plugins/PCG/Source/PCG/Private/Graph/PCGGraphExecutor.cpp
--- a/Engine/Plugins/PCG/Source/PCG/Private/Graph/PCGGraphExecutor.cpp
+++ b/Engine/Plugins/PCG/Source/PCG/Private/Graph/PCGGraphExecutor.cpp
@@ -43,11 +43,6 @@ static TAutoConsoleVariable<float> CVarMaxPercentageOfThreadsToUse(
 	0.9f,
 	TEXT("Maximum percentage of number of threads for concurrent PCG processing"));
 
-static TAutoConsoleVariable<bool> CVarStripEmptyPointData(
-	TEXT("pcg.StripEmptyPointData"),
-	false,
-	TEXT("Will strip empty point data from being passed along as output"));
-
 namespace PCGGraphExecutor
 {
 	TAutoConsoleVariable<float> CVarTimePerFrame(
@@ -597,7 +592,6 @@ void FPCGGraphExecutor::Execute()
 	const float MaxPercentageOfThreadsToUse = FMath::Clamp(CVarMaxPercentageOfThreadsToUse.GetValueOnAnyThread(), 0.0f, 1.0f);
 	const int32 MaxNumThreads = FMath::Max(0, FMath::Min((int32)(FPlatformMisc::NumberOfCoresIncludingHyperthreads() * MaxPercentageOfThreadsToUse), CVarMaxNumTasks.GetValueOnAnyThread() - 1));
 	const bool bAllowMultiDispatch = PCGGraphExecutor::CVarGraphMultithreading.GetValueOnAnyThread();
-	const bool bStripEmptyPointData = CVarStripEmptyPointData.GetValueOnAnyThread();
 	const bool bDynamicTaskCulling = PCGGraphExecutor::CVarDynamicTaskCulling.GetValueOnAnyThread();
 
 #if WITH_EDITOR
@@ -780,7 +774,6 @@ void FPCGGraphExecutor::Execute()
 					Task.Context->InitializeSettings();
 					Task.Context->TaskId = Task.NodeId;
 					Task.Context->CompiledTaskId = Task.CompiledTaskId;
-					Task.Context->InputData.Crc = TaskInput.Crc;
 					Task.Context->DependenciesCrc = DependenciesCrc;
 					Task.Context->Stack = (Task.StackContext && Task.StackIndex != INDEX_NONE) ? Task.StackContext->GetStack(Task.StackIndex) : nullptr;
 				}
@@ -869,7 +862,7 @@ void FPCGGraphExecutor::Execute()
 			}
 		}
 
-		auto PostTaskExecute = [this, &bAnyTaskEnded, bStripEmptyPointData](int32 TaskIndex)
+		auto PostTaskExecute = [this, &bAnyTaskEnded](int32 TaskIndex)
 		{
 			FPCGGraphActiveTask& ActiveTask = ActiveTasks[TaskIndex];
 			check(ActiveTask.Context);
@@ -898,8 +891,6 @@ void FPCGGraphExecutor::Execute()
 			if (!ActiveTask.bWasCancelled)
 #endif
 			{
-				PCGGraphExecutionLogging::LogTaskExecuteOutputCRC(ActiveTask);
-
 				// Store result in cache as needed - done here because it needs to be done on the main thread
 
 				// Don't store if errors or warnings present
@@ -919,31 +910,25 @@ void FPCGGraphExecutor::Execute()
 			CurrentlyUsedThreads -= ActiveTask.Context->AsyncState.NumAvailableTasks;
 
 #if WITH_EDITOR
-			// Execute debug display code as needed - done here because it needs to be done on the main thread
-			// Additional note: this needs to be executed before the StoreResults since debugging might cancel further tasks
-			ActiveTask.Element->DebugDisplay(ActiveTask.Context.Get());
-
-			if (UPCGComponent* SourceComponent = ActiveTask.Context->SourceComponent.Get())
+			if (!ActiveTask.bWasCancelled)
 			{
-				if (ActiveTask.StackIndex != INDEX_NONE)
-				{
-					const FPCGStack* Stack = ActiveTask.StackContext->GetStack(ActiveTask.StackIndex);
-					SourceComponent->StoreInspectionData(Stack, ActiveTask.Context->Node, ActiveTask.Context->InputData, ActiveTask.Context->OutputData);
-				}
-			}
-#endif
+				// Execute debug display code as needed - done here because it needs to be done on the main thread
+				// Additional note: this needs to be executed before the StoreResults since debugging might cancel further tasks
+				ActiveTask.Element->DebugDisplay(ActiveTask.Context.Get());
 
-			// TODO: Evaluation code for finding elements that produce empty point data. This is temporary.
-			if (bStripEmptyPointData)
-			{
-				const int32 NumRemoved = ActiveTask.Context->OutputData.StripEmptyPointData();
-				if (NumRemoved > 0)
+				if (UPCGComponent* SourceComponent = ActiveTask.Context->SourceComponent.Get())
 				{
-					UE_LOG(LogPCG, Log, TEXT("%d empty point data stripped from node: %s"), NumRemoved, *ActiveTask.Context->Node->GetNodeTitle(EPCGNodeTitleType::ListView).ToString());
+					if (ActiveTask.StackIndex != INDEX_NONE)
+					{
+						const FPCGStack* Stack = ActiveTask.StackContext->GetStack(ActiveTask.StackIndex);
+						SourceComponent->StoreInspectionData(Stack, ActiveTask.Context->Node, ActiveTask.Context->InputData, ActiveTask.Context->OutputData);
+					}
 				}
 			}
+#endif
 
-			// Store output in data map
+			// Store output in data map.
+			// TODO - investigate if we should avoid doing this if the task was cancelled.
 			StoreResults(ActiveTask.NodeId, ActiveTask.Context->OutputData);
 
 			// Book-keeping
@@ -1204,14 +1189,18 @@ void FPCGGraphExecutor::BuildTaskInput(const FPCGGraphTask& Task, FPCGDataCollec
 			InPin ? *InPin->Properties.Label.ToString() : TEXT("MissingPin"));
 	};
 
-	// Initialize a Crc onto which each input Crc will be combined.
-	FPCGCrc Crc(Task.Inputs.Num());
+	// Initialize a Crc onto which each input Crc will be combined (using random prime number).
+	FPCGCrc Crc(1000033);
 
 	// Random prime numbers to use as placeholders in the CRC computation when there are no defined in/out pins.
 	// Note that they aren't strictly needed, but will make sure we don't introduce issues if we rework this bit of code.
 	constexpr uint32 DefaultHashForNoInputPin = 955333;
 	constexpr uint32 DefaultHashForNoOutputPin = 999983;
 
+	// Hoisted out of loop for performance reasons.
+	TArray<FPCGTaggedData, TInlineAllocator<16>> InputDataOnPin;
+	TArray<FPCGCrc, TInlineAllocator<16>> InputDataCrcsOnPin;
+
 	for (const FPCGGraphTaskInput& Input : Task.Inputs)
 	{
 		check(OutputData.Contains(Input.TaskId));
@@ -1223,6 +1212,7 @@ void FPCGGraphExecutor::BuildTaskInput(const FPCGGraphTask& Task, FPCGDataCollec
 		}
 
 		const bool bAllowMultipleData = Input.OutPin ? Input.OutPin->Properties.bAllowMultipleData : true;
+		const uint32 InputPinLabelCrc = Input.OutPin ? GetTypeHash(Input.OutPin->Properties.Label) : DefaultHashForNoOutputPin;
 
 		// Enforce single data - if already have input for this pin, don't add more. Early check before other side effects below.
 		if (Input.OutPin && !bAllowMultipleData && TaskInput.GetInputCountByPin(Input.OutPin->Properties.Label) > 0)
@@ -1235,30 +1225,43 @@ void FPCGGraphExecutor::BuildTaskInput(const FPCGGraphTask& Task, FPCGDataCollec
 
 		TaskInput.bCancelExecution |= InputCollection.bCancelExecution;
 
-		// Get input data at the given pin (or everything)
 		const int32 TaggedDataOffset = TaskInput.TaggedData.Num();
+
+		// Get input data at the given pin (or everything). This will add the data and include the input pin Crc to uniquely identify
+		// inputs per-pin, or use a placeholder for symmetry.
+		// Note: The input data CRC will already contain the output pin (calculated in element post execute).
 		if (Input.InPin)
 		{
-			// Proceed carefully when adding data items - if pin is single-data, only add first item.
-			TArray<FPCGTaggedData> InputsOnPin = InputCollection.GetInputsByPin(Input.InPin->Properties.Label);
-			if (InputsOnPin.Num() > 1 && !bAllowMultipleData)
-			{
-				LogDiscardedData(Input.OutPin);
-				TaskInput.TaggedData.Add(InputsOnPin[0]);
-			}
-			else
+			InputDataOnPin.Reset();
+			InputDataCrcsOnPin.Reset();
+			InputCollection.GetInputsAndCrcsByPin(Input.InPin->Properties.Label, InputDataOnPin, InputDataCrcsOnPin);
+
+			if (!InputDataOnPin.IsEmpty())
 			{
-				TaskInput.TaggedData.Append(InputsOnPin);
+				// Proceed carefully when adding data items - if pin is single-data, only add first item.
+				if (!ensure(InputDataOnPin.Num() == InputDataCrcsOnPin.Num()))
+				{
+					InputDataCrcsOnPin.SetNumZeroed(InputDataOnPin.Num());
+				}
+
+				const int NumberDataItemsToTake = bAllowMultipleData ? InputDataOnPin.Num() : 1;
+
+				TaskInput.AddDataForPin(
+					MakeArrayView(InputDataOnPin.GetData(), NumberDataItemsToTake),
+					MakeArrayView(InputDataCrcsOnPin.GetData(), NumberDataItemsToTake),
+					InputPinLabelCrc);
+
+				if (NumberDataItemsToTake < InputDataOnPin.Num())
+				{
+					LogDiscardedData(Input.OutPin);
+				}
 			}
 		}
 		else
 		{
-			TaskInput.TaggedData.Append(InputCollection.TaggedData);
+			TaskInput.AddData(InputCollection.TaggedData, InputCollection.DataCrcs);
 		}
 
-		// Write input pin name (e.g. name of the output pin on the dependency node) Crc to uniquely identify inputs per-pin, or use a placeholder for symmetry.
-		Crc.Combine(Input.InPin ? GetTypeHash(Input.InPin->Properties.Label) : DefaultHashForNoInputPin);
-
 		// Apply labelling on data; technically, we should ensure that we do this only for pass-through nodes,
 		// Otherwise we could also null out the label on the input...
 		if (Input.OutPin)
@@ -1268,21 +1271,10 @@ void FPCGGraphExecutor::BuildTaskInput(const FPCGGraphTask& Task, FPCGDataCollec
 				TaskInput.TaggedData[TaggedDataIndex].Pin = Input.OutPin->Properties.Label;
 			}
 		}
-
-		// Write output pin name (e.g. input pin node on this node) Crc to uniquely identify inputs per-pin, or use a placeholder for symmetry.
-		Crc.Combine(Input.OutPin ? GetTypeHash(Input.OutPin->Properties.Label) : DefaultHashForNoOutputPin);
-
-		// This chains the Crc of each input to produce a Crc that covers all of them.
-		if (InputCollection.Crc.IsValid())
-		{
-			Crc.Combine(InputCollection.Crc);
-		}
 	}
 
 	// Then combine params if needed
 	CombineParams(Task.NodeId, TaskInput);
-
-	TaskInput.Crc = Crc;
 }
 
 void FPCGGraphExecutor::CombineParams(FPCGTaskId InTaskId, FPCGDataCollection& InTaskInput)
diff --git a/Engine/Plugins/PCG/Source/PCG/Private/PCGComponent.cpp b/Engine/Plugins/PCG/Source/PCG/Private/PCGComponent.cpp
--- a/Engine/Plugins/PCG/Source/PCG/Private/PCGComponent.cpp
+++ b/Engine/Plugins/PCG/Source/PCG/Private/PCGComponent.cpp
@@ -1978,10 +1978,7 @@ void UPCGComponent::StoreInspectionData(const FPCGStack* InStack, const UPCGNode
 				StackFrames.Emplace(Pin);
 
 				FPCGDataCollection PinDataCollection;
-				PinDataCollection.TaggedData = InData.GetInputsByPin(Pin->Properties.Label);
-				// The data collection for each pin is given the Crc from the data collection. This is to enable inspecting the normal node output Crc
-				// when cache debugging is enabled.
-				PinDataCollection.Crc = InData.Crc;
+				InData.GetInputsAndCrcsByPin(Pin->Properties.Label, PinDataCollection.TaggedData, PinDataCollection.DataCrcs);
 
 				if (!PinDataCollection.TaggedData.IsEmpty())
 				{
diff --git a/Engine/Plugins/PCG/Source/PCG/Private/PCGData.cpp b/Engine/Plugins/PCG/Source/PCG/Private/PCGData.cpp
--- a/Engine/Plugins/PCG/Source/PCG/Private/PCGData.cpp
+++ b/Engine/Plugins/PCG/Source/PCG/Private/PCGData.cpp
@@ -105,6 +105,25 @@ bool FPCGTaggedData::operator!=(const FPCGTaggedData& Other) const
 	return !operator==(Other);
 }
 
+FPCGCrc FPCGTaggedData::ComputeCrc(bool bFullDataCrc) const
+{
+	FArchiveCrc32 Ar;
+
+	Ar << const_cast<FName&>(Pin);
+
+	if (Data)
+	{
+		const FPCGCrc DataCrc = Data->GetOrComputeCrc(bFullDataCrc);
+		uint32 CrcValue = DataCrc.GetValue();
+		Ar << CrcValue;
+	}
+
+	// TODO: Ensuring tags are sorted could prevent CRC change.
+	Ar << const_cast<TSet<FString>&>(Tags);
+
+	return FPCGCrc(Ar.GetCrc());
+}
+
 TArray<FPCGTaggedData> FPCGDataCollection::GetInputs() const
 {
 	return TaggedData.FilterByPredicate([](const FPCGTaggedData& Data) {
@@ -392,26 +411,33 @@ void FPCGDataCollection::AddReferences(FReferenceCollector& Collector)
 	}
 }
 
-FPCGCrc FPCGDataCollection::ComputeCrc(bool bFullDataCrc)
+void FPCGDataCollection::ComputeCrcs(bool bFullDataCrc)
 {
-	// If there is no data, will return valid Crc==0 which is fine. No such thing as an invalid FPCGDataCollection
-	FArchiveCrc32 Ar;
+	DataCrcs.SetNumUninitialized(TaggedData.Num(), /*bAllowShrinking=*/false);
 
-	for (FPCGTaggedData& Data : TaggedData)
+	for (int I = 0; I < TaggedData.Num(); ++I)
 	{
-		Ar << Data.Pin;
+		DataCrcs[I] = TaggedData[I].ComputeCrc(bFullDataCrc);
+	}
+}
 
-		if (Data.Data)
-		{
-			const FPCGCrc DataCrc = Data.Data->GetOrComputeCrc(bFullDataCrc);
-			uint32 Result = DataCrc.GetValue();
-			Ar << Result;
-		}
+void FPCGDataCollection::AddData(const TConstArrayView<FPCGTaggedData>& InData, const TConstArrayView<FPCGCrc>& InDataCrcs)
+{
+	TaggedData.Append(InData);
+	DataCrcs.Append(InDataCrcs);
+}
 
-		Ar << Data.Tags;
-	}
+void FPCGDataCollection::AddDataForPin(const TConstArrayView<FPCGTaggedData>& InData, const TConstArrayView<FPCGCrc>& InDataCrcs, uint32 InputPinLabelCrc)
+{
+	TaggedData.Append(InData);
+	DataCrcs.Append(InDataCrcs);
 
-	return FPCGCrc(Ar.GetCrc());
+	// Add input pin label to Crc to uniquely identify inputs per-pin, or use a placeholder for symmetry.
+	// Note that the cached data Crc will already contain the output pin (calculated in element post execute).
+	for (int I = (DataCrcs.Num() - InDataCrcs.Num()); I < DataCrcs.Num(); ++I)
+	{
+		DataCrcs[I].Combine(InputPinLabelCrc);
+	}
 }
 
 void FPCGDataCollection::Reset()
diff --git a/Engine/Plugins/PCG/Source/PCG/Private/PCGElement.cpp b/Engine/Plugins/PCG/Source/PCG/Private/PCGElement.cpp
--- a/Engine/Plugins/PCG/Source/PCG/Private/PCGElement.cpp
+++ b/Engine/Plugins/PCG/Source/PCG/Private/PCGElement.cpp
@@ -170,8 +170,8 @@ void IPCGElement::PostExecute(FPCGContext* Context) const
 		// execution. Performing a detailed CRC of output data can detect real change in the data and halt the cascade of execution.
 		const bool bShouldComputeFullOutputDataCrc = ShouldComputeFullOutputDataCrc(Context);
 
-		// Compute Crc from output data
-		Context->OutputData.Crc = Context->OutputData.ComputeCrc(bShouldComputeFullOutputDataCrc);
+		// Compute Crc from output data which will include output pin labels.
+		Context->OutputData.ComputeCrcs(bShouldComputeFullOutputDataCrc);
 	}
 
 #if WITH_EDITOR
@@ -474,7 +474,15 @@ bool IPCGElement::IsCacheableInstance(const UPCGSettingsInterface* InSettingsInt
 void IPCGElement::GetDependenciesCrc(const FPCGDataCollection& InInput, const UPCGSettings* InSettings, UPCGComponent* InComponent, FPCGCrc& OutCrc) const
 {
 	TRACE_CPUPROFILER_EVENT_SCOPE_TEXT(*FString::Printf(TEXT("IPCGElement::GetDependenciesCrc (%s)"), InSettings ? *InSettings->GetFName().ToString() : TEXT("")));
-	FPCGCrc Crc = InInput.Crc;
+
+	// Start from a random prime.
+	FPCGCrc Crc(1000003);
+
+	// The cached data CRCs are computed in FPCGGraphExecutor::BuildTaskInput and incorporate data CRC, tags, output pin label and input pin label.
+	for (const FPCGCrc& DataCrc : InInput.DataCrcs)
+	{
+		Crc.Combine(DataCrc);
+	}
 
 	if (InSettings)
 	{
diff --git a/Engine/Plugins/PCG/Source/PCG/Private/Tests/PCGTestsCommon.cpp b/Engine/Plugins/PCG/Source/PCG/Private/Tests/PCGTestsCommon.cpp
--- a/Engine/Plugins/PCG/Source/PCG/Private/Tests/PCGTestsCommon.cpp
+++ b/Engine/Plugins/PCG/Source/PCG/Private/Tests/PCGTestsCommon.cpp
@@ -58,8 +58,8 @@ namespace PCGTestsCommon
 		check(TestGraph);
 		TestPCGComponent->SetGraphLocal(TestGraph);
 
-		// Tests don't currently need CRCs, so just assigned a trivial valid CRC to avoid asserts.
-		InputData.Crc = FPCGCrc(0);
+		// Initialize CRC to avoid asserts.
+		InputData.ComputeCrcs(/*bFullDataCrc=*/false);
 #else
 		TestActor = nullptr;
 		TestPCGComponent = nullptr;
diff --git a/Engine/Plugins/PCG/Source/PCG/Private/Utils/PCGGraphExecutionLogging.cpp b/Engine/Plugins/PCG/Source/PCG/Private/Utils/PCGGraphExecutionLogging.cpp
--- a/Engine/Plugins/PCG/Source/PCG/Private/Utils/PCGGraphExecutionLogging.cpp
+++ b/Engine/Plugins/PCG/Source/PCG/Private/Utils/PCGGraphExecutionLogging.cpp
@@ -206,21 +206,6 @@ namespace PCGGraphExecutionLogging
 #endif
 	}
 
-	void LogTaskExecuteOutputCRC(const FPCGGraphActiveTask& Task)
-	{
-#if WITH_EDITOR
-		if (!LogEnabled() || !Task.Context || !Task.Context->SourceComponent.Get() || !Task.Context->SourceComponent->GetOwner())
-		{
-			return;
-		}
-
-		UE_LOG(LogPCG, Log, TEXT("         [%s] %s\t\tOUTPUT CRC %u"),
-			*Task.Context->SourceComponent->GetOwner()->GetName(),
-			*FString::Printf(TEXT("%u'%s'"), Task.NodeId, Task.Context->Node ? *Task.Context->Node->GetNodeTitle(EPCGNodeTitleType::ListView).ToString() : TEXT("")),
-			Task.Context->OutputData.Crc.GetValue());
-#endif
-	}
-
 	void LogTaskCullingBegin(FPCGTaskId CompletedTaskId, uint64 InactiveOutputPinBitmask, const TArray<FPCGPinId>& PinIdsToDeactivate)
 	{
 #if WITH_EDITOR
diff --git a/Engine/Plugins/PCG/Source/PCG/Private/Utils/PCGGraphExecutionLogging.h b/Engine/Plugins/PCG/Source/PCG/Private/Utils/PCGGraphExecutionLogging.h
--- a/Engine/Plugins/PCG/Source/PCG/Private/Utils/PCGGraphExecutionLogging.h
+++ b/Engine/Plugins/PCG/Source/PCG/Private/Utils/PCGGraphExecutionLogging.h
@@ -29,7 +29,6 @@ namespace PCGGraphExecutionLogging
 
 	void LogTaskExecute(const FPCGGraphTask& Task);
 	void LogTaskExecuteCachingDisabled(const FPCGGraphTask& Task);
-	void LogTaskExecuteOutputCRC(const FPCGGraphActiveTask& Task);
 
 	void LogTaskCullingBegin(FPCGTaskId CompletedTaskId, uint64 InactiveOutputPinBitmask, const TArray<FPCGPinId>& PinIdsToDeactivate);
 	void LogTaskCullingBeginLoop(FPCGTaskId PinTaskId, uint64 PinIndex, const TArray<FPCGPinId>& PinIdsToDeactivate);
diff --git a/Engine/Plugins/PCG/Source/PCG/Public/Elements/PCGDataFromActor.h b/Engine/Plugins/PCG/Source/PCG/Public/Elements/PCGDataFromActor.h
--- a/Engine/Plugins/PCG/Source/PCG/Public/Elements/PCGDataFromActor.h
+++ b/Engine/Plugins/PCG/Source/PCG/Public/Elements/PCGDataFromActor.h
@@ -122,6 +122,7 @@ class FPCGDataFromActorElement : public IPCGElement
 public:
 	virtual bool CanExecuteOnlyOnMainThread(FPCGContext* Context) const override { return true; }
 	virtual void GetDependenciesCrc(const FPCGDataCollection& InInput, const UPCGSettings* InSettings, UPCGComponent* InComponent, FPCGCrc& OutCrc) const override;
+	virtual bool ShouldComputeFullOutputDataCrc(FPCGContext* Context) const override { return true; }
 
 protected:
 	virtual FPCGContext* CreateContext() override;
diff --git a/Engine/Plugins/PCG/Source/PCG/Public/PCGData.h b/Engine/Plugins/PCG/Source/PCG/Public/PCGData.h
--- a/Engine/Plugins/PCG/Source/PCG/Public/PCGData.h
+++ b/Engine/Plugins/PCG/Source/PCG/Public/PCGData.h
@@ -92,6 +92,7 @@ struct PCG_API FPCGTaggedData
 	UPROPERTY(BlueprintReadWrite, EditAnywhere, Category = Data)
 	TSet<FString> Tags;
 
+	/** The label of the pin that this data was either emitted from or received on. */
 	UPROPERTY(BlueprintReadWrite, EditAnywhere, Category = Data)
 	FName Pin = NAME_None;
 
@@ -101,6 +102,8 @@ struct PCG_API FPCGTaggedData
 
 	bool operator==(const FPCGTaggedData& Other) const;
 	bool operator!=(const FPCGTaggedData& Other) const;
+
+	FPCGCrc ComputeCrc(bool bFullDataCrc) const;
 };
 
 USTRUCT(BlueprintType)
@@ -110,11 +113,31 @@ struct PCG_API FPCGDataCollection
 
 	/** Returns all spatial data in the collection */
 	TArray<FPCGTaggedData> GetInputs() const;
-	/** Returns all data on a given pin */
+	/** Returns all data on a given pin. */
 	TArray<FPCGTaggedData> GetInputsByPin(const FName& InPinLabel) const;
 	/** Returns all spatial data on a given pin */
 	TArray<FPCGTaggedData> GetSpatialInputsByPin(const FName& InPinLabel) const;
-	
+
+	/** Returns all data and corresponding cached data CRCs for a given pin. */
+	template<typename AllocatorType1, typename AllocatorType2>
+	void GetInputsAndCrcsByPin(const FName& InPinLabel, TArray<FPCGTaggedData, AllocatorType1>& OutData, TArray<FPCGCrc, AllocatorType2>& OutDataCrcs) const
+	{
+		if (!ensure(TaggedData.Num() == DataCrcs.Num()))
+		{
+			// CRCs are not up to date. Error recovery - add 0 CRCs.
+			const_cast<TArray<FPCGCrc>&>(DataCrcs).SetNumZeroed(TaggedData.Num());
+		}
+
+		for (int I = 0; I < TaggedData.Num(); ++I)
+		{
+			if (ensure(TaggedData[I].Data) && TaggedData[I].Pin == InPinLabel)
+			{
+				OutData.Add(TaggedData[I]);
+				OutDataCrcs.Add(DataCrcs[I]);
+			}
+		}
+	}
+
 	/** Gets number of data items on a given pin */
 	int32 GetInputCountByPin(const FName& InPinLabel) const;
 	/** Gets number of spatial data items on a given pin */
@@ -163,8 +186,13 @@ struct PCG_API FPCGDataCollection
 	bool operator!=(const FPCGDataCollection& Other) const;
 	void AddReferences(FReferenceCollector& Collector);
 
-	/** Computes Crc for this data. */
-	FPCGCrc ComputeCrc(bool bFullDataCrc);
+	/** Computes CRCs for all data items. */
+	void ComputeCrcs(bool bFullDataCrc);
+
+	/** Add data and CRCs to collection. */
+	void AddData(const TConstArrayView<FPCGTaggedData>& InData, const TConstArrayView<FPCGCrc>& InDataCrcs);
+	/** Add data and CRCs to collection with pin label combined into the CRC. */
+	void AddDataForPin(const TConstArrayView<FPCGTaggedData>& InData, const TConstArrayView<FPCGCrc>& InDataCrcs, uint32 InputPinLabelCrc);
 
 	/** Cleans up the collection, but does not unroot any previously rooted data. */
 	void Reset();
@@ -182,8 +210,8 @@ struct PCG_API FPCGDataCollection
 	/** This flag is used to cancel further computation or for the debug/isolate feature */
 	bool bCancelExecution = false;
 
-	/** A snapshot of the internal state of the data. If any dependency (setting, node input or external data) changes then this value should change. */
-	FPCGCrc Crc;
+	/** Per-data CRC which will capture tags, data, output pin and in some cases input pin too. */
+	TArray<FPCGCrc> DataCrcs;
 
 	/** After the task is complete, bit j is set if output pin index j is deactivated. Stored here so that it can be retrieved from the cache. */
 	uint64 InactiveOutputPinBitmask = 0;
diff --git a/Engine/Plugins/PCG/Source/PCG/Public/PCGEngineSettings.h b/Engine/Plugins/PCG/Source/PCG/Public/PCGEngineSettings.h
--- a/Engine/Plugins/PCG/Source/PCG/Public/PCGEngineSettings.h
+++ b/Engine/Plugins/PCG/Source/PCG/Public/PCGEngineSettings.h
@@ -44,11 +44,8 @@ public:
 	/** Time in MS between a landscape change and PCG refresh. Set it to 0 or negative value to disable the delay. */
 	UPROPERTY(EditAnywhere, Config, Category = "Tracking|Advanced", meta = (ConsoleVariable = "pcg.LandscapeRefreshTimeDelayMS"))
 	int32 LandscapeRefreshTimeDelayMS = 1000;
-
-	/** Strips empty point data from being passed along as output. */
-	UPROPERTY(EditAnywhere, Config, Category = DataFlow, meta = (ConsoleVariable = "pcg.StripEmptyPointData"))
-	bool StripEmptyPointData = false;
 #endif // WITH_EDITORONLY_DATA
+
 	// Begin UDeveloperSettings Interface
 	virtual FName GetCategoryName() const override;
 #if WITH_EDITOR
diff --git a/Engine/Plugins/PCG/Source/PCGEditor/Private/SPCGEditorGraphAttributeListView.cpp b/Engine/Plugins/PCG/Source/PCGEditor/Private/SPCGEditorGraphAttributeListView.cpp
--- a/Engine/Plugins/PCG/Source/PCGEditor/Private/SPCGEditorGraphAttributeListView.cpp
+++ b/Engine/Plugins/PCG/Source/PCGEditor/Private/SPCGEditorGraphAttributeListView.cpp
@@ -685,6 +685,7 @@ void SPCGEditorGraphAttributeListView::RefreshAttributeList()
 
 	const FPCGTaggedData& TaggedData = InspectionData->TaggedData[DataIndex];
 	const UPCGData* PCGData = TaggedData.Data;
+	const FPCGCrc Crc = InspectionData->DataCrcs.IsValidIndex(DataIndex) ? InspectionData->DataCrcs[DataIndex] : FPCGCrc(0);
 
 	if (const UPCGParamData* PCGParamData = Cast<const UPCGParamData>(PCGData))
 	{
@@ -709,7 +710,7 @@ void SPCGEditorGraphAttributeListView::RefreshAttributeList()
 			else
 			{
 				// If cache debugging enabled, write CRC to help diagnose missed-dependency issues
-				InfoTextBlock->SetText(FText::Format(LOCTEXT("MetadataInfoTextBlockWithCrcFmt", "Number of metadata: {0}  CRC: {1}"), ItemKeyUpperBound - ItemKeyLowerBound, InspectionData->Crc.GetValue()));
+				InfoTextBlock->SetText(FText::Format(LOCTEXT("MetadataInfoTextBlockWithCrcFmt", "Number of metadata: {0}  CRC: {1}"), ItemKeyUpperBound - ItemKeyLowerBound, Crc.GetValue()));
 			}
 		}
 	}
@@ -742,7 +743,7 @@ void SPCGEditorGraphAttributeListView::RefreshAttributeList()
 			else
 			{
 				// If cache debugging enabled, write CRC to help diagnose missed-dependency issues
-				InfoTextBlock->SetText(FText::Format(LOCTEXT("PointInfoTextBlockWithCrcFmt", "Number of points: {0}, CRC: {1}"), NumPoints, InspectionData->Crc.GetValue()));
+				InfoTextBlock->SetText(FText::Format(LOCTEXT("PointInfoTextBlockWithCrcFmt", "Number of points: {0}, CRC: {1}"), NumPoints, Crc.GetValue()));
 			}
 		}
 	}
