[PCG] Missing Initialize from data in the PointNeighborhood node

#rnx

diff --git a/Engine/Plugins/PCG/Source/PCG/Private/Elements/PCGPointNeighborhood.cpp b/Engine/Plugins/PCG/Source/PCG/Private/Elements/PCGPointNeighborhood.cpp
--- a/Engine/Plugins/PCG/Source/PCG/Private/Elements/PCGPointNeighborhood.cpp
+++ b/Engine/Plugins/PCG/Source/PCG/Private/Elements/PCGPointNeighborhood.cpp
@@ -80,6 +80,7 @@ bool FPCGPointNeighborhoodElement::ExecuteInternal(FPCGContext* Context) const
 		const TArray<FPCGPoint>& SrcPoints = InputPointData->GetPoints();
 		FPCGTaggedData& Output = Outputs.Add_GetRef(Inputs[i]);
 		UPCGPointData* OutputPointData = NewObject<UPCGPointData>();
+		OutputPointData->InitializeFromData(InputPointData);
 		TArray<FPCGPoint>& OutputPoints = OutputPointData->GetMutablePoints();
 		OutputPoints.SetNumUninitialized(SrcPoints.Num());
 		Output.Data = OutputPointData;
