[PCG] Fix crash when inspecting interior sampled splines
* Allow SampleInteriorData to consume a null context without crashing

#rnx

diff --git a/Engine/Plugins/PCG/Source/PCG/Private/Elements/PCGSplineSampler.cpp b/Engine/Plugins/PCG/Source/PCG/Private/Elements/PCGSplineSampler.cpp
--- a/Engine/Plugins/PCG/Source/PCG/Private/Elements/PCGSplineSampler.cpp
+++ b/Engine/Plugins/PCG/Source/PCG/Private/Elements/PCGSplineSampler.cpp
@@ -881,7 +881,7 @@ namespace PCGSplineSamplerHelpers
  
 	void SampleInteriorData(FPCGContext* Context, const UPCGPolyLineData* LineData, const UPCGSpatialData* InBoundingShape, const UPCGSpatialData* InProjectionTarget, const FPCGProjectionParams& InProjectionParams, const FPCGSplineSamplerParams& Params, UPCGPointData* OutPointData)
 	{
-		check(Context && LineData && OutPointData);
+		check(LineData && OutPointData);
 
 		const FPCGSplineStruct* Spline = nullptr;
 
@@ -891,12 +891,12 @@ namespace PCGSplineSamplerHelpers
 		}
 		else if (const UPCGLandscapeSplineData* LandscapeSplineData = Cast<UPCGLandscapeSplineData>(LineData))
 		{
-			PCGE_LOG_C(Error, GraphAndLog, Context, LOCTEXT("LandscapeSplinesNotSupported", "Input data of type Landscape Spline are not supported for interior sampling"));
+			PCGLog::LogErrorOnGraph(LOCTEXT("LandscapeSplinesNotSupported", "Input data of type Landscape Spline are not supported for interior sampling"), Context);
 			return;
 		}
 		else
 		{
-			PCGE_LOG_C(Error, GraphAndLog, Context, LOCTEXT("CouldNotCreateSplineData", "Could not create UPCGSplineData from LineData"));
+			PCGLog::LogErrorOnGraph(LOCTEXT("CouldNotCreateSplineData", "Could not create UPCGSplineData from LineData"), Context);
 			return;
 		}
 
@@ -904,7 +904,7 @@ namespace PCGSplineSamplerHelpers
 
 		if (!Spline->IsClosedLoop())
 		{
-			PCGE_LOG_C(Error, GraphAndLog, Context, LOCTEXT("ShapeNotClosed", "Interior sampling only generates for closed shapes, enable the 'Closed Loop' setting on the spline"));
+			PCGLog::LogErrorOnGraph(LOCTEXT("ShapeNotClosed", "Interior sampling only generates for closed shapes, enable the 'Closed Loop' setting on the spline"), Context);
 			return;
 		}
 
@@ -1060,8 +1060,9 @@ namespace PCGSplineSamplerHelpers
 		const FVector::FReal MaxY = FMath::FloorToDouble(MaxPoint.Y / Params.InteriorSampleSpacing) * Params.InteriorSampleSpacing;
 
 		constexpr int32 MinIterationPerDispatch = 4;
+		const int32 NumAvailableTasks = Context ? Context->AsyncState.NumAvailableTasks : 1;
 		const int32 NumIterations = (MaxY + UE_KINDA_SMALL_NUMBER - MinY) / Params.InteriorSampleSpacing;
-		const int32 NumDispatch = FMath::Max(1, FMath::Min(Context->AsyncState.NumAvailableTasks, NumIterations / MinIterationPerDispatch));
+		const int32 NumDispatch = FMath::Max(1, FMath::Min(NumAvailableTasks, NumIterations / MinIterationPerDispatch));
 		const int32 NumIterationsPerDispatch = NumIterations / NumDispatch;
 
 		TArray<TArray<TTuple<FTransform, FVector, float>>> InteriorSplinePointData;
@@ -1230,7 +1231,7 @@ namespace PCGSplineSamplerHelpers
 
 		if (bAnyRowFailedToSample)
 		{
-			PCGE_LOG_C(Error, GraphAndLog, Context, LOCTEXT("IntersectionTestFailed", "One or more rows of the spline interior failed to sample (intersection test failed). Ensure the spline points are not overlapping."));
+			PCGLog::LogErrorOnGraph(LOCTEXT("IntersectionTestFailed", "One or more rows of the spline interior failed to sample (intersection test failed). Ensure the spline points are not overlapping."), Context);
 		}
 
 		// Finally, gather the data and push to the points
