[PCG] Fix occasional assert during graph actions (typically wiring new nodes) by deferring editor node reconstructs until after action is completed.

diff --git a/Engine/Plugins/PCG/Source/PCGEditor/Private/PCGEditorGraphSchemaActions.cpp b/Engine/Plugins/PCG/Source/PCGEditor/Private/PCGEditorGraphSchemaActions.cpp
--- a/Engine/Plugins/PCG/Source/PCGEditor/Private/PCGEditorGraphSchemaActions.cpp
+++ b/Engine/Plugins/PCG/Source/PCGEditor/Private/PCGEditorGraphSchemaActions.cpp
@@ -23,6 +23,30 @@
 
 #define LOCTEXT_NAMESPACE "PCGEditorGraphSchemaActions"
 
+/** Disables reconstruct on the owning node of a pin and re-enables in destructor. */
+struct FPCGDeferNodeReconstructScope
+{
+	FPCGDeferNodeReconstructScope(UEdGraphPin* FromPin)
+		: Node(FromPin ? Cast<UPCGEditorGraphNodeBase>(FromPin->GetOwningNode()) : nullptr)
+	{
+		if (Node)
+		{
+			Node->EnableDeferredReconstruct();
+		}
+	}
+
+	~FPCGDeferNodeReconstructScope()
+	{
+		if (Node)
+		{
+			Node->DisableDeferredReconstruct();
+		}
+	}
+
+private:
+	UPCGEditorGraphNodeBase* Node = nullptr;
+};
+
 UEdGraphNode* FPCGEditorGraphSchemaAction_NewNativeElement::PerformAction(UEdGraph* ParentGraph, UEdGraphPin* FromPin, const FVector2D Location, bool bSelectNewNode)
 {
 	UPCGEditorGraph* EditorGraph = Cast<UPCGEditorGraph>(ParentGraph);
@@ -39,6 +63,9 @@ UEdGraphNode* FPCGEditorGraphSchemaAction_NewNativeElement::PerformAction(UEdGra
 		return nullptr;
 	}
 
+	// Important - do not reconstruct the editor graph node/pins midway through this function as this will invalidate FromPin.
+	const FPCGDeferNodeReconstructScope DisableReconstruct(FromPin);
+
 	const FScopedTransaction Transaction(*FPCGEditorCommon::ContextIdentifier, LOCTEXT("PCGEditorNewNativeElement", "PCG Editor: New Native Element"), nullptr);
 	EditorGraph->Modify();
 
@@ -112,6 +139,9 @@ UEdGraphNode* FPCGEditorGraphSchemaAction_NewSettingsElement::PerformAction(UEdG
 		return nullptr;
 	}
 
+	// Important - do not reconstruct the editor graph node/pins midway through this function as this will invalidate FromPin.
+	const FPCGDeferNodeReconstructScope DisableReconstruct(FromPin);
+
 	bool bCreateInstance = false;
 
 	if(Behavior != EPCGEditorNewSettingsBehavior::Normal)
@@ -220,6 +250,9 @@ UEdGraphNode* FPCGEditorGraphSchemaAction_NewSettingsElement::MakeSettingsNode(U
 		return nullptr;
 	}
 
+	// Important - do not reconstruct the editor graph node/pins midway through this function as this will invalidate InFromPin.
+	const FPCGDeferNodeReconstructScope DisableReconstruct(InFromPin);
+
 	const FScopedTransaction Transaction(*FPCGEditorCommon::ContextIdentifier, LOCTEXT("PCGEditorNewSettingsElement", "PCG Editor: New Settings Element"), nullptr);
 	InEditorGraph->Modify();
 
@@ -275,6 +308,9 @@ UEdGraphNode* FPCGEditorGraphSchemaAction_NewBlueprintElement::PerformAction(UEd
 		return nullptr;
 	}
 
+	// Important - do not reconstruct the editor graph node/pins midway through this function as this will invalidate FromPin.
+	const FPCGDeferNodeReconstructScope DisableReconstruct(FromPin);
+
 	const FScopedTransaction Transaction(*FPCGEditorCommon::ContextIdentifier, LOCTEXT("PCGEditorNewBlueprintELement", "PCG Editor: New Blueprint Element"), nullptr);
 	EditorGraph->Modify();
 
@@ -323,6 +359,9 @@ UEdGraphNode* FPCGEditorGraphSchemaAction_NewSubgraphElement::PerformAction(UEdG
 		return nullptr;
 	}
 
+	// Important - do not reconstruct the editor graph node/pins midway through this function as this will invalidate FromPin.
+	const FPCGDeferNodeReconstructScope DisableReconstruct(FromPin);
+
 	UPCGGraph* Subgraph = CastChecked<UPCGGraph>(SubgraphObjectPath.TryLoad());
 
 	const FScopedTransaction Transaction(*FPCGEditorCommon::ContextIdentifier, LOCTEXT("PCGEditorNewSubgraphElement", "PCG Editor: New Subgraph Element"), nullptr);
@@ -399,6 +438,9 @@ UEdGraphNode* FPCGEditorGraphSchemaAction_NewReroute::PerformAction(class UEdGra
 		return nullptr;
 	}
 
+	// Important - do not reconstruct the editor graph node/pins midway through this function as this will invalidate FromPin.
+	const FPCGDeferNodeReconstructScope DisableReconstruct(FromPin);
+
 	const FScopedTransaction Transaction(*FPCGEditorCommon::ContextIdentifier, LOCTEXT("PCGEditorNewReroute", "PCG Editor: New Reroute Node"), nullptr);
 	EditorGraph->Modify();
 
@@ -499,6 +541,9 @@ UEdGraphNode* FPCGEditorGraphSchemaAction_NewNamedRerouteDeclaration::PerformAct
 		return nullptr;
 	}
 
+	// Important - do not reconstruct the editor graph node/pins midway through this function as this will invalidate FromPin.
+	const FPCGDeferNodeReconstructScope DisableReconstruct(FromPin);
+
 	const FScopedTransaction Transaction(*FPCGEditorCommon::ContextIdentifier, LOCTEXT("PCGEditorNewNamedRerouteDeclaration", "PCG Editor: New Named Reroute Node Declaration"), nullptr);
 	EditorGraph->Modify();
 
