Fixed a bAllowShrinking argument.

diff --git a/Engine/Plugins/PCG/Source/PCG/Private/PCGData.cpp b/Engine/Plugins/PCG/Source/PCG/Private/PCGData.cpp
--- a/Engine/Plugins/PCG/Source/PCG/Private/PCGData.cpp
+++ b/Engine/Plugins/PCG/Source/PCG/Private/PCGData.cpp
@@ -413,7 +413,7 @@ void FPCGDataCollection::AddReferences(FReferenceCollector& Collector)
 
 void FPCGDataCollection::ComputeCrcs(bool bFullDataCrc)
 {
-	DataCrcs.SetNumUninitialized(TaggedData.Num(), /*bAllowShrinking=*/false);
+	DataCrcs.SetNumUninitialized(TaggedData.Num(), EAllowShrinking::No);
 
 	for (int I = 0; I < TaggedData.Num(); ++I)
 	{
