[PCG] Add EPCGTextureFilter for Point and Bilinear filtering on PCGTextureData

diff --git a/Engine/Plugins/PCG/Source/PCG/Private/Data/PCGTextureData.cpp b/Engine/Plugins/PCG/Source/PCG/Private/Data/PCGTextureData.cpp
--- a/Engine/Plugins/PCG/Source/PCG/Private/Data/PCGTextureData.cpp
+++ b/Engine/Plugins/PCG/Source/PCG/Private/Data/PCGTextureData.cpp
@@ -45,23 +45,45 @@ namespace PCGTextureSamplingHelpers
 	ValueType SampleInternal(FVector2D PositionLocalSpace,
 		int32 Width,
 		int32 Height,
+		EPCGTextureFilter Filter,
 		TFunctionRef<ValueType(int32 Index)> SamplingFunction)
 	{
-		// Accounts for texel values being at texel centers
-		const double TexelX = (PositionLocalSpace.X * Width - 0.5);
-		const double TexelY = (PositionLocalSpace.Y * Height - 0.5);
+		const double TexelX = PositionLocalSpace.X * Width;
+		const double TexelY = PositionLocalSpace.Y * Height;
 
-		const int32 X0 = FMath::Clamp(FMath::FloorToInt(TexelX), 0, Width - 1);
-		const int32 X1 = FMath::Min(X0 + 1, Width - 1);
-		const int32 Y0 = FMath::Clamp(FMath::FloorToInt(TexelY), 0, Height - 1);
-		const int32 Y1 = FMath::Min(Y0 + 1, Height - 1);
+		ValueType Result{};
 
-		const ValueType SampleX0Y0 = SamplingFunction(X0 + Y0 * Width);
-		const ValueType SampleX1Y0 = SamplingFunction(X1 + Y0 * Width);
-		const ValueType SampleX0Y1 = SamplingFunction(X0 + Y1 * Width);
-		const ValueType SampleX1Y1 = SamplingFunction(X1 + Y1 * Width);
+		if (Filter == EPCGTextureFilter::Point)
+		{
+			const int32 X = FMath::Clamp(FMath::FloorToInt(TexelX), 0, Width - 1);
+			const int32 Y = FMath::Clamp(FMath::FloorToInt(TexelY), 0, Height - 1);
+
+			Result = SamplingFunction(X + Y * Width);
+		}
+		else if (Filter == EPCGTextureFilter::Bilinear)
+		{
+			// Accounts for texel values being at texel centers
+			const double TexelXOffset = TexelX - 0.5;
+			const double TexelYOffset = TexelY - 0.5;
+
+			const int32 X0 = FMath::Clamp(FMath::FloorToInt(TexelXOffset), 0, Width - 1);
+			const int32 X1 = FMath::Min(X0 + 1, Width - 1);
+			const int32 Y0 = FMath::Clamp(FMath::FloorToInt(TexelYOffset), 0, Height - 1);
+			const int32 Y1 = FMath::Min(Y0 + 1, Height - 1);
+
+			const ValueType SampleX0Y0 = SamplingFunction(X0 + Y0 * Width);
+			const ValueType SampleX1Y0 = SamplingFunction(X1 + Y0 * Width);
+			const ValueType SampleX0Y1 = SamplingFunction(X0 + Y1 * Width);
+			const ValueType SampleX1Y1 = SamplingFunction(X1 + Y1 * Width);
+
+			Result = FMath::BiLerp(SampleX0Y0, SampleX1Y0, SampleX0Y1, SampleX1Y1, TexelXOffset - X0, TexelYOffset - Y0);
+		}
+		else
+		{
+			ensureMsgf(false, TEXT("Unrecognized PCG texture filtering mode."));
+		}
 
-		return FMath::BiLerp(SampleX0Y0, SampleX1Y0, SampleX0Y1, SampleX1Y1, TexelX - X0, TexelY - Y0);
+		return Result;
 	}
 
 	template<typename ValueType>
@@ -117,7 +139,7 @@ namespace PCGTextureSamplingHelpers
 			Pos = FVector2D(X, Y);
 		}
 
-		SampledValue = SampleInternal(Pos, Width, Height, SamplingFunction);
+		SampledValue = SampleInternal(Pos, Width, Height, InTextureData->Filter, SamplingFunction);
 		return true;
 	}
 
@@ -272,7 +294,7 @@ bool UPCGBaseTextureData::SamplePointLocal(const FVector2D& LocalPosition, FVect
 	Pos.X = FMath::Frac(LocalPosition.X);
 	Pos.Y = FMath::Frac(LocalPosition.Y);
 
-	const FLinearColor OutSample = PCGTextureSamplingHelpers::SampleInternal<FLinearColor>(Pos, Width, Height, [this](int32 Index) { return ColorData[Index]; });
+	const FLinearColor OutSample = PCGTextureSamplingHelpers::SampleInternal<FLinearColor>(Pos, Width, Height, Filter, [this](int32 Index) { return ColorData[Index]; });
 
 	OutColor = OutSample;
 	OutDensity = (DensityFunction == EPCGTextureDensityFunction::Ignore) ? 1.0f : PCGTextureSamplingHelpers::SampleFloatChannel(OutSample, ColorChannel);
diff --git a/Engine/Plugins/PCG/Source/PCG/Private/Elements/PCGTextureSampler.cpp b/Engine/Plugins/PCG/Source/PCG/Private/Elements/PCGTextureSampler.cpp
--- a/Engine/Plugins/PCG/Source/PCG/Private/Elements/PCGTextureSampler.cpp
+++ b/Engine/Plugins/PCG/Source/PCG/Private/Elements/PCGTextureSampler.cpp
@@ -172,6 +172,7 @@ bool FPCGTextureSamplerElement::ExecuteInternal(FPCGContext* InContext) const
 	const bool bUseAbsoluteTransform = Settings->bUseAbsoluteTransform;
 	const EPCGTextureDensityFunction DensityFunction = Settings->DensityFunction;
 	const EPCGTextureColorChannel ColorChannel = Settings->ColorChannel;
+	const EPCGTextureFilter Filter = Settings->Filter;
 	const float TexelSize = Settings->TexelSize;
 	const bool bUseAdvancedTiling = Settings->bUseAdvancedTiling;
 	const FVector2D& Tiling = Settings->Tiling;
@@ -221,6 +222,7 @@ bool FPCGTextureSamplerElement::ExecuteInternal(FPCGContext* InContext) const
 
 	TextureData->DensityFunction = DensityFunction;
 	TextureData->ColorChannel = ColorChannel;
+	TextureData->Filter = Filter;
 	TextureData->TexelSize = TexelSize;
 	TextureData->bUseAdvancedTiling = bUseAdvancedTiling;
 	TextureData->Tiling = Tiling;
diff --git a/Engine/Plugins/PCG/Source/PCG/Public/Data/PCGTextureData.h b/Engine/Plugins/PCG/Source/PCG/Public/Data/PCGTextureData.h
--- a/Engine/Plugins/PCG/Source/PCG/Public/Data/PCGTextureData.h
+++ b/Engine/Plugins/PCG/Source/PCG/Public/Data/PCGTextureData.h
@@ -26,6 +26,14 @@ enum class EPCGTextureDensityFunction : uint8
 	Multiply
 };
 
+/** Method used to determine the value for a sample based on the value of nearby texels. */
+UENUM(BlueprintType)
+enum class EPCGTextureFilter : uint8
+{
+	Point UMETA(Tooltip="Takes the value of whatever texel the sample lands in."),
+	Bilinear UMETA(Tooltip="Bilinearly interpolates the values of the four nearest texels to the sample location.")
+};
+
 UENUM()
 enum class EPCGTextureAddressMode : uint8
 {
@@ -74,6 +82,10 @@ public:
 	UPROPERTY(BlueprintReadWrite, EditAnywhere, Category = Settings)
 	EPCGTextureColorChannel ColorChannel = EPCGTextureColorChannel::Alpha;
 
+	/** Method used to determine the value for a sample based on the value of nearby texels. */
+	UPROPERTY(BlueprintReadWrite, EditAnywhere, Category = Settings, meta = (PCG_Overridable))
+	EPCGTextureFilter Filter = EPCGTextureFilter::Bilinear;
+
 	/** The size of one texel in cm, used when calling ToPointData. */
 	UPROPERTY(BlueprintReadWrite, EditAnywhere, Category = Settings, meta = (UIMin = "1.0", ClampMin = "1.0"))
 	float TexelSize = 50.0f;
diff --git a/Engine/Plugins/PCG/Source/PCG/Public/Elements/PCGTextureSampler.h b/Engine/Plugins/PCG/Source/PCG/Public/Elements/PCGTextureSampler.h
--- a/Engine/Plugins/PCG/Source/PCG/Public/Elements/PCGTextureSampler.h
+++ b/Engine/Plugins/PCG/Source/PCG/Public/Elements/PCGTextureSampler.h
@@ -60,13 +60,16 @@ public:
 	UPROPERTY(BlueprintReadWrite, EditAnywhere, Category = Settings, meta = (EditCondition = bDisplayTextureArrayIndex, EditConditionHides, HideEditConditionToggle, ClampMin = '0', PCG_Overridable))
 	int TextureArrayIndex = 0;
 
-	// Common members in BaseTextureData
 	UPROPERTY(BlueprintReadWrite, EditAnywhere, Category = SpatialData, meta = (PCG_Overridable))
 	EPCGTextureDensityFunction DensityFunction = EPCGTextureDensityFunction::Multiply;
 
 	UPROPERTY(BlueprintReadWrite, EditAnywhere, Category = Settings, meta = (PCG_Overridable))
 	EPCGTextureColorChannel ColorChannel = EPCGTextureColorChannel::Alpha;
 
+	/** Method used to determine the value for a sample based on the value of nearby texels. */
+	UPROPERTY(BlueprintReadWrite, EditAnywhere, Category = Settings, meta = (PCG_Overridable))
+	EPCGTextureFilter Filter = EPCGTextureFilter::Bilinear;
+
 	/** The size of one texel in cm, used when calling ToPointData. */
 	UPROPERTY(BlueprintReadWrite, EditAnywhere, Category = Settings, meta = (UIMin = "1.0", ClampMin = "1.0", PCG_Overridable))
 	float TexelSize = 50.0f;
