[PCG] Add Attributes now supports the addition of multiple attributes in one go

#rnx

#changelist validated
#virtualized

diff --git a/Engine/Plugins/PCG/Source/PCG/Private/Elements/PCGCreateAttribute.cpp b/Engine/Plugins/PCG/Source/PCG/Private/Elements/PCGCreateAttribute.cpp
--- a/Engine/Plugins/PCG/Source/PCG/Private/Elements/PCGCreateAttribute.cpp
+++ b/Engine/Plugins/PCG/Source/PCG/Private/Elements/PCGCreateAttribute.cpp
@@ -295,14 +295,20 @@ bool FPCGAddAttributeElement::ExecuteInternal(FPCGContext* Context) const
 	}
 
 	// Otherwise, is is like a copy
-	const UPCGParamData* SourceParamData = !SourceParams.IsEmpty() ? Cast<UPCGParamData>(SourceParams[0].Data) : nullptr;
-	if (!SourceParamData)
+	const UPCGParamData* FirstSourceParamData = !SourceParams.IsEmpty() ? Cast<UPCGParamData>(SourceParams[0].Data) : nullptr;
+	if (!FirstSourceParamData)
 	{
 		// Nothing to do
 		Context->OutputData.TaggedData = Inputs;
 		return true;
 	}
 
+	// If we copy all attributes, support to have multiple source param. Overwise, add a warning
+	if (SourceParams.Num() > 1 && !Settings->bCopyAllAttributes)
+	{
+		PCGLog::LogWarningOnGraph(LOCTEXT("MultiAttributeWhenNoCopyAll", "Multiple source param detected in the Attributes pin, but we do not copy all attributes. We will only look into the first source param."), Context);
+	}
+
 	for (int32 i = 0; i < Inputs.Num(); ++i)
 	{
 		const UPCGData* InputData = Inputs[i].Data;
@@ -314,14 +320,20 @@ bool FPCGAddAttributeElement::ExecuteInternal(FPCGContext* Context) const
 		UPCGData* TargetData = InputData->DuplicateData();
 		FPCGTaggedData& Output = Context->OutputData.TaggedData.Add_GetRef(Inputs[i]);
 
-		bool bSuccess = false;
+		bool bSuccess = true;
 		if (Settings->bCopyAllAttributes)
 		{
-			bSuccess = PCGMetadataHelpers::CopyAllAttributes(SourceParamData, TargetData, Context);
+			for (const FPCGTaggedData& SourceParamData : SourceParams)
+			{
+				if (const UPCGParamData* ParamData = Cast<UPCGParamData>(SourceParamData.Data))
+				{
+					bSuccess &= PCGMetadataHelpers::CopyAllAttributes(ParamData, TargetData, Context);
+				}
+			}
 		}
 		else
 		{
-			bSuccess = PCGMetadataHelpers::CopyAttributes(SourceParamData, Settings->InputSource, TargetData, Settings->OutputTarget, /*bSameOrigin=*/false, Context);
+			bSuccess = PCGMetadataHelpers::CopyAttributes(FirstSourceParamData, Settings->InputSource, TargetData, Settings->OutputTarget, /*bSameOrigin=*/false, Context);
 		}
 
 		if (bSuccess)
