[PCG] Add context menu to debug object tree
--The user is now able to jump to a specfic node/graph by right clicking the object inthe debug tree

#changelist validated
#virtualized

diff --git a/Engine/Plugins/PCG/Source/PCGEditor/Private/PCGEditor.cpp b/Engine/Plugins/PCG/Source/PCGEditor/Private/PCGEditor.cpp
--- a/Engine/Plugins/PCG/Source/PCGEditor/Private/PCGEditor.cpp
+++ b/Engine/Plugins/PCG/Source/PCGEditor/Private/PCGEditor.cpp
@@ -369,6 +369,26 @@ void FPCGEditor::JumpToNode(const UEdGraphNode* InNode)
 	}
 }
 
+void FPCGEditor::JumpToNode(const UPCGNode* InNode)
+{
+	if (!ensure(PCGEditorGraph))
+	{
+		return;
+	}
+
+	for (const UEdGraphNode* EdGraphNode : PCGEditorGraph->Nodes)
+	{
+		if (const UPCGEditorGraphNodeBase* PCGEdGraphNode = Cast<UPCGEditorGraphNodeBase>(EdGraphNode))
+		{
+			if (PCGEdGraphNode->GetPCGNode() == InNode)
+			{
+				JumpToNode(EdGraphNode);
+				break;
+			}
+		}
+	}
+}
+
 void FPCGEditor::RegisterTabSpawners(const TSharedRef<FTabManager>& InTabManager)
 {
 	WorkspaceMenuCategory = InTabManager->AddLocalWorkspaceMenuCategory(LOCTEXT("WorkspaceMenu_PCGEditor", "PCG Editor"));
diff --git a/Engine/Plugins/PCG/Source/PCGEditor/Private/PCGEditor.h b/Engine/Plugins/PCG/Source/PCGEditor/Private/PCGEditor.h
--- a/Engine/Plugins/PCG/Source/PCGEditor/Private/PCGEditor.h
+++ b/Engine/Plugins/PCG/Source/PCGEditor/Private/PCGEditor.h
@@ -37,12 +37,15 @@ public:
 	/** Edits the specified PCGGraph */
 	void Initialize(const EToolkitMode::Type InMode, const TSharedPtr<class IToolkitHost>& InToolkitHost, UPCGGraph* InPCGGraph);
 
-	/** Get the PCG graph being edited */
+	/** Get the PCG editor graph being edited. */
 	UPCGEditorGraph* GetPCGEditorGraph();
 
 	/** Gets/Creates the PCG graph editor for a given PCG graph */
 	static UPCGEditorGraph* GetPCGEditorGraph(UPCGGraph* InGraph);
 
+	/** Get the PCG graph being edited. */
+	const UPCGGraph* GetPCGGraph() { return PCGGraphBeingEdited; }
+
 	/** Sets the execution stack that want to inspect. */
 	void SetStackBeingInspected(const FPCGStack& FullStack);
 
@@ -54,6 +57,7 @@ public:
 
 	/** Focus the graph view on a specific node */
 	void JumpToNode(const UEdGraphNode* InNode);
+	void JumpToNode(const UPCGNode* InNode);
 
 	/** Helper to get to the subsystem. */
 	static class UPCGSubsystem* GetSubsystem();
diff --git a/Engine/Plugins/PCG/Source/PCGEditor/Private/SPCGEditorGraphDebugObjectTree.cpp b/Engine/Plugins/PCG/Source/PCGEditor/Private/SPCGEditorGraphDebugObjectTree.cpp
--- a/Engine/Plugins/PCG/Source/PCGEditor/Private/SPCGEditorGraphDebugObjectTree.cpp
+++ b/Engine/Plugins/PCG/Source/PCGEditor/Private/SPCGEditorGraphDebugObjectTree.cpp
@@ -252,16 +252,17 @@ void SPCGEditorGraphDebugObjectTree::Construct(const FArguments& InArgs, TShared
 		.Thickness(FVector2D(12.0f, 12.0f));
 
 	DebugObjectTreeView = SNew(STreeView<FPCGEditorGraphDebugObjectItemPtr>)
-				.TreeItemsSource(&RootItems)
-				.OnGenerateRow(this, &SPCGEditorGraphDebugObjectTree::MakeTreeRowWidget)
-				.OnGetChildren(this, &SPCGEditorGraphDebugObjectTree::OnGetChildren)
-				.OnSelectionChanged(this, &SPCGEditorGraphDebugObjectTree::OnSelectionChanged)
-				.OnExpansionChanged(this, &SPCGEditorGraphDebugObjectTree::OnExpansionChanged)
-				.OnSetExpansionRecursive(this, &SPCGEditorGraphDebugObjectTree::OnSetExpansionRecursive)
-				.ItemHeight(18)
-				.AllowOverscroll(EAllowOverscroll::No)
-				.ExternalScrollbar(VerticalScrollBar)
-				.ConsumeMouseWheel(EConsumeMouseWheel::Always);
+		.TreeItemsSource(&RootItems)
+		.OnGenerateRow(this, &SPCGEditorGraphDebugObjectTree::MakeTreeRowWidget)
+		.OnGetChildren(this, &SPCGEditorGraphDebugObjectTree::OnGetChildren)
+		.OnSelectionChanged(this, &SPCGEditorGraphDebugObjectTree::OnSelectionChanged)
+		.OnExpansionChanged(this, &SPCGEditorGraphDebugObjectTree::OnExpansionChanged)
+		.OnSetExpansionRecursive(this, &SPCGEditorGraphDebugObjectTree::OnSetExpansionRecursive)
+		.ItemHeight(18)
+		.AllowOverscroll(EAllowOverscroll::No)
+		.ExternalScrollbar(VerticalScrollBar)
+		.ConsumeMouseWheel(EConsumeMouseWheel::Always)
+		.OnContextMenuOpening(FOnContextMenuOpening::CreateSP(this, &SPCGEditorGraphDebugObjectTree::OpenContextMenu));
 
 	const TSharedRef<SWidget> SetButton = PropertyCustomizationHelpers::MakeUseSelectedButton(
 		FSimpleDelegate::CreateSP(this, &SPCGEditorGraphDebugObjectTree::SetDebugObjectFromSelection_OnClicked),
@@ -880,7 +881,14 @@ void SPCGEditorGraphDebugObjectTree::OnGetChildren(FPCGEditorGraphDebugObjectIte
 
 void SPCGEditorGraphDebugObjectTree::OnSelectionChanged(FPCGEditorGraphDebugObjectItemPtr InItem, ESelectInfo::Type InSelectInfo)
 {
-	if (const FPCGStack* Stack = (InItem && InItem->IsDebuggable()) ? InItem->GetPCGStack() : nullptr)
+	// If the user selects a new item, record the previous stack. This is helpful to give previous selection information to 
+	// context menu commands, because right clicking an item changes the selection.
+	if (InSelectInfo != ESelectInfo::Direct)
+	{
+		PreviouslySelectedStack = SelectedStack;
+	}
+
+	if (const FPCGStack* Stack = InItem ? InItem->GetPCGStack() : nullptr)
 	{
 		SelectedStack = *Stack;
 	}
@@ -985,4 +993,93 @@ void SPCGEditorGraphDebugObjectTree::ExpandAndSelectDebugObject(FPCGEditorGraphD
 	}
 }
 
+TSharedPtr<SWidget> SPCGEditorGraphDebugObjectTree::OpenContextMenu()
+{
+	FMenuBuilder MenuBuilder(/*bInShouldCloseWindowAfterMenuSelection=*/true, /*InCommandList=*/nullptr);
+	MenuBuilder.AddMenuEntry(
+		LOCTEXT("JumpToSelectedGraph", "Jump To"),
+		LOCTEXT("JumpToSelectedGraphTooltip", "Jumps to the selected graph."),
+		FSlateIcon(FAppStyle::GetAppStyleSetName(), "SystemWideCommands.FindInContentBrowser.Small"),
+		FUIAction(FExecuteAction::CreateSP(this, &SPCGEditorGraphDebugObjectTree::ContextMenu_JumpToGraphInTree), FCanExecuteAction::CreateSP(this, &SPCGEditorGraphDebugObjectTree::ContextMenu_JumpToGraphInTree_CanExecute))
+	);
+
+	return MenuBuilder.MakeWidget();
+}
+
+void SPCGEditorGraphDebugObjectTree::ContextMenu_JumpToGraphInTree()
+{
+	// If we were previously debugging a target and we're jumping to a parent graph, jump to the subgraph/loop node
+	// that corresponds to the previously selected stack.
+	// 
+	// Example:
+	//		Jump target: Component/TopGraph
+	//		Previous stack: Component/TopGraph/SubgraphNode/Subgraph
+	// 
+	// Here the "subject" of the jump here is SubgraphNode - obtain this from the previous stack, frame index after TopGraph.
+
+	const TArray<FPCGEditorGraphDebugObjectItemPtr> SelectedItems = DebugObjectTreeView->GetSelectedItems();
+	for (const FPCGEditorGraphDebugObjectItemPtr& SelectedItem : SelectedItems)
+	{
+		UPCGGraph* JumpToPCGGraph = nullptr;
+		UPCGNode* JumpToPCGNode = nullptr;
+
+		if (SelectedItem->IsDebuggable())
+		{
+			// Debuggable target graphs correspond to the currently edited graph. For this case open the parent graph and jump to the corresponding subgraph node.
+			
+			// Search the stack for the parent graph.
+			for (int i = SelectedStack.GetStackFrames().Num() - 2; i > 0; --i) 
+			{
+				if (SelectedStack.GetStackFrames()[i].Object != nullptr && SelectedStack.GetStackFrames()[i].Object->IsA<UPCGGraph>())
+				{
+					JumpToPCGGraph = const_cast<UPCGGraph*>(Cast<const UPCGGraph>(SelectedStack.GetStackFrames()[i].Object));
+					JumpToPCGNode = const_cast<UPCGNode*>(Cast<const UPCGNode>(SelectedStack.GetStackFrames()[i + 1].Object));
+					break;
+				}
+			}
+		}
+		else 
+		{
+			// Non-debuggable target graphs are graphs that are not the current edited graph.
+			JumpToPCGGraph = const_cast<UPCGGraph*>(SelectedItem->GetPCGGraph());
+
+			// When the user right clicks, the selection moves. PreviouslySelectedStack will hold the previous selection, and
+			// we retrieve our subject node from there, which can be found 1 frame after the jump to graph, hence the Num() below.
+			const int SubjectNodeFrameIndex = SelectedStack.GetStackFrames().Num();
+			if (PreviouslySelectedStack.GetStackFrames().IsValidIndex(SubjectNodeFrameIndex))
+			{
+				JumpToPCGNode = const_cast<UPCGNode*>(Cast<const UPCGNode>(PreviouslySelectedStack.GetStackFrames()[SubjectNodeFrameIndex].Object));
+			}
+		}
+
+		if (JumpToPCGGraph && JumpToPCGNode)
+		{
+			GEditor->GetEditorSubsystem<UAssetEditorSubsystem>()->OpenEditorForAsset(JumpToPCGGraph);
+
+			if (IAssetEditorInstance* EditorInstance = GEditor->GetEditorSubsystem<UAssetEditorSubsystem>()->FindEditorForAsset(JumpToPCGGraph, /*bFocusIfOpen*/true))
+			{
+				static_cast<FPCGEditor*>(EditorInstance)->JumpToNode(JumpToPCGNode);
+			}
+		}
+	}
+}
+
+bool SPCGEditorGraphDebugObjectTree::ContextMenu_JumpToGraphInTree_CanExecute() const
+{
+	const TArray<FPCGEditorGraphDebugObjectItemPtr> SelectedItems = DebugObjectTreeView->GetSelectedItems();
+	for (const FPCGEditorGraphDebugObjectItemPtr& SelectedItem : SelectedItems)
+	{
+		const UPCGGraph* PCGGraphBeingEdited = PCGEditor.IsValid() ? PCGEditor.Pin()->GetPCGGraph() : nullptr;
+		const UPCGGraph* SelectedPCGGraph = SelectedItem->GetPCGGraph();
+
+		// Offer jump-to command if any selected item is a graph or a subgraph loop iteration.
+		if (SelectedPCGGraph || SelectedItem->IsLoopIteration())
+		{
+			return true;
+		}
+	}
+
+	return false;
+}
+
 #undef LOCTEXT_NAMESPACE
diff --git a/Engine/Plugins/PCG/Source/PCGEditor/Private/SPCGEditorGraphDebugObjectTree.h b/Engine/Plugins/PCG/Source/PCGEditor/Private/SPCGEditorGraphDebugObjectTree.h
--- a/Engine/Plugins/PCG/Source/PCGEditor/Private/SPCGEditorGraphDebugObjectTree.h
+++ b/Engine/Plugins/PCG/Source/PCGEditor/Private/SPCGEditorGraphDebugObjectTree.h
@@ -49,6 +49,8 @@ public:
 	virtual UPCGComponent* GetPCGComponent() const = 0;
 	virtual const FPCGStack* GetPCGStack() const { return nullptr; }
 	virtual const UObject* GetObject() const  = 0;
+	virtual const UPCGGraph* GetPCGGraph() const { return nullptr; }
+	virtual bool IsLoopIteration() const { return false; }
 
 protected:
 	TWeakPtr<FPCGEditorGraphDebugObjectItem> Parent;
@@ -99,6 +101,7 @@ public:
 	virtual const UObject* GetObject() const override { return PCGComponent.Get(); };
 	virtual const FPCGStack* GetPCGStack() const override { return &PCGStack; }
 	virtual bool IsDebuggable() const override { return bIsDebuggable; }
+	virtual const UPCGGraph* GetPCGGraph() const override { return PCGGraph.Get(); }
 
 protected:
 	TWeakObjectPtr<UPCGComponent> PCGComponent = nullptr;
@@ -129,6 +132,7 @@ public:
 	virtual const UObject* GetObject() const override { return PCGNode.Get(); };
 	virtual const FPCGStack* GetPCGStack() const override { return &PCGStack; }
 	virtual bool IsDebuggable() const override { return bIsDebuggable; }
+	virtual const UPCGGraph* GetPCGGraph() const override { return PCGGraph.Get(); }
 
 protected:
 	TWeakObjectPtr<const UPCGNode> PCGNode = nullptr;
@@ -162,6 +166,8 @@ public:
 	virtual const FPCGStack* GetPCGStack() const override { return &PCGStack; }
 	virtual int32 GetSortPriority() const override { return LoopIndex; }
 	virtual bool IsDebuggable() const override { return bIsDebuggable; }
+	virtual bool IsLoopIteration() const override { return true; }
+	virtual const UPCGGraph* GetPCGGraph() const override { return Cast<UPCGGraph>(LoopedPCGGraph.Get()); }
 
 protected:
 	int32 LoopIndex = INDEX_NONE;
@@ -179,9 +185,7 @@ public:
 	SLATE_BEGIN_ARGS(SPCGEditorGraphDebugObjectItemRow)
 		: _OnDoubleClickFunc()
 	{}
-
 		SLATE_ARGUMENT(FDoubleClickFunc, OnDoubleClickFunc)
-
 	SLATE_END_ARGS()
 
 	void Construct(const FArguments& InArgs, const TSharedRef<STableViewBase>& InOwnerTableView, FPCGEditorGraphDebugObjectItemPtr InItem);
@@ -245,6 +249,12 @@ private:
 	/** Expand the given row and select the deepest entry as the debug object if it is unambiguous (the only entry at its level in the tree). */
 	void ExpandAndSelectDebugObject(FPCGEditorGraphDebugObjectItemPtr InItem);
 
+	TSharedPtr<SWidget> OpenContextMenu();
+
+	/** Jump-to context menu command. */
+	void ContextMenu_JumpToGraphInTree();
+	bool ContextMenu_JumpToGraphInTree_CanExecute() const;
+
 	TWeakPtr<FPCGEditor> PCGEditor;
 
 	TSharedPtr<STreeView<FPCGEditorGraphDebugObjectItemPtr>> DebugObjectTreeView;
@@ -262,5 +272,8 @@ private:
 	/** Used to retain item selection state across tree refreshes. */
 	FPCGStack SelectedStack;
 
+	/** The previous stack that the user selected. */
+	FPCGStack PreviouslySelectedStack;
+
 	const UPCGNode* PCGNodeBeingInspected = nullptr;
 };
