[PCG] Add an advanced option for Generate on Drop, to have Generate on demand (no generation in PIE) but still generate the PCG when it is droped in the world

#rnx

#changelist validated
#virtualized

diff --git a/Engine/Plugins/PCG/Source/PCG/Private/PCGComponent.cpp b/Engine/Plugins/PCG/Source/PCG/Private/PCGComponent.cpp
--- a/Engine/Plugins/PCG/Source/PCG/Private/PCGComponent.cpp
+++ b/Engine/Plugins/PCG/Source/PCG/Private/PCGComponent.cpp
@@ -2152,9 +2152,16 @@ bool UPCGComponent::ShouldGenerateBPPCGAddedToWorld() const
 	}
 	else
 	{
-		// Generate on drop can be disabled by global settings or locally by not having "GenerateOnLoad" as a generation trigger.
-		const UPCGEngineSettings* Settings = GetDefault<UPCGEngineSettings>();
-		return Settings ? (Settings->bGenerateOnDrop && bForceGenerateOnBPAddedToWorld && (GenerationTrigger == EPCGComponentGenerationTrigger::GenerateOnLoad)) : false;
+		// Generate on drop can be disabled by global settings or locally by not having "GenerateOnLoad" as a generation trigger (and Generate on Drop option to false).
+		if (const UPCGEngineSettings* Settings = GetDefault<UPCGEngineSettings>())
+		{
+			return Settings->bGenerateOnDrop && bForceGenerateOnBPAddedToWorld &&
+				(GenerationTrigger == EPCGComponentGenerationTrigger::GenerateOnLoad || (GenerationTrigger == EPCGComponentGenerationTrigger::GenerateOnDemand && bGenerateOnDropWhenTriggerOnDemand));
+		}
+		else
+		{
+			return false;
+		}
 	}
 }
 
diff --git a/Engine/Plugins/PCG/Source/PCG/Public/PCGComponent.h b/Engine/Plugins/PCG/Source/PCG/Public/PCGComponent.h
--- a/Engine/Plugins/PCG/Source/PCG/Public/PCGComponent.h
+++ b/Engine/Plugins/PCG/Source/PCG/Public/PCGComponent.h
@@ -222,6 +222,10 @@ public:
 	UPROPERTY(BlueprintReadOnly, EditAnywhere, Category = Settings, meta = (EditCondition = "!bIsComponentLocal", EditConditionHides, DisplayPriority = 200))
 	EPCGComponentGenerationTrigger GenerationTrigger = EPCGComponentGenerationTrigger::GenerateOnLoad;
 
+	/** When Generation Trigger is OnDemand, we can still force the component to generate on drop. */
+	UPROPERTY(BlueprintReadOnly, EditAnywhere, Category = "Settings|Advanced" , meta = (EditCondition = "GenerationTrigger == EPCGComponentGenerationTrigger::GenerateOnDemand"))
+	bool bGenerateOnDropWhenTriggerOnDemand = false;
+
 	/** Manual overrides for the graph generation radii and cleanup radius multiplier. */
 	UPROPERTY(BlueprintReadWrite, EditAnywhere, Category = RuntimeGeneration, meta = (EditCondition = "GenerationTrigger == EPCGComponentGenerationTrigger::GenerateAtRuntime", EditConditionHides))
 	bool bOverrideGenerationRadii = false;
