[PCG] Fix automated test failure from classes not initialized properly

#rnx

diff --git a/Engine/Plugins/PCG/Source/PCG/Public/PCGWorldActor.h b/Engine/Plugins/PCG/Source/PCG/Public/PCGWorldActor.h
--- a/Engine/Plugins/PCG/Source/PCG/Public/PCGWorldActor.h
+++ b/Engine/Plugins/PCG/Source/PCG/Public/PCGWorldActor.h
@@ -24,11 +24,11 @@ struct FPCGPartitionActorRecord
 
 	/** The grid size this actor lives on. */
 	UPROPERTY(VisibleAnywhere, Category = Debug)
-	uint32 GridSize;
+	uint32 GridSize = 0;
 
 	/** The specific grid cell this actor lives in. */
 	UPROPERTY(VisibleAnywhere, Category = Debug)
-	FIntVector GridCoords;
+	FIntVector GridCoords = FIntVector::ZeroValue;
 
 	bool operator==(const FPCGPartitionActorRecord& InOther) const;
 	friend uint32 GetTypeHash(const FPCGPartitionActorRecord& In);
