[PCG] Add option to compute alpha and distance along curve as attributes in PCGSplineSampler
* Alpha is a value in [0,1] along the curve where each segment of the curve represents the same interval. For example, if the curve has three segments, points on those segments would lie in [0,0.333], [0.333, 0.666], [0.666, 1] respectively, and interpolate across those intervals using the spline's interpolation.
* Distance is the total cumulative distance along the spline to reach a given sample point

diff --git a/Engine/Plugins/PCG/Source/PCG/Private/Data/PCGPolyLineData.cpp b/Engine/Plugins/PCG/Source/PCG/Private/Data/PCGPolyLineData.cpp
--- a/Engine/Plugins/PCG/Source/PCG/Private/Data/PCGPolyLineData.cpp
+++ b/Engine/Plugins/PCG/Source/PCG/Private/Data/PCGPolyLineData.cpp
@@ -37,3 +37,10 @@ void UPCGPolyLineData::GetTangentsAtSegmentStart(int SegmentIndex, FVector& OutA
 	OutArriveTangent = FVector::Zero();
 	OutLeaveTangent = FVector::Zero();
 }
+
+float UPCGPolyLineData::GetAlphaAtDistance(int SegmentIndex, FVector::FReal Distance) const
+{
+	const int32 NumSegments = GetNumSegments();
+	const FVector::FReal SegmentLength = GetSegmentLength(SegmentIndex);
+	return SegmentIndex / static_cast<float>(NumSegments) + (Distance / SegmentLength / NumSegments);
+}
diff --git a/Engine/Plugins/PCG/Source/PCG/Private/Elements/PCGSplineSampler.cpp b/Engine/Plugins/PCG/Source/PCG/Private/Elements/PCGSplineSampler.cpp
--- a/Engine/Plugins/PCG/Source/PCG/Private/Elements/PCGSplineSampler.cpp
+++ b/Engine/Plugins/PCG/Source/PCG/Private/Elements/PCGSplineSampler.cpp
@@ -346,6 +346,8 @@ namespace PCGSplineSamplerHelpers
 		FVector LeaveTangent = FVector::Zero();
 		FVector::FReal PreviousDeltaAngle = 0;
 		FVector::FReal NextDeltaAngle = 0;
+		FVector::FReal Alpha = 0.0f;
+		FVector::FReal Distance = 0.0f;
 	};
 
 	void SetSeed(FPCGPoint& Point, const FVector& LSPosition, const FPCGSplineSamplerParams& Params)
@@ -380,6 +382,8 @@ namespace PCGSplineSamplerHelpers
 		FStepSampler(const UPCGPolyLineData* InLineData, const FPCGSplineSamplerParams& Params)
 			: LineData(InLineData)
 			, bComputeCurvature(Params.bComputeCurvature)
+			, bComputeAlpha(Params.bComputeAlpha)
+			, bComputeDistance(Params.bComputeDistance)
 		{
 			check(LineData);
 			CurrentSegmentIndex = 0;
@@ -390,7 +394,10 @@ namespace PCGSplineSamplerHelpers
 
 		const UPCGPolyLineData* LineData = nullptr;
 		int CurrentSegmentIndex = 0;
+		FVector::FReal DistanceToCurrentSegment = 0.0f;
 		bool bComputeCurvature = false;
+		bool bComputeAlpha = false;
+		bool bComputeDistance = false;
 	};
 
 	struct FSubdivisionStepSampler : public FStepSampler
@@ -445,6 +452,22 @@ namespace PCGSplineSamplerHelpers
 				}
 			}
 
+			if (bComputeAlpha)
+			{
+				OutResult.Alpha = LineData->GetAlphaAtDistance(SegmentIndex, DistanceAlongSegment);
+			}
+
+			if (bComputeDistance)
+			{
+				// When we step onto a new segment, add the length of the previous segment onto the distance to our current segment.
+				if (SegmentIndex > 0 && !bLastKeyPoint && SubpointIndex == 0)
+				{
+					DistanceToCurrentSegment += LineData->GetSegmentLength(SegmentIndex - 1);
+				}
+
+				OutResult.Distance = DistanceToCurrentSegment + DistanceAlongSegment;
+			}
+
 			if (SubpointIndex == 0)
 			{
 				const FVector::FReal PreviousSegmentLength = LineData->GetSegmentLength(PreviousSegmentIndex);
@@ -481,9 +504,9 @@ namespace PCGSplineSamplerHelpers
 			return CurrentSegmentIndex >= NumSegmentsWhenDone;
 		}
 
-		int NumSegments;
-		int SubdivisionsPerSegment;
-		int SubpointIndex;
+		int NumSegments = 0;
+		int SubdivisionsPerSegment = 0;
+		int SubpointIndex = 0;
 		bool bComputeTangents = false;
 	};
 
@@ -513,11 +536,27 @@ namespace PCGSplineSamplerHelpers
 				OutResult.Curvature = LineData->GetCurvatureAtDistance(CurrentSegmentIndex, CurrentDistance);
 			}
 
+			if (bComputeAlpha)
+			{
+				OutResult.Alpha = LineData->GetAlphaAtDistance(CurrentSegmentIndex, CurrentDistance);
+			}
+
+			if (bComputeDistance)
+			{
+				OutResult.Distance = DistanceToCurrentSegment + CurrentDistance;
+			}
+
 			CurrentDistance += DistanceIncrement;
 			while(CurrentDistance > CurrentSegmentLength)
 			{
 				CurrentDistance -= CurrentSegmentLength;
 				++CurrentSegmentIndex;
+
+				if (bComputeDistance)
+				{
+					DistanceToCurrentSegment += CurrentSegmentLength;
+				}
+
 				if (!IsDone())
 				{
 					CurrentSegmentLength = LineData->GetSegmentLength(CurrentSegmentIndex);
@@ -534,8 +573,8 @@ namespace PCGSplineSamplerHelpers
 			return CurrentSegmentIndex >= LineData->GetNumSegments();
 		}
 
-		FVector::FReal DistanceIncrement;
-		FVector::FReal CurrentDistance;
+		FVector::FReal CurrentDistance = 0.0f;
+		FVector::FReal DistanceIncrement = 0.0f;
 	};
 
 	struct FDimensionSampler
@@ -556,7 +595,9 @@ namespace PCGSplineSamplerHelpers
 					|| Params.bComputeCurvature
 					|| Params.bComputeSegmentIndex
 					|| Params.bComputeSubsegmentIndex
-					|| Params.bComputeTangents))
+					|| Params.bComputeTangents
+					|| Params.bComputeAlpha
+					|| Params.bComputeDistance))
 			{
 				constexpr double DefaultValue = 0.0;
 				if (Params.bComputeDirectionDelta)
@@ -591,6 +632,18 @@ namespace PCGSplineSamplerHelpers
 					LeaveTangentAttribute = OutPointData->Metadata->FindOrCreateAttribute<FVector>(Params.LeaveTangentAttribute, FVector::Zero());
 					bSetMetadata |= (LeaveTangentAttribute != nullptr);
 				}
+
+				if (Params.bComputeAlpha)
+				{
+					AlphaAttribute = OutPointData->Metadata->FindOrCreateAttribute<double>(Params.AlphaAttribute, DefaultValue);
+					bSetMetadata |= (AlphaAttribute != nullptr);
+				}
+
+				if (Params.bComputeDistance)
+				{
+					DistanceAttribute = OutPointData->Metadata->FindOrCreateAttribute<double>(Params.DistanceAttribute, DefaultValue);
+					bSetMetadata |= (DistanceAttribute != nullptr);
+				}
 			}
 		}
 
@@ -632,6 +685,16 @@ namespace PCGSplineSamplerHelpers
 			{
 				LeaveTangentAttribute->SetValue(OutPoint.MetadataEntry, InResult.LeaveTangent);
 			}
+
+			if (AlphaAttribute)
+			{
+				AlphaAttribute->SetValue(OutPoint.MetadataEntry, InResult.Alpha);
+			}
+
+			if (DistanceAttribute)
+			{
+				DistanceAttribute->SetValue(OutPoint.MetadataEntry, InResult.Distance);
+			}
 		}
 
 		virtual void Sample(const FSamplerResult& InResult, UPCGPointData* OutPointData)
@@ -679,6 +742,8 @@ namespace PCGSplineSamplerHelpers
 		FPCGMetadataAttribute<int>* SubsegmentIndexAttribute = nullptr;
 		FPCGMetadataAttribute<FVector>* LeaveTangentAttribute = nullptr;
 		FPCGMetadataAttribute<FVector>* ArriveTangentAttribute = nullptr;
+		FPCGMetadataAttribute<double>* AlphaAttribute = nullptr;
+		FPCGMetadataAttribute<double>* DistanceAttribute = nullptr;
 	};
 
 	/** Samples in a volume surrounding the poly line. */
diff --git a/Engine/Plugins/PCG/Source/PCG/Public/Data/PCGPolyLineData.h b/Engine/Plugins/PCG/Source/PCG/Public/Data/PCGPolyLineData.h
--- a/Engine/Plugins/PCG/Source/PCG/Public/Data/PCGPolyLineData.h
+++ b/Engine/Plugins/PCG/Source/PCG/Public/Data/PCGPolyLineData.h
@@ -21,18 +21,41 @@ public:
 	virtual FBox GetBounds() const override;
 	//~End UPCGSpatialData interface
 
+	/** Get the world-space transform of the entire line. */
 	virtual FTransform GetTransform() const { return FTransform::Identity; }
+
+	/** Get the number of segments in this line. If the line is closed, this will be the same as the number of control points in the line. */
 	virtual int GetNumSegments() const PURE_VIRTUAL(UPCGPolyLineData::GetNumSegments, return 0;);
+
+	/** Get the length of a specific segment of the line. */
 	virtual FVector::FReal GetSegmentLength(int SegmentIndex) const PURE_VIRTUAL(UPCGPolyLineData::GetSegmentLength, return 0;);
+
+	/** Get the total length of the line. */
 	virtual FVector::FReal GetLength() const;
+
+	/** Get the location at a distance along the line. */
 	virtual FTransform GetTransformAtDistance(int SegmentIndex, FVector::FReal Distance, bool bWorldSpace = true, FBox* OutBounds = nullptr) const PURE_VIRTUAL(UPCGPolyLine::GetTransformAtDistance, return FTransform(););
+
+	/** Get the location at a distance along the line. */
+	virtual FVector GetLocationAtDistance(int SegmentIndex, FVector::FReal Distance, bool bWorldSpace = true) const { return GetTransformAtDistance(SegmentIndex, Distance, bWorldSpace).GetLocation(); }
+
+	/** Get the curvature at a distance along the line. */
 	virtual FVector::FReal GetCurvatureAtDistance(int SegmentIndex, FVector::FReal Distance) const { return 0; }
+
+	/**
+	 * Get a value [0,1] representing how far along the point is to the end of the line. Each segment on the line represents a same-size interval.
+	 * For example, if there are three segments, each segment will take up 0.333... of the interval.
+	 */
+	virtual float GetAlphaAtDistance(int SegmentIndex, FVector::FReal Distance) const;
+
+	/** Get the arrive and leave tangents for a control point via its segment index. */
 	virtual void GetTangentsAtSegmentStart(int SegmentIndex, FVector& OutArriveTangent, FVector& OutLeaveTangent) const;
-	virtual bool IsClosed() const { return false; }
 
-	virtual FVector GetLocationAtDistance(int SegmentIndex, FVector::FReal Distance, bool bWorldSpace = true) const { return GetTransformAtDistance(SegmentIndex, Distance, bWorldSpace).GetLocation(); }
+	/** True if the line is a closed loop. */
+	virtual bool IsClosed() const { return false; }
 };
 
+
 #if UE_ENABLE_INCLUDE_ORDER_DEPRECATED_IN_5_2
 #include "CoreMinimal.h"
 #endif
diff --git a/Engine/Plugins/PCG/Source/PCG/Public/Elements/PCGSplineSampler.h b/Engine/Plugins/PCG/Source/PCG/Public/Elements/PCGSplineSampler.h
--- a/Engine/Plugins/PCG/Source/PCG/Public/Elements/PCGSplineSampler.h
+++ b/Engine/Plugins/PCG/Source/PCG/Public/Elements/PCGSplineSampler.h
@@ -120,6 +120,23 @@ struct PCG_API FPCGSplineSamplerParams
 	UPROPERTY(BlueprintReadWrite, EditAnywhere, Category = Settings, meta = (EditCondition = "bComputeTangents"))
 	FName LeaveTangentAttribute = TEXT("LeaveTangent");
 
+	UPROPERTY(BlueprintReadWrite, EditAnywhere, Category = Settings, meta = (InlineEditConditionToggle))
+	bool bComputeAlpha = false;
+
+	/**
+	 * Attribute that will contain a value in [0,1] representing how far along the point is to the end of the line. Each segment on the line represents a same-size interval.
+	 * For example, if there are three segments, each segment will take up 0.333... of the interval.
+	 */
+	UPROPERTY(BlueprintReadWrite, EditAnywhere, Category = Settings, meta = (EditCondition = "bComputeAlpha"))
+	FName AlphaAttribute = TEXT("Alpha");
+
+	UPROPERTY(BlueprintReadWrite, EditAnywhere, Category = Settings, meta = (InlineEditConditionToggle))
+	bool bComputeDistance = false;
+
+	/** Attribute that will contain the distance along the spline at the sample point. */
+	UPROPERTY(BlueprintReadWrite, EditAnywhere, Category = Settings, meta = (EditCondition = "bComputeDistance"))
+	FName DistanceAttribute = TEXT("Distance");
+
 	/** If no Bounding Shape input is provided, the actor bounds are used to limit the sample generation domain.
 	* This option allows ignoring the actor bounds and generating over the entire spline. Use with caution as this
 	* may generate a lot of points.
