[PCG] Add Cast node, to cast attribute to another type

* Add unit tests
* Refactor CopyAttribute helper, to support casting.

#rnx

diff --git a/Engine/Plugins/PCG/Source/PCG/Private/Elements/Metadata/PCGMetadataElement.cpp b/Engine/Plugins/PCG/Source/PCG/Private/Elements/Metadata/PCGMetadataElement.cpp
--- a/Engine/Plugins/PCG/Source/PCG/Private/Elements/Metadata/PCGMetadataElement.cpp
+++ b/Engine/Plugins/PCG/Source/PCG/Private/Elements/Metadata/PCGMetadataElement.cpp
@@ -201,7 +201,15 @@ bool FPCGMetadataOperationElement::ExecuteInternal(FPCGContext* Context) const
 		}
 		else
 		{
-			bSuccess = PCGMetadataHelpers::CopyAttributes(SourceData, Settings->InputSource, SampledData, Settings->OutputTarget, /*bSameOrigin=*/!SourceAttributeSet, Context);
+			PCGMetadataHelpers::FPCGCopyAttributeParams Params{};
+			Params.SourceData = SourceData;
+			Params.TargetData = SampledData;
+			Params.InputSource = Settings->InputSource;
+			Params.OutputTarget = Settings->OutputTarget;
+			Params.OptionalContext = Context;
+			Params.bSameOrigin = !SourceAttributeSet;
+
+			bSuccess = PCGMetadataHelpers::CopyAttribute(Params);
 		}
 
 		if (bSuccess)
diff --git a/Engine/Plugins/PCG/Source/PCG/Private/Elements/PCGAttributeCast.cpp b/Engine/Plugins/PCG/Source/PCG/Private/Elements/PCGAttributeCast.cpp
new file mode 100644
--- /dev/null
+++ b/Engine/Plugins/PCG/Source/PCG/Private/Elements/PCGAttributeCast.cpp
@@ -0,0 +1,116 @@
+// Copyright Epic Games, Inc. All Rights Reserved.
+
+#include "Elements/PCGAttributeCast.h"
+
+#include "PCGContext.h"
+#include "PCGParamData.h"
+#include "Elements/Metadata/PCGMetadataElementCommon.h"
+#include "Metadata/PCGMetadataAttributeTraits.h"
+#include "Metadata/Accessors/IPCGAttributeAccessor.h"
+#include "Metadata/Accessors/PCGAttributeAccessorKeys.h"
+#include "Metadata/Accessors/PCGAttributeAccessorHelpers.h"
+
+#define LOCTEXT_NAMESPACE "PCGAttributeCastElement"
+
+TArray<FPCGPinProperties> UPCGAttributeCastSettings::OutputPinProperties() const
+{
+	TArray<FPCGPinProperties> PinProperties;
+	PinProperties.Emplace(PCGPinConstants::DefaultOutputLabel, EPCGDataType::Any);
+
+	return PinProperties;
+}
+
+FPCGElementPtr UPCGAttributeCastSettings::CreateElement() const
+{
+	return MakeShared<FPCGAttributeCastElement>();
+}
+
+#if WITH_EDITOR
+TArray<FPCGPreConfiguredSettingsInfo> UPCGAttributeCastSettings::GetPreconfiguredInfo() const
+{
+	return PCGMetadataElementCommon::FillPreconfiguredSettingsInfoFromEnum<EPCGMetadataTypes>({ EPCGMetadataTypes::Count, EPCGMetadataTypes::Unknown }, LOCTEXT("PreconfigureInfo", "Cast to "));
+}
+#endif
+
+FString UPCGAttributeCastSettings::GetAdditionalTitleInformation() const
+{
+#if WITH_EDITOR
+	const bool bIsTypeOverridden = IsPropertyOverriddenByPin(GET_MEMBER_NAME_CHECKED(UPCGAttributeCastSettings, OutputType));
+	const bool bSelectorOverridden = IsPropertyOverriddenByPin(GET_MEMBER_NAME_CHECKED(UPCGAttributeCastSettings, InputSource)) || IsPropertyOverriddenByPin(GET_MEMBER_NAME_CHECKED(UPCGAttributeCastSettings, OutputTarget));
+
+	if (bIsTypeOverridden)
+	{
+		return FString();
+	}
+	else
+	{
+		if (bSelectorOverridden)
+		{
+			return FText::Format(LOCTEXT("CastType", "To {0}"), PCG::Private::GetTypeNameText((uint16)OutputType)).ToString();
+		}
+	}
+#endif
+
+	return FText::Format(LOCTEXT("FullInfo", "{0} -> {1} as {2}"), InputSource.GetDisplayText(), OutputTarget.GetDisplayText(), PCG::Private::GetTypeNameText((uint16)OutputType)).ToString();
+}
+
+FName UPCGAttributeCastSettings::AdditionalTaskName() const
+{
+	return FName(*FText::Format(LOCTEXT("AdditionalTaskName", "Attribute Cast: {0}"), PCG::Private::GetTypeNameText(static_cast<uint16>(OutputType))).ToString());
+}
+
+void UPCGAttributeCastSettings::ApplyPreconfiguredSettings(const FPCGPreConfiguredSettingsInfo& PreconfiguredInfo)
+{
+	if (const UEnum* EnumPtr = StaticEnum<EPCGMetadataTypes>())
+	{
+		if (EnumPtr->IsValidEnumValue(PreconfiguredInfo.PreconfiguredIndex))
+		{
+			OutputType = EPCGMetadataTypes(PreconfiguredInfo.PreconfiguredIndex);
+		}
+	}
+}
+
+bool FPCGAttributeCastElement::ExecuteInternal(FPCGContext* Context) const
+{
+	check(Context);
+
+	const UPCGAttributeCastSettings* Settings = Context->GetInputSettings<UPCGAttributeCastSettings>();
+	check(Settings);
+
+	TArray<FPCGTaggedData> Inputs = Context->InputData.GetInputsByPin(PCGPinConstants::DefaultInputLabel);
+
+	for (int32 i = 0; i < Inputs.Num(); ++i)
+	{
+		const FPCGTaggedData& Input = Inputs[i];
+
+		// Only support types that supports metadata
+		if (!ensure(Input.Data) || !Input.Data->ConstMetadata())
+		{
+			PCGE_LOG(Error, GraphAndLog, FText::Format(LOCTEXT("InvalidData", "Only support Spatial and Attribute Set, Input {0} is neither of them. Skipped"), FText::AsNumber(i)));
+			continue;
+		}
+
+		FPCGTaggedData& Output = Context->OutputData.TaggedData.Add_GetRef(Input);
+
+		UPCGData* OutputData = Input.Data->DuplicateData();
+		check(OutputData);
+
+		PCGMetadataHelpers::FPCGCopyAttributeParams Params{};
+		Params.SourceData = Input.Data;
+		Params.TargetData = OutputData;
+		Params.InputSource = Settings->InputSource;
+		Params.OutputTarget = Settings->OutputTarget;
+		Params.OptionalContext = Context;
+		Params.bSameOrigin = true;
+		Params.OutputType = Settings->OutputType;
+
+		if (PCGMetadataHelpers::CopyAttribute(Params))
+		{
+			Output.Data = OutputData;
+		}
+	}
+
+	return true;
+}
+
+#undef LOCTEXT_NAMESPACE
\ No newline at end of file
diff --git a/Engine/Plugins/PCG/Source/PCG/Private/Elements/PCGCreateAttribute.cpp b/Engine/Plugins/PCG/Source/PCG/Private/Elements/PCGCreateAttribute.cpp
--- a/Engine/Plugins/PCG/Source/PCG/Private/Elements/PCGCreateAttribute.cpp
+++ b/Engine/Plugins/PCG/Source/PCG/Private/Elements/PCGCreateAttribute.cpp
@@ -333,7 +333,15 @@ bool FPCGAddAttributeElement::ExecuteInternal(FPCGContext* Context) const
 		}
 		else
 		{
-			bSuccess = PCGMetadataHelpers::CopyAttributes(FirstSourceParamData, Settings->InputSource, TargetData, Settings->OutputTarget, /*bSameOrigin=*/false, Context);
+			PCGMetadataHelpers::FPCGCopyAttributeParams Params{};
+			Params.SourceData = FirstSourceParamData;
+			Params.TargetData = TargetData;
+			Params.InputSource = Settings->InputSource;
+			Params.OutputTarget = Settings->OutputTarget;
+			Params.OptionalContext = Context;
+			Params.bSameOrigin = false;
+
+			bSuccess = PCGMetadataHelpers::CopyAttribute(Params);
 		}
 
 		if (bSuccess)
diff --git a/Engine/Plugins/PCG/Source/PCG/Private/Helpers/PCGMetadataHelpers.cpp b/Engine/Plugins/PCG/Source/PCG/Private/Helpers/PCGMetadataHelpers.cpp
--- a/Engine/Plugins/PCG/Source/PCG/Private/Helpers/PCGMetadataHelpers.cpp
+++ b/Engine/Plugins/PCG/Source/PCG/Private/Helpers/PCGMetadataHelpers.cpp
@@ -109,7 +109,8 @@ namespace PCGMetadataHelpers
 		return false;
 	}
 
-	bool CopyAttributes(UPCGData* TargetData, const UPCGData* SourceData, const TArray<TPair<FPCGAttributePropertyInputSelector, FPCGAttributePropertyOutputSelector>>& AttributeSelectors, bool bSameOrigin, FPCGContext* OptionalContext)
+
+	bool CopyAttributes(UPCGData* TargetData, const UPCGData* SourceData, const TArray<TTuple<FPCGAttributePropertyInputSelector, FPCGAttributePropertyOutputSelector, EPCGMetadataTypes>>& AttributeSelectorsWithOutputType, bool bSameOrigin, FPCGContext* OptionalContext)
 	{
 		check(TargetData && SourceData);
 		const UPCGMetadata* SourceMetadata = SourceData->ConstMetadata();
@@ -122,10 +123,11 @@ namespace PCGMetadataHelpers
 
 		bool bSuccess = false;
 
-		for (const auto& SelectorPair : AttributeSelectors)
+		for (const auto& SelectorTuple : AttributeSelectorsWithOutputType)
 		{
-			const FPCGAttributePropertyInputSelector& InputSource = SelectorPair.Key;
-			const FPCGAttributePropertyOutputSelector& OutputTarget = SelectorPair.Value;
+			const FPCGAttributePropertyInputSelector& InputSource = SelectorTuple.Get<0>();
+			const FPCGAttributePropertyOutputSelector& OutputTarget = SelectorTuple.Get<1>();
+			const EPCGMetadataTypes RequestedOutputType = SelectorTuple.Get<2>();
 
 			const FName LocalSourceAttribute = InputSource.GetName();
 			const FName LocalDestinationAttribute = OutputTarget.GetName();
@@ -142,8 +144,10 @@ namespace PCGMetadataHelpers
 			const bool bOutputHasAnyExtra = !OutputTarget.GetExtraNames().IsEmpty();
 			const bool bSourceIsAttribute = InputSource.GetSelection() == EPCGAttributePropertySelection::Attribute;
 			const bool bTargetIsAttribute = OutputTarget.GetSelection() == EPCGAttributePropertySelection::Attribute;
+			// Cast is only required if it is on an output attribute that has no extra (that we will create)
+			const bool bOutputTypeCast = bTargetIsAttribute && !bOutputHasAnyExtra && (RequestedOutputType != EPCGMetadataTypes::Unknown);
 
-			const bool bNeedAccessors = bIsMultiEntries || bInputHasAnyExtra || bOutputHasAnyExtra || !bSourceIsAttribute || !bTargetIsAttribute;
+			const bool bNeedAccessors = bIsMultiEntries || bInputHasAnyExtra || bOutputHasAnyExtra || !bSourceIsAttribute || !bTargetIsAttribute || bOutputTypeCast;
 
 			// If no accessor, copy over the attribute
 			if (!bNeedAccessors)
@@ -183,17 +187,39 @@ namespace PCGMetadataHelpers
 					continue;
 				}
 
+				const uint16 OutputType = bOutputTypeCast ? static_cast<uint16>(RequestedOutputType) : InputAccessor->GetUnderlyingType();
+
+				if (bOutputTypeCast && InputAccessor->GetUnderlyingType() == OutputType && bSameOrigin && LocalSourceAttribute == LocalDestinationAttribute)
+				{
+					// Nothing to do if we try to cast an attribute on itself with the same type
+					continue;
+				}
+
+				// If we have a cast, make sure it is valid
+				if (bOutputTypeCast && !PCG::Private::IsBroadcastableOrConstructible(InputAccessor->GetUnderlyingType(), OutputType))
+				{
+					PCGLog::LogWarningOnGraph(FText::Format(LOCTEXT("CastInvalid", "Cannot convert InputAttribute '{0}' of type {1} into {2}"), InputSource.GetDisplayText(), PCG::Private::GetTypeNameText(InputAccessor->GetUnderlyingType()), PCG::Private::GetTypeNameText(OutputType)), OptionalContext);
+					continue;
+				}
+
 				// If the target is an attribute, only create a new one if the attribute we don't have any extra.
 				// If it has any extra, it will try to write to it.
 				if (!bOutputHasAnyExtra && bTargetIsAttribute)
 				{
-					auto CreateAttribute = [TargetMetadata, LocalDestinationAttribute](auto Dummy)
+					auto CreateAttribute = [TargetMetadata, LocalDestinationAttribute, &InputAccessor](auto Dummy) -> bool
 					{
 						using AttributeType = decltype(Dummy);
-						return PCGMetadataElementCommon::ClearOrCreateAttribute<AttributeType>(TargetMetadata, LocalDestinationAttribute) != nullptr;
+						AttributeType DefaultValue{};
+						if (!InputAccessor->Get(DefaultValue, FPCGAttributeAccessorKeysEntries(PCGInvalidEntryKey), EPCGAttributeAccessorFlags::AllowBroadcast | EPCGAttributeAccessorFlags::AllowConstructible))
+						{
+							// It's OK to fail getting the default value, if for example the input accessor is a property. In that case, just fallback on 0.
+							DefaultValue = PCG::Private::MetadataTraits<AttributeType>::ZeroValue();
+						}
+
+						return PCGMetadataElementCommon::ClearOrCreateAttribute<AttributeType>(TargetMetadata, LocalDestinationAttribute, std::move(DefaultValue)) != nullptr;
 					};
 
-					if (!PCGMetadataAttribute::CallbackWithRightType(InputAccessor->GetUnderlyingType(), CreateAttribute))
+					if (!PCGMetadataAttribute::CallbackWithRightType(OutputType, CreateAttribute))
 					{
 						PCGLog::LogWarningOnGraph(FText::Format(LOCTEXT("FailedToCreateNewAttribute", "Failed to create new attribute '{0}'"), FText::FromName(LocalDestinationAttribute)), OptionalContext);
 						continue;
@@ -215,10 +241,10 @@ namespace PCGMetadataHelpers
 					continue;
 				}
 
-				// Final verification, if we can put the value of input into output
-				if (!PCG::Private::IsBroadcastableOrConstructible(InputAccessor->GetUnderlyingType(), OutputAccessor->GetUnderlyingType()))
+				// Final verification (if not already done), if we can put the value of input into output
+				if (!bOutputTypeCast && !PCG::Private::IsBroadcastableOrConstructible(OutputType, OutputAccessor->GetUnderlyingType()))
 				{
-					PCGLog::LogErrorOnGraph(FText::Format(LOCTEXT("CannotConvertTypes", "Cannot convert input type {0} into output type {1}"), PCG::Private::GetTypeNameText(InputAccessor->GetUnderlyingType()), PCG::Private::GetTypeNameText(OutputAccessor->GetUnderlyingType())), OptionalContext);
+					PCGLog::LogErrorOnGraph(FText::Format(LOCTEXT("CannotConvertTypes", "Cannot convert input type {0} into output type {1}"), PCG::Private::GetTypeNameText(OutputType), PCG::Private::GetTypeNameText(OutputAccessor->GetUnderlyingType())), OptionalContext);
 					continue;
 				}
 
@@ -244,19 +270,19 @@ namespace PCGMetadataHelpers
 		return bSuccess;
 	}
 
-	bool CopyAttributes(const UPCGData* SourceData, const FPCGAttributePropertyInputSelector& _InputSource, UPCGData* TargetData, const FPCGAttributePropertyOutputSelector& _OutputTarget, bool bSameOrigin, FPCGContext* OptionalContext)
+	bool CopyAttribute(const FPCGCopyAttributeParams& InParams)
 	{
-		if (!TargetData || !SourceData)
+		if (!InParams.TargetData || !InParams.SourceData)
 		{
 			return false;
 		}
 
-		TArray<TPair<FPCGAttributePropertyInputSelector, FPCGAttributePropertyOutputSelector>> AttributeSelectors;
-		FPCGAttributePropertyInputSelector InputSource = _InputSource.CopyAndFixLast(SourceData);
-		FPCGAttributePropertyOutputSelector OutputTarget = _OutputTarget.CopyAndFixSource(&InputSource, SourceData);
+		TArray<TTuple<FPCGAttributePropertyInputSelector, FPCGAttributePropertyOutputSelector, EPCGMetadataTypes>> AttributeSelectors;
+		FPCGAttributePropertyInputSelector InputSource = InParams.InputSource.CopyAndFixLast(InParams.SourceData);
+		FPCGAttributePropertyOutputSelector OutputTarget = InParams.OutputTarget.CopyAndFixSource(&InputSource, InParams.SourceData);
 
-		AttributeSelectors.Emplace(MoveTemp(InputSource), MoveTemp(OutputTarget));
-		return CopyAttributes(TargetData, SourceData, AttributeSelectors, bSameOrigin, OptionalContext);
+		AttributeSelectors.Emplace(MoveTemp(InputSource), MoveTemp(OutputTarget), InParams.OutputType);
+		return CopyAttributes(InParams.TargetData, InParams.SourceData, AttributeSelectors, InParams.bSameOrigin, InParams.OptionalContext);
 	}
 
 	bool CopyAllAttributes(const UPCGData* SourceData, UPCGData* TargetData, FPCGContext* OptionalContext)
@@ -272,16 +298,17 @@ namespace PCGMetadataHelpers
 			return false;
 		}
 
-		TArray<TPair<FPCGAttributePropertyInputSelector, FPCGAttributePropertyOutputSelector>> AttributeSelectors;
+		TArray<TTuple<FPCGAttributePropertyInputSelector, FPCGAttributePropertyOutputSelector, EPCGMetadataTypes>> AttributeSelectors;
 		TArray<FName> AttributeNames;
 		TArray<EPCGMetadataTypes> AttributeTypes;
 		SourceMetadata->GetAttributes(AttributeNames, AttributeTypes);
 
 		for (const FName& AttributeName : AttributeNames)
 		{
-			TPair<FPCGAttributePropertyInputSelector, FPCGAttributePropertyOutputSelector>& Selectors = AttributeSelectors.Emplace_GetRef();
-			Selectors.Key.SetAttributeName(AttributeName);
-			Selectors.Value.SetAttributeName(AttributeName);
+			TTuple<FPCGAttributePropertyInputSelector, FPCGAttributePropertyOutputSelector, EPCGMetadataTypes>& Selectors = AttributeSelectors.Emplace_GetRef();
+			Selectors.Get<0>().SetAttributeName(AttributeName);
+			Selectors.Get<1>().SetAttributeName(AttributeName);
+			Selectors.Get<2>() = EPCGMetadataTypes::Unknown;
 		}
 
 		return CopyAttributes(TargetData, SourceData, AttributeSelectors, /*bSameOrigin=*/false, OptionalContext);
diff --git a/Engine/Plugins/PCG/Source/PCG/Private/Tests/Elements/PCGAttributeCastTest.cpp b/Engine/Plugins/PCG/Source/PCG/Private/Tests/Elements/PCGAttributeCastTest.cpp
new file mode 100644
--- /dev/null
+++ b/Engine/Plugins/PCG/Source/PCG/Private/Tests/Elements/PCGAttributeCastTest.cpp
@@ -0,0 +1,257 @@
+// Copyright Epic Games, Inc. All Rights Reserved.
+
+#include "Tests/PCGTestsCommon.h"
+
+#include "PCGComponent.h"
+#include "PCGContext.h"
+#include "PCGParamData.h"
+#include "Data/PCGPointData.h"
+
+#include "Elements/PCGAttributeCast.h"
+
+class FPCGAttributeCastTestBase : public FPCGTestBaseClass
+{
+public:
+	using FPCGTestBaseClass::FPCGTestBaseClass;
+
+protected:
+	template <typename InType, typename OutType, typename PCGDataType>
+	bool RunTestInternal(TArrayView<const InType> InData, TArrayView<const OutType> ExpectedOutData)
+	{
+		PCGTestsCommon::FTestData TestData;
+		UPCGAttributeCastSettings* Settings = PCGTestsCommon::GenerateSettings<UPCGAttributeCastSettings>(TestData);
+		check(Settings);
+
+		check(InData.Num() == ExpectedOutData.Num() && InData.Num() > 1);
+
+		const FName InputAttributeName = TEXT("InAttr");
+		const FName OutputAttributeName = TEXT("OutAttr");
+		Settings->InputSource.SetAttributeName(InputAttributeName);
+		Settings->OutputTarget.SetAttributeName(OutputAttributeName);
+		Settings->OutputType = static_cast<EPCGMetadataTypes>(PCG::Private::MetadataTypes<OutType>::Id);
+
+		PCGDataType* InputData = NewObject<PCGDataType>();
+		FPCGMetadataAttribute<InType>* Attribute = InputData->Metadata->template CreateAttribute<InType>(InputAttributeName, InData[0], true, false);
+		for (int32 i = 1; i < InData.Num(); ++i)
+		{
+			PCGMetadataEntryKey EntryKey = InputData->Metadata->AddEntry();
+			Attribute->SetValue(EntryKey, InData[i]);
+
+			if constexpr (std::is_same_v<PCGDataType, UPCGPointData>)
+			{
+				FPCGPoint& Point = InputData->GetMutablePoints().Emplace_GetRef();
+				Point.MetadataEntry = EntryKey;
+			}
+		}
+
+		FPCGTaggedData& InputTaggedData = TestData.InputData.TaggedData.Emplace_GetRef();
+		InputTaggedData.Data = InputData;
+		InputTaggedData.Pin = PCGPinConstants::DefaultInputLabel;
+
+		FPCGElementPtr TestElement = TestData.Settings->GetElement();
+		TUniquePtr<FPCGContext> Context = TestData.InitializeTestContext();
+
+		while (!TestElement->Execute(Context.Get())) {}
+
+		UTEST_EQUAL("There is one output", Context->OutputData.TaggedData.Num(), 1);
+
+		const PCGDataType* OutputData = Cast<const PCGDataType>(Context->OutputData.TaggedData[0].Data);
+
+		UTEST_NOT_NULL("Output is an of the same type as input", OutputData);
+
+		if constexpr (std::is_same_v<PCGDataType, UPCGPointData>)
+		{
+			UTEST_EQUAL("Output has the same number of points", InputData->GetPoints().Num(), OutputData->GetPoints().Num());
+		}
+
+		const FPCGMetadataAttribute<OutType>* OutputAttribute = OutputData->Metadata->template GetConstTypedAttribute<OutType>(OutputAttributeName);
+
+		UTEST_NOT_NULL("Output metadata has the expected attribute", OutputAttribute);
+
+		UTEST_EQUAL("Default value was casted correctly", OutputAttribute->GetValue(PCGDefaultValueKey), ExpectedOutData[0]);
+		for (int32 i = 1; i < ExpectedOutData.Num(); ++i)
+		{
+			PCGMetadataEntryKey EntryKey(i - 1);
+			if constexpr (std::is_same_v<PCGDataType, UPCGPointData>)
+			{
+				EntryKey = OutputData->GetPoints()[i-1].MetadataEntry;
+			}
+
+			UTEST_EQUAL(*FString::Printf(TEXT("Entry %d was casted correctly"), i - 1), OutputAttribute->GetValueFromItemKey(EntryKey), ExpectedOutData[i]);
+		}
+
+		return true;
+	}
+};
+
+IMPLEMENT_CUSTOM_SIMPLE_AUTOMATION_TEST(FPCGAttributeCast_AttributeSetIntToDouble, FPCGAttributeCastTestBase, "Plugins.PCG.AttributeCast.AttributeSetIntToDouble", PCGTestsCommon::TestFlags)
+IMPLEMENT_CUSTOM_SIMPLE_AUTOMATION_TEST(FPCGAttributeCast_AttributeSetDoubleToInt, FPCGAttributeCastTestBase, "Plugins.PCG.AttributeCast.AttributeSetDoubleToInt", PCGTestsCommon::TestFlags)
+IMPLEMENT_CUSTOM_SIMPLE_AUTOMATION_TEST(FPCGAttributeCast_AttributeSetDoubleToBool, FPCGAttributeCastTestBase, "Plugins.PCG.AttributeCast.AttributeSetDoubleToBool", PCGTestsCommon::TestFlags)
+IMPLEMENT_CUSTOM_SIMPLE_AUTOMATION_TEST(FPCGAttributeCast_AttributeSetBoolToDouble, FPCGAttributeCastTestBase, "Plugins.PCG.AttributeCast.AttributeSetBoolToDouble", PCGTestsCommon::TestFlags)
+IMPLEMENT_CUSTOM_SIMPLE_AUTOMATION_TEST(FPCGAttributeCast_AttributeSetIntToString, FPCGAttributeCastTestBase, "Plugins.PCG.AttributeCast.AttributeSetIntToString", PCGTestsCommon::TestFlags)
+IMPLEMENT_CUSTOM_SIMPLE_AUTOMATION_TEST(FPCGAttributeCast_AttributeSetQuatToRotator, FPCGAttributeCastTestBase, "Plugins.PCG.AttributeCast.AttributeSetQuatToRotator", PCGTestsCommon::TestFlags)
+
+IMPLEMENT_CUSTOM_SIMPLE_AUTOMATION_TEST(FPCGAttributeCast_PointsIntToDouble, FPCGAttributeCastTestBase, "Plugins.PCG.AttributeCast.PointsIntToDouble", PCGTestsCommon::TestFlags)
+
+IMPLEMENT_CUSTOM_SIMPLE_AUTOMATION_TEST(FPCGAttributeCast_PointsIntToDoubleToDensity, FPCGTestBaseClass, "Plugins.PCG.AttributeCast.PointsIntToDoubleToDensity", PCGTestsCommon::TestFlags)
+IMPLEMENT_CUSTOM_SIMPLE_AUTOMATION_TEST(FPCGAttributeCast_SameType, FPCGTestBaseClass, "Plugins.PCG.AttributeCast.SameType", PCGTestsCommon::TestFlags)
+IMPLEMENT_CUSTOM_SIMPLE_AUTOMATION_TEST(FPCGAttributeCast_InvalidCast, FPCGTestBaseClass, "Plugins.PCG.AttributeCast.InvalidCast", PCGTestsCommon::TestFlags)
+
+bool FPCGAttributeCast_AttributeSetIntToDouble::RunTest(const FString& Parameters)
+{
+	return RunTestInternal<int, double, UPCGParamData>({ 5, 1, 2 }, { 5.0, 1.0, 2.0 });
+}
+
+bool FPCGAttributeCast_AttributeSetDoubleToInt::RunTest(const FString& Parameters)
+{
+	return RunTestInternal<double, int64, UPCGParamData>({ 5.1, 1.6, 2.3 }, { 5, 1, 2 });
+}
+
+bool FPCGAttributeCast_AttributeSetDoubleToBool::RunTest(const FString& Parameters)
+{
+	return RunTestInternal<double, bool, UPCGParamData>({ 5.1, 1.6, 2.3, 0.0, 0.0, 2.7 }, { true, true, true, false, false, true });
+}
+
+bool FPCGAttributeCast_AttributeSetBoolToDouble::RunTest(const FString& Parameters)
+{
+	return RunTestInternal<bool, double, UPCGParamData>({ true, true, true, false, false, true }, { 1.0, 1.0, 1.0, 0.0, 0.0, 1.0 });
+}
+
+bool FPCGAttributeCast_AttributeSetIntToString::RunTest(const FString& Parameters)
+{
+	return RunTestInternal<int, FString, UPCGParamData>({ 0, 8, 7, 4, -5 }, { TEXT("0"), TEXT("8"), TEXT("7"), TEXT("4"), TEXT("-5") });
+}
+
+bool FPCGAttributeCast_AttributeSetQuatToRotator::RunTest(const FString& Parameters)
+{
+	TArray<FQuat, TInlineAllocator<4>> QuatValues = {
+		FQuat::Identity,
+		FQuat(-0.70710678118654757, 0, 0, 0.70710678118654768), // (90, 0, 0)
+		FQuat(0, -0.38268343236508978, 0, 0.92387953251128663), // (0, 45, 0)
+		FQuat(0, 0, 0.25881904510252079, 0.96592582628906842)  // (0, 0, 30)
+	};
+
+	TArray<FRotator, TInlineAllocator<4>> RotatorValues = {
+		FRotator::ZeroRotator,
+		FRotator::MakeFromEuler({90, 0, 0}),
+		FRotator::MakeFromEuler({0, 45, 0}),
+		FRotator::MakeFromEuler({0, 0, 30})
+	};
+
+	return RunTestInternal<FQuat, FRotator, UPCGParamData>(QuatValues, RotatorValues);
+}
+
+bool FPCGAttributeCast_PointsIntToDouble::RunTest(const FString& Parameters)
+{
+	return RunTestInternal<int, double, UPCGPointData>({ 5, 1, 2 }, { 5.0, 1.0, 2.0 });
+}
+
+bool FPCGAttributeCast_PointsIntToDoubleToDensity::RunTest(const FString& Parameters)
+{
+	PCGTestsCommon::FTestData TestData;
+	UPCGAttributeCastSettings* Settings = PCGTestsCommon::GenerateSettings<UPCGAttributeCastSettings>(TestData);
+	check(Settings);
+
+	const FName InputAttributeName = TEXT("Int");
+	Settings->InputSource.SetAttributeName(InputAttributeName);
+	Settings->OutputTarget.SetPointProperty(EPCGPointProperties::Density);
+	Settings->OutputType = EPCGMetadataTypes::Double;
+
+	UPCGPointData* InputPoints = NewObject<UPCGPointData>();
+	FPCGMetadataAttribute<int32>* Attribute = InputPoints->Metadata->CreateAttribute<int32>(InputAttributeName, 5, true, false);
+
+	TArray<FPCGPoint>& Points = InputPoints->GetMutablePoints();
+	FPCGPoint& Point1 = Points.Emplace_GetRef();
+	FPCGPoint& Point2 = Points.Emplace_GetRef();
+	InputPoints->Metadata->InitializeOnSet(Point1.MetadataEntry);
+	InputPoints->Metadata->InitializeOnSet(Point2.MetadataEntry);
+	Attribute->SetValue(Point1.MetadataEntry, 1);
+	Attribute->SetValue(Point2.MetadataEntry, 2);
+
+	FPCGTaggedData& InputTaggedData = TestData.InputData.TaggedData.Emplace_GetRef();
+	InputTaggedData.Data = InputPoints;
+	InputTaggedData.Pin = PCGPinConstants::DefaultInputLabel;
+
+	FPCGElementPtr TestElement = TestData.Settings->GetElement();
+	TUniquePtr<FPCGContext> Context = TestData.InitializeTestContext();
+
+	while (!TestElement->Execute(Context.Get())) {}
+
+	UTEST_EQUAL("There is one output", Context->OutputData.TaggedData.Num(), 1);
+
+	const UPCGPointData* OutputPoints = Cast<const UPCGPointData>(Context->OutputData.TaggedData[0].Data);
+
+	UTEST_NOT_NULL("Output is a point data", OutputPoints);
+	UTEST_EQUAL("Output has 2 points", OutputPoints->GetPoints().Num(), 2);
+
+	UTEST_EQUAL("First point has the right value", OutputPoints->GetPoint(0).Density, 1.0f);
+	UTEST_EQUAL("Second point has the right value", OutputPoints->GetPoint(1).Density, 2.0f);
+
+	return true;
+}
+
+bool FPCGAttributeCast_SameType::RunTest(const FString& Parameters)
+{
+	PCGTestsCommon::FTestData TestData;
+	UPCGAttributeCastSettings* Settings = PCGTestsCommon::GenerateSettings<UPCGAttributeCastSettings>(TestData);
+	check(Settings);
+
+	const FName InputAttributeName = TEXT("Int");
+	Settings->InputSource.SetAttributeName(InputAttributeName);
+	Settings->OutputTarget.SetAttributeName(InputAttributeName);
+	Settings->OutputType = EPCGMetadataTypes::Integer32;
+
+	UPCGParamData* InputParam = NewObject<UPCGParamData>();
+	FPCGMetadataAttribute<int32>* Attribute = InputParam->Metadata->CreateAttribute<int32>(InputAttributeName, 5, true, false);
+	Attribute->SetValue(InputParam->Metadata->AddEntry(), 1);
+	Attribute->SetValue(InputParam->Metadata->AddEntry(), 2);
+
+	FPCGTaggedData& InputTaggedData = TestData.InputData.TaggedData.Emplace_GetRef();
+	InputTaggedData.Data = InputParam;
+	InputTaggedData.Pin = PCGPinConstants::DefaultInputLabel;
+
+	FPCGElementPtr TestElement = TestData.Settings->GetElement();
+	TUniquePtr<FPCGContext> Context = TestData.InitializeTestContext();
+
+	while (!TestElement->Execute(Context.Get())) {}
+
+	UTEST_EQUAL("There is one output", Context->OutputData.TaggedData.Num(), 1);
+	UTEST_EQUAL("Output is same as input", Context->OutputData.TaggedData[0].Data.Get(), static_cast<const UPCGData*>(InputParam));
+
+	return true;
+}
+
+bool FPCGAttributeCast_InvalidCast::RunTest(const FString& Parameters)
+{
+	PCGTestsCommon::FTestData TestData;
+	UPCGAttributeCastSettings* Settings = PCGTestsCommon::GenerateSettings<UPCGAttributeCastSettings>(TestData);
+	check(Settings);
+
+	const FName InputAttributeName = TEXT("String");
+	Settings->InputSource.SetAttributeName(InputAttributeName);
+	Settings->OutputTarget.SetAttributeName(InputAttributeName);
+	Settings->OutputType = EPCGMetadataTypes::Integer32;
+
+	UPCGParamData* InputParam = NewObject<UPCGParamData>();
+	FPCGMetadataAttribute<FString>* Attribute = InputParam->Metadata->CreateAttribute<FString>(InputAttributeName, TEXT("Hello"), true, false);
+	Attribute->SetValue(InputParam->Metadata->AddEntry(), TEXT("Hey"));
+	Attribute->SetValue(InputParam->Metadata->AddEntry(), TEXT("Hi"));
+
+	FPCGTaggedData& InputTaggedData = TestData.InputData.TaggedData.Emplace_GetRef();
+	InputTaggedData.Data = InputParam;
+	InputTaggedData.Pin = PCGPinConstants::DefaultInputLabel;
+
+	FPCGElementPtr TestElement = TestData.Settings->GetElement();
+	TUniquePtr<FPCGContext> Context = TestData.InitializeTestContext();
+
+	AddExpectedError(TEXT("Cannot convert InputAttribute 'String' of type String into Integer32"));
+
+	while (!TestElement->Execute(Context.Get())) {}
+
+	UTEST_EQUAL("There is one output", Context->OutputData.TaggedData.Num(), 1);
+
+	// Failed cast will just forward.
+	UTEST_EQUAL("Output is same as input", Context->OutputData.TaggedData[0].Data.Get(), static_cast<const UPCGData*>(InputParam));
+
+	return true;
+}
\ No newline at end of file
diff --git a/Engine/Plugins/PCG/Source/PCG/Public/Elements/PCGAttributeCast.h b/Engine/Plugins/PCG/Source/PCG/Public/Elements/PCGAttributeCast.h
new file mode 100644
--- /dev/null
+++ b/Engine/Plugins/PCG/Source/PCG/Public/Elements/PCGAttributeCast.h
@@ -0,0 +1,54 @@
+// Copyright Epic Games, Inc. All Rights Reserved.
+
+#pragma once
+
+#include "PCGElement.h"
+#include "PCGSettings.h"
+#include "Metadata/PCGAttributePropertySelector.h"
+#include "Metadata/PCGMetadataAttributeTraits.h"
+
+#include "PCGAttributeCast.generated.h"
+
+/**
+* Cast an attribute to another type. Support broadcastable cast (like double -> FVector) and constructible cast (like double -> float)
+*/
+UCLASS(BlueprintType, ClassGroup = (Procedural))
+class PCG_API UPCGAttributeCastSettings : public UPCGSettings
+{
+	GENERATED_BODY()
+
+public:
+	UPROPERTY(BlueprintReadWrite, EditAnywhere, Category = Settings, meta = (PCG_Overridable))
+	FPCGAttributePropertyInputSelector InputSource;
+
+	UPROPERTY(BlueprintReadWrite, EditAnywhere, Category = Settings, meta = (PCG_Overridable))
+	EPCGMetadataTypes OutputType = EPCGMetadataTypes::Float;
+
+	UPROPERTY(BlueprintReadWrite, EditAnywhere, Category = Settings, meta = (PCG_Overridable, PCG_DiscardPropertySelection, PCG_DiscardExtraSelection))
+	FPCGAttributePropertyOutputSelector OutputTarget;
+
+	//~Begin UPCGSettings interface
+public:
+#if WITH_EDITOR
+	virtual FName GetDefaultNodeName() const override { return FName("AttributeCast"); }
+	virtual FText GetDefaultNodeTitle() const override { return NSLOCTEXT("PCGAttributeCastElement", "NodeTitle", "Attribute Cast"); }
+	virtual EPCGSettingsType GetType() const override { return EPCGSettingsType::Metadata; }
+	virtual bool HasDynamicPins() const override { return true; }
+	virtual TArray<FPCGPreConfiguredSettingsInfo> GetPreconfiguredInfo() const override;
+	virtual bool OnlyExposePreconfiguredSettings() const override { return true; }
+#endif
+	virtual FString GetAdditionalTitleInformation() const override;
+	virtual FName AdditionalTaskName() const override;
+	virtual void ApplyPreconfiguredSettings(const FPCGPreConfiguredSettingsInfo& PreconfiguredInfo) override;
+	virtual FPCGElementPtr CreateElement() const override;
+
+protected:
+	virtual TArray<FPCGPinProperties> OutputPinProperties() const;
+	//~End UPCGSettings interface
+};
+
+class PCG_API FPCGAttributeCastElement : public IPCGElement
+{
+protected:
+	virtual bool ExecuteInternal(FPCGContext* Context) const override;
+};
diff --git a/Engine/Plugins/PCG/Source/PCG/Public/Helpers/PCGMetadataHelpers.h b/Engine/Plugins/PCG/Source/PCG/Public/Helpers/PCGMetadataHelpers.h
--- a/Engine/Plugins/PCG/Source/PCG/Public/Helpers/PCGMetadataHelpers.h
+++ b/Engine/Plugins/PCG/Source/PCG/Public/Helpers/PCGMetadataHelpers.h
@@ -2,13 +2,14 @@
 
 #pragma once
 
+#include "Metadata/PCGAttributePropertySelector.h"
+#include "Metadata/PCGMetadataAttributeTraits.h"
+
 #include "HAL/Platform.h"
 
 class FPCGMetadataAttributeBase;
 class UPCGData;
 class UPCGMetadata;
-struct FPCGAttributePropertyInputSelector;
-struct FPCGAttributePropertyOutputSelector;
 struct FPCGContext;
 struct FSoftObjectPath;
 template <typename FuncType> class TFunction;
@@ -28,10 +29,32 @@ namespace PCGMetadataHelpers
 	/** Create a lambda that will construct a soft object path from an underlying attribute of type FSoftObjectPath or FString. Returns true if successful. */
 	[[nodiscard]] PCG_API bool CreateObjectOrClassPathGetter(const FPCGMetadataAttributeBase* InAttributeBase, TFunction<void(int64, FSoftObjectPath&)>& OutGetter);
 
-	/** Copy the attribute coming from Source Data into Target Data.
-	* Can set bSameOrigin to true if SourceData and TargetData have the same origin (if TargetData was initialized from SourceData). 
-	*/
-	PCG_API bool CopyAttributes(const UPCGData* SourceData, const FPCGAttributePropertyInputSelector& InputSource, UPCGData* TargetData, const FPCGAttributePropertyOutputSelector& OutputTarget, bool bSameOrigin, FPCGContext* OptionalContext = nullptr);
+	struct PCG_API FPCGCopyAttributeParams
+	{
+		/** Source data where the attribute is coming from */
+		const UPCGData* SourceData = nullptr;
+
+		/** Target data to write to */
+		UPCGData* TargetData = nullptr;
+
+		/** Selector for the attribute in SourceData */
+		FPCGAttributePropertyInputSelector InputSource;
+
+		/** Selector for the attribute in TargetData */
+		FPCGAttributePropertyOutputSelector OutputTarget;
+
+		/** Optional context for logging */
+		FPCGContext* OptionalContext = nullptr;
+
+		/** Will convert the output attribute to this type if not Unknown */
+		EPCGMetadataTypes OutputType = EPCGMetadataTypes::Unknown;
+
+		/** If SourceData and TargetData have the same origin (if TargetData was initialized from SourceData) */
+		bool bSameOrigin = false;
+	};
+
+	/** Copy the attribute coming from Source Data into Target Data. */
+	PCG_API bool CopyAttribute(const FPCGCopyAttributeParams& InParams);
 
 	/** Copy all the attributes coming from Source Data into Target Data. */
 	PCG_API bool CopyAllAttributes(const UPCGData* SourceData, UPCGData* TargetData, FPCGContext* OptionalContext = nullptr);
