Exposed enum value from attribute noise, should have been a blueprint type.

#rnx

diff --git a/Engine/Plugins/PCG/Source/PCG/Public/Elements/PCGAttributeNoise.h b/Engine/Plugins/PCG/Source/PCG/Public/Elements/PCGAttributeNoise.h
--- a/Engine/Plugins/PCG/Source/PCG/Public/Elements/PCGAttributeNoise.h
+++ b/Engine/Plugins/PCG/Source/PCG/Public/Elements/PCGAttributeNoise.h
@@ -13,7 +13,7 @@
 
 #include "PCGAttributeNoise.generated.h"
 
-UENUM()
+UENUM(BlueprintType)
 enum class EPCGAttributeNoiseMode : uint8
 {
 	Set,
