[PCG] Static Mesh Spawner now supports "Property Overrides"

+ ISMs will be multi-partitioned based on property overrides
+ The override will be applied, similar to Spawn Actor
* PCGActorPropertyOverride is now PCGObjectPropertyOverride to support UObject or UStruct

diff --git a/Engine/Plugins/PCG/Config/BasePCG.ini b/Engine/Plugins/PCG/Config/BasePCG.ini
--- a/Engine/Plugins/PCG/Config/BasePCG.ini
+++ b/Engine/Plugins/PCG/Config/BasePCG.ini
@@ -75,6 +75,7 @@
 +ClassRedirects=(OldName="PCGSortPointsSettings",NewName="/Script/PCG.PCGSortAttributesSettings")
 
 +StructRedirects=(OldName="PCGPointFilterThresholdSettings",NewName="/Script/PCG.PCGAttributeFilterThresholdSettings")
++StructRedirects=(OldName="PCGActorPropertyOverrideDescription",NewName="/Script/PCG.PCGObjectPropertyOverrideDescription")
 
 +EnumRedirects=(OldName="EPCGDensityNoiseMode",NewName="EPCGAttributeNoiseMode")
 +EnumRedirects=(OldName="EPCGPointFilterOperator",NewName="EPCGAttributeFilterOperator")
diff --git a/Engine/Plugins/PCG/Source/PCG/Private/Elements/PCGApplyOnActor.cpp b/Engine/Plugins/PCG/Source/PCG/Private/Elements/PCGApplyOnActor.cpp
--- a/Engine/Plugins/PCG/Source/PCG/Private/Elements/PCGApplyOnActor.cpp
+++ b/Engine/Plugins/PCG/Source/PCG/Private/Elements/PCGApplyOnActor.cpp
@@ -21,7 +21,7 @@ TArray<FPCGPinProperties> UPCGApplyOnActorSettings::InputPinProperties() const
 {
 	TArray<FPCGPinProperties> PinProperties;
 	PinProperties.Emplace(PCGPinConstants::DefaultDependencyOnlyLabel, EPCGDataType::Any, /*bAllowMultipleConnections=*/true, /*bAllowMultipleData=*/true, PCGApplyOnActorConstants::DependencyTooltip);
-	PinProperties.Add(PCGActorPropertyOverrideHelpers::CreateActorPropertiesOverridePin(PCGApplyOnActorConstants::ActorPropertyOverridesLabel, PCGApplyOnActorConstants::ActorPropertyOverridesTooltip));
+	PinProperties.Add(PCGObjectPropertyOverrideHelpers::CreateObjectPropertiesOverridePin(PCGApplyOnActorConstants::ActorPropertyOverridesLabel, PCGApplyOnActorConstants::ActorPropertyOverridesTooltip));
 	return PinProperties;
 }
 
@@ -46,7 +46,7 @@ bool FPCGApplyOnActorElement::ExecuteInternal(FPCGContext* Context) const
 		return true;
 	}
 
-	PCGActorPropertyOverrideHelpers::ApplyOverridesFromParams(Settings->PropertyOverrideDescriptions, TargetActor, PCGApplyOnActorConstants::ActorPropertyOverridesLabel, Context);
+	PCGObjectPropertyOverrideHelpers::ApplyOverridesFromParams(Settings->PropertyOverrideDescriptions, TargetActor, PCGApplyOnActorConstants::ActorPropertyOverridesLabel, Context);
 
 	for (UFunction* Function : PCGHelpers::FindUserFunctions(TargetActor->GetClass(), Settings->PostProcessFunctionNames, { UPCGFunctionPrototypes::GetPrototypeWithNoParams() }, Context))
 	{
diff --git a/Engine/Plugins/PCG/Source/PCG/Private/Elements/PCGCreateTargetActor.cpp b/Engine/Plugins/PCG/Source/PCG/Private/Elements/PCGCreateTargetActor.cpp
--- a/Engine/Plugins/PCG/Source/PCG/Private/Elements/PCGCreateTargetActor.cpp
+++ b/Engine/Plugins/PCG/Source/PCG/Private/Elements/PCGCreateTargetActor.cpp
@@ -45,7 +45,7 @@ FPCGElementPtr UPCGCreateTargetActor::CreateElement() const
 TArray<FPCGPinProperties> UPCGCreateTargetActor::InputPinProperties() const
 {
 	TArray<FPCGPinProperties> PinProperties;
-	PinProperties.Add(PCGActorPropertyOverrideHelpers::CreateActorPropertiesOverridePin(PCGCreateTargetActorConstants::ActorPropertyOverridesLabel, PCGCreateTargetActorConstants::ActorPropertyOverridesTooltip));
+	PinProperties.Add(PCGObjectPropertyOverrideHelpers::CreateObjectPropertiesOverridePin(PCGCreateTargetActorConstants::ActorPropertyOverridesLabel, PCGCreateTargetActorConstants::ActorPropertyOverridesTooltip));
 	return PinProperties;
 }
 
@@ -293,7 +293,7 @@ bool FPCGCreateTargetActorElement::ExecuteInternal(FPCGContext* Context) const
 	GeneratedActor->Tags.Add(PCGHelpers::DefaultPCGActorTag);
 
 	// Apply property overrides to the GeneratedActor
-	PCGActorPropertyOverrideHelpers::ApplyOverridesFromParams(Settings->PropertyOverrideDescriptions, GeneratedActor, PCGCreateTargetActorConstants::ActorPropertyOverridesLabel, Context);
+	PCGObjectPropertyOverrideHelpers::ApplyOverridesFromParams(Settings->PropertyOverrideDescriptions, GeneratedActor, PCGCreateTargetActorConstants::ActorPropertyOverridesLabel, Context);
 
 	for (UFunction* Function : PCGHelpers::FindUserFunctions(GeneratedActor->GetClass(), Settings->PostProcessFunctionNames, { UPCGFunctionPrototypes::GetPrototypeWithNoParams() }, Context))
 	{
diff --git a/Engine/Plugins/PCG/Source/PCG/Private/Elements/PCGSpawnActor.cpp b/Engine/Plugins/PCG/Source/PCG/Private/Elements/PCGSpawnActor.cpp
--- a/Engine/Plugins/PCG/Source/PCG/Private/Elements/PCGSpawnActor.cpp
+++ b/Engine/Plugins/PCG/Source/PCG/Private/Elements/PCGSpawnActor.cpp
@@ -754,7 +754,7 @@ void FPCGSpawnActorElement::SpawnActors(FPCGSubgraphContext* Context, AActor* Ta
 
 	check(TemplateActor);
 
-	FPCGActorOverrides ActorOverrides(TemplateActor);
+	FPCGObjectOverrides ActorOverrides(TemplateActor);
 	ActorOverrides.Initialize(Settings->SpawnedActorPropertyOverrideDescriptions, TemplateActor, PointData, Context);
 
 	FActorSpawnParameters SpawnParams;
diff --git a/Engine/Plugins/PCG/Source/PCG/Private/Elements/PCGStaticMeshSpawner.cpp b/Engine/Plugins/PCG/Source/PCG/Private/Elements/PCGStaticMeshSpawner.cpp
--- a/Engine/Plugins/PCG/Source/PCG/Private/Elements/PCGStaticMeshSpawner.cpp
+++ b/Engine/Plugins/PCG/Source/PCG/Private/Elements/PCGStaticMeshSpawner.cpp
@@ -12,6 +12,7 @@
 #include "Helpers/PCGHelpers.h"
 #include "InstanceDataPackers/PCGInstanceDataPackerBase.h"
 #include "MeshSelectors/PCGMeshSelectorBase.h"
+#include "MeshSelectors/PCGMeshSelectorByAttribute.h"
 #include "MeshSelectors/PCGMeshSelectorWeighted.h"
 
 #include "Components/InstancedStaticMeshComponent.h"
@@ -519,6 +520,20 @@ void UPCGStaticMeshSpawnerSettings::PostEditChangeProperty(FPropertyChangedEvent
 
 	Super::PostEditChangeProperty(PropertyChangedEvent);
 }
+
+bool UPCGStaticMeshSpawnerSettings::CanEditChange(const FProperty* InProperty) const
+{
+	// TODO: In place temporarily, until the other two modes are supported
+	if (InProperty->GetFName() == GET_MEMBER_NAME_CHECKED(UPCGStaticMeshSpawnerSettings, StaticMeshComponentPropertyOverrides))
+	{
+		if (!MeshSelectorType->IsChildOf(UPCGMeshSelectorByAttribute::StaticClass()))
+		{
+			return false;
+		}
+	}
+
+	return Super::CanEditChange(InProperty);
+}
 #endif
 
 void UPCGStaticMeshSpawnerSettings::SetMeshSelectorType(TSubclassOf<UPCGMeshSelectorBase> InMeshSelectorType) 
diff --git a/Engine/Plugins/PCG/Source/PCG/Private/MeshSelectors/PCGMeshSelectorByAttribute.cpp b/Engine/Plugins/PCG/Source/PCG/Private/MeshSelectors/PCGMeshSelectorByAttribute.cpp
--- a/Engine/Plugins/PCG/Source/PCG/Private/MeshSelectors/PCGMeshSelectorByAttribute.cpp
+++ b/Engine/Plugins/PCG/Source/PCG/Private/MeshSelectors/PCGMeshSelectorByAttribute.cpp
@@ -3,10 +3,14 @@
 #include "MeshSelectors/PCGMeshSelectorByAttribute.h"
 
 #include "Data/PCGPointData.h"
+#include "Data/PCGSpatialData.h"
 #include "Elements/PCGStaticMeshSpawner.h"
 #include "Elements/PCGStaticMeshSpawnerContext.h"
-#include "Data/PCGSpatialData.h"
+#include "Elements/Metadata/PCGMetadataElementCommon.h"
 #include "MeshSelectors/PCGMeshSelectorBase.h"
+#include "Metadata/PCGObjectPropertyOverride.h"
+#include "Metadata/PCGMetadataPartitionCommon.h"
+#include "Metadata/Accessors/PCGAttributeAccessorHelpers.h"
 
 #include "Engine/StaticMesh.h"
 
@@ -22,13 +26,15 @@ namespace PCGMeshSelectorAttribute
 		const FSoftISMComponentDescriptor& TemplateDescriptor,
 		TSoftObjectPtr<UStaticMesh> Mesh,
 		const TArray<TSoftObjectPtr<UMaterialInterface>>& MaterialOverrides,
-		bool bReverseCulling)
+		bool bReverseCulling,
+		const int AttributePartitionIndex = INDEX_NONE)
 	{
 		for (FPCGMeshInstanceList& InstanceList : InstanceLists)
 		{
 			if (InstanceList.Descriptor.StaticMesh == Mesh &&
 				InstanceList.Descriptor.bReverseCulling == bReverseCulling &&
-				InstanceList.Descriptor.OverrideMaterials == MaterialOverrides)
+				InstanceList.Descriptor.OverrideMaterials == MaterialOverrides &&
+				InstanceList.AttributePartitionIndex == AttributePartitionIndex)
 			{
 				return InstanceList;
 			}
@@ -38,6 +44,7 @@ namespace PCGMeshSelectorAttribute
 		NewInstanceList.Descriptor.StaticMesh = Mesh;
 		NewInstanceList.Descriptor.OverrideMaterials = MaterialOverrides;
 		NewInstanceList.Descriptor.bReverseCulling = bReverseCulling;
+		NewInstanceList.AttributePartitionIndex = AttributePartitionIndex;
 
 		return NewInstanceList;
 	}
@@ -96,7 +103,7 @@ void UPCGMeshSelectorByAttribute::PostLoad()
 
 bool UPCGMeshSelectorByAttribute::SelectInstances(
 	FPCGStaticMeshSpawnerContext& Context,
-	const UPCGStaticMeshSpawnerSettings* Settings, 
+	const UPCGStaticMeshSpawnerSettings* Settings,
 	const UPCGPointData* InPointData,
 	TArray<FPCGMeshInstanceList>& OutMeshInstances,
 	UPCGPointData* OutPointData) const
@@ -136,7 +143,7 @@ bool UPCGMeshSelectorByAttribute::SelectInstances(
 	}
 	else if (PCG::Private::IsOfTypes<FString>(AttributeBase->GetTypeId()))
 	{
-		MeshPathGetter = [AttributeBase, &Context](PCGMetadataValueKey InValueKey, FSoftObjectPath& OutMeshPath)
+		MeshPathGetter = [AttributeBase](PCGMetadataValueKey InValueKey, FSoftObjectPath& OutMeshPath)
 		{
 			const FPCGMetadataAttribute<FString>* Attribute = static_cast<const FPCGMetadataAttribute<FString>*>(AttributeBase);
 
@@ -181,12 +188,65 @@ bool UPCGMeshSelectorByAttribute::SelectInstances(
 	constexpr int32 TimeSlicingCheckFrequency = 1024;
 	TMap<PCGMetadataValueKey, TSoftObjectPtr<UStaticMesh>>& ValueKeyToMesh = Context.ValueKeyToMesh;
 
+	// No partition needed if there are no overrides
+	if (Settings->StaticMeshComponentPropertyOverrides.IsEmpty())
+	{
+		Context.bPartitionDone = true;
+	}
+
+	if (!Context.bPartitionDone)
+	{
+		// Validate all the selectors are actual FSoftISMComponentDescriptor properties
+		TArray<FPCGAttributePropertySelector> ValidSelectorOverrides;
+		for (const FPCGObjectPropertyOverrideDescription& PropertyOverride : Settings->StaticMeshComponentPropertyOverrides)
+		{
+			if (FSoftISMComponentDescriptor::StaticStruct()->FindPropertyByName(FName(PropertyOverride.PropertyTarget)))
+			{
+				ValidSelectorOverrides.Emplace_GetRef() = FPCGAttributePropertySelector::CreateFromOtherSelector<FPCGAttributePropertySelector>(PropertyOverride.InputSource);
+			}
+			else
+			{
+				PCGE_LOG_C(Error, GraphAndLog, &Context, FText::Format(LOCTEXT("OverriddenPropertyNotFound", "Property '{0}' not a valid property with an ISM Descriptor. It will be ignored."), FText::FromString(PropertyOverride.PropertyTarget)));
+			}
+		}
+
+		// If there are valid overrides, partition the points on those attributes so that an instance can be created for each
+		if (!ValidSelectorOverrides.IsEmpty())
+		{
+			Context.AttributeOverridePartition = PCGMetadataPartitionCommon::AttributeGenericPartition(InPointData, ValidSelectorOverrides, &Context);
+		}
+
+		// Set the descriptors to match the partition count. Uninitialized, because it will be copied from the template below
+		Context.OverriddenDescriptors.Reserve(Context.AttributeOverridePartition.Num());
+		for (int I = 0; I < Context.AttributeOverridePartition.Num(); ++I)
+		{
+			FSoftISMComponentDescriptor& Descriptor = Context.OverriddenDescriptors.Add_GetRef(TemplateDescriptor);
+
+			// Use the Object Override to map the user's input selector and property to the descriptor
+			FPCGObjectOverrides Overrides(&Descriptor);
+			Overrides.Initialize(Settings->StaticMeshComponentPropertyOverrides, &Descriptor, InPointData, &Context);
+
+			// Since they are already partitioned and identical, we can just use the value on the first point
+			check(!Context.AttributeOverridePartition[I].IsEmpty());
+			int32 AnyPointIndexOnThisPartition = Context.AttributeOverridePartition[I][0];
+			Overrides.Apply(AnyPointIndexOnThisPartition);
+		}
+
+		// Given partitioning is expensive, check if we're out of time for this frame
+		Context.bPartitionDone = true;
+		if (Context.ShouldStop())
+		{
+			return false;
+		}
+	}
+
 	const TArray<FPCGPoint>& Points = InPointData->GetPoints();
 
 	while (CurrentPointIndex < Points.Num())
 	{
-		const FPCGPoint& Point = Points[CurrentPointIndex++];
-		
+		int32 ThisPointIndex = CurrentPointIndex++;
+		const FPCGPoint& Point = Points[ThisPointIndex];
+
 		const PCGMetadataValueKey ValueKey = AttributeBase->GetValueKey(Point.MetadataEntry);
 		TSoftObjectPtr<UStaticMesh>* NewMesh = ValueKeyToMesh.Find(ValueKey);
 		TSoftObjectPtr<UStaticMesh> Mesh = nullptr;
@@ -223,20 +283,35 @@ bool UPCGMeshSelectorByAttribute::SelectInstances(
 			continue;
 		}
 
+		// TODO: Revisit this when attribute partitioning is returned in a more optimized form
+		// The partition index is used to assign the point to the correct partition's instance
+		int32 CurrentPointPartitionIndex = INDEX_NONE;
+		for (int PartitionIndex = 0; PartitionIndex < Context.AttributeOverridePartition.Num(); ++PartitionIndex)
+		{
+			TArray<int32>& Partition = Context.AttributeOverridePartition[PartitionIndex];
+			if (Partition.Contains(ThisPointIndex))
+			{
+				CurrentPointPartitionIndex = PartitionIndex;
+				break;
+			}
+		}
+
+		// If the point wasn't found in a partition (likely due to no partitions or an invalid property target), just default to the template
+		const FSoftISMComponentDescriptor& CurrentPartitionDescriptor = CurrentPointPartitionIndex != INDEX_NONE ? Context.OverriddenDescriptors[CurrentPointPartitionIndex] : TemplateDescriptor;
 		const bool bReverseTransform = (Point.Transform.GetDeterminant() < 0);
-		FPCGMeshInstanceList& InstanceList = PCGMeshSelectorAttribute::GetInstanceList(OutMeshInstances, TemplateDescriptor, Mesh, MaterialOverrideHelper.GetMaterialOverrides(Point.MetadataEntry), bReverseTransform);
+
+		FPCGMeshInstanceList& InstanceList = PCGMeshSelectorAttribute::GetInstanceList(OutMeshInstances, CurrentPartitionDescriptor, Mesh, MaterialOverrideHelper.GetMaterialOverrides(Point.MetadataEntry), bReverseTransform, CurrentPointPartitionIndex);
 		InstanceList.Instances.Emplace(Point.Transform);
 		InstanceList.InstancesMetadataEntry.Emplace(Point.MetadataEntry);
 
 		if (OutPointData && Settings->bApplyMeshBoundsToPoints)
 		{
 			TArray<int32>& PointIndices = Context.MeshToOutPoints.FindOrAdd(Mesh).FindOrAdd(OutPointData);
-			// CurrentPointIndex - 1, because CurrentPointIndex is incremented at the beginning of the loop
-			PointIndices.Emplace(CurrentPointIndex - 1);
+			PointIndices.Emplace(ThisPointIndex);
 		}
 
 		// Check if we should stop here and continue in a subsequent call
-		if(CurrentPointIndex - LastCheckpointIndex >= TimeSlicingCheckFrequency)
+		if (CurrentPointIndex - LastCheckpointIndex >= TimeSlicingCheckFrequency)
 		{
 			if (Context.ShouldStop())
 			{
diff --git a/Engine/Plugins/PCG/Source/PCG/Private/MeshSelectors/PCGMeshSelectorWeighted.cpp b/Engine/Plugins/PCG/Source/PCG/Private/MeshSelectors/PCGMeshSelectorWeighted.cpp
--- a/Engine/Plugins/PCG/Source/PCG/Private/MeshSelectors/PCGMeshSelectorWeighted.cpp
+++ b/Engine/Plugins/PCG/Source/PCG/Private/MeshSelectors/PCGMeshSelectorWeighted.cpp
@@ -122,6 +122,12 @@ bool UPCGMeshSelectorWeighted::SelectInstances(
 	TArray<TArray<FPCGMeshInstanceList>>& MeshInstances = Context.WeightedMeshInstances;
 	TArray<int>& CumulativeWeights = Context.CumulativeWeights;
 
+	// TODO: Remove this log once the other selection modes are available
+	if (!Settings->StaticMeshComponentPropertyOverrides.IsEmpty())
+	{
+		PCGE_LOG_C(Log, LogOnly, &Context, LOCTEXT("AttributeToPropertyOverrideUnavailable", "Attribute to Property Overrides are only currently available with the 'By Attribute' Selector"));
+	}
+
 	// Setup
 	if (Context.CurrentPointIndex == 0)
 	{
diff --git a/Engine/Plugins/PCG/Source/PCG/Private/MeshSelectors/PCGMeshSelectorWeightedByCategory.cpp b/Engine/Plugins/PCG/Source/PCG/Private/MeshSelectors/PCGMeshSelectorWeightedByCategory.cpp
--- a/Engine/Plugins/PCG/Source/PCG/Private/MeshSelectors/PCGMeshSelectorWeightedByCategory.cpp
+++ b/Engine/Plugins/PCG/Source/PCG/Private/MeshSelectors/PCGMeshSelectorWeightedByCategory.cpp
@@ -69,6 +69,12 @@ bool UPCGMeshSelectorWeightedByCategory::SelectInstances(
 	// unmarked points will fallback to the MeshEntries associated with the DefaultValueKey
 	PCGMetadataValueKey DefaultValueKey = PCGDefaultValueKey;
 
+	// TODO: Remove this log once the other selection modes are available
+	if (!Settings->StaticMeshComponentPropertyOverrides.IsEmpty())
+	{
+		PCGE_LOG_C(Log, LogOnly, &Context, LOCTEXT("AttributeToPropertyOverrideUnavailable", "Attribute to Property Overrides are only currently available with the 'By Attribute' Selector"));
+	}
+
 	// Setup
 	if (Context.CurrentPointIndex == 0)
 	{
diff --git a/Engine/Plugins/PCG/Source/PCG/Private/Metadata/PCGActorPropertyOverride.cpp b/Engine/Plugins/PCG/Source/PCG/Private/Metadata/PCGActorPropertyOverride.cpp
deleted file mode 100644
diff --git a/Engine/Plugins/PCG/Source/PCG/Private/Metadata/PCGObjectPropertyOverride.cpp b/Engine/Plugins/PCG/Source/PCG/Private/Metadata/PCGObjectPropertyOverride.cpp
new file mode 100644
--- /dev/null
+++ b/Engine/Plugins/PCG/Source/PCG/Private/Metadata/PCGObjectPropertyOverride.cpp
@@ -0,0 +1,132 @@
+// Copyright Epic Games, Inc. All Rights Reserved.
+
+#include "Metadata/PCGObjectPropertyOverride.h"
+
+#include "PCGContext.h"
+#include "PCGData.h"
+#include "PCGModule.h"
+#include "PCGParamData.h"
+#include "Metadata/Accessors/PCGAttributeAccessorHelpers.h"
+
+#include "GameFramework/Actor.h"
+
+#define LOCTEXT_NAMESPACE "PCGObjectPropertyOverride"
+
+namespace PCGObjectPropertyOverrideHelpers
+{
+	FPCGPinProperties CreateObjectPropertiesOverridePin(FName Label, const FText& Tooltip)
+	{
+		FPCGPinProperties ObjectOverridePinProperties(Label, EPCGDataType::Param, /*bAllowMultipleConnections=*/false, /*bAllowMultipleData=*/false, Tooltip);
+		ObjectOverridePinProperties.bAdvancedPin = true;
+		return ObjectOverridePinProperties;
+	}
+
+	void ApplyOverridesFromParams(const TArray<FPCGObjectPropertyOverrideDescription>& InObjectPropertyOverrideDescriptions, AActor* TargetActor, FName OverridesPinLabel, FPCGContext* Context)
+	{
+		if (!Context)
+		{
+			return;
+		}
+
+		const TArray<FPCGTaggedData> OverrideInputs = Context->InputData.GetInputsByPin(OverridesPinLabel);
+
+		if (OverrideInputs.Num() == 0)
+		{
+			return;
+		}
+		else if (OverrideInputs.Num() > 1)
+		{
+			PCGLog::LogWarningOnGraph(FText::Format(LOCTEXT("MoreThanOneData", "More than one data was found on pin '{0}'. Only using the first one."), FText::FromName(OverridesPinLabel)), Context);
+		}
+
+		const UPCGParamData* ParamData = Cast<const UPCGParamData>(OverrideInputs[0].Data);
+
+		if (!ParamData)
+		{
+			PCGLog::LogErrorOnGraph(LOCTEXT("InvalidActorOverrideData", "Invalid input data type for Actor Property Overrides pin, must be of type Param."), Context);
+			return;
+		}
+
+		FPCGObjectOverrides ActorOverrides(TargetActor);
+		ActorOverrides.Initialize(InObjectPropertyOverrideDescriptions, TargetActor, ParamData, Context);
+		if (!ActorOverrides.Apply(/*InputKeyIndex=*/0)) // Use the First Entry of the param data for override (similar to what is done in Parameter Overrides in FPCGContext)
+		{
+			PCGLog::LogErrorOnGraph(FText::Format(LOCTEXT("ApplyOverrideFailed", "Failed to apply property overrides to actor '%s' from an attribute set."), FText::FromName(TargetActor->GetClass()->GetFName())), Context);
+		}
+	}
+}
+
+void FPCGObjectSingleOverride::Initialize(const FPCGAttributePropertySelector& InputSelector, const FString& OutputProperty, const UStruct* TemplateClass, const UPCGData* SourceData, FPCGContext* Context)
+{
+	InputKeys = PCGAttributeAccessorHelpers::CreateConstKeys(SourceData, InputSelector);
+	ObjectOverrideInputAccessor = PCGAttributeAccessorHelpers::CreateConstAccessor(SourceData, InputSelector);
+
+	if (!ObjectOverrideInputAccessor.IsValid())
+	{
+		PCGLog::LogWarningOnGraph(FText::Format(NSLOCTEXT("PCGObjectPropertyOverride", "OverrideInputInvalid", "ObjectOverride for input '{0}' is invalid or unsupported. Check the attribute or property selection."), InputSelector.GetDisplayText()), Context);
+		return;
+	}
+
+	const FPCGAttributePropertySelector OutputSelector = FPCGAttributePropertySelector::CreateSelectorFromString(OutputProperty);
+	// TODO: Move implementation into a new helper: PCGAttributeAccessorHelpers::CreatePropertyAccessor(const FPCGAttributePropertySelector& Selector, UStruct* Class)
+	const TArray<FString>& ExtraNames = OutputSelector.GetExtraNames();
+	if (ExtraNames.IsEmpty())
+	{
+		ObjectOverrideOutputAccessor = PCGAttributeAccessorHelpers::CreatePropertyAccessor(FName(OutputProperty), TemplateClass);
+	}
+	else
+	{
+		TArray<FName> PropertyNames;
+		PropertyNames.Reserve(ExtraNames.Num() + 1);
+		PropertyNames.Add(OutputSelector.GetAttributeName());
+		for (const FString& Name : ExtraNames)
+		{
+			PropertyNames.Add(FName(Name));
+		}
+
+		ObjectOverrideOutputAccessor = PCGAttributeAccessorHelpers::CreatePropertyChainAccessor(PropertyNames, TemplateClass);
+	}
+
+	if (!ObjectOverrideOutputAccessor.IsValid())
+	{
+		PCGLog::LogWarningOnGraph(FText::Format(NSLOCTEXT("PCGObjectPropertyOverride", "OverrideOutputInvalid", "ObjectOverride for object property '{0}' is invalid or unsupported. Check the attribute or property selection."), FText::FromString(OutputProperty)), Context);
+		return;
+	}
+
+	if (!PCG::Private::IsBroadcastableOrConstructible(ObjectOverrideInputAccessor->GetUnderlyingType(), ObjectOverrideOutputAccessor->GetUnderlyingType()))
+	{
+		PCGLog::LogWarningOnGraph(
+			FText::Format(
+				NSLOCTEXT("PCGObjectPropertyOverride", "TypesIncompatible", "ObjectOverride cannot set input '{0}' to output '{1}'. Cannot convert type '{2}' to type '{3}'. Will be skipped."),
+				InputSelector.GetDisplayText(),
+				FText::FromString(OutputProperty),
+				PCG::Private::GetTypeNameText(ObjectOverrideInputAccessor->GetUnderlyingType()),
+				PCG::Private::GetTypeNameText(ObjectOverrideOutputAccessor->GetUnderlyingType())),
+			Context);
+
+		ObjectOverrideInputAccessor.Reset();
+		ObjectOverrideOutputAccessor.Reset();
+		return;
+	}
+
+	auto CreateGetterSetter = [this](auto Dummy)
+	{
+		using Type = decltype(Dummy);
+
+		ObjectOverrideFunction = &FPCGObjectSingleOverride::ApplyImpl<Type>;
+	};
+
+	PCGMetadataAttribute::CallbackWithRightType(ObjectOverrideOutputAccessor->GetUnderlyingType(), CreateGetterSetter);
+}
+
+bool FPCGObjectSingleOverride::IsValid() const
+{
+	return InputKeys.IsValid() && ObjectOverrideInputAccessor.IsValid() && ObjectOverrideOutputAccessor.IsValid() && ObjectOverrideFunction;
+}
+
+bool FPCGObjectSingleOverride::Apply(int32 InputKeyIndex, IPCGAttributeAccessorKeys& OutputKey)
+{
+	return Invoke(ObjectOverrideFunction, this, InputKeyIndex, OutputKey);
+}
+
+#undef LOCTEXT_NAMESPACE
diff --git a/Engine/Plugins/PCG/Source/PCG/Public/Elements/PCGApplyOnActor.h b/Engine/Plugins/PCG/Source/PCG/Public/Elements/PCGApplyOnActor.h
--- a/Engine/Plugins/PCG/Source/PCG/Public/Elements/PCGApplyOnActor.h
+++ b/Engine/Plugins/PCG/Source/PCG/Public/Elements/PCGApplyOnActor.h
@@ -3,7 +3,7 @@
 #pragma once
 
 #include "PCGSettings.h"
-#include "Metadata/PCGActorPropertyOverride.h"
+#include "Metadata/PCGObjectPropertyOverride.h"
 
 #include "PCGApplyOnActor.generated.h"
 
@@ -38,7 +38,7 @@ public:
 
 	/** Override the default property values on the target actor. Applied before post-process functions. */
 	UPROPERTY(EditAnywhere, BlueprintReadWrite, Category = Settings)
-	TArray<FPCGActorPropertyOverrideDescription> PropertyOverrideDescriptions;
+	TArray<FPCGObjectPropertyOverrideDescription> PropertyOverrideDescriptions;
 
 	/** Specify a list of functions to be called on the target actor. Functions need to be parameter-less and with "CallInEditor" flag enabled. */
 	UPROPERTY(BlueprintReadWrite, EditAnywhere, Category = Settings)
diff --git a/Engine/Plugins/PCG/Source/PCG/Public/Elements/PCGCreateTargetActor.h b/Engine/Plugins/PCG/Source/PCG/Public/Elements/PCGCreateTargetActor.h
--- a/Engine/Plugins/PCG/Source/PCG/Public/Elements/PCGCreateTargetActor.h
+++ b/Engine/Plugins/PCG/Source/PCG/Public/Elements/PCGCreateTargetActor.h
@@ -3,7 +3,7 @@
 #pragma once
 
 #include "PCGSettings.h"
-#include "Metadata/PCGActorPropertyOverride.h"
+#include "Metadata/PCGObjectPropertyOverride.h"
 
 #include "PCGCreateTargetActor.generated.h"
 
@@ -74,7 +74,7 @@ public:
 
 	/** Override the default property values on the created target actor. Applied before post-process functions. */
 	UPROPERTY(EditAnywhere, BlueprintReadWrite, Category = Settings)
-	TArray<FPCGActorPropertyOverrideDescription> PropertyOverrideDescriptions;
+	TArray<FPCGObjectPropertyOverrideDescription> PropertyOverrideDescriptions;
 
 	/** Specify a list of functions to be called on the target actor after creation. Functions need to be parameter-less and with "CallInEditor" flag enabled. */
 	UPROPERTY(BlueprintReadWrite, EditAnywhere, Category = Settings)
diff --git a/Engine/Plugins/PCG/Source/PCG/Public/Elements/PCGSpawnActor.h b/Engine/Plugins/PCG/Source/PCG/Public/Elements/PCGSpawnActor.h
--- a/Engine/Plugins/PCG/Source/PCG/Public/Elements/PCGSpawnActor.h
+++ b/Engine/Plugins/PCG/Source/PCG/Public/Elements/PCGSpawnActor.h
@@ -3,7 +3,7 @@
 #pragma once
 
 #include "PCGSubgraph.h"
-#include "Metadata/PCGActorPropertyOverride.h"
+#include "Metadata/PCGObjectPropertyOverride.h"
 
 #include "PCGSpawnActor.generated.h"
 
@@ -74,7 +74,7 @@ public:
 	TObjectPtr<AActor> TemplateActor;
 
 	UPROPERTY(EditAnywhere, BlueprintReadWrite, Category = Settings, meta = (EditCondition = "Option != EPCGSpawnActorOption::CollapseActors", EditConditionHides))
-	TArray<FPCGActorPropertyOverrideDescription> SpawnedActorPropertyOverrideDescriptions;
+	TArray<FPCGObjectPropertyOverrideDescription> SpawnedActorPropertyOverrideDescriptions;
 
 	UPROPERTY(meta = (PCG_Overridable))
 	TSoftObjectPtr<AActor> RootActor;
diff --git a/Engine/Plugins/PCG/Source/PCG/Public/Elements/PCGStaticMeshSpawner.h b/Engine/Plugins/PCG/Source/PCG/Public/Elements/PCGStaticMeshSpawner.h
--- a/Engine/Plugins/PCG/Source/PCG/Public/Elements/PCGStaticMeshSpawner.h
+++ b/Engine/Plugins/PCG/Source/PCG/Public/Elements/PCGStaticMeshSpawner.h
@@ -4,6 +4,8 @@
 
 #include "PCGSettings.h"
 
+#include "Metadata/PCGObjectPropertyOverride.h"
+
 #include "Engine/CollisionProfile.h"
 
 #include "PCGStaticMeshSpawner.generated.h"
@@ -64,6 +66,7 @@ public:
 	virtual void PostLoad() override;
 #if WITH_EDITOR
 	virtual void PostEditChangeProperty(FPropertyChangedEvent& PropertyChangedEvent) override;
+	virtual bool CanEditChange(const FProperty* InProperty) const override;
 	// ~End UObject interface
 #endif
 
@@ -88,6 +91,12 @@ public:
 	UPROPERTY(BlueprintReadOnly, VisibleAnywhere, Instanced, Category = InstanceDataPacker)
 	TObjectPtr<UPCGInstanceDataPackerBase> InstanceDataPackerParameters = nullptr;
 
+	/** Map an attribute directly to an ISM Descriptor property, the value of which will be overriden when generated.
+	 * Note: Currently only enabled using SelectByAttribute mesh selection.
+	 */
+	UPROPERTY(BlueprintReadWrite, EditAnywhere, Category = Settings)
+	TArray<FPCGObjectPropertyOverrideDescription> StaticMeshComponentPropertyOverrides;
+
 	/** Attribute name to store mesh SoftObjectPaths inside if the output pin is connected. Note: Will overwrite existing data if the attribute name already exists. */
 	UPROPERTY(BlueprintReadWrite, EditAnywhere, Category = Settings)
 	FName OutAttributeName = NAME_None;
@@ -124,7 +133,7 @@ public:
 protected:
 	virtual FPCGContext* CreateContext() override;
 	virtual bool PrepareDataInternal(FPCGContext* Context) const override;
-	virtual bool ExecuteInternal(FPCGContext* Context) const override;	
+	virtual bool ExecuteInternal(FPCGContext* Context) const override;
 	void SpawnStaticMeshInstances(FPCGStaticMeshSpawnerContext* Context, const FPCGMeshInstanceList& InstanceList, AActor* TargetActor, const FPCGPackedCustomData& PackedCustomData) const;
 };
 
diff --git a/Engine/Plugins/PCG/Source/PCG/Public/Elements/PCGStaticMeshSpawnerContext.h b/Engine/Plugins/PCG/Source/PCG/Public/Elements/PCGStaticMeshSpawnerContext.h
--- a/Engine/Plugins/PCG/Source/PCG/Public/Elements/PCGStaticMeshSpawnerContext.h
+++ b/Engine/Plugins/PCG/Source/PCG/Public/Elements/PCGStaticMeshSpawnerContext.h
@@ -46,6 +46,7 @@ struct FPCGStaticMeshSpawnerContext : public FPCGContext, public IPCGAsyncLoadin
 	// Per-input context variables
 	bool bCurrentInputSetup = false;
 	bool bSelectionDone = false;
+	bool bPartitionDone = false;
 
 	const UPCGPointData* CurrentPointData = nullptr;
 	UPCGPointData* CurrentOutputPointData = nullptr;
@@ -68,5 +69,9 @@ struct FPCGStaticMeshSpawnerContext : public FPCGContext, public IPCGAsyncLoadin
 	// Used in the weighted by category selector
 	TMap<PCGMetadataValueKey, FPCGInstancesAndWeights> CategoryEntryToInstancesAndWeights;
 
+	// Used for mesh property overrides
+	TArray<TArray<int32>> AttributeOverridePartition;
+	TArray<FSoftISMComponentDescriptor> OverriddenDescriptors;
+
 	void ResetInputIterationData();
 };
\ No newline at end of file
diff --git a/Engine/Plugins/PCG/Source/PCG/Public/MeshSelectors/PCGMeshSelectorBase.h b/Engine/Plugins/PCG/Source/PCG/Public/MeshSelectors/PCGMeshSelectorBase.h
--- a/Engine/Plugins/PCG/Source/PCG/Public/MeshSelectors/PCGMeshSelectorBase.h
+++ b/Engine/Plugins/PCG/Source/PCG/Public/MeshSelectors/PCGMeshSelectorBase.h
@@ -28,6 +28,7 @@ struct FPCGMeshInstanceList
 
 	explicit FPCGMeshInstanceList(const FSoftISMComponentDescriptor& InDescriptor)
 		: Descriptor(InDescriptor)
+		, AttributePartitionIndex(INDEX_NONE)
 	{}
 
 	UPROPERTY(EditAnywhere, Category = Settings)
@@ -38,6 +39,9 @@ struct FPCGMeshInstanceList
 
 	UPROPERTY(BlueprintReadWrite, EditAnywhere, Category = Settings)
 	TArray<int64> InstancesMetadataEntry;
+
+	/** Tracks which partition the instance list belongs to. */
+	int64 AttributePartitionIndex;
 };
 
 UENUM()
diff --git a/Engine/Plugins/PCG/Source/PCG/Public/Metadata/PCGActorPropertyOverride.h b/Engine/Plugins/PCG/Source/PCG/Public/Metadata/PCGActorPropertyOverride.h
deleted file mode 100644
diff --git a/Engine/Plugins/PCG/Source/PCG/Public/Metadata/PCGObjectPropertyOverride.h b/Engine/Plugins/PCG/Source/PCG/Public/Metadata/PCGObjectPropertyOverride.h
new file mode 100644
--- /dev/null
+++ b/Engine/Plugins/PCG/Source/PCG/Public/Metadata/PCGObjectPropertyOverride.h
@@ -0,0 +1,196 @@
+// Copyright Epic Games, Inc. All Rights Reserved.
+
+#pragma once
+
+#include "PCGData.h"
+#include "PCGPin.h"
+#include "Accessors/PCGAttributeAccessorHelpers.h"
+#include "Metadata/PCGAttributePropertySelector.h"
+#include "Metadata/PCGMetadata.h"
+#include "Metadata/Accessors/IPCGAttributeAccessor.h"
+#include "Metadata/Accessors/PCGAttributeAccessorKeys.h"
+
+#include "PCGObjectPropertyOverride.generated.h"
+
+class IPCGAttributeAccessor;
+struct FPCGContext;
+
+/**
+* Represents the override source (to be read) and the object property (to be written).
+*/
+USTRUCT(BlueprintType)
+struct FPCGObjectPropertyOverrideDescription
+{
+	GENERATED_BODY()
+
+	FPCGObjectPropertyOverrideDescription() = default;
+
+	FPCGObjectPropertyOverrideDescription(const FPCGAttributePropertyInputSelector& InInputSource, const FString& InPropertyTarget)
+		: InputSource(InInputSource)
+		, PropertyTarget(InPropertyTarget)
+	{}
+
+	/** Provide an attribute or property to read the override value from. */
+	UPROPERTY(EditAnywhere, Category = Settings, meta = (PCG_Overridable))
+	FPCGAttributePropertyInputSelector InputSource;
+
+	/**
+	* Provide an object property name to be overridden. If you have a property "A" on your object, use "A" as the property target.
+	*
+	* For example, if you want to override the "Is Editor Only" flag, find it in the details panel, right-click, select 'Copy Internal Name', and paste that as the property target.
+	*
+	* If you have a component property, such as the static mesh of a static mesh component, use "StaticMeshComponent.StaticMesh".
+	*/
+	UPROPERTY(EditAnywhere, Category = Settings)
+	FString PropertyTarget;
+};
+
+namespace PCGObjectPropertyOverrideHelpers
+{
+	/** Create an advanced ParamData pin for capturing property overrides. */
+	FPCGPinProperties CreateObjectPropertiesOverridePin(FName Label, const FText& Tooltip);
+
+	/** Apply property overrides to the TargetObject directly from the ObjectPropertiesOverride pin. Use CreateObjectPropertiesOverridePin(). */
+	void ApplyOverridesFromParams(const TArray<FPCGObjectPropertyOverrideDescription>& InObjectPropertyOverrideDescriptions, AActor* TargetActor, FName OverridesPinLabel, FPCGContext* Context);
+}
+
+/**
+* Represents a single property override on the provided object. Applies an override function to read the InputAccessor
+* and write its value to the OutputAccessor.
+* 
+* The InputAccessor's InputKeys are created from the given SourceData and InputSelector.
+*/
+struct FPCGObjectSingleOverride
+{
+	/** Initialize the single object override. Call before using Apply(InputKeyIndex, OutputKey). */
+	void Initialize(const FPCGAttributePropertySelector& InputSelector, const FString& OutputProperty, const UStruct* TemplateClass, const UPCGData* SourceData, FPCGContext* Context);
+
+	/** Returns true if initialization succeeded in creating the accessors and accessor keys. */
+	bool IsValid() const;
+
+	/** Applies a single property override to the object by reading from the InputAccessor at the given KeyIndex, and writing to the OutputKey which represents the object property. */
+	bool Apply(int32 InputKeyIndex, IPCGAttributeAccessorKeys& OutputKey);
+
+private:
+	TUniquePtr<const IPCGAttributeAccessorKeys> InputKeys;
+	TUniquePtr<const IPCGAttributeAccessor> ObjectOverrideInputAccessor;
+	TUniquePtr<IPCGAttributeAccessor> ObjectOverrideOutputAccessor;
+
+	// InputKeyIndex, OutputKeys
+	using ApplyOverrideFunction = bool(FPCGObjectSingleOverride::*)(int32, IPCGAttributeAccessorKeys&);
+	ApplyOverrideFunction ObjectOverrideFunction = nullptr;
+
+	template <typename Type>
+	bool ApplyImpl(int32 InputKeyIndex, IPCGAttributeAccessorKeys& OutputKey)
+	{
+		if (!IsValid())
+		{
+			return false;
+		}
+
+		Type Value{};
+		check(ObjectOverrideInputAccessor.IsValid());
+		if (ObjectOverrideInputAccessor->Get<Type>(Value, InputKeyIndex, *InputKeys.Get(), EPCGAttributeAccessorFlags::AllowBroadcast | EPCGAttributeAccessorFlags::AllowConstructible))
+		{
+			check(ObjectOverrideOutputAccessor.IsValid());
+			if (ObjectOverrideOutputAccessor->Set<Type>(Value, OutputKey))
+			{
+				return true;
+			}
+		}
+
+		return false;
+	}
+};
+
+/**
+* Represents a set of property overrides for the provided object. Provide a SourceData to read from, and a collection of ObjectPropertyOverrides matching the TemplateObject's class properties.
+*/
+template <typename T>
+struct FPCGObjectOverrides
+{
+	FPCGObjectOverrides(T* TemplateObject) : OutputKey(TemplateObject)
+	{}
+
+	/** Initialize the object overrides. Call before using Apply(InputKeyIndex). */
+	void Initialize(const TArray<FPCGObjectPropertyOverrideDescription>& OverrideDescriptions, T* TemplateObject, const UPCGData* SourceData, FPCGContext* Context)
+	{
+		if (!TemplateObject)
+		{
+			PCGLog::LogErrorOnGraph(NSLOCTEXT("PCGObjectPropertyOverride", "InitializeOverrideFailedNoObject", "Failed to initialize property overrides. No template object was provided."), Context);
+			return;
+		}
+
+		const UStruct* ClassObject = nullptr;
+		if constexpr (TIsDerivedFrom<T, UObject>::Value)
+		{
+			ClassObject = TemplateObject->GetClass();
+		}
+		else
+		{
+			ClassObject = TemplateObject->StaticStruct();
+		}
+
+		ObjectSingleOverrides.Reserve(OverrideDescriptions.Num());
+
+		for (int32 i = 0; i < OverrideDescriptions.Num(); ++i)
+		{
+			FPCGAttributePropertyInputSelector InputSelector = OverrideDescriptions[i].InputSource.CopyAndFixLast(SourceData);
+
+			if (InputSelector.GetSelection() == EPCGAttributePropertySelection::Attribute)
+			{
+				const UPCGMetadata* Metadata = SourceData ? SourceData->ConstMetadata() : nullptr;
+				const FName AttributeName = InputSelector.GetAttributeName();
+
+				if (!Metadata || !Metadata->HasAttribute(AttributeName))
+				{
+					PCGLog::LogWarningOnGraph(FText::Format(NSLOCTEXT("PCGObjectPropertyOverride", "InvalidAttribute", "Tried to initialize ObjectOverride for input '{0}', but the attribute '{1}' does not exist."), InputSelector.GetDisplayText(), FText::FromName(AttributeName)), Context);
+					continue;
+				}
+			}
+
+			const FString& OutputProperty = OverrideDescriptions[i].PropertyTarget;
+
+			FPCGObjectSingleOverride Override;
+			Override.Initialize(InputSelector, OutputProperty, ClassObject, SourceData, Context);
+
+			if (Override.IsValid())
+			{
+				ObjectSingleOverrides.Add(std::move(Override));
+			}
+			else
+			{
+				PCGLog::LogErrorOnGraph(FText::Format(NSLOCTEXT("PCGObjectPropertyOverride", "InitializeOverrideFailed", "Failed to initialize override '{0}' for property {1} on object '{2}'."), InputSelector.GetDisplayText(), FText::FromString(OutputProperty), FText::FromName(ClassObject->GetFName())), Context);
+			}
+		}
+	}
+
+	/** Applies each property override to the object by reading from the InputAccessor at the given KeyIndex, and writing to the OutputKey which represents the object property. */
+	bool Apply(int32 InputKeyIndex)
+	{
+		bool bAllSucceeded = true;
+
+		for (FPCGObjectSingleOverride& ObjectSingleOverride : ObjectSingleOverrides)
+		{
+			bAllSucceeded &= ObjectSingleOverride.Apply(InputKeyIndex, OutputKey);
+		}
+
+		return bAllSucceeded;
+	}
+
+private:
+	FPCGAttributeAccessorKeysSingleObjectPtr<T> OutputKey;
+	TArray<FPCGObjectSingleOverride> ObjectSingleOverrides;
+};
+
+USTRUCT(BlueprintType, meta=(Deprecated = "5.4", DeprecationMessage="Use FPCGObjectPropertyOverrideDescription instead."))
+struct FPCGActorPropertyOverride
+{
+	GENERATED_BODY()
+
+	UPROPERTY(EditAnywhere, Category = Settings, meta = (PCG_Overridable))
+	FPCGAttributePropertyInputSelector InputSource;
+
+	UPROPERTY(EditAnywhere, Category = Settings)
+	FString PropertyTarget;
+};
