[PCG] Relax the refresh on graph parameters that doesn't impact the graph

* Impact any changes to the graph parameters or overrides.
* Also remove any refresh if the change to the graph is cosmetic

#rnx

diff --git a/Engine/Plugins/PCG/Source/PCG/Private/PCGComponent.cpp b/Engine/Plugins/PCG/Source/PCG/Private/PCGComponent.cpp
--- a/Engine/Plugins/PCG/Source/PCG/Private/PCGComponent.cpp
+++ b/Engine/Plugins/PCG/Source/PCG/Private/PCGComponent.cpp
@@ -1474,6 +1474,12 @@ void UPCGComponent::RefreshAfterGraphChanged(UPCGGraphInterface* InGraph, EPCGCh
 		return;
 	}
 
+	if (ChangeType == EPCGChangeType::Cosmetic || ChangeType == EPCGChangeType::None)
+	{
+		// If it is a cosmetic change (or no change), nothing to do
+		return;
+	}
+
 	const bool bHasGraph = (InGraph && InGraph->GetGraph());
 
 	const bool bIsStructural = ((ChangeType & (EPCGChangeType::Edge | EPCGChangeType::Structural)) != EPCGChangeType::None);
@@ -3207,7 +3213,7 @@ void FPCGComponentInstanceData::ApplyToComponent(UActorComponent* Component, con
 
 #if WITH_EDITOR
 			// bDirtyGenerated is transient.
-			PCGComponent->bDirtyGenerated = SourceComponent->bDirtyGenerated; 
+			PCGComponent->bDirtyGenerated = SourceComponent->bDirtyGenerated;
 
 			// While this is serialized, it is not properly copied because it is not visible. This is needed otherwise a refresh can retoggle from not generated to generated
 			PCGComponent->bForceGenerateOnBPAddedToWorld = SourceComponent->bForceGenerateOnBPAddedToWorld;
diff --git a/Engine/Plugins/PCG/Source/PCG/Private/PCGGraph.cpp b/Engine/Plugins/PCG/Source/PCG/Private/PCGGraph.cpp
--- a/Engine/Plugins/PCG/Source/PCG/Private/PCGGraph.cpp
+++ b/Engine/Plugins/PCG/Source/PCG/Private/PCGGraph.cpp
@@ -196,6 +196,49 @@ bool UPCGGraphInterface::IsEquivalent(const UPCGGraphInterface* Other) const
 	return true;
 }
 
+EPCGChangeType UPCGGraphInterface::GetChangeTypeForGraphParameterChange(EPCGGraphParameterEvent InChangeType, FName InChangedPropertyName)
+{
+	// If the parameter had its order changed in the struct, was just added or was removed but was not used in the graph, it is not a change that requires a refresh, so we go with Cosmetic change type.
+	if (InChangeType == EPCGGraphParameterEvent::PropertyMoved || InChangeType == EPCGGraphParameterEvent::Added || InChangeType == EPCGGraphParameterEvent::RemovedUnused)
+	{
+		return EPCGChangeType::Cosmetic;
+	}
+
+	// If it is not linked to a single property, or it was removed and used, we need to refresh, so we go with Settings change type.
+	if (InChangedPropertyName == NAME_None || InChangeType == EPCGGraphParameterEvent::RemovedUsed)
+	{
+		return EPCGChangeType::Settings;
+	}
+
+	const UPCGGraph* Graph = GetGraph();
+	const FInstancedPropertyBag* UserParameters = GetUserParametersStruct();
+	if (!ensure(Graph && UserParameters))
+	{
+		// Should never happen, but if there is no graph nor user parameters, there is nothing to do.
+		return EPCGChangeType::None;
+	}
+
+	// Finally if anything change on a property that has an impact for the graph, look for GetUserParameters nodes for this property, to only refresh if the property is used.
+	for (const UPCGNode* Node : Graph->GetNodes())
+	{
+		if (!Node)
+		{
+			continue;
+		}
+
+		if (UPCGUserParameterGetSettings* Settings = Cast<UPCGUserParameterGetSettings>(Node->GetSettings()))
+		{
+			if (Settings->PropertyName == InChangedPropertyName)
+			{
+				return EPCGChangeType::Settings;
+			}
+		}
+	}
+
+	// At this point, we didn't find any node that use our property, so no refresh needed.
+	return EPCGChangeType::Cosmetic;
+}
+
 /****************************
 * UPCGGraph
 ****************************/
@@ -1191,7 +1234,7 @@ void UPCGGraph::NotifyGraphParametersChanged(EPCGGraphParameterEvent InChangeTyp
 	OnGraphParametersChangedDelegate.Broadcast(this, InChangeType, InChangedPropertyName);
 	bIsNotifying = false;
 
-	NotifyGraphChanged(EPCGChangeType::Settings);
+	NotifyGraphChanged(GetChangeTypeForGraphParameterChange(InChangeType, InChangedPropertyName));
 }
 
 void UPCGGraph::OnNodeChanged(UPCGNode* InNode, EPCGChangeType ChangeType)
@@ -1230,8 +1273,8 @@ void UPCGGraph::PreEditChange(FProperty* InProperty)
 
 	if (InProperty->GetFName() == GET_MEMBER_NAME_CHECKED(UPCGGraph, UserParameters))
 	{
-		// We need to keep track of the number of properties, to detect if a property was added/removed/modified
-		NumberOfUserParametersPreEdit = UserParameters.GetNumPropertiesInBag();
+		// We need to keep track of the previous property bag, to detect if a property was added/removed/renamed/moved/modified...
+		PreviousPropertyBag = UserParameters.GetPropertyBagStruct();
 	}
 }
 
@@ -1253,8 +1296,10 @@ void UPCGGraph::PostEditChangeProperty(FPropertyChangedEvent& PropertyChangedEve
 	}
 	else if (PropertyName == GET_MEMBER_NAME_CHECKED(UPCGGraph, UserParameters))
 	{
-		EPCGGraphParameterEvent ChangeType;
-		int32 NumberOfUserParametersPostEdit = UserParameters.GetNumPropertiesInBag();
+		EPCGGraphParameterEvent ChangeType = EPCGGraphParameterEvent::None;
+		const int32 NumberOfUserParametersPreEdit = PreviousPropertyBag ? PreviousPropertyBag->GetPropertyDescs().Num() : 0;
+		const int32 NumberOfUserParametersPostEdit = UserParameters.GetNumPropertiesInBag();
+		FName ChangedPropertyName = NAME_None;
 
 		if (NumberOfUserParametersPostEdit > NumberOfUserParametersPreEdit)
 		{
@@ -1262,14 +1307,43 @@ void UPCGGraph::PostEditChangeProperty(FPropertyChangedEvent& PropertyChangedEve
 		}
 		else if (NumberOfUserParametersPostEdit < NumberOfUserParametersPreEdit)
 		{
-			ChangeType = EPCGGraphParameterEvent::Removed;
+			// Removed, but not knowing if it is used or not yet.
+			ChangeType = EPCGGraphParameterEvent::RemovedUnused;
 		}
-		else //NumberOfUserParametersPostEdit == NumberOfUserParametersPreEdit
+		else if (PreviousPropertyBag) // && NumberOfUserParametersPostEdit == NumberOfUserParametersPreEdit
 		{
-			ChangeType = EPCGGraphParameterEvent::PropertyModified;
+			for (int32 i = 0; i < NumberOfUserParametersPostEdit; ++i)
+			{
+				const FPropertyBagPropertyDesc& PreDesc = PreviousPropertyBag->GetPropertyDescs()[i];
+				const FPropertyBagPropertyDesc& PostDesc = UserParameters.GetPropertyBagStruct()->GetPropertyDescs()[i];
+
+				// Not Same ID -> Moved
+				if (PreDesc.ID != PostDesc.ID)
+				{
+					ChangeType = EPCGGraphParameterEvent::PropertyMoved;
+					break;
+				}
+				// Same ID but different name -> Renamed
+				else if (PreDesc.Name != PostDesc.Name)
+				{
+					ChangeType = EPCGGraphParameterEvent::PropertyRenamed;
+					ChangedPropertyName = PostDesc.Name;
+					break;
+				}
+				// Same name but different type -> Type modified
+				else if (!PostDesc.CompatibleType(PreDesc))
+				{
+					ChangeType = EPCGGraphParameterEvent::PropertyTypeModified;
+					ChangedPropertyName = PostDesc.Name;
+					break;
+				}
+			}
 		}
 
-		OnGraphParametersChanged(ChangeType, NAME_None);
+		if (ChangeType != EPCGGraphParameterEvent::None)
+		{
+			OnGraphParametersChanged(ChangeType, ChangedPropertyName);
+		}
 	}
 	else if (PropertyChangedEvent.MemberProperty && PropertyChangedEvent.MemberProperty->GetOwnerStruct() == UserParameters.GetPropertyBagStruct())
 	{
@@ -1282,7 +1356,7 @@ void UPCGGraph::PostEditChangeProperty(FPropertyChangedEvent& PropertyChangedEve
 		NotifyGraphChanged(EPCGChangeType::Structural | EPCGChangeType::GenerationGrid);
 	}
 
-	NumberOfUserParametersPreEdit = 0;
+	PreviousPropertyBag = nullptr;
 }
 
 void UPCGGraph::FixInvalidEdges()
@@ -1370,8 +1444,9 @@ bool UPCGGraph::UserParametersIsPinTypeAccepted(FEdGraphPinType InPinType, bool
 void UPCGGraph::OnGraphParametersChanged(EPCGGraphParameterEvent InChangeType, FName InChangedPropertyName)
 {
 	bool bWasModified = false;
+	EPCGGraphParameterEvent ChangeType = InChangeType;
 
-	if (InChangeType == EPCGGraphParameterEvent::Removed || InChangeType == EPCGGraphParameterEvent::PropertyModified)
+	if (InChangeType == EPCGGraphParameterEvent::RemovedUsed || InChangeType == EPCGGraphParameterEvent::RemovedUnused || InChangeType == EPCGGraphParameterEvent::PropertyRenamed)
 	{
 		// Look for all the Get Parameter nodes and make sure to delete all nodes that doesn't exist anymore
 		TArray<UPCGNode*> NodesToRemove;
@@ -1388,6 +1463,7 @@ void UPCGGraph::OnGraphParametersChanged(EPCGGraphParameterEvent InChangeType, F
 				const FPropertyBagPropertyDesc* PropertyDesc = UserParameters.FindPropertyDescByID(Settings->PropertyGuid);
 				if (!PropertyDesc)
 				{
+					ChangeType = EPCGGraphParameterEvent::RemovedUsed;
 					NodesToRemove.Add(Node);
 				}
 				else if (Settings->PropertyName != PropertyDesc->Name)
@@ -1573,13 +1649,19 @@ void UPCGGraphInstance::PostEditChangeProperty(FPropertyChangedEvent& PropertyCh
 		return;
 	}
 
+	const FName PropertyName = PropertyChangedEvent.GetPropertyName();
+
 	// We need to be careful and only capture `Graph` if it is our graph and not a graph parameter called `Graph`!
-	if (PropertyChangedEvent.Property->GetOwnerClass() == UPCGGraphInstance::StaticClass() && PropertyChangedEvent.GetPropertyName() == GET_MEMBER_NAME_CHECKED(UPCGGraphInstance, Graph))
+	if (PropertyChangedEvent.Property->GetOwnerClass() == UPCGGraphInstance::StaticClass() && PropertyName == GET_MEMBER_NAME_CHECKED(UPCGGraphInstance, Graph))
 	{
 		SetupCallbacks();
 
 		RefreshParameters(EPCGGraphParameterEvent::GraphChanged);
 	}
+	else if (PropertyName == GET_MEMBER_NAME_CHECKED(FPCGOverrideInstancedPropertyBag, PropertiesIDsOverridden))
+	{
+		// Changes to this property will be caught by the update on override status, and "OnGraphParametersChanged" will be called accordingly.
+	}
 	else if (PropertyChangedEvent.MemberProperty && PropertyChangedEvent.MemberProperty->GetOwnerStruct() == ParametersOverrides.Parameters.GetPropertyBagStruct())
 	{
 		OnGraphParametersChanged(this, EPCGGraphParameterEvent::ValueModifiedLocally, PropertyChangedEvent.GetMemberPropertyName());
@@ -1718,7 +1800,7 @@ void UPCGGraphInstance::NotifyGraphParametersChanged(EPCGGraphParameterEvent InC
 	OnGraphParametersChangedDelegate.Broadcast(this, InChangeType, InChangedPropertyName);
 
 	// Also propagates the changes
-	OnGraphChanged(Graph, EPCGChangeType::Settings);
+	OnGraphChanged(Graph, GetChangeTypeForGraphParameterChange(InChangeType, InChangedPropertyName));
 }
 #endif // WITH_EDITOR
 
@@ -1825,10 +1907,13 @@ void UPCGGraphInstance::ResetPropertyToDefault(const FProperty* InProperty)
 
 	Modify();
 
-	ParametersOverrides.ResetPropertyToDefault(InProperty, Graph->GetUserParametersStruct());
+	bool bValueChanged = ParametersOverrides.ResetPropertyToDefault(InProperty, Graph->GetUserParametersStruct());
 
 #if WITH_EDITOR
-	NotifyGraphParametersChanged(EPCGGraphParameterEvent::ValueModifiedLocally, InProperty->GetFName());
+	if (bValueChanged)
+	{
+		NotifyGraphParametersChanged(EPCGGraphParameterEvent::ValueModifiedLocally, InProperty->GetFName());
+	}
 #endif // WITH_EDITOR
 }
 
@@ -1869,9 +1954,12 @@ bool FPCGOverrideInstancedPropertyBag::RefreshParameters(const FInstancedPropert
 #endif // WITH_EDITOR
 		break;
 	}
-	case EPCGGraphParameterEvent::Added:
-	case EPCGGraphParameterEvent::Removed:
-	case EPCGGraphParameterEvent::PropertyModified:
+	case EPCGGraphParameterEvent::Added: // fall-through
+	case EPCGGraphParameterEvent::RemovedUnused: // fall-through
+	case EPCGGraphParameterEvent::RemovedUsed: // fall-through
+	case EPCGGraphParameterEvent::PropertyRenamed: // fall-through
+	case EPCGGraphParameterEvent::PropertyMoved: // fall-through
+	case EPCGGraphParameterEvent::PropertyTypeModified:
 	{
 		bWasModified = true;
 		const FPropertyBagPropertyDesc* ThisPropertyDesc = Parameters.FindPropertyDescByName(InChangedPropertyName);
@@ -1993,14 +2081,13 @@ bool FPCGOverrideInstancedPropertyBag::UpdatePropertyOverride(const FProperty* I
 	// Reset the value if it is not marked overridden anymore.
 	if (!bMarkAsOverridden)
 	{
-		ResetPropertyToDefault(InProperty, ParentUserParameters);
-		return true;
+		return ResetPropertyToDefault(InProperty, ParentUserParameters);
 	}
 
 	return false;
 }
 
-void FPCGOverrideInstancedPropertyBag::ResetPropertyToDefault(const FProperty* InProperty, const FInstancedPropertyBag* ParentUserParameters)
+bool FPCGOverrideInstancedPropertyBag::ResetPropertyToDefault(const FProperty* InProperty, const FInstancedPropertyBag* ParentUserParameters)
 {
 	check(ParentUserParameters);
 
@@ -2009,8 +2096,14 @@ void FPCGOverrideInstancedPropertyBag::ResetPropertyToDefault(const FProperty* I
 
 	if (OriginalPropertyDesc && ThisPropertyDesc)
 	{
-		PCGGraphUtils::CopyPropertyValue(OriginalPropertyDesc, *ParentUserParameters, ThisPropertyDesc, Parameters);
+		if (!PCGGraphUtils::ArePropertiesIdentical(OriginalPropertyDesc, *ParentUserParameters, ThisPropertyDesc, Parameters))
+		{
+			PCGGraphUtils::CopyPropertyValue(OriginalPropertyDesc, *ParentUserParameters, ThisPropertyDesc, Parameters);
+			return true;
+		}
 	}
+
+	return false;
 }
 
 bool FPCGOverrideInstancedPropertyBag::IsPropertyOverridden(const FProperty* InProperty) const
diff --git a/Engine/Plugins/PCG/Source/PCG/Public/PCGGraph.h b/Engine/Plugins/PCG/Source/PCG/Public/PCGGraph.h
--- a/Engine/Plugins/PCG/Source/PCG/Public/PCGGraph.h
+++ b/Engine/Plugins/PCG/Source/PCG/Public/PCGGraph.h
@@ -24,11 +24,15 @@ enum class EPCGGraphParameterEvent
 	GraphChanged,
 	GraphPostLoad,
 	Added,
-	Removed,
-	PropertyModified,
+	RemovedUnused,
+	RemovedUsed,
+	PropertyMoved,
+	PropertyRenamed,
+	PropertyTypeModified,
 	ValueModifiedLocally,
 	ValueModifiedByParent,
-	MultiplePropertiesAdded
+	MultiplePropertiesAdded,
+	None
 };
 
 #if WITH_EDITOR
@@ -51,8 +55,8 @@ public:
 	/** Add/Remove given property from overrides, and reset its value if it is removed. Returns true if the value was changed. */
 	bool UpdatePropertyOverride(const FProperty* InProperty, bool bMarkAsOverridden, const FInstancedPropertyBag* ParentUserParameters);
 
-	/** Reset overridden property to its parent value. */
-	void ResetPropertyToDefault(const FProperty* InProperty, const FInstancedPropertyBag* ParentUserParameters);
+	/** Reset overridden property to its parent value. Return true if the value was different. */
+	bool ResetPropertyToDefault(const FProperty* InProperty, const FInstancedPropertyBag* ParentUserParameters);
 
 	/** Return if the property is currently marked overridden. */
 	bool IsPropertyOverridden(const FProperty* InProperty) const;
@@ -162,6 +166,9 @@ public:
 protected:
 	virtual void OnGraphParametersChanged(EPCGGraphParameterEvent InChangeType, FName InChangedPropertyName) PURE_VIRTUAL(UPCGGraphInterface::OnGraphParametersChanged, )
 	virtual FInstancedPropertyBag* GetMutableUserParametersStruct() PURE_VIRTUAL(UPCGGraphInterface::GetMutableUserParametersStruct, return nullptr;)
+
+	/** Detecting if we need to refresh the graph depending on the type of change in the Graph Parameter. */
+	EPCGChangeType GetChangeTypeForGraphParameterChange(EPCGGraphParameterEvent InChangeType, FName InChangedPropertyName);
 };
 
 UCLASS(BlueprintType, ClassGroup = (Procedural), hidecategories=(Object))
@@ -415,12 +422,14 @@ private:
 	/** Remove invalid edges and edges to nodes that are not present in the node array. */
 	void FixInvalidEdges();
 
+	// Keep track of the previous PropertyBag, to see if we had a change in the number of properties, or if it is a rename/move.
+	TObjectPtr<const UPropertyBag> PreviousPropertyBag;
+
 	int32 GraphChangeNotificationsDisableCounter = 0;
-	bool bDelayedChangeNotification = false;
 	EPCGChangeType DelayedChangeType = EPCGChangeType::None;
+	bool bDelayedChangeNotification = false;
 	bool bIsNotifying = false;
 	bool bUserPausedNotificationsInGraphEditor = false;
-	int32 NumberOfUserParametersPreEdit = 0;
 	bool bIsInspecting = false;
 #endif // WITH_EDITOR
 };
diff --git a/Engine/Plugins/PCG/Source/PCGEditor/Private/PCGEditorGraph.cpp b/Engine/Plugins/PCG/Source/PCGEditor/Private/PCGEditorGraph.cpp
--- a/Engine/Plugins/PCG/Source/PCGEditor/Private/PCGEditorGraph.cpp
+++ b/Engine/Plugins/PCG/Source/PCGEditor/Private/PCGEditorGraph.cpp
@@ -290,7 +290,7 @@ void UPCGEditorGraph::CreateLinks(UPCGEditorGraphNodeBase* GraphNode, bool bCrea
 
 void UPCGEditorGraph::OnGraphUserParametersChanged(UPCGGraphInterface* InGraph, EPCGGraphParameterEvent ChangeType, FName ChangedPropertyName)
 {
-	if (ChangeType != EPCGGraphParameterEvent::Removed || InGraph != PCGGraph)
+	if ((ChangeType != EPCGGraphParameterEvent::RemovedUnused && ChangeType != EPCGGraphParameterEvent::RemovedUsed) || InGraph != PCGGraph)
 	{
 		return;
 	}
