[PCG] Fix: Automation test 'Plugins.PCG.Texture.OffsetTilingRotation' no longer fails with errors

diff --git a/Engine/Plugins/PCG/Source/PCG/Private/Tests/Data/PCGTextureDataTest.cpp b/Engine/Plugins/PCG/Source/PCG/Private/Tests/Data/PCGTextureDataTest.cpp
--- a/Engine/Plugins/PCG/Source/PCG/Private/Tests/Data/PCGTextureDataTest.cpp
+++ b/Engine/Plugins/PCG/Source/PCG/Private/Tests/Data/PCGTextureDataTest.cpp
@@ -44,10 +44,13 @@ bool FPCGTextureDataOffsetTilingRotation::RunTest(const FString& Parameters)
 			FRandomStream RandomStream;
 			FPCGPoint OutPoint;
 
-			const float TextureSpacePixelX = static_cast<float>(WhitePixelX) / TextureSize;
-			const float ScaledPixelX = (2.0 * TextureSpacePixelX) - 1.0;
-			const float TextureSpacePixelY = static_cast<float>(WhitePixelY) / TextureSize;
-			const float ScaledPixelY = (2.0 * TextureSpacePixelY) - 1;
+			const float PixelCenterX = WhitePixelX + 0.5f;
+			const float PixelCenterY = WhitePixelY + 0.5f;
+
+			const float TextureSpacePixelX = PixelCenterX / TextureSize;
+			const float ScaledPixelX = (2.0f * TextureSpacePixelX) - 1.0f;
+			const float TextureSpacePixelY = PixelCenterY / TextureSize;
+			const float ScaledPixelY = (2.0f * TextureSpacePixelY) - 1.0f;
 
 			{
 				TextureData->Rotation = 0.f;
