[PCG] Fixed a bug where SMS Property Overrides were only taking the override from the first connected data

#rnx

diff --git a/Engine/Plugins/PCG/Source/PCG/Private/Elements/PCGStaticMeshSpawnerContext.cpp b/Engine/Plugins/PCG/Source/PCG/Private/Elements/PCGStaticMeshSpawnerContext.cpp
--- a/Engine/Plugins/PCG/Source/PCG/Private/Elements/PCGStaticMeshSpawnerContext.cpp
+++ b/Engine/Plugins/PCG/Source/PCG/Private/Elements/PCGStaticMeshSpawnerContext.cpp
@@ -6,6 +6,7 @@ void FPCGStaticMeshSpawnerContext::ResetInputIterationData()
 {
 	bCurrentInputSetup = false;
 	bSelectionDone = false;
+	bPartitionDone = false;
 	CurrentPointData = nullptr;
 	CurrentOutputPointData = nullptr;
 	MaterialOverrideHelper.Reset();
@@ -15,4 +16,6 @@ void FPCGStaticMeshSpawnerContext::ResetInputIterationData()
 	MeshToValueKey.Reset();
 	CumulativeWeights.Reset();
 	CategoryEntryToInstancesAndWeights.Reset();
+	AttributeOverridePartition.Reset();
+	OverriddenDescriptors.Reset();
 }
\ No newline at end of file
