[PCG] Fix crash when assigning a graph level instance to itself.

#rnx

diff --git a/Engine/Plugins/PCG/Source/PCG/Private/PCGGraph.cpp b/Engine/Plugins/PCG/Source/PCG/Private/PCGGraph.cpp
--- a/Engine/Plugins/PCG/Source/PCG/Private/PCGGraph.cpp
+++ b/Engine/Plugins/PCG/Source/PCG/Private/PCGGraph.cpp
@@ -336,7 +336,7 @@ void UPCGGraph::PostLoad()
 
 	// Ensure that all nodes are loaded (& updated their deprecated data)
 	for (UPCGNode* Node : Nodes)
-	{
+		{
 		Node->ConditionalPostLoad();
 	}
 
@@ -1636,6 +1636,8 @@ void UPCGGraphInstance::PreEditChange(FProperty* InProperty)
 	// We need to be careful and only capture `Graph` if it is our graph and not a graph parameter called `Graph`!
 	if (InProperty->GetOwnerClass() == UPCGGraphInstance::StaticClass() && InProperty->GetFName() == GET_MEMBER_NAME_CHECKED(UPCGGraphInstance, Graph))
 	{
+		PreGraphCache = Graph;
+
 		TeardownCallbacks();
 	}
 }
@@ -1654,9 +1656,20 @@ void UPCGGraphInstance::PostEditChangeProperty(FPropertyChangedEvent& PropertyCh
 	// We need to be careful and only capture `Graph` if it is our graph and not a graph parameter called `Graph`!
 	if (PropertyChangedEvent.Property->GetOwnerClass() == UPCGGraphInstance::StaticClass() && PropertyName == GET_MEMBER_NAME_CHECKED(UPCGGraphInstance, Graph))
 	{
+		// If the new graph is itself, return to the previous value
+		if (this == Graph)
+		{
+			UE_LOG(LogPCG, Error, TEXT("Try to set the graph of a graph instance to itself, would cause infinite recursion."));
+			Graph = PreGraphCache;
+		}
+
 		SetupCallbacks();
 
-		RefreshParameters(EPCGGraphParameterEvent::GraphChanged);
+		// No need to refresh if it is the same graph
+		if (Graph == PreGraphCache)
+		{
+			RefreshParameters(EPCGGraphParameterEvent::GraphChanged);
+		}
 	}
 	else if (PropertyName == GET_MEMBER_NAME_CHECKED(FPCGOverrideInstancedPropertyBag, PropertiesIDsOverridden))
 	{
@@ -1666,6 +1679,8 @@ void UPCGGraphInstance::PostEditChangeProperty(FPropertyChangedEvent& PropertyCh
 	{
 		OnGraphParametersChanged(this, EPCGGraphParameterEvent::ValueModifiedLocally, PropertyChangedEvent.GetMemberPropertyName());
 	}
+
+	PreGraphCache = nullptr;
 }
 
 void UPCGGraphInstance::PreEditUndo()
@@ -1674,7 +1689,7 @@ void UPCGGraphInstance::PreEditUndo()
 
 	TeardownCallbacks();
 
-	UndoRedoGraphCache = Graph;
+	PreGraphCache = Graph;
 }
 
 void UPCGGraphInstance::PostEditUndo()
diff --git a/Engine/Plugins/PCG/Source/PCG/Public/PCGGraph.h b/Engine/Plugins/PCG/Source/PCG/Public/PCGGraph.h
--- a/Engine/Plugins/PCG/Source/PCG/Public/PCGGraph.h
+++ b/Engine/Plugins/PCG/Source/PCG/Public/PCGGraph.h
@@ -495,8 +495,8 @@ public:
 
 private:
 #if WITH_EDITORONLY_DATA
-	// Transient, to keep track the undo/redo changed the graph.
-	UPCGGraphInterface* UndoRedoGraphCache = nullptr;
+	// Transient, to keep track of the previous graph when it changed.
+	UPCGGraphInterface* PreGraphCache = nullptr;
 #endif // WITH_EDITORONLY_DATA
 };
 
