[PCG] Changes on the Graph instance on the PCG Component should not trigger changes if not impactful

#rnx

diff --git a/Engine/Plugins/PCG/Source/PCG/Private/PCGComponent.cpp b/Engine/Plugins/PCG/Source/PCG/Private/PCGComponent.cpp
--- a/Engine/Plugins/PCG/Source/PCG/Private/PCGComponent.cpp
+++ b/Engine/Plugins/PCG/Source/PCG/Private/PCGComponent.cpp
@@ -1623,7 +1623,7 @@ void UPCGComponent::PostEditChangeProperty(FPropertyChangedEvent& PropertyChange
 	}
 	else if (PropName == GET_MEMBER_NAME_CHECKED(UPCGComponent, GraphInstance))
 	{
-		OnGraphChanged(GraphInstance, EPCGChangeType::Structural | EPCGChangeType::GenerationGrid);
+		// If anything happens on the graph instance, it will be handled there.
 	}
 	else if (PropName == GET_MEMBER_NAME_CHECKED(UPCGComponent, InputType))
 	{
@@ -2892,6 +2892,12 @@ void UPCGComponent::RegisterDynamicTracking(const UPCGSettings* InSettings, cons
 
 	for (const TPair<FPCGSelectionKey, bool>& It : InDynamicKeysAndCulling)
 	{
+		// Make sure to not register null assets
+		if (It.Key.Selection == EPCGActorSelection::ByPath && It.Key.ObjectPath.IsNull())
+		{
+			continue;
+		}
+
 		CurrentExecutionDynamicTracking.FindOrAdd(It.Key).Emplace(InSettings, It.Value);
 	}
 }
