[PCG] Add safe-guards to avoid crashing in post load when a node doesn't exist anymore.

#rnx

diff --git a/Engine/Plugins/PCG/Source/PCG/Private/PCGGraph.cpp b/Engine/Plugins/PCG/Source/PCG/Private/PCGGraph.cpp
--- a/Engine/Plugins/PCG/Source/PCG/Private/PCGGraph.cpp
+++ b/Engine/Plugins/PCG/Source/PCG/Private/PCGGraph.cpp
@@ -335,15 +335,32 @@ void UPCGGraph::PostLoad()
 	Cast<UPCGGraphInputOutputSettings>(OutputNode->GetSettings())->SetInput(false);
 
 	// Ensure that all nodes are loaded (& updated their deprecated data)
-	for (UPCGNode* Node : Nodes)
+	// If a node is null (can happen if an asset was saved with a node that don't exist in the current session), we don't want to crash. So remove the faulty node and warn the user.
+	for (int32 i = Nodes.Num() - 1; i >= 0; --i)
+	{
+		if (!Nodes[i])
+		{
+			UE_LOG(LogPCG, Error, TEXT("Graph %s has a node that doesn't exist anymore. Check if you are missing a plugin or if you saved an asset with an old settings that was removed/renamed."), *GetPathName());
+			Nodes.RemoveAtSwap(i);
+		}
+		else
 		{
-		Node->ConditionalPostLoad();
+			Nodes[i]->ConditionalPostLoad();
+		}
 	}
 
 	// Also do this for ExtraNodes
-	for (UObject* ExtraNode : ExtraEditorNodes)
+	for (int32 i = ExtraEditorNodes.Num() - 1; i >= 0; --i)
 	{
-		ExtraNode->ConditionalPostLoad();
+		if (!ExtraEditorNodes[i])
+		{
+			UE_LOG(LogPCG, Error, TEXT("Graph %s has an extra non-PCG node that doesn't exist anymore. Check if you are missing a plugin or if you saved an asset with an old settings that was removed/renamed."), *GetPathName());
+			ExtraEditorNodes.RemoveAtSwap(i);
+		}
+		else
+		{
+			ExtraEditorNodes[i]->ConditionalPostLoad();
+		}
 	}
 
 	// Create a copy to iterate through the nodes while more might be added
