Introduce visual edge feedback when debugging (e.g. fat edges) and editor settings to control the appropriate values.
Enabled inspection on reroutes and named reroutes.

diff --git a/Engine/Plugins/PCG/Source/PCG/Public/Elements/PCGReroute.h b/Engine/Plugins/PCG/Source/PCG/Public/Elements/PCGReroute.h
--- a/Engine/Plugins/PCG/Source/PCG/Public/Elements/PCGReroute.h
+++ b/Engine/Plugins/PCG/Source/PCG/Public/Elements/PCGReroute.h
@@ -24,7 +24,6 @@ public:
 
 	//~Begin UPCGSettingsInterface interface
 	virtual bool CanBeDisabled() const override { return false; }
-	virtual bool CanBeDebugged() const override { return false; }
 	//~End UPCGSettingsInterface interface
 
 	//~Begin UPCGSettings interface
diff --git a/Engine/Plugins/PCG/Source/PCGEditor/Private/PCGEditorGraph.cpp b/Engine/Plugins/PCG/Source/PCGEditor/Private/PCGEditorGraph.cpp
--- a/Engine/Plugins/PCG/Source/PCGEditor/Private/PCGEditorGraph.cpp
+++ b/Engine/Plugins/PCG/Source/PCGEditor/Private/PCGEditorGraph.cpp
@@ -15,6 +15,34 @@
 
 #include "EdGraph/EdGraphPin.h"
 
+namespace PCGEditorGraphUtils
+{
+	void GetInspectablePin(const UPCGNode* InNode, const UPCGPin* InPin, const UPCGNode*& OutNode, const UPCGPin*& OutPin)
+	{
+		OutNode = InNode;
+		OutPin = InPin;
+
+		// Basically, this is needed so we can go up the graph when the selected node/pin combo is on a reroute node.
+		while (OutPin && OutPin->IsOutputPin() &&
+			OutNode && OutNode->GetSettings() && OutNode->GetSettings()->IsA<UPCGRerouteSettings>())
+		{
+			// Since it's a reroute node, we can look at the inbound edge (if any) on the reroute node and go up there
+			check(OutNode->GetInputPin(PCGPinConstants::DefaultInputLabel));
+			const TArray<TObjectPtr<UPCGEdge>>& Edges = OutNode->GetInputPin(PCGPinConstants::DefaultInputLabel)->Edges;
+			// A reroute node can have at most one inbound edge, but we still need to make sure it exists
+			if (Edges.Num() == 1)
+			{
+				OutPin = Edges[0]->InputPin;
+				OutNode = OutPin->Node;
+			}
+			else
+			{
+				break;
+			}
+		}
+	}
+}
+
 void UPCGEditorGraph::InitFromNodeGraph(UPCGGraph* InPCGGraph)
 {
 	check(InPCGGraph && !PCGGraph);
diff --git a/Engine/Plugins/PCG/Source/PCGEditor/Private/PCGEditorGraph.h b/Engine/Plugins/PCG/Source/PCGEditor/Private/PCGEditorGraph.h
--- a/Engine/Plugins/PCG/Source/PCGEditor/Private/PCGEditorGraph.h
+++ b/Engine/Plugins/PCG/Source/PCGEditor/Private/PCGEditorGraph.h
@@ -10,8 +10,15 @@
 
 class FPCGEditor;
 class UPCGNode;
+class UPCGPin;
 class UPCGEditorGraphNodeBase;
 
+namespace PCGEditorGraphUtils
+{
+	/** Goes up the graph to the first non-reroute node it can, assuming it's properly connected. */
+	void GetInspectablePin(const UPCGNode* InNode, const UPCGPin* InPin, const UPCGNode*& OutNode, const UPCGPin*& OutPin);
+}
+
 UCLASS()
 class UPCGEditorGraph : public UEdGraph
 {
diff --git a/Engine/Plugins/PCG/Source/PCGEditor/Private/PCGEditorGraphSchema.cpp b/Engine/Plugins/PCG/Source/PCGEditor/Private/PCGEditorGraphSchema.cpp
--- a/Engine/Plugins/PCG/Source/PCGEditor/Private/PCGEditorGraphSchema.cpp
+++ b/Engine/Plugins/PCG/Source/PCGEditor/Private/PCGEditorGraphSchema.cpp
@@ -2,6 +2,7 @@
 
 #include "PCGEditorGraphSchema.h"
 
+#include "PCGComponent.h"
 #include "PCGEdge.h"
 #include "PCGGraph.h"
 #include "PCGPin.h"
@@ -13,6 +14,7 @@
 #include "Elements/PCGReroute.h"
 #include "Elements/PCGUserParameterGet.h"
 
+#include "PCGEditor.h"
 #include "PCGEditorCommon.h"
 #include "PCGEditorGraph.h"
 #include "PCGEditorGraphNodeBase.h"
@@ -706,6 +708,64 @@ FPCGEditorConnectionDrawingPolicy::FPCGEditorConnectionDrawingPolicy(int32 InBac
 	ArrowRadius = FVector2D::ZeroVector;
 }
 
+bool FPCGEditorConnectionDrawingPolicy::UpdateParamsIfDebugging(UEdGraphPin* OutputPin, UEdGraphPin* InputPin, FConnectionParams& Params)
+{
+	check(OutputPin && InputPin);
+
+	// Early validation
+	const UPCGEditorGraphNodeBase* UpstreamEditorNode = CastChecked<const UPCGEditorGraphNodeBase>(OutputPin->GetOwningNode());
+	if (!UpstreamEditorNode || !UpstreamEditorNode->GetPCGNode())
+	{
+		return false;
+	}
+
+	const UPCGNode* NodeToInspect = nullptr;
+	const UPCGPin* PinToInspect = nullptr;
+
+	// Walk up the graph if the current node is a reroute node because there is no associated inspection data.
+	PCGEditorGraphUtils::GetInspectablePin(UpstreamEditorNode->GetPCGNode(), UpstreamEditorNode->GetPCGNode()->GetOutputPin(OutputPin->GetFName()), NodeToInspect, PinToInspect);
+
+	if (!NodeToInspect || !PinToInspect)
+	{
+		return false;
+	}
+
+	// Early out if we aren't in debug mode
+	if (!Graph || !Graph->GetEditor().IsValid())
+	{
+		return false;
+	}
+
+	const FPCGEditor* Editor = Graph->GetEditor().Pin().Get();
+	const FPCGStack* PCGStack = Editor ? Editor->GetStackBeingInspected() : nullptr;
+	UPCGComponent* PCGComponent = Editor ? Editor->GetPCGComponentBeingInspected() : nullptr;
+
+	if (!Editor || !PCGStack || !PCGComponent)
+	{
+		return false;
+	}
+
+	FPCGStack Stack = *PCGStack;
+	TArray<FPCGStackFrame>& StackFrames = Stack.GetStackFramesMutable();
+	StackFrames.Reserve(StackFrames.Num() + 2);
+	StackFrames.Emplace(NodeToInspect);
+	StackFrames.Emplace(PinToInspect);
+
+	if (const FPCGDataCollection* DataCollection = PCGComponent->GetInspectionData(Stack))
+	{
+		if (DataCollection->TaggedData.Num() > 1)
+		{
+			Params.WireThickness *= GetDefault<UPCGEditorSettings>()->MultiDataEdgeDebugEmphasis;
+		}
+	}
+	else
+	{
+		Params.WireColor = Params.WireColor.Desaturate(GetDefault<UPCGEditorSettings>()->EmptyEdgeDebugDesaturateFactor);
+	}
+
+	return true;
+}
+
 void FPCGEditorConnectionDrawingPolicy::DetermineWiringStyle(UEdGraphPin* OutputPin, UEdGraphPin* InputPin, /*inout*/ FConnectionParams& Params)
 {
 	FConnectionDrawingPolicy::DetermineWiringStyle(OutputPin, InputPin, Params);
@@ -713,7 +773,7 @@ void FPCGEditorConnectionDrawingPolicy::DetermineWiringStyle(UEdGraphPin* Output
 	// Emphasize wire thickness on hovered pins
 	if (HoveredPins.Contains(InputPin) && HoveredPins.Contains(OutputPin))
 	{
-		Params.WireThickness = Params.WireThickness * 3;
+		Params.WireThickness *= GetDefault<UPCGEditorSettings>()->HoverEdgeEmphasis;
 	}
 
 	// Base the color of the wire on the color of the output pin
@@ -725,27 +785,31 @@ void FPCGEditorConnectionDrawingPolicy::DetermineWiringStyle(UEdGraphPin* Output
 	// Desaturate connection if downstream node is disabled or if the data on this wire won't be used
 	if (InputPin && OutputPin)
 	{
-		const UPCGEditorGraphNodeBase* EditorNode = CastChecked<const UPCGEditorGraphNodeBase>(InputPin->GetOwningNode());
-		const UPCGNode* PCGNode = EditorNode ? EditorNode->GetPCGNode() : nullptr;
-		const UPCGPin* PCGPin = PCGNode ? PCGNode->GetInputPin(InputPin->GetFName()) : nullptr;
-		const UPCGEditorGraphNodeBase* UpstreamEditorNode = CastChecked<const UPCGEditorGraphNodeBase>(OutputPin->GetOwningNode());
-		const UPCGEditorGraphNodeBase* DownstreamEditorNode = CastChecked<const UPCGEditorGraphNodeBase>(InputPin->GetOwningNode());
-
-		if (PCGPin && UpstreamEditorNode && DownstreamEditorNode)
+		// Try to apply debugging/dynamic visualization - if it fails, fall back to static visualization
+		if (!UpdateParamsIfDebugging(OutputPin, InputPin, Params))
 		{
-			const bool bDownstreamNodeForceDisabled = DownstreamEditorNode->IsDisplayAsDisabledForced();
+			const UPCGEditorGraphNodeBase* EditorNode = CastChecked<const UPCGEditorGraphNodeBase>(InputPin->GetOwningNode());
+			const UPCGNode* PCGNode = EditorNode ? EditorNode->GetPCGNode() : nullptr;
+			const UPCGPin* PCGPin = PCGNode ? PCGNode->GetInputPin(InputPin->GetFName()) : nullptr;
+			const UPCGEditorGraphNodeBase* UpstreamEditorNode = CastChecked<const UPCGEditorGraphNodeBase>(OutputPin->GetOwningNode());
+			const UPCGEditorGraphNodeBase* DownstreamEditorNode = CastChecked<const UPCGEditorGraphNodeBase>(InputPin->GetOwningNode());
 
-			// Look for the PCG edge that correlates with passed in (OutputPin, InputPin) edge
-			const TObjectPtr<UPCGEdge>* PCGEdge = PCGPin->Edges.FindByPredicate([UpstreamEditorNode, OutputPin](const UPCGEdge* ConnectedPCGEdge)
+			if (PCGPin && UpstreamEditorNode && DownstreamEditorNode)
 			{
-				return UpstreamEditorNode->GetPCGNode() == ConnectedPCGEdge->InputPin->Node && ConnectedPCGEdge->InputPin->Properties.Label == OutputPin->GetFName();
-			});
-			const bool bDownstreamNodeDoesNotUseData = PCGEdge && !PCGNode->IsEdgeUsedByNodeExecution(*PCGEdge);
+				const bool bDownstreamNodeForceDisabled = DownstreamEditorNode->IsDisplayAsDisabledForced();
 
-			// If edge found and is not used, gray it out
-			if (bDownstreamNodeForceDisabled || bDownstreamNodeDoesNotUseData)
-			{
-				Params.WireColor = Params.WireColor.Desaturate(0.7f);
+				// Look for the PCG edge that correlates with passed in (OutputPin, InputPin) edge
+				const TObjectPtr<UPCGEdge>* PCGEdge = PCGPin->Edges.FindByPredicate([UpstreamEditorNode, OutputPin](const UPCGEdge* ConnectedPCGEdge)
+				{
+					return UpstreamEditorNode->GetPCGNode() == ConnectedPCGEdge->InputPin->Node && ConnectedPCGEdge->InputPin->Properties.Label == OutputPin->GetFName();
+				});
+				const bool bDownstreamNodeDoesNotUseData = PCGEdge && !PCGNode->IsEdgeUsedByNodeExecution(*PCGEdge);
+
+				// If edge found and is not used, gray it out
+				if (bDownstreamNodeForceDisabled || bDownstreamNodeDoesNotUseData)
+				{
+					Params.WireColor = Params.WireColor.Desaturate(GetDefault<UPCGEditorSettings>()->EmptyEdgeDebugDesaturateFactor);
+				}
 			}
 		}
 	}
diff --git a/Engine/Plugins/PCG/Source/PCGEditor/Private/PCGEditorGraphSchema.h b/Engine/Plugins/PCG/Source/PCGEditor/Private/PCGEditorGraphSchema.h
--- a/Engine/Plugins/PCG/Source/PCGEditor/Private/PCGEditorGraphSchema.h
+++ b/Engine/Plugins/PCG/Source/PCGEditor/Private/PCGEditorGraphSchema.h
@@ -50,5 +50,7 @@ public:
 	virtual void DetermineWiringStyle(UEdGraphPin* OutputPin, UEdGraphPin* InputPin, /*inout*/ FConnectionParams& Params) override;
 
 protected:
+	bool UpdateParamsIfDebugging(UEdGraphPin* OutputPin, UEdGraphPin* InputPin, FConnectionParams& Params);
+
 	UPCGEditorGraph* Graph;
 };
diff --git a/Engine/Plugins/PCG/Source/PCGEditor/Private/SPCGEditorGraphAttributeListView.cpp b/Engine/Plugins/PCG/Source/PCGEditor/Private/SPCGEditorGraphAttributeListView.cpp
--- a/Engine/Plugins/PCG/Source/PCGEditor/Private/SPCGEditorGraphAttributeListView.cpp
+++ b/Engine/Plugins/PCG/Source/PCGEditor/Private/SPCGEditorGraphAttributeListView.cpp
@@ -4,6 +4,7 @@
 
 #include "PCGComponent.h"
 #include "PCGData.h"
+#include "PCGGraph.h"
 #include "PCGNode.h"
 #include "PCGParamData.h"
 #include "PCGPin.h"
@@ -18,6 +19,7 @@
 #include "Metadata/Accessors/PCGCustomAccessor.h"
 
 #include "PCGEditor.h"
+#include "PCGEditorGraph.h"
 #include "PCGEditorGraphNodeBase.h"
 
 #include "Fonts/FontMeasure.h"
@@ -641,6 +643,13 @@ const FPCGDataCollection* SPCGEditorGraphAttributeListView::GetInspectionData()
 		return nullptr;
 	}
 
+	PCGEditorGraphUtils::GetInspectablePin(PCGNode, Pin, PCGNode, Pin);
+
+	if (!PCGNode || !Pin)
+	{
+		return nullptr;
+	}
+
 	const TSharedPtr<FPCGEditor> PCGEditor = PCGEditorPtr.Pin();
 	const FPCGStack* PCGStack = PCGEditor->GetStackBeingInspected();
 	if (!PCGStack)
diff --git a/Engine/Plugins/PCG/Source/PCGEditor/Public/PCGEditorSettings.h b/Engine/Plugins/PCG/Source/PCGEditor/Public/PCGEditorSettings.h
--- a/Engine/Plugins/PCG/Source/PCGEditor/Public/PCGEditorSettings.h
+++ b/Engine/Plugins/PCG/Source/PCGEditor/Public/PCGEditorSettings.h
@@ -149,6 +149,18 @@ public:
 	UPROPERTY(EditAnywhere, config, Category = Node, meta = (HideAlphaChannel))
 	TMap<TSubclassOf<UPCGSettings>, FLinearColor> OverrideNodeColorByClass;
 
+	/** Hover edge size emphasis */
+	UPROPERTY(EditAnywhere, config, Category = Wire, meta = (ClampMin="1", ClampMax="5"))
+	float HoverEdgeEmphasis = 1.5f;
+
+	/** Debug mode edge emphasis on multi-data */
+	UPROPERTY(EditAnywhere, config, Category = Wire, meta = (ClampMin="1", ClampMax="5"))
+	float MultiDataEdgeDebugEmphasis = 2.5f;
+
+	/** Debug mode edge desaturation with empty data */
+	UPROPERTY(EditAnywhere, config, Category = Wire, meta = (ClampMin="0", ClampMax="1"))
+	float EmptyEdgeDebugDesaturateFactor = 0.5f;
+
 	/** Specify if we want to jump to definition in case of double click on native PCG Nodes */
 	UPROPERTY(EditAnywhere, config, Category = Workflow)
 	bool bEnableNavigateToNativeNodes;
