[PCG] Disallow to have graphs instances depending on themselves.

* Pattern like A depends on B and B depends on A is now not possible.
* Also small fix to make sure we clean up any inavlid edges if invalid nodes were detected.

#rnx

diff --git a/Engine/Plugins/PCG/Source/PCG/Private/PCGGraph.cpp b/Engine/Plugins/PCG/Source/PCG/Private/PCGGraph.cpp
--- a/Engine/Plugins/PCG/Source/PCG/Private/PCGGraph.cpp
+++ b/Engine/Plugins/PCG/Source/PCG/Private/PCGGraph.cpp
@@ -336,11 +336,15 @@ void UPCGGraph::PostLoad()
 
 	// Ensure that all nodes are loaded (& updated their deprecated data)
 	// If a node is null (can happen if an asset was saved with a node that don't exist in the current session), we don't want to crash. So remove the faulty node and warn the user.
+	// Keep track if that ever happen to force an edge cleanup
+	bool bHasInvalidNode = false;
+
 	for (int32 i = Nodes.Num() - 1; i >= 0; --i)
 	{
 		if (!Nodes[i])
 		{
 			UE_LOG(LogPCG, Error, TEXT("Graph %s has a node that doesn't exist anymore. Check if you are missing a plugin or if you saved an asset with an old settings that was removed/renamed."), *GetPathName());
+			bHasInvalidNode = true;
 			Nodes.RemoveAtSwap(i);
 		}
 		else
@@ -396,7 +400,7 @@ void UPCGGraph::PostLoad()
 		}
 	}
 
-	if (PCGGraph::CVarFixInvalidEdgesOnPostLoad.GetValueOnAnyThread())
+	if (bHasInvalidNode || PCGGraph::CVarFixInvalidEdgesOnPostLoad.GetValueOnAnyThread())
 	{
 		FixInvalidEdges();
 	}
@@ -1673,17 +1677,17 @@ void UPCGGraphInstance::PostEditChangeProperty(FPropertyChangedEvent& PropertyCh
 	// We need to be careful and only capture `Graph` if it is our graph and not a graph parameter called `Graph`!
 	if (PropertyChangedEvent.Property->GetOwnerClass() == UPCGGraphInstance::StaticClass() && PropertyName == GET_MEMBER_NAME_CHECKED(UPCGGraphInstance, Graph))
 	{
-		// If the new graph is itself, return to the previous value
-		if (this == Graph)
+		// If the new graph hierarchy has this graph in it, return to the previous value.
+		if (Graph && !CanGraphInterfaceBeSet(Graph))
 		{
-			UE_LOG(LogPCG, Error, TEXT("Try to set the graph of a graph instance to itself, would cause infinite recursion."));
+			UE_LOG(LogPCG, Error, TEXT("Attempting to assign %s would cause infinite recursion in the graph instance hierarchy, this is not allowed."), *Graph->GetPathName());
 			Graph = PreGraphCache;
 		}
 
 		SetupCallbacks();
 
 		// No need to refresh if it is the same graph
-		if (Graph == PreGraphCache)
+		if (Graph != PreGraphCache)
 		{
 			RefreshParameters(EPCGGraphParameterEvent::GraphChanged);
 		}
@@ -1781,9 +1785,9 @@ void UPCGGraphInstance::SetupCallbacks()
 
 void UPCGGraphInstance::SetGraph(UPCGGraphInterface* InGraph)
 {
-	if (InGraph == this)
+	if (InGraph && !CanGraphInterfaceBeSet(InGraph))
 	{
-		UE_LOG(LogPCG, Error, TEXT("Try to set the graph of a graph instance to itself, would cause infinite recursion."));
+		UE_LOG(LogPCG, Error, TEXT("Attempting to assign %s would cause infinite recursion in the graph instance hierarchy, this is not allowed."), *InGraph->GetPathName());
 		return;
 	}
 
@@ -1959,6 +1963,23 @@ bool UPCGGraphInstance::IsGraphParameterOverridden(const FName PropertyName) con
 	return (ParametersOverrides.Parameters.FindPropertyDescByName(PropertyName) != nullptr);
 }
 
+bool UPCGGraphInstance::CanGraphInterfaceBeSet(const UPCGGraphInterface* GraphInterface) const
+{
+	if (GraphInterface == this)
+	{
+		return false;
+	}
+
+	const UPCGGraphInstance* GraphInstance = Cast<const UPCGGraphInstance>(GraphInterface);
+	// Can always set a normal graph (or null graph)
+	if (!GraphInstance)
+	{
+		return true;
+	}
+	
+	return CanGraphInterfaceBeSet(GraphInstance->Graph);
+}
+
 bool FPCGOverrideInstancedPropertyBag::RefreshParameters(const FInstancedPropertyBag* ParentUserParameters, EPCGGraphParameterEvent InChangeType, FName InChangedPropertyName)
 {
 	check(ParentUserParameters);
diff --git a/Engine/Plugins/PCG/Source/PCG/Public/PCGGraph.h b/Engine/Plugins/PCG/Source/PCG/Public/PCGGraph.h
--- a/Engine/Plugins/PCG/Source/PCG/Public/PCGGraph.h
+++ b/Engine/Plugins/PCG/Source/PCG/Public/PCGGraph.h
@@ -493,6 +493,13 @@ public:
 	virtual const FInstancedPropertyBag* GetUserParametersStruct() const override { return &ParametersOverrides.Parameters; }
 	virtual bool IsGraphParameterOverridden(const FName PropertyName) const override;
 
+	/** 
+	* When setting a graph instance as a base to another graph instance, we need to make sure we don't find this graph in the graph hierarchy.
+	* Otherwise it would cause infinite recursion (like A is an instance of B and B is an instance of A).
+	* This function will go up the graph hierarchy and returning false if `this` is ever encountered.
+	*/
+	bool CanGraphInterfaceBeSet(const UPCGGraphInterface* GraphInterface) const;
+
 private:
 #if WITH_EDITORONLY_DATA
 	// Transient, to keep track of the previous graph when it changed.
