[PCG] Allow Get Actor Data to trigger generation on Preview components. Bulider support for generating Preview components.

diff --git a/Engine/Plugins/PCG/Source/PCG/Private/Elements/PCGDataFromActor.cpp b/Engine/Plugins/PCG/Source/PCG/Private/Elements/PCGDataFromActor.cpp
--- a/Engine/Plugins/PCG/Source/PCG/Private/Elements/PCGDataFromActor.cpp
+++ b/Engine/Plugins/PCG/Source/PCG/Private/Elements/PCGDataFromActor.cpp
@@ -2,6 +2,7 @@
 
 #include "Elements/PCGDataFromActor.h"
 
+#include "PCGActorAndComponentMapping.h"
 #include "PCGComponent.h"
 #include "PCGCustomVersion.h"
 #include "PCGSubsystem.h"
@@ -11,6 +12,7 @@
 #include "Elements/PCGMergeElement.h"
 #include "Grid/PCGPartitionActor.h"
 #include "Helpers/PCGHelpers.h"
+#include "Utils/PCGGraphExecutionLogging.h"
 
 #include "Algo/AnyOf.h"
 #include "GameFramework/Actor.h"
@@ -409,19 +411,34 @@ bool FPCGDataFromActorElement::ExecuteInternal(FPCGContext* InContext) const
 
 	if (Context->bPerformedQuery)
 	{
+#if WITH_EDITOR
+		// Remove ignored change origins now that we've completed the wait tasks.
+		UPCGComponent* OriginalComponent = Context->SourceComponent->GetOriginalComponent();
+		for (TObjectKey<UObject>& IgnoredChangeOriginKey : Context->IgnoredChangeOrigins)
+		{
+			if (UObject* IgnoredChangeOrigin = IgnoredChangeOriginKey.ResolveObjectPtr())
+			{
+				OriginalComponent->StopIgnoringChangeOriginDuringGeneration(IgnoredChangeOrigin);
+			}
+		}
+#endif
+
 		ProcessActors(Context, Settings, Context->FoundActors);
 	}
 
 	return true;
 }
 
-void FPCGDataFromActorElement::GatherWaitTasks(AActor* FoundActor, FPCGContext* Context, TArray<FPCGTaskId>& OutWaitTasks) const
+void FPCGDataFromActorElement::GatherWaitTasks(AActor* FoundActor, FPCGContext* InContext, TArray<FPCGTaskId>& OutWaitTasks) const
 {
 	if (!FoundActor)
 	{
 		return;
 	}
 
+	FPCGDataFromActorContext* Context = static_cast<FPCGDataFromActorContext*>(InContext);
+	check(Context);
+
 	const UPCGDataFromActorSettings* Settings = Context->GetInputSettings<UPCGDataFromActorSettings>();
 	check(Settings);
 
@@ -446,10 +463,37 @@ void FPCGDataFromActorElement::GatherWaitTasks(AActor* FoundActor, FPCGContext*
 
 	for (UPCGComponent* Component : PCGComponents)
 	{
-		if (Component->IsGenerating() && Component->GetOwner() != ThisOwner)
+		if (Component->GetOwner() == ThisOwner)
+		{
+			continue;
+		}
+
+		if (Component->IsGenerating())
 		{
 			OutWaitTasks.Add(Component->GetGenerationTaskId());
 		}
+		else if (!Component->bGenerated && (Component->GetSerializedEditingMode() == EPCGEditorDirtyMode::Preview) && Component->GetOwner())
+		{
+#if WITH_EDITOR
+			// Signal that any change notifications from generating upstream component should not trigger re-executions of this component.
+			// Such change notifications can cancel the current execution.
+			// Note: Uses owner because FPCGActorAndComponentMapping::OnPCGGraphGeneratedOrCleaned reports change on owner.
+			SourceComponent->GetOriginalComponent()->StartIgnoringChangeOriginDuringGeneration(Component->GetOwner());
+			Context->IgnoredChangeOrigins.Add(Component->GetOwner());
+#endif
+
+			const FPCGTaskId GenerateTask = Component->GenerateLocalGetTaskId(EPCGComponentGenerationTrigger::GenerateOnDemand, /*bForce=*/false);
+			if (GenerateTask != InvalidPCGTaskId)
+			{
+				PCGGraphExecutionLogging::LogGraphScheduleDependency(Component);
+
+				OutWaitTasks.Add(GenerateTask);
+			}
+			else
+			{
+				PCGGraphExecutionLogging::LogGraphScheduleDependencyFailed(Component);
+			}
+		}
 	}
 }
 
diff --git a/Engine/Plugins/PCG/Source/PCG/Private/Graph/PCGGraphExecutor.cpp b/Engine/Plugins/PCG/Source/PCG/Private/Graph/PCGGraphExecutor.cpp
--- a/Engine/Plugins/PCG/Source/PCG/Private/Graph/PCGGraphExecutor.cpp
+++ b/Engine/Plugins/PCG/Source/PCG/Private/Graph/PCGGraphExecutor.cpp
@@ -283,6 +283,8 @@ TSet<UPCGComponent*> FPCGGraphExecutor::Cancel(TFunctionRef<bool(TWeakObjectPtr<
 		return CancelledComponents;
 	}
 
+	PCGGraphExecutionLogging::LogComponentCancellation(CancelledComponents);
+
 	auto TryAbortScheduledTasks = [](FPCGGraphScheduleTask& ScheduledTask)
 	{
 		if (ScheduledTask.bHasAbortCallbacks)
diff --git a/Engine/Plugins/PCG/Source/PCG/Private/PCGActorAndComponentMapping.cpp b/Engine/Plugins/PCG/Source/PCG/Private/PCGActorAndComponentMapping.cpp
--- a/Engine/Plugins/PCG/Source/PCG/Private/PCGActorAndComponentMapping.cpp
+++ b/Engine/Plugins/PCG/Source/PCG/Private/PCGActorAndComponentMapping.cpp
@@ -10,6 +10,7 @@
 #include "Grid/PCGPartitionActor.h"
 #include "Helpers/PCGActorHelpers.h"
 #include "Helpers/PCGHelpers.h"
+#include "Utils/PCGGraphExecutionLogging.h"
 
 #include "Landscape.h"
 #include "LandscapeProxy.h"
@@ -1672,6 +1673,12 @@ void FPCGActorAndComponentMapping::OnObjectChanged(UObject* InObject, const FAct
 				continue;
 			}
 
+			if (PCGComponent->GetOriginalComponent()->IsIgnoringChangeOrigin(InObject))
+			{
+				PCGGraphExecutionLogging::LogChangeOriginIgnoredForComponent(InObject, PCGComponent);
+				continue;
+			}
+
 			// Don't mark "Owner changed" if the change originate from a PCG Component. It will be delegated to the ClearCacheForActor.
 			// It is necessary to avoid infine loops when there are multiple PCG components on one actor, and one component was generated.
 			const bool bOwnerChanged = (PCGComponent->GetOwner() == InObject) && (!InOriginatingChangeObject || !InOriginatingChangeObject->IsA<UPCGComponent>());
@@ -1714,6 +1721,12 @@ void FPCGActorAndComponentMapping::OnObjectChanged(UObject* InObject, const FAct
 				return;
 			}
 
+			if (ComponentRef.Component->GetOriginalComponent()->IsIgnoringChangeOrigin(InObject))
+			{
+				PCGGraphExecutionLogging::LogChangeOriginIgnoredForComponent(InObject, ComponentRef.Component);
+				return;
+			}
+
 			bool bShouldDirty = ComponentRef.Component->ShouldTrackLandscape() && InObject->IsA<ALandscapeProxy>();
 			bShouldDirty |= ClearCacheForKeys(MatchedKeys, ComponentRef.Component, /*bIntersect=*/true, InOriginatingChangeObject);
 			if (bShouldDirty)
@@ -1741,6 +1754,12 @@ void FPCGActorAndComponentMapping::OnObjectChanged(UObject* InObject, const FAct
 				return;
 			}
 
+			if (ComponentRef.Component->GetOriginalComponent()->IsIgnoringChangeOrigin(InObject))
+			{
+				PCGGraphExecutionLogging::LogChangeOriginIgnoredForComponent(InObject, ComponentRef.Component);
+				return;
+			}
+
 			check(CurrentActorBoundsPtr);
 			const FBox ComponentBounds = ComponentRef.Bounds.GetBox();
 			const bool bIntersect = CurrentActorBoundsPtr->Intersect(ComponentBounds);
diff --git a/Engine/Plugins/PCG/Source/PCG/Private/PCGComponent.cpp b/Engine/Plugins/PCG/Source/PCG/Private/PCGComponent.cpp
--- a/Engine/Plugins/PCG/Source/PCG/Private/PCGComponent.cpp
+++ b/Engine/Plugins/PCG/Source/PCG/Private/PCGComponent.cpp
@@ -31,6 +31,7 @@
 #include "RuntimeGen/SchedulingPolicies/PCGSchedulingPolicyBase.h"
 #include "RuntimeGen/SchedulingPolicies/PCGSchedulingPolicyDistanceAndDirection.h"
 #include "Utils/PCGGeneratedResourcesLogging.h"
+#include "Utils/PCGGraphExecutionLogging.h"
 
 #include "LandscapeComponent.h"
 #include "LandscapeProxy.h"
@@ -384,7 +385,7 @@ FPCGTaskId UPCGComponent::CreateGenerateTask(bool bForce, const TArray<FPCGTaskI
 
 void UPCGComponent::PostProcessGraph(const FBox& InNewBounds, bool bInGenerated, FPCGContext* Context)
 {
-	PCGGeneratedResourcesLogging::LogPostProcessGraph();
+	PCGGraphExecutionLogging::LogPostProcessGraph(this);
 
 	LastGeneratedBounds = InNewBounds;
 
@@ -394,6 +395,10 @@ void UPCGComponent::PostProcessGraph(const FBox& InNewBounds, bool bInGenerated,
 
 	GeneratedGraphOutput.Reset();
 
+#if WITH_EDITOR
+	ResetIgnoredChangeOrigins(/*bLogIfAnyPresent=*/true);
+#endif
+
 	if (bInGenerated)
 	{
 		bGenerated = true;
@@ -545,6 +550,11 @@ void UPCGComponent::OnProcessGraphAborted(bool bQuiet)
 		UE_LOG(LogPCG, Warning, TEXT("Process Graph was called but aborted, check for errors in log if you expected a result."));
 	}
 
+#if WITH_EDITOR
+	// On abort, there may be ignores still registered, silently remove these.
+	ResetIgnoredChangeOrigins(/*bLogIfAnyPresent=*/false);
+#endif
+
 	CleanupUnusedManagedResources();
 
 	CurrentGenerationTask = InvalidPCGTaskId;
@@ -3007,6 +3017,63 @@ void UPCGComponent::UpdateDynamicTracking()
 	}
 }
 
+void UPCGComponent::StartIgnoringChangeOriginDuringGeneration(UObject* InChangeOriginToIgnore)
+{
+	FWriteScopeLock Lock(IgnoredChangeOriginsLock);
+
+	if (int32* FoundCounter = IgnoredChangeOriginsToCounters.Find(InChangeOriginToIgnore))
+	{
+		int32& Counter = *FoundCounter; // Put in local variable to evade SA warnings.
+		ensure(Counter >= 0);
+
+		Counter = FMath::Max(0, Counter) + 1;
+	}
+	else
+	{
+
+		IgnoredChangeOriginsToCounters.Add(InChangeOriginToIgnore, 1);
+	}
+}
+
+void UPCGComponent::StopIgnoringChangeOriginDuringGeneration(UObject* InChangeOriginToIgnore)
+{
+	FWriteScopeLock Lock(IgnoredChangeOriginsLock);
+
+	int32* FoundCounter = IgnoredChangeOriginsToCounters.Find(InChangeOriginToIgnore);
+	if (ensure(FoundCounter))
+	{
+		int32& Counter = *FoundCounter; // Put in local variable to evade SA warnings.
+		ensure(Counter > 0);
+
+		if (--Counter <= 0)
+		{
+			IgnoredChangeOriginsToCounters.Remove(InChangeOriginToIgnore);
+		}
+	}
+}
+
+bool UPCGComponent::IsIgnoringChangeOrigin(UObject* InChangeOrigin)
+{
+	FReadScopeLock Lock(IgnoredChangeOriginsLock);
+	const int32* Counter = IgnoredChangeOriginsToCounters.Find(InChangeOrigin);
+	return Counter && ensure(*Counter > 0);
+}
+
+void UPCGComponent::ResetIgnoredChangeOrigins(bool bLogIfAnyPresent)
+{
+	FWriteScopeLock Lock(IgnoredChangeOriginsLock);
+
+	if (bLogIfAnyPresent && !IgnoredChangeOriginsToCounters.IsEmpty())
+	{
+		UE_LOG(LogPCG, Warning, TEXT("[%s/%s] ResetIgnoredChangeOrigins: IgnoredChangeOrigins should be empty but %d found, purged."),
+			GetOwner() ? *GetOwner()->GetName() : TEXT("MISSINGACTOR"),
+			*GetName(),
+			IgnoredChangeOriginsToCounters.Num());
+	}
+
+	IgnoredChangeOriginsToCounters.Reset();
+}
+
 #endif // WITH_EDITOR
 
 void UPCGComponent::SetManagedResources(const TArray<TObjectPtr<UPCGManagedResource>>& Resources)
diff --git a/Engine/Plugins/PCG/Source/PCG/Private/Utils/PCGGeneratedResourcesLogging.cpp b/Engine/Plugins/PCG/Source/PCG/Private/Utils/PCGGeneratedResourcesLogging.cpp
--- a/Engine/Plugins/PCG/Source/PCG/Private/Utils/PCGGeneratedResourcesLogging.cpp
+++ b/Engine/Plugins/PCG/Source/PCG/Private/Utils/PCGGeneratedResourcesLogging.cpp
@@ -182,18 +182,6 @@ namespace PCGGeneratedResourcesLogging
 #endif
 	}
 
-	void LogPostProcessGraph()
-	{
-#if WITH_EDITOR
-		if (!LogEnabled())
-		{
-			return;
-		}
-
-		UE_LOG(LogPCG, Log, TEXT("[PCGMANAGEDRESOURCES] UPCGComponent::PostProcessGraph"));
-#endif
-	}
-
 	void LogCreateCleanupTask(bool bRemoveComponents)
 	{
 #if WITH_EDITOR
diff --git a/Engine/Plugins/PCG/Source/PCG/Private/Utils/PCGGeneratedResourcesLogging.h b/Engine/Plugins/PCG/Source/PCG/Private/Utils/PCGGeneratedResourcesLogging.h
--- a/Engine/Plugins/PCG/Source/PCG/Private/Utils/PCGGeneratedResourcesLogging.h
+++ b/Engine/Plugins/PCG/Source/PCG/Private/Utils/PCGGeneratedResourcesLogging.h
@@ -20,8 +20,6 @@ namespace PCGGeneratedResourcesLogging
 	void LogCleanupLocalImmediateResource(const UPCGManagedResource* Resource);
 	void LogCleanupLocalImmediateFinished(const TArray<UPCGManagedResource*>& GeneratedResources);
 	
-	void LogPostProcessGraph();
-
 	void LogCreateCleanupTask(bool bRemoveComponents);
 	void LogCreateCleanupTaskResource(const UPCGManagedResource* Resource);
 	void LogCreateCleanupTaskFinished(const TArray<UPCGManagedResource*>& GeneratedResources);
diff --git a/Engine/Plugins/PCG/Source/PCG/Private/Utils/PCGGraphExecutionLogging.cpp b/Engine/Plugins/PCG/Source/PCG/Private/Utils/PCGGraphExecutionLogging.cpp
--- a/Engine/Plugins/PCG/Source/PCG/Private/Utils/PCGGraphExecutionLogging.cpp
+++ b/Engine/Plugins/PCG/Source/PCG/Private/Utils/PCGGraphExecutionLogging.cpp
@@ -4,6 +4,7 @@
 
 #include "PCGComponent.h"
 #include "PCGContext.h"
+#include "PCGGraph.h"
 #include "PCGModule.h"
 #include "Graph/PCGGraphExecutor.h"
 
@@ -129,11 +130,40 @@ namespace PCGGraphExecutionLogging
 			return;
 		}
 
-		UE_LOG(LogPCG, Log, TEXT("[%s] --- SCHEDULE GRAPH ---"),
-			(SourceComponent && SourceComponent->GetOwner()) ? *SourceComponent->GetOwner()->GetName() : TEXT("MISSINGCOMPONENT"));
+		UE_LOG(LogPCG, Log, TEXT("[%s/%s] --- SCHEDULE GRAPH ---"),
+			(SourceComponent && SourceComponent->GetOwner()) ? *SourceComponent->GetOwner()->GetName() : TEXT("MISSINGCOMPONENT"),
+			(SourceComponent && SourceComponent->GetGraph()) ? *SourceComponent->GetGraph()->GetName() : TEXT("MISSINGGRAPH"));
 #endif
 	}
 
+	void LogGraphScheduleDependency(const UPCGComponent* InComponent)
+	{
+#if WITH_EDITOR
+		if (!LogEnabled())
+		{
+			return;
+		}
+
+		UE_LOG(LogPCG, Log, TEXT("[%s/%s] --- SCHEDULE GRAPH FOR DEPENDENCY ---"),
+			(InComponent && InComponent->GetOwner()) ? *InComponent->GetOwner()->GetName() : TEXT("MISSINGCOMPONENT"),
+			(InComponent && InComponent->GetGraph()) ? *InComponent->GetGraph()->GetName() : TEXT("MISSINGGRAPH"));
+#endif
+	}
+
+	void LogGraphScheduleDependencyFailed(const UPCGComponent* InComponent)
+	{
+#if WITH_EDITOR
+		if (!LogEnabled())
+		{
+			return;
+		}
+
+		UE_LOG(LogPCG, Error, TEXT("[%s/%s] Failed to schedule dependency"),
+			(InComponent && InComponent->GetOwner()) ? *InComponent->GetOwner()->GetName() : TEXT("MISSINGCOMPONENT"),
+			(InComponent && InComponent->GetGraph()) ? *InComponent->GetGraph()->GetName() : TEXT("MISSINGGRAPH"));
+#endif
+	}
+	
 	void LogGraphPostSchedule(const TMap<FPCGTaskId, FPCGGraphTask>& Tasks, const TMap<FPCGTaskId, TSet<FPCGTaskId>>& TaskSuccessors)
 	{
 #if WITH_EDITOR
@@ -148,6 +178,52 @@ namespace PCGGraphExecutionLogging
 #endif
 	}
 
+	void LogPostProcessGraph(const UPCGComponent* InSourceComponent)
+	{
+#if WITH_EDITOR
+		if (!LogEnabled())
+		{
+			return;
+		}
+
+		UE_LOG(LogPCG, Log, TEXT("[%s/%s] UPCGComponent::PostProcessGraph"),
+			(InSourceComponent && InSourceComponent->GetOwner()) ? *InSourceComponent->GetOwner()->GetName() : TEXT("MISSINGCOMPONENT"),
+			(InSourceComponent && InSourceComponent->GetGraph()) ? *InSourceComponent->GetGraph()->GetName() : TEXT("MISSINGGRAPH"));
+#endif
+	}
+
+	void LogComponentCancellation(const TSet<UPCGComponent*>& CancelledComponents)
+	{
+#if WITH_EDITOR
+		if (!LogEnabled())
+		{
+			return;
+		}
+
+		for (const UPCGComponent* Component : CancelledComponents)
+		{
+			UE_LOG(LogPCG, Log, TEXT("[%s/%s] Component cancelled"),
+				(Component && Component->GetOwner()) ? *Component->GetOwner()->GetName() : TEXT("MISSINGCOMPONENT"),
+				(Component && Component->GetGraph()) ? *Component->GetGraph()->GetName() : TEXT("MISSINGGRAPH"));
+		}
+#endif
+	}
+
+	void LogChangeOriginIgnoredForComponent(const UObject* InObject, const UPCGComponent* InComponent)
+	{
+#if WITH_EDITOR
+		if (!LogEnabled())
+		{
+			return;
+		}
+
+		UE_LOG(LogPCG, Log, TEXT("[%s/%s] Change origin ignored: '%s'"),
+			(InComponent && InComponent->GetOwner()) ? *InComponent->GetOwner()->GetName() : TEXT("MISSINGCOMPONENT"),
+			(InComponent && InComponent->GetGraph()) ? *InComponent->GetGraph()->GetName() : TEXT("MISSINGGRAPH"),
+			InObject ? *InObject->GetName() : TEXT("MISSINGOBJECT"));
+#endif
+	}
+
 	void LogGraphExecuteFrameFinished()
 	{
 #if WITH_EDITOR
@@ -186,8 +262,9 @@ namespace PCGGraphExecutionLogging
 			return;
 		}
 
-		UE_LOG(LogPCG, Log, TEXT("         [%s] %s\t\tEXECUTE"),
+		UE_LOG(LogPCG, Log, TEXT("         [%s/%s] %s\t\tEXECUTE"),
 			*Task.SourceComponent->GetOwner()->GetName(),
+			Task.SourceComponent->GetGraph() ? *Task.SourceComponent->GetGraph()->GetName() : TEXT("MISSINGGRAPH"),
 			*FString::Printf(TEXT("%u'%s'"), Task.NodeId, Task.Node ? *Task.Node->GetNodeTitle(EPCGNodeTitleType::ListView).ToString() : TEXT("")));
 #endif
 	}
@@ -200,8 +277,9 @@ namespace PCGGraphExecutionLogging
 			return;
 		}
 
-		UE_LOG(LogPCG, Warning, TEXT("[%s] %s\t\tCACHING DISABLED"),
+		UE_LOG(LogPCG, Warning, TEXT("[%s/%s] %s\t\tCACHING DISABLED"),
 			*Task.SourceComponent->GetOwner()->GetName(),
+			Task.SourceComponent->GetGraph() ? *Task.SourceComponent->GetGraph()->GetName() : TEXT("MISSINGGRAPH"),
 			*FString::Printf(TEXT("%u'%s'"), Task.NodeId, Task.Node ? *Task.Node->GetNodeTitle(EPCGNodeTitleType::ListView).ToString() : TEXT("")));
 #endif
 	}
diff --git a/Engine/Plugins/PCG/Source/PCG/Private/Utils/PCGGraphExecutionLogging.h b/Engine/Plugins/PCG/Source/PCG/Private/Utils/PCGGraphExecutionLogging.h
--- a/Engine/Plugins/PCG/Source/PCG/Private/Utils/PCGGraphExecutionLogging.h
+++ b/Engine/Plugins/PCG/Source/PCG/Private/Utils/PCGGraphExecutionLogging.h
@@ -24,7 +24,16 @@ namespace PCGGraphExecutionLogging
 	void LogGraphTasks(const TArray<FPCGGraphTask>& Tasks);
 
 	void LogGraphSchedule(const UPCGComponent* SourceComponent);
+	void LogGraphScheduleDependency(const UPCGComponent* InComponent);
+	void LogGraphScheduleDependencyFailed(const UPCGComponent* InComponent);
+
 	void LogGraphPostSchedule(const TMap<FPCGTaskId, FPCGGraphTask>& Tasks, const TMap<FPCGTaskId, TSet<FPCGTaskId>>& TaskSuccessors);
+	void LogPostProcessGraph(const UPCGComponent* InSourceComponent);
+
+	void LogComponentCancellation(const TSet<UPCGComponent*>& CancelledComponents);
+
+	void LogChangeOriginIgnoredForComponent(const UObject* InObject, const UPCGComponent* InComponent);
+
 	void LogGraphExecuteFrameFinished();
 
 	void LogTaskExecute(const FPCGGraphTask& Task);
diff --git a/Engine/Plugins/PCG/Source/PCG/Public/Elements/PCGDataFromActor.h b/Engine/Plugins/PCG/Source/PCG/Public/Elements/PCGDataFromActor.h
--- a/Engine/Plugins/PCG/Source/PCG/Public/Elements/PCGDataFromActor.h
+++ b/Engine/Plugins/PCG/Source/PCG/Public/Elements/PCGDataFromActor.h
@@ -6,6 +6,8 @@
 #include "PCGSettings.h"
 #include "Elements/PCGActorSelector.h"
 
+#include "UObject/ObjectKey.h"
+
 #include "PCGDataFromActor.generated.h"
 
 UENUM()
@@ -115,6 +117,11 @@ struct FPCGDataFromActorContext : public FPCGContext
 {
 	TArray<AActor*> FoundActors;
 	bool bPerformedQuery = false;
+
+#if WITH_EDITOR
+	/** Any change origin ignores we added, to solve dependency issues (like upstream execution cancelling downstream graph). */
+	TArray<TObjectKey<UObject>> IgnoredChangeOrigins;
+#endif
 };
 
 class FPCGDataFromActorElement : public IPCGElement
@@ -127,7 +134,7 @@ public:
 protected:
 	virtual FPCGContext* CreateContext() override;
 	virtual bool ExecuteInternal(FPCGContext* Context) const override;
-	void GatherWaitTasks(AActor* FoundActor, FPCGContext* Context, TArray<FPCGTaskId>& OutWaitTasks) const;
+	void GatherWaitTasks(AActor* FoundActor, FPCGContext* InContext, TArray<FPCGTaskId>& OutWaitTasks) const;
 	virtual void ProcessActors(FPCGContext* Context, const UPCGDataFromActorSettings* Settings, const TArray<AActor*>& FoundActors) const;
 	virtual void ProcessActor(FPCGContext* Context, const UPCGDataFromActorSettings* Settings, AActor* FoundActor) const;
 
diff --git a/Engine/Plugins/PCG/Source/PCG/Public/PCGComponent.h b/Engine/Plugins/PCG/Source/PCG/Public/PCGComponent.h
--- a/Engine/Plugins/PCG/Source/PCG/Public/PCGComponent.h
+++ b/Engine/Plugins/PCG/Source/PCG/Public/PCGComponent.h
@@ -341,6 +341,12 @@ public:
 
 	/** To be called by an element to notify the component that this settings have a dynamic dependency. */
 	void RegisterDynamicTracking(const UPCGSettings* InSettings, const TArrayView<TPair<FPCGSelectionKey, bool>>& InDynamicKeysAndCulling);
+
+	/** For duration of the current/next generation, any change triggers from this change origin will be discarded. */
+	void StartIgnoringChangeOriginDuringGeneration(UObject* InChangeOriginToIgnore);
+	void StopIgnoringChangeOriginDuringGeneration(UObject* InChangeOriginToIgnore);
+	bool IsIgnoringChangeOrigin(UObject* InChangeOrigin);
+	void ResetIgnoredChangeOrigins(bool bLogIfAnyPresent);
 #endif
 
 	/** Utility function (mostly for tests) to properly set the value of bIsComponentPartitioned.
@@ -583,6 +589,12 @@ private:
 	/** Map from nodes to stacks to mask of output pins that were deactivated during execution. */
 	mutable TMap<TObjectKey<const UPCGNode>, TMap<const FPCGStack, uint64>> NodeToStackToInactivePinMask;
 	mutable FRWLock NodeToStackToInactivePinMaskLock;
+
+	/** The tracking system will not trigger a generation on this component for these change origins. Populated within the scope
+	* of an element. Entries removed when counter is decremented to 0, so empty map means no active ignores.
+	*/
+	TMap<TObjectKey<UObject>, int32> IgnoredChangeOriginsToCounters;
+	FRWLock IgnoredChangeOriginsLock;
 #endif
 
 	mutable FCriticalSection GeneratedResourcesLock;
diff --git a/Engine/Plugins/PCG/Source/PCGEditor/Private/WorldPartitionBuilder/PCGWorldPartitionBuilder.cpp b/Engine/Plugins/PCG/Source/PCGEditor/Private/WorldPartitionBuilder/PCGWorldPartitionBuilder.cpp
--- a/Engine/Plugins/PCG/Source/PCGEditor/Private/WorldPartitionBuilder/PCGWorldPartitionBuilder.cpp
+++ b/Engine/Plugins/PCG/Source/PCGEditor/Private/WorldPartitionBuilder/PCGWorldPartitionBuilder.cpp
@@ -46,6 +46,8 @@ UPCGWorldPartitionBuilder::UPCGWorldPartitionBuilder(const FObjectInitializer& O
 	{
 		bGenerateEditingModeNormalComponents = HasParam("GenerateComponentEditingModeNormal");
 
+		bGenerateEditingModePreviewComponents = HasParam("GenerateComponentEditingModePreview");
+
 		bOneComponentAtATime = HasParam("OneComponentAtATime");
 
 		FString IncludeGraphNamesValue;
@@ -97,7 +99,8 @@ bool UPCGWorldPartitionBuilder::RunInternal(UWorld* World, const FCellInfo& InCe
 
 			// Accept based on editing mode.
 			return (InComponent->GetSerializedEditingMode() == EPCGEditorDirtyMode::LoadAsPreview)
-				|| (bGenerateEditingModeNormalComponents && InComponent->GetSerializedEditingMode() == EPCGEditorDirtyMode::Normal);
+				|| (bGenerateEditingModeNormalComponents && InComponent->GetSerializedEditingMode() == EPCGEditorDirtyMode::Normal)
+				|| (bGenerateEditingModePreviewComponents && InComponent->GetSerializedEditingMode() == EPCGEditorDirtyMode::Preview);
 		};
 
 		PCGWorldPartitionBuilder::CollectComponentsToGenerate(World, ComponentFilter, ComponentsToGenerate);
diff --git a/Engine/Plugins/PCG/Source/PCGEditor/Private/WorldPartitionBuilder/PCGWorldPartitionBuilder.h b/Engine/Plugins/PCG/Source/PCGEditor/Private/WorldPartitionBuilder/PCGWorldPartitionBuilder.h
--- a/Engine/Plugins/PCG/Source/PCGEditor/Private/WorldPartitionBuilder/PCGWorldPartitionBuilder.h
+++ b/Engine/Plugins/PCG/Source/PCGEditor/Private/WorldPartitionBuilder/PCGWorldPartitionBuilder.h
@@ -42,6 +42,9 @@ private:
 	/** Include components which have editing mode set to Normal. */
 	bool bGenerateEditingModeNormalComponents = false;
 
+	/** Include components which have editing mode set to Preview. */
+	bool bGenerateEditingModePreviewComponents = false;
+
 	/** If non empty, components with each graph name will be generated (if editing mode is included), in order. */
 	TArray<FName> IncludeGraphNames;
 
