[PCG] Defer debug object tree refreshing while a debugged component is generating, which would lead to losing the debug object selection.

This could happen if a tree happened to be refreshed at a bad time, such as just after hte dynamic execution stacks are purged. I had previously managed to prevent this happening in the past, but some recent change has thrown it out of sync again, so introducing this constraint to avoid the problem.

diff --git a/Engine/Plugins/PCG/Source/PCGEditor/Private/SPCGEditorGraphDebugObjectTree.cpp b/Engine/Plugins/PCG/Source/PCGEditor/Private/SPCGEditorGraphDebugObjectTree.cpp
--- a/Engine/Plugins/PCG/Source/PCGEditor/Private/SPCGEditorGraphDebugObjectTree.cpp
+++ b/Engine/Plugins/PCG/Source/PCGEditor/Private/SPCGEditorGraphDebugObjectTree.cpp
@@ -339,8 +339,14 @@ void SPCGEditorGraphDebugObjectTree::Tick(const FGeometry& AllottedGeometry, con
 
 	if (bNeedsRefresh)
 	{
-		bNeedsRefresh = false;
-		RefreshTree();
+		// Updating the tree while the inspected component is generating can be bad as the selection can
+		// be lost. Don't change the tree if we're inspecting something that is generating.
+		const UPCGComponent* InspectedComponent = SelectedStack.GetRootComponent();
+		if(InspectedComponent == nullptr || !InspectedComponent->IsGenerating())
+		{
+			bNeedsRefresh = false;
+			RefreshTree();
+		}
 	}
 }
 
