[PCG] Fix cleanup failing on partitioned PCG components in non-WP maps

#rnx

diff --git a/Engine/Plugins/PCG/Source/PCG/Private/PCGComponent.cpp b/Engine/Plugins/PCG/Source/PCG/Private/PCGComponent.cpp
--- a/Engine/Plugins/PCG/Source/PCG/Private/PCGComponent.cpp
+++ b/Engine/Plugins/PCG/Source/PCG/Private/PCGComponent.cpp
@@ -1271,6 +1271,11 @@ void UPCGComponent::Serialize(FArchive& Ar)
 		if (Ar.IsLoading())
 		{
 			DataVersion = Ar.CustomVer(FPCGCustomVersion::GUID);
+
+			if (DataVersion < FPCGCustomVersion::SupportPartitionedComponentsInNonPartitionedLevels)
+			{
+				bDisableIsComponentPartitionedOnLoad = true;
+			}
 		}
 
 		if (DataVersion >= FPCGCustomVersion::DynamicTrackingKeysSerializedInComponent)
@@ -1306,6 +1311,19 @@ void UPCGComponent::PostLoad()
 		bIsPartitioned = false;
 	}
 
+	// Components marked as partitioned in non-WP worlds from BEFORE partitioning was supported in non-WP worlds can leak resources. To fix this, we can just unset 'bIsComponentPartitioned'.
+	if (bDisableIsComponentPartitionedOnLoad)
+	{
+		const UWorld* World = GetOwner() ? GetOwner()->GetWorld() : nullptr;
+
+		if (IsPartitioned() && !IsManagedByRuntimeGenSystem() && World && World->GetWorldPartition() == nullptr)
+		{
+			bIsComponentPartitioned = false;
+		}
+
+		bDisableIsComponentPartitionedOnLoad = false;
+	}
+
 	/** Deprecation code, should be removed once generated data has been updated */
 	if (GetOwner() && bGenerated && GeneratedResources.Num() == 0)
 	{
diff --git a/Engine/Plugins/PCG/Source/PCG/Public/PCGComponent.h b/Engine/Plugins/PCG/Source/PCG/Public/PCGComponent.h
--- a/Engine/Plugins/PCG/Source/PCG/Public/PCGComponent.h
+++ b/Engine/Plugins/PCG/Source/PCG/Public/PCGComponent.h
@@ -427,6 +427,9 @@ private:
 	// Note for upgrade: can be safely replaced by bIsComponentPartitioned. Needed a new variable to change the default value. Kept to allow proper value change, cannot be deprecated. Do not use.
 	UPROPERTY()
 	bool bIsPartitioned = true;
+
+	/** Track if component should disable 'bIsComponentPartitioned'. Used to deprecate partitioned components in non-WP levels before support for partitioning in non-WP levels. */
+	bool bDisableIsComponentPartitionedOnLoad = false;
 #endif // WITH_EDITORONLY_DATA
 
 private:
diff --git a/Engine/Plugins/PCG/Source/PCG/Public/PCGCustomVersion.h b/Engine/Plugins/PCG/Source/PCG/Public/PCGCustomVersion.h
--- a/Engine/Plugins/PCG/Source/PCG/Public/PCGCustomVersion.h
+++ b/Engine/Plugins/PCG/Source/PCG/Public/PCGCustomVersion.h
@@ -71,6 +71,9 @@ struct PCG_API FPCGCustomVersion
 		// Added dynamic tracking to the PCG component serialization
 		DynamicTrackingKeysSerializedInComponent = 20,
 
+		// Supporting partitioned components in non-partitioned levels
+		SupportPartitionedComponentsInNonPartitionedLevels = 21,
+
 		// -----<new versions can be added above this line>-------------------------------------------------
 		VersionPlusOne,
 		LatestVersion = VersionPlusOne - 1
