[PCG] Fix display name on 'bForceEditorOnlyCPUSampling'

#rnx

diff --git a/Engine/Plugins/PCG/Source/PCG/Public/Elements/PCGTextureSampler.h b/Engine/Plugins/PCG/Source/PCG/Public/Elements/PCGTextureSampler.h
--- a/Engine/Plugins/PCG/Source/PCG/Public/Elements/PCGTextureSampler.h
+++ b/Engine/Plugins/PCG/Source/PCG/Public/Elements/PCGTextureSampler.h
@@ -102,7 +102,7 @@ public:
 	 * Reading from CPU memory will be faster and more accurate than reading from GPU memory, since the texture will not be subject to compression or resolution clamping. Enable
 	 * this flag to force a duplicate of the texture with the correct settings for CPU memory access. This is editor-only.
 	 */
-	UPROPERTY(BlueprintReadWrite, EditAnywhere, Category = "Settings")
+	UPROPERTY(BlueprintReadWrite, EditAnywhere, Category = "Settings", meta = (DisplayName = "Force Editor Only CPU Sampling"))
 	bool bForceEditorOnlyCPUSampling = false;
 #endif
 
