[PCG] Minor updates to Print String node

* Clamp print duration to min 0
* Disable print message on cleanup

#rnx

diff --git a/Engine/Plugins/PCG/Source/PCG/Private/Elements/PCGPrintElement.cpp b/Engine/Plugins/PCG/Source/PCG/Private/Elements/PCGPrintElement.cpp
--- a/Engine/Plugins/PCG/Source/PCG/Private/Elements/PCGPrintElement.cpp
+++ b/Engine/Plugins/PCG/Source/PCG/Private/Elements/PCGPrintElement.cpp
@@ -22,6 +22,28 @@ namespace PCGPrintElementConstants
 	FText Delimiter = LOCTEXT("Delimiter", "::");
 }
 
+#if WITH_EDITOR
+
+namespace PCGPrintElementHelpers
+{
+	void CleanUpMessage(uint64& OutMessageHashKey)
+	{
+		if (GEngine && (OutMessageHashKey != (uint64)-1) && GEngine->OnScreenDebugMessageExists(OutMessageHashKey))
+		{
+			GEngine->RemoveOnScreenDebugMessage(OutMessageHashKey);
+			OutMessageHashKey = (uint64)-1;
+		}
+	}
+}
+
+bool UPCGManagedDebugStringMessageKey::Release(bool bHardRelease, TSet<TSoftObjectPtr<AActor>>& OutActorsToDelete)
+{
+	PCGPrintElementHelpers::CleanUpMessage(HashKey);
+	return Super::Release(bHardRelease, OutActorsToDelete);
+}
+
+#endif // WITH_EDITOR
+
 TArray<FPCGPinProperties> UPCGPrintElementSettings::InputPinProperties() const
 {
 	TArray<FPCGPinProperties> PinProperties;
@@ -111,14 +133,22 @@ bool FPCGPrintElement::ExecuteInternal(FPCGContext* Context) const
 	{
 		if (UEditorEngine* Editor = static_cast<UEditorEngine*>(GEngine))
 		{
-			uint32 HashKey = HashCombine(GetTypeHash(Context->Node->GetFName()), Context->Node->GetUniqueID());
+			uint32 HashKey32 = HashCombine(GetTypeHash(Context->Node->GetFName()), Context->Node->GetUniqueID());
 
 			if (Settings->bPrintPerComponent)
 			{
-				HashKey = HashCombine(HashKey, Component->GetUniqueID());
+				HashKey32 = HashCombine(HashKey32, Component->GetUniqueID());
 			}
 
-			Editor->AddOnScreenDebugMessage(static_cast<uint64>(HashKey), Settings->PrintToScreenDuration, Settings->PrintToScreenColor, FinalString);
+			check(Context->SourceComponent.IsValid());
+			UPCGManagedDebugStringMessageKey* ManagedMessageKey = NewObject<UPCGManagedDebugStringMessageKey>(Context->SourceComponent.Get());
+			ManagedMessageKey->HashKey = static_cast<uint64>(HashKey32);
+
+			Context->SourceComponent->AddToManagedResources(ManagedMessageKey);
+
+			// Count 0 as infinite duration to match other debug features in PCG
+			const double Duration = Settings->PrintToScreenDuration > 0.0 ? Settings->PrintToScreenDuration : std::numeric_limits<double>::max();
+			Editor->AddOnScreenDebugMessage(ManagedMessageKey->HashKey, Duration, Settings->PrintToScreenColor, FinalString);
 		}
 	}
 #endif // WITH_EDITOR
diff --git a/Engine/Plugins/PCG/Source/PCG/Public/Elements/PCGPrintElement.h b/Engine/Plugins/PCG/Source/PCG/Public/Elements/PCGPrintElement.h
--- a/Engine/Plugins/PCG/Source/PCG/Public/Elements/PCGPrintElement.h
+++ b/Engine/Plugins/PCG/Source/PCG/Public/Elements/PCGPrintElement.h
@@ -5,6 +5,8 @@
 #include "PCGElement.h"
 #include "PCGSettings.h"
 
+#include "PCGManagedResource.h"
+
 #include "Logging/LogVerbosity.h"
 
 #include "PCGPrintElement.generated.h"
@@ -17,6 +19,23 @@ enum class EPCGPrintVerbosity : uint8
 	Error = ELogVerbosity::Error
 };
 
+/** Used to track the debug message to properly remove it upon regen or clean up. */
+UCLASS(ClassGroup = (Procedural))
+class UPCGManagedDebugStringMessageKey : public UPCGManagedResource
+{
+	GENERATED_BODY()
+
+#if WITH_EDITOR
+	// Debug should always be transient
+	virtual void ChangeTransientState(EPCGEditorDirtyMode NewEditingMode) override {}
+	virtual bool Release(bool bHardRelease, TSet<TSoftObjectPtr<AActor>>& OutActorsToDelete) override;
+	virtual bool CanBeUsed() const override { return false; }
+
+public:
+	uint64 HashKey = (uint64)-1;
+#endif // WITH_EDITOR
+};
+
 /**
  * Issues a specified message to the log, and optionally to the graph and/or screen.
  * Note: This node will not function in shipping builds.
@@ -60,7 +79,7 @@ public:
 	bool bDisplayOnNode = false;
 
 	/** Use the component as part of the key hash and print a message for each component with this node. */
-	UPROPERTY(BlueprintReadWrite, EditAnywhere, Category = Settings, meta = (EditCondition = "bPrintToScreen", EditConditionHides))
+	UPROPERTY(BlueprintReadWrite, EditAnywhere, Category = Settings, meta = (PCG_Overridable))
 	bool bPrintPerComponent = true;
 
 #if WITH_EDITORONLY_DATA
@@ -69,7 +88,7 @@ public:
 	bool bPrintToScreen = false;
 
 	/** The duration (in seconds) of the on screen message. */
-	UPROPERTY(BlueprintReadWrite, EditAnywhere, Category = Settings, meta = (PCG_Overridable, EditCondition = "bPrintToScreen", EditConditionHides))
+	UPROPERTY(BlueprintReadWrite, EditAnywhere, Category = Settings, meta = (PCG_Overridable, EditCondition = "bPrintToScreen", EditConditionHides, ClampMin = "0.0"))
 	double PrintToScreenDuration = 15.0;
 
 	/** The color of the on screen message. */
