[PCG] Make sure CleanupLocal is working even without a subsystem

* If we ever have resources on the component and we need to clean them up, we should not have to depend on the subsystem

#rnx

diff --git a/Engine/Plugins/PCG/Source/PCG/Private/PCGComponent.cpp b/Engine/Plugins/PCG/Source/PCG/Private/PCGComponent.cpp
--- a/Engine/Plugins/PCG/Source/PCG/Private/PCGComponent.cpp
+++ b/Engine/Plugins/PCG/Source/PCG/Private/PCGComponent.cpp
@@ -625,7 +625,7 @@ FPCGTaskId UPCGComponent::CleanupInternal(bool bRemoveComponents, const TArray<F
 
 void UPCGComponent::CancelGeneration()
 {
-	if (CurrentGenerationTask != InvalidPCGTaskId)
+	if (CurrentGenerationTask != InvalidPCGTaskId && GetSubsystem())
 	{
 		GetSubsystem()->CancelGeneration(this);
 	}
@@ -859,11 +859,16 @@ void UPCGComponent::CleanupLocalImmediate(bool bRemoveComponents, bool bCleanupL
 	PCGGeneratedResourcesLogging::LogCleanupLocalImmediate(bRemoveComponents, GeneratedResources);
 
 	UPCGSubsystem* Subsystem = GetSubsystem();
-	check(Subsystem);
 
-	bool bHasUnbounded = false;
-	PCGHiGenGrid::FSizeArray GridSizes;
-	ensure(PCGHelpers::GetGenerationGridSizes(GetGraph(), Subsystem->GetPCGWorldActor(), GridSizes, bHasUnbounded));
+	// Cleanup Local should work even if we don't have any subsytem. In cook (or in other places), if Cleanup is necessary, we need to make sure to 
+	// cleanup the managed resources on the component even if we don't have a subsystem. If we don't have a subsystem, assume we are unbounded.
+	bool bHasUnbounded = true;
+
+	if (Subsystem)
+	{
+		PCGHiGenGrid::FSizeArray GridSizes;
+		ensure(PCGHelpers::GetGenerationGridSizes(GetGraph(), Subsystem->GetPCGWorldActor(), GridSizes, bHasUnbounded));
+	}
 
 	// Cleanup original component if non-partitioned, or if it has nodes that will execute at the Unbounded level.
 	if (!IsPartitioned() || bHasUnbounded)
@@ -920,7 +925,7 @@ void UPCGComponent::CleanupLocalImmediate(bool bRemoveComponents, bool bCleanupL
 	}
 
 	// If the component is partitioned, we will forward the calls to its local components.
-	if (bCleanupLocalComponents && IsPartitioned())
+	if (Subsystem && bCleanupLocalComponents && IsPartitioned())
 	{
 		Subsystem->CleanupLocalComponentsImmediate(this, bRemoveComponents);
 	}
