[PCG] Add some fallback if an ObjectPtr is overriden by something that has not the same class

#rnx

diff --git a/Engine/Plugins/PCG/Source/PCG/Private/PCGContext.cpp b/Engine/Plugins/PCG/Source/PCG/Private/PCGContext.cpp
--- a/Engine/Plugins/PCG/Source/PCG/Private/PCGContext.cpp
+++ b/Engine/Plugins/PCG/Source/PCG/Private/PCGContext.cpp
@@ -230,6 +230,38 @@ void FPCGContext::OverrideSettings()
 			FPCGAttributeAccessorKeysSingleObjectPtr PropertyObjectKey(Container);
 			PropertyAccessor->Set<PropertyType>(Value, PropertyObjectKey);
 
+			// Add some validation on Object Ptr overrides, to make sure that the object/class stored there has the right class
+			if constexpr (std::is_base_of_v<FSoftObjectPath, PropertyType>)
+			{
+				if (const FObjectPropertyBase* ObjectProperty = CastField<const FObjectPropertyBase>(Param.Properties.Last()))
+				{
+					if (PropertyAccessor->Get<PropertyType>(Value, PropertyObjectKey))
+					{
+						bool bInvalid = false;
+						if constexpr (std::is_same_v<FSoftObjectPath, PropertyType>)
+						{
+							if (const UObject* Object = Value.ResolveObject())
+							{
+								bInvalid = ObjectProperty->PropertyClass && Object->GetClass() && !Object->GetClass()->IsChildOf(ObjectProperty->PropertyClass);
+							}
+						}
+						else if constexpr (std::is_same_v<FSoftClassPath, PropertyType>)
+						{
+							if (const UClass* Class = Cast<UClass>(Value.ResolveObject()))
+							{
+								bInvalid = ObjectProperty->PropertyClass && !Class->IsChildOf(ObjectProperty->PropertyClass);
+							}
+						}
+
+						if (bInvalid)
+						{
+							PCGE_LOG_C(Error, GraphAndLog, this, FText::Format(LOCTEXT("WrongObjectClass", "Parameter '{0}' was set with an attribute that is not a child class. It will be nulled out."), FText::FromName(Param.Label)));
+							PropertyAccessor->Set<PropertyType>(PropertyType{}, PropertyObjectKey);
+						}
+					}
+				}
+			}
+
 			return true;
 		});
 
