[PCG] Add support for dropping Graph Instances in graphs

#rnx

diff --git a/Engine/Plugins/PCG/Source/PCGEditor/Private/PCGEditorGraphSchema.cpp b/Engine/Plugins/PCG/Source/PCGEditor/Private/PCGEditorGraphSchema.cpp
--- a/Engine/Plugins/PCG/Source/PCGEditor/Private/PCGEditorGraphSchema.cpp
+++ b/Engine/Plugins/PCG/Source/PCGEditor/Private/PCGEditorGraphSchema.cpp
@@ -604,7 +604,7 @@ void UPCGEditorGraphSchema::DroppedAssetsOnGraph(const TArray<FAssetData>& Asset
 	{
 		if (const UObject* Asset = AssetData.GetAsset())
 		{
-			if (Asset->IsA<UPCGGraph>())
+			if (Asset->IsA<UPCGGraphInterface>())
 			{
 				FPCGEditorGraphSchemaAction_NewSubgraphElement NewSubgraphAction;
 				NewSubgraphAction.SubgraphObjectPath = AssetData.GetSoftObjectPath();
@@ -649,7 +649,7 @@ void UPCGEditorGraphSchema::GetAssetsGraphHoverMessage(const TArray<FAssetData>&
 	{
 		if (const UObject* Asset = AssetData.GetAsset())
 		{
-			if (Asset->IsA<UPCGGraph>() || Asset->IsA<UPCGSettings>() || PCGEditorUtils::IsAssetPCGBlueprint(AssetData))
+			if (Asset->IsA<UPCGGraphInterface>() || Asset->IsA<UPCGSettings>() || PCGEditorUtils::IsAssetPCGBlueprint(AssetData))
 			{
 				OutOkIcon = true;
 				return;
diff --git a/Engine/Plugins/PCG/Source/PCGEditor/Private/PCGEditorGraphSchemaActions.cpp b/Engine/Plugins/PCG/Source/PCGEditor/Private/PCGEditorGraphSchemaActions.cpp
--- a/Engine/Plugins/PCG/Source/PCGEditor/Private/PCGEditorGraphSchemaActions.cpp
+++ b/Engine/Plugins/PCG/Source/PCGEditor/Private/PCGEditorGraphSchemaActions.cpp
@@ -362,7 +362,7 @@ UEdGraphNode* FPCGEditorGraphSchemaAction_NewSubgraphElement::PerformAction(UEdG
 	// Important - do not reconstruct the editor graph node/pins midway through this function as this will invalidate FromPin.
 	const FPCGDeferNodeReconstructScope DisableReconstruct(FromPin);
 
-	UPCGGraph* Subgraph = CastChecked<UPCGGraph>(SubgraphObjectPath.TryLoad());
+	UPCGGraphInterface* Subgraph = CastChecked<UPCGGraphInterface>(SubgraphObjectPath.TryLoad());
 
 	const FScopedTransaction Transaction(*FPCGEditorCommon::ContextIdentifier, LOCTEXT("PCGEditorNewSubgraphElement", "PCG Editor: New Subgraph Element"), nullptr);
 	EditorGraph->Modify();
