[PCG] Enable Diff details on PCGGraph and PCGGraphInstance

#rnx

diff --git a/Engine/Plugins/PCG/Source/PCGEditor/Private/AssetDefinitions/AssetDefinition_PCGGraph.cpp b/Engine/Plugins/PCG/Source/PCGEditor/Private/AssetDefinitions/AssetDefinition_PCGGraph.cpp
--- a/Engine/Plugins/PCG/Source/PCGEditor/Private/AssetDefinitions/AssetDefinition_PCGGraph.cpp
+++ b/Engine/Plugins/PCG/Source/PCGEditor/Private/AssetDefinitions/AssetDefinition_PCGGraph.cpp
@@ -5,6 +5,8 @@
 #include "PCGEditor.h"
 #include "PCGGraph.h"
 
+#include "SDetailsDiff.h"
+
 #define LOCTEXT_NAMESPACE "AssetDefinition_PCGGraph"
 
 FText UAssetDefinition_PCGGraph::GetAssetDisplayName() const
@@ -28,4 +30,18 @@ EAssetCommandResult UAssetDefinition_PCGGraph::OpenAssets(const FAssetOpenArgs&
 	return EAssetCommandResult::Handled;
 }
 
+EAssetCommandResult UAssetDefinition_PCGGraph::PerformAssetDiff(const FAssetDiffArgs& DiffArgs) const
+{
+	const UPCGGraph* OldGraph = Cast<UPCGGraph>(DiffArgs.OldAsset);
+	const UPCGGraph* NewGraph = Cast<UPCGGraph>(DiffArgs.NewAsset);
+
+	if (NewGraph == nullptr && OldGraph == nullptr)
+	{
+		return EAssetCommandResult::Unhandled;
+	}
+
+	SDetailsDiff::CreateDiffWindow(OldGraph, NewGraph, DiffArgs.OldRevision, DiffArgs.NewRevision, UPCGGraph::StaticClass());
+	return EAssetCommandResult::Handled;
+}
+
 #undef LOCTEXT_NAMESPACE
\ No newline at end of file
diff --git a/Engine/Plugins/PCG/Source/PCGEditor/Private/AssetDefinitions/AssetDefinition_PCGGraph.h b/Engine/Plugins/PCG/Source/PCGEditor/Private/AssetDefinitions/AssetDefinition_PCGGraph.h
--- a/Engine/Plugins/PCG/Source/PCGEditor/Private/AssetDefinitions/AssetDefinition_PCGGraph.h
+++ b/Engine/Plugins/PCG/Source/PCGEditor/Private/AssetDefinitions/AssetDefinition_PCGGraph.h
@@ -15,4 +15,5 @@ public:
 	virtual FText GetAssetDisplayName() const override;
 	virtual TSoftClassPtr<UObject> GetAssetClass() const override;
 	virtual EAssetCommandResult OpenAssets(const FAssetOpenArgs& OpenArgs) const override;
+	virtual EAssetCommandResult PerformAssetDiff(const FAssetDiffArgs& DiffArgs) const override;
 };
diff --git a/Engine/Plugins/PCG/Source/PCGEditor/Private/AssetDefinitions/AssetDefinition_PCGGraphInstance.cpp b/Engine/Plugins/PCG/Source/PCGEditor/Private/AssetDefinitions/AssetDefinition_PCGGraphInstance.cpp
--- a/Engine/Plugins/PCG/Source/PCGEditor/Private/AssetDefinitions/AssetDefinition_PCGGraphInstance.cpp
+++ b/Engine/Plugins/PCG/Source/PCGEditor/Private/AssetDefinitions/AssetDefinition_PCGGraphInstance.cpp
@@ -4,6 +4,8 @@
 
 #include "PCGGraph.h"
 
+#include "SDetailsDiff.h"
+
 #define LOCTEXT_NAMESPACE "AssetDefinition_PCGGraphInstance"
 
 FText UAssetDefinition_PCGGraphInstance::GetAssetDisplayName() const
@@ -25,4 +27,18 @@ TSoftClassPtr<UObject> UAssetDefinition_PCGGraphInstance::GetAssetClass() const
 	return UPCGGraphInstance::StaticClass();
 }
 
+EAssetCommandResult UAssetDefinition_PCGGraphInstance::PerformAssetDiff(const FAssetDiffArgs& DiffArgs) const
+{
+	const UPCGGraphInstance* OldGraph = Cast<UPCGGraphInstance>(DiffArgs.OldAsset);
+	const UPCGGraphInstance* NewGraph = Cast<UPCGGraphInstance>(DiffArgs.NewAsset);
+
+	if (NewGraph == nullptr && OldGraph == nullptr)
+	{
+		return EAssetCommandResult::Unhandled;
+	}
+
+	SDetailsDiff::CreateDiffWindow(OldGraph, NewGraph, DiffArgs.OldRevision, DiffArgs.NewRevision, UPCGGraphInstance::StaticClass());
+	return EAssetCommandResult::Handled;
+}
+
 #undef LOCTEXT_NAMESPACE
diff --git a/Engine/Plugins/PCG/Source/PCGEditor/Private/AssetDefinitions/AssetDefinition_PCGGraphInstance.h b/Engine/Plugins/PCG/Source/PCGEditor/Private/AssetDefinitions/AssetDefinition_PCGGraphInstance.h
--- a/Engine/Plugins/PCG/Source/PCGEditor/Private/AssetDefinitions/AssetDefinition_PCGGraphInstance.h
+++ b/Engine/Plugins/PCG/Source/PCGEditor/Private/AssetDefinitions/AssetDefinition_PCGGraphInstance.h
@@ -15,4 +15,5 @@ public:
 	virtual FText GetAssetDisplayName() const override;
 	virtual FLinearColor GetAssetColor() const override;
 	virtual TSoftClassPtr<UObject> GetAssetClass() const override;
+	EAssetCommandResult PerformAssetDiff(const FAssetDiffArgs& DiffArgs) const override;
 };
