[PCG] Add PCGWaterInterop plugin
* Add UPCGWaterSplineData and 'Get Water Spline Data' node
* Write WaterSplineComponent metadata to PCGWaterSplineData metadata for read-only access during PCG execution
* Add C++ support for custom metadata on PolyLine subtypes
* Add the option to compute spline input key as an attribute in the SplineSampler (used to interpolate metadata on splines)
* Precompute a reparameterization of the LandscapeSplineData for more efficient computation

diff --git a/Engine/Plugins/Experimental/PCGWaterInterop/PCGWaterInterop.uplugin b/Engine/Plugins/Experimental/PCGWaterInterop/PCGWaterInterop.uplugin
new file mode 100644
--- /dev/null
+++ b/Engine/Plugins/Experimental/PCGWaterInterop/PCGWaterInterop.uplugin
@@ -0,0 +1,35 @@
+{
+	"FileVersion": 3,
+	"Version": 1,
+	"VersionName": "0.1",
+	"FriendlyName": "Procedural Content Generation Framework (PCG) Water Interop",
+	"Description": "Extra plugin for Procedural Content Generation Framework interacting with the Water system.",
+	"Category": "Editor",
+	"CreatedBy": "Epic Games, Inc.",
+	"CreatedByURL": "https://epicgames.com",
+	"DocsURL": "https://docs.unrealengine.com/latest/en-US/procedural-content-generation--framework-in-unreal-engine/",
+	"MarketplaceURL": "",
+	"SupportURL": "",
+	"EnabledByDefault": false,
+	"CanContainContent": true,
+	"IsBetaVersion": false,
+	"IsExperimentalVersion": true,
+	"Installed": false,
+	"Modules": [
+		{
+			"Name": "PCGWaterInterop",
+			"Type": "Runtime",
+			"LoadingPhase": "Default"
+		}
+	],
+	"Plugins": [
+		{
+			"Name": "PCG",
+			"Enabled": true
+		},
+		{
+			"Name": "Water",
+			"Enabled": true
+		}
+	]
+}
diff --git a/Engine/Plugins/Experimental/PCGWaterInterop/Source/PCGWaterInterop/PCGWaterInterop.Build.cs b/Engine/Plugins/Experimental/PCGWaterInterop/Source/PCGWaterInterop/PCGWaterInterop.Build.cs
new file mode 100644
--- /dev/null
+++ b/Engine/Plugins/Experimental/PCGWaterInterop/Source/PCGWaterInterop/PCGWaterInterop.Build.cs
@@ -0,0 +1,22 @@
+// Copyright Epic Games, Inc. All Rights Reserved.
+
+namespace UnrealBuildTool.Rules
+{
+	public class PCGWaterInterop : ModuleRules
+	{
+		public PCGWaterInterop(ReadOnlyTargetRules Target) : base(Target)
+		{
+			PCHUsage = PCHUsageMode.UseExplicitOrSharedPCHs;
+
+			PublicDependencyModuleNames.AddRange(
+				new string[] {
+					"Core",
+					"CoreUObject",
+					"Engine",
+					"PCG",
+					"Water",
+				}
+			);
+		}
+	}
+}
diff --git a/Engine/Plugins/Experimental/PCGWaterInterop/Source/PCGWaterInterop/Private/Data/PCGWaterSplineData.cpp b/Engine/Plugins/Experimental/PCGWaterInterop/Source/PCGWaterInterop/Private/Data/PCGWaterSplineData.cpp
new file mode 100644
--- /dev/null
+++ b/Engine/Plugins/Experimental/PCGWaterInterop/Source/PCGWaterInterop/Private/Data/PCGWaterSplineData.cpp
@@ -0,0 +1,121 @@
+// Copyright Epic Games, Inc. All Rights Reserved.
+
+#include "Data/PCGWaterSplineData.h"
+
+#include "Data/PCGPointData.h"
+#include "Metadata/PCGMetadata.h"
+
+#include "WaterSplineComponent.h"
+#include "WaterSplineMetadata.h"
+
+#include UE_INLINE_GENERATED_CPP_BY_NAME(PCGWaterSplineData)
+
+namespace PCGWaterSplineDataConstants
+{
+	const FName DepthAttributeName = TEXT("Depth");
+	const FName WaterVelocityScalarAttributeName = TEXT("WaterVelocityScalar");
+	const FName RiverWidthAttributeName = TEXT("RiverWidth");
+	const FName AudioIntensityAttributeName = TEXT("AudioIntensity");
+}
+
+void UPCGWaterSplineData::Initialize(const UWaterSplineComponent* InSpline)
+{
+	check(InSpline && Metadata);
+
+	Super::Initialize(InSpline);
+
+	if (const UWaterSplineMetadata* WaterSplineMetadata = Cast<UWaterSplineMetadata>(InSpline->GetSplinePointsMetadata()))
+	{
+		WaterSplineMetadataStruct.Depth = WaterSplineMetadata->Depth;
+		WaterSplineMetadataStruct.WaterVelocityScalar = WaterSplineMetadata->WaterVelocityScalar;
+		WaterSplineMetadataStruct.RiverWidth = WaterSplineMetadata->RiverWidth;
+		WaterSplineMetadataStruct.AudioIntensity = WaterSplineMetadata->AudioIntensity;
+
+		Metadata->CreateAttribute<float>(PCGWaterSplineDataConstants::DepthAttributeName, InSpline->WaterSplineDefaults.DefaultDepth, /*bAllowsInterpolation=*/true, /*bOverrideParent=*/true);
+		Metadata->CreateAttribute<float>(PCGWaterSplineDataConstants::WaterVelocityScalarAttributeName, InSpline->WaterSplineDefaults.DefaultVelocity, /*bAllowsInterpolation=*/true, /*bOverrideParent=*/true);
+		Metadata->CreateAttribute<float>(PCGWaterSplineDataConstants::RiverWidthAttributeName, InSpline->WaterSplineDefaults.DefaultWidth, /*bAllowsInterpolation=*/true, /*bOverrideParent=*/true);
+		Metadata->CreateAttribute<float>(PCGWaterSplineDataConstants::AudioIntensityAttributeName, InSpline->WaterSplineDefaults.DefaultAudioIntensity, /*bAllowsInterpolation=*/true, /*bOverrideParent=*/true);
+	}
+}
+
+bool UPCGWaterSplineData::SamplePoint(const FTransform& Transform, const FBox& Bounds, FPCGPoint& OutPoint, UPCGMetadata* OutMetadata) const
+{
+	// Find nearest point on spline
+	const FVector InPosition = Transform.GetLocation();
+	const float NearestPointKey = SplineStruct.FindInputKeyClosestToWorldLocation(InPosition);
+	FTransform NearestTransform = SplineStruct.GetTransformAtSplineInputKey(NearestPointKey, ESplineCoordinateSpace::World, true);
+
+	// WaterSpline samples are scaled by river width on X instead of Y, so we swizzle XY to get the correct transform.
+	const FVector Scale = NearestTransform.GetScale3D();
+	NearestTransform.SetScale3D(FVector(Scale.Y, Scale.X, Scale.Z));
+
+	const FVector LocalPoint = NearestTransform.InverseTransformPosition(InPosition);
+	
+	// Linear fall off based on the distance to the nearest point
+	const float Distance = LocalPoint.Length();
+	if (Distance > 1.0f)
+	{
+		return false;
+	}
+	else
+	{
+		OutPoint.Transform = Transform;
+		OutPoint.SetLocalBounds(Bounds);
+		OutPoint.Density = 1.0f - Distance;
+
+		WriteMetadataToPoint(NearestPointKey, OutPoint, OutMetadata);
+
+		return true;
+	}
+}
+
+UPCGSpatialData* UPCGWaterSplineData::CopyInternal() const
+{
+	UPCGWaterSplineData* NewSplineData = NewObject<UPCGWaterSplineData>();
+	NewSplineData->WaterSplineMetadataStruct = WaterSplineMetadataStruct;
+
+	CopySplineData(NewSplineData);
+
+	return NewSplineData;
+}
+
+FTransform UPCGWaterSplineData::GetTransformAtDistance(int SegmentIndex, FVector::FReal Distance, bool bWorldSpace, FBox* OutBounds) const
+{
+	FTransform Transform = Super::GetTransformAtDistance(SegmentIndex, Distance, bWorldSpace, OutBounds);
+
+	// WaterSpline samples are scaled by river width on X instead of Y, so we swizzle XY to get the correct transform.
+	const FVector Scale = Transform.GetScale3D();
+	Transform.SetScale3D(FVector(Scale.Y, Scale.X, Scale.Z));
+
+	return Transform;
+}
+
+void UPCGWaterSplineData::WriteMetadataToPoint(float InputKey, FPCGPoint& OutPoint, UPCGMetadata* OutMetadata) const
+{
+	if (!OutMetadata)
+	{
+		return;
+	}
+
+	OutMetadata->InitializeOnSet(OutPoint.MetadataEntry);
+
+	if (FPCGMetadataAttribute<float>* DepthAttribute = OutMetadata->GetMutableTypedAttribute<float>(PCGWaterSplineDataConstants::DepthAttributeName))
+	{
+		DepthAttribute->SetValue(OutPoint.MetadataEntry, WaterSplineMetadataStruct.Depth.Eval(InputKey));
+	}
+
+	if (FPCGMetadataAttribute<float>* VelocityAttribute = OutMetadata->GetMutableTypedAttribute<float>(PCGWaterSplineDataConstants::WaterVelocityScalarAttributeName))
+	{
+		VelocityAttribute->SetValue(OutPoint.MetadataEntry, WaterSplineMetadataStruct.WaterVelocityScalar.Eval(InputKey));
+	}
+
+	if (FPCGMetadataAttribute<float>* WidthAttribute = OutMetadata->GetMutableTypedAttribute<float>(PCGWaterSplineDataConstants::RiverWidthAttributeName))
+	{
+		WidthAttribute->SetValue(OutPoint.MetadataEntry, WaterSplineMetadataStruct.RiverWidth.Eval(InputKey));
+	}
+
+	if (FPCGMetadataAttribute<float>* AudioIntensityAttribute = OutMetadata->GetMutableTypedAttribute<float>(PCGWaterSplineDataConstants::AudioIntensityAttributeName))
+	{
+		AudioIntensityAttribute->SetValue(OutPoint.MetadataEntry, WaterSplineMetadataStruct.AudioIntensity.Eval(InputKey));
+	}
+}
diff --git a/Engine/Plugins/Experimental/PCGWaterInterop/Source/PCGWaterInterop/Private/Elements/PCGGetWaterSpline.cpp b/Engine/Plugins/Experimental/PCGWaterInterop/Source/PCGWaterInterop/Private/Elements/PCGGetWaterSpline.cpp
new file mode 100644
--- /dev/null
+++ b/Engine/Plugins/Experimental/PCGWaterInterop/Source/PCGWaterInterop/Private/Elements/PCGGetWaterSpline.cpp
@@ -0,0 +1,60 @@
+// Copyright Epic Games, Inc. All Rights Reserved.
+
+#include "Elements/PCGGetWaterSpline.h"
+
+#include "Data/PCGWaterSplineData.h"
+
+#include "WaterSplineComponent.h"
+
+#include UE_INLINE_GENERATED_CPP_BY_NAME(PCGGetWaterSpline)
+
+#define LOCTEXT_NAMESPACE "PCGGetWaterSplineElement"
+
+UPCGGetWaterSplineSettings::UPCGGetWaterSplineSettings()
+{
+	bDisplayModeSettings = false;
+	Mode = EPCGGetDataFromActorMode::ParseActorComponents;
+}
+
+#if WITH_EDITOR
+FText UPCGGetWaterSplineSettings::GetNodeTooltipText() const
+{
+	return LOCTEXT("GetWaterSplineTooltip", "Builds a collection of data from WaterSplineComponents on the selected actors.");
+}
+#endif
+
+FPCGElementPtr UPCGGetWaterSplineSettings::CreateElement() const
+{
+	return MakeShared<FPCGGetWaterSplineElement>();
+}
+
+void FPCGGetWaterSplineElement::ProcessActor(FPCGContext* Context, const UPCGDataFromActorSettings* Settings, AActor* FoundActor) const
+{
+	check(Context && Settings);
+
+	if (!FoundActor || !IsValid(FoundActor))
+	{
+		return;
+	}
+
+	auto NameTagsToStringTags = [](const FName& InName) { return InName.ToString(); };
+	TSet<FString> ActorTags;
+	Algo::Transform(FoundActor->Tags, ActorTags, NameTagsToStringTags);
+
+	// Prepare data on a component basis
+	TInlineComponentArray<UWaterSplineComponent*, 4> WaterSplines;
+	FoundActor->GetComponents(WaterSplines);
+
+	for (UWaterSplineComponent* WaterSplineComponent : WaterSplines)
+	{
+		UPCGWaterSplineData* SplineData = NewObject<UPCGWaterSplineData>();
+		SplineData->Initialize(WaterSplineComponent);
+
+		FPCGTaggedData& TaggedData = Context->OutputData.TaggedData.Emplace_GetRef();
+		TaggedData.Data = SplineData;
+		Algo::Transform(WaterSplineComponent->ComponentTags, TaggedData.Tags, NameTagsToStringTags);
+		TaggedData.Tags.Append(ActorTags);
+	}
+}
+
+#undef LOCTEXT_NAMESPACE
diff --git a/Engine/Plugins/Experimental/PCGWaterInterop/Source/PCGWaterInterop/Private/PCGWaterInteropModule.cpp b/Engine/Plugins/Experimental/PCGWaterInterop/Source/PCGWaterInterop/Private/PCGWaterInteropModule.cpp
new file mode 100644
--- /dev/null
+++ b/Engine/Plugins/Experimental/PCGWaterInterop/Source/PCGWaterInterop/Private/PCGWaterInteropModule.cpp
@@ -0,0 +1,9 @@
+// Copyright Epic Games, Inc. All Rights Reserved.
+
+#include "PCGWaterInteropModule.h"
+
+#include "Modules/ModuleManager.h"
+
+IMPLEMENT_MODULE(FPCGWaterInteropModule, PCGWaterInterop);
+
+DEFINE_LOG_CATEGORY(LogPCGWaterInterop);
diff --git a/Engine/Plugins/Experimental/PCGWaterInterop/Source/PCGWaterInterop/Public/Data/PCGWaterSplineData.h b/Engine/Plugins/Experimental/PCGWaterInterop/Source/PCGWaterInterop/Public/Data/PCGWaterSplineData.h
new file mode 100644
--- /dev/null
+++ b/Engine/Plugins/Experimental/PCGWaterInterop/Source/PCGWaterInterop/Public/Data/PCGWaterSplineData.h
@@ -0,0 +1,56 @@
+// Copyright Epic Games, Inc. All Rights Reserved.
+
+#pragma once
+
+#include "Data/PCGSplineData.h"
+
+#include "PCGWaterSplineData.generated.h"
+
+class UWaterSplineComponent;
+
+/** Helper struct to hold a copy of the WaterSplineMetadata contents. */
+USTRUCT()
+struct FPCGWaterSplineMetadataStruct
+{
+	GENERATED_BODY()
+
+	/** The Depth of the water at this vertex. */
+	UPROPERTY()
+	FInterpCurveFloat Depth;
+
+	/** The Current of the water at this vertex. Magnitude and direction. */
+	UPROPERTY()
+	FInterpCurveFloat WaterVelocityScalar;
+
+	/** Rivers Only: The width of the river (from center) in each direction. */
+	UPROPERTY()
+	FInterpCurveFloat RiverWidth;
+
+	/** A scalar used to define intensity of the water audio along the spline. */
+	UPROPERTY()
+	FInterpCurveFloat AudioIntensity;
+};
+
+UCLASS(BlueprintType, ClassGroup = (Procedural))
+class UPCGWaterSplineData : public UPCGSplineData
+{
+	GENERATED_BODY()
+
+public:
+	void Initialize(const UWaterSplineComponent* InSpline);
+
+	//~Begin UPCGSpatialData interface
+	virtual bool SamplePoint(const FTransform& Transform, const FBox& Bounds, FPCGPoint& OutPoint, UPCGMetadata* OutMetadata) const override;
+protected:
+	virtual UPCGSpatialData* CopyInternal() const override;
+	//~End UPCGSpatialData interface
+
+	//~Begin UPCGPolyLineData interface
+	virtual FTransform GetTransformAtDistance(int SegmentIndex, FVector::FReal Distance, bool bWorldSpace = true, FBox* OutBounds = nullptr) const override;
+	virtual void WriteMetadataToPoint(float InputKey, FPCGPoint& OutPoint, UPCGMetadata* OutMetadata) const override;
+	virtual bool HasCustomMetadata() const override { return true; }
+	//~End UPCGPolyLineData interface
+
+	UPROPERTY()
+	FPCGWaterSplineMetadataStruct WaterSplineMetadataStruct;
+};
diff --git a/Engine/Plugins/Experimental/PCGWaterInterop/Source/PCGWaterInterop/Public/Elements/PCGGetWaterSpline.h b/Engine/Plugins/Experimental/PCGWaterInterop/Source/PCGWaterInterop/Public/Elements/PCGGetWaterSpline.h
new file mode 100644
--- /dev/null
+++ b/Engine/Plugins/Experimental/PCGWaterInterop/Source/PCGWaterInterop/Public/Elements/PCGGetWaterSpline.h
@@ -0,0 +1,36 @@
+// Copyright Epic Games, Inc. All Rights Reserved.
+
+#pragma once
+
+#include "PCGSettings.h"
+#include "Elements/PCGDataFromActor.h"
+#include "Elements/PCGTypedGetter.h"
+
+#include "PCGGetWaterSpline.generated.h"
+
+/** Builds a collection of data from WaterSplineComponents on the selected actors. */
+UCLASS(BlueprintType, ClassGroup = (Procedural))
+class UPCGGetWaterSplineSettings : public UPCGGetSplineSettings
+{
+	GENERATED_BODY()
+
+public:
+	UPCGGetWaterSplineSettings();
+
+	//~Begin UPCGSettings interface
+#if WITH_EDITOR
+	virtual FName GetDefaultNodeName() const override { return FName(TEXT("GetWaterSplineData")); }
+	virtual FText GetDefaultNodeTitle() const override { return NSLOCTEXT("PCGGetWaterSplineElement", "NodeTitle", "Get Water Spline Data"); }
+	virtual FText GetNodeTooltipText() const override;
+#endif
+
+protected:
+	virtual FPCGElementPtr CreateElement() const override;
+	//~End UPCGSettings
+};
+
+class FPCGGetWaterSplineElement : public FPCGDataFromActorElement
+{
+protected:
+	virtual void ProcessActor(FPCGContext* Context, const UPCGDataFromActorSettings* Settings, AActor* FoundActor) const override;
+};
diff --git a/Engine/Plugins/Experimental/PCGWaterInterop/Source/PCGWaterInterop/Public/PCGWaterInteropModule.h b/Engine/Plugins/Experimental/PCGWaterInterop/Source/PCGWaterInterop/Public/PCGWaterInteropModule.h
new file mode 100644
--- /dev/null
+++ b/Engine/Plugins/Experimental/PCGWaterInterop/Source/PCGWaterInterop/Public/PCGWaterInteropModule.h
@@ -0,0 +1,16 @@
+// Copyright Epic Games, Inc. All Rights Reserved.
+
+#pragma once
+
+#include "Modules/ModuleInterface.h"
+#include "Stats/Stats.h"
+
+DECLARE_LOG_CATEGORY_EXTERN(LogPCGWaterInterop, Log, All);
+
+class FPCGWaterInteropModule final : public IModuleInterface
+{
+public:
+	// ~IModuleInterface implementation
+	virtual bool SupportsDynamicReloading() override { return true; };
+	// ~End IModuleInterface implementation
+};
diff --git a/Engine/Plugins/PCG/Source/PCG/Private/Data/PCGLandscapeSplineData.cpp b/Engine/Plugins/PCG/Source/PCG/Private/Data/PCGLandscapeSplineData.cpp
--- a/Engine/Plugins/PCG/Source/PCG/Private/Data/PCGLandscapeSplineData.cpp
+++ b/Engine/Plugins/PCG/Source/PCG/Private/Data/PCGLandscapeSplineData.cpp
@@ -42,6 +42,14 @@ void UPCGLandscapeSplineData::Initialize(ULandscapeSplinesComponent* InSplineCom
 {
 	check(InSplineComponent);
 	Spline = InSplineComponent;
+
+	UpdateReparamTable();
+}
+
+void UPCGLandscapeSplineData::PostLoad()
+{
+	Super::PostLoad();
+	UpdateReparamTable();
 }
 
 void UPCGLandscapeSplineData::AddToCrc(FArchiveCrc32& Ar, bool bFullDataCrc) const
@@ -65,151 +73,121 @@ int UPCGLandscapeSplineData::GetNumSegments() const
 
 FVector::FReal UPCGLandscapeSplineData::GetSegmentLength(int SegmentIndex) const
 {
-	check(Spline.IsValid());
-	check(SegmentIndex >= 0 && SegmentIndex < Spline->GetSegments().Num());
-
-	const ULandscapeSplineSegment* Segment = Spline->GetSegments()[SegmentIndex];
-	const TArray<FLandscapeSplineInterpPoint>& InterpPoints = Segment->GetPoints();
-	FVector::FReal Length = 0;
-
-	for (int PointIndex = 1; PointIndex < InterpPoints.Num(); ++PointIndex)
-	{
-		Length += (InterpPoints[PointIndex].Center - InterpPoints[PointIndex - 1].Center).Length();
-	}
-
-	return Length;
+	return GetDistanceAtSegmentStart(SegmentIndex + 1) - GetDistanceAtSegmentStart(SegmentIndex);
 }
 
 FTransform UPCGLandscapeSplineData::GetTransformAtDistance(int SegmentIndex, FVector::FReal Distance, bool bWorldSpace, FBox* OutBounds) const
 {
-	check(Spline.IsValid());
-	check(SegmentIndex >= 0 && SegmentIndex < Spline->GetSegments().Num());
+	check(Spline.IsValid() && Spline->GetSegments().IsValidIndex(SegmentIndex));
 
 	const ULandscapeSplineSegment* Segment = Spline->GetSegments()[SegmentIndex];
 	check(Segment);
 
-	const FLandscapeSplineSegmentConnection& Start = Segment->Connections[0];
-	const FLandscapeSplineSegmentConnection& End = Segment->Connections[1];
-
 	const TArray<FLandscapeSplineInterpPoint>& InterpPoints = Segment->GetPoints();
-	FVector::FReal Length = FMath::Max(0, Distance);
 
-	for (int PointIndex = 1; PointIndex < InterpPoints.Num(); ++PointIndex)
-	{
-		const FLandscapeSplineInterpPoint& PreviousPoint = InterpPoints[PointIndex - 1];
-		const FLandscapeSplineInterpPoint& CurrentPoint = InterpPoints[PointIndex];
+	int32 PointIndex;
+	FVector::FReal Alpha;
+	GetInterpPointAtDistance(SegmentIndex, Distance, PointIndex, /*bComputeAlpha=*/true, Alpha);
 
-		const FVector::FReal SegmentLength = (CurrentPoint.Center - PreviousPoint.Center).Length();
-		if (SegmentLength > Length || PointIndex == InterpPoints.Num() - 1)
-		{
-			const FVector XAxis = CurrentPoint.Center - PreviousPoint.Center;
-			const FVector PreviousYAxis = PreviousPoint.Right - PreviousPoint.Center;
-			const FVector CurrentYAxis = CurrentPoint.Right - CurrentPoint.Center;
-			const FVector PreviousZAxis = (XAxis ^ PreviousYAxis).GetSafeNormal(UE_SMALL_NUMBER, FVector::ZAxisVector);
-			const FVector CurrentZAxis = (XAxis ^ CurrentYAxis).GetSafeNormal(UE_SMALL_NUMBER, FVector::ZAxisVector);
+	check(InterpPoints.IsValidIndex(PointIndex));
+	const FLandscapeSplineInterpPoint& PreviousPoint = InterpPoints[PointIndex];
+	const FLandscapeSplineInterpPoint& CurrentPoint = InterpPoints[FMath::Min(PointIndex + 1, InterpPoints.Num() - 1)]; // If our RightPoint ends up being on the next segment, clamp it back to the current segment.
 
-			FTransform PreviousTransform = FTransform(XAxis, PreviousYAxis, PreviousZAxis, PreviousPoint.Center);
-			FTransform CurrentTransform = FTransform(XAxis, CurrentYAxis, CurrentZAxis, CurrentPoint.Center);
+	const FVector XAxis = CurrentPoint.Center - PreviousPoint.Center;
+	const FVector PreviousYAxis = PreviousPoint.Right - PreviousPoint.Center;
+	const FVector CurrentYAxis = CurrentPoint.Right - CurrentPoint.Center;
+	const FVector PreviousZAxis = (XAxis ^ PreviousYAxis).GetSafeNormal(UE_SMALL_NUMBER, FVector::ZAxisVector);
+	const FVector CurrentZAxis = (XAxis ^ CurrentYAxis).GetSafeNormal(UE_SMALL_NUMBER, FVector::ZAxisVector);
 
-			const FVector::FReal BlendRatio = ((SegmentLength > Length) ? (Length / SegmentLength) : 1.0);
-			PreviousTransform.BlendWith(CurrentTransform, BlendRatio);
+	FTransform PreviousTransform = FTransform(XAxis, PreviousYAxis, PreviousZAxis, PreviousPoint.Center);
+	FTransform CurrentTransform = FTransform(XAxis, CurrentYAxis, CurrentZAxis, CurrentPoint.Center);
 
-			if (OutBounds)
-			{
-				// Important note: the box here is going to be useful to be able to specify the relative sizes of the falloffs
-				*OutBounds = FBox::BuildAABB(FVector::ZeroVector, FVector::OneVector);
-				OutBounds->Min.Y *= (CurrentPoint.FalloffLeft - CurrentPoint.Center).Length() / (CurrentPoint.Left - CurrentPoint.Center).Length();
-				OutBounds->Max.Y *= (CurrentPoint.FalloffRight - CurrentPoint.Center).Length() / (CurrentPoint.Right - CurrentPoint.Center).Length();
-			}
+	PreviousTransform.BlendWith(CurrentTransform, Alpha);
 
-			if (bWorldSpace)
-			{
-				return PreviousTransform * Spline->GetComponentTransform();
-			}
-			else
-			{
-				return PreviousTransform;
-			}
-		}
-		else
-		{
-			Length -= SegmentLength;
-		}
+	if (OutBounds)
+	{
+		// Important note: the box here is going to be useful to be able to specify the relative sizes of the falloffs
+		*OutBounds = FBox::BuildAABB(FVector::ZeroVector, FVector::OneVector);
+		OutBounds->Min.Y *= (CurrentPoint.FalloffLeft - CurrentPoint.Center).Length() / (CurrentPoint.Left - CurrentPoint.Center).Length();
+		OutBounds->Max.Y *= (CurrentPoint.FalloffRight - CurrentPoint.Center).Length() / (CurrentPoint.Right - CurrentPoint.Center).Length();
 	}
-	
-	check(0);
-	return FTransform();
+
+	if (bWorldSpace)
+	{
+		PreviousTransform *= Spline->GetComponentTransform();
+	}
+
+	return PreviousTransform;
 }
 
 FVector::FReal UPCGLandscapeSplineData::GetCurvatureAtDistance(int SegmentIndex, FVector::FReal Distance) const
 {
-	check(Spline.IsValid());
-	check(SegmentIndex >= 0 && SegmentIndex < Spline->GetSegments().Num());
+	check(Spline.IsValid() && Spline->GetSegments().IsValidIndex(SegmentIndex));
 
 	const ULandscapeSplineSegment* Segment = Spline->GetSegments()[SegmentIndex];
 	check(Segment);
 
-	const FLandscapeSplineSegmentConnection& Start = Segment->Connections[0];
-	const FLandscapeSplineSegmentConnection& End = Segment->Connections[1];
-
 	const TArray<FLandscapeSplineInterpPoint>& InterpPoints = Segment->GetPoints();
 
-	// Need at least three points to compute the curvature
+	// Need at least three points to compute the curvature.
 	if (InterpPoints.Num() < 3)
 	{
 		return 0;
 	}
 
-	FVector::FReal Length = FMath::Max(0, Distance);
+	int32 PointIndex;
+	FVector::FReal Alpha;
+	GetInterpPointAtDistance(SegmentIndex, Distance, PointIndex, /*bComputeAlpha=*/false, Alpha);
 
-	for (int PointIndex = 1; PointIndex < InterpPoints.Num(); ++PointIndex)
+	// If our sample overshoots the segment, clamp it back to the last point.
+	if (PointIndex == InterpPoints.Num() - 1)
 	{
-		const FLandscapeSplineInterpPoint& PreviousPoint = InterpPoints[PointIndex - 1];
-		const FLandscapeSplineInterpPoint& CurrentPoint = InterpPoints[PointIndex];
-
-		const FVector::FReal SegmentLength = (CurrentPoint.Center - PreviousPoint.Center).Length();
-		if (SegmentLength > Length || PointIndex == InterpPoints.Num() - 1)
-		{
-			FVector FirstDerivative = FVector::ZeroVector;
-			FVector SecondDerivative = FVector::ZeroVector;
+		--PointIndex;
+	}
 
-			// Compute curvature using finite differences - here h is 1 because that's the only base unit we have.
-			// Warning: precision will be poor
-			// if last point -> use backward 2nd derivative
-			if (PointIndex == InterpPoints.Num() - 1)
-			{
-				const FLandscapeSplineInterpPoint& PreviousPreviousPoint = InterpPoints[PointIndex - 2];
+	// We don't need to clamp the current point index like we do in GetTransformAtDistance(), because we've already
+	// decremented the PointIndex so that we can perform the backward 2nd derivative.
+	const FLandscapeSplineInterpPoint& PreviousPoint = InterpPoints[PointIndex];
+	const FLandscapeSplineInterpPoint& CurrentPoint = InterpPoints[PointIndex + 1];
 
-				// f'(x) = (f(x) - f(x-h)) / h
-				FirstDerivative = (CurrentPoint.Center - PreviousPoint.Center);
-				// f''(x) = (f(x) - 2f(x - h) + f(x - 2h)) / h2
-				SecondDerivative = (CurrentPoint.Center - 2 * PreviousPoint.Center + PreviousPreviousPoint.Center);
-			}
-			// otherwise -> use central 2nd derivative
-			else
-			{
-				const FLandscapeSplineInterpPoint& NextPoint = InterpPoints[PointIndex + 1];
+	FVector FirstDerivative = FVector::ZeroVector;
+	FVector SecondDerivative = FVector::ZeroVector;
 
-				// f'(x) ~= (f(x+h) - f(x-h)) / 2h
-				FirstDerivative = (NextPoint.Center - PreviousPoint.Center) / 2.0;
-				// f''(x) = (f(x+h) - 2f(x) + f(x-h)) / h2
-				SecondDerivative = NextPoint.Center - 2 * CurrentPoint.Center + PreviousPoint.Center;
-			}
+	// Compute curvature using finite differences - here h is 1 because that's the only base unit we have.
+	// Warning: Precision will be poor.
+	// If last point -> use backward 2nd derivative.
+	if (PointIndex >= InterpPoints.Num() - 2)
+	{
+		const FLandscapeSplineInterpPoint& PreviousPreviousPoint = InterpPoints[PointIndex - 1];
 
-			const FVector::FReal FirstDerivativeLength = FMath::Max(FirstDerivative.Length(), UE_DOUBLE_SMALL_NUMBER);
-			const FVector ForwardVector = FirstDerivative / FirstDerivativeLength;
-			const FVector CurvatureVector = SecondDerivative - (SecondDerivative | ForwardVector) * ForwardVector;
-			const FVector::FReal Curvature = CurvatureVector.Length() / FirstDerivativeLength;
+		// f'(x) = (f(x) - f(x-h)) / h
+		FirstDerivative = (CurrentPoint.Center - PreviousPoint.Center);
+		// f''(x) = (f(x) - 2f(x - h) + f(x - 2h)) / h2
+		SecondDerivative = (CurrentPoint.Center - 2 * PreviousPoint.Center + PreviousPreviousPoint.Center);
+	}
+	// Otherwise, use central 2nd derivative.
+	else
+	{
+		const FLandscapeSplineInterpPoint& NextPoint = InterpPoints[PointIndex + 2];
 
-			return FMath::Sign(CurvatureVector | (CurrentPoint.Right - CurrentPoint.Center)) * Curvature;
-		}
-		else
-		{
-			Length -= SegmentLength;
-		}
+		// f'(x) ~= (f(x+h) - f(x-h)) / 2h
+		FirstDerivative = (NextPoint.Center - PreviousPoint.Center) / 2.0;
+		// f''(x) = (f(x+h) - 2f(x) + f(x-h)) / h2
+		SecondDerivative = NextPoint.Center - 2 * CurrentPoint.Center + PreviousPoint.Center;
 	}
 
-	return 0;
+	const FVector::FReal FirstDerivativeLength = FMath::Max(FirstDerivative.Length(), UE_DOUBLE_SMALL_NUMBER);
+	const FVector ForwardVector = FirstDerivative / FirstDerivativeLength;
+	const FVector CurvatureVector = SecondDerivative - (SecondDerivative | ForwardVector) * ForwardVector;
+	const FVector::FReal Curvature = CurvatureVector.Length() / FirstDerivativeLength;
+
+	return FMath::Sign(CurvatureVector | (CurrentPoint.Right - CurrentPoint.Center)) * Curvature;
+}
+
+float UPCGLandscapeSplineData::GetInputKeyAtDistance(int SegmentIndex, FVector::FReal Distance) const
+{
+	const float FullDistance = GetDistanceAtSegmentStart(SegmentIndex) + Distance;
+	return ReparamTable.Eval(FullDistance, 0.0f);
 }
 
 void UPCGLandscapeSplineData::GetTangentsAtSegmentStart(int SegmentIndex, FVector& OutArriveTangent, FVector& OutLeaveTangent) const
@@ -231,6 +209,23 @@ void UPCGLandscapeSplineData::GetTangentsAtSegmentStart(int SegmentIndex, FVecto
 	OutLeaveTangent = LeavePoint ? LeavePoint->Rotation.Vector() * CurrentSegment->Connections[0].TangentLen : FVector::Zero();
 }
 
+FVector::FReal UPCGLandscapeSplineData::GetDistanceAtSegmentStart(int SegmentIndex) const
+{
+	check(Spline.IsValid() && Spline->GetSegments().IsValidIndex(SegmentIndex));
+
+	const TArray<TObjectPtr<ULandscapeSplineSegment>>& Segments = Spline->GetSegments();
+	int32 ReparamIndex = 0;
+
+	for (int32 Index = 0; Index < SegmentIndex; ++Index)
+	{
+		// NumPoints - 1 to avoid double-counting the control points, which overlap at the start + end of each segment.
+		ReparamIndex += ensure(Segments[Index]) ? Segments[Index]->GetPoints().Num() - 1 : 0;
+	}
+
+	check(ReparamTable.Points.IsValidIndex(ReparamIndex));
+	return ReparamTable.Points[ReparamIndex].InVal;
+}
+
 const UPCGPointData* UPCGLandscapeSplineData::CreatePointData(FPCGContext* Context) const
 {
 	check(Spline.IsValid());
@@ -338,6 +333,68 @@ UPCGSpatialData* UPCGLandscapeSplineData::CopyInternal() const
 	UPCGLandscapeSplineData* NewLandscapeSplineData = NewObject<UPCGLandscapeSplineData>();
 
 	NewLandscapeSplineData->Spline = Spline;
+	NewLandscapeSplineData->ReparamTable = ReparamTable;
 
 	return NewLandscapeSplineData;
 }
+
+void UPCGLandscapeSplineData::UpdateReparamTable()
+{
+	check(Spline.IsValid());
+
+	const TArray<TObjectPtr<ULandscapeSplineSegment>>& Segments = Spline->GetSegments();
+	const int32 NumSegments = Segments.Num();
+
+	FVector::FReal AccumulatedDistance = 0.0f;
+
+	// Add a point for the first control point of the spline.
+	ReparamTable.Points.Emplace(AccumulatedDistance, /*InputKey=*/0, /*ArriveTangent=*/0.0f, /*LeaveTangent=*/0.0f, CIM_Linear);
+
+	// Create a curve mapping DistanceAlongSpline -> InputKey at that point. Accumulate the distance over each segment as we insert points into the ReparamTable.
+	for (int32 SegmentIndex = 0; SegmentIndex < NumSegments; ++SegmentIndex)
+	{
+		const TObjectPtr<ULandscapeSplineSegment>& Segment = Segments[SegmentIndex];
+		check(Segment);
+
+		const TArray<FLandscapeSplineInterpPoint>& InterpPoints = Segment->GetPoints();
+		const int32 NumPoints = InterpPoints.Num();
+
+		// Add a point for each InterpPoint on the segment.
+		for (int PointIndex = 1; PointIndex < NumPoints; ++PointIndex)
+		{
+			const FLandscapeSplineInterpPoint& Start = InterpPoints[PointIndex - 1];
+			const FLandscapeSplineInterpPoint& End = InterpPoints[PointIndex];
+			AccumulatedDistance += FVector::Distance(Start.Center, End.Center);
+
+			const float Param = PointIndex / (NumPoints - 1.0f);
+			ReparamTable.Points.Emplace(AccumulatedDistance, SegmentIndex + Param, /*ArriveTangent=*/0.0f, /*LeaveTangent=*/0.0f, CIM_Linear);
+		}
+	}
+}
+
+void UPCGLandscapeSplineData::GetInterpPointAtDistance(int SegmentIndex, FVector::FReal Distance, int32& OutPointIndex, bool bComputeAlpha, FVector::FReal& OutAlpha) const
+{
+	const FVector::FReal DistanceToSegment = GetDistanceAtSegmentStart(SegmentIndex);
+	const FVector::FReal DistanceToSample = DistanceToSegment + Distance; // Total distance along the spline to the point we are sampling at.
+
+	const int32 ReparamIndex = ReparamTable.GetPointIndexForInputValue(DistanceToSegment + Distance);
+	check(ReparamTable.Points.IsValidIndex(ReparamIndex));
+
+	// Find the index of the InterpPoint which begins the segment our desired transform lies on.
+	// We can get this from the ReparamTable because it should contain one sample for each InterpPoint.
+	const int32 SegmentReparamIndex = ReparamTable.GetPointIndexForInputValue(DistanceToSegment);
+	OutPointIndex = ReparamIndex - SegmentReparamIndex;
+
+	if (bComputeAlpha)
+	{
+		// Blend between LeftPoint and RightPoint by finding the ratio of distance to our sample vs length of the segment.
+		// This ratio should be relative to the interp segment (which is formed by the line from LeftPoint to RightPoint),
+		// since that is what we're blending on, not the entire spline segment.
+		const FVector::FReal DistanceToPrevPoint = ReparamTable.Points[ReparamIndex].InVal;
+		const FVector::FReal DistanceToNextPoint = ReparamTable.Points[FMath::Min(ReparamIndex + 1, ReparamTable.Points.Num() - 1)].InVal;
+		const FVector::FReal InterpSegmentLength = DistanceToNextPoint - DistanceToPrevPoint;
+		const FVector::FReal InterpSegmentDistanceToSample = DistanceToSample - DistanceToPrevPoint;
+
+		OutAlpha = FMath::IsNearlyZero(InterpSegmentDistanceToSample) ? 0.0f : FMath::Clamp(InterpSegmentDistanceToSample / InterpSegmentLength, 0.0f, 1.0f);
+	}
+}
diff --git a/Engine/Plugins/PCG/Source/PCG/Private/Data/PCGSplineData.cpp b/Engine/Plugins/PCG/Source/PCG/Private/Data/PCGSplineData.cpp
--- a/Engine/Plugins/PCG/Source/PCG/Private/Data/PCGSplineData.cpp
+++ b/Engine/Plugins/PCG/Source/PCG/Private/Data/PCGSplineData.cpp
@@ -13,7 +13,7 @@
 
 #include UE_INLINE_GENERATED_CPP_BY_NAME(PCGSplineData)
 
-void UPCGSplineData::Initialize(USplineComponent* InSpline)
+void UPCGSplineData::Initialize(const USplineComponent* InSpline)
 {
 	check(InSpline);
 
@@ -134,6 +134,12 @@ FVector::FReal UPCGSplineData::GetCurvatureAtDistance(int SegmentIndex, FVector:
 	return FMath::Sign(RightVector | CurvatureVector) * Curvature;
 }
 
+float UPCGSplineData::GetInputKeyAtDistance(int SegmentIndex, FVector::FReal Distance) const
+{
+	const float FullDistance = GetDistanceAtSegmentStart(SegmentIndex) + Distance;
+	return SplineStruct.SplineCurves.ReparamTable.Eval(FullDistance, 0.0f);
+}
+
 void UPCGSplineData::GetTangentsAtSegmentStart(int SegmentIndex, FVector& OutArriveTangent, FVector& OutLeaveTangent) const
 {
 	check(SplineStruct.SplineCurves.Position.Points.IsValidIndex(SegmentIndex));
@@ -141,6 +147,11 @@ void UPCGSplineData::GetTangentsAtSegmentStart(int SegmentIndex, FVector& OutArr
 	OutLeaveTangent = SplineStruct.SplineCurves.Position.Points[SegmentIndex].LeaveTangent;
 }
 
+FVector::FReal UPCGSplineData::GetDistanceAtSegmentStart(int SegmentIndex) const
+{
+	return SplineStruct.GetDistanceAlongSplineAtSplinePoint(SegmentIndex);
+}
+
 const UPCGPointData* UPCGSplineData::CreatePointData(FPCGContext* Context) const
 {
 	TRACE_CPUPROFILER_EVENT_SCOPE(UPCGSplineData::CreatePointData);
@@ -209,12 +220,17 @@ UPCGSpatialData* UPCGSplineData::CopyInternal() const
 {
 	UPCGSplineData* NewSplineData = NewObject<UPCGSplineData>();
 
-	NewSplineData->SplineStruct = SplineStruct;
-	NewSplineData->CachedBounds = CachedBounds;
+	CopySplineData(NewSplineData);
 
 	return NewSplineData;
 }
 
+void UPCGSplineData::CopySplineData(UPCGSplineData* InCopy) const
+{
+	InCopy->SplineStruct = SplineStruct;
+	InCopy->CachedBounds = CachedBounds;
+}
+
 bool UPCGSplineProjectionData::SamplePoint(const FTransform& InTransform, const FBox& InBounds, FPCGPoint& OutPoint, UPCGMetadata* OutMetadata) const
 {
 	// TODO: support metadata - we don't currently have a good representation of what metadata entries mean for non-point data
diff --git a/Engine/Plugins/PCG/Source/PCG/Private/Data/PCGSplineStruct.cpp b/Engine/Plugins/PCG/Source/PCG/Private/Data/PCGSplineStruct.cpp
--- a/Engine/Plugins/PCG/Source/PCG/Private/Data/PCGSplineStruct.cpp
+++ b/Engine/Plugins/PCG/Source/PCG/Private/Data/PCGSplineStruct.cpp
@@ -76,7 +76,7 @@ namespace PCGSplineStruct
 	}
 }
 
-void FPCGSplineStruct::Initialize(USplineComponent* InSplineComponent)
+void FPCGSplineStruct::Initialize(const USplineComponent* InSplineComponent)
 {
 	check(InSplineComponent);
 	SplineCurves = InSplineComponent->SplineCurves;
diff --git a/Engine/Plugins/PCG/Source/PCG/Private/Elements/PCGSplineSampler.cpp b/Engine/Plugins/PCG/Source/PCG/Private/Elements/PCGSplineSampler.cpp
--- a/Engine/Plugins/PCG/Source/PCG/Private/Elements/PCGSplineSampler.cpp
+++ b/Engine/Plugins/PCG/Source/PCG/Private/Elements/PCGSplineSampler.cpp
@@ -342,6 +342,7 @@ namespace PCGSplineSamplerHelpers
 		FVector::FReal Curvature = 0;
 		int SegmentIndex = 0;
 		int SubsegmentIndex = 0;
+		float InputKey = 0.0f;
 		FVector ArriveTangent = FVector::Zero();
 		FVector LeaveTangent = FVector::Zero();
 		FVector::FReal PreviousDeltaAngle = 0;
@@ -430,6 +431,7 @@ namespace PCGSplineSamplerHelpers
 			OutTransform = LineData->GetTransformAtDistance(SegmentIndex, DistanceAlongSegment, /*bWorldSpace=*/false, &OutBox);
 			OutResult.SegmentIndex = LineData->IsClosed() ? CurrentSegmentIndex : SegmentIndex;
 			OutResult.SubsegmentIndex = SubpointIndex;
+			OutResult.InputKey = LineData->GetInputKeyAtDistance(SegmentIndex, DistanceAlongSegment);
 
 			if (bComputeCurvature)
 			{
@@ -526,6 +528,7 @@ namespace PCGSplineSamplerHelpers
 			FBox& OutBox = OutResult.Box;
 			OutTransform = LineData->GetTransformAtDistance(CurrentSegmentIndex, CurrentDistance, /*bWorldSpace=*/false, &OutBox);
 			OutResult.SegmentIndex = CurrentSegmentIndex;
+			OutResult.InputKey = LineData->GetInputKeyAtDistance(CurrentSegmentIndex, CurrentDistance);
 
 			// Set min/max to half of extent
 			OutBox.Min.X *= 0.5 * DistanceIncrement / OutTransform.GetScale3D().X;
@@ -588,62 +591,66 @@ namespace PCGSplineSamplerHelpers
 			ProjectionTargetData = InProjectionTarget;
 			ProjectionParams = InProjectionParams;
 
+			UPCGMetadata* Metadata = OutPointData ? OutPointData->Metadata : nullptr;
+
 			// Initialize metadata accessors if needed
-			if (OutPointData
-				&& OutPointData->Metadata
-				&& (Params.bComputeDirectionDelta
-					|| Params.bComputeCurvature
-					|| Params.bComputeSegmentIndex
-					|| Params.bComputeSubsegmentIndex
-					|| Params.bComputeTangents
-					|| Params.bComputeAlpha
-					|| Params.bComputeDistance))
+			if (Metadata)
 			{
 				constexpr double DefaultValue = 0.0;
+
+				bHasCustomMetadata = LineData->HasCustomMetadata();
+				bSetMetadata |= bHasCustomMetadata;
+
 				if (Params.bComputeDirectionDelta)
 				{
-					NextDirectionDeltaAttribute = OutPointData->Metadata->FindOrCreateAttribute<double>(Params.NextDirectionDeltaAttribute, DefaultValue);
+					NextDirectionDeltaAttribute = Metadata->FindOrCreateAttribute<double>(Params.NextDirectionDeltaAttribute, DefaultValue);
 					bSetMetadata |= (NextDirectionDeltaAttribute != nullptr);
 				}
 
 				if (Params.bComputeCurvature)
 				{
-					CurvatureAttribute = OutPointData->Metadata->FindOrCreateAttribute<double>(Params.CurvatureAttribute, DefaultValue);
+					CurvatureAttribute = Metadata->FindOrCreateAttribute<double>(Params.CurvatureAttribute, DefaultValue);
 					bSetMetadata |= (CurvatureAttribute != nullptr);
 				}
 
 				if (Params.bComputeSegmentIndex)
 				{
-					SegmentIndexAttribute = OutPointData->Metadata->FindOrCreateAttribute<int>(Params.SegmentIndexAttribute, static_cast<int>(DefaultValue));
+					SegmentIndexAttribute = Metadata->FindOrCreateAttribute<int>(Params.SegmentIndexAttribute, static_cast<int>(DefaultValue));
 					bSetMetadata |= (SegmentIndexAttribute != nullptr);
 				}
 
 				if (Params.bComputeSubsegmentIndex)
 				{
-					SubsegmentIndexAttribute = OutPointData->Metadata->FindOrCreateAttribute<int>(Params.SubsegmentIndexAttribute, static_cast<int>(DefaultValue));
+					SubsegmentIndexAttribute = Metadata->FindOrCreateAttribute<int>(Params.SubsegmentIndexAttribute, static_cast<int>(DefaultValue));
 					bSetMetadata |= (SubsegmentIndexAttribute != nullptr);
 				}
 
 				if (Params.bComputeTangents)
 				{
-					ArriveTangentAttribute = OutPointData->Metadata->FindOrCreateAttribute<FVector>(Params.ArriveTangentAttribute, FVector::Zero());
+					ArriveTangentAttribute = Metadata->FindOrCreateAttribute<FVector>(Params.ArriveTangentAttribute, FVector::Zero());
 					bSetMetadata |= (ArriveTangentAttribute != nullptr);
 
-					LeaveTangentAttribute = OutPointData->Metadata->FindOrCreateAttribute<FVector>(Params.LeaveTangentAttribute, FVector::Zero());
+					LeaveTangentAttribute = Metadata->FindOrCreateAttribute<FVector>(Params.LeaveTangentAttribute, FVector::Zero());
 					bSetMetadata |= (LeaveTangentAttribute != nullptr);
 				}
 
 				if (Params.bComputeAlpha)
 				{
-					AlphaAttribute = OutPointData->Metadata->FindOrCreateAttribute<double>(Params.AlphaAttribute, DefaultValue);
+					AlphaAttribute = Metadata->FindOrCreateAttribute<double>(Params.AlphaAttribute, DefaultValue);
 					bSetMetadata |= (AlphaAttribute != nullptr);
 				}
 
 				if (Params.bComputeDistance)
 				{
-					DistanceAttribute = OutPointData->Metadata->FindOrCreateAttribute<double>(Params.DistanceAttribute, DefaultValue);
+					DistanceAttribute = Metadata->FindOrCreateAttribute<double>(Params.DistanceAttribute, DefaultValue);
 					bSetMetadata |= (DistanceAttribute != nullptr);
 				}
+
+				if (Params.bComputeInputKey)
+				{
+					InputKeyAttribute = Metadata->FindOrCreateAttribute<float>(Params.InputKeyAttribute, static_cast<float>(DefaultValue));
+					bSetMetadata |= (InputKeyAttribute != nullptr);
+				}
 			}
 		}
 
@@ -656,6 +663,11 @@ namespace PCGSplineSamplerHelpers
 				OutMetadata->InitializeOnSet(OutPoint.MetadataEntry);
 			}
 
+			if (bHasCustomMetadata)
+			{
+				LineData->WriteMetadataToPoint(InResult.InputKey, OutPoint, OutMetadata);
+			}
+
 			if (NextDirectionDeltaAttribute)
 			{
 				NextDirectionDeltaAttribute->SetValue(OutPoint.MetadataEntry, InResult.NextDeltaAngle);
@@ -695,6 +707,11 @@ namespace PCGSplineSamplerHelpers
 			{
 				DistanceAttribute->SetValue(OutPoint.MetadataEntry, InResult.Distance);
 			}
+
+			if (InputKeyAttribute)
+			{
+				InputKeyAttribute->SetValue(OutPoint.MetadataEntry, InResult.InputKey);
+			}
 		}
 
 		virtual void Sample(const FSamplerResult& InResult, UPCGPointData* OutPointData)
@@ -736,6 +753,7 @@ namespace PCGSplineSamplerHelpers
 		FPCGProjectionParams ProjectionParams;
 
 		bool bSetMetadata = false;
+		bool bHasCustomMetadata = false;
 		FPCGMetadataAttribute<double>* NextDirectionDeltaAttribute = nullptr;
 		FPCGMetadataAttribute<double>* CurvatureAttribute = nullptr;
 		FPCGMetadataAttribute<int>* SegmentIndexAttribute = nullptr;
@@ -744,6 +762,7 @@ namespace PCGSplineSamplerHelpers
 		FPCGMetadataAttribute<FVector>* ArriveTangentAttribute = nullptr;
 		FPCGMetadataAttribute<double>* AlphaAttribute = nullptr;
 		FPCGMetadataAttribute<double>* DistanceAttribute = nullptr;
+		FPCGMetadataAttribute<float>* InputKeyAttribute = nullptr;
 	};
 
 	/** Samples in a volume surrounding the poly line. */
diff --git a/Engine/Plugins/PCG/Source/PCG/Public/Data/PCGLandscapeSplineData.h b/Engine/Plugins/PCG/Source/PCG/Public/Data/PCGLandscapeSplineData.h
--- a/Engine/Plugins/PCG/Source/PCG/Public/Data/PCGLandscapeSplineData.h
+++ b/Engine/Plugins/PCG/Source/PCG/Public/Data/PCGLandscapeSplineData.h
@@ -17,6 +17,10 @@ class PCG_API UPCGLandscapeSplineData : public UPCGPolyLineData
 public:
 	void Initialize(ULandscapeSplinesComponent* InSplineComponent);
 
+	// ~Begin UObject interface
+	virtual void PostLoad() override;
+	// ~End UObject interface
+
 	// ~Begin UPCGData interface
 	virtual EPCGDataType GetDataType() const override { return EPCGDataType::LandscapeSpline; }
 	virtual void AddToCrc(FArchiveCrc32& Ar, bool bFullDataCrc) const override;
@@ -28,7 +32,9 @@ public:
 	virtual FVector::FReal GetSegmentLength(int SegmentIndex) const override;
 	virtual FTransform GetTransformAtDistance(int SegmentIndex, FVector::FReal Distance, bool bWorldSpace = true, FBox* OutBounds = nullptr) const override;
 	virtual FVector::FReal GetCurvatureAtDistance(int SegmentIndex, FVector::FReal Distance) const override;
+	virtual float GetInputKeyAtDistance(int SegmentIndex, FVector::FReal Distance) const override;
 	virtual void GetTangentsAtSegmentStart(int SegmentIndex, FVector& OutArriveTangent, FVector& OutLeaveTangent) const override;
+	virtual FVector::FReal GetDistanceAtSegmentStart(int SegmentIndex) const override;
 	//~End UPCGPolyLineData interface
 
 	//~Begin UPCGSpatialDataWithPointCache interface
@@ -42,9 +48,18 @@ protected:
 	virtual UPCGSpatialData* CopyInternal() const override;
 	//~End UPCGSpatialData interface
 
+	/** Recompute the reparameterization of the spline by distance. */
+	void UpdateReparamTable();
+
+	/** Get the index of the first interp point before a given distance. Also computes the Alpha describing how far (in [0, 1]) the sample is between the left and right points. */
+	void GetInterpPointAtDistance(int SegmentIndex, FVector::FReal Distance, int32& OutPointIndex, bool bComputeAlpha, FVector::FReal& OutAlpha) const;
+
 protected:
 	UPROPERTY(BlueprintReadWrite, EditAnywhere, Category = SourceData)
 	TWeakObjectPtr<ULandscapeSplinesComponent> Spline;
+
+	/** Reparameterization of the spline by distance. Useful to query the InputKey at arbitrary distance. */
+	FInterpCurveFloat ReparamTable;
 };
 
 #if UE_ENABLE_INCLUDE_ORDER_DEPRECATED_IN_5_2
diff --git a/Engine/Plugins/PCG/Source/PCG/Public/Data/PCGPolyLineData.h b/Engine/Plugins/PCG/Source/PCG/Public/Data/PCGPolyLineData.h
--- a/Engine/Plugins/PCG/Source/PCG/Public/Data/PCGPolyLineData.h
+++ b/Engine/Plugins/PCG/Source/PCG/Public/Data/PCGPolyLineData.h
@@ -48,11 +48,23 @@ public:
 	 */
 	virtual float GetAlphaAtDistance(int SegmentIndex, FVector::FReal Distance) const;
 
+	/** Get the input key at a distance along the line. InputKey is a float value in [0, N], where N is the number of control points. Each range [i, i+1] represents an interpolation from 0 to 1 across spline segment i. */
+	virtual float GetInputKeyAtDistance(int SegmentIndex, FVector::FReal Distance) const { return 0; }
+
 	/** Get the arrive and leave tangents for a control point via its segment index. */
 	virtual void GetTangentsAtSegmentStart(int SegmentIndex, FVector& OutArriveTangent, FVector& OutLeaveTangent) const;
 
+	/** Get the cumulative distance along the line to the start of a segment. */
+	virtual FVector::FReal GetDistanceAtSegmentStart(int SegmentIndex) const { return 0; }
+
 	/** True if the line is a closed loop. */
 	virtual bool IsClosed() const { return false; }
+
+	/** If a PolyLine subtype has custom metadata, it can use this virtual call to write that into the PCG Metadata per-point. */
+	virtual void WriteMetadataToPoint(float InputKey, FPCGPoint& OutPoint, UPCGMetadata* OutMetadata) const {}
+
+	/** True if the line holds custom metadata. */
+	virtual bool HasCustomMetadata() const { return false; }
 };
 
 
diff --git a/Engine/Plugins/PCG/Source/PCG/Public/Data/PCGSplineData.h b/Engine/Plugins/PCG/Source/PCG/Public/Data/PCGSplineData.h
--- a/Engine/Plugins/PCG/Source/PCG/Public/Data/PCGSplineData.h
+++ b/Engine/Plugins/PCG/Source/PCG/Public/Data/PCGSplineData.h
@@ -21,7 +21,7 @@ class PCG_API UPCGSplineData : public UPCGPolyLineData
 	GENERATED_BODY()
 
 public:
-	void Initialize(USplineComponent* InSpline);
+	void Initialize(const USplineComponent* InSpline);
 	void Initialize(const TArray<FSplinePoint>& InSplinePoints, bool bInClosedLoop, const FTransform& InTransform);
 	void Initialize(const FPCGSplineStruct& InSplineStruct);
 	void ApplyTo(USplineComponent* InSpline);
@@ -38,7 +38,9 @@ public:
 	virtual FVector GetLocationAtDistance(int SegmentIndex, FVector::FReal Distance, bool bWorldSpace = true) const override;
 	virtual FTransform GetTransformAtDistance(int SegmentIndex, FVector::FReal Distance, bool bWorldSpace = true, FBox* OutBounds = nullptr) const override;
 	virtual FVector::FReal GetCurvatureAtDistance(int SegmentIndex, FVector::FReal Distance) const override;
+	virtual float GetInputKeyAtDistance(int SegmentIndex, FVector::FReal Distance) const override;
 	virtual void GetTangentsAtSegmentStart(int SegmentIndex, FVector& OutArriveTangent, FVector& OutLeaveTangent) const override;
+	virtual FVector::FReal GetDistanceAtSegmentStart(int SegmentIndex) const override;
 	virtual bool IsClosed() const override { return SplineStruct.bClosedLoop; }
 	//~End UPCGPolyLineData interface
 
@@ -54,6 +56,8 @@ protected:
 	virtual UPCGSpatialData* CopyInternal() const override;
 	//~End UPCGSpatialData interface
 
+	virtual void CopySplineData(UPCGSplineData* InCopy) const;
+
 public:
 	// Minimal data needed to replicate the behavior from USplineComponent
 	UPROPERTY()
diff --git a/Engine/Plugins/PCG/Source/PCG/Public/Data/PCGSplineStruct.h b/Engine/Plugins/PCG/Source/PCG/Public/Data/PCGSplineStruct.h
--- a/Engine/Plugins/PCG/Source/PCG/Public/Data/PCGSplineStruct.h
+++ b/Engine/Plugins/PCG/Source/PCG/Public/Data/PCGSplineStruct.h
@@ -13,7 +13,7 @@ struct PCG_API FPCGSplineStruct
 {
 	GENERATED_BODY()
 
-	void Initialize(USplineComponent* InSplineComponent);
+	void Initialize(const USplineComponent* InSplineComponent);
 	void Initialize(const TArray<FSplinePoint>& InSplinePoints, bool bIsClosedLoop, const FTransform& InTransform);
 	void ApplyTo(USplineComponent* InSplineComponent);
 
diff --git a/Engine/Plugins/PCG/Source/PCG/Public/Elements/PCGDataFromActor.h b/Engine/Plugins/PCG/Source/PCG/Public/Elements/PCGDataFromActor.h
--- a/Engine/Plugins/PCG/Source/PCG/Public/Elements/PCGDataFromActor.h
+++ b/Engine/Plugins/PCG/Source/PCG/Public/Elements/PCGDataFromActor.h
@@ -124,7 +124,7 @@ struct FPCGDataFromActorContext : public FPCGContext
 #endif
 };
 
-class FPCGDataFromActorElement : public IPCGElement
+class PCG_API FPCGDataFromActorElement : public IPCGElement
 {
 public:
 	virtual bool CanExecuteOnlyOnMainThread(FPCGContext* Context) const override { return true; }
diff --git a/Engine/Plugins/PCG/Source/PCG/Public/Elements/PCGSplineSampler.h b/Engine/Plugins/PCG/Source/PCG/Public/Elements/PCGSplineSampler.h
--- a/Engine/Plugins/PCG/Source/PCG/Public/Elements/PCGSplineSampler.h
+++ b/Engine/Plugins/PCG/Source/PCG/Public/Elements/PCGSplineSampler.h
@@ -137,6 +137,13 @@ struct PCG_API FPCGSplineSamplerParams
 	UPROPERTY(BlueprintReadWrite, EditAnywhere, Category = Settings, meta = (EditCondition = "bComputeDistance"))
 	FName DistanceAttribute = TEXT("Distance");
 
+	UPROPERTY(BlueprintReadWrite, EditAnywhere, Category = Settings, meta = (InlineEditConditionToggle))
+	bool bComputeInputKey = false;
+
+	/** Attribute that will contain the spline input key, a float value between [0, N], where N is the number of control points. Each range [i, i+1] represents an interpolation from 0 to 1 across spline segment i. */
+	UPROPERTY(BlueprintReadWrite, EditAnywhere, Category = Settings, meta = (EditCondition = "bComputeInputKey"))
+	FName InputKeyAttribute = TEXT("InputKey");
+
 	/** If no Bounding Shape input is provided, the actor bounds are used to limit the sample generation domain.
 	* This option allows ignoring the actor bounds and generating over the entire spline. Use with caution as this
 	* may generate a lot of points.
