[PCG] Fix check failure in PCGLandscapeSplineData::GetDistanceAtSegmentStart()

diff --git a/Engine/Plugins/PCG/Source/PCG/Private/Data/PCGLandscapeSplineData.cpp b/Engine/Plugins/PCG/Source/PCG/Private/Data/PCGLandscapeSplineData.cpp
--- a/Engine/Plugins/PCG/Source/PCG/Private/Data/PCGLandscapeSplineData.cpp
+++ b/Engine/Plugins/PCG/Source/PCG/Private/Data/PCGLandscapeSplineData.cpp
@@ -211,7 +211,8 @@ void UPCGLandscapeSplineData::GetTangentsAtSegmentStart(int SegmentIndex, FVecto
 
 FVector::FReal UPCGLandscapeSplineData::GetDistanceAtSegmentStart(int SegmentIndex) const
 {
-	check(Spline.IsValid() && Spline->GetSegments().IsValidIndex(SegmentIndex));
+	// Allow SegmentIndex == NumSegments, which indicates we want the distance to the final control point, which is like saying "Start of the Nth segment".
+	check(Spline.IsValid() && SegmentIndex >= 0 && SegmentIndex <= Spline->GetSegments().Num());
 
 	const TArray<TObjectPtr<ULandscapeSplineSegment>>& Segments = Spline->GetSegments();
 	int32 ReparamIndex = 0;
