[PCG] Fix bug with Filter Data By Index not accepting Attribute Sets

#rnx

diff --git a/Engine/Plugins/PCG/Source/PCG/Private/Elements/PCGFilterByIndex.cpp b/Engine/Plugins/PCG/Source/PCG/Private/Elements/PCGFilterByIndex.cpp
--- a/Engine/Plugins/PCG/Source/PCG/Private/Elements/PCGFilterByIndex.cpp
+++ b/Engine/Plugins/PCG/Source/PCG/Private/Elements/PCGFilterByIndex.cpp
@@ -57,13 +57,13 @@ bool FPCGFilterByIndexElement::ExecuteInternal(FPCGContext* Context) const
 	const UPCGFilterByIndexSettings* Settings = Context->GetInputSettings<UPCGFilterByIndexSettings>();
 	check(Settings);
 
-	if (Context->InputData.GetInputs().IsEmpty())
+	const TArray<FPCGTaggedData> Inputs = Context->InputData.GetInputsByPin(PCGPinConstants::DefaultInputLabel);
+
+	if (Inputs.IsEmpty())
 	{
 		return true;
 	}
 
-	TArray<FPCGTaggedData> Inputs = Context->InputData.GetInputsByPin(PCGPinConstants::DefaultInputLabel);
-
 	PCGIndexing::FPCGIndexCollection FilteredIndices(Inputs.Num());
 	// Parse the indices and switch through possible issues.
 	switch (PCGParser::ParseIndexRanges(FilteredIndices, Settings->SelectedIndices))
