Exposed attribute filter operator

#rnx

diff --git a/Engine/Plugins/PCG/Source/PCG/Public/Elements/PCGAttributeFilter.h b/Engine/Plugins/PCG/Source/PCG/Public/Elements/PCGAttributeFilter.h
--- a/Engine/Plugins/PCG/Source/PCG/Public/Elements/PCGAttributeFilter.h
+++ b/Engine/Plugins/PCG/Source/PCG/Public/Elements/PCGAttributeFilter.h
@@ -8,7 +8,7 @@
 
 #include "PCGAttributeFilter.generated.h"
 
-UENUM()
+UENUM(BlueprintType)
 enum class EPCGAttributeFilterOperator : uint8
 {
 	Greater UMETA(DisplayName=">"),
