Added missing overrides on the match & set attribute node

#rnx

diff --git a/Engine/Plugins/PCG/Source/PCG/Public/Elements/PCGMatchAndSetAttributes.h b/Engine/Plugins/PCG/Source/PCG/Public/Elements/PCGMatchAndSetAttributes.h
--- a/Engine/Plugins/PCG/Source/PCG/Public/Elements/PCGMatchAndSetAttributes.h
+++ b/Engine/Plugins/PCG/Source/PCG/Public/Elements/PCGMatchAndSetAttributes.h
@@ -65,7 +65,7 @@ protected:
 
 public:
 	/** Controls whether selection of the attribute set values to copy will be done by matching point-to-attribute set (true) or done randomly (false). */
-	UPROPERTY(BlueprintReadWrite, EditAnywhere, Category = Settings)
+	UPROPERTY(BlueprintReadWrite, EditAnywhere, Category = Settings, meta = (PCG_Overridable))
 	bool bMatchAttributes = false;
 
 	/** Attribute from the point data to select & match. */
@@ -73,7 +73,7 @@ public:
 	FPCGAttributePropertyInputSelector InputAttribute;
 
 	/** Attribute from the attribute set to match against. */
-	UPROPERTY(BlueprintReadWrite, EditAnywhere, Category = Settings, meta = (EditCondition = "bMatchAttributes"))
+	UPROPERTY(BlueprintReadWrite, EditAnywhere, Category = Settings, meta = (EditCondition = "bMatchAttributes", PCG_Overridable))
 	FName MatchAttribute = NAME_None;
 
 	/** Controls whether points that have no valid match in the attribute set are kept as is (default values) or removed from the output. */
@@ -96,7 +96,7 @@ public:
 	FPCGAttributePropertyInputSelector MaxDistanceInputAttribute;
 
 	/** Controls whether we will use the attribute provided in the Input Weight Attribute to perform entry selection. */
-	UPROPERTY(BlueprintReadWrite, EditAnywhere, Category = Settings)
+	UPROPERTY(BlueprintReadWrite, EditAnywhere, Category = Settings, meta = (PCG_Overridable))
 	bool bUseInputWeightAttribute = false;
 
 	/** Input weight from the points, assumed to be in the [0, 1] range. */
