[PCG] Fix dirtying PCG actor with generated graph output data when Preview editing mode

diff --git a/Engine/Plugins/PCG/Source/PCG/Private/PCGComponent.cpp b/Engine/Plugins/PCG/Source/PCG/Private/PCGComponent.cpp
--- a/Engine/Plugins/PCG/Source/PCG/Private/PCGComponent.cpp
+++ b/Engine/Plugins/PCG/Source/PCG/Private/PCGComponent.cpp
@@ -424,7 +424,7 @@ void UPCGComponent::PostProcessGraph(const FBox& InNewBounds, bool bInGenerated,
 					{
 						FPCGTaggedData& DuplicatedTaggedData = GeneratedGraphOutput.TaggedData.Add_GetRef(TaggedData);
 						DuplicatedTaggedData.Data = DuplicatedData;
-						DuplicatedData->Rename(nullptr, this);
+						DuplicatedData->Rename(nullptr, this, IsInPreviewMode() ? REN_DoNotDirty : REN_None);
 						DuplicatedData->Flatten();
 					}
 				}
