[PCG] Add WallTime for profiling to see the performance of tasks that are paused.

#rnx

diff --git a/Engine/Plugins/PCG/Source/PCG/Private/Utils/PCGExtraCapture.cpp b/Engine/Plugins/PCG/Source/PCG/Private/Utils/PCGExtraCapture.cpp
--- a/Engine/Plugins/PCG/Source/PCG/Private/Utils/PCGExtraCapture.cpp
+++ b/Engine/Plugins/PCG/Source/PCG/Private/Utils/PCGExtraCapture.cpp
@@ -27,7 +27,8 @@ void PCGUtils::FExtraCapture::Update(const PCGUtils::FScopedCall& InScopedCall)
 		return;
 	}
 
-	const double ThisFrameTime = FPlatformTime::Seconds() - InScopedCall.StartTime;
+	const double CurrentTime = FPlatformTime::Seconds();
+	const double ThisFrameTime = CurrentTime - InScopedCall.StartTime;
 
 	FScopeLock ScopedLock(&Lock);
 
@@ -39,11 +40,24 @@ void PCGUtils::FExtraCapture::Update(const PCGUtils::FScopedCall& InScopedCall)
 		Timer = FCallTime(); // reset it
 		break;
 	case EPCGExecutionPhase::PrepareData:
+		if (Timer.PrepareDataFrameCount == 0)
+		{
+			Timer.PrepareDataStartTime = InScopedCall.StartTime;
+		}
+
+		Timer.PrepareDataFrameCount++;
 		Timer.PrepareDataTime = ThisFrameTime;
+		Timer.PrepareDataWallTime = CurrentTime - Timer.PrepareDataStartTime;
 		break;
 	case EPCGExecutionPhase::Execute:
+		if (Timer.ExecutionFrameCount == 0)
+		{
+			Timer.ExecutionStartTime = InScopedCall.StartTime;
+		}
+
 		Timer.ExecutionTime += ThisFrameTime;
 		Timer.ExecutionFrameCount++;
+		Timer.ExecutionWallTime = CurrentTime - Timer.ExecutionStartTime;
 
 		Timer.MaxExecutionFrameTime = FMath::Max(Timer.MaxExecutionFrameTime, ThisFrameTime);
 		Timer.MinExecutionFrameTime = FMath::Min(Timer.MinExecutionFrameTime, ThisFrameTime);
@@ -153,7 +167,9 @@ namespace PCGUtils
 			BuildTreeInfo(Child);
 
 			Info.CallTime.PrepareDataTime += Child.CallTime.PrepareDataTime;
+			Info.CallTime.PrepareDataWallTime += Child.CallTime.PrepareDataWallTime;
 			Info.CallTime.ExecutionTime += Child.CallTime.ExecutionTime;
+			Info.CallTime.ExecutionWallTime += Child.CallTime.ExecutionWallTime;
 			Info.CallTime.PostExecuteTime += Child.CallTime.PostExecuteTime;
 			Info.CallTime.MinExecutionFrameTime = FMath::Min(Info.CallTime.MinExecutionFrameTime, Child.CallTime.MinExecutionFrameTime);
 			Info.CallTime.MaxExecutionFrameTime = FMath::Max(Info.CallTime.MaxExecutionFrameTime, Child.CallTime.MaxExecutionFrameTime);
diff --git a/Engine/Plugins/PCG/Source/PCG/Public/Utils/PCGExtraCapture.h b/Engine/Plugins/PCG/Source/PCG/Public/Utils/PCGExtraCapture.h
--- a/Engine/Plugins/PCG/Source/PCG/Public/Utils/PCGExtraCapture.h
+++ b/Engine/Plugins/PCG/Source/PCG/Public/Utils/PCGExtraCapture.h
@@ -18,12 +18,19 @@ namespace PCGUtils
 	struct PCG_API FCallTime
 	{
 		// sum of all frames
-		double ExecutionTime = 0;
+		double ExecutionStartTime = 0.0;
+		double ExecutionTime = 0.0;
+		double ExecutionWallTime = 0.0;
 		// how many frames
 		int32 ExecutionFrameCount = 0;
 		double MinExecutionFrameTime = MAX_dbl;
 		double MaxExecutionFrameTime = 0.0;
+
+		double PrepareDataStartTime = 0.0;
 		double PrepareDataTime = 0.0;
+		double PrepareDataWallTime = 0.0;
+		int32 PrepareDataFrameCount = 0;
+
 		double PostExecuteTime = 0.0;
 	};
 
diff --git a/Engine/Plugins/PCG/Source/PCGEditor/Private/SPCGEditorGraphProfilingView.cpp b/Engine/Plugins/PCG/Source/PCGEditor/Private/SPCGEditorGraphProfilingView.cpp
--- a/Engine/Plugins/PCG/Source/PCGEditor/Private/SPCGEditorGraphProfilingView.cpp
+++ b/Engine/Plugins/PCG/Source/PCGEditor/Private/SPCGEditorGraphProfilingView.cpp
@@ -24,25 +24,37 @@ namespace PCGEditorGraphProfilingView
 	/** Names of the columns in the attribute list */
 	const FName NAME_Node = FName(TEXT("Node"));
 	const FName NAME_PrepareDataTime = FName(TEXT("PrepareDataTime"));
+	const FName NAME_PrepareDataWallTime = FName(TEXT("PrepareData_WallTime"));
 	const FName NAME_MinExecutionFrameTime = FName(TEXT("MinFrameTime"));
 	const FName NAME_MaxExecutionFrameTime = FName(TEXT("MaxFrameTime"));
-	const FName NAME_TotalExecutionTime = FName(TEXT("TotalExecutionTime"));
+	const FName NAME_ExecutionTime = FName(TEXT("ExecutionTime"));
+	const FName NAME_ExecutionWallTime = FName(TEXT("Execution_WallTime"));
 	const FName NAME_NbExecutionFrames = FName(TEXT("NbExecutionFrames"));
+	const FName NAME_TotalTime = FName(TEXT("TotalTime"));
+	const FName NAME_TotalWallTime = FName(TEXT("Total_WallTime"));
 
 	/** Labels of the columns */
 	const FText TEXT_NodeLabel = LOCTEXT("NodeLabel", "Node");
-	const FText TEXT_PrepareDataTimeLabel = LOCTEXT("PrepareDataTimeLabel", "PrepareData Phase (ms)");
+	const FText TEXT_PrepareDataTimeLabel = LOCTEXT("PrepareDataTimeLabel", "PrepareData Time (ms)");
+	const FText TEXT_PrepareDataWallTimeLabel = LOCTEXT("PrepareDataWallTimeLabel", "PrepareData WallTime (ms)");
 	const FText TEXT_MinExecutionFrameTimeLabel = LOCTEXT("MinExecutionFrameTimeLabel", "Min Frame Time (ms)");
 	const FText TEXT_MaxExecutionFrameTimeLabel = LOCTEXT("MaxExecutionFrameTimeLabel", "Max Frame Time (ms)");
-	const FText TEXT_TotalExecutionTimeLabel = LOCTEXT("TotalExecutionTimeLabel", "Total Time (ms)");
+	const FText TEXT_ExecutionTimeLabel = LOCTEXT("ExecutionTimeLabel", "Execution Time (ms)");
+	const FText TEXT_ExecutionWallTimeLabel = LOCTEXT("ExecutionWallTimeLabel", "Execution Wall Time (ms)");
 	const FText TEXT_NbExecutionFramesLabel = LOCTEXT("NbExecutionFramesLabel", "Exec Frames");
+	const FText TEXT_TotalTimeLabel = LOCTEXT("TotalTimeLabel", "Total Time (ms)");
+	const FText TEXT_TotalWallTimeLabel = LOCTEXT("TotalWallTimeLabel", "Total Wall Time (ms)");
 
 	/** Tooltips */
 	const FText TEXT_PrepareDataTimeTooltip = LOCTEXT("PrepareDataTimeTooltip", "Cost of the PrepareData execution phase which some nodes use to process the incoming data.");
+	const FText TEXT_PrepareDataWallTimeTooltip = LOCTEXT("PrepareDataWallTimeTooltip", "Total real time elapsed between prepare data first being called until completion, including any wait/sleep time.");
 	const FText TEXT_MinExecutionFrameTimeTooltip = LOCTEXT("MinExecutionFrameTimeTooltip", "The minimum time spent of all execution frames.");
 	const FText TEXT_MaxExecutionFrameTimeTooltip = LOCTEXT("MaxExecutionFrameTimeTooltip", "The maximum time spent of all execution frames.");
-	const FText TEXT_TotalExecutionTimeTooltip = LOCTEXT("TotalExecutionTimeTooltip", "The total time spent, summed over all execution frames and phases.");
+	const FText TEXT_ExecutionTimeTooltip = LOCTEXT("ExecutionTimeTooltip", "The total time spent for execution, summed over all execution frames.");
+	const FText TEXT_ExecutionWallTimeTooltip = LOCTEXT("ExecutionWallTimeTooltip", "Total real time elapsed between execute first being called until completion, including any wait/sleep time.");
 	const FText TEXT_NbExecutionFramesTooltip = LOCTEXT("NbExecutionFramesTooltip", "The number of frames in which one or more execution phases were executed.");
+	const FText TEXT_TotalTimeTooltip = LOCTEXT("TotalTimeTooltip", "The total time spent in this node, summed over all execution and prepare data frames.");
+	const FText TEXT_TotalWallTimeTooltip = LOCTEXT("TotalWallTimeTooltip", "Total real time elapsed between the first call until completion, including any wait/sleep time.");
 }
 
 void SPCGProfilingListViewItemRow::Construct(const FArguments& InArgs, const TSharedRef<STableViewBase>& InOwnerTableView, const PCGProfilingListViewItemPtr& Item)
@@ -82,15 +94,36 @@ TSharedRef<SWidget> SPCGProfilingListViewItemRow::GenerateWidgetForColumn(const
 			// In ms
 			ColumnData = ((InternalItem->CallTime.MaxExecutionFrameTime >= 0) ? FText::AsNumber(InternalItem->CallTime.MaxExecutionFrameTime * 1000.0) : FText());
 		}
-		else if (ColumnId == PCGEditorGraphProfilingView::NAME_TotalExecutionTime)
+		else if (ColumnId == PCGEditorGraphProfilingView::NAME_ExecutionTime)
 		{
 			// In ms
 			ColumnData = FText::AsNumber(InternalItem->CallTime.ExecutionTime * 1000.0);
 		}
+		else if (ColumnId == PCGEditorGraphProfilingView::NAME_ExecutionWallTime)
+		{
+			// In ms
+			ColumnData = FText::AsNumber(InternalItem->CallTime.ExecutionWallTime * 1000.0);
+		}
 		else if (ColumnId == PCGEditorGraphProfilingView::NAME_PrepareDataTime)
 		{
+			// In ms
 			ColumnData = FText::AsNumber(InternalItem->CallTime.PrepareDataTime * 1000.0);
 		}
+		else if (ColumnId == PCGEditorGraphProfilingView::NAME_PrepareDataWallTime)
+		{
+			// In ms
+			ColumnData = FText::AsNumber(InternalItem->CallTime.PrepareDataWallTime * 1000.0);
+		}
+		else if (ColumnId == PCGEditorGraphProfilingView::NAME_TotalTime)
+		{
+			// In ms
+			ColumnData = FText::AsNumber((InternalItem->CallTime.PrepareDataTime + InternalItem->CallTime.ExecutionTime) * 1000.0);
+		}
+		else if (ColumnId == PCGEditorGraphProfilingView::NAME_TotalWallTime)
+		{
+			// In ms
+			ColumnData = FText::AsNumber((InternalItem->CallTime.PrepareDataWallTime + InternalItem->CallTime.ExecutionWallTime) * 1000.0);
+		}
 	}
 
 	return SNew(STextBlock)
@@ -135,7 +168,7 @@ void SPCGEditorGraphProfilingView::Construct(const FArguments& InArgs, TSharedPt
 		PCGEditor->OnInspectedStackChangedDelegate.AddSP(this, &SPCGEditorGraphProfilingView::OnDebugStackChanged);
 	}
 
-	SortingColumn = PCGEditorGraphProfilingView::NAME_TotalExecutionTime;
+	SortingColumn = PCGEditorGraphProfilingView::NAME_TotalTime;
 	SortMode = EColumnSortMode::Descending;
 	ListViewHeader = CreateHeaderRowWidget();
 
@@ -204,6 +237,22 @@ void SPCGEditorGraphProfilingView::Construct(const FArguments& InArgs, TSharedPt
 				.Text(this, &SPCGEditorGraphProfilingView::GetTotalTimeLabel)
 				.MinDesiredWidth(50.0f)
 			]
+			+SHorizontalBox::Slot()
+			.Padding(FMargin(30.0f, 0.0f, 2.0f, 0.0f))
+			.VAlign(EVerticalAlignment::VAlign_Center)
+			.AutoWidth()
+			[
+				SNew(STextBlock)
+				.Text(LOCTEXT("TotalExecutionWallTime", "Total Wall Time:"))
+			]
+			+SHorizontalBox::Slot()
+			.VAlign(EVerticalAlignment::VAlign_Center)
+			.AutoWidth()
+			[
+				SNew(STextBlock)
+				.Text(this, &SPCGEditorGraphProfilingView::GetTotalWallTimeLabel)
+				.MinDesiredWidth(50.0f)
+			]
 		]
 		+SVerticalBox::Slot()
 		.FillHeight(1.0f)
@@ -261,6 +310,15 @@ TSharedRef<SHeaderRow> SPCGEditorGraphProfilingView::CreateHeaderRowWidget()
 		.OnSort(this, &SPCGEditorGraphProfilingView::OnSortColumnHeader)
 		.InitialSortMode(EColumnSortMode::Descending)
 		.DefaultTooltip(PCGEditorGraphProfilingView::TEXT_PrepareDataTimeTooltip)
+		+ SHeaderRow::Column(PCGEditorGraphProfilingView::NAME_PrepareDataWallTime)
+		.ManualWidth(170.0f)
+		.DefaultLabel(PCGEditorGraphProfilingView::TEXT_PrepareDataWallTimeLabel)
+		.HAlignHeader(HAlign_Center)
+		.HAlignCell(HAlign_Right)
+		.SortMode(this, &SPCGEditorGraphProfilingView::GetColumnSortMode, PCGEditorGraphProfilingView::NAME_PrepareDataWallTime)
+		.OnSort(this, &SPCGEditorGraphProfilingView::OnSortColumnHeader)
+		.InitialSortMode(EColumnSortMode::Descending)
+		.DefaultTooltip(PCGEditorGraphProfilingView::TEXT_PrepareDataWallTimeTooltip)
 		+ SHeaderRow::Column(PCGEditorGraphProfilingView::NAME_NbExecutionFrames)
 		.ManualWidth(105.0f)
 		.DefaultLabel(PCGEditorGraphProfilingView::TEXT_NbExecutionFramesLabel)
@@ -288,15 +346,42 @@ TSharedRef<SHeaderRow> SPCGEditorGraphProfilingView::CreateHeaderRowWidget()
 		.OnSort(this, &SPCGEditorGraphProfilingView::OnSortColumnHeader)
 		.InitialSortMode(EColumnSortMode::Descending)
 		.DefaultTooltip(PCGEditorGraphProfilingView::TEXT_MaxExecutionFrameTimeTooltip)
-		+ SHeaderRow::Column(PCGEditorGraphProfilingView::NAME_TotalExecutionTime)
+		+ SHeaderRow::Column(PCGEditorGraphProfilingView::NAME_ExecutionTime)
+		.ManualWidth(120.0f)
+		.DefaultLabel(PCGEditorGraphProfilingView::TEXT_ExecutionTimeLabel)
+		.HAlignHeader(HAlign_Center)
+		.HAlignCell(HAlign_Right)
+		.SortMode(this, &SPCGEditorGraphProfilingView::GetColumnSortMode, PCGEditorGraphProfilingView::NAME_ExecutionTime)
+		.OnSort(this, &SPCGEditorGraphProfilingView::OnSortColumnHeader)
+		.InitialSortMode(EColumnSortMode::Descending)
+		.DefaultTooltip(PCGEditorGraphProfilingView::TEXT_ExecutionTimeTooltip)
+		+ SHeaderRow::Column(PCGEditorGraphProfilingView::NAME_ExecutionWallTime)
 		.ManualWidth(120.0f)
-		.DefaultLabel(PCGEditorGraphProfilingView::TEXT_TotalExecutionTimeLabel)
+		.DefaultLabel(PCGEditorGraphProfilingView::TEXT_ExecutionWallTimeLabel)
 		.HAlignHeader(HAlign_Center)
 		.HAlignCell(HAlign_Right)
-		.SortMode(this, &SPCGEditorGraphProfilingView::GetColumnSortMode, PCGEditorGraphProfilingView::NAME_TotalExecutionTime)
+		.SortMode(this, &SPCGEditorGraphProfilingView::GetColumnSortMode, PCGEditorGraphProfilingView::NAME_ExecutionWallTime)
 		.OnSort(this, &SPCGEditorGraphProfilingView::OnSortColumnHeader)
 		.InitialSortMode(EColumnSortMode::Descending)
-		.DefaultTooltip(PCGEditorGraphProfilingView::TEXT_TotalExecutionTimeTooltip);
+		.DefaultTooltip(PCGEditorGraphProfilingView::TEXT_ExecutionWallTimeTooltip)
+		+ SHeaderRow::Column(PCGEditorGraphProfilingView::NAME_TotalTime)
+		.ManualWidth(120.0f)
+		.DefaultLabel(PCGEditorGraphProfilingView::TEXT_TotalTimeLabel)
+		.HAlignHeader(HAlign_Center)
+		.HAlignCell(HAlign_Right)
+		.SortMode(this, &SPCGEditorGraphProfilingView::GetColumnSortMode, PCGEditorGraphProfilingView::NAME_TotalTime)
+		.OnSort(this, &SPCGEditorGraphProfilingView::OnSortColumnHeader)
+		.InitialSortMode(EColumnSortMode::Descending)
+		.DefaultTooltip(PCGEditorGraphProfilingView::TEXT_TotalTimeTooltip)
+		+ SHeaderRow::Column(PCGEditorGraphProfilingView::NAME_TotalWallTime)
+		.ManualWidth(120.0f)
+		.DefaultLabel(PCGEditorGraphProfilingView::TEXT_TotalWallTimeLabel)
+		.HAlignHeader(HAlign_Center)
+		.HAlignCell(HAlign_Right)
+		.SortMode(this, &SPCGEditorGraphProfilingView::GetColumnSortMode, PCGEditorGraphProfilingView::NAME_TotalWallTime)
+		.OnSort(this, &SPCGEditorGraphProfilingView::OnSortColumnHeader)
+		.InitialSortMode(EColumnSortMode::Descending)
+		.DefaultTooltip(PCGEditorGraphProfilingView::TEXT_TotalWallTimeTooltip);
 }
 
 FText SPCGEditorGraphProfilingView::GetTotalTimeLabel() const
@@ -304,6 +389,11 @@ FText SPCGEditorGraphProfilingView::GetTotalTimeLabel() const
 	return FText::Format(LOCTEXT("GraphTotalTimeLabel", "{0} s"), TotalTime);
 }
 
+FText SPCGEditorGraphProfilingView::GetTotalWallTimeLabel() const
+{
+	return FText::Format(LOCTEXT("GraphTotalWallTimeLabel", "{0} s"), TotalWallTime);
+}
+
 void SPCGEditorGraphProfilingView::OnSortColumnHeader(const EColumnSortPriority::Type SortPriority, const FName& ColumnId, const EColumnSortMode::Type NewSortMode)
 {
 	if (SortingColumn == ColumnId)
@@ -441,6 +531,7 @@ FReply SPCGEditorGraphProfilingView::Refresh()
 	PCGUtils::FCallTreeInfo TreeInfo = Component->ExtraCapture.CalculateCallTreeInfo(Component);
 
 	TotalTime = TreeInfo.CallTime.ExecutionTime;
+	TotalWallTime = TreeInfo.CallTime.ExecutionWallTime;
 
 	ListViewItems.Reserve(TreeInfo.Children.Num());
 
@@ -460,6 +551,10 @@ FReply SPCGEditorGraphProfilingView::Refresh()
 				{
 					isLess = A->CallTime.PrepareDataTime < B->CallTime.PrepareDataTime;
 				}
+				else if (SortingColumn == PCGEditorGraphProfilingView::NAME_PrepareDataWallTime)
+				{
+					isLess = A->CallTime.PrepareDataWallTime < B->CallTime.PrepareDataWallTime;
+				}
 				else if (SortingColumn == PCGEditorGraphProfilingView::NAME_MinExecutionFrameTime)
 				{
 					isLess = A->CallTime.MinExecutionFrameTime < B->CallTime.MinExecutionFrameTime;
@@ -468,14 +563,26 @@ FReply SPCGEditorGraphProfilingView::Refresh()
 				{
 					isLess = A->CallTime.MaxExecutionFrameTime < B->CallTime.MaxExecutionFrameTime;
 				}
-				else if (SortingColumn == PCGEditorGraphProfilingView::NAME_TotalExecutionTime)
+				else if (SortingColumn == PCGEditorGraphProfilingView::NAME_ExecutionTime)
 				{
 					isLess = A->CallTime.ExecutionTime < B->CallTime.ExecutionTime;
 				}
+				else if (SortingColumn == PCGEditorGraphProfilingView::NAME_ExecutionWallTime)
+				{
+					isLess = A->CallTime.ExecutionWallTime < B->CallTime.ExecutionWallTime;
+				}
 				else if (SortingColumn == PCGEditorGraphProfilingView::NAME_NbExecutionFrames)
 				{
 					isLess = A->CallTime.ExecutionFrameCount < B->CallTime.ExecutionFrameCount;
 				}
+				else if (SortingColumn == PCGEditorGraphProfilingView::NAME_TotalTime)
+				{
+					isLess = (A->CallTime.ExecutionTime + A->CallTime.PrepareDataTime) < (B->CallTime.ExecutionTime + B->CallTime.PrepareDataTime);
+				}
+				else if (SortingColumn == PCGEditorGraphProfilingView::NAME_TotalWallTime)
+				{
+					isLess = (A->CallTime.ExecutionWallTime + A->CallTime.PrepareDataWallTime) < (B->CallTime.ExecutionWallTime + B->CallTime.PrepareDataWallTime);
+				}
 
 				return SortMode == EColumnSortMode::Ascending ? isLess : !isLess;
 			});
diff --git a/Engine/Plugins/PCG/Source/PCGEditor/Private/SPCGEditorGraphProfilingView.h b/Engine/Plugins/PCG/Source/PCGEditor/Private/SPCGEditorGraphProfilingView.h
--- a/Engine/Plugins/PCG/Source/PCGEditor/Private/SPCGEditorGraphProfilingView.h
+++ b/Engine/Plugins/PCG/Source/PCGEditor/Private/SPCGEditorGraphProfilingView.h
@@ -71,6 +71,7 @@ private:
 	void OnSortColumnHeader(const EColumnSortPriority::Type SortPriority, const FName& ColumnId, const EColumnSortMode::Type NewSortMode);
 	EColumnSortMode::Type GetColumnSortMode(const FName ColumnId) const;
 	FText GetTotalTimeLabel() const;
+	FText GetTotalWallTimeLabel() const;
 
 	/** Pointer back to the PCG editor that owns us */
 	TWeakPtr<FPCGEditor> PCGEditorPtr;
@@ -90,6 +91,7 @@ private:
 	EColumnSortMode::Type SortMode = EColumnSortMode::Type::None;
 
 	double TotalTime = 0.0;
+	double TotalWallTime = 0.0;
 
 	bool bExpandSubgraph = true;
 };
