[PCG] Another issue where Preview component gets package dirtied. Flip order of flatten and outering.

diff --git a/Engine/Plugins/PCG/Source/PCG/Private/PCGComponent.cpp b/Engine/Plugins/PCG/Source/PCG/Private/PCGComponent.cpp
--- a/Engine/Plugins/PCG/Source/PCG/Private/PCGComponent.cpp
+++ b/Engine/Plugins/PCG/Source/PCG/Private/PCGComponent.cpp
@@ -424,8 +424,11 @@ void UPCGComponent::PostProcessGraph(const FBox& InNewBounds, bool bInGenerated,
 					{
 						FPCGTaggedData& DuplicatedTaggedData = GeneratedGraphOutput.TaggedData.Add_GetRef(TaggedData);
 						DuplicatedTaggedData.Data = DuplicatedData;
-						DuplicatedData->Rename(nullptr, this, IsInPreviewMode() ? REN_DoNotDirty : REN_None);
+
+						// NOTE: Flatten before outering to this, as doing the flatten afterwards can dirty this package.
 						DuplicatedData->Flatten();
+
+						DuplicatedData->Rename(nullptr, this, IsInPreviewMode() ? REN_DoNotDirty : REN_None);
 					}
 				}
 			}
