[PCG] Avoid pointless regenerations when changing debug object within a single component.

diff --git a/Engine/Plugins/PCG/Source/PCGEditor/Private/PCGEditor.cpp b/Engine/Plugins/PCG/Source/PCGEditor/Private/PCGEditor.cpp
--- a/Engine/Plugins/PCG/Source/PCGEditor/Private/PCGEditor.cpp
+++ b/Engine/Plugins/PCG/Source/PCGEditor/Private/PCGEditor.cpp
@@ -235,12 +235,17 @@ void FPCGEditor::SetStackBeingInspected(const FPCGStack& FullStack)
 {
 	UPCGComponent* OldComponent = PCGComponentBeingInspected.Get();
 	UPCGComponent* NewComponent = const_cast<UPCGComponent*>(FullStack.GetRootComponent());
+	const bool bComponentChanged = (NewComponent != OldComponent);
 
 	if (StackBeingInspected != FullStack)
 	{
 		if (OldComponent)
 		{
-			OldComponent->DisableInspection();
+			if (bComponentChanged)
+			{
+				OldComponent->DisableInspection();
+			}
+
 			if (PCGGraphBeingEdited)
 			{
 				PCGGraphBeingEdited->DisableInspection();
@@ -256,7 +261,11 @@ void FPCGEditor::SetStackBeingInspected(const FPCGStack& FullStack)
 
 		if (NewComponent)
 		{
-			PCGComponentBeingInspected->EnableInspection();
+			if (bComponentChanged)
+			{
+				PCGComponentBeingInspected->EnableInspection();
+			}
+
 			if (PCGGraphBeingEdited)
 			{
 				PCGGraphBeingEdited->EnableInspection();
@@ -285,7 +294,7 @@ void FPCGEditor::SetStackBeingInspected(const FPCGStack& FullStack)
 
 void FPCGEditor::UpdateDebugAfterComponentSelection(UPCGComponent* InOldComponent, UPCGComponent* InNewComponent, bool bInNewComponentStartedInspecting)
 {
-	if (!ensure(PCGGraphBeingEdited))
+	if (!ensure(PCGGraphBeingEdited) || (InOldComponent == InNewComponent))
 	{
 		return;
 	}
