[PCG] Fix issue where selecting a folder in the debug object tree displays all nodes as culled. Instead clear inspection state.

diff --git a/Engine/Plugins/PCG/Source/PCGEditor/Private/PCGEditor.cpp b/Engine/Plugins/PCG/Source/PCGEditor/Private/PCGEditor.cpp
--- a/Engine/Plugins/PCG/Source/PCGEditor/Private/PCGEditor.cpp
+++ b/Engine/Plugins/PCG/Source/PCGEditor/Private/PCGEditor.cpp
@@ -233,65 +233,76 @@ UPCGEditorGraph* FPCGEditor::GetPCGEditorGraph()
 
 void FPCGEditor::SetStackBeingInspected(const FPCGStack& FullStack)
 {
+	if (FullStack == StackBeingInspected)
+	{
+		// No-op if we're already inspecting this stack.
+		return;
+	}
+
 	UPCGComponent* OldComponent = PCGComponentBeingInspected.Get();
 	UPCGComponent* NewComponent = const_cast<UPCGComponent*>(FullStack.GetRootComponent());
 	const bool bComponentChanged = (NewComponent != OldComponent);
 
-	if (StackBeingInspected != FullStack)
+	if (OldComponent)
 	{
-		if (OldComponent)
+		if (bComponentChanged)
 		{
-			if (bComponentChanged)
-			{
-				OldComponent->DisableInspection();
-			}
+			OldComponent->DisableInspection();
+		}
 
-			if (PCGGraphBeingEdited)
-			{
-				PCGGraphBeingEdited->DisableInspection();
-			}
+		if (PCGGraphBeingEdited)
+		{
+			PCGGraphBeingEdited->DisableInspection();
 		}
+	}
 
-		const bool bNewComponentStartedInspecting = NewComponent && !NewComponent->IsInspecting();
+	const bool bNewComponentStartedInspecting = NewComponent && !NewComponent->IsInspecting();
 
-		PCGComponentBeingInspected = NewComponent;
+	PCGComponentBeingInspected = NewComponent;
 
-		StackBeingInspected = FullStack;
-		OnInspectedStackChangedDelegate.Broadcast(StackBeingInspected);
+	StackBeingInspected = FullStack;
+	OnInspectedStackChangedDelegate.Broadcast(StackBeingInspected);
 
-		if (NewComponent)
+	if (NewComponent)
+	{
+		if (bComponentChanged)
 		{
-			if (bComponentChanged)
-			{
-				PCGComponentBeingInspected->EnableInspection();
-			}
+			PCGComponentBeingInspected->EnableInspection();
+		}
 
-			if (PCGGraphBeingEdited)
-			{
-				PCGGraphBeingEdited->EnableInspection();
-			}
+		if (PCGGraphBeingEdited)
+		{
+			PCGGraphBeingEdited->EnableInspection();
 		}
+	}
 
-		UpdateDebugAfterComponentSelection(OldComponent, NewComponent, bNewComponentStartedInspecting);
+	UpdateDebugAfterComponentSelection(OldComponent, NewComponent, bNewComponentStartedInspecting);
 
-		check(PCGEditorGraph);
-		for (UEdGraphNode* Node : PCGEditorGraph->Nodes)
+	check(PCGEditorGraph);
+	for (UEdGraphNode* Node : PCGEditorGraph->Nodes)
+	{
+		if (UPCGEditorGraphNodeBase* PCGNode = Cast<UPCGEditorGraphNodeBase>(Node))
 		{
-			if (UPCGEditorGraphNodeBase* PCGNode = Cast<UPCGEditorGraphNodeBase>(Node))
-			{
-				// Update now that component has changed. Will fire OnNodeChanged if necessary.
-				EPCGChangeType ChangeType = PCGNode->UpdateErrorsAndWarnings();
-				ChangeType |= PCGNode->UpdateStructuralVisualization(NewComponent, &StackBeingInspected);
+			// Update now that component has changed. Will fire OnNodeChanged if necessary.
+			EPCGChangeType ChangeType = PCGNode->UpdateErrorsAndWarnings();
+			ChangeType |= PCGNode->UpdateStructuralVisualization(NewComponent, &StackBeingInspected);
 
-				if (ChangeType != EPCGChangeType::None)
-				{
-					PCGNode->ReconstructNode();
-				}
+			if (ChangeType != EPCGChangeType::None)
+			{
+				PCGNode->ReconstructNode();
 			}
 		}
 	}
 }
 
+void FPCGEditor::ClearStackBeingInspected()
+{
+	if (GetStackBeingInspected())
+	{
+		SetStackBeingInspected(FPCGStack());
+	}
+}
+
 void FPCGEditor::UpdateDebugAfterComponentSelection(UPCGComponent* InOldComponent, UPCGComponent* InNewComponent, bool bInNewComponentStartedInspecting)
 {
 	if (!ensure(PCGGraphBeingEdited) || (InOldComponent == InNewComponent))
diff --git a/Engine/Plugins/PCG/Source/PCGEditor/Private/PCGEditor.h b/Engine/Plugins/PCG/Source/PCGEditor/Private/PCGEditor.h
--- a/Engine/Plugins/PCG/Source/PCGEditor/Private/PCGEditor.h
+++ b/Engine/Plugins/PCG/Source/PCGEditor/Private/PCGEditor.h
@@ -49,6 +49,9 @@ public:
 	/** Sets the execution stack that want to inspect. */
 	void SetStackBeingInspected(const FPCGStack& FullStack);
 
+	/** Clear current inspection. */
+	void ClearStackBeingInspected();
+
 	/** Gets the PCG component we are debugging */
 	UPCGComponent* GetPCGComponentBeingInspected() const { return const_cast<UPCGComponent*>(StackBeingInspected.GetRootComponent()); }
 	
diff --git a/Engine/Plugins/PCG/Source/PCGEditor/Private/SPCGEditorGraphDebugObjectTree.cpp b/Engine/Plugins/PCG/Source/PCGEditor/Private/SPCGEditorGraphDebugObjectTree.cpp
--- a/Engine/Plugins/PCG/Source/PCGEditor/Private/SPCGEditorGraphDebugObjectTree.cpp
+++ b/Engine/Plugins/PCG/Source/PCGEditor/Private/SPCGEditorGraphDebugObjectTree.cpp
@@ -908,16 +908,14 @@ void SPCGEditorGraphDebugObjectTree::OnSelectionChanged(FPCGEditorGraphDebugObje
 		return;
 	}
 
-	if (!InItem)
+	// Only attempt to inspect stacks that correspond to the edited graph. Other graphs need to be inspected in their own editor.
+	if (InItem && SelectedStack.GetGraphForCurrentFrame() == PCGEditor.Pin()->GetPCGGraph())
 	{
-		PCGEditor.Pin()->SetStackBeingInspected(FPCGStack());
-		return;
+		PCGEditor.Pin()->SetStackBeingInspected(SelectedStack);
 	}
-
-	const FPCGStack* InspectedStack = PCGEditor.Pin()->GetStackBeingInspected();
-	if (!InspectedStack || SelectedStack != *InspectedStack)
+	else
 	{
-		PCGEditor.Pin()->SetStackBeingInspected(SelectedStack);
+		PCGEditor.Pin()->ClearStackBeingInspected();
 	}
 }
 
