[PCG] Fixed a minor bug with Print String node where the message was not properly cleaned up in BiomeCore

* Also Owner will print ActorLabel rather than full name
* Prefixing with Node name is no longer requiring "By Component"

#rnx

diff --git a/Engine/Plugins/PCG/Source/PCG/Private/Elements/PCGPrintElement.cpp b/Engine/Plugins/PCG/Source/PCG/Private/Elements/PCGPrintElement.cpp
--- a/Engine/Plugins/PCG/Source/PCG/Private/Elements/PCGPrintElement.cpp
+++ b/Engine/Plugins/PCG/Source/PCG/Private/Elements/PCGPrintElement.cpp
@@ -11,7 +11,6 @@
 #include "GameFramework/Actor.h"
 
 #if WITH_EDITOR
-#include "EngineGlobals.h"
 #include "Editor/EditorEngine.h"
 #endif // WITH_EDITOR
 
@@ -22,17 +21,17 @@ namespace PCGPrintElementConstants
 	FText Delimiter = LOCTEXT("Delimiter", "::");
 }
 
-#if WITH_EDITOR
-
 namespace PCGPrintElementHelpers
 {
 	void CleanUpMessage(uint64& OutMessageHashKey)
 	{
+#if WITH_EDITOR
 		if (GEngine && (OutMessageHashKey != (uint64)-1) && GEngine->OnScreenDebugMessageExists(OutMessageHashKey))
 		{
 			GEngine->RemoveOnScreenDebugMessage(OutMessageHashKey);
 			OutMessageHashKey = (uint64)-1;
 		}
+#endif // WITH_EDITOR
 	}
 }
 
@@ -42,8 +41,6 @@ bool UPCGManagedDebugStringMessageKey::Release(bool bHardRelease, TSet<TSoftObje
 	return Super::Release(bHardRelease, OutActorsToDelete);
 }
 
-#endif // WITH_EDITOR
-
 TArray<FPCGPinProperties> UPCGPrintElementSettings::InputPinProperties() const
 {
 	TArray<FPCGPinProperties> PinProperties;
@@ -82,11 +79,15 @@ bool FPCGPrintElement::ExecuteInternal(FPCGContext* Context) const
 
 	// Build the message prefixes
 	const UPCGComponent* Component = Context->SourceComponent->GetOriginalComponent();
-	if (Component)
+	if (Settings->bPrintPerComponent && Component)
 	{
 		if (Settings->bPrefixWithOwner)
 		{
+#if WITH_EDITOR
+			Prefixes.Emplace(Component->GetOwner() ? Component->GetOwner()->GetActorLabel() : LOCTEXT("MissingOwner", "Missing Owner").ToString());
+#else
 			Prefixes.Emplace(Component->GetOwner() ? Component->GetOwner()->GetName() : LOCTEXT("MissingOwner", "Missing Owner").ToString());
+#endif // WITH_EDITOR
 		}
 
 		if (Settings->bPrefixWithComponent)
diff --git a/Engine/Plugins/PCG/Source/PCG/Public/Elements/PCGPrintElement.h b/Engine/Plugins/PCG/Source/PCG/Public/Elements/PCGPrintElement.h
--- a/Engine/Plugins/PCG/Source/PCG/Public/Elements/PCGPrintElement.h
+++ b/Engine/Plugins/PCG/Source/PCG/Public/Elements/PCGPrintElement.h
@@ -25,15 +25,17 @@ class UPCGManagedDebugStringMessageKey : public UPCGManagedResource
 {
 	GENERATED_BODY()
 
+public:
 #if WITH_EDITOR
 	// Debug should always be transient
 	virtual void ChangeTransientState(EPCGEditorDirtyMode NewEditingMode) override {}
+#endif // WITH_EDITOR
+
 	virtual bool Release(bool bHardRelease, TSet<TSoftObjectPtr<AActor>>& OutActorsToDelete) override;
 	virtual bool CanBeUsed() const override { return false; }
 
-public:
+	UPROPERTY()
 	uint64 HashKey = (uint64)-1;
-#endif // WITH_EDITOR
 };
 
 /**
@@ -109,7 +111,7 @@ public:
     bool bPrefixWithGraph = true;
 
 	/** Prefix the message with the name of the node. */
-	UPROPERTY(BlueprintReadWrite, EditAnywhere, Category = Settings, meta = (EditCondition = "bPrintPerComponent", EditConditionHides))
+	UPROPERTY(BlueprintReadWrite, EditAnywhere, Category = Settings)
 	bool bPrefixWithNode = true;
 
 	/** Enable the functionality of this node. Disable to bypass printing. */
