[PCG] Created new Visualize Attribute node

+ Prints an attribute's value to string with an offset for each point input

diff --git a/Engine/Plugins/PCG/Source/PCG/Private/Elements/PCGVisualizeAttribute.cpp b/Engine/Plugins/PCG/Source/PCG/Private/Elements/PCGVisualizeAttribute.cpp
new file mode 100644
--- /dev/null
+++ b/Engine/Plugins/PCG/Source/PCG/Private/Elements/PCGVisualizeAttribute.cpp
@@ -0,0 +1,135 @@
+// Copyright Epic Games, Inc. All Rights Reserved.
+
+#include "Elements/PCGVisualizeAttribute.h"
+
+#include "PCGComponent.h"
+#include "PCGContext.h"
+#include "PCGDebugDrawComponent.h"
+#include "Data/PCGPointData.h"
+#include "Metadata/Accessors/PCGAttributeAccessorHelpers.h"
+
+#define LOCTEXT_NAMESPACE "PCGVisualizeAttributeElement"
+
+FPCGElementPtr UPCGVisualizeAttributeSettings::CreateElement() const
+{
+	return MakeShared<FPCGVisualizeAttribute>();
+}
+
+bool FPCGVisualizeAttribute::ExecuteInternal(FPCGContext* Context) const
+{
+	TRACE_CPUPROFILER_EVENT_SCOPE(FPCGVisualizeAttribute::ExecuteInternal);
+
+	Context->OutputData = Context->InputData;
+
+#if !(UE_BUILD_SHIPPING || UE_BUILD_TEST) || USE_LOGGING_IN_SHIPPING
+	const UPCGVisualizeAttributeSettings* Settings = Context->GetInputSettings<UPCGVisualizeAttributeSettings>();
+	check(Settings);
+
+	if (!Settings->bVisualizeEnabled)
+	{
+		return true;
+	}
+
+	if (!ensure(Context->SourceComponent.IsValid()))
+	{
+		return true;
+	}
+
+	const TArray<FPCGTaggedData> Inputs = Context->InputData.GetInputsByPin(PCGPinConstants::DefaultInputLabel);
+
+	for (const FPCGTaggedData& InputData : Inputs)
+	{
+		const UPCGPointData* PointData = Cast<const UPCGPointData>(InputData.Data);
+		if (!PointData)
+		{
+			PCGLog::LogErrorOnGraph(LOCTEXT("InvalidInputType", "Input data type is not supported, only supports Point Data."), Context);
+			continue;
+		}
+
+		// No points, early out
+		if (PointData->GetPoints().IsEmpty())
+		{
+			continue;
+		}
+
+		if (!PointData->Metadata)
+		{
+			PCGLog::LogErrorOnGraph(LOCTEXT("InvalidMetadata", "Input has no metadata."), Context);
+			continue;
+		}
+
+		FPCGAttributePropertyInputSelector InputSource = Settings->AttributeSource.CopyAndFixLast(PointData);
+
+		TUniquePtr<const IPCGAttributeAccessor> Accessor = PCGAttributeAccessorHelpers::CreateConstAccessor(PointData, InputSource);
+		TUniquePtr<const IPCGAttributeAccessorKeys> Keys = PCGAttributeAccessorHelpers::CreateConstKeys(PointData, InputSource);
+
+		if (!Accessor.IsValid() || !Keys.IsValid())
+		{
+			PCGLog::LogWarningOnGraph(FText::Format(LOCTEXT("AttributeDoesNotExist", "Input attribute/property '{0}' does not exist - skipped."), InputSource.GetDisplayText()), Context);
+			continue;
+		}
+
+		TArray<FString> ValuesToString;
+		ValuesToString.SetNum(Keys->GetNum());
+
+		// Get all the attribute values at once in range, for efficiency
+		if (!Accessor->GetRange<FString>(ValuesToString, 0, *Keys, EPCGAttributeAccessorFlags::AllowBroadcast))
+		{
+			PCGLog::LogErrorOnGraph(LOCTEXT("InvalidValues", "Could not get the range of attribute values."), Context);
+			continue;
+		}
+
+		UPCGDebugDrawComponent* DebugDrawComponent = nullptr;
+		if (AActor* TargetActor = Context->GetTargetActor(PointData))
+		{
+			DebugDrawComponent = NewObject<UPCGDebugDrawComponent>(TargetActor, FName(LOCTEXT("PCGDebugDrawComponent", "PCG Debug Draw Component").ToString(), RF_Transient));
+
+			TargetActor->Modify(/*bAlwaysMarkDirty=*/false);
+			DebugDrawComponent->RegisterComponent();
+			TargetActor->AddInstanceComponent(DebugDrawComponent);
+			UPCGManagedDebugDrawComponent* ManagedComponent = NewObject<UPCGManagedDebugDrawComponent>(Context->SourceComponent.Get());
+			ManagedComponent->GeneratedComponent = DebugDrawComponent;
+			Context->SourceComponent->AddToManagedResources(ManagedComponent);
+		}
+		else
+		{
+			PCGLog::LogErrorOnGraph(LOCTEXT("NoTargetActor", "Internal error: Could not find target actor."), Context);
+			continue;
+		}
+
+		const TArray<FPCGPoint>& Points = PointData->GetPoints();
+		TArray<FDebugRenderSceneProxy::FText3d> DebugStrings;
+		DebugStrings.Reserve(FMath::Min(Points.Num(), Settings->PointLimit));
+
+		// TODO: Optimize which points will be drawn by distance, or other measures, rather than by Index
+		// Debug messages are already culled by camera frustum, but a point limit is desired to prevent overloading
+		for (int I = 0; I < Points.Num() && I < Settings->PointLimit; ++I)
+		{
+			FStringBuilderBase CompoundedString;
+
+			CompoundedString += Settings->CustomPrefixString;
+
+			if (Settings->bPrefixWithIndex)
+			{
+				CompoundedString += FString::Printf(TEXT("[%d]"), I);
+			}
+
+			if (Settings->bPrefixWithAttributeName)
+			{
+				CompoundedString += FString::Printf(TEXT("[%s]"), *InputSource.GetDisplayText().ToString());
+			}
+
+			CompoundedString = FString::Printf(TEXT("%s %s"), *CompoundedString, *ValuesToString[I]);
+
+			DebugStrings.Emplace(CompoundedString.ToString(), Points[I].Transform.GetLocation() + Settings->LocalOffset, Settings->Color);
+		}
+
+		DebugDrawComponent->AddDebugStrings(DebugStrings);
+		DebugDrawComponent->StartTimer(Settings->Duration);
+	}
+#endif // !(UE_BUILD_SHIPPING || UE_BUILD_TEST) || USE_LOGGING_IN_SHIPPING
+
+	return true;
+}
+
+#undef LOCTEXT_NAMESPACE
diff --git a/Engine/Plugins/PCG/Source/PCG/Private/PCGDebugDrawComponent.cpp b/Engine/Plugins/PCG/Source/PCG/Private/PCGDebugDrawComponent.cpp
new file mode 100644
--- /dev/null
+++ b/Engine/Plugins/PCG/Source/PCG/Private/PCGDebugDrawComponent.cpp
@@ -0,0 +1,91 @@
+// Copyright Epic Games, Inc. All Rights Reserved.
+
+#include "PCGDebugDrawComponent.h"
+
+#include "PCGModule.h"
+
+#if WITH_EDITOR
+#include "Editor.h"
+#endif // WITH_EDITOR
+
+#include "TimerManager.h"
+#include "Engine/World.h"
+
+#define LOCTEXT_NAMESPACE "PCGDebugVisuals"
+
+void UPCGDebugDrawComponent::AddDebugString(const FDebugRenderSceneProxy::FText3d& TextIn3D)
+{
+	SceneProxy->Texts.Add(TextIn3D);
+}
+
+void UPCGDebugDrawComponent::AddDebugString(const FString& InString, const FVector& InLocation, const FLinearColor& InColor)
+{
+	SceneProxy->Texts.Emplace(InString, InLocation, InColor);
+}
+
+void UPCGDebugDrawComponent::AddDebugStrings(const TArrayView<FDebugRenderSceneProxy::FText3d> TextIn3DArray)
+{
+	SceneProxy->Texts.Append(TextIn3DArray);
+}
+
+void UPCGDebugDrawComponent::StartTimer(const float DurationInMilliseconds)
+{
+#if WITH_EDITOR
+	if (GEditor)
+	{
+		GEditor->GetTimerManager()->SetTimer(TimerHandle, FTimerDelegate::CreateUObject(this, &UPCGDebugDrawComponent::OnTimerElapsed), DurationInMilliseconds, /*InbLoop=*/true);
+	}
+	else
+#endif //WITH_EDITOR
+	{
+		UWorld* World = GetWorld();
+		if (ensure(World))
+		{
+			World->GetTimerManager().SetTimer(TimerHandle, FTimerDelegate::CreateUObject(this, &UPCGDebugDrawComponent::OnTimerElapsed), DurationInMilliseconds, /*InbLoop=*/true);
+		}
+	}
+
+	if (TimerHandle.IsValid())
+	{
+		DebugDrawDelegateHelper.InitDelegateHelper(SceneProxy);
+	}
+}
+
+void UPCGDebugDrawComponent::ClearTimer()
+{
+#if WITH_EDITOR
+	if (GEditor)
+	{
+		GEditor->GetTimerManager()->ClearTimer(TimerHandle);
+	}
+	else
+#endif //WITH_EDITOR
+	{
+		UWorld* World = GetWorld();
+		if (ensure(World))
+		{
+			World->GetTimerManager().ClearTimer(TimerHandle);
+		}
+	}
+}
+
+FDebugRenderSceneProxy* UPCGDebugDrawComponent::CreateDebugSceneProxy()
+{
+	SceneProxy = new FDebugRenderSceneProxy(this);
+	SceneProxy->ViewFlagName = PCGEngineShowFlags::Debug;
+	return SceneProxy;
+}
+
+void UPCGDebugDrawComponent::OnUnregister()
+{
+	ClearTimer();
+	Super::OnUnregister();
+}
+
+void UPCGDebugDrawComponent::OnTimerElapsed()
+{
+	ClearTimer();
+	DestroyComponent();
+}
+
+#undef LOCTEXT_NAMESPACE
diff --git a/Engine/Plugins/PCG/Source/PCG/Private/PCGModule.cpp b/Engine/Plugins/PCG/Source/PCG/Private/PCGModule.cpp
--- a/Engine/Plugins/PCG/Source/PCG/Private/PCGModule.cpp
+++ b/Engine/Plugins/PCG/Source/PCG/Private/PCGModule.cpp
@@ -10,6 +10,7 @@
 #if WITH_EDITOR
 #include "Elements/PCGDifferenceElement.h"
 #include "ISettingsModule.h"
+#include "ShowFlags.h"
 #include "Tests/Determinism/PCGDeterminismNativeTests.h"
 #include "Tests/Determinism/PCGDifferenceDeterminismTest.h"
 #endif
@@ -46,6 +47,8 @@ void FPCGModule::StartupModule()
 	PCGDeterminismTests::FNativeTestRegistry::Create();
 
 	RegisterNativeElementDeterminismTests();
+
+	FEngineShowFlags::RegisterCustomShowFlag(PCGEngineShowFlags::Debug, /*DefaultEnabled=*/true, EShowFlagGroup::SFG_Developer, LOCTEXT("ShowFlagDisplayName", "PCG Debug"));
 }
 
 void FPCGModule::ShutdownModule()
diff --git a/Engine/Plugins/PCG/Source/PCG/Public/Elements/PCGVisualizeAttribute.h b/Engine/Plugins/PCG/Source/PCG/Public/Elements/PCGVisualizeAttribute.h
new file mode 100644
--- /dev/null
+++ b/Engine/Plugins/PCG/Source/PCG/Public/Elements/PCGVisualizeAttribute.h
@@ -0,0 +1,80 @@
+// Copyright Epic Games, Inc. All Rights Reserved.
+
+#pragma once
+
+#include "PCGElement.h"
+#include "PCGSettings.h"
+
+#include "DebugRenderSceneProxy.h"
+
+#include "PCGVisualizeAttribute.generated.h"
+
+/**
+ * Visualizes a selected attribute on screen at each point's transform.
+ */
+UCLASS(BlueprintType, ClassGroup = (Procedural))
+class UPCGVisualizeAttributeSettings : public UPCGSettings
+{
+	GENERATED_BODY()
+
+public:
+	//~Begin UPCGSettings interface
+#if WITH_EDITOR
+	virtual FName GetDefaultNodeName() const override { return FName(TEXT("VisualizeAttribute")); }
+	virtual FText GetDefaultNodeTitle() const override { return NSLOCTEXT("PCGVisualizeAttributeElement", "NodeTitle", "Visualize Attribute"); }
+	virtual FText GetNodeTooltipText() const override { return NSLOCTEXT("PCGVisualizeAttributeElement", "NodeTooltip", "Visualizes a selected attribute on screen at each point's transform."); }
+	virtual EPCGSettingsType GetType() const override { return EPCGSettingsType::Debug; }
+#endif
+
+protected:
+	virtual FPCGElementPtr CreateElement() const override;
+	virtual TArray<FPCGPinProperties> InputPinProperties() const override { return Super::DefaultPointInputPinProperties(); }
+	virtual TArray<FPCGPinProperties> OutputPinProperties() const override { return Super::DefaultPointOutputPinProperties(); }
+	//~End UPCGSettings interface
+
+public:
+	/** This attribute will be have it's value printed in proximity to each input point's transform. */
+	UPROPERTY(BlueprintReadWrite, EditAnywhere, Category = Settings, meta = (PCG_Overridable))
+	FPCGAttributePropertyInputSelector AttributeSource;
+
+	/** A custom added prefix to which the attribute value will be appended. */
+	UPROPERTY(BlueprintReadWrite, EditAnywhere, Category = Settings, meta = (PCG_Overridable))
+	FString CustomPrefixString;
+
+	/** Prefix the printed value with the point's index. */
+	UPROPERTY(BlueprintReadWrite, EditAnywhere, Category = Settings, meta = (PCG_Overridable))
+	bool bPrefixWithIndex = false;
+
+	/** Prefix the printed value with the attribute's name. */
+	UPROPERTY(BlueprintReadWrite, EditAnywhere, Category = Settings, meta = (PCG_Overridable))
+	bool bPrefixWithAttributeName = true;
+
+	/** A local offset from the point's location to draw the text. */
+	UPROPERTY(BlueprintReadWrite, EditAnywhere, Category = Settings, meta = (PCG_Overridable))
+	FVector LocalOffset = FVector(0, 0, 0);
+
+	/** The color of the on displayed value. */
+	UPROPERTY(BlueprintReadWrite, EditAnywhere, Category = Settings, meta = (PCG_Overridable))
+	FColor Color = FColor::Cyan;
+
+	/** The duration (in seconds) of the displayed value. */
+	UPROPERTY(BlueprintReadWrite, EditAnywhere, Category = Settings, meta = (PCG_Overridable, ClampMin = "0.0"))
+	double Duration = 30.0;
+
+	/** The limit of points to draw debug messages. */
+	UPROPERTY(BlueprintReadWrite, EditAnywhere, Category = Settings, meta = (PCG_Overridable, ClampMin = "1", ClampMax="4096"))
+	int32 PointLimit = 256;
+
+	/** The visualizer is enabled. Useful for dynamically overriding. */
+	UPROPERTY(BlueprintReadWrite, EditAnywhere, Category = Settings, meta = (PCG_Overridable))
+	bool bVisualizeEnabled = true;
+};
+
+class FPCGVisualizeAttribute : public IPCGElement
+{
+public:
+	virtual bool IsCacheable(const UPCGSettings* InSettings) const override { return false; }
+
+protected:
+	virtual bool ExecuteInternal(FPCGContext* Context) const override;
+};
\ No newline at end of file
diff --git a/Engine/Plugins/PCG/Source/PCG/Public/PCGDebugDrawComponent.h b/Engine/Plugins/PCG/Source/PCG/Public/PCGDebugDrawComponent.h
new file mode 100644
--- /dev/null
+++ b/Engine/Plugins/PCG/Source/PCG/Public/PCGDebugDrawComponent.h
@@ -0,0 +1,59 @@
+// Copyright Epic Games, Inc. All Rights Reserved.
+
+#pragma once
+
+#include "PCGManagedResource.h"
+
+#include "DebugRenderSceneProxy.h"
+#include "Debug/DebugDrawComponent.h"
+
+#include "PCGDebugDrawComponent.generated.h"
+
+/** 
+* Manages the editor DebugDrawComponent for displaying debug information within SIE and PIE.
+* Note: Must be a managed component to ensure proper resource management during generation and cleanup.
+*/
+UCLASS(ClassGroup = (Procedural))
+class UPCGManagedDebugDrawComponent : public UPCGManagedComponent
+{
+	GENERATED_BODY()
+
+#if WITH_EDITOR
+	// Debug Draw should always be transient, regardless of editing mode
+	virtual void ChangeTransientState(EPCGEditorDirtyMode NewEditingMode) override {}
+#endif // WITH_EDITOR
+};
+
+/**
+ * A transient component intended to visualize attribute information of PCG Point Data, such as printing the string value in World Space. Future support may include: vectors, flow fields, heatmaps, etc.
+ * Note: Uses a built in timer for users that may want debug information for a limited timeframe during generation.
+ */
+UCLASS(Transient, ClassGroup = (Procedural))
+class UPCGDebugDrawComponent : public UDebugDrawComponent
+{
+	GENERATED_BODY()
+
+public:
+	void AddDebugString(const FDebugRenderSceneProxy::FText3d& TextIn3D);
+	void AddDebugString(const FString& InString, const FVector& InLocation, const FLinearColor& InColor);
+	void AddDebugStrings(const TArrayView<FDebugRenderSceneProxy::FText3d> TextIn3DArray);
+
+	void StartTimer(float DurationInMilliseconds);
+
+protected:
+	//~Begin UDebugDrawComponent interface
+	virtual FDebugRenderSceneProxy* CreateDebugSceneProxy() override;
+	//~End UDebugDrawComponent interface
+
+	//~Begin UPrimitiveComponent interface
+	virtual void OnUnregister() override;
+	//~End UPrimitiveComponent interface
+
+	virtual void OnTimerElapsed();
+
+	FDebugRenderSceneProxy* SceneProxy = nullptr;
+	FTimerHandle TimerHandle;
+
+private:
+	void ClearTimer();
+};
\ No newline at end of file
diff --git a/Engine/Plugins/PCG/Source/PCG/Public/PCGModule.h b/Engine/Plugins/PCG/Source/PCG/Public/PCGModule.h
--- a/Engine/Plugins/PCG/Source/PCG/Public/PCGModule.h
+++ b/Engine/Plugins/PCG/Source/PCG/Public/PCGModule.h
@@ -17,6 +17,11 @@ namespace PCGLog
 	PCG_API void LogWarningOnGraph(const FText& InMsg, const FPCGContext* InContext = nullptr);
 }
 
+namespace PCGEngineShowFlags
+{
+	static constexpr TCHAR Debug[] = TEXT("PCGDebug");
+}
+
 // Stats
 DECLARE_STATS_GROUP(TEXT("PCG"), STATGROUP_PCG, STATCAT_Advanced);
 
