[PCG] Fix: Native nodes now forward tags

diff --git a/Engine/Plugins/PCG/Source/PCG/Private/Elements/PCGCreateAttribute.cpp b/Engine/Plugins/PCG/Source/PCG/Private/Elements/PCGCreateAttribute.cpp
--- a/Engine/Plugins/PCG/Source/PCG/Private/Elements/PCGCreateAttribute.cpp
+++ b/Engine/Plugins/PCG/Source/PCG/Private/Elements/PCGCreateAttribute.cpp
@@ -287,7 +287,7 @@ bool FPCGAddAttributeElement::ExecuteInternal(FPCGContext* Context) const
 				OutputMetadata->AddEntry();
 			}
 
-			FPCGTaggedData& NewData = Context->OutputData.TaggedData.Emplace_GetRef();
+			FPCGTaggedData& NewData = Context->OutputData.TaggedData.Add_GetRef(Inputs[i]);
 			NewData.Data = OutputData;
 		}
 
