Added reduce mode sum

diff --git a/Engine/Plugins/PCG/Source/PCG/Private/Elements/PCGAttributeReduceElement.cpp b/Engine/Plugins/PCG/Source/PCG/Private/Elements/PCGAttributeReduceElement.cpp
--- a/Engine/Plugins/PCG/Source/PCG/Private/Elements/PCGAttributeReduceElement.cpp
+++ b/Engine/Plugins/PCG/Source/PCG/Private/Elements/PCGAttributeReduceElement.cpp
@@ -107,6 +107,24 @@ namespace PCGAttributeReduceElement
 				});
 		}
 	}
+
+	template<typename T>
+	bool Sum(const IPCGAttributeAccessorKeys& Keys, const IPCGAttributeAccessor& Accessor, T& OutValue)
+	{
+		if constexpr (!PCG::Private::MetadataTraits<T>::CanSubAdd)
+		{
+			return false;
+		}
+		else
+		{
+			OutValue = PCG::Private::MetadataTraits<T>::ZeroValue();
+
+			return PCGMetadataElementCommon::ApplyOnAccessor<T>(Keys, Accessor, [&OutValue](const T& InValue, int32)
+			{
+				OutValue = PCG::Private::MetadataTraits<T>::Add(OutValue, InValue);
+			});
+		}
+	}
 }
 
 #if WITH_EDITOR
@@ -282,6 +300,9 @@ bool FPCGAttributeReduceElement::ExecuteInternal(FPCGContext* Context) const
 			case EPCGAttributeReduceOperation::Min:
 				bSuccess = PCGAttributeReduceElement::MinMax<AttributeType, /*bIsMin*/true>(*Keys, *Accessor, OutputValue);
 				break;
+			case EPCGAttributeReduceOperation::Sum:
+				bSuccess = PCGAttributeReduceElement::Sum<AttributeType>(*Keys, *Accessor, OutputValue);
+				break;
 			default:
 				break;
 			}
diff --git a/Engine/Plugins/PCG/Source/PCG/Public/Elements/PCGAttributeReduceElement.h b/Engine/Plugins/PCG/Source/PCG/Public/Elements/PCGAttributeReduceElement.h
--- a/Engine/Plugins/PCG/Source/PCG/Public/Elements/PCGAttributeReduceElement.h
+++ b/Engine/Plugins/PCG/Source/PCG/Public/Elements/PCGAttributeReduceElement.h
@@ -13,7 +13,8 @@ enum class EPCGAttributeReduceOperation
 {
 	Average,
 	Max,
-	Min
+	Min,
+	Sum
 };
 
 /**
