[PCG] The Distance Node has been upgraded to utilize an input accessor rather than an attribute name

* Also, it was refactored to be more optimized with setting a range of attributes

#rnx

diff --git a/Engine/Plugins/PCG/Source/PCG/Private/Elements/PCGDistance.cpp b/Engine/Plugins/PCG/Source/PCG/Private/Elements/PCGDistance.cpp
--- a/Engine/Plugins/PCG/Source/PCG/Private/Elements/PCGDistance.cpp
+++ b/Engine/Plugins/PCG/Source/PCG/Private/Elements/PCGDistance.cpp
@@ -8,6 +8,7 @@
 #include "Data/PCGPointData.h"
 #include "Elements/PCGGather.h"
 #include "Helpers/PCGAsync.h"
+#include "Metadata/Accessors/PCGAttributeAccessorHelpers.h"
 
 #define LOCTEXT_NAMESPACE "PCGDistanceElement"
 
@@ -82,6 +83,30 @@ FPCGElementPtr UPCGDistanceSettings::CreateElement() const
 	return MakeShared<FPCGDistanceElement>();
 }
 
+void UPCGDistanceSettings::PostLoad()
+{
+	Super::PostLoad();
+
+#if WITH_EDITOR
+	if (AttributeName_DEPRECATED != PCGDistanceConstants::DefaultOutputAttributeName)
+	{
+		// "None" was previously used to indicate that nothing should be written to attribute
+		if (AttributeName_DEPRECATED == NAME_None)
+		{
+			bOutputToAttribute = false;
+			OutputAttribute.SetAttributeName(PCGDistanceConstants::DefaultOutputAttributeName);
+		}
+		else
+		{
+			bOutputToAttribute = true;
+			OutputAttribute.SetAttributeName(AttributeName_DEPRECATED);
+		}
+
+		AttributeName_DEPRECATED = PCGDistanceConstants::DefaultOutputAttributeName;
+	}
+#endif // WITH_EDITOR
+}
+
 bool FPCGDistanceElement::ExecuteInternal(FPCGContext* Context) const
 {
 	TRACE_CPUPROFILER_EVENT_SCOPE(FPCGDistanceElement::Execute);
@@ -97,7 +122,6 @@ bool FPCGDistanceElement::ExecuteInternal(FPCGContext* Context) const
 	const UPCGDistanceSettings* Settings = Context->GetInputSettings<UPCGDistanceSettings>();
 	check(Settings);
 
-	const FName AttributeName = Settings->AttributeName;
 	const bool bSetDensity = Settings->bSetDensity;
 	const bool bOutputDistanceVector = Settings->bOutputDistanceVector;
 	const PCGDistanceShape SourceShape = Settings->SourceShape;
@@ -157,95 +181,161 @@ bool FPCGDistanceElement::ExecuteInternal(FPCGContext* Context) const
 		OutputData->InitializeFromData(SourcePointData);
 		Outputs.Add_GetRef(Source).Data = OutputData;
 
-		FPCGMetadataAttribute<float>* ScalarAttribute = nullptr;
-		FPCGMetadataAttribute<FVector>* VectorAttribute = nullptr;
+		if (Settings->bOutputToAttribute && Settings->OutputAttribute.IsBasicAttribute())
+		{
+			if (bOutputDistanceVector)
+			{
+				OutputData->Metadata->FindOrCreateAttribute<FVector>(Settings->OutputAttribute.GetAttributeName());
+			}
+			else
+			{
+				OutputData->Metadata->FindOrCreateAttribute<double>(Settings->OutputAttribute.GetAttributeName());
+			}
+		}
+
+		TUniquePtr<IPCGAttributeAccessor> Accessor;
+		TUniquePtr<IPCGAttributeAccessorKeys> Keys;
+
+		if (Settings->bOutputToAttribute)
+		{
+			Accessor = PCGAttributeAccessorHelpers::CreateAccessor(OutputData, Settings->OutputAttribute);
+		}
+
+		// If the selected attribute is a property or extra property and not the correct type, invalidate the accessor
+		if (Accessor.IsValid())
+		{
+			using PCG::Private::MetadataTypes;
+
+			if ((bOutputDistanceVector && !IsBroadcastableOrConstructible(Accessor->GetUnderlyingType(), MetadataTypes<FVector>::Id)) ||
+				(!bOutputDistanceVector && !IsBroadcastableOrConstructible(Accessor->GetUnderlyingType(), MetadataTypes<double>::Id)))
+			{
+				PCGE_LOG(Warning, GraphAndLog, LOCTEXT("InvalidAccessorType", "Selected type for Output Attribute is incompatible with distance as output."));
+				Accessor = nullptr;
+			}
+		}
+
+		TArray<FPCGPoint>& OutPoints = OutputData->GetMutablePoints();
+		const TArray<FPCGPoint>& SourcePoints = SourcePointData->GetPoints();
+		OutPoints.SetNumUninitialized(SourcePoints.Num());
 
-		if (AttributeName != NAME_None && !bOutputDistanceVector)
+		if (Accessor.IsValid())
 		{
-			ScalarAttribute = OutputData->Metadata->FindOrCreateAttribute<float>(AttributeName, 0.0f);
+			Keys = PCGAttributeAccessorHelpers::CreateKeys(OutputData, Settings->OutputAttribute);
+
+			if (!Keys.IsValid())
+			{
+				PCGE_LOG(Warning, GraphAndLog, LOCTEXT("CannotCreateAccessorKeys", "Cannot create accessor keys on output points."));
+				Accessor = nullptr;
+			}
 		}
 
-		if (AttributeName != NAME_None && bOutputDistanceVector)
+		struct TemporaryResultCache
 		{
-			VectorAttribute = OutputData->Metadata->FindOrCreateAttribute<FVector>(AttributeName, FVector::ZeroVector);
+			TArray<double> Distances;
+			TArray<FVector> DistanceVectors;
+		} ResultCache;
+
+		// Set up a cache so we can set all the attributes in a single range set
+		if (Accessor.IsValid())
+		{
+			if (bOutputDistanceVector)
+			{
+				ResultCache.DistanceVectors.SetNumUninitialized(SourcePoints.Num());
+			}
+			else
+			{
+				ResultCache.Distances.SetNumUninitialized(SourcePoints.Num());
+			}
 		}
 
-		FPCGAsync::AsyncPointProcessing(Context, SourcePointData->GetPoints(), OutputData->GetMutablePoints(),
-			[OutputData, SourceShape, TargetShape, &TargetPointDatas, MaximumDistance, MaximumDistanceRecip, ScalarAttribute, VectorAttribute, bSetDensity](const FPCGPoint& SourcePoint, FPCGPoint& OutPoint) {
+		auto ProcessDistance = [SourceShape, TargetShape, &TargetPointDatas, MaximumDistance, MaximumDistanceRecip, bSetDensity, bOutputDistanceVector, &OutPoints, &SourcePoints, &ResultCache, bWriteToAttribute = Accessor.IsValid()](int32 ReadIndex, int32 WriteIndex)
+		{
+			FPCGPoint& OutPoint = OutPoints[WriteIndex];
+			const FPCGPoint& SourcePoint = SourcePoints[ReadIndex];
 
-				OutPoint = SourcePoint;
+			OutPoint = SourcePoint;
 
-				const FBoxSphereBounds SourceQueryBounds = FBoxSphereBounds(FBox(SourcePoint.BoundsMin - FVector(MaximumDistance), SourcePoint.BoundsMax + FVector(MaximumDistance))).TransformBy(SourcePoint.Transform);
+			const FBoxSphereBounds SourceQueryBounds = FBoxSphereBounds(FBox(SourcePoint.BoundsMin - FVector(MaximumDistance), SourcePoint.BoundsMax + FVector(MaximumDistance))).TransformBy(SourcePoint.Transform);
 
-				const FVector SourceCenter = SourcePoint.Transform.TransformPosition(SourcePoint.GetLocalCenter());
+			const FVector SourceCenter = SourcePoint.Transform.TransformPosition(SourcePoint.GetLocalCenter());
 
-				double MinDistanceSquared = MaximumDistance * MaximumDistance;
-				FVector MinDistanceVector = FVector::ZeroVector;
+			double MinDistanceSquared = MaximumDistance * MaximumDistance;
+			FVector MinDistanceVector = FVector::ZeroVector;
 
-				// Signed distance field for calculating the closest point of source and target
-				auto CalculateSDF = [&MinDistanceSquared, &MinDistanceVector, &SourcePoint, SourceCenter, SourceShape, TargetShape](const FPCGPointRef& TargetPointRef)
+			// Signed distance field for calculating the closest point of source and target
+			auto CalculateSDF = [&MinDistanceSquared, &MinDistanceVector, &SourcePoint, SourceCenter, SourceShape, TargetShape](const FPCGPointRef& TargetPointRef)
+			{
+				// If the source pointer and target pointer are the same, ignore distance to the exact same point
+				if (&SourcePoint == TargetPointRef.Point)
 				{
-					// If the source pointer and target pointer are the same, ignore distance to the exact same point
-					if (&SourcePoint == TargetPointRef.Point)
-					{
-						return;
-					}
-					
-					const FPCGPoint& TargetPoint = *TargetPointRef.Point;
-					const FVector& TargetCenter = TargetPointRef.Bounds.Origin;
-
-					const FVector SourceShapePos = PCGDistance::CalcPosition(SourceShape, SourcePoint, TargetPoint, SourceCenter, TargetCenter);
-					const FVector TargetShapePos = PCGDistance::CalcPosition(TargetShape, TargetPoint, SourcePoint, TargetCenter, SourceCenter);
-
-					const FVector ToTargetShapeDir = TargetShapePos - SourceShapePos;
-					const FVector ToTargetCenterDir = TargetCenter - SourceCenter;
-
-					const double Sign = FMath::Sign(ToTargetShapeDir.Dot(ToTargetCenterDir));
-					const double ThisDistanceSquared = ToTargetShapeDir.SquaredLength() * Sign;
-
-					if (ThisDistanceSquared < MinDistanceSquared)
-					{
-						MinDistanceSquared = ThisDistanceSquared;
-						MinDistanceVector = ToTargetShapeDir;
-					}
-				};
-
-				for (const UPCGPointData* TargetPointData : TargetPointDatas)
+					return;
+				}
+
+				const FPCGPoint& TargetPoint = *TargetPointRef.Point;
+				const FVector& TargetCenter = TargetPointRef.Bounds.Origin;
+
+				const FVector SourceShapePos = PCGDistance::CalcPosition(SourceShape, SourcePoint, TargetPoint, SourceCenter, TargetCenter);
+				const FVector TargetShapePos = PCGDistance::CalcPosition(TargetShape, TargetPoint, SourcePoint, TargetCenter, SourceCenter);
+
+				const FVector ToTargetShapeDir = TargetShapePos - SourceShapePos;
+				const FVector ToTargetCenterDir = TargetCenter - SourceCenter;
+
+				const double Sign = FMath::Sign(ToTargetShapeDir.Dot(ToTargetCenterDir));
+				const double ThisDistanceSquared = ToTargetShapeDir.SquaredLength() * Sign;
+
+				if (ThisDistanceSquared < MinDistanceSquared)
 				{
-					const UPCGPointData::PointOctree& Octree = TargetPointData->GetOctree();
+					MinDistanceSquared = ThisDistanceSquared;
+					MinDistanceVector = ToTargetShapeDir;
+				}
+			};
 
-					Octree.FindElementsWithBoundsTest(
+			for (const UPCGPointData* TargetPointData : TargetPointDatas)
+			{
+				const UPCGPointData::PointOctree& Octree = TargetPointData->GetOctree();
+
+				Octree.FindElementsWithBoundsTest(
 						FBoxCenterAndExtent(SourceQueryBounds.Origin, SourceQueryBounds.BoxExtent),
-						CalculateSDF
-					);
-				}
+						CalculateSDF);
+			}
 
-				const double Distance = FMath::Sign(MinDistanceSquared) * FMath::Sqrt(FMath::Abs(MinDistanceSquared));
+			const double Distance = FMath::Sign(MinDistanceSquared) * FMath::Sqrt(FMath::Abs(MinDistanceSquared));
 
-				if (ScalarAttribute || VectorAttribute)
+			if (bWriteToAttribute)
+			{
+				if (bOutputDistanceVector)
 				{
-					OutputData->Metadata->InitializeOnSet(OutPoint.MetadataEntry);
+					ResultCache.DistanceVectors[WriteIndex] = MinDistanceVector;
 				}
-
-				if (ScalarAttribute)
+				else
 				{
-					ScalarAttribute->SetValue(OutPoint.MetadataEntry, Distance);
+					ResultCache.Distances[WriteIndex] = Distance;
 				}
+			}
 
-				if (VectorAttribute)
+			if (bSetDensity)
+			{
+				OutPoint.Density = MaximumDistance > UE_DOUBLE_SMALL_NUMBER ? (FMath::Clamp(Distance, -MaximumDistance, MaximumDistance) * MaximumDistanceRecip) : 1.0f;
+			}
+
+			return true;
+		};
+
+		if (FPCGAsync::AsyncProcessingOneToOneEx(&Context->AsyncState, SourcePoints.Num(), /*InitializeFunc=*/[]{}, ProcessDistance, /*bEnableTimeSlicing=*/false))
+		{
+			if (Accessor.IsValid())
+			{
+				// Set all the attributes at once
+				if (bOutputDistanceVector)
 				{
-					VectorAttribute->SetValue(OutPoint.MetadataEntry, MinDistanceVector);
+					Accessor->SetRange<FVector>(ResultCache.DistanceVectors, 0, *Keys);
 				}
-				
-				if (bSetDensity)
+				else
 				{
-					// Set density instead
-					OutPoint.Density = MaximumDistance > UE_DOUBLE_SMALL_NUMBER ? (FMath::Clamp(Distance, -MaximumDistance, MaximumDistance) * MaximumDistanceRecip) : 1.0f;
+					Accessor->SetRange<double>(ResultCache.Distances, 0, *Keys);
 				}
-
-				return true;
 			}
-		);
+		}
 	}
 
 	return true;
diff --git a/Engine/Plugins/PCG/Source/PCG/Private/Tests/Elements/PCGDistanceTest.cpp b/Engine/Plugins/PCG/Source/PCG/Private/Tests/Elements/PCGDistanceTest.cpp
--- a/Engine/Plugins/PCG/Source/PCG/Private/Tests/Elements/PCGDistanceTest.cpp
+++ b/Engine/Plugins/PCG/Source/PCG/Private/Tests/Elements/PCGDistanceTest.cpp
@@ -56,11 +56,11 @@ bool FPCGDistanceTest_PointToPoint::RunTest(const FString& Parameters)
 
 	UTEST_EQUAL("Output point count", OutPoints.Num(), 2);
 
-	const FPCGMetadataAttribute<float> *DistanceAttribute = OutPointData->Metadata->GetConstTypedAttribute<float>(Settings->AttributeName);
+	const FPCGMetadataAttribute<double> *DistanceAttribute = OutPointData->Metadata->GetConstTypedAttribute<double>(Settings->OutputAttribute.GetAttributeName());
 	UTEST_NOT_NULL("Distance attribute", DistanceAttribute);
 
-	UTEST_EQUAL_TOLERANCE("Point 0 distance", DistanceAttribute->GetValue(OutPoints[0].MetadataEntry), 100.f, 0.01f);
-	UTEST_EQUAL_TOLERANCE("Point 1 distance", DistanceAttribute->GetValue(OutPoints[1].MetadataEntry), 50.f, 0.01f);
+	UTEST_EQUAL_TOLERANCE("Point 0 distance", DistanceAttribute->GetValue(OutPoints[0].MetadataEntry), 100.0, 0.01);
+	UTEST_EQUAL_TOLERANCE("Point 1 distance", DistanceAttribute->GetValue(OutPoints[1].MetadataEntry), 50.0, 0.01);
 
 	return true;
 }
@@ -75,7 +75,7 @@ bool FPCGDistanceTest_SetDensity::RunTest(const FString& Parameters)
 
 	Settings->SourceShape = PCGDistanceShape::Center;
 	Settings->TargetShape = PCGDistanceShape::Center;
-	Settings->AttributeName = NAME_None;
+	Settings->OutputAttribute = FPCGAttributePropertySelector();
 	Settings->bSetDensity = true;
 	Settings->MaximumDistance = 200.f;
 
@@ -178,11 +178,11 @@ bool FPCGDistanceTest_PointToSphere::RunTest(const FString& Parameters)
 
 	UTEST_EQUAL("Output point count", OutPoints.Num(), 2);
 
-	const FPCGMetadataAttribute<float> *DistanceAttribute = OutPointData->Metadata->GetConstTypedAttribute<float>(Settings->AttributeName);
+	const FPCGMetadataAttribute<double> *DistanceAttribute = OutPointData->Metadata->GetConstTypedAttribute<double>(Settings->OutputAttribute.GetAttributeName());
 	UTEST_NOT_NULL("Distance attribute", DistanceAttribute);
 
-	UTEST_EQUAL_TOLERANCE("Point 0 distance", DistanceAttribute->GetValue(OutPoints[0].MetadataEntry), 100.f-TargetPointRadius, 0.01f);
-	UTEST_EQUAL_TOLERANCE("Point 1 distance", DistanceAttribute->GetValue(OutPoints[1].MetadataEntry), 50.f-TargetPointRadius, 0.01f);
+	UTEST_EQUAL_TOLERANCE("Point 0 distance", DistanceAttribute->GetValue(OutPoints[0].MetadataEntry), 100.0 - TargetPointRadius, 0.01);
+	UTEST_EQUAL_TOLERANCE("Point 1 distance", DistanceAttribute->GetValue(OutPoints[1].MetadataEntry), 50.0 - TargetPointRadius, 0.01);
 
 	return true;
 }
@@ -243,11 +243,11 @@ bool FPCGDistanceTest_PointToBox::RunTest(const FString& Parameters)
 
 	UTEST_EQUAL("Output point count", OutPoints.Num(), 2);
 
-	const FPCGMetadataAttribute<float> *DistanceAttribute = OutPointData->Metadata->GetConstTypedAttribute<float>(Settings->AttributeName);
+	const FPCGMetadataAttribute<double> *DistanceAttribute = OutPointData->Metadata->GetConstTypedAttribute<double>(Settings->OutputAttribute.GetAttributeName());
 	UTEST_NOT_NULL("Distance attribute", DistanceAttribute);
 
-	UTEST_EQUAL_TOLERANCE("Point 0 distance", DistanceAttribute->GetValue(OutPoints[0].MetadataEntry), 10.f, 0.01f);
-	UTEST_EQUAL_TOLERANCE("Point 1 distance", DistanceAttribute->GetValue(OutPoints[1].MetadataEntry), 10.f, 0.01f);
+	UTEST_EQUAL_TOLERANCE("Point 0 distance", DistanceAttribute->GetValue(OutPoints[0].MetadataEntry), 10.0, 0.01);
+	UTEST_EQUAL_TOLERANCE("Point 1 distance", DistanceAttribute->GetValue(OutPoints[1].MetadataEntry), 10.0, 0.01);
 
 	return true;
 }
diff --git a/Engine/Plugins/PCG/Source/PCG/Public/Elements/PCGDistance.h b/Engine/Plugins/PCG/Source/PCG/Public/Elements/PCGDistance.h
--- a/Engine/Plugins/PCG/Source/PCG/Public/Elements/PCGDistance.h
+++ b/Engine/Plugins/PCG/Source/PCG/Public/Elements/PCGDistance.h
@@ -20,6 +20,11 @@ namespace PCGDistance
 	extern const FName TargetLabel;
 }
 
+namespace PCGDistanceConstants
+{
+	const FName DefaultOutputAttributeName = TEXT("Distance");
+}
+
 /**
  * Calculates the distance between two points (inherently a n*n operation)
  */
@@ -45,13 +50,20 @@ protected:
 	//~End UPCGSettings interface
 
 public:
+	// ~Begin UObject interface
+	virtual void PostLoad() override;
+	// ~End UObject interface
 
-	/** The name of the attribute to store on the point.Use 'None' to disable */
-	UPROPERTY(BlueprintReadWrite, EditAnywhere, Category = Settings, meta = (PCG_Overridable))
-	FName AttributeName = TEXT("Distance");
+	/** Output the distance or distance vector to an attribute. */
+	UPROPERTY(BlueprintReadWrite, EditAnywhere, Category = Settings, meta = (InlineEditConditionToggle, PCG_Overridable))
+	bool bOutputToAttribute = true;
+
+	/** The attribute output for the resulting distance value. */
+	UPROPERTY(BlueprintReadWrite, EditAnywhere, Category = Settings, meta = (EditCondition = "bOutputToAttribute", PCG_Overridable))
+	FPCGAttributePropertySelector OutputAttribute = FPCGAttributePropertySelector::CreateAttributeSelector(PCGDistanceConstants::DefaultOutputAttributeName);
 
 	/** Controls whether the attribute will be a scalar or a vector */
-	UPROPERTY(BlueprintReadWrite, EditAnywhere, Category = Settings, meta = (PCG_Overridable))
+	UPROPERTY(BlueprintReadWrite, EditAnywhere, Category = Settings, meta = (EditCondition = "bOutputToAttribute", EditConditionHides, PCG_Overridable))
 	bool bOutputDistanceVector = false;
 
 	/** If true, will also set the density to be 0 - 1 based on MaximumDistance */
@@ -69,15 +81,15 @@ public:
 	/** What shape is used on the 'target' points */
 	UPROPERTY(BlueprintReadWrite, EditAnywhere, Category = Settings, meta = (PCG_Overridable))
 	PCGDistanceShape TargetShape = PCGDistanceShape::SphereBounds;
+
+#if WITH_EDITORONLY_DATA
+	UPROPERTY(meta=(DeprecatedProperty, DeprecationMessage="Use OutputAttribute selector instead."))
+	FName AttributeName_DEPRECATED = PCGDistanceConstants::DefaultOutputAttributeName;
+#endif // WITH_EDITORONLY_DATA
 };
 
 class FPCGDistanceElement : public FPCGPointProcessingElementBase
 {
 protected:
 	virtual bool ExecuteInternal(FPCGContext* Context) const override;
-};
-
-
-#if UE_ENABLE_INCLUDE_ORDER_DEPRECATED_IN_5_2
-#include "CoreMinimal.h"
-#endif
+};
\ No newline at end of file
