[PCG][Hotfix] Fix minor bug causing UPCGDebugDrawComponent's SceneProxy to overwrite itself on multi data

#rnx

diff --git a/Engine/Plugins/PCG/Source/PCG/Private/Elements/PCGVisualizeAttribute.cpp b/Engine/Plugins/PCG/Source/PCG/Private/Elements/PCGVisualizeAttribute.cpp
--- a/Engine/Plugins/PCG/Source/PCG/Private/Elements/PCGVisualizeAttribute.cpp
+++ b/Engine/Plugins/PCG/Source/PCG/Private/Elements/PCGVisualizeAttribute.cpp
@@ -35,6 +35,8 @@ bool FPCGVisualizeAttribute::ExecuteInternal(FPCGContext* Context) const
 		return true;
 	}
 
+	TArray<UPCGDebugDrawComponent*> DebugDrawComponents;
+
 	const TArray<FPCGTaggedData> Inputs = Context->InputData.GetInputsByPin(PCGPinConstants::DefaultInputLabel);
 
 	for (const FPCGTaggedData& InputData : Inputs)
@@ -80,7 +82,21 @@ bool FPCGVisualizeAttribute::ExecuteInternal(FPCGContext* Context) const
 		}
 
 		UPCGDebugDrawComponent* DebugDrawComponent = nullptr;
-		if (AActor* TargetActor = Context->GetTargetActor(PointData))
+
+		AActor* TargetActor = Context->GetTargetActor(PointData);
+		if (TargetActor)
+		{
+			// Grab it if it exists already
+			DebugDrawComponent = TargetActor->GetComponentByClass<UPCGDebugDrawComponent>();
+		}
+		else
+		{
+			PCGLog::LogErrorOnGraph(LOCTEXT("NoTargetActor", "Internal error: Could not find target actor."), Context);
+			continue;
+		}
+
+		// If it doesn't exist already, create a new one
+		if (!DebugDrawComponent)
 		{
 			DebugDrawComponent = NewObject<UPCGDebugDrawComponent>(TargetActor, FName(LOCTEXT("PCGDebugDrawComponent", "PCG Debug Draw Component").ToString(), RF_Transient));
 
@@ -91,11 +107,6 @@ bool FPCGVisualizeAttribute::ExecuteInternal(FPCGContext* Context) const
 			ManagedComponent->GeneratedComponent = DebugDrawComponent;
 			Context->SourceComponent->AddToManagedResources(ManagedComponent);
 		}
-		else
-		{
-			PCGLog::LogErrorOnGraph(LOCTEXT("NoTargetActor", "Internal error: Could not find target actor."), Context);
-			continue;
-		}
 
 		const TArray<FPCGPoint>& Points = PointData->GetPoints();
 		TArray<FDebugRenderSceneProxy::FText3d> DebugStrings;
@@ -124,7 +135,14 @@ bool FPCGVisualizeAttribute::ExecuteInternal(FPCGContext* Context) const
 			DebugStrings.Emplace(CompoundedString.ToString(), Points[I].Transform.GetLocation() + Settings->LocalOffset, Settings->Color);
 		}
 
+		check(DebugDrawComponent);
 		DebugDrawComponent->AddDebugStrings(DebugStrings);
+		DebugDrawComponents.AddUnique(DebugDrawComponent);
+	}
+
+	for (UPCGDebugDrawComponent* DebugDrawComponent : DebugDrawComponents)
+	{
+		check(DebugDrawComponent);
 		DebugDrawComponent->StartTimer(Settings->Duration);
 	}
 #endif // !(UE_BUILD_SHIPPING || UE_BUILD_TEST) || USE_LOGGING_IN_SHIPPING
