[PCG] Massage logs a little to make job logs easier to parse

diff --git a/Engine/Plugins/PCG/Source/PCGEditor/Private/WorldPartitionBuilder/PCGWorldPartitionBuilder.cpp b/Engine/Plugins/PCG/Source/PCGEditor/Private/WorldPartitionBuilder/PCGWorldPartitionBuilder.cpp
--- a/Engine/Plugins/PCG/Source/PCGEditor/Private/WorldPartitionBuilder/PCGWorldPartitionBuilder.cpp
+++ b/Engine/Plugins/PCG/Source/PCGEditor/Private/WorldPartitionBuilder/PCGWorldPartitionBuilder.cpp
@@ -202,6 +202,8 @@ bool UPCGWorldPartitionBuilder::PostRun(UWorld* World, FPackageSourceControlHelp
 		}
 	}
 
+	UE_LOG(LogPCGWorldPartitionBuilder, Log, TEXT("PostRun: %d packages modified, %d packages deleted."), DirtyPackages.Num(), PackagesToDelete.Num());
+
 	if (!SavePackages(DirtyPackages, PackageHelper))
 	{
 		return false;
@@ -326,7 +328,7 @@ bool PCGWorldPartitionBuilder::GenerateComponents(
 		// Last minute validations, done here just prior to generation (after component has passed all previous filters) to minimize spam.
 		if (!Component->bActivated)
 		{
-			UE_LOG(LogPCGWorldPartitionBuilder, Warning, TEXT("'Activated' toggle was set false on PCG component on actor '%s' label '%s' graph '%s'. Component skipped."),
+			UE_LOG(LogPCGWorldPartitionBuilder, Log, TEXT("'Activated' toggle was set false on PCG component on actor '%s' label '%s' graph '%s'. Component skipped."),
 				*Component->GetOwner()->GetName(),
 				*Component->GetOwner()->GetActorNameOrLabel(),
 				*Component->GetGraph()->GetName());
@@ -335,7 +337,7 @@ bool PCGWorldPartitionBuilder::GenerateComponents(
 
 		if (Component->IsManagedByRuntimeGenSystem())
 		{
-			UE_LOG(LogPCGWorldPartitionBuilder, Warning, TEXT("PCG component generation trigger is set to run-time generation on actor '%s' label '%s' graph '%s'. Component skipped."),
+			UE_LOG(LogPCGWorldPartitionBuilder, Log, TEXT("PCG component generation trigger is set to run-time generation on actor '%s' label '%s' graph '%s'. Component skipped."),
 				*Component->GetOwner()->GetName(),
 				*Component->GetOwner()->GetActorNameOrLabel(),
 				*Component->GetGraph()->GetName());
