Created a native PCG data asset class and an accompanying level-to-data-asset menu action. Replaces the editor utility blueprint that we had before.

diff --git a/Engine/Plugins/PCG/Source/PCG/Private/Elements/IO/PCGLoadAssetElement.cpp b/Engine/Plugins/PCG/Source/PCG/Private/Elements/IO/PCGLoadAssetElement.cpp
new file mode 100644
--- /dev/null
+++ b/Engine/Plugins/PCG/Source/PCG/Private/Elements/IO/PCGLoadAssetElement.cpp
@@ -0,0 +1,124 @@
+// Copyright Epic Games, Inc. All Rights Reserved.
+
+#include "Elements/IO/PCGLoadAssetElement.h"
+
+#include "PCGModule.h"
+
+#include "AssetRegistry/AssetData.h"
+
+#if WITH_EDITOR
+void UPCGLoadDataAssetSettings::GetStaticTrackedKeys(FPCGSelectionKeyToSettingsMap& OutKeysToSettings, TArray<TObjectPtr<const UPCGGraph>>& OutVisitedGraphs) const
+{
+	if (Asset.IsNull())
+	{
+		return;
+	}
+
+	FPCGSelectionKey Key = FPCGSelectionKey::CreateFromPath(Asset.ToSoftObjectPath());
+
+	OutKeysToSettings.FindOrAdd(Key).Emplace(this, /*bCulling=*/false);
+}
+
+void UPCGLoadDataAssetSettings::PostEditChangeProperty(FPropertyChangedEvent& PropertyChangedEvent)
+{
+	if (PropertyChangedEvent.Property && PropertyChangedEvent.GetPropertyName() == GET_MEMBER_NAME_CHECKED(UPCGLoadDataAssetSettings, Asset))
+	{
+		UpdateFromData();
+	}
+
+	Super::PostEditChangeProperty(PropertyChangedEvent);
+}
+#endif // WITH_EDITOR
+
+FPCGElementPtr UPCGLoadDataAssetSettings::CreateElement() const
+{
+	return MakeShared<FPCGDataAssetElement>();
+}
+
+FString UPCGLoadDataAssetSettings::GetAdditionalTitleInformation() const
+{
+	return AssetName.IsEmpty() ? Asset.ToSoftObjectPath().GetAssetName() : AssetName;
+}
+
+void UPCGLoadDataAssetSettings::SetFromAsset(const FAssetData& InAsset)
+{
+	Asset = nullptr;
+
+	if (const UClass* AssetClass = InAsset.GetClass())
+	{
+		if (AssetClass->IsChildOf(UPCGDataAsset::StaticClass()))
+		{
+			Asset = TSoftObjectPtr<UPCGDataAsset>(InAsset.GetSoftObjectPath());
+		}
+	}
+
+	UpdateFromData();
+	// TODO : notify?
+}
+
+void UPCGLoadDataAssetSettings::UpdateFromData()
+{
+	// Populate pins based on data present, in order, in the data collection.
+	if (UPCGDataAsset* AssetData = Asset.LoadSynchronous())
+	{
+		TArray<FPCGPinProperties> NewPins;
+		const FPCGDataCollection& Data = AssetData->Data;
+
+		for (const FPCGTaggedData& TaggedData : Data.TaggedData)
+		{
+			if (!TaggedData.Data)
+			{
+				continue;
+			}
+
+			FPCGPinProperties* MatchingPin = NewPins.FindByPredicate([&TaggedData](const FPCGPinProperties& PinProperty) { return PinProperty.Label == TaggedData.Pin; });
+
+			if (!MatchingPin)
+			{
+				NewPins.Emplace_GetRef(TaggedData.Pin, TaggedData.Data->GetDataType());
+			}
+			else
+			{
+				MatchingPin->AllowedTypes |= TaggedData.Data->GetDataType();
+			}
+		}
+
+		Pins = NewPins;
+
+		// Update rest of cached data (name, tooltip, color, ...)
+		AssetName = AssetData->Name;
+#if WITH_EDITOR
+		AssetDescription = AssetData->Description;
+		AssetColor = AssetData->Color;
+#endif
+	}
+	else
+	{
+		Pins = TArray<FPCGPinProperties>();
+		AssetName = FString();
+#if WITH_EDITOR
+		AssetDescription = FText::GetEmpty();
+		AssetColor = FLinearColor::White;
+#endif
+	}
+}
+
+bool FPCGDataAssetElement::ExecuteInternal(FPCGContext* Context) const
+{
+	TRACE_CPUPROFILER_EVENT_SCOPE(FPCGDataAssetElement::Execute);
+
+	check(Context);
+	const UPCGLoadDataAssetSettings* Settings = Context->GetInputSettings<UPCGLoadDataAssetSettings>();
+	check(Settings);
+
+	if (UPCGDataAsset* AssetData = Settings->Asset.LoadSynchronous())
+	{
+		Context->OutputData = AssetData->Data;
+	}
+	else if (!Settings->Asset.IsNull() || Settings->bWarnIfNoAsset)
+	{
+		PCGE_LOG(Warning, GraphAndLog, FText::Format(NSLOCTEXT("PCGLoadDataAssetSettings", "UnableToLoadAsset", "Unable to load PCG asset '{0}'"), FText::FromString(Settings->Asset.ToString())));
+	}
+
+	return true;
+}
\ No newline at end of file
diff --git a/Engine/Plugins/PCG/Source/PCG/Public/Elements/IO/PCGLoadAssetElement.h b/Engine/Plugins/PCG/Source/PCG/Public/Elements/IO/PCGLoadAssetElement.h
new file mode 100644
--- /dev/null
+++ b/Engine/Plugins/PCG/Source/PCG/Public/Elements/IO/PCGLoadAssetElement.h
@@ -0,0 +1,81 @@
+// Copyright Epic Games, Inc. All Rights Reserved.
+
+#pragma once
+
+#include "PCGContext.h"
+#include "PCGDataAsset.h"
+#include "PCGSettings.h"
+
+struct FAssetData;
+
+#include "PCGLoadAssetElement.generated.h"
+
+/** Loader/Executor of PCG data assets */
+UCLASS(BlueprintType, ClassGroup = (Procedural))
+class PCG_API UPCGLoadDataAssetSettings : public UPCGSettings
+{
+	GENERATED_BODY()
+
+public:
+	//~UObject interface implementation
+#if WITH_EDITOR
+	virtual void PostEditChangeProperty(FPropertyChangedEvent& PropertyChangedEvent) override;
+#endif
+	//~End UObject interface implementation
+
+	//~UPCGSettings interface implementation
+#if WITH_EDITOR
+	virtual FName GetDefaultNodeName() const override { return FName(TEXT("PCGLoadDataAsset")); }
+	virtual FText GetDefaultNodeTitle() const override { return NSLOCTEXT("PCGLoadDataAssetSettings", "NodeTitle", "Load PCG Data Asset"); }
+	virtual FText GetNodeTooltipText() const override { return AssetDescription; }
+	virtual FLinearColor GetNodeTitleColor() const override { return AssetColor; }
+	virtual EPCGSettingsType GetType() const override { return EPCGSettingsType::Generic; } // TODO: change type?
+	virtual void GetStaticTrackedKeys(FPCGSelectionKeyToSettingsMap& OutKeysToSettings, TArray<TObjectPtr<const UPCGGraph>>& OutVisitedGraphs) const override;
+	virtual bool CanDynamicalyTrackKeys() const override { return true; }
+#endif // WITH_EDITOR
+
+	virtual bool HasDynamicPins() const override { return true; }
+	virtual bool HasFlippedTitleLines() const override { return true; }
+	virtual FString GetAdditionalTitleInformation() const override;
+
+protected:
+	virtual FPCGElementPtr CreateElement() const override;
+	// TODO tracking if data is stored in external asset
+
+	virtual TArray<FPCGPinProperties> InputPinProperties() const override { return TArray<FPCGPinProperties>(); }
+	virtual TArray<FPCGPinProperties> OutputPinProperties() const override { return Pins; }
+	// ~End UPCGSettings interface
+
+public:
+	virtual void SetFromAsset(const FAssetData& InAsset);
+	virtual void UpdateFromData();
+
+	UPROPERTY(BlueprintReadWrite, EditAnywhere, Category = Data, meta = (PCG_Overridable))
+	TSoftObjectPtr<UPCGDataAsset> Asset;
+
+	UPROPERTY(BlueprintReadOnly, VisibleAnywhere, Category = Data, meta = (NoResetToDefault))
+	TArray<FPCGPinProperties> Pins;
+
+	// Cached from the data when loaded
+	UPROPERTY(BlueprintReadOnly, VisibleAnywhere, Category = Data)
+	FString AssetName;
+
+#if WITH_EDITORONLY_DATA
+	// Cached from the data when loaded
+	UPROPERTY(BlueprintReadOnly, VisibleAnywhere, Category = Data)
+	FText AssetDescription = FText::GetEmpty();
+
+	// Cached from the data when loaded
+	UPROPERTY(BlueprintReadOnly, VisibleAnywhere, Category = Data)
+	FLinearColor AssetColor = FLinearColor::White;
+#endif // WITH_EDITORONLY_DATA
+
+	UPROPERTY(BlueprintReadWrite, EditAnywhere, Category = Settings)
+	bool bWarnIfNoAsset = true;
+};
+
+class PCG_API FPCGDataAssetElement : public IPCGElement
+{
+protected:
+	virtual bool ExecuteInternal(FPCGContext* Context) const override;
+};
diff --git a/Engine/Plugins/PCG/Source/PCG/Public/PCGDataAsset.h b/Engine/Plugins/PCG/Source/PCG/Public/PCGDataAsset.h
new file mode 100644
--- /dev/null
+++ b/Engine/Plugins/PCG/Source/PCG/Public/PCGDataAsset.h
@@ -0,0 +1,55 @@
+// Copyright Epic Games, Inc. All Rights Reserved.
+
+#pragma once
+
+#include "PCGData.h"
+#include "PCGSettings.h"
+
+#include "PCGDataAsset.generated.h"
+
+/** Container for PCG data exported as standalone objects */
+UCLASS(hidecategories=Object, BlueprintType, ClassGroup=(Procedural))
+class PCG_API UPCGDataAsset : public UObject
+{
+	GENERATED_BODY()
+
+public:
+	// TODO: crc?
+	// TODO: add compression on serialization
+	// TODO: add array of references to other assets for recursive parsing
+	// TODO: additional embedded graph?
+
+	/** Contains direct data owned by this data asset */
+	UPROPERTY(BlueprintReadOnly, VisibleAnywhere, Category = Data, meta = (NoResetToDefault))
+	FPCGDataCollection Data;
+
+	/** Alternative name (instead of asset name), can be left empty */
+	UPROPERTY(BlueprintReadWrite, EditAnywhere, Category = Data, AssetRegistrySearchable)
+	FString Name;
+
+#if WITH_EDITORONLY_DATA
+	/** Custom class to create settings/node in the graph when dragged in the editor. Can be left empty. */
+	UPROPERTY(BlueprintReadOnly, VisibleAnywhere, Category = Data, AssetRegistrySearchable, meta = (NoResetToDefault));
+	TSubclassOf<UPCGSettings> SettingsClass;
+
+	/** Reference to originating level */
+	UPROPERTY(BlueprintReadOnly, VisibleAnywhere, Category = Data, AssetRegistrySearchable, meta = (NoResetToDefault))
+	FSoftObjectPath LevelPath;
+
+	/** Custom node color to use for this specific asset */
+	UPROPERTY(BlueprintReadWrite, EditAnywhere, Category = Data)
+	FLinearColor Color = FLinearColor::White;
+
+	/** Controls whether the asset will be visible from the palette & contextual menus. */
+	UPROPERTY(BlueprintReadWrite, EditAnywhere, Category = AssetInfo, AssetRegistrySearchable)
+	bool bExposeToLibrary = true;
+
+	/** Controls in what category the asset will appear in the palette & contextual menus. */
+	UPROPERTY(BlueprintReadWrite, EditAnywhere, Category = AssetInfo, AssetRegistrySearchable)
+	FText Category;
+
+	/** Additional description (tooltip) shown on the asset/loader node. */
+	UPROPERTY(BlueprintReadWrite, EditAnywhere, Category = AssetInfo, AssetRegistrySearchable)
+	FText Description;
+#endif // WITH_EDITORONLY_DATA
+};
\ No newline at end of file
diff --git a/Engine/Plugins/PCG/Source/PCGEditor/Private/AssetDefinitions/AssetDefinition_PCGAsset.cpp b/Engine/Plugins/PCG/Source/PCGEditor/Private/AssetDefinitions/AssetDefinition_PCGAsset.cpp
new file mode 100644
--- /dev/null
+++ b/Engine/Plugins/PCG/Source/PCGEditor/Private/AssetDefinitions/AssetDefinition_PCGAsset.cpp
@@ -0,0 +1,32 @@
+// Copyright Epic Games, Inc. All Rights Reserved.
+
+#include "AssetDefinition_PCGAsset.h"
+
+#include "PCGDataAsset.h"
+
+#include "Misc/AssetCategoryPath.h"
+
+#define LOCTEXT_NAMESPACE "AssetDefinition_PCGAsset"
+
+FText UAssetDefinition_PCGAsset::GetAssetDisplayName() const
+{
+	return LOCTEXT("DisplayName", "PCG Data Asset");
+}
+
+FLinearColor UAssetDefinition_PCGAsset::GetAssetColor() const
+{
+	return FColor::Emerald;
+}
+
+TSoftClassPtr<UObject> UAssetDefinition_PCGAsset::GetAssetClass() const
+{
+	return UPCGDataAsset::StaticClass();
+}
+
+TConstArrayView<FAssetCategoryPath> UAssetDefinition_PCGAsset::GetAssetCategories() const
+{
+	static const FAssetCategoryPath Categories[] = { FAssetCategoryPath(LOCTEXT("PCGCategory", "PCG")) };
+	return Categories;
+}
+
+#undef LOCTEXT_NAMESPACE
diff --git a/Engine/Plugins/PCG/Source/PCGEditor/Private/AssetDefinitions/AssetDefinition_PCGAsset.h b/Engine/Plugins/PCG/Source/PCGEditor/Private/AssetDefinitions/AssetDefinition_PCGAsset.h
new file mode 100644
--- /dev/null
+++ b/Engine/Plugins/PCG/Source/PCGEditor/Private/AssetDefinitions/AssetDefinition_PCGAsset.h
@@ -0,0 +1,20 @@
+// Copyright Epic Games, Inc. All Rights Reserved.
+
+#pragma once
+
+#include "AssetDefinitionDefault.h"
+
+#include "AssetDefinition_PCGAsset.generated.h"
+
+/** Asset definition class for PCG Assets. Used to control asset name, color & categories in the content browser. */
+UCLASS()
+class UAssetDefinition_PCGAsset : public UAssetDefinitionDefault
+{
+	GENERATED_BODY()
+
+public:
+	virtual FText GetAssetDisplayName() const override;
+	virtual FLinearColor GetAssetColor() const override;
+	virtual TSoftClassPtr<UObject> GetAssetClass() const override;
+	virtual TConstArrayView<FAssetCategoryPath> GetAssetCategories() const override;
+};
diff --git a/Engine/Plugins/PCG/Source/PCGEditor/Private/PCGEditorCommon.h b/Engine/Plugins/PCG/Source/PCGEditor/Private/PCGEditorCommon.h
--- a/Engine/Plugins/PCG/Source/PCGEditor/Private/PCGEditorCommon.h
+++ b/Engine/Plugins/PCG/Source/PCGEditor/Private/PCGEditorCommon.h
@@ -11,8 +11,9 @@ enum class EPCGElementType : uint8
 	Blueprint = 1 << 1,
 	Subgraph = 1 << 2,
 	Settings = 1 << 3,
-	Other = 1 << 4,
-	All = (1 << 5) - 1
+	Asset = 1 << 4,
+	Other = 1 << 5,
+	All = (1 << 6) - 1
 };
 ENUM_CLASS_FLAGS(EPCGElementType);
 
diff --git a/Engine/Plugins/PCG/Source/PCGEditor/Private/PCGEditorGraphSchema.cpp b/Engine/Plugins/PCG/Source/PCGEditor/Private/PCGEditorGraphSchema.cpp
--- a/Engine/Plugins/PCG/Source/PCGEditor/Private/PCGEditorGraphSchema.cpp
+++ b/Engine/Plugins/PCG/Source/PCGEditor/Private/PCGEditorGraphSchema.cpp
@@ -3,6 +3,7 @@
 #include "PCGEditorGraphSchema.h"
 
 #include "PCGComponent.h"
+#include "PCGDataAsset.h"
 #include "PCGEdge.h"
 #include "PCGGraph.h"
 #include "PCGPin.h"
@@ -52,6 +53,10 @@ void UPCGEditorGraphSchema::GetPaletteActions(FGraphActionMenuBuilder& ActionMen
 	{
 		GetSettingsElementActions(ActionMenuBuilder, /*bIsContextual=*/false);
 	}
+	if (!!(InPCGElementTypeFilter & EPCGElementType::Asset))
+	{
+		GetDataAssetActions(ActionMenuBuilder);
+	}
 	if (!!(InPCGElementTypeFilter & EPCGElementType::Other))
 	{
 		GetNamedRerouteUsageActions(ActionMenuBuilder);
@@ -67,6 +72,7 @@ void UPCGEditorGraphSchema::GetGraphContextActions(FGraphContextMenuBuilder& Con
 	GetSubgraphElementActions(ContextMenuBuilder);
 	GetBlueprintElementActions(ContextMenuBuilder);
 	GetSettingsElementActions(ContextMenuBuilder, /*bIsContextual=*/true);
+	GetDataAssetActions(ContextMenuBuilder);
 	GetNamedRerouteUsageActions(ContextMenuBuilder, ContextMenuBuilder.CurrentGraph);
 	GetExtraElementActions(ContextMenuBuilder);
 }
@@ -591,6 +597,32 @@ void UPCGEditorGraphSchema::GetNamedRerouteUsageActions(FGraphActionMenuBuilder&
 	}
 }
 
+void UPCGEditorGraphSchema::GetDataAssetActions(FGraphActionMenuBuilder& ActionMenuBuilder) const
+{
+	PCGEditorUtils::ForEachPCGAssetData([&ActionMenuBuilder](const FAssetData& AssetData)
+	{
+		const bool bExposeToLibrary = AssetData.GetTagValueRef<bool>(GET_MEMBER_NAME_CHECKED(UPCGDataAsset, bExposeToLibrary));
+		if (bExposeToLibrary)
+		{
+			const FText MenuDesc = FText::FromString(FName::NameToDisplayString(AssetData.AssetName.ToString(), false));
+			const FText Category = AssetData.GetTagValueRef<FText>(GET_MEMBER_NAME_CHECKED(UPCGDataAsset, Category));
+			const FText Description = AssetData.GetTagValueRef<FText>(GET_MEMBER_NAME_CHECKED(UPCGDataAsset, Description));
+			const FString AssetName = AssetData.GetTagValueRef<FString>(GET_MEMBER_NAME_CHECKED(UPCGDataAsset, Name));
+
+			const FSoftClassPath SettingsClassPath = FSoftClassPath(AssetData.GetTagValueRef<FString>(GET_MEMBER_NAME_CHECKED(UPCGDataAsset, SettingsClass)));
+			TSubclassOf<UPCGSettings> SettingsClass = SettingsClassPath.TryLoadClass<UPCGSettings>();
+
+			TSharedPtr<FPCGEditorGraphSchemaAction_NewLoadAssetElement> NewLoadDataAssetAction(new FPCGEditorGraphSchemaAction_NewLoadAssetElement(Category, AssetName.IsEmpty() ? MenuDesc : FText::FromString(AssetName), Description, 0));
+			NewLoadDataAssetAction->Asset = AssetData;
+			NewLoadDataAssetAction->SettingsClass = SettingsClass;
+
+			ActionMenuBuilder.AddAction(NewLoadDataAssetAction);
+		}
+
+		return true;
+	});
+}
+
 void UPCGEditorGraphSchema::DroppedAssetsOnGraph(const TArray<FAssetData>& Assets, const FVector2D& GraphPosition, UEdGraph* Graph) const
 {
 	FVector2D GraphPositionOffset = GraphPosition;
@@ -602,9 +634,9 @@ void UPCGEditorGraphSchema::DroppedAssetsOnGraph(const TArray<FAssetData>& Asset
 
 	for (const FAssetData& AssetData : Assets)
 	{
-		if (const UObject* Asset = AssetData.GetAsset())
+		if(const UClass* AssetClass = AssetData.GetClass())
 		{
-			if (Asset->IsA<UPCGGraphInterface>())
+			if(AssetClass->IsChildOf(UPCGGraphInterface::StaticClass()))
 			{
 				FPCGEditorGraphSchemaAction_NewSubgraphElement NewSubgraphAction;
 				NewSubgraphAction.SubgraphObjectPath = AssetData.GetSoftObjectPath();
@@ -620,7 +652,15 @@ void UPCGEditorGraphSchema::DroppedAssetsOnGraph(const TArray<FAssetData>& Asset
 				NewBlueprintAction.PerformAction(Graph, NullFromPin, GraphPositionOffset);
 				GraphPositionOffset.Y += PositionOffsetIncrementY;
 			}
-			else if (Asset->IsA<UPCGSettings>())
+			else if(AssetClass->IsChildOf(UPCGDataAsset::StaticClass()))
+			{
+				FPCGEditorGraphSchemaAction_NewLoadAssetElement NewLoadAssetAction;
+				NewLoadAssetAction.Asset = AssetData;
+				NewLoadAssetAction.SettingsClass = FSoftClassPath(AssetData.GetTagValueRef<FString>(GET_MEMBER_NAME_CHECKED(UPCGDataAsset, SettingsClass))).TryLoadClass<UPCGSettings>();
+				NewLoadAssetAction.PerformAction(Graph, NullFromPin, GraphPositionOffset);
+				GraphPositionOffset.Y += PositionOffsetIncrementY;
+			}
+			else if(AssetClass->IsChildOf(UPCGSettings::StaticClass()))
 			{
 				// Delay creation so we can open a menu, once, if needed.
 				SettingsPaths.Add(AssetData.GetSoftObjectPath());
@@ -647,14 +687,15 @@ void UPCGEditorGraphSchema::GetAssetsGraphHoverMessage(const TArray<FAssetData>&
 {
 	for (const FAssetData& AssetData : Assets)
 	{
-		if (const UObject* Asset = AssetData.GetAsset())
+		// TODO: get class from asset data, shouldn't require loading
+		if(const UClass* AssetClass = AssetData.GetClass())
 		{
-			if (Asset->IsA<UPCGGraphInterface>() || Asset->IsA<UPCGSettings>() || PCGEditorUtils::IsAssetPCGBlueprint(AssetData))
+			if(AssetClass->IsChildOf(UPCGGraphInterface::StaticClass()) || AssetClass->IsChildOf(UPCGSettings::StaticClass()) || AssetClass->IsChildOf(UPCGDataAsset::StaticClass()) || PCGEditorUtils::IsAssetPCGBlueprint(AssetData))
 			{
 				OutOkIcon = true;
 				return;
 			}
-			else if (Asset->IsA<UBlueprint>())
+			else if(AssetClass->IsChildOf(UBlueprint::StaticClass()))
 			{
 				OutTooltipText = LOCTEXT("PCGEditorDropAssetInvalidBP", "Blueprint does not derive from UPCGBlueprintElement").ToString();
 				OutOkIcon = false;
diff --git a/Engine/Plugins/PCG/Source/PCGEditor/Private/PCGEditorGraphSchema.h b/Engine/Plugins/PCG/Source/PCGEditor/Private/PCGEditorGraphSchema.h
--- a/Engine/Plugins/PCG/Source/PCGEditor/Private/PCGEditorGraphSchema.h
+++ b/Engine/Plugins/PCG/Source/PCGEditor/Private/PCGEditorGraphSchema.h
@@ -39,6 +39,7 @@ private:
 	void GetSettingsElementActions(FGraphActionMenuBuilder& ActionMenuBuilder, bool bIsContextual) const;
 	void GetExtraElementActions(FGraphActionMenuBuilder& ActionMenuBuilder) const;
 	void GetNamedRerouteUsageActions(FGraphActionMenuBuilder& ActionMenuBuilder, const UEdGraph* CurrentGraph = nullptr) const;
+	void GetDataAssetActions(FGraphActionMenuBuilder& ActionMenuBuilder) const;
 
 	virtual bool TryCreateConnectionInternal(UEdGraphPin* A, UEdGraphPin* B, bool bAddConversionNodeIfNeeded) const;
 };
diff --git a/Engine/Plugins/PCG/Source/PCGEditor/Private/PCGEditorGraphSchemaActions.cpp b/Engine/Plugins/PCG/Source/PCGEditor/Private/PCGEditorGraphSchemaActions.cpp
--- a/Engine/Plugins/PCG/Source/PCGEditor/Private/PCGEditorGraphSchemaActions.cpp
+++ b/Engine/Plugins/PCG/Source/PCGEditor/Private/PCGEditorGraphSchemaActions.cpp
@@ -7,9 +7,12 @@
 #include "PCGEditorGraphNode.h"
 #include "PCGEditorGraphNodeReroute.h"
 #include "PCGEditorModule.h"
+
+#include "PCGDataAsset.h"
 #include "PCGGraph.h"
 #include "PCGNode.h"
 #include "PCGSubgraph.h"
+#include "Elements/IO/PCGLoadAssetElement.h"
 #include "Elements/PCGExecuteBlueprint.h"
 #include "Elements/PCGReroute.h"
 #include "Elements/PCGUserParameterGet.h"
@@ -292,6 +295,55 @@ UEdGraphNode* FPCGEditorGraphSchemaAction_NewSettingsElement::MakeSettingsNode(U
 	return NewNode;
 }
 
+UEdGraphNode* FPCGEditorGraphSchemaAction_NewLoadAssetElement::PerformAction(UEdGraph* ParentGraph, UEdGraphPin* FromPin, const FVector2D Location, bool bSelectNewNode)
+{
+	UPCGEditorGraph* EditorGraph = Cast<UPCGEditorGraph>(ParentGraph);
+	if (!EditorGraph)
+	{
+		UE_LOG(LogPCGEditor, Error, TEXT("Invalid EditorGraph"));
+		return nullptr;
+	}
+
+	UPCGGraph* PCGGraph = EditorGraph->GetPCGGraph();
+	if (!PCGGraph)
+	{
+		UE_LOG(LogPCGEditor, Error, TEXT("Invalid PCGGraph"));
+		return nullptr;
+	}
+
+	if(!Asset.GetSoftObjectPath().IsValid())
+	{
+		UE_LOG(LogPCGEditor, Error, TEXT("Invalid asset path."));
+		return nullptr;
+	}
+
+	// Important - do not reconstruct the editor graph node/pins midway through this function as this will invalidate FromPin.
+	const FPCGDeferNodeReconstructScope DisableReconstruct(FromPin);
+
+	const FScopedTransaction Transaction(*FPCGEditorCommon::ContextIdentifier, LOCTEXT("PCGEditorNewBlueprintELement", "PCG Editor: New Blueprint Element"), nullptr);
+	EditorGraph->Modify();
+
+	UPCGSettings* DefaultNodeSettings = nullptr;
+	UPCGNode* NewPCGNode = PCGGraph->AddNodeOfType(SettingsClass ? SettingsClass.Get() : UPCGLoadDataAssetSettings::StaticClass(), DefaultNodeSettings);
+	UPCGLoadDataAssetSettings* DefaultLoadAssetSettings = CastChecked<UPCGLoadDataAssetSettings>(DefaultNodeSettings);
+
+	DefaultLoadAssetSettings->SetFromAsset(Asset);
+
+	NewPCGNode->UpdateAfterSettingsChangeDuringCreation();
+
+	FGraphNodeCreator<UPCGEditorGraphNode> NodeCreator(*EditorGraph);
+	UPCGEditorGraphNode* NewNode = NodeCreator.CreateUserInvokedNode(bSelectNewNode);
+	NewNode->Construct(NewPCGNode);
+	NewNode->NodePosX = Location.X;
+	NewNode->NodePosY = Location.Y;
+	NodeCreator.Finalize();
+
+	NewPCGNode->PositionX = Location.X;
+	NewPCGNode->PositionY = Location.Y;
+
+	return NewNode;
+}
+
 UEdGraphNode* FPCGEditorGraphSchemaAction_NewBlueprintElement::PerformAction(UEdGraph* ParentGraph, UEdGraphPin* FromPin, const FVector2D Location, bool bSelectNewNode)
 {
 	UPCGEditorGraph* EditorGraph = Cast<UPCGEditorGraph>(ParentGraph);
diff --git a/Engine/Plugins/PCG/Source/PCGEditor/Private/PCGEditorGraphSchemaActions.h b/Engine/Plugins/PCG/Source/PCGEditor/Private/PCGEditorGraphSchemaActions.h
--- a/Engine/Plugins/PCG/Source/PCGEditor/Private/PCGEditorGraphSchemaActions.h
+++ b/Engine/Plugins/PCG/Source/PCGEditor/Private/PCGEditorGraphSchemaActions.h
@@ -129,6 +129,33 @@ struct FPCGEditorGraphSchemaAction_NewBlueprintElement : public FEdGraphSchemaAc
 	//~ End FEdGraphSchemaAction Interface
 };
 
+USTRUCT()
+struct FPCGEditorGraphSchemaAction_NewLoadAssetElement : public FEdGraphSchemaAction
+{
+	GENERATED_BODY()
+
+	// Inherit the base class's constructors
+	using FEdGraphSchemaAction::FEdGraphSchemaAction;
+
+	// Simple type info
+	static FName StaticGetTypeId()
+	{
+		static FName Type("FPCGEditorGraphSchemaAction_NewLoadAssetElement");
+		return Type;
+	}
+
+	UPROPERTY()
+	FAssetData Asset;
+
+	UPROPERTY()
+	TSubclassOf<UPCGSettings> SettingsClass;
+
+	//~ Begin FEdGraphSchemaAction Interface
+	virtual FName GetTypeId() const override { return StaticGetTypeId(); }
+	virtual UEdGraphNode* PerformAction(UEdGraph* ParentGraph, UEdGraphPin* FromPin, const FVector2D Location, bool bSelectNewNode = true) override;
+	//~ End FEdGraphSchemaAction Interface
+};
+
 USTRUCT()
 struct FPCGEditorGraphSchemaAction_NewSubgraphElement : public FEdGraphSchemaAction
 {
diff --git a/Engine/Plugins/PCG/Source/PCGEditor/Private/PCGEditorMenuUtils.cpp b/Engine/Plugins/PCG/Source/PCGEditor/Private/PCGEditorMenuUtils.cpp
new file mode 100644
--- /dev/null
+++ b/Engine/Plugins/PCG/Source/PCGEditor/Private/PCGEditorMenuUtils.cpp
@@ -0,0 +1,61 @@
+// Copyright Epic Games, Inc. All Rights Reserved.
+
+#include "PCGEditorMenuUtils.h"
+
+#include "PCGLevelToAsset.h"
+
+#include "Engine/World.h"
+#include "Misc/ScopedSlowTask.h"
+
+#define LOCTEXT_NAMESPACE "PCGEditorMenuUtils"
+
+namespace PCGEditorMenuUtils
+{
+	FToolMenuSection& CreatePCGSection(UToolMenu* Menu)
+	{
+		const FName LevelSectionName = TEXT("PCG");
+		FToolMenuSection* SectionPtr = Menu->FindSection(LevelSectionName);
+		if (!SectionPtr)
+		{
+			SectionPtr = &(Menu->AddSection(LevelSectionName, LOCTEXT("PCGSectionLabel", "Procedural (PCG)")));
+		}
+
+		FToolMenuSection& Section = *SectionPtr;
+		return Section;
+	}
+
+	void CreateOrUpdatePCGAssetFromMenu(UToolMenu* Menu, TArray<FAssetData>& InWorldAssets)
+	{
+		TArray<FAssetData> TempWorldAssets;
+		for (const FAssetData& Asset : InWorldAssets)
+		{
+			if (Asset.IsInstanceOf<UWorld>())
+			{
+				TempWorldAssets.Add(Asset);
+			}
+		}
+
+		if (TempWorldAssets.IsEmpty())
+		{
+			return;
+		}
+
+		FToolMenuSection& Section = CreatePCGSection(Menu);
+
+		FToolUIAction UIAction;
+		UIAction.ExecuteAction.BindLambda([WorldAssets = MoveTemp(TempWorldAssets)](const FToolMenuContext& MenuContext)
+		{
+			FScopedSlowTask SlowTask(0.0f, LOCTEXT("CreateOrUpdatePCGAssetsInProgress", "Creating PCG Assets..."));
+			UPCGLevelToAsset::CreateOrUpdatePCGAssets(WorldAssets);
+		});
+
+		Section.AddMenuEntry(
+			"CreateOrUpdatePCGAssetFromMenu",
+			LOCTEXT("CreateOrUpdatePCGAssetFromMenu", "Create or Update PCG Asset from Level(s)"),
+			TAttribute<FText>(),
+			FSlateIcon(),
+			UIAction);
+	}
+}
+
+#undef LOCTEXT_NAMESPACE
\ No newline at end of file
diff --git a/Engine/Plugins/PCG/Source/PCGEditor/Private/PCGEditorModule.cpp b/Engine/Plugins/PCG/Source/PCGEditor/Private/PCGEditorModule.cpp
--- a/Engine/Plugins/PCG/Source/PCGEditor/Private/PCGEditorModule.cpp
+++ b/Engine/Plugins/PCG/Source/PCGEditor/Private/PCGEditorModule.cpp
@@ -4,18 +4,23 @@
 
 #include "PCGEditorCommands.h"
 #include "PCGEditorGraphNodeFactory.h"
+#include "PCGEditorMenuUtils.h"
 #include "PCGEditorSettings.h"
 #include "PCGEditorStyle.h"
 #include "PCGEditorUtils.h"
 #include "PCGSubsystem.h"
 #include "PCGVolumeFactory.h"
 
+#include "ContentBrowserMenuContexts.h"
+#include "ContentBrowserModule.h"
 #include "Editor.h"
 #include "EditorModeManager.h"
 #include "EditorModes.h"
+#include "IContentBrowserSingleton.h"
 #include "ISettingsModule.h"
 #include "PropertyEditorModule.h"
 #include "ToolMenus.h"
+
 #include "Details/EnumSelectorDetails.h"
 #include "Details/PCGAttributePropertySelectorDetails.h"
 #include "Details/PCGBlueprintSettingsDetails.h"
@@ -30,9 +35,10 @@
 void FPCGEditorModule::StartupModule()
 {
 	RegisterDetailsCustomizations();
-	RegisterMenuExtensions();
 	RegisterSettings();
 
+	UToolMenus::RegisterStartupCallback(FSimpleMulticastDelegate::FDelegate::CreateRaw(this, &FPCGEditorModule::RegisterMenuExtensions));
+
 	FPCGEditorCommands::Register();
 	FPCGEditorStyle::Register();
 
@@ -143,14 +149,35 @@ void FPCGEditorModule::UnregisterDetailsCustomizations()
 void FPCGEditorModule::RegisterMenuExtensions()
 {
 	FToolMenuOwnerScoped OwnerScoped(this);
-	UToolMenu* Menu = UToolMenus::Get()->ExtendMenu("LevelEditor.MainMenu.Tools");
-	FToolMenuSection& Section = Menu->AddSection("PCGToolsSection", LOCTEXT("PCGToolsSection", "Procedural Generation Tools"));
-
-	Section.AddSubMenu(
-		"PCGToolsSubMenu",
-		LOCTEXT("PCGSubMenu", "PCG Framework"),
-		LOCTEXT("PCGSubMenu_Tooltip", "PCG Framework related functionality"),
-		FNewMenuDelegate::CreateRaw(this, &FPCGEditorModule::PopulateMenuActions));
+	
+	if (UToolMenu* Menu = UToolMenus::Get()->ExtendMenu("LevelEditor.MainMenu.Tools"))
+	{
+		FToolMenuSection& Section = Menu->AddSection("PCGToolsSection", LOCTEXT("PCGToolsSection", "Procedural Generation Tools"));
+
+		Section.AddSubMenu(
+			"PCGToolsSubMenu",
+			LOCTEXT("PCGSubMenu", "PCG Framework"),
+			LOCTEXT("PCGSubMenu_Tooltip", "Procedural Content Generation (PCG) Framework related functionality"),
+			FNewMenuDelegate::CreateRaw(this, &FPCGEditorModule::PopulateMenuActions));
+	}
+
+	if (UToolMenu* WorldAssetMenu = UToolMenus::Get()->ExtendMenu("ContentBrowser.AssetContextMenu.AssetActionsSubMenu"))
+	{
+		// Use a dynamic section here because we might have plugins registering at a later time
+		FToolMenuSection& Section = WorldAssetMenu->AddDynamicSection("PCG", FNewToolMenuDelegate::CreateLambda([this](UToolMenu* ToolMenu)
+		{
+			if (!GEditor || GEditor->GetPIEWorldContext() || !ToolMenu)
+			{
+				return;
+			}
+
+			if (UContentBrowserAssetContextMenuContext* AssetMenuContext = ToolMenu->Context.FindContext<UContentBrowserAssetContextMenuContext>())
+			{
+				PCGEditorMenuUtils::CreateOrUpdatePCGAssetFromMenu(ToolMenu, AssetMenuContext->SelectedAssets);
+			}
+
+		}), FToolMenuInsert(NAME_None, EToolMenuInsertType::Default));
+	}
 }
 
 void FPCGEditorModule::UnregisterMenuExtensions()
diff --git a/Engine/Plugins/PCG/Source/PCGEditor/Private/PCGEditorUtils.cpp b/Engine/Plugins/PCG/Source/PCGEditor/Private/PCGEditorUtils.cpp
--- a/Engine/Plugins/PCG/Source/PCGEditor/Private/PCGEditorUtils.cpp
+++ b/Engine/Plugins/PCG/Source/PCGEditor/Private/PCGEditorUtils.cpp
@@ -2,13 +2,14 @@
 
 #include "PCGEditorUtils.h"
 
-#include "Blueprint/BlueprintSupport.h"
+#include "PCGDataAsset.h"
 #include "PCGGraph.h"
 #include "Elements/PCGExecuteBlueprint.h"
 
 #include "AssetToolsModule.h"
 #include "AssetRegistry/ARFilter.h"
 #include "AssetRegistry/AssetRegistryModule.h"
+#include "Blueprint/BlueprintSupport.h"
 #include "IAssetTools.h"
 #include "Kismet2/BlueprintEditorUtils.h"
 #include "Kismet2/KismetEditorUtilities.h"
@@ -89,6 +90,15 @@ void PCGEditorUtils::ForEachPCGGraphAssetData(TFunctionRef<bool(const FAssetData
 	ForEachAssetData(Filter, InFunc);
 }
 
+void PCGEditorUtils::ForEachPCGAssetData(TFunctionRef<bool(const FAssetData&)> InFunc)
+{
+	FARFilter Filter;
+	Filter.ClassPaths.Add(UPCGDataAsset::StaticClass()->GetClassPathName());
+	Filter.bRecursiveClasses = true;
+
+	ForEachAssetData(Filter, InFunc);
+}
+
 void PCGEditorUtils::ForcePCGBlueprintVariableVisibility()
 {
 	ForEachPCGBlueprintAssetData([](const FAssetData& AssetData)
diff --git a/Engine/Plugins/PCG/Source/PCGEditor/Private/PCGEditorUtils.h b/Engine/Plugins/PCG/Source/PCGEditor/Private/PCGEditorUtils.h
--- a/Engine/Plugins/PCG/Source/PCGEditor/Private/PCGEditorUtils.h
+++ b/Engine/Plugins/PCG/Source/PCGEditor/Private/PCGEditorUtils.h
@@ -27,6 +27,7 @@ namespace PCGEditorUtils
 	void ForEachPCGBlueprintAssetData(TFunctionRef<bool(const FAssetData&)> InFunc);
 	void ForEachPCGSettingsAssetData(TFunctionRef<bool(const FAssetData&)> InFunc);
 	void ForEachPCGGraphAssetData(TFunctionRef<bool(const FAssetData&)> InFunc);
+	void ForEachPCGAssetData(TFunctionRef<bool(const FAssetData&)> InFunc);
 
 	/** Asset deprecation methods */
 	void ForcePCGBlueprintVariableVisibility();
diff --git a/Engine/Plugins/PCG/Source/PCGEditor/Private/PCGLevelToAsset.cpp b/Engine/Plugins/PCG/Source/PCGEditor/Private/PCGLevelToAsset.cpp
new file mode 100644
--- /dev/null
+++ b/Engine/Plugins/PCG/Source/PCGEditor/Private/PCGLevelToAsset.cpp
@@ -0,0 +1,318 @@
+// Copyright Epic Games, Inc. All Rights Reserved.
+
+#include "PCGLevelToAsset.h"
+
+#include "PCGEditorModule.h"
+
+#include "Data/PCGPointData.h"
+#include "Helpers/PCGActorHelpers.h"
+#include "Helpers/PCGHelpers.h"
+#include "Metadata/PCGMetadata.h"
+
+#include "FileHelpers.h"
+#include "AssetRegistry/AssetRegistryModule.h"
+#include "Components/InstancedStaticMeshComponent.h"
+#include "UObject/UObjectGlobals.h"
+#include "UObject/Package.h"
+
+void UPCGLevelToAsset::CreateOrUpdatePCGAssets(const TArray<FAssetData>& WorldAssets, TSubclassOf<UPCGLevelToAsset> ExporterSubclass)
+{
+	TArray<UPackage*> PackagesToSave;
+
+	for (const FAssetData& WorldAsset : WorldAssets)
+	{
+		if (UPackage* Package = CreateOrUpdatePCGAsset(TSoftObjectPtr<UWorld>(WorldAsset.GetSoftObjectPath()), ExporterSubclass))
+		{
+			PackagesToSave.Add(Package);
+			// TODO: consider if we should garbage collect?
+		}
+	}
+
+	// Save the file(s)
+	// TODO: check if we should just dirty and not save (could be an option here)
+	if (!PackagesToSave.IsEmpty())
+	{
+		FEditorFileUtils::PromptForCheckoutAndSave(PackagesToSave, false, false);
+	}
+}
+
+UPackage* UPCGLevelToAsset::CreateOrUpdatePCGAsset(TSoftObjectPtr<UWorld> WorldPath, TSubclassOf<UPCGLevelToAsset> ExporterSubclass)
+{
+	return CreateOrUpdatePCGAsset(WorldPath.LoadSynchronous(), ExporterSubclass);
+}
+
+UPackage* UPCGLevelToAsset::CreateOrUpdatePCGAsset(UWorld* Level, TSubclassOf<UPCGLevelToAsset> ExporterSubclass)
+{
+	if (!Level)
+	{
+		return nullptr;
+	}
+
+	UPCGLevelToAsset* Exporter = nullptr;
+	if (ExporterSubclass)
+	{
+		Exporter = NewObject<UPCGLevelToAsset>(GetTransientPackage(), ExporterSubclass);
+	}
+	else
+	{
+		Exporter = NewObject<UPCGLevelToAsset>(GetTransientPackage());
+	}
+
+	if (!Exporter)
+	{
+		UE_LOG(LogPCGEditor, Error, TEXT("Unable to create Level to Settings exporter."));
+		return nullptr;
+	}
+
+	const FString AssetName = Level->GetName() + TEXT("_PCG");
+	// TODO: since we store the level SOP in the PCG data assets, we can technically search if there is already an asset matching this level and select this one instead.
+	const FString PackageName = FPaths::Combine(FPackageName::GetLongPackagePath(Level->GetPackage()->GetName()), AssetName);
+
+	// Implementation note: this will cause a warning if the file can't be loaded (if the file doesn't exist, for example). It doesn't seem possible to quiet this.
+	UPackage* Package = LoadPackage(nullptr, *PackageName, LOAD_None);
+
+	UPCGDataAsset* Asset = nullptr;
+	bool NewAssetCreated = false;
+
+	if (Package)
+	{
+		UObject* Object = FindObjectFast<UObject>(Package, *AssetName);
+		if (Object && Object->GetClass() != Exporter->GetAssetType())
+		{
+			Object->SetFlags(RF_Transient);
+			Object->Rename(nullptr, GetTransientPackage(), REN_DontCreateRedirectors | REN_NonTransactional);
+			NewAssetCreated = true;
+		}
+		else
+		{
+			Asset = Cast<UPCGDataAsset>(Object);
+		}
+	}
+	else
+	{
+		Package = CreatePackage(*PackageName);
+		NewAssetCreated = true;
+	}
+
+	if (!Asset)
+	{
+		const EObjectFlags Flags = RF_Public | RF_Standalone | RF_Transactional;
+		Asset = NewObject<UPCGDataAsset>(Package, Exporter->GetAssetType(), FName(*AssetName), Flags);
+	}
+
+	if (Asset)
+	{
+		if (NewAssetCreated)
+		{
+			// Notify the asset registry
+			FAssetRegistryModule::AssetCreated(Asset);
+		}
+
+		Exporter->ExportLevel(Level, PackageName, Asset);
+
+		// Mark the package dirty...
+		Package->MarkPackageDirty();
+
+		// Make sure everybody knows we changed these settings.
+		if (!NewAssetCreated)
+		{
+			FCoreUObjectDelegates::BroadcastOnObjectModified(Asset);
+		}
+	}
+
+	return Package;
+}
+
+TSubclassOf<UPCGDataAsset> UPCGLevelToAsset::BP_GetAssetType_Implementation() const
+{
+	return UPCGDataAsset::StaticClass();
+}
+
+TSubclassOf<UPCGDataAsset> UPCGLevelToAsset::GetAssetType() const
+{
+	return BP_GetAssetType();
+}
+
+bool UPCGLevelToAsset::ExportLevel(UWorld* Level, const FString& PackageName, UPCGDataAsset* Asset)
+{
+	return BP_ExportLevel(Level, PackageName, Asset);
+}
+
+bool UPCGLevelToAsset::BP_ExportLevel_Implementation(UWorld* Level, const FString& PackageName, UPCGDataAsset* Asset)
+{
+	check(Level && Asset);
+	Asset->LevelPath = FSoftObjectPath(Level);
+	Asset->Description = FText::Format(NSLOCTEXT("PCGLevelToAsset", "DefaultDescriptionOnExportedLevel", "Generated from level: {0}"), FText::FromString(Level->GetName()));
+
+	FPCGDataCollection& DataCollection = Asset->Data;
+
+	// Create Root Data
+	UPCGPointData* RootPointData = NewObject<UPCGPointData>(Asset);
+	UPCGMetadata* RootMetadata = RootPointData->MutableMetadata();
+	TArray<FPCGPoint>& Roots = RootPointData->GetMutablePoints();
+	RootMetadata->CreateAttribute<FString>(TEXT("Name"), Level->GetName(), /*bAllowsInterpolation=*/false, /*bOverrideParent=*/true);
+	RootMetadata->CreateAttribute<FSoftObjectPath>(TEXT("Source"), PackageName, /*bAllowsInterpolation=*/false, /*bOverrideParent=*/true);
+
+	// Add to data collection
+	{
+		FPCGTaggedData& RootsTaggedData = DataCollection.TaggedData.Emplace_GetRef();
+		RootsTaggedData.Data = RootPointData;
+		RootsTaggedData.Pin = TEXT("Root");
+	}
+
+	// Create PCGData
+	UPCGPointData* PointData = NewObject<UPCGPointData>(Asset);
+	UPCGMetadata* PointMetadata = PointData->MutableMetadata();
+	TArray<FPCGPoint>& Points = PointData->GetMutablePoints();
+
+	// Add to data collection
+	{
+		FPCGTaggedData& PointsTaggedData = DataCollection.TaggedData.Emplace_GetRef();
+		PointsTaggedData.Data = PointData;
+		PointsTaggedData.Pin = TEXT("Points");
+	}
+
+	// Common data shared across steps
+	FBox AllActorBounds(EForceInit::ForceInit);
+
+	// Attribute setup on the points
+	FPCGMetadataAttribute<FSoftObjectPath>* MaterialAttribute = PointMetadata->CreateAttribute<FSoftObjectPath>(TEXT("Material"), FSoftObjectPath(), /*bAllowsInterpolation=*/false, /*bOverrideParent=*/true);
+	FPCGMetadataAttribute<FSoftObjectPath>* MeshAttribute = PointMetadata->CreateAttribute<FSoftObjectPath>(TEXT("Mesh"), FSoftObjectPath(), /*bAllowsInterpolation=*/false, /*bOverrideParent=*/true);
+	FPCGMetadataAttribute<int64>* HierarchyDepthAttribute = PointMetadata->CreateAttribute<int64>(TEXT("HierarchyDepth"), 0, /*bAllowsInterpolation=*/false, /*bOverrideParent=*/true);
+	FPCGMetadataAttribute<int64>* ActorIndexAttribute = PointMetadata->CreateAttribute<int64>(TEXT("ActorIndex"), -1, /*bAllowsInterpolation=*/false, /*bOverrideParent=*/true);
+	FPCGMetadataAttribute<int64>* ParentIndexAttribute = PointMetadata->CreateAttribute<int64>(TEXT("ParentIndex"), -1, /*bAllowsInterpolation=*/false, /*bOverrideParent=*/true);
+	FPCGMetadataAttribute<FTransform>* RelativeTransformAttribute = PointMetadata->CreateAttribute<FTransform>(TEXT("RelativeTransform"), FTransform::Identity, /*bAllowsInterpolation=*/false, /*bOverrideParent=*/true);
+
+	TMap<FName, FPCGMetadataAttribute<int64>*> TagToAttributeMap;
+
+	// Hierarchy root point
+	{
+		FPCGPoint& RootPoint = Points.Emplace_GetRef(FTransform::Identity, /*Density=*/1.0f, /*Seed=*/0);
+		RootPoint.Steepness = 1.0f;
+		RootPoint.BoundsMin = RootPoint.BoundsMax = FVector::Zero();
+
+		PointMetadata->InitializeOnSet(RootPoint.MetadataEntry);
+		ActorIndexAttribute->SetValue(RootPoint.MetadataEntry, 0);
+	}
+
+	// Build actor-index map
+	TMap<AActor*, int> ActorIndexMap;
+	int LastActorIndex = 1; // Since the root is the "first" point we'll have, we'll have the map start from 1.
+	UPCGActorHelpers::ForEachActorInWorld(Level, AActor::StaticClass(), [&ActorIndexMap, &LastActorIndex](AActor* Actor)
+	{
+		ActorIndexMap.Add(Actor, LastActorIndex++);
+		return true;
+	});
+
+	// Create points
+	UPCGActorHelpers::ForEachActorInWorld(Level, AActor::StaticClass(), [&](AActor* Actor)
+	{
+		// TODO Actor-level decisions if any; if the actor is "consumed" at this step, make sure to update AllActorBounds as well.
+
+		// Otherwise, parse "known" actor components
+		TArray<UStaticMeshComponent*> SMCs;
+		Actor->GetComponents(SMCs);
+
+		if (SMCs.IsEmpty()) // early out
+		{
+			return true;
+		}
+
+		const FBox ActorBounds = PCGHelpers::GetActorBounds(Actor, /*bIgnorePCGCreatedComponents=*/true);
+		AllActorBounds += ActorBounds;
+
+		for (FName ActorTag : Actor->Tags)
+		{
+			if (!TagToAttributeMap.Contains(ActorTag))
+			{
+				TagToAttributeMap.Add(ActorTag, PointMetadata->CreateAttribute<int64>(ActorTag, 0, /*bAllowsInterpolation=*/false, /*bOverrideParent=*/true));
+			}
+		}
+
+		// Prepare actor-level data that's propagated to all points
+		const FTransform& ActorTransform = Actor->GetTransform();
+		const int64 ActorIndex = ActorIndexMap[Actor];
+		AActor* ParentActor = Actor->GetAttachParentActor();
+		const int64 ParentActorIndex = ParentActor ? ActorIndexMap[ParentActor] : 0;
+		const FTransform RelativeTransform = ParentActor ? ActorTransform.GetRelativeTransform(ParentActor->GetTransform()) : ActorTransform;
+
+		// Hierarchy depth, starts at 1 if the actor doesn't have a parent
+		int HierarchyDepth = 1;
+		while (ParentActor)
+		{
+			++HierarchyDepth;
+			ParentActor = ParentActor->GetAttachParentActor();
+		}
+
+		auto MakePoint = [&](const FTransform& Transform, const FSoftObjectPath& MeshPath, const FBox& MeshBounds, const TArray<UMaterialInterface*>& MeshMaterials)
+		{
+			FPCGPoint& Point = Points.Emplace_GetRef(Transform, /*Density=*/1.0f, PCGHelpers::ComputeSeedFromPosition(Transform.GetLocation()));
+			Point.BoundsMin = MeshBounds.Min;
+			Point.BoundsMax = MeshBounds.Max;
+			Point.Steepness = 1.0f;
+
+			PointMetadata->InitializeOnSet(Point.MetadataEntry);
+			MeshAttribute->SetValue(Point.MetadataEntry, MeshPath);
+
+			if (!MeshMaterials.IsEmpty())
+			{
+				MaterialAttribute->SetValue(Point.MetadataEntry, FSoftObjectPath(MeshMaterials[0]));
+			}
+
+			ActorIndexAttribute->SetValue(Point.MetadataEntry, ActorIndex);
+			ParentIndexAttribute->SetValue(Point.MetadataEntry, ParentActorIndex);
+			RelativeTransformAttribute->SetValue(Point.MetadataEntry, RelativeTransform);
+			HierarchyDepthAttribute->SetValue(Point.MetadataEntry, HierarchyDepth);
+
+			// For all tags, set attribute value to 1.
+			for (FName ActorTag : Actor->Tags)
+			{
+				TagToAttributeMap[ActorTag]->SetValue(Point.MetadataEntry, 1);
+			}
+		};
+
+		for (UStaticMeshComponent* SMC : SMCs)
+		{
+			TObjectPtr<UStaticMesh> StaticMesh = SMC->GetStaticMesh();
+
+			if (!StaticMesh)
+			{
+				continue;
+			}
+
+			const FSoftObjectPath MeshPath(StaticMesh);
+			const FBox MeshBounds = StaticMesh->GetBoundingBox();
+			TArray<UMaterialInterface*> Materials = SMC->GetMaterials();
+
+			// For all instances (or a single instance if this is a static mesh and not an ISM)
+			// if a static mesh -> use actor transform (which might be wrong?)
+			// if ISM -> get instance transform in world space
+			if (UInstancedStaticMeshComponent* ISMC = Cast<UInstancedStaticMeshComponent>(SMC))
+			{
+				const int InstanceCount = ISMC->GetNumInstances();
+				for (int I = 0; I < InstanceCount; ++I)
+				{
+					FTransform InstanceTransform;
+					ISMC->GetInstanceTransform(I, InstanceTransform, /*bWorldSpace=*/true);
+
+					MakePoint(InstanceTransform, MeshPath, MeshBounds, Materials);
+				}
+			}
+			else
+			{
+				MakePoint(ActorTransform, MeshPath, MeshBounds, Materials);
+			}
+		}
+
+		return true;
+	});
+
+	// Finally, create root point in the root data
+	{
+		FPCGPoint& RootPoint = Roots.Emplace_GetRef(FTransform::Identity, 1.0f, 0);
+		RootPoint.BoundsMin = AllActorBounds.Min;
+		RootPoint.BoundsMax = AllActorBounds.Max;
+	}
+
+	return true;
+}
\ No newline at end of file
diff --git a/Engine/Plugins/PCG/Source/PCGEditor/Public/PCGEditorMenuUtils.h b/Engine/Plugins/PCG/Source/PCGEditor/Public/PCGEditorMenuUtils.h
new file mode 100644
--- /dev/null
+++ b/Engine/Plugins/PCG/Source/PCGEditor/Public/PCGEditorMenuUtils.h
@@ -0,0 +1,12 @@
+// Copyright Epic Games, Inc. All Rights Reserved.
+
+#pragma once
+
+#include "ToolMenus.h"
+#include "AssetRegistry/AssetData.h"
+
+namespace PCGEditorMenuUtils
+{
+	PCGEDITOR_API FToolMenuSection& CreatePCGSection(UToolMenu* Menu);
+	PCGEDITOR_API void CreateOrUpdatePCGAssetFromMenu(UToolMenu* Menu, TArray<FAssetData>& WorldAssets);
+}
\ No newline at end of file
diff --git a/Engine/Plugins/PCG/Source/PCGEditor/Public/PCGLevelToAsset.h b/Engine/Plugins/PCG/Source/PCGEditor/Public/PCGLevelToAsset.h
new file mode 100644
--- /dev/null
+++ b/Engine/Plugins/PCG/Source/PCGEditor/Public/PCGLevelToAsset.h
@@ -0,0 +1,40 @@
+// Copyright Epic Games, Inc. All Rights Reserved.
+
+#pragma once
+
+#include "PCGDataAsset.h"
+
+#include "PCGLevelToAsset.generated.h"
+
+class UPackage;
+
+UCLASS(BlueprintType, Blueprintable, meta = (ShowWorldContextPin))
+class PCGEDITOR_API UPCGLevelToAsset : public UObject
+{
+	GENERATED_BODY()
+
+public:
+	/** Creates/Updates on PCG Asset per given world. Allows exporter subclassing by passing in a Subclass. */
+	static void CreateOrUpdatePCGAssets(const TArray<FAssetData>& WorldAssets, TSubclassOf<UPCGLevelToAsset> Subclass = {});
+
+	/** Creates/Updates on PCG Asset for a specific world. Allows exporter subclassing by passing in a Subclass. Will return null if it fails, or the package that was modified otherwise. */
+	static UPackage* CreateOrUpdatePCGAsset(TSoftObjectPtr<UWorld> World, TSubclassOf<UPCGLevelToAsset> Subclass = {});
+
+	/** Creates/Updates a PCG Asset for a specific world, allowing exporter subclassing (and settings creation by extension). Will return null if it fails, or the package that was modified otherwise. */
+	static UPackage* CreateOrUpdatePCGAsset(UWorld* World, TSubclassOf<UPCGLevelToAsset> Subclass = {});
+
+	/** Parses the level and fills in the provided data asset. Implement this in BP to drive the generation in a custom manner. */
+	UFUNCTION(BlueprintNativeEvent, meta = (DisplayName = "Export Level", ForceAsFunction))
+	bool BP_ExportLevel(UWorld* Level, const FString& PackageName, UPCGDataAsset* Asset);
+
+	/** Returns the subtype of PCG data asset generated by the Level To Asset process. Override this in the instances where you would not use the default base class. */
+	UFUNCTION(BlueprintNativeEvent, meta = (DisplayName = "Get Asset Type", ForceAsFunction))
+	TSubclassOf<UPCGDataAsset> BP_GetAssetType() const;
+
+protected:
+	/** Drives the asset creation for the data parsed from the level, will return UPCGDataAsset by default. */
+	virtual TSubclassOf<UPCGDataAsset> GetAssetType() const;
+
+	/** Parses the level and pushes data to the asset. */
+	virtual bool ExportLevel(UWorld* Level, const FString& PackageName, UPCGDataAsset* Asset);
+};
\ No newline at end of file
