[PCG] Fixed a minor bug causing Spawn Actor node to ignore spawning when Collapsing

#rnx

diff --git a/Engine/Plugins/PCG/Source/PCG/Private/Elements/PCGSpawnActor.cpp b/Engine/Plugins/PCG/Source/PCG/Private/Elements/PCGSpawnActor.cpp
--- a/Engine/Plugins/PCG/Source/PCG/Private/Elements/PCGSpawnActor.cpp
+++ b/Engine/Plugins/PCG/Source/PCG/Private/Elements/PCGSpawnActor.cpp
@@ -44,7 +44,7 @@ public:
 
 	bool InitializeForData(const UPCGData* InData, UPCGData* OutData)
 	{
-		if (!InData || !InData->IsA<UPCGPointData>() || !OutData || !OutData->IsA<UPCGPointData>())
+		if (!InData || !InData->IsA<UPCGPointData>())
 		{
 			return false;
 		}
