[PCG] PCG Component now allows Force Generation

diff --git a/Engine/Plugins/PCG/Source/PCGEditor/Private/Details/PCGComponentDetails.cpp b/Engine/Plugins/PCG/Source/PCGEditor/Private/Details/PCGComponentDetails.cpp
--- a/Engine/Plugins/PCG/Source/PCGEditor/Private/Details/PCGComponentDetails.cpp
+++ b/Engine/Plugins/PCG/Source/PCGEditor/Private/Details/PCGComponentDetails.cpp
@@ -2,7 +2,9 @@
 
 #include "Details/PCGComponentDetails.h"
 #include "PCGComponent.h"
+#include "PCGSubsystem.h"
 
+#include "Framework/Application/SlateApplication.h"
 #include "Widgets/Text/STextBlock.h"
 #include "Widgets/Input/SButton.h"
 #include "DetailLayoutBuilder.h"
@@ -70,6 +72,7 @@ void FPCGComponentDetails::CustomizeDetails(IDetailLayoutBuilder& DetailBuilder)
 			[
 				SNew(SButton)
 				.OnClicked(this, &FPCGComponentDetails::OnGenerateClicked)
+				.ToolTipText(FText::FromString("Generates graph data. \nCtrl + Click flushes the cache and force generates."))
 				.Visibility(this, &FPCGComponentDetails::GenerateButtonVisible)
 				[
 					SNew(STextBlock)
@@ -215,7 +218,15 @@ FReply FPCGComponentDetails::OnGenerateClicked()
 	{
 		if (Component.IsValid())
 		{
-			Component.Get()->Generate();
+			bool bForce = false;
+			FModifierKeysState ModifierKeys = FSlateApplication::Get().GetModifierKeys();
+			if (ModifierKeys.IsControlDown())
+			{
+				Component->GetSubsystem()->FlushCache();
+				bForce = true;
+			}
+
+			Component.Get()->Generate(bForce);
 		}
 	}
 
