[PCG] Don't try to generate upstream preview components if they are deactivated, and include some execution logs in bulider log which is a bit starved of info at the moment.

diff --git a/Engine/Plugins/PCG/Source/PCG/Private/Elements/PCGDataFromActor.cpp b/Engine/Plugins/PCG/Source/PCG/Private/Elements/PCGDataFromActor.cpp
--- a/Engine/Plugins/PCG/Source/PCG/Private/Elements/PCGDataFromActor.cpp
+++ b/Engine/Plugins/PCG/Source/PCG/Private/Elements/PCGDataFromActor.cpp
@@ -472,7 +472,7 @@ void FPCGDataFromActorElement::GatherWaitTasks(AActor* FoundActor, FPCGContext*
 		{
 			OutWaitTasks.Add(Component->GetGenerationTaskId());
 		}
-		else if (!Component->bGenerated && (Component->GetSerializedEditingMode() == EPCGEditorDirtyMode::Preview) && Component->GetOwner())
+		else if (!Component->bGenerated && Component->bActivated && (Component->GetSerializedEditingMode() == EPCGEditorDirtyMode::Preview) && Component->GetOwner())
 		{
 #if WITH_EDITOR
 			// Signal that any change notifications from generating upstream component should not trigger re-executions of this component.
diff --git a/Engine/Plugins/PCG/Source/PCG/Private/Utils/PCGGraphExecutionLogging.cpp b/Engine/Plugins/PCG/Source/PCG/Private/Utils/PCGGraphExecutionLogging.cpp
--- a/Engine/Plugins/PCG/Source/PCG/Private/Utils/PCGGraphExecutionLogging.cpp
+++ b/Engine/Plugins/PCG/Source/PCG/Private/Utils/PCGGraphExecutionLogging.cpp
@@ -124,8 +124,7 @@ namespace PCGGraphExecutionLogging
 
 	void LogGraphSchedule(const UPCGComponent* SourceComponent)
 	{
-#if WITH_EDITOR
-		if (!LogEnabled())
+		if (!LogEnabled() && !IsRunningCommandlet())
 		{
 			return;
 		}
@@ -133,13 +132,11 @@ namespace PCGGraphExecutionLogging
 		UE_LOG(LogPCG, Log, TEXT("[%s/%s] --- SCHEDULE GRAPH ---"),
 			(SourceComponent && SourceComponent->GetOwner()) ? *SourceComponent->GetOwner()->GetName() : TEXT("MISSINGCOMPONENT"),
 			(SourceComponent && SourceComponent->GetGraph()) ? *SourceComponent->GetGraph()->GetName() : TEXT("MISSINGGRAPH"));
-#endif
 	}
 
 	void LogGraphScheduleDependency(const UPCGComponent* InComponent)
 	{
-#if WITH_EDITOR
-		if (!LogEnabled())
+		if (!LogEnabled() && !IsRunningCommandlet())
 		{
 			return;
 		}
@@ -147,21 +144,18 @@ namespace PCGGraphExecutionLogging
 		UE_LOG(LogPCG, Log, TEXT("[%s/%s] --- SCHEDULE GRAPH FOR DEPENDENCY ---"),
 			(InComponent && InComponent->GetOwner()) ? *InComponent->GetOwner()->GetName() : TEXT("MISSINGCOMPONENT"),
 			(InComponent && InComponent->GetGraph()) ? *InComponent->GetGraph()->GetName() : TEXT("MISSINGGRAPH"));
-#endif
 	}
 
 	void LogGraphScheduleDependencyFailed(const UPCGComponent* InComponent)
 	{
-#if WITH_EDITOR
-		if (!LogEnabled())
+		if (!LogEnabled() && !IsRunningCommandlet())
 		{
 			return;
 		}
 
-		UE_LOG(LogPCG, Error, TEXT("[%s/%s] Failed to schedule dependency"),
+		UE_LOG(LogPCG, Warning, TEXT("[%s/%s] Failed to schedule dependency"),
 			(InComponent && InComponent->GetOwner()) ? *InComponent->GetOwner()->GetName() : TEXT("MISSINGCOMPONENT"),
 			(InComponent && InComponent->GetGraph()) ? *InComponent->GetGraph()->GetName() : TEXT("MISSINGGRAPH"));
-#endif
 	}
 	
 	void LogGraphPostSchedule(const TMap<FPCGTaskId, FPCGGraphTask>& Tasks, const TMap<FPCGTaskId, TSet<FPCGTaskId>>& TaskSuccessors)
@@ -180,8 +174,7 @@ namespace PCGGraphExecutionLogging
 
 	void LogPostProcessGraph(const UPCGComponent* InSourceComponent)
 	{
-#if WITH_EDITOR
-		if (!LogEnabled())
+		if (!LogEnabled() && !IsRunningCommandlet())
 		{
 			return;
 		}
@@ -189,13 +182,11 @@ namespace PCGGraphExecutionLogging
 		UE_LOG(LogPCG, Log, TEXT("[%s/%s] UPCGComponent::PostProcessGraph"),
 			(InSourceComponent && InSourceComponent->GetOwner()) ? *InSourceComponent->GetOwner()->GetName() : TEXT("MISSINGCOMPONENT"),
 			(InSourceComponent && InSourceComponent->GetGraph()) ? *InSourceComponent->GetGraph()->GetName() : TEXT("MISSINGGRAPH"));
-#endif
 	}
 
 	void LogComponentCancellation(const TSet<UPCGComponent*>& CancelledComponents)
 	{
-#if WITH_EDITOR
-		if (!LogEnabled())
+		if (!LogEnabled() && !IsRunningCommandlet())
 		{
 			return;
 		}
@@ -206,7 +197,6 @@ namespace PCGGraphExecutionLogging
 				(Component && Component->GetOwner()) ? *Component->GetOwner()->GetName() : TEXT("MISSINGCOMPONENT"),
 				(Component && Component->GetGraph()) ? *Component->GetGraph()->GetName() : TEXT("MISSINGGRAPH"));
 		}
-#endif
 	}
 
 	void LogChangeOriginIgnoredForComponent(const UObject* InObject, const UPCGComponent* InComponent)
