[PCG] Add Search to profiling view + re-organize column names/sizes

* Also fix Prepare data time not being accumulated.

#rnx

diff --git a/Engine/Plugins/PCG/Source/PCG/Private/Utils/PCGExtraCapture.cpp b/Engine/Plugins/PCG/Source/PCG/Private/Utils/PCGExtraCapture.cpp
--- a/Engine/Plugins/PCG/Source/PCG/Private/Utils/PCGExtraCapture.cpp
+++ b/Engine/Plugins/PCG/Source/PCG/Private/Utils/PCGExtraCapture.cpp
@@ -46,7 +46,7 @@ void PCGUtils::FExtraCapture::Update(const PCGUtils::FScopedCall& InScopedCall)
 		}
 
 		Timer.PrepareDataFrameCount++;
-		Timer.PrepareDataTime = ThisFrameTime;
+		Timer.PrepareDataTime += ThisFrameTime;
 		Timer.PrepareDataWallTime = CurrentTime - Timer.PrepareDataStartTime;
 		break;
 	case EPCGExecutionPhase::Execute:
diff --git a/Engine/Plugins/PCG/Source/PCGEditor/Private/SPCGEditorGraphProfilingView.cpp b/Engine/Plugins/PCG/Source/PCGEditor/Private/SPCGEditorGraphProfilingView.cpp
--- a/Engine/Plugins/PCG/Source/PCGEditor/Private/SPCGEditorGraphProfilingView.cpp
+++ b/Engine/Plugins/PCG/Source/PCGEditor/Private/SPCGEditorGraphProfilingView.cpp
@@ -12,6 +12,7 @@
 #include "Framework/Views/TableViewMetadata.h"
 #include "Widgets/Input/SButton.h"
 #include "Widgets/Input/SCheckBox.h"
+#include "Widgets/Input/SSearchBox.h"
 #include "Widgets/Layout/SScrollBox.h"
 #include "Widgets/Views/SListView.h"
 
@@ -35,26 +36,26 @@ namespace PCGEditorGraphProfilingView
 
 	/** Labels of the columns */
 	const FText TEXT_NodeLabel = LOCTEXT("NodeLabel", "Node");
-	const FText TEXT_PrepareDataTimeLabel = LOCTEXT("PrepareDataTimeLabel", "PrepareData Time (ms)");
-	const FText TEXT_PrepareDataWallTimeLabel = LOCTEXT("PrepareDataWallTimeLabel", "PrepareData WallTime (ms)");
-	const FText TEXT_MinExecutionFrameTimeLabel = LOCTEXT("MinExecutionFrameTimeLabel", "Min Frame Time (ms)");
-	const FText TEXT_MaxExecutionFrameTimeLabel = LOCTEXT("MaxExecutionFrameTimeLabel", "Max Frame Time (ms)");
-	const FText TEXT_ExecutionTimeLabel = LOCTEXT("ExecutionTimeLabel", "Execution Time (ms)");
-	const FText TEXT_ExecutionWallTimeLabel = LOCTEXT("ExecutionWallTimeLabel", "Execution Wall Time (ms)");
+	const FText TEXT_PrepareDataTimeLabel = LOCTEXT("PrepareDataTimeLabel", "Prepare");
+	const FText TEXT_PrepareDataWallTimeLabel = LOCTEXT("PrepareDataWallTimeLabel", "Prepare WallTime");
+	const FText TEXT_MinExecutionFrameTimeLabel = LOCTEXT("MinExecutionFrameTimeLabel", "Min FrameTime");
+	const FText TEXT_MaxExecutionFrameTimeLabel = LOCTEXT("MaxExecutionFrameTimeLabel", "Max FrameTime");
+	const FText TEXT_ExecutionTimeLabel = LOCTEXT("ExecutionTimeLabel", "Exec");
+	const FText TEXT_ExecutionWallTimeLabel = LOCTEXT("ExecutionWallTimeLabel", "Exec WallTime");
 	const FText TEXT_NbExecutionFramesLabel = LOCTEXT("NbExecutionFramesLabel", "Exec Frames");
-	const FText TEXT_TotalTimeLabel = LOCTEXT("TotalTimeLabel", "Total Time (ms)");
-	const FText TEXT_TotalWallTimeLabel = LOCTEXT("TotalWallTimeLabel", "Total Wall Time (ms)");
+	const FText TEXT_TotalTimeLabel = LOCTEXT("TotalTimeLabel", "Total");
+	const FText TEXT_TotalWallTimeLabel = LOCTEXT("TotalWallTimeLabel", "Total WallTime");
 
 	/** Tooltips */
-	const FText TEXT_PrepareDataTimeTooltip = LOCTEXT("PrepareDataTimeTooltip", "Cost of the PrepareData execution phase which some nodes use to process the incoming data.");
-	const FText TEXT_PrepareDataWallTimeTooltip = LOCTEXT("PrepareDataWallTimeTooltip", "Total real time elapsed between prepare data first being called until completion, including any wait/sleep time.");
-	const FText TEXT_MinExecutionFrameTimeTooltip = LOCTEXT("MinExecutionFrameTimeTooltip", "The minimum time spent of all execution frames.");
-	const FText TEXT_MaxExecutionFrameTimeTooltip = LOCTEXT("MaxExecutionFrameTimeTooltip", "The maximum time spent of all execution frames.");
-	const FText TEXT_ExecutionTimeTooltip = LOCTEXT("ExecutionTimeTooltip", "The total time spent for execution, summed over all execution frames.");
-	const FText TEXT_ExecutionWallTimeTooltip = LOCTEXT("ExecutionWallTimeTooltip", "Total real time elapsed between execute first being called until completion, including any wait/sleep time.");
+	const FText TEXT_PrepareDataTimeTooltip = LOCTEXT("PrepareDataTimeTooltip", "Cost of the PrepareData execution phase which some nodes use to process the incoming data, in ms.");
+	const FText TEXT_PrepareDataWallTimeTooltip = LOCTEXT("PrepareDataWallTimeTooltip", "Total real time elapsed between prepare data first being called until completion, including any wait/sleep time, in ms.");
+	const FText TEXT_MinExecutionFrameTimeTooltip = LOCTEXT("MinExecutionFrameTimeTooltip", "The minimum time spent of all execution frames, in ms.");
+	const FText TEXT_MaxExecutionFrameTimeTooltip = LOCTEXT("MaxExecutionFrameTimeTooltip", "The maximum time spent of all execution frames, in ms.");
+	const FText TEXT_ExecutionTimeTooltip = LOCTEXT("ExecutionTimeTooltip", "The total time spent for execution, summed over all execution frames, in ms.");
+	const FText TEXT_ExecutionWallTimeTooltip = LOCTEXT("ExecutionWallTimeTooltip", "Total real time elapsed between execute first being called until completion, including any wait/sleep time, in ms.");
 	const FText TEXT_NbExecutionFramesTooltip = LOCTEXT("NbExecutionFramesTooltip", "The number of frames in which one or more execution phases were executed.");
-	const FText TEXT_TotalTimeTooltip = LOCTEXT("TotalTimeTooltip", "The total time spent in this node, summed over all execution and prepare data frames.");
-	const FText TEXT_TotalWallTimeTooltip = LOCTEXT("TotalWallTimeTooltip", "Total real time elapsed between the first call until completion, including any wait/sleep time.");
+	const FText TEXT_TotalTimeTooltip = LOCTEXT("TotalTimeTooltip", "The total time spent in this node, summed over all execution and prepare data frames, in ms.");
+	const FText TEXT_TotalWallTimeTooltip = LOCTEXT("TotalWallTimeTooltip", "Total real time elapsed between the first call until completion, including any wait/sleep time, in ms.");
 }
 
 void SPCGProfilingListViewItemRow::Construct(const FArguments& InArgs, const TSharedRef<STableViewBase>& InOwnerTableView, const PCGProfilingListViewItemPtr& Item)
@@ -126,9 +127,18 @@ TSharedRef<SWidget> SPCGProfilingListViewItemRow::GenerateWidgetForColumn(const
 		}
 	}
 
-	return SNew(STextBlock)
+	TSharedRef<SWidget> DataWidget = 
+		SNew(STextBlock)
 		.Text(ColumnData)
 		.OverflowPolicy(ETextOverflowPolicy::Ellipsis);
+
+	// Add the internal name of the node as tooltip
+	if (ColumnId == PCGEditorGraphProfilingView::NAME_Node && InternalItem->PCGNode)
+	{
+		DataWidget->SetToolTipText(FText::FromString(InternalItem->PCGNode->GetName()));
+	}
+
+	return DataWidget;
 }
 
 void SPCGEditorGraphProfilingView::OnItemDoubleClicked(PCGProfilingListViewItemPtr Item)
@@ -198,6 +208,14 @@ void SPCGEditorGraphProfilingView::Construct(const FArguments& InArgs, TSharedPt
 			SNew(SHorizontalBox)
 			+SHorizontalBox::Slot()
 			.AutoWidth()
+			[
+				SNew(SSearchBox)
+				.HintText(LOCTEXT("PCGGraphSearchHint", "Enter text to find nodes..."))
+				.OnTextChanged(this, &SPCGEditorGraphProfilingView::OnSearchTextChanged)
+				.OnTextCommitted(this, &SPCGEditorGraphProfilingView::OnSearchTextCommitted)
+			]
+			+SHorizontalBox::Slot()
+			.AutoWidth()
 			.Padding(1.0f, 1.0f)
 			[
 				SNew(SButton)
@@ -295,14 +313,14 @@ TSharedRef<SHeaderRow> SPCGEditorGraphProfilingView::CreateHeaderRowWidget()
 		.ResizeMode(ESplitterResizeMode::FixedPosition)
 		.CanSelectGeneratedColumn(true)
 		+ SHeaderRow::Column(PCGEditorGraphProfilingView::NAME_Node)
-		.ManualWidth(350.0f)
+		.ManualWidth(300.0f)
 		.DefaultLabel(PCGEditorGraphProfilingView::TEXT_NodeLabel)
 		.HAlignHeader(HAlign_Center)
 		.HAlignCell(HAlign_Left)
 		.SortMode(this, &SPCGEditorGraphProfilingView::GetColumnSortMode, PCGEditorGraphProfilingView::NAME_Node)
 		.OnSort(this, &SPCGEditorGraphProfilingView::OnSortColumnHeader)
 		+ SHeaderRow::Column(PCGEditorGraphProfilingView::NAME_PrepareDataTime)
-		.ManualWidth(170.0f)
+		.ManualWidth(75.0f)
 		.DefaultLabel(PCGEditorGraphProfilingView::TEXT_PrepareDataTimeLabel)
 		.HAlignHeader(HAlign_Center)
 		.HAlignCell(HAlign_Right)
@@ -311,7 +329,7 @@ TSharedRef<SHeaderRow> SPCGEditorGraphProfilingView::CreateHeaderRowWidget()
 		.InitialSortMode(EColumnSortMode::Descending)
 		.DefaultTooltip(PCGEditorGraphProfilingView::TEXT_PrepareDataTimeTooltip)
 		+ SHeaderRow::Column(PCGEditorGraphProfilingView::NAME_PrepareDataWallTime)
-		.ManualWidth(170.0f)
+		.ManualWidth(110.0f)
 		.DefaultLabel(PCGEditorGraphProfilingView::TEXT_PrepareDataWallTimeLabel)
 		.HAlignHeader(HAlign_Center)
 		.HAlignCell(HAlign_Right)
@@ -320,7 +338,7 @@ TSharedRef<SHeaderRow> SPCGEditorGraphProfilingView::CreateHeaderRowWidget()
 		.InitialSortMode(EColumnSortMode::Descending)
 		.DefaultTooltip(PCGEditorGraphProfilingView::TEXT_PrepareDataWallTimeTooltip)
 		+ SHeaderRow::Column(PCGEditorGraphProfilingView::NAME_NbExecutionFrames)
-		.ManualWidth(105.0f)
+		.ManualWidth(90.0f)
 		.DefaultLabel(PCGEditorGraphProfilingView::TEXT_NbExecutionFramesLabel)
 		.HAlignHeader(HAlign_Center)
 		.HAlignCell(HAlign_Right)
@@ -329,7 +347,7 @@ TSharedRef<SHeaderRow> SPCGEditorGraphProfilingView::CreateHeaderRowWidget()
 		.InitialSortMode(EColumnSortMode::Descending)
 		.DefaultTooltip(PCGEditorGraphProfilingView::TEXT_NbExecutionFramesTooltip)
 		+ SHeaderRow::Column(PCGEditorGraphProfilingView::NAME_MinExecutionFrameTime)
-		.ManualWidth(155.0f)
+		.ManualWidth(110.0f)
 		.DefaultLabel(PCGEditorGraphProfilingView::TEXT_MinExecutionFrameTimeLabel)
 		.HAlignHeader(HAlign_Center)
 		.HAlignCell(HAlign_Right)
@@ -338,7 +356,7 @@ TSharedRef<SHeaderRow> SPCGEditorGraphProfilingView::CreateHeaderRowWidget()
 		.InitialSortMode(EColumnSortMode::Descending)
 		.DefaultTooltip(PCGEditorGraphProfilingView::TEXT_MinExecutionFrameTimeTooltip)
 		+ SHeaderRow::Column(PCGEditorGraphProfilingView::NAME_MaxExecutionFrameTime)
-		.ManualWidth(155.0f)
+		.ManualWidth(110.0f)
 		.DefaultLabel(PCGEditorGraphProfilingView::TEXT_MaxExecutionFrameTimeLabel)
 		.HAlignHeader(HAlign_Center)
 		.HAlignCell(HAlign_Right)
@@ -347,7 +365,7 @@ TSharedRef<SHeaderRow> SPCGEditorGraphProfilingView::CreateHeaderRowWidget()
 		.InitialSortMode(EColumnSortMode::Descending)
 		.DefaultTooltip(PCGEditorGraphProfilingView::TEXT_MaxExecutionFrameTimeTooltip)
 		+ SHeaderRow::Column(PCGEditorGraphProfilingView::NAME_ExecutionTime)
-		.ManualWidth(120.0f)
+		.ManualWidth(75.0f)
 		.DefaultLabel(PCGEditorGraphProfilingView::TEXT_ExecutionTimeLabel)
 		.HAlignHeader(HAlign_Center)
 		.HAlignCell(HAlign_Right)
@@ -356,7 +374,7 @@ TSharedRef<SHeaderRow> SPCGEditorGraphProfilingView::CreateHeaderRowWidget()
 		.InitialSortMode(EColumnSortMode::Descending)
 		.DefaultTooltip(PCGEditorGraphProfilingView::TEXT_ExecutionTimeTooltip)
 		+ SHeaderRow::Column(PCGEditorGraphProfilingView::NAME_ExecutionWallTime)
-		.ManualWidth(120.0f)
+		.ManualWidth(100.0f)
 		.DefaultLabel(PCGEditorGraphProfilingView::TEXT_ExecutionWallTimeLabel)
 		.HAlignHeader(HAlign_Center)
 		.HAlignCell(HAlign_Right)
@@ -365,7 +383,7 @@ TSharedRef<SHeaderRow> SPCGEditorGraphProfilingView::CreateHeaderRowWidget()
 		.InitialSortMode(EColumnSortMode::Descending)
 		.DefaultTooltip(PCGEditorGraphProfilingView::TEXT_ExecutionWallTimeTooltip)
 		+ SHeaderRow::Column(PCGEditorGraphProfilingView::NAME_TotalTime)
-		.ManualWidth(120.0f)
+		.ManualWidth(75.0f)
 		.DefaultLabel(PCGEditorGraphProfilingView::TEXT_TotalTimeLabel)
 		.HAlignHeader(HAlign_Center)
 		.HAlignCell(HAlign_Right)
@@ -374,7 +392,7 @@ TSharedRef<SHeaderRow> SPCGEditorGraphProfilingView::CreateHeaderRowWidget()
 		.InitialSortMode(EColumnSortMode::Descending)
 		.DefaultTooltip(PCGEditorGraphProfilingView::TEXT_TotalTimeTooltip)
 		+ SHeaderRow::Column(PCGEditorGraphProfilingView::NAME_TotalWallTime)
-		.ManualWidth(120.0f)
+		.ManualWidth(100.0f)
 		.DefaultLabel(PCGEditorGraphProfilingView::TEXT_TotalWallTimeLabel)
 		.HAlignHeader(HAlign_Center)
 		.HAlignCell(HAlign_Right)
@@ -461,6 +479,7 @@ namespace PCGEditorGraphProfilingView
 		const TMap<const UPCGNode*, UPCGEditorGraphNode*>& EditorNodeLookup,
 		bool bExpandSubgraph,
 		const FString& FolderName,
+		const FString& SearchString,
 		UPCGEditorGraphNode* CurrentEditorNode = nullptr)
 	{
 		for (const PCGUtils::FCallTreeInfo& Info : TreeInfo)
@@ -474,8 +493,10 @@ namespace PCGEditorGraphProfilingView
 				Fullname += Info.Node->GetNodeTitle(EPCGNodeTitleType::ListView).ToString();
 			}
 
+			const bool bFilteredIn = SearchString.IsEmpty() || Fullname.Find(SearchString) != INDEX_NONE || (Info.Node && Info.Node->GetName().Find(SearchString) != INDEX_NONE);
+
 			// don't show inclusive times when bExpandSubgraph is on or the total sum of everything will mismatch reality
-			if ((Info.Children.IsEmpty() || bExpandSubgraph == false) && Info.CallTime.MaxExecutionFrameTime > 0)
+			if (bFilteredIn && (Info.Children.IsEmpty() || bExpandSubgraph == false) && Info.CallTime.MaxExecutionFrameTime > 0)
 			{
 				PCGProfilingListViewItemPtr ListViewItem = MakeShared<FPCGProfilingListViewItem>();
 
@@ -491,7 +512,7 @@ namespace PCGEditorGraphProfilingView
 
 			if (bExpandSubgraph && !Info.Children.IsEmpty())
 			{
-				AddListItems(OutListViewItems, Info.Children, EditorNodeLookup, /* bExpandSubgraph */ true, Fullname + "/", EditorNode);
+				AddListItems(OutListViewItems, Info.Children, EditorNodeLookup, /* bExpandSubgraph */ true, Fullname + "/", SearchString, EditorNode);
 			}
 		}
 	}
@@ -536,7 +557,7 @@ FReply SPCGEditorGraphProfilingView::Refresh()
 	ListViewItems.Reserve(TreeInfo.Children.Num());
 
 	//TODO: could turn this into a tree instead of expanding into a list
-	PCGEditorGraphProfilingView::AddListItems(ListViewItems, TreeInfo.Children, EditorNodeLookup, bExpandSubgraph, FString());
+	PCGEditorGraphProfilingView::AddListItems(ListViewItems, TreeInfo.Children, EditorNodeLookup, bExpandSubgraph, FString(), SearchValue);
 
 	if (SortingColumn != NAME_None && SortMode != EColumnSortMode::None)
 	{
@@ -620,4 +641,15 @@ TSharedRef<ITableRow> SPCGEditorGraphProfilingView::OnGenerateRow(PCGProfilingLi
 	return SNew(SPCGProfilingListViewItemRow, OwnerTable, Item);
 }
 
+void SPCGEditorGraphProfilingView::OnSearchTextChanged(const FText& InText)
+{
+	SearchValue = InText.ToString();
+	Refresh();
+}
+
+void SPCGEditorGraphProfilingView::OnSearchTextCommitted(const FText& InText, ETextCommit::Type InCommitType)
+{
+	OnSearchTextChanged(InText);
+}
+
 #undef LOCTEXT_NAMESPACE
diff --git a/Engine/Plugins/PCG/Source/PCGEditor/Private/SPCGEditorGraphProfilingView.h b/Engine/Plugins/PCG/Source/PCGEditor/Private/SPCGEditorGraphProfilingView.h
--- a/Engine/Plugins/PCG/Source/PCGEditor/Private/SPCGEditorGraphProfilingView.h
+++ b/Engine/Plugins/PCG/Source/PCGEditor/Private/SPCGEditorGraphProfilingView.h
@@ -73,6 +73,12 @@ private:
 	FText GetTotalTimeLabel() const;
 	FText GetTotalWallTimeLabel() const;
 
+	/** Called when user changes the text they are searching for */
+	void OnSearchTextChanged(const FText& InText);
+
+	/** Called when user changes commits text to the search box */
+	void OnSearchTextCommitted(const FText& InText, ETextCommit::Type InCommitType);
+
 	/** Pointer back to the PCG editor that owns us */
 	TWeakPtr<FPCGEditor> PCGEditorPtr;
 
@@ -94,4 +100,7 @@ private:
 	double TotalWallTime = 0.0;
 
 	bool bExpandSubgraph = true;
+
+	/** The string to search for */
+	FString SearchValue;
 };
