[PCG] Fix usecase where the user wants to save a graph instance on itself

#rnx

diff --git a/Engine/Plugins/PCG/Source/PCGEditor/Private/Details/PCGGraphInstanceDetails.cpp b/Engine/Plugins/PCG/Source/PCGEditor/Private/Details/PCGGraphInstanceDetails.cpp
--- a/Engine/Plugins/PCG/Source/PCGEditor/Private/Details/PCGGraphInstanceDetails.cpp
+++ b/Engine/Plugins/PCG/Source/PCGEditor/Private/Details/PCGGraphInstanceDetails.cpp
@@ -74,7 +74,7 @@ FReply FPCGGraphInstanceDetails::OnSaveInstanceClicked()
 {
 	const UPCGGraphInstance* GraphInstance = !SelectedGraphInstances.IsEmpty() ? SelectedGraphInstances[0].Get() : nullptr;
 
-	if (!GraphInstance || !GraphInstance->Graph)
+	if (!GraphInstance || !IsValid(GraphInstance) || !GraphInstance->Graph)
 	{
 		return FReply::Handled();
 	}
@@ -86,6 +86,10 @@ FReply FPCGGraphInstanceDetails::OnSaveInstanceClicked()
 	FString NewAssetName;
 	PCGEditorUtils::GetParentPackagePathAndUniqueName(GraphInstance, LOCTEXT("NewPCGGraphInstanceAsset", "NewPCGGraphInstance").ToString(), NewPackageName, NewAssetName);
 
+	// Saving info about the graph instance, since there is nothing preventing the user to override the asset of the graph instance they want to save. In that case, GraphInstance is destroyed, so we need to keep the info alive.
+	UPCGGraphInterface* Graph = GraphInstance->Graph;
+	FPCGOverrideInstancedPropertyBag Overrides = GraphInstance->ParametersOverrides;
+
 	UPCGGraphInstance* NewPCGGraphInstance = Cast<UPCGGraphInstance>(AssetTools.CreateAssetWithDialog(NewAssetName, NewPackageName, GraphInstance->GetClass(), Factory, "PCGEditor_SaveGraphInstance"));
 
 	if (NewPCGGraphInstance == nullptr)
@@ -94,8 +98,8 @@ FReply FPCGGraphInstanceDetails::OnSaveInstanceClicked()
 		return FReply::Handled();
 	}
 
-	NewPCGGraphInstance->SetGraph(GraphInstance->Graph);
-	NewPCGGraphInstance->ParametersOverrides = GraphInstance->ParametersOverrides;
+	NewPCGGraphInstance->SetGraph(Graph);
+	NewPCGGraphInstance->ParametersOverrides = std::move(Overrides);
 
 	// Save the new asset
 	UEditorAssetLibrary::SaveLoadedAsset(NewPCGGraphInstance);
