[PCG] Improvements to Attribute Noise node

* Fix issue where invalid operation (no attribute or not valid type) was crashing
* Now supports AttributeSets
* Optimize the get/set of values using GetRange on accessor
* Remove the use of a temporary bufer.
* Fix tests and add new ones, for the bug fix and the new features.
* Fixing doc on Async range

#rnx

diff --git a/Engine/Plugins/PCG/Source/PCG/Private/Elements/PCGAttributeNoise.cpp b/Engine/Plugins/PCG/Source/PCG/Private/Elements/PCGAttributeNoise.cpp
--- a/Engine/Plugins/PCG/Source/PCG/Private/Elements/PCGAttributeNoise.cpp
+++ b/Engine/Plugins/PCG/Source/PCG/Private/Elements/PCGAttributeNoise.cpp
@@ -153,6 +153,22 @@ FPCGElementPtr UPCGAttributeNoiseSettings::CreateElement() const
 	return MakeShared<FPCGAttributeNoiseElement>();
 }
 
+TArray<FPCGPinProperties> UPCGAttributeNoiseSettings::InputPinProperties() const
+{
+	TArray<FPCGPinProperties> PinProperties;
+	PinProperties.Emplace(PCGPinConstants::DefaultInputLabel, EPCGDataType::Point | EPCGDataType::Param);
+
+	return PinProperties;
+}
+
+TArray<FPCGPinProperties> UPCGAttributeNoiseSettings::OutputPinProperties() const
+{
+	TArray<FPCGPinProperties> PinProperties;
+	PinProperties.Emplace(PCGPinConstants::DefaultOutputLabel, EPCGDataType::Point | EPCGDataType::Param);
+
+	return PinProperties;
+}
+
 FPCGContext* FPCGAttributeNoiseElement::CreateContext()
 {
 	return new FPCGAttributeNoiseContext();
@@ -168,7 +184,7 @@ bool FPCGAttributeNoiseElement::ExecuteInternal(FPCGContext* InContext) const
 	const UPCGAttributeNoiseSettings* Settings = Context->GetInputSettings<UPCGAttributeNoiseSettings>();
 	check(Settings);
 
-	TArray<FPCGTaggedData> Inputs = Context->InputData.GetInputs();
+	TArray<FPCGTaggedData> Inputs = Context->InputData.GetInputsByPin(PCGPinConstants::DefaultInputLabel);
 	TArray<FPCGTaggedData>& Outputs = Context->OutputData.TaggedData;
 
 	// Precompute a seed based on the settings one and the component one
@@ -180,40 +196,24 @@ bool FPCGAttributeNoiseElement::ExecuteInternal(FPCGContext* InContext) const
 		const FPCGTaggedData& Input = Inputs[CurrentInput];
 
 		TRACE_CPUPROFILER_EVENT_SCOPE(FPCGAttributeNoiseElement::InputLoop);
-		
-		const TArray<FPCGPoint>* InputPoints = nullptr;
-		TArray<FPCGPoint>* OutputPoints = nullptr;
 
 		if (!Context->bDataPreparedForCurrentInput)
 		{
-			const UPCGSpatialData* SpatialData = Cast<UPCGSpatialData>(Input.Data);
-
-			if (!SpatialData)
+			const UPCGData* InputData = Input.Data;
+			if (!InputData || !InputData->ConstMetadata())
 			{
-				PCGE_LOG(Error, GraphAndLog, LOCTEXT("InputMissingSpatialData", "Unable to get Spatial data from input"));
+				PCGE_LOG(Error, GraphAndLog, FText::Format(LOCTEXT("InputUnsuportedData", "Data {0} is neither spatial nor an attribute set, unsupported."), Context->CurrentInput));
 				Context->CurrentInput++;
 				continue;
 			}
 
-			const UPCGPointData* PointData = SpatialData->ToPointData(Context);
-
-			if (!PointData)
-			{
-				PCGE_LOG(Error, GraphAndLog, LOCTEXT("InputMissingPointData", "Unable to get Point data from input"));
-				Context->CurrentInput++;
-				continue;
-			}
-
-			InputPoints = &PointData->GetPoints();
-
-			Context->InputSource = Settings->InputSource.CopyAndFixLast(PointData);
+			Context->InputSource = Settings->InputSource.CopyAndFixLast(InputData);
 
 			// Create a dummy accessor on the input before allocating the output, to avoid doing useless allocation.
-			TUniquePtr<const IPCGAttributeAccessor> TempInputAccessor = PCGAttributeAccessorHelpers::CreateConstAccessor(PointData, Context->InputSource);
+			TUniquePtr<const IPCGAttributeAccessor> TempInputAccessor = PCGAttributeAccessorHelpers::CreateConstAccessor(InputData, Context->InputSource);
 			if (!TempInputAccessor)
 			{
-				Outputs.RemoveAt(Outputs.Num() - 1);
-				PCGE_LOG(Error, GraphAndLog, FText::Format(LOCTEXT("CantCreateAccessor", "Could not find Attribute/Property {0}"), FText::FromName(Context->InputSource.GetName())));
+				PCGE_LOG(Error, GraphAndLog, FText::Format(LOCTEXT("CantCreateAccessor", "Could not find Attribute/Property '{0}'"), Context->InputSource.GetDisplayText()));
 				Context->CurrentInput++;
 				continue;
 			}
@@ -221,8 +221,7 @@ bool FPCGAttributeNoiseElement::ExecuteInternal(FPCGContext* InContext) const
 			// Also need to make sure the accessor is a "noisable" type
 			if (!PCG::Private::IsOfTypes<int32, int64, float, double, FVector, FVector2D, FVector4, FRotator, FQuat>(TempInputAccessor->GetUnderlyingType()))
 			{
-				Outputs.RemoveAt(Outputs.Num() - 1);
-				PCGE_LOG(Error, GraphAndLog, FText::Format(LOCTEXT("AttributeIsNotANumericalType", "Attribute/Property {0} is not a numerical type, we can't apply noise to it."), FText::FromName(Context->InputSource.GetName())));
+				PCGE_LOG(Error, GraphAndLog, FText::Format(LOCTEXT("AttributeIsNotANumericalType", "Attribute/Property '{0}' is not a numerical type, we can't apply noise to it."), Context->InputSource.GetDisplayText()));
 				Context->CurrentInput++;
 				continue;
 			}
@@ -230,58 +229,48 @@ bool FPCGAttributeNoiseElement::ExecuteInternal(FPCGContext* InContext) const
 			TempInputAccessor.Reset();
 
 			FPCGTaggedData& Output = Outputs.Add_GetRef(Input);
-			UPCGPointData* OutputData = NewObject<UPCGPointData>();
-			OutputData->InitializeFromData(PointData);
-			OutputPoints = &OutputData->GetMutablePoints();
-			OutputPoints->SetNumUninitialized(InputPoints->Num());
+			UPCGData* OutputData = InputData->DuplicateData();
 			Output.Data = OutputData;
 
 			// Then create the accessor/keys
-			Context->InputAccessor = PCGAttributeAccessorHelpers::CreateConstAccessor(OutputData, Context->InputSource);
-			Context->Keys = PCGAttributeAccessorHelpers::CreateKeys(OutputData, Context->InputSource);
+			Context->InputAccessor = PCGAttributeAccessorHelpers::CreateConstAccessor(InputData, Context->InputSource);
+			Context->InputKeys = PCGAttributeAccessorHelpers::CreateConstKeys(InputData, Context->InputSource);
 
 			// It won't fail because we validated on the input data and output is initialized from input.
-			check(Context->InputAccessor && Context->Keys);
-
-			// Allocate temp buffer and create output accessor if needed
-			const bool bValid = PCGMetadataAttribute::CallbackWithRightType(Context->InputAccessor->GetUnderlyingType(), [this, Context, Settings, OutputData](auto&& Dummy) -> bool
-			{
-				check(Context);
+			check(Context->InputAccessor && Context->InputKeys);
 
-				using AttributeType = std::decay_t<decltype(Dummy)>;
-				int32 NumPoints = Context->Keys->GetNum();
-				Context->TempValuesBuffer.SetNumUninitialized(sizeof(AttributeType) * NumPoints);
+			Context->OutputTarget = Settings->OutputTarget.CopyAndFixSource(&Context->InputSource, Input.Data);
 
-				Context->OutputTarget = Settings->OutputTarget.CopyAndFixSource(&Context->InputSource);
-
-				Context->OutputAccessor = PCGAttributeAccessorHelpers::CreateAccessor(OutputData, Context->OutputTarget);
-				if (!Context->OutputAccessor && Context->OutputTarget.GetSelection() == EPCGAttributePropertySelection::Attribute)
+			Context->OutputAccessor = PCGAttributeAccessorHelpers::CreateAccessor(OutputData, Context->OutputTarget);
+			if (!Context->OutputAccessor && Context->OutputTarget.IsBasicAttribute())
+			{
+				PCGMetadataAttribute::CallbackWithRightType(Context->InputAccessor->GetUnderlyingType(), [Context, OutputData](auto&& Dummy)
 				{
-					OutputData->Metadata->CreateAttribute<AttributeType>(Context->OutputTarget.GetName(), AttributeType{}, /*bAllowsInterpolation=*/ true, /*bOverrideParent=*/false);
-					Context->OutputAccessor = PCGAttributeAccessorHelpers::CreateAccessor(OutputData, Context->OutputTarget);
-				}
+					using AttributeType = std::decay_t<decltype(Dummy)>;
 
-				if (!Context->OutputAccessor)
-				{
-					PCGE_LOG(Error, GraphAndLog, FText::Format(LOCTEXT("OutputTargetInvalid", "Failed to find/create Attribute/Property {0}."), FText::FromName(Context->OutputTarget.GetName())));
-					return false;
-				}
+					OutputData->MutableMetadata()->CreateAttribute<AttributeType>(Context->OutputTarget.GetName(), AttributeType{}, /*bAllowsInterpolation=*/ true, /*bOverrideParent=*/false);
+					Context->OutputAccessor = PCGAttributeAccessorHelpers::CreateAccessor(OutputData, Context->OutputTarget);
+				});
+			}
 
-				if (!PCG::Private::IsBroadcastable(Context->InputAccessor->GetUnderlyingType(), Context->OutputAccessor->GetUnderlyingType()))
-				{
-					PCGE_LOG(Error, GraphAndLog, FText::Format(LOCTEXT("CantBroadcast", "Cannot convert Attribute {0} ({1}) into Attribute {2} ({3})."),
-						FText::FromName(Context->InputSource.GetName()),
-						FText::FromString(PCG::Private::GetTypeName(Context->InputAccessor->GetUnderlyingType())),
-						FText::FromName(Context->OutputTarget.GetName()),
-						FText::FromString(PCG::Private::GetTypeName(Context->OutputAccessor->GetUnderlyingType()))));
-					return false;
-				}
+			Context->OutputKeys = PCGAttributeAccessorHelpers::CreateKeys(OutputData, Context->OutputTarget);
 
-				return true;
-			});
+			if (!Context->OutputAccessor || !Context->OutputKeys)
+			{
+				PCGE_LOG(Error, GraphAndLog, FText::Format(LOCTEXT("OutputTargetInvalid", "Failed to find/create Attribute/Property '{0}'."), Context->OutputTarget.GetDisplayText()));
+				Outputs.RemoveAt(Outputs.Num() - 1);
+				Context->CurrentInput++;
+				continue;
+			}
 
-			if (!bValid)
+			if (!PCG::Private::IsBroadcastableOrConstructible(Context->InputAccessor->GetUnderlyingType(), Context->OutputAccessor->GetUnderlyingType()))
 			{
+				PCGE_LOG(Error, GraphAndLog, FText::Format(LOCTEXT("CantBroadcast", "Cannot convert Attribute '{0}' ({1}) into Attribute '{2}' ({3})."),
+					Context->InputSource.GetDisplayText(),
+					PCG::Private::GetTypeNameText(Context->InputAccessor->GetUnderlyingType()),
+					Context->OutputTarget.GetDisplayText(),
+					PCG::Private::GetTypeNameText(Context->OutputAccessor->GetUnderlyingType())));
+				
 				Outputs.RemoveAt(Outputs.Num() - 1);
 				Context->CurrentInput++;
 				continue;
@@ -289,55 +278,52 @@ bool FPCGAttributeNoiseElement::ExecuteInternal(FPCGContext* InContext) const
 
 			Context->bDataPreparedForCurrentInput = true;
 		}
-		else
-		{
-			OutputPoints = &(const_cast<UPCGPointData*>(CastChecked<UPCGPointData>(Outputs[CurrentInput].Data)))->GetMutablePoints();
-			// Note: for deprecation purposes, we consider that the input here is a spatial data (even though the pin is typed to be a Point Data)
-			// hence the need to call ToPointData, as otherwise a direct case to UPCGPointData would fail.
-			InputPoints = &(CastChecked<const UPCGSpatialData>(Inputs[CurrentInput].Data)->ToPointData(Context))->GetPoints();
-		}
 
-		check(InputPoints && OutputPoints);
+		// Input points are solely used to use the point seed for randomness. Otherwise, it will just use the index of the element as the seed (combined with component + settings seed of course)
+		const UPCGPointData* InputPointData = Cast<UPCGPointData>(Input.Data);
+		const TArray<FPCGPoint>* InputPoints = InputPointData ? &InputPointData->GetPoints() : nullptr;
 
 		// Force clamp on Density
 		const bool bClampResult = Settings->bClampResult || (Context->OutputTarget.GetSelection() == EPCGAttributePropertySelection::PointProperty && Context->OutputTarget.GetPointProperty() == EPCGPointProperties::Density);
 
-		// Dummy Initialize, we already initialized the output points before.
-		auto Initialize = []() {};
-
-		const bool bDone = FPCGAsync::AsyncProcessingOneToOneEx(&Context->AsyncState, OutputPoints->Num(), Initialize, [InputPoints, OutputPoints, Settings, Seed, bClampResult, Context](int32 StartReadIndex, int32 StartWriteIndex) -> int32
+		const bool bDone = PCGMetadataAttribute::CallbackWithRightType(Context->InputAccessor->GetUnderlyingType(), [Settings, Context, Seed, bClampResult, InputPoints](auto&& Dummy) -> bool
 		{
-			PCGMetadataAttribute::CallbackWithRightType(Context->InputAccessor->GetUnderlyingType(), [Settings, StartReadIndex, StartWriteIndex, InputPoints, OutputPoints, Context, Seed, bClampResult](auto&& Dummy)
+			using AttributeType = std::decay_t<decltype(Dummy)>;
+			constexpr int32 ChunkSize = 64;
+			const int32 NumIterations = Context->InputKeys->GetNum();
+
+			// No init
+			auto Initialize = []() {};
+			// It's a 1 for 1 operation, should never move
+			auto MoveDataRange = [](int32, int32, int32) { ensure(false); };
+			// It's finished if we processed all elements.
+			auto Finished = [NumIterations](int32 Count) { ensure(NumIterations == Count); };
+
+			return FPCGAsync::AsyncProcessingRangeEx(&Context->AsyncState, NumIterations, Initialize, [Settings, Seed, bClampResult, Context, ChunkSize, InputPoints, NumIterations](int32 StartReadIndex, int32 StartWriteIndex, int32 Count) -> int32
 			{
-				using AttributeType = std::decay_t<decltype(Dummy)>;
-
-				// Copying the point
-				FPCGPoint& OutPoint = (*OutputPoints)[StartWriteIndex];
-				OutPoint = (*InputPoints)[StartReadIndex];
+				TArray<AttributeType, TInlineAllocator<ChunkSize>> Values;
+				Values.SetNumUninitialized(Count);
 
-				// Reinterpret the buffer to store our temporary value.
-				AttributeType& Value = *(reinterpret_cast<AttributeType*>(Context->TempValuesBuffer.GetData()) + StartReadIndex);
-
-				if (Context->InputAccessor->Get<AttributeType>(Value, StartReadIndex, *Context->Keys))
+				if (Context->InputAccessor->GetRange<AttributeType>(Values, StartReadIndex, *Context->InputKeys))
 				{
-					FRandomStream RandomSource(PCGHelpers::ComputeSeed(Seed, OutPoint.Seed));
-					PCGAttributeNoiseSettings::ProcessNoise(Value, RandomSource, Settings, bClampResult);
+					for (int32 i = 0; i < Count; ++i)
+					{
+						// Use the point seed if we have points, otherwise the index. Don't start at 0 (that's why there is a +1)
+						// Warning: It makes it order independant for points, but order dependant for the rest.
+						const int32 ElementSeed = InputPoints ? (*InputPoints)[StartReadIndex + i].Seed : StartReadIndex + i + 1;
+						FRandomStream RandomSource(PCGHelpers::ComputeSeed(Seed, ElementSeed));
+						PCGAttributeNoiseSettings::ProcessNoise(Values[i], RandomSource, Settings, bClampResult);
+					}
+
+					Context->OutputAccessor->SetRange<AttributeType>(Values, StartWriteIndex, *Context->OutputKeys, EPCGAttributeAccessorFlags::AllowBroadcast | EPCGAttributeAccessorFlags::AllowConstructible);
 				}
-			});
 
-			return true;
-		}, /*bEnableTimeSlicing=*/true);
+				return Count;
+			}, MoveDataRange, Finished, /*bEnableTimeSlicing=*/true, ChunkSize);
+		});
 
 		if (bDone)
 		{
-			PCGMetadataAttribute::CallbackWithRightType(Context->InputAccessor->GetUnderlyingType(), [Context](auto&& Value)
-			{
-				using AttributeType = std::decay_t<decltype(Value)>;
-				int32 NumPoints = Context->Keys->GetNum();
-				TArrayView<AttributeType> Values(reinterpret_cast<AttributeType*>(Context->TempValuesBuffer.GetData()), NumPoints);
-				Context->OutputAccessor->SetRange<AttributeType>(Values, 0, *Context->Keys, EPCGAttributeAccessorFlags::AllowBroadcast);
-			});
-			
 			Context->CurrentInput++;
 			Context->bDataPreparedForCurrentInput = false;
 		}
diff --git a/Engine/Plugins/PCG/Source/PCG/Private/Tests/Elements/PCGAttributeNoiseTest.cpp b/Engine/Plugins/PCG/Source/PCG/Private/Tests/Elements/PCGAttributeNoiseTest.cpp
--- a/Engine/Plugins/PCG/Source/PCG/Private/Tests/Elements/PCGAttributeNoiseTest.cpp
+++ b/Engine/Plugins/PCG/Source/PCG/Private/Tests/Elements/PCGAttributeNoiseTest.cpp
@@ -2,12 +2,12 @@
 
 #include "Tests/PCGTestsCommon.h"
 
-#include "Data/PCGSpatialData.h"
 #include "PCGComponent.h"
-
-#include "Data/PCGPointData.h"
-#include "Elements/PCGAttributeNoise.h"
 #include "PCGContext.h"
+#include "PCGParamData.h"
+#include "Elements/PCGAttributeNoise.h"
+#include "Data/PCGPointData.h"
+#include "Data/PCGSpatialData.h"
 
 IMPLEMENT_CUSTOM_SIMPLE_AUTOMATION_TEST(FPCGAttributeNoiseDensityTest, FPCGTestBaseClass, "Plugins.PCG.AttributeNoise.Density", PCGTestsCommon::TestFlags)
 
@@ -32,6 +32,7 @@ bool FPCGAttributeNoiseDensityTest::RunTest(const FString& Parameters)
 
 	FPCGTaggedData& TaggedData = TestData.InputData.TaggedData.Emplace_GetRef(FPCGTaggedData());
 	TaggedData.Data = PointData;
+	TaggedData.Pin = PCGPinConstants::DefaultInputLabel;
 
 	auto ValidateDensityNoise = [this, &TestData, NoiseElement, Settings](TArray<float> ExpectedOutput) -> bool
 	{
@@ -43,7 +44,7 @@ bool FPCGAttributeNoiseDensityTest::RunTest(const FString& Parameters)
 		const TArray<FPCGTaggedData>& Inputs = Context->InputData.GetInputs();
 		const TArray<FPCGTaggedData>& Outputs = Context->OutputData.GetInputs();
 
-		if (!TestEqual("Valid number of outputs", Inputs.Num(), Outputs.Num()))
+		if (!TestEqual("Valid number of outputs", Outputs.Num(), Inputs.Num()))
 		{
 			return false;
 		}
@@ -231,3 +232,165 @@ bool FPCGAttributeNoiseDensityTest::RunTest(const FString& Parameters)
 
 	return bTestPassed;
 }
+
+IMPLEMENT_CUSTOM_SIMPLE_AUTOMATION_TEST(FPCGAttributeNoiseNotExistingAttributeTest, FPCGTestBaseClass, "Plugins.PCG.AttributeNoise.NotExistingAttribute", PCGTestsCommon::TestFlags)
+
+bool FPCGAttributeNoiseNotExistingAttributeTest::RunTest(const FString& Parameters)
+{
+	PCGTestsCommon::FTestData TestData;
+	UPCGAttributeNoiseSettings* Settings = PCGTestsCommon::GenerateSettings<UPCGAttributeNoiseSettings>(TestData);
+	check(Settings);
+
+	Settings->InputSource.SetAttributeName(TEXT("Hi"));
+
+	UPCGPointData* PointData = PCGTestsCommon::CreateRandomPointData(/*PointCount*/5, /*Seed*/ 42, /*bRandomDensity=*/true);
+
+	FPCGTaggedData& TaggedData = TestData.InputData.TaggedData.Emplace_GetRef();
+	TaggedData.Data = PointData;
+	TaggedData.Pin = PCGPinConstants::DefaultInputLabel;
+
+	FPCGElementPtr NoiseElement = TestData.Settings->GetElement();
+	TUniquePtr<FPCGContext> Context = TestData.InitializeTestContext();
+
+	AddExpectedError(TEXT("Could not find Attribute/Property 'Hi'"));
+
+	while (!NoiseElement->Execute(Context.Get())) {}
+
+	UTEST_EQUAL("No output", Context->OutputData.TaggedData.Num(), 0);
+
+	return true;
+}
+
+IMPLEMENT_CUSTOM_SIMPLE_AUTOMATION_TEST(FPCGAttributeNoiseInvalidTypeTest, FPCGTestBaseClass, "Plugins.PCG.AttributeNoise.InvalidType", PCGTestsCommon::TestFlags)
+
+bool FPCGAttributeNoiseInvalidTypeTest::RunTest(const FString& Parameters)
+{
+	PCGTestsCommon::FTestData TestData;
+	UPCGAttributeNoiseSettings* Settings = PCGTestsCommon::GenerateSettings<UPCGAttributeNoiseSettings>(TestData);
+	check(Settings);
+
+	Settings->InputSource.SetAttributeName(TEXT("MyStr"));
+
+	UPCGPointData* PointData = PCGTestsCommon::CreateRandomPointData(/*PointCount*/5, /*Seed*/ 42, /*bRandomDensity=*/true);
+	FPCGMetadataAttribute<FString>* StrAttribute = PointData->Metadata->CreateAttribute<FString>(TEXT("MyStr"), FString{}, /*bAllowInterpolation=*/false, /*bOverrideParent=*/false);
+	for (FPCGPoint& Point : PointData->GetMutablePoints())
+	{
+		PointData->Metadata->InitializeOnSet(Point.MetadataEntry);
+		StrAttribute->SetValue(Point.MetadataEntry, TEXT("Hey"));
+	}
+
+	FPCGTaggedData& TaggedData = TestData.InputData.TaggedData.Emplace_GetRef();
+	TaggedData.Data = PointData;
+	TaggedData.Pin = PCGPinConstants::DefaultInputLabel;
+
+	FPCGElementPtr NoiseElement = TestData.Settings->GetElement();
+	TUniquePtr<FPCGContext> Context = TestData.InitializeTestContext();
+
+	AddExpectedError(TEXT("Attribute/Property 'MyStr' is not a numerical type, we can't apply noise to it."));
+
+	while (!NoiseElement->Execute(Context.Get())) {}
+
+	UTEST_EQUAL("No output", Context->OutputData.TaggedData.Num(), 0);
+
+	return true;
+}
+
+IMPLEMENT_CUSTOM_SIMPLE_AUTOMATION_TEST(FPCGAttributeNoiseAttributeSetTest, FPCGTestBaseClass, "Plugins.PCG.AttributeNoise.AttributeSet", PCGTestsCommon::TestFlags)
+
+bool FPCGAttributeNoiseAttributeSetTest::RunTest(const FString& Parameters)
+{
+	PCGTestsCommon::FTestData TestData(42);
+	UPCGAttributeNoiseSettings* Settings = PCGTestsCommon::GenerateSettings<UPCGAttributeNoiseSettings>(TestData);
+	check(Settings);
+
+	const FName InputAttrName = TEXT("Attr");
+	const FName OutputAttrName = TEXT("OutAttr");
+	constexpr int32 NbElements = 5;
+
+	Settings->InputSource.SetAttributeName(InputAttrName);
+	Settings->OutputTarget.SetAttributeName(OutputAttrName);
+
+	UPCGParamData* ParamData = NewObject<UPCGParamData>();
+	FPCGMetadataAttribute<float>* FloatAttribute = ParamData->Metadata->CreateAttribute<float>(InputAttrName, 0.f, /*bAllowInterpolation=*/true, /*bOverrideParent=*/false);
+	for (int32 i = 0; i < NbElements; ++i)
+	{
+		FloatAttribute->SetValue(ParamData->Metadata->AddEntry(), (float)i / NbElements);
+	}
+
+	FPCGTaggedData& TaggedData = TestData.InputData.TaggedData.Emplace_GetRef();
+	TaggedData.Data = ParamData;
+	TaggedData.Pin = PCGPinConstants::DefaultInputLabel;
+
+	FPCGElementPtr NoiseElement = TestData.Settings->GetElement();
+	TUniquePtr<FPCGContext> Context = TestData.InitializeTestContext();
+
+	while (!NoiseElement->Execute(Context.Get())) {}
+
+	UTEST_EQUAL("1 output", Context->OutputData.TaggedData.Num(), 1);
+
+	const UPCGParamData* OutParamData = Cast<const UPCGParamData>(Context->OutputData.TaggedData[0].Data);
+	UTEST_NOT_NULL("Output is param", OutParamData);
+
+	const FPCGMetadataAttribute<float>* OutAttribute = OutParamData->Metadata->GetConstTypedAttribute<float>(OutputAttrName);
+	UTEST_NOT_NULL("Output attrbiute exists", OutAttribute);
+
+	// Taken from execution
+	const float ExpectedValues[NbElements] = { 0.382462f, 0.193192f, 0.134616f, 0.782528f, 0.084569f };
+	for (int32 i = 0; i < NbElements; ++i)
+	{
+		float Value = OutAttribute->GetValueFromItemKey(PCGMetadataEntryKey(i));
+		UTEST_EQUAL(FString::Printf(TEXT("Value %d is noised as expected"), i), Value, ExpectedValues[i]);
+	}
+
+	return true;
+}
+
+IMPLEMENT_CUSTOM_SIMPLE_AUTOMATION_TEST(FPCGAttributeNoiseOutputAttributeExistingTest, FPCGTestBaseClass, "Plugins.PCG.AttributeNoise.OutputAttributeExisting", PCGTestsCommon::TestFlags)
+
+bool FPCGAttributeNoiseOutputAttributeExistingTest::RunTest(const FString& Parameters)
+{
+	PCGTestsCommon::FTestData TestData(42);
+	UPCGAttributeNoiseSettings* Settings = PCGTestsCommon::GenerateSettings<UPCGAttributeNoiseSettings>(TestData);
+	check(Settings);
+
+	const FName InputAttrName = TEXT("Attr");
+	constexpr int32 NbElements = 5;
+
+	Settings->InputSource.Update(TEXT("Attr.X"));
+	Settings->OutputTarget.Update(TEXT("Attr.Z"));
+
+	UPCGParamData* ParamData = NewObject<UPCGParamData>();
+	FPCGMetadataAttribute<FVector>* VectorAttribute = ParamData->Metadata->CreateAttribute<FVector>(InputAttrName, FVector::ZeroVector, /*bAllowInterpolation=*/true, /*bOverrideParent=*/false);
+	for (int32 i = 0; i < NbElements; ++i)
+	{
+		VectorAttribute->SetValue(ParamData->Metadata->AddEntry(), FVector((double)i / NbElements));
+	}
+
+	FPCGTaggedData& TaggedData = TestData.InputData.TaggedData.Emplace_GetRef();
+	TaggedData.Data = ParamData;
+	TaggedData.Pin = PCGPinConstants::DefaultInputLabel;
+
+	FPCGElementPtr NoiseElement = TestData.Settings->GetElement();
+	TUniquePtr<FPCGContext> Context = TestData.InitializeTestContext();
+
+	while (!NoiseElement->Execute(Context.Get())) {}
+
+	UTEST_EQUAL("1 output", Context->OutputData.TaggedData.Num(), 1);
+
+	const UPCGParamData* OutParamData = Cast<const UPCGParamData>(Context->OutputData.TaggedData[0].Data);
+	UTEST_NOT_NULL("Output is param", OutParamData);
+
+	const FPCGMetadataAttribute<FVector>* OutAttribute = OutParamData->Metadata->GetConstTypedAttribute<FVector>(InputAttrName);
+	UTEST_NOT_NULL("Output attrbiute exists", OutAttribute);
+
+	// Taken from execution
+	const double ExpectedValues[NbElements] = { 0.382462, 0.193192, 0.134616, 0.782528, 0.084569 };
+	for (int32 i = 0; i < NbElements; ++i)
+	{
+		FVector Value = OutAttribute->GetValueFromItemKey(PCGMetadataEntryKey(i));
+		UTEST_EQUAL(FString::Printf(TEXT("Value %d for X component is the same"), i), Value.X, (double)i / NbElements);
+		UTEST_EQUAL(FString::Printf(TEXT("Value %d for Z component is noised as expected"), i), Value.Z, ExpectedValues[i]);
+	}
+
+	return true;
+}
\ No newline at end of file
diff --git a/Engine/Plugins/PCG/Source/PCG/Public/Elements/PCGAttributeNoise.h b/Engine/Plugins/PCG/Source/PCG/Public/Elements/PCGAttributeNoise.h
--- a/Engine/Plugins/PCG/Source/PCG/Public/Elements/PCGAttributeNoise.h
+++ b/Engine/Plugins/PCG/Source/PCG/Public/Elements/PCGAttributeNoise.h
@@ -47,11 +47,12 @@ public:
 	virtual TArray<FText> GetNodeTitleAliases() const { return { NSLOCTEXT("PCGAttributeNoiseSettings", "DensityNoiseAlias", "Density Noise") }; }
 	virtual EPCGSettingsType GetType() const override { return EPCGSettingsType::Metadata; }
 #endif
+	virtual bool HasDynamicPins() const override { return true; }
 	virtual bool IsInputPinRequiredByExecution(const UPCGPin* InPin) const override { return !InPin->Properties.bAdvancedPin; }
 
 protected:
-	virtual TArray<FPCGPinProperties> InputPinProperties() const override { return Super::DefaultPointInputPinProperties(); }
-	virtual TArray<FPCGPinProperties> OutputPinProperties() const override { return Super::DefaultPointOutputPinProperties(); }
+	virtual TArray<FPCGPinProperties> InputPinProperties() const override;
+	virtual TArray<FPCGPinProperties> OutputPinProperties() const override;
 	virtual FPCGElementPtr CreateElement() const override;
 	//~End UPCGSettings interface
 
@@ -107,9 +108,8 @@ struct FPCGAttributeNoiseContext : public FPCGContext
 	FPCGAttributePropertyOutputSelector OutputTarget;
 	TUniquePtr<const IPCGAttributeAccessor> InputAccessor;
 	TUniquePtr<IPCGAttributeAccessor> OutputAccessor;
-	TUniquePtr<IPCGAttributeAccessorKeys> Keys;
-
-	TArray<uint8> TempValuesBuffer;
+	TUniquePtr<const IPCGAttributeAccessorKeys> InputKeys;
+	TUniquePtr<IPCGAttributeAccessorKeys> OutputKeys;
 };
 
 class FPCGAttributeNoiseElement : public IPCGElement
diff --git a/Engine/Plugins/PCG/Source/PCG/Public/Helpers/PCGAsync.h b/Engine/Plugins/PCG/Source/PCG/Public/Helpers/PCGAsync.h
--- a/Engine/Plugins/PCG/Source/PCG/Public/Helpers/PCGAsync.h
+++ b/Engine/Plugins/PCG/Source/PCG/Public/Helpers/PCGAsync.h
@@ -329,7 +329,7 @@ namespace FPCGAsync
 	* @param AsyncState - The context containing the information about how many tasks we can launch, async read/write index for the current job and a function to know if we need to stop processing.
 	* @param NumIterations - The number of calls that will be done to the provided function, also an upper bound on the number of data generated.
 	* @param Initialize - Signature: void(). A function that will be called once on the first timeslice, where you can reserve data for processing
-	* @param ProcessRange - Signature: bool(int32 StartReadIndex, int32 StartWriteIndex, int32 Count). A function that processes a range of values and returns the number of written values 
+	* @param ProcessRange - Signature: int32(int32 StartReadIndex, int32 StartWriteIndex, int32 Count). A function that processes a range of values and returns the number of written values 
 	* @param MoveDataRange - Signature: void(int32 ReadIndex, int32 WriteIndex, int32 Count). If the processing filters points, this will be used to move elements in chunk from one range to another 
 	* @param Finished - Signature: void(int32 Count). Called once on finished, and tells you the total count of points written.
 	* @param bEnableTimeSlicing - If false, we will not stop until all the processing is done.
