[PCG] Fix deleted files not submitted from world partition builder

diff --git a/Engine/Plugins/PCG/Source/PCGEditor/Private/WorldPartitionBuilder/PCGWorldPartitionBuilder.cpp b/Engine/Plugins/PCG/Source/PCGEditor/Private/WorldPartitionBuilder/PCGWorldPartitionBuilder.cpp
--- a/Engine/Plugins/PCG/Source/PCGEditor/Private/WorldPartitionBuilder/PCGWorldPartitionBuilder.cpp
+++ b/Engine/Plugins/PCG/Source/PCGEditor/Private/WorldPartitionBuilder/PCGWorldPartitionBuilder.cpp
@@ -225,6 +225,7 @@ bool UPCGWorldPartitionBuilder::PostRun(UWorld* World, FPackageSourceControlHelp
 
 	TArray<FString> FilesToSubmit;
 	FilesToSubmit.Append(SourceControlHelpers::PackageFilenames(DirtyPackages));
+	FilesToSubmit.Append(SourceControlHelpers::PackageFilenames(PackagesToDelete));
 
 	const FString ChangeDescription = FString::Printf(TEXT("Generated PCG components for world '%s'"), *World->GetName());
 	return OnFilesModified(FilesToSubmit, ChangeDescription);
