Clarified usage of preview/load-as-preview modes on the pcg component

#rnx

diff --git a/Engine/Plugins/PCG/Source/PCG/Public/PCGCommon.h b/Engine/Plugins/PCG/Source/PCG/Public/PCGCommon.h
--- a/Engine/Plugins/PCG/Source/PCG/Public/PCGCommon.h
+++ b/Engine/Plugins/PCG/Source/PCG/Public/PCGCommon.h
@@ -247,7 +247,7 @@ UENUM()
 enum class EPCGEditorDirtyMode : uint8
 {
 	Normal UMETA(Tooltip="Normal editing mode where generation changes (generation, cleanup) dirty the component and its resources."),
-	Preview UMETA(Tooltip="Editing mode where generation changes (generation, cleanup, resources) on the component will not trigger any dirty state, but will also not save any of the generated resources."),
+	Preview UMETA(Tooltip="Editing mode where generation changes (generation, cleanup, resources) on the component will not trigger any dirty state, but will also not save any of the generated resources. Also represents the state after loading from the Load as Preview edit mode, where this will hold the last saved generation until a new generation is triggered."),
 	LoadAsPreview UMETA(Tooltip="Acts as the normal editing mode until the next load of the component, at which state it acts as-if-transient, namely that any further generation changes will not dirty the component.")
 };
 
diff --git a/Engine/Plugins/PCG/Source/PCG/Public/PCGComponent.h b/Engine/Plugins/PCG/Source/PCG/Public/PCGComponent.h
--- a/Engine/Plugins/PCG/Source/PCG/Public/PCGComponent.h
+++ b/Engine/Plugins/PCG/Source/PCG/Public/PCGComponent.h
@@ -405,7 +405,7 @@ protected:
 	UPROPERTY(Transient, EditAnywhere, Category = "Editing Settings", meta = (DisplayName = "Editing Mode", EditCondition = "!bIsComponentLocal && GenerationTrigger != EPCGComponentGenerationTrigger::GenerateAtRuntime", DisplayPriority = 300))
 	EPCGEditorDirtyMode CurrentEditingMode = EPCGEditorDirtyMode::Normal;
 
-	UPROPERTY(VisibleAnywhere, Category = Debug, meta = (NoResetToDefault))
+	UPROPERTY(VisibleAnywhere, Category = "Editing Settings", meta = (NoResetToDefault, DisplayPriority = 300))
 	EPCGEditorDirtyMode SerializedEditingMode = EPCGEditorDirtyMode::Normal;
 
 	/** Used to store the CurrentEditingMode when it is forcefully changed by another system, such as runtime generation. */
