[PCG] Dependency tracking for PCG generation improvements and fixes

* Recent change to trigger a generation for all depending PCG components that has no data was not taken into account when gathering dependencies and so were sitting there forever, preventing refresh.
* Dependency gathering on PCG start generating was gathering all dependencies, while it only needs to gather PCG data related dependencies.
* Some changes are already taken care of in the PCG Component, and should not trigger a refresh (regression with the tracking refactor)

#rnx

diff --git a/Engine/Plugins/PCG/Source/PCG/Private/PCGActorAndComponentMapping.cpp b/Engine/Plugins/PCG/Source/PCG/Private/PCGActorAndComponentMapping.cpp
--- a/Engine/Plugins/PCG/Source/PCG/Private/PCGActorAndComponentMapping.cpp
+++ b/Engine/Plugins/PCG/Source/PCG/Private/PCGActorAndComponentMapping.cpp
@@ -1608,15 +1608,26 @@ void FPCGActorAndComponentMapping::OnObjectPropertyChanged(UObject* InObject, FP
 
 void FPCGActorAndComponentMapping::OnObjectChanged(UObject* InObject, const FActorPreviousData* InPreviousData, const UObject* InOriginatingChangeObject, int32 LevelInstanceDepth, bool bNoRefreshOwner)
 {
-	// Nothing to do if we track nothing
-	if (CulledTrackedKeysToComponentsMap.IsEmpty() && AlwaysTrackedKeysToComponentsMap.IsEmpty())
+	TRACE_CPUPROFILER_EVENT_SCOPE(FPCGActorAndComponentMapping::OnObjectChanged);
+
+	// Nothing to do if we track nothing or there is no object.
+	if (!InObject || (CulledTrackedKeysToComponentsMap.IsEmpty() && AlwaysTrackedKeysToComponentsMap.IsEmpty()))
 	{
 		return;
 	}
 
-	TRACE_CPUPROFILER_EVENT_SCOPE(FPCGActorAndComponentMapping::OnObjectChanged);
+	// Don't react to what the PCG Component is already reacting to (itself and graphs).
+	static const TArray<const UClass*> ExcludedClasses =
+	{
+		UPCGComponent::StaticClass(),
+		UPCGGraphInterface::StaticClass()
+	};
+
+	if (Algo::AnyOf(ExcludedClasses, [InObject](const UClass* Class) -> bool { return InObject->IsA(Class); }))
+	{
+		return;
+	}
 
-	check(InObject);
 	AActor* Actor = Cast<AActor>(InObject);
 
 	ensure(!PCGSubsystem || !Actor || Actor->GetWorld() == PCGSubsystem->GetWorld());
@@ -1931,13 +1942,21 @@ void FPCGActorAndComponentMapping::OnPCGGraphStartsGenerating(UPCGComponent* InC
 	// When a graph starts generating, look for component that depends on it and keep a count.
 	// When this graph will be done generating, we'll trigger a dependency generation. But if multiple graphs are generated at the same time and 
 	// they all contribute to the same dependency, we'll trigger the dependency only when all the graphs are done generating
+	// We only need to gather components that depends on PCG Data, because other dependency changes will be caught by other engine callbacks (such as OnObjectPropertyChanged)
 	TSet<UPCGComponent*> TrackedComponents;
 	TSet<FName> RemovedTags;
 
 	const FBox ComponentBounds = InComponent->GetGridBounds();
 
-	for (auto& It : CulledTrackedKeysToComponentsMap)
+	auto IsKeyTrackingPCGData = [](const FPCGSelectionKey& InKey) { return InKey.OptionalExtraDependency && InKey.OptionalExtraDependency->IsChildOf(UPCGComponent::StaticClass()); };
+
+	for (TPair<FPCGSelectionKey, TSet<UPCGComponent*>>& It : CulledTrackedKeysToComponentsMap)
 	{
+		if (!IsKeyTrackingPCGData(It.Key))
+		{
+			continue;
+		}
+
 		TSet<UPCGComponent*> TempTrackedComponents;
 		It.Key.IsMatching(InComponent->GetOwner(), RemovedTags, It.Value, &TempTrackedComponents);
 
@@ -1951,14 +1970,26 @@ void FPCGActorAndComponentMapping::OnPCGGraphStartsGenerating(UPCGComponent* InC
 		}
 	}
 
-	for (auto& It : AlwaysTrackedKeysToComponentsMap)
+	for (TPair<FPCGSelectionKey, TSet<UPCGComponent*>>& It : AlwaysTrackedKeysToComponentsMap)
 	{
+		if (!IsKeyTrackingPCGData(It.Key))
+		{
+			continue;
+		}
+
 		It.Key.IsMatching(InComponent->GetOwner(), RemovedTags, It.Value, &TrackedComponents);
 	}
 
 	for (UPCGComponent* Component : TrackedComponents)
 	{
-		if (Component == InComponent)
+		// Don't have a dependency on itself.
+		if (!Component || Component == InComponent)
+		{
+			continue;
+		}
+
+		// If the tracked component is currently ignoring the refresh from InComponent, don't add to the dependencies
+		if (Component->IsIgnoringChangeOrigin(InComponent->GetOwner()))
 		{
 			continue;
 		}
