[PCG] Demote warning on graph cache resize mechanism

No longer necessary as the function works well, and this warning just noises up build logs

diff --git a/Engine/Plugins/PCG/Source/PCG/Private/Graph/PCGGraphCache.cpp b/Engine/Plugins/PCG/Source/PCG/Private/Graph/PCGGraphCache.cpp
--- a/Engine/Plugins/PCG/Source/PCG/Private/Graph/PCGGraphCache.cpp
+++ b/Engine/Plugins/PCG/Source/PCG/Private/Graph/PCGGraphCache.cpp
@@ -220,7 +220,7 @@ void FPCGGraphCache::GrowCache_Unsafe()
 
 	// Resize and flush
 	CacheData.Empty(CacheData.Num() * 2);
-	UE_LOG(LogPCG, Warning, TEXT("Graph cache doubled in capacity to %d entries."), CacheData.Max());
+	UE_LOG(LogPCG, Log, TEXT("Graph cache doubled in capacity to %d entries."), CacheData.Max());
 
 	// Copy back. Restore temporal order.
 	for (TLruCache<FPCGCacheEntryKey, FPCGDataCollection>::TIterator It(CacheDataCopy); It; ++It)
