[PCG] Change bAdvancedPin to a PinStatus enum, to enable subgraphs/BP to have required pins

* With the full deprecation

#rnx

diff --git a/Engine/Plugins/PCG/Source/PCG/Private/Elements/ControlFlow/PCGMultiSelect.cpp b/Engine/Plugins/PCG/Source/PCG/Private/Elements/ControlFlow/PCGMultiSelect.cpp
--- a/Engine/Plugins/PCG/Source/PCG/Private/Elements/ControlFlow/PCGMultiSelect.cpp
+++ b/Engine/Plugins/PCG/Source/PCG/Private/Elements/ControlFlow/PCGMultiSelect.cpp
@@ -56,7 +56,7 @@ EPCGDataType UPCGMultiSelectSettings::GetCurrentPinTypes(const UPCGPin* Pin) con
 
 		for (int Index = 0; Index < InputPins.Num(); ++Index)
 		{
-			if (!InputPins[Index]->Properties.bAdvancedPin)
+			if (!InputPins[Index]->Properties.IsAdvancedPin())
 			{
 				InputTypeUnion |= GetTypeUnionOfIncidentEdges(InputPins[Index]->Properties.Label);
 			}
diff --git a/Engine/Plugins/PCG/Source/PCG/Private/Elements/ControlFlow/PCGSwitch.cpp b/Engine/Plugins/PCG/Source/PCG/Private/Elements/ControlFlow/PCGSwitch.cpp
--- a/Engine/Plugins/PCG/Source/PCG/Private/Elements/ControlFlow/PCGSwitch.cpp
+++ b/Engine/Plugins/PCG/Source/PCG/Private/Elements/ControlFlow/PCGSwitch.cpp
@@ -140,12 +140,14 @@ EPCGChangeType UPCGSwitchSettings::GetChangeTypeForProperty(const FName& InPrope
 TArray<FPCGPinProperties> UPCGSwitchSettings::InputPinProperties() const
 {
 	TArray<FPCGPinProperties> PinProperties;
-	PinProperties.Emplace(PCGPinConstants::DefaultInputLabel,
+	FPCGPinProperties& InputPinProperty = PinProperties.Emplace_GetRef(PCGPinConstants::DefaultInputLabel,
 		EPCGDataType::Any,
 		/*bInAllowMultipleConnections=*/true,
 		/*bAllowMultipleData=*/true,
 		LOCTEXT("OutputPinTooltip", "All input will be forwarded directly to the selected output pin."));
 
+	InputPinProperty.SetRequiredPin();
+
 	return PinProperties;
 }
 
diff --git a/Engine/Plugins/PCG/Source/PCG/Private/Elements/Metadata/PCGMetadataElement.cpp b/Engine/Plugins/PCG/Source/PCG/Private/Elements/Metadata/PCGMetadataElement.cpp
--- a/Engine/Plugins/PCG/Source/PCG/Private/Elements/Metadata/PCGMetadataElement.cpp
+++ b/Engine/Plugins/PCG/Source/PCG/Private/Elements/Metadata/PCGMetadataElement.cpp
@@ -93,7 +93,9 @@ void UPCGMetadataOperationSettings::ApplyDeprecation(UPCGNode* InOutNode)
 TArray<FPCGPinProperties> UPCGMetadataOperationSettings::InputPinProperties() const
 {
 	TArray<FPCGPinProperties> Properties;
-	Properties.Emplace(PCGPinConstants::DefaultInputLabel, EPCGDataType::Point);
+	FPCGPinProperties& InputPinProperty = Properties.Emplace_GetRef(PCGPinConstants::DefaultInputLabel, EPCGDataType::Point);
+	InputPinProperty.SetRequiredPin();
+
 	Properties.Emplace(PCGMetadataOperationSettings::AttributeLabel, EPCGDataType::Param, /*bInAllowMultipleConnections=*/ false, /*bAllowMultipleData=*/ false, PCGMetadataOperationSettings::AttributeTooltip);
 	return Properties;
 }
diff --git a/Engine/Plugins/PCG/Source/PCG/Private/Elements/Metadata/PCGMetadataOpElementBase.cpp b/Engine/Plugins/PCG/Source/PCG/Private/Elements/Metadata/PCGMetadataOpElementBase.cpp
--- a/Engine/Plugins/PCG/Source/PCG/Private/Elements/Metadata/PCGMetadataOpElementBase.cpp
+++ b/Engine/Plugins/PCG/Source/PCG/Private/Elements/Metadata/PCGMetadataOpElementBase.cpp
@@ -197,7 +197,13 @@ TArray<FPCGPinProperties> UPCGMetadataSettingsBase::InputPinProperties() const
 		const FName PinLabel = GetInputPinLabel(InputPinIndex);
 		if (PinLabel != NAME_None)
 		{
-			PinProperties.Emplace(PinLabel, EPCGDataType::Any);
+			FPCGPinProperties& PinProperty = PinProperties.Emplace_GetRef(PinLabel, EPCGDataType::Any);
+
+			const bool bSupportDefaultValue = DoesInputSupportDefaultValue(InputPinIndex);
+			if (!bSupportDefaultValue)
+			{
+				PinProperty.SetRequiredPin();
+			}
 
 #if WITH_EDITOR
 			TArray<FText> AllTooltips;
@@ -222,14 +228,14 @@ TArray<FPCGPinProperties> UPCGMetadataSettingsBase::InputPinProperties() const
 				AllTooltips.Add(LOCTEXT("PinTooltipForwardInput", "This input will be forwarded to the output."));
 			}
 
-			if (DoesInputSupportDefaultValue(InputPinIndex))
+			if (bSupportDefaultValue)
 			{
 				AllTooltips.Add(FText::Format(LOCTEXT("PinTooltipDefaultValue", "Pin is optional, will use default value if not connected ({0})"), FText::FromString(GetDefaultValueString(InputPinIndex))));
 			}
 
 			if (!AllTooltips.IsEmpty())
 			{
-				PinProperties.Last().Tooltip = FText::Join(FText::FromString(TEXT("\n")), AllTooltips);
+				PinProperty.Tooltip = FText::Join(FText::FromString(TEXT("\n")), AllTooltips);
 			}
 #endif // WITH_EDITOR
 		}
diff --git a/Engine/Plugins/PCG/Source/PCG/Private/Elements/Metadata/PCGMetadataPartition.cpp b/Engine/Plugins/PCG/Source/PCG/Private/Elements/Metadata/PCGMetadataPartition.cpp
--- a/Engine/Plugins/PCG/Source/PCG/Private/Elements/Metadata/PCGMetadataPartition.cpp
+++ b/Engine/Plugins/PCG/Source/PCG/Private/Elements/Metadata/PCGMetadataPartition.cpp
@@ -12,14 +12,6 @@
 
 #define LOCTEXT_NAMESPACE "PCGMetadataPartitionElement"
 
-TArray<FPCGPinProperties> UPCGMetadataPartitionSettings::InputPinProperties() const
-{
-	TArray<FPCGPinProperties> PinProperties;
-	PinProperties.Emplace(PCGPinConstants::DefaultInputLabel, EPCGDataType::Any);
-
-	return PinProperties;
-}
-
 TArray<FPCGPinProperties> UPCGMetadataPartitionSettings::OutputPinProperties() const
 {
 	TArray<FPCGPinProperties> PinProperties;
diff --git a/Engine/Plugins/PCG/Source/PCG/Private/Elements/Metadata/PCGMetadataRenameElement.cpp b/Engine/Plugins/PCG/Source/PCG/Private/Elements/Metadata/PCGMetadataRenameElement.cpp
--- a/Engine/Plugins/PCG/Source/PCG/Private/Elements/Metadata/PCGMetadataRenameElement.cpp
+++ b/Engine/Plugins/PCG/Source/PCG/Private/Elements/Metadata/PCGMetadataRenameElement.cpp
@@ -26,14 +26,6 @@ FString UPCGMetadataRenameSettings::GetAdditionalTitleInformation() const
 	}
 }
 
-TArray<FPCGPinProperties> UPCGMetadataRenameSettings::InputPinProperties() const
-{
-	TArray<FPCGPinProperties> PinProperties;
-	PinProperties.Emplace(PCGPinConstants::DefaultInputLabel, EPCGDataType::Any);
-
-	return PinProperties;
-}
-
 TArray<FPCGPinProperties> UPCGMetadataRenameSettings::OutputPinProperties() const
 {
 	TArray<FPCGPinProperties> PinProperties;
diff --git a/Engine/Plugins/PCG/Source/PCG/Private/Elements/PCGAddTag.cpp b/Engine/Plugins/PCG/Source/PCG/Private/Elements/PCGAddTag.cpp
--- a/Engine/Plugins/PCG/Source/PCG/Private/Elements/PCGAddTag.cpp
+++ b/Engine/Plugins/PCG/Source/PCG/Private/Elements/PCGAddTag.cpp
@@ -8,14 +8,6 @@
 
 #define LOCTEXT_NAMESPACE "PCGAddTagElement"
 
-TArray<FPCGPinProperties> UPCGAddTagSettings::InputPinProperties() const
-{
-	TArray<FPCGPinProperties> PinProperties;
-	PinProperties.Emplace(PCGPinConstants::DefaultInputLabel, EPCGDataType::Any);
-
-	return PinProperties;
-}
-
 TArray<FPCGPinProperties> UPCGAddTagSettings::OutputPinProperties() const
 {
 	TArray<FPCGPinProperties> PinProperties;
diff --git a/Engine/Plugins/PCG/Source/PCG/Private/Elements/PCGAttributeFilter.cpp b/Engine/Plugins/PCG/Source/PCG/Private/Elements/PCGAttributeFilter.cpp
--- a/Engine/Plugins/PCG/Source/PCG/Private/Elements/PCGAttributeFilter.cpp
+++ b/Engine/Plugins/PCG/Source/PCG/Private/Elements/PCGAttributeFilter.cpp
@@ -274,7 +274,8 @@ EPCGDataType UPCGAttributeFilteringSettings::GetCurrentPinTypes(const UPCGPin* I
 TArray<FPCGPinProperties> UPCGAttributeFilteringSettings::InputPinProperties() const
 {
 	TArray<FPCGPinProperties> PinProperties;
-	PinProperties.Emplace(PCGAttributeFilterConstants::DataToFilterLabel, EPCGDataType::Any);
+	FPCGPinProperties& DataToFilterPinProperty = PinProperties.Emplace_GetRef(PCGAttributeFilterConstants::DataToFilterLabel, EPCGDataType::Any);
+	DataToFilterPinProperty.SetRequiredPin();
 
 #if WITH_EDITOR
 	PinProperties.Last().Tooltip = PCGAttributeFilterConstants::InputPinTooltip;
@@ -403,7 +404,8 @@ EPCGDataType UPCGAttributeFilteringRangeSettings::GetCurrentPinTypes(const UPCGP
 TArray<FPCGPinProperties> UPCGAttributeFilteringRangeSettings::InputPinProperties() const
 {
 	TArray<FPCGPinProperties> PinProperties;
-	PinProperties.Emplace(PCGAttributeFilterConstants::DataToFilterLabel, EPCGDataType::Any);
+	FPCGPinProperties& DataToFilterPinProperty = PinProperties.Emplace_GetRef(PCGAttributeFilterConstants::DataToFilterLabel, EPCGDataType::Any);
+	DataToFilterPinProperty.SetRequiredPin();
 
 #if WITH_EDITOR
 	PinProperties.Last().Tooltip = PCGAttributeFilterConstants::InputPinTooltip;
diff --git a/Engine/Plugins/PCG/Source/PCG/Private/Elements/PCGAttributeGetFromIndexElement.cpp b/Engine/Plugins/PCG/Source/PCG/Private/Elements/PCGAttributeGetFromIndexElement.cpp
--- a/Engine/Plugins/PCG/Source/PCG/Private/Elements/PCGAttributeGetFromIndexElement.cpp
+++ b/Engine/Plugins/PCG/Source/PCG/Private/Elements/PCGAttributeGetFromIndexElement.cpp
@@ -32,7 +32,8 @@ FText UPCGAttributeGetFromIndexSettings::GetNodeTooltipText() const
 TArray<FPCGPinProperties> UPCGAttributeGetFromIndexSettings::InputPinProperties() const
 {
 	TArray<FPCGPinProperties> PinProperties;
-	PinProperties.Emplace(PCGPinConstants::DefaultInputLabel, EPCGDataType::Param);
+	FPCGPinProperties& InputPinProperty = PinProperties.Emplace_GetRef(PCGPinConstants::DefaultInputLabel, EPCGDataType::Param);
+	InputPinProperty.SetRequiredPin();
 
 	return PinProperties;
 }
diff --git a/Engine/Plugins/PCG/Source/PCG/Private/Elements/PCGAttributeGetFromPointIndexElement.cpp b/Engine/Plugins/PCG/Source/PCG/Private/Elements/PCGAttributeGetFromPointIndexElement.cpp
--- a/Engine/Plugins/PCG/Source/PCG/Private/Elements/PCGAttributeGetFromPointIndexElement.cpp
+++ b/Engine/Plugins/PCG/Source/PCG/Private/Elements/PCGAttributeGetFromPointIndexElement.cpp
@@ -56,7 +56,8 @@ void UPCGAttributeGetFromPointIndexSettings::PostLoad()
 TArray<FPCGPinProperties> UPCGAttributeGetFromPointIndexSettings::InputPinProperties() const
 {
 	TArray<FPCGPinProperties> PinProperties;
-	PinProperties.Emplace(PCGPinConstants::DefaultInputLabel, EPCGDataType::Point);
+	FPCGPinProperties& InputPinProperty = PinProperties.Emplace_GetRef(PCGPinConstants::DefaultInputLabel, EPCGDataType::Point);
+	InputPinProperty.SetRequiredPin();
 
 	return PinProperties;
 }
diff --git a/Engine/Plugins/PCG/Source/PCG/Private/Elements/PCGAttributeNoise.cpp b/Engine/Plugins/PCG/Source/PCG/Private/Elements/PCGAttributeNoise.cpp
--- a/Engine/Plugins/PCG/Source/PCG/Private/Elements/PCGAttributeNoise.cpp
+++ b/Engine/Plugins/PCG/Source/PCG/Private/Elements/PCGAttributeNoise.cpp
@@ -156,8 +156,8 @@ FPCGElementPtr UPCGAttributeNoiseSettings::CreateElement() const
 TArray<FPCGPinProperties> UPCGAttributeNoiseSettings::InputPinProperties() const
 {
 	TArray<FPCGPinProperties> PinProperties;
-	PinProperties.Emplace(PCGPinConstants::DefaultInputLabel, EPCGDataType::Point | EPCGDataType::Param);
-
+	FPCGPinProperties& InputPinProperty = PinProperties.Emplace_GetRef(PCGPinConstants::DefaultInputLabel, EPCGDataType::Point | EPCGDataType::Param);
+	InputPinProperty.SetRequiredPin();
 	return PinProperties;
 }
 
diff --git a/Engine/Plugins/PCG/Source/PCG/Private/Elements/PCGAttributeSelectElement.cpp b/Engine/Plugins/PCG/Source/PCG/Private/Elements/PCGAttributeSelectElement.cpp
--- a/Engine/Plugins/PCG/Source/PCG/Private/Elements/PCGAttributeSelectElement.cpp
+++ b/Engine/Plugins/PCG/Source/PCG/Private/Elements/PCGAttributeSelectElement.cpp
@@ -247,7 +247,8 @@ FString UPCGAttributeSelectSettings::GetAdditionalTitleInformation() const
 TArray<FPCGPinProperties> UPCGAttributeSelectSettings::InputPinProperties() const
 {
 	TArray<FPCGPinProperties> PinProperties;
-	PinProperties.Emplace(PCGPinConstants::DefaultInputLabel, EPCGDataType::Spatial, /*bInAllowMultipleConnections=*/ false);
+	FPCGPinProperties& InputPinProperty = PinProperties.Emplace_GetRef(PCGPinConstants::DefaultInputLabel, EPCGDataType::Spatial, /*bInAllowMultipleConnections=*/ false);
+	InputPinProperty.SetRequiredPin();
 
 	return PinProperties;
 }
diff --git a/Engine/Plugins/PCG/Source/PCG/Private/Elements/PCGAttributeTransferElement.cpp b/Engine/Plugins/PCG/Source/PCG/Private/Elements/PCGAttributeTransferElement.cpp
--- a/Engine/Plugins/PCG/Source/PCG/Private/Elements/PCGAttributeTransferElement.cpp
+++ b/Engine/Plugins/PCG/Source/PCG/Private/Elements/PCGAttributeTransferElement.cpp
@@ -103,7 +103,9 @@ FString UPCGAttributeTransferSettings::GetAdditionalTitleInformation() const
 TArray<FPCGPinProperties> UPCGAttributeTransferSettings::InputPinProperties() const
 {
 	TArray<FPCGPinProperties> PinProperties;
-	PinProperties.Emplace(PCGAttributeTransferConstants::TargetLabel, EPCGDataType::Spatial, /*bInAllowMultipleConnections=*/ false);
+	FPCGPinProperties& TargetPinProperty = PinProperties.Emplace_GetRef(PCGAttributeTransferConstants::TargetLabel, EPCGDataType::Spatial, /*bInAllowMultipleConnections=*/ false);
+	TargetPinProperty.SetRequiredPin();
+
 	PinProperties.Emplace(PCGAttributeTransferConstants::SourceLabel, EPCGDataType::Spatial, /*bInAllowMultipleConnections=*/ false);
 
 	return PinProperties;
diff --git a/Engine/Plugins/PCG/Source/PCG/Private/Elements/PCGCollapseElement.cpp b/Engine/Plugins/PCG/Source/PCG/Private/Elements/PCGCollapseElement.cpp
--- a/Engine/Plugins/PCG/Source/PCG/Private/Elements/PCGCollapseElement.cpp
+++ b/Engine/Plugins/PCG/Source/PCG/Private/Elements/PCGCollapseElement.cpp
@@ -20,7 +20,8 @@ bool UPCGCollapseSettings::GetCompactNodeIcon(FName& OutCompactNodeIcon) const
 TArray<FPCGPinProperties> UPCGCollapseSettings::InputPinProperties() const
 {
 	TArray<FPCGPinProperties> PinProperties;
-	PinProperties.Emplace(PCGPinConstants::DefaultInputLabel, EPCGDataType::Spatial);
+	FPCGPinProperties& InputPinProperty = PinProperties.Emplace_GetRef(PCGPinConstants::DefaultInputLabel, EPCGDataType::Spatial);
+	InputPinProperty.SetRequiredPin();
 
 	return PinProperties;
 }
@@ -33,7 +34,8 @@ FPCGElementPtr UPCGCollapseSettings::CreateElement() const
 TArray<FPCGPinProperties> UPCGConvertToPointDataSettings::InputPinProperties() const
 {
 	TArray<FPCGPinProperties> PinProperties;
-	PinProperties.Emplace(PCGPinConstants::DefaultInputLabel, EPCGDataType::Param);
+	FPCGPinProperties& InputPinProperty = PinProperties.Emplace_GetRef(PCGPinConstants::DefaultInputLabel, EPCGDataType::Param);
+	InputPinProperty.SetRequiredPin();
 
 	return PinProperties;
 }
diff --git a/Engine/Plugins/PCG/Source/PCG/Private/Elements/PCGCopyPoints.cpp b/Engine/Plugins/PCG/Source/PCG/Private/Elements/PCGCopyPoints.cpp
--- a/Engine/Plugins/PCG/Source/PCG/Private/Elements/PCGCopyPoints.cpp
+++ b/Engine/Plugins/PCG/Source/PCG/Private/Elements/PCGCopyPoints.cpp
@@ -24,8 +24,12 @@ FText UPCGCopyPointsSettings::GetNodeTooltipText() const
 TArray<FPCGPinProperties> UPCGCopyPointsSettings::InputPinProperties() const
 {
 	TArray<FPCGPinProperties> PinProperties;
-	PinProperties.Emplace(PCGCopyPointsConstants::SourcePointsLabel, EPCGDataType::Point, /*bAllowMultipleConnections=*/true);
-	PinProperties.Emplace(PCGCopyPointsConstants::TargetPointsLabel, EPCGDataType::Point, /*bAllowMultipleConnections=*/true);
+	FPCGPinProperties& SourcePinProperty = PinProperties.Emplace_GetRef(PCGCopyPointsConstants::SourcePointsLabel, EPCGDataType::Point, /*bAllowMultipleConnections=*/true);
+	SourcePinProperty.SetRequiredPin();
+
+	FPCGPinProperties& TargetPinProperty = PinProperties.Emplace_GetRef(PCGCopyPointsConstants::TargetPointsLabel, EPCGDataType::Point, /*bAllowMultipleConnections=*/true);
+	TargetPinProperty.SetRequiredPin();
+
 	return PinProperties;
 }
 
diff --git a/Engine/Plugins/PCG/Source/PCG/Private/Elements/PCGCreateAttribute.cpp b/Engine/Plugins/PCG/Source/PCG/Private/Elements/PCGCreateAttribute.cpp
--- a/Engine/Plugins/PCG/Source/PCG/Private/Elements/PCGCreateAttribute.cpp
+++ b/Engine/Plugins/PCG/Source/PCG/Private/Elements/PCGCreateAttribute.cpp
@@ -179,7 +179,9 @@ void UPCGAddAttributeSettings::ApplyDeprecation(UPCGNode* InOutNode)
 TArray<FPCGPinProperties> UPCGAddAttributeSettings::InputPinProperties() const
 {
 	TArray<FPCGPinProperties> PinProperties;
-	PinProperties.Emplace(PCGPinConstants::DefaultInputLabel, EPCGDataType::Any);
+	FPCGPinProperties& InputPinProperty = PinProperties.Emplace_GetRef(PCGPinConstants::DefaultInputLabel, EPCGDataType::Any);
+	InputPinProperty.SetRequiredPin();
+
 	PinProperties.Emplace(PCGCreateAttributeConstants::AttributesLabel, EPCGDataType::Param, /*bInAllowMultipleConnections=*/true, /*bAllowMultipleData=*/true, PCGCreateAttributeConstants::AttributesTooltip);
 
 	return PinProperties;
diff --git a/Engine/Plugins/PCG/Source/PCG/Private/Elements/PCGDeleteAttributesElement.cpp b/Engine/Plugins/PCG/Source/PCG/Private/Elements/PCGDeleteAttributesElement.cpp
--- a/Engine/Plugins/PCG/Source/PCG/Private/Elements/PCGDeleteAttributesElement.cpp
+++ b/Engine/Plugins/PCG/Source/PCG/Private/Elements/PCGDeleteAttributesElement.cpp
@@ -99,14 +99,6 @@ FString UPCGDeleteAttributesSettings::GetAdditionalTitleInformation() const
 	}
 }
 
-TArray<FPCGPinProperties> UPCGDeleteAttributesSettings::InputPinProperties() const
-{
-	TArray<FPCGPinProperties> PinProperties;
-	PinProperties.Emplace(PCGPinConstants::DefaultInputLabel, EPCGDataType::Any);
-
-	return PinProperties;
-}
-
 TArray<FPCGPinProperties> UPCGDeleteAttributesSettings::OutputPinProperties() const
 {
 	TArray<FPCGPinProperties> PinProperties;
diff --git a/Engine/Plugins/PCG/Source/PCG/Private/Elements/PCGDifferenceElement.cpp b/Engine/Plugins/PCG/Source/PCG/Private/Elements/PCGDifferenceElement.cpp
--- a/Engine/Plugins/PCG/Source/PCG/Private/Elements/PCGDifferenceElement.cpp
+++ b/Engine/Plugins/PCG/Source/PCG/Private/Elements/PCGDifferenceElement.cpp
@@ -109,6 +109,7 @@ TArray<FPCGPinProperties> UPCGDifferenceSettings::InputPinProperties() const
 {
 	TArray<FPCGPinProperties> PinProperties;
 	FPCGPinProperties& PinPropertiesSource = PinProperties.Emplace_GetRef(PCGDifferenceConstants::SourceLabel, EPCGDataType::Spatial);
+	PinPropertiesSource.SetRequiredPin();
 	FPCGPinProperties& PinPropertiesDifferences = PinProperties.Emplace_GetRef(PCGDifferenceConstants::DifferencesLabel, EPCGDataType::Spatial);
 
 #if WITH_EDITOR
diff --git a/Engine/Plugins/PCG/Source/PCG/Private/Elements/PCGDistance.cpp b/Engine/Plugins/PCG/Source/PCG/Private/Elements/PCGDistance.cpp
--- a/Engine/Plugins/PCG/Source/PCG/Private/Elements/PCGDistance.cpp
+++ b/Engine/Plugins/PCG/Source/PCG/Private/Elements/PCGDistance.cpp
@@ -56,6 +56,7 @@ TArray<FPCGPinProperties> UPCGDistanceSettings::InputPinProperties() const
 {
 	TArray<FPCGPinProperties> PinProperties;
 	FPCGPinProperties& PinPropertySource = PinProperties.Emplace_GetRef(PCGDistance::SourceLabel, EPCGDataType::Point);
+	PinPropertySource.SetRequiredPin();
 	FPCGPinProperties& PinPropertyTarget = PinProperties.Emplace_GetRef(PCGDistance::TargetLabel, EPCGDataType::Point);
 
 #if WITH_EDITOR
diff --git a/Engine/Plugins/PCG/Source/PCG/Private/Elements/PCGFilterDataBase.cpp b/Engine/Plugins/PCG/Source/PCG/Private/Elements/PCGFilterDataBase.cpp
--- a/Engine/Plugins/PCG/Source/PCG/Private/Elements/PCGFilterDataBase.cpp
+++ b/Engine/Plugins/PCG/Source/PCG/Private/Elements/PCGFilterDataBase.cpp
@@ -22,14 +22,6 @@ void UPCGFilterDataBaseSettings::ApplyDeprecationBeforeUpdatePins(UPCGNode* InOu
 
 #endif // WITH_EDITOR
 
-TArray<FPCGPinProperties> UPCGFilterDataBaseSettings::InputPinProperties() const
-{
-	TArray<FPCGPinProperties> PinProperties;
-	PinProperties.Emplace(PCGPinConstants::DefaultInputLabel, EPCGDataType::Any);
-
-	return PinProperties;
-}
-
 TArray<FPCGPinProperties> UPCGFilterDataBaseSettings::OutputPinProperties() const
 {
 	TArray<FPCGPinProperties> PinProperties;
diff --git a/Engine/Plugins/PCG/Source/PCG/Private/Elements/PCGGather.cpp b/Engine/Plugins/PCG/Source/PCG/Private/Elements/PCGGather.cpp
--- a/Engine/Plugins/PCG/Source/PCG/Private/Elements/PCGGather.cpp
+++ b/Engine/Plugins/PCG/Source/PCG/Private/Elements/PCGGather.cpp
@@ -10,7 +10,8 @@
 TArray<FPCGPinProperties> UPCGGatherSettings::InputPinProperties() const
 {
 	TArray<FPCGPinProperties> PinProperties;
-	PinProperties.Emplace(PCGPinConstants::DefaultInputLabel, EPCGDataType::Any);
+	FPCGPinProperties& InputPinProperty = PinProperties.Emplace_GetRef(PCGPinConstants::DefaultInputLabel, EPCGDataType::Any);
+	InputPinProperty.SetRequiredPin();
 	PinProperties.Emplace(PCGPinConstants::DefaultDependencyOnlyLabel, EPCGDataType::Any, /*bAllowMultipleConnections=*/true, /*bAllowMultipleData=*/true, LOCTEXT("DependencyPinTooltip", "Data passed to this pin will be used to order execution but will otherwise not contribute to the results of this node."));
 
 	return PinProperties;
diff --git a/Engine/Plugins/PCG/Source/PCG/Private/Elements/PCGGetBounds.cpp b/Engine/Plugins/PCG/Source/PCG/Private/Elements/PCGGetBounds.cpp
--- a/Engine/Plugins/PCG/Source/PCG/Private/Elements/PCGGetBounds.cpp
+++ b/Engine/Plugins/PCG/Source/PCG/Private/Elements/PCGGetBounds.cpp
@@ -25,7 +25,8 @@ FText UPCGGetBoundsSettings::GetNodeTooltipText() const
 TArray<FPCGPinProperties> UPCGGetBoundsSettings::InputPinProperties() const
 {
 	TArray<FPCGPinProperties> PinProperties;
-	PinProperties.Emplace(PCGPinConstants::DefaultInputLabel, EPCGDataType::Spatial);
+	FPCGPinProperties& InputPinProperty = PinProperties.Emplace_GetRef(PCGPinConstants::DefaultInputLabel, EPCGDataType::Spatial);
+	InputPinProperty.SetRequiredPin();
 
 	return PinProperties;
 }
diff --git a/Engine/Plugins/PCG/Source/PCG/Private/Elements/PCGInnerIntersectionElement.cpp b/Engine/Plugins/PCG/Source/PCG/Private/Elements/PCGInnerIntersectionElement.cpp
--- a/Engine/Plugins/PCG/Source/PCG/Private/Elements/PCGInnerIntersectionElement.cpp
+++ b/Engine/Plugins/PCG/Source/PCG/Private/Elements/PCGInnerIntersectionElement.cpp
@@ -20,6 +20,7 @@ TArray<FPCGPinProperties> UPCGInnerIntersectionSettings::InputPinProperties() co
 {
 	TArray<FPCGPinProperties> PinProperties;
 	FPCGPinProperties& SourcePinProperty = PinProperties.Emplace_GetRef(PCGPinConstants::DefaultInputLabel, EPCGDataType::Spatial);
+	SourcePinProperty.SetRequiredPin();
 
 #if WITH_EDITOR
 	SourcePinProperty.Tooltip = LOCTEXT("SourcePinTooltip", "Source spatial data from which to conduct the intersection. Empty spatial data will be ignored.");
diff --git a/Engine/Plugins/PCG/Source/PCG/Private/Elements/PCGMatchAndSetAttributes.cpp b/Engine/Plugins/PCG/Source/PCG/Private/Elements/PCGMatchAndSetAttributes.cpp
--- a/Engine/Plugins/PCG/Source/PCG/Private/Elements/PCGMatchAndSetAttributes.cpp
+++ b/Engine/Plugins/PCG/Source/PCG/Private/Elements/PCGMatchAndSetAttributes.cpp
@@ -55,7 +55,9 @@ UPCGMatchAndSetAttributesSettings::UPCGMatchAndSetAttributesSettings()
 TArray<FPCGPinProperties> UPCGMatchAndSetAttributesSettings::InputPinProperties() const
 {
 	TArray<FPCGPinProperties> PinProperties;
-	PinProperties.Emplace(PCGPinConstants::DefaultInputLabel, EPCGDataType::Point | EPCGDataType::Param);
+	FPCGPinProperties& InputPinProperty = PinProperties.Emplace_GetRef(PCGPinConstants::DefaultInputLabel, EPCGDataType::Point | EPCGDataType::Param);
+	InputPinProperty.SetRequiredPin();
+
 	PinProperties.Emplace(PCGMatchAndSetAttributesConstants::MatchDataLabel, 
 		EPCGDataType::Param, 
 		/*bAllowMultipleConnection=*/false, 
diff --git a/Engine/Plugins/PCG/Source/PCG/Private/Elements/PCGMergeAttributes.cpp b/Engine/Plugins/PCG/Source/PCG/Private/Elements/PCGMergeAttributes.cpp
--- a/Engine/Plugins/PCG/Source/PCG/Private/Elements/PCGMergeAttributes.cpp
+++ b/Engine/Plugins/PCG/Source/PCG/Private/Elements/PCGMergeAttributes.cpp
@@ -24,7 +24,8 @@ FText UPCGMergeAttributesSettings::GetNodeTooltipText() const
 TArray<FPCGPinProperties> UPCGMergeAttributesSettings::InputPinProperties() const
 {
 	TArray<FPCGPinProperties> PinProperties;
-	PinProperties.Emplace(PCGPinConstants::DefaultInputLabel, EPCGDataType::Param, /*bAllowMultipleConnections=*/true);
+	FPCGPinProperties& InputPinProperty = PinProperties.Emplace_GetRef(PCGPinConstants::DefaultInputLabel, EPCGDataType::Param, /*bAllowMultipleConnections=*/true);
+	InputPinProperty.SetRequiredPin();
 
 	return PinProperties;
 }
diff --git a/Engine/Plugins/PCG/Source/PCG/Private/Elements/PCGMergeElement.cpp b/Engine/Plugins/PCG/Source/PCG/Private/Elements/PCGMergeElement.cpp
--- a/Engine/Plugins/PCG/Source/PCG/Private/Elements/PCGMergeElement.cpp
+++ b/Engine/Plugins/PCG/Source/PCG/Private/Elements/PCGMergeElement.cpp
@@ -19,10 +19,7 @@ FText UPCGMergeSettings::GetNodeTooltipText() const
 
 TArray<FPCGPinProperties> UPCGMergeSettings::InputPinProperties() const
 {
-	TArray<FPCGPinProperties> PinProperties;
-	PinProperties.Emplace(PCGPinConstants::DefaultInputLabel, EPCGDataType::Point, /*bAllowMultipleConnections=*/true);
-
-	return PinProperties;
+	return Super::DefaultPointInputPinProperties();
 }
 
 TArray<FPCGPinProperties> UPCGMergeSettings::OutputPinProperties() const
diff --git a/Engine/Plugins/PCG/Source/PCG/Private/Elements/PCGNumberOfElements.cpp b/Engine/Plugins/PCG/Source/PCG/Private/Elements/PCGNumberOfElements.cpp
--- a/Engine/Plugins/PCG/Source/PCG/Private/Elements/PCGNumberOfElements.cpp
+++ b/Engine/Plugins/PCG/Source/PCG/Private/Elements/PCGNumberOfElements.cpp
@@ -29,10 +29,7 @@ TArray<FPCGPinProperties> UPCGNumberOfElementsBaseSettings::OutputPinProperties(
 
 TArray<FPCGPinProperties> UPCGNumberOfPointsSettings::InputPinProperties() const
 {
-	TArray<FPCGPinProperties> PinProperties;
-	PinProperties.Emplace(PCGPinConstants::DefaultInputLabel, EPCGDataType::Point, /*bInAllowMultipleConnections=*/true, /*bAllowMultipleData=*/true);
-
-	return PinProperties;
+	return Super::DefaultPointInputPinProperties();
 }
 
 FPCGElementPtr UPCGNumberOfPointsSettings::CreateElement() const
@@ -48,7 +45,8 @@ FPCGElementPtr UPCGNumberOfPointsSettings::CreateElement() const
 TArray<FPCGPinProperties> UPCGNumberOfEntriesSettings::InputPinProperties() const
 {
 	TArray<FPCGPinProperties> PinProperties;
-	PinProperties.Emplace(PCGPinConstants::DefaultInputLabel, EPCGDataType::Param, /*bInAllowMultipleConnections=*/true, /*bAllowMultipleData=*/true);
+	FPCGPinProperties& InputPinProperty = PinProperties.Emplace_GetRef(PCGPinConstants::DefaultInputLabel, EPCGDataType::Param, /*bInAllowMultipleConnections=*/true, /*bAllowMultipleData=*/true);
+	InputPinProperty.SetRequiredPin();
 
 	return PinProperties;
 }
diff --git a/Engine/Plugins/PCG/Source/PCG/Private/Elements/PCGOuterIntersectionElement.cpp b/Engine/Plugins/PCG/Source/PCG/Private/Elements/PCGOuterIntersectionElement.cpp
--- a/Engine/Plugins/PCG/Source/PCG/Private/Elements/PCGOuterIntersectionElement.cpp
+++ b/Engine/Plugins/PCG/Source/PCG/Private/Elements/PCGOuterIntersectionElement.cpp
@@ -13,6 +13,7 @@ TArray<FPCGPinProperties> UPCGOuterIntersectionSettings::StaticInputPinPropertie
 {
 	TArray<FPCGPinProperties> StaticPinProperties;
 	FPCGPinProperties& PrimaryPinProperties = StaticPinProperties.Emplace_GetRef(PCGIntersectionConstants::PrimaryLabel, EPCGDataType::Spatial);
+	PrimaryPinProperties.SetRequiredPin();
 	FPCGPinProperties& SecondaryPinProperties = StaticPinProperties.Emplace_GetRef(FName(PCGIntersectionConstants::SecondaryLabel.ToString() + FString(TEXT("1"))), EPCGDataType::Spatial);
 
 #if WITH_EDITOR
diff --git a/Engine/Plugins/PCG/Source/PCG/Private/Elements/PCGPointMatchAndSet.cpp b/Engine/Plugins/PCG/Source/PCG/Private/Elements/PCGPointMatchAndSet.cpp
--- a/Engine/Plugins/PCG/Source/PCG/Private/Elements/PCGPointMatchAndSet.cpp
+++ b/Engine/Plugins/PCG/Source/PCG/Private/Elements/PCGPointMatchAndSet.cpp
@@ -52,12 +52,8 @@ void UPCGPointMatchAndSetSettings::ApplyDeprecation(UPCGNode* InOutNode)
 
 TArray<FPCGPinProperties> UPCGPointMatchAndSetSettings::InputPinProperties() const
 {
-	TArray<FPCGPinProperties> PinProperties;
-
-	PinProperties.Emplace(PCGPinConstants::DefaultInputLabel, EPCGDataType::Point);
 	// TODO Add param support?
-
-	return PinProperties;
+	return Super::DefaultPointInputPinProperties();
 }
 
 TArray<FPCGPinProperties> UPCGPointMatchAndSetSettings::OutputPinProperties() const
diff --git a/Engine/Plugins/PCG/Source/PCG/Private/Elements/PCGProjectionElement.cpp b/Engine/Plugins/PCG/Source/PCG/Private/Elements/PCGProjectionElement.cpp
--- a/Engine/Plugins/PCG/Source/PCG/Private/Elements/PCGProjectionElement.cpp
+++ b/Engine/Plugins/PCG/Source/PCG/Private/Elements/PCGProjectionElement.cpp
@@ -17,8 +17,11 @@
 TArray<FPCGPinProperties> UPCGProjectionSettings::InputPinProperties() const
 {
 	TArray<FPCGPinProperties> PinProperties;
-	PinProperties.Emplace(PCGPinConstants::DefaultInputLabel, EPCGDataType::Concrete, /*bInAllowMultipleConnections=*/true, /*bAllowMultipleData=*/true,  LOCTEXT("ProjectionSourcePinTooltip", "The data to project."));
-	PinProperties.Emplace(PCGProjectionConstants::ProjectionTargetLabel, EPCGDataType::Concrete, /*bAllowMultipleConnections=*/false, /*bAllowMultipleData=*/false, LOCTEXT("ProjectionTargetPinTooltip", "The projection target."));
+	FPCGPinProperties& InputPinProperty = PinProperties.Emplace_GetRef(PCGPinConstants::DefaultInputLabel, EPCGDataType::Concrete, /*bInAllowMultipleConnections=*/true, /*bAllowMultipleData=*/true,  LOCTEXT("ProjectionSourcePinTooltip", "The data to project."));
+	InputPinProperty.SetRequiredPin();
+
+	FPCGPinProperties& ProjectionPinProperty = PinProperties.Emplace_GetRef(PCGProjectionConstants::ProjectionTargetLabel, EPCGDataType::Concrete, /*bAllowMultipleConnections=*/false, /*bAllowMultipleData=*/false, LOCTEXT("ProjectionTargetPinTooltip", "The projection target."));
+	ProjectionPinProperty.SetRequiredPin();
 
 	return PinProperties;
 }
diff --git a/Engine/Plugins/PCG/Source/PCG/Private/Elements/PCGSplineSampler.cpp b/Engine/Plugins/PCG/Source/PCG/Private/Elements/PCGSplineSampler.cpp
--- a/Engine/Plugins/PCG/Source/PCG/Private/Elements/PCGSplineSampler.cpp
+++ b/Engine/Plugins/PCG/Source/PCG/Private/Elements/PCGSplineSampler.cpp
@@ -1383,7 +1383,9 @@ FText UPCGSplineSamplerSettings::GetNodeTooltipText() const
 TArray<FPCGPinProperties> UPCGSplineSamplerSettings::InputPinProperties() const
 {
 	TArray<FPCGPinProperties> PinProperties;
-	PinProperties.Emplace(PCGSplineSamplerConstants::SplineLabel, EPCGDataType::PolyLine, /*bAllowMultipleConnections=*/true, /*bAllowMultipleData=*/true);
+	FPCGPinProperties& SplinePinProperty = PinProperties.Emplace_GetRef(PCGSplineSamplerConstants::SplineLabel, EPCGDataType::PolyLine, /*bAllowMultipleConnections=*/true, /*bAllowMultipleData=*/true);
+	SplinePinProperty.SetRequiredPin();
+
 	// Only one connection allowed, user can union multiple shapes.
 	PinProperties.Emplace(PCGSplineSamplerConstants::BoundingShapeLabel, EPCGDataType::Spatial, /*bInAllowMultipleConnections=*/false, /*bAllowMultipleData=*/false, LOCTEXT("SplineSamplerBoundingShapePinTooltip",
 		"Optional. All sampled points must be contained within this shape."
diff --git a/Engine/Plugins/PCG/Source/PCG/Private/Elements/PCGSurfaceSampler.cpp b/Engine/Plugins/PCG/Source/PCG/Private/Elements/PCGSurfaceSampler.cpp
--- a/Engine/Plugins/PCG/Source/PCG/Private/Elements/PCGSurfaceSampler.cpp
+++ b/Engine/Plugins/PCG/Source/PCG/Private/Elements/PCGSurfaceSampler.cpp
@@ -247,11 +247,14 @@ FText UPCGSurfaceSamplerSettings::GetNodeTooltipText() const
 TArray<FPCGPinProperties> UPCGSurfaceSamplerSettings::InputPinProperties() const
 {
 	TArray<FPCGPinProperties> PinProperties;
-	PinProperties.Emplace(PCGSurfaceSamplerConstants::SurfaceLabel, EPCGDataType::Surface, /*bAllowMultipleConnections=*/true, /*bAllowMultipleData=*/true, LOCTEXT("SurfaceSamplerSurfacePinTooltip",
+
+	FPCGPinProperties& SurfacePinProperty = PinProperties.Emplace_GetRef(PCGSurfaceSamplerConstants::SurfaceLabel, EPCGDataType::Surface, /*bAllowMultipleConnections=*/true, /*bAllowMultipleData=*/true, LOCTEXT("SurfaceSamplerSurfacePinTooltip",
 		"The surface to sample with points. Points will be generated in the two dimensional footprint of the combined bounds of the Surface and the Bounding Shape (if any) "
 		"and then projected onto this surface. If this input is omitted then the network of shapes connected to the Bounding Shape pin will be inspected for a surface "
 		"shape to use to project the points onto."
 	));
+	SurfacePinProperty.SetRequiredPin();
+
 	// Only one connection allowed, user can union multiple shapes
 	PinProperties.Emplace(PCGSurfaceSamplerConstants::BoundingShapeLabel, EPCGDataType::Spatial, /*bInAllowMultipleConnections=*/false, /*bAllowMultipleData=*/false, LOCTEXT("SurfaceSamplerBoundingShapePinTooltip",
 		"All sampled points must be contained within this shape. If this input is omitted then bounds will be taken from the actor so that points are contained within actor bounds. "
diff --git a/Engine/Plugins/PCG/Source/PCG/Private/Elements/PCGUnionElement.cpp b/Engine/Plugins/PCG/Source/PCG/Private/Elements/PCGUnionElement.cpp
--- a/Engine/Plugins/PCG/Source/PCG/Private/Elements/PCGUnionElement.cpp
+++ b/Engine/Plugins/PCG/Source/PCG/Private/Elements/PCGUnionElement.cpp
@@ -10,7 +10,8 @@
 TArray<FPCGPinProperties> UPCGUnionSettings::InputPinProperties() const
 {
 	TArray<FPCGPinProperties> PinProperties;
-	PinProperties.Emplace(PCGPinConstants::DefaultInputLabel, EPCGDataType::Spatial);
+	FPCGPinProperties& InputPinProperty = PinProperties.Emplace_GetRef(PCGPinConstants::DefaultInputLabel, EPCGDataType::Spatial);
+	InputPinProperty.SetRequiredPin();
 
 	return PinProperties;
 }
diff --git a/Engine/Plugins/PCG/Source/PCG/Private/Elements/PCGVolumeSampler.cpp b/Engine/Plugins/PCG/Source/PCG/Private/Elements/PCGVolumeSampler.cpp
--- a/Engine/Plugins/PCG/Source/PCG/Private/Elements/PCGVolumeSampler.cpp
+++ b/Engine/Plugins/PCG/Source/PCG/Private/Elements/PCGVolumeSampler.cpp
@@ -145,10 +145,13 @@ FText UPCGVolumeSamplerSettings::GetNodeTooltipText() const
 TArray<FPCGPinProperties> UPCGVolumeSamplerSettings::InputPinProperties() const
 {
 	TArray<FPCGPinProperties> PinProperties;
+
 	// Spatial is ok - volume sampling just needs bounds.
-	PinProperties.Emplace(PCGVolumeSamplerConstants::VolumeLabel, EPCGDataType::Spatial, /*bAllowMultipleConnections=*/true, /*bAllowMultipleData=*/true, LOCTEXT("VolumeSamplerVolumePinTooltip",
+	FPCGPinProperties& VolumePinProperty = PinProperties.Emplace_GetRef(PCGVolumeSamplerConstants::VolumeLabel, EPCGDataType::Spatial, /*bAllowMultipleConnections=*/true, /*bAllowMultipleData=*/true, LOCTEXT("VolumeSamplerVolumePinTooltip",
 		"The volume to sample with points. Can be any spatial data that can provide bounds."
 	));
+	VolumePinProperty.SetRequiredPin();
+
 	// Only one connection allowed, user can union multiple shapes.
 	PinProperties.Emplace(PCGVolumeSamplerConstants::BoundingShapeLabel, EPCGDataType::Spatial, /*bInAllowMultipleConnections=*/false, /*bAllowMultipleData=*/false, LOCTEXT("VolumeSamplerBoundingShapePinTooltip",
 		"Optional. All sampled points must be contained within this shape."
diff --git a/Engine/Plugins/PCG/Source/PCG/Private/Graph/PCGGraphCompiler.cpp b/Engine/Plugins/PCG/Source/PCG/Private/Graph/PCGGraphCompiler.cpp
--- a/Engine/Plugins/PCG/Source/PCG/Private/Graph/PCGGraphCompiler.cpp
+++ b/Engine/Plugins/PCG/Source/PCG/Private/Graph/PCGGraphCompiler.cpp
@@ -522,11 +522,14 @@ bool FPCGGraphCompiler::CalculateStaticallyActiveRecursive(FPCGTaskId InTaskId,
 			continue;
 		}
 
-		bHasAnyNonAdvancedPins |= !InputPin->Properties.bAdvancedPin;
-
-		if (Node->IsInputPinRequiredByExecution(InputPin))
+		if (!InputPin->Properties.IsAdvancedPin())
 		{
-			PinsRequiringActiveConnection.AddUnique(InputPin->Properties.Label);
+			bHasAnyNonAdvancedPins = true;
+
+			if (Node->IsInputPinRequiredByExecution(InputPin))
+			{
+				PinsRequiringActiveConnection.AddUnique(InputPin->Properties.Label);
+			}
 		}
 	}
 	
@@ -538,7 +541,7 @@ bool FPCGGraphCompiler::CalculateStaticallyActiveRecursive(FPCGTaskId InTaskId,
 		const UPCGPin* InputPin = Input.OutPin;
 
 		// Only non-advanced input pins play a part in determining active/inactive state.
-		if (InputPin && InputPin->Properties.bAdvancedPin)
+		if (InputPin && InputPin->Properties.IsAdvancedPin())
 		{
 			continue;
 		}
@@ -842,7 +845,7 @@ void FPCGGraphCompiler::CalculateDynamicActivePinDependencies(FPCGTaskId InTaskI
 		// is active. We build a disjunction that expresses this.
 		for (const FPCGGraphTaskInput& Input : InOutCompiledTasks[InTaskId].Inputs)
 		{
-			if (Input.OutPin && Input.OutPin->Properties.bAdvancedPin)
+			if (Input.OutPin && Input.OutPin->Properties.IsAdvancedPin())
 			{
 				// Advanced input pins never participate in keeping node active.
 				continue;
diff --git a/Engine/Plugins/PCG/Source/PCG/Private/Metadata/PCGObjectPropertyOverride.cpp b/Engine/Plugins/PCG/Source/PCG/Private/Metadata/PCGObjectPropertyOverride.cpp
--- a/Engine/Plugins/PCG/Source/PCG/Private/Metadata/PCGObjectPropertyOverride.cpp
+++ b/Engine/Plugins/PCG/Source/PCG/Private/Metadata/PCGObjectPropertyOverride.cpp
@@ -17,7 +17,7 @@ namespace PCGObjectPropertyOverrideHelpers
 	FPCGPinProperties CreateObjectPropertiesOverridePin(FName Label, const FText& Tooltip)
 	{
 		FPCGPinProperties ObjectOverridePinProperties(Label, EPCGDataType::Param, /*bAllowMultipleConnections=*/false, /*bAllowMultipleData=*/false, Tooltip);
-		ObjectOverridePinProperties.bAdvancedPin = true;
+		ObjectOverridePinProperties.SetAdvancedPin();
 		return ObjectOverridePinProperties;
 	}
 
diff --git a/Engine/Plugins/PCG/Source/PCG/Private/PCGInputOutputSettings.cpp b/Engine/Plugins/PCG/Source/PCG/Private/PCGInputOutputSettings.cpp
--- a/Engine/Plugins/PCG/Source/PCG/Private/PCGInputOutputSettings.cpp
+++ b/Engine/Plugins/PCG/Source/PCG/Private/PCGInputOutputSettings.cpp
@@ -117,7 +117,7 @@ void UPCGGraphInputOutputSettings::SetInput(bool bInIsInput)
 
 			Algo::Transform(StaticLabels, Pins, [DefaultPinDataType](const PCGInputOutputPrivate::FLabelAndTooltip& InLabelAndTooltip) {
 				FPCGPinProperties Res = FPCGPinProperties(InLabelAndTooltip.Label, DefaultPinDataType, /*bMultiConnections=*/true, /*bMultiData=*/true, InLabelAndTooltip.Tooltip);
-				Res.bAdvancedPin = true;
+				Res.SetAdvancedPin();
 				return Res;
 			});
 		}
@@ -167,7 +167,7 @@ TArray<FPCGPinProperties> UPCGGraphInputOutputSettings::DefaultInputPinPropertie
 
 	Algo::Transform(StaticLabels, PinProperties, [DefaultPinDataType](const PCGInputOutputPrivate::FLabelAndTooltip& InLabelAndTooltip) {
 		FPCGPinProperties Res = FPCGPinProperties(InLabelAndTooltip.Label, DefaultPinDataType, /*bMultiConnections=*/true, /*bMultiData=*/true, InLabelAndTooltip.Tooltip);
-		Res.bAdvancedPin = true;
+		Res.SetAdvancedPin();
 		return Res;
 	});
 
diff --git a/Engine/Plugins/PCG/Source/PCG/Private/PCGPin.cpp b/Engine/Plugins/PCG/Source/PCG/Private/PCGPin.cpp
--- a/Engine/Plugins/PCG/Source/PCG/Private/PCGPin.cpp
+++ b/Engine/Plugins/PCG/Source/PCG/Private/PCGPin.cpp
@@ -90,10 +90,25 @@ bool FPCGPinProperties::operator==(const FPCGPinProperties& Other) const
 		bAllowMultipleConnections == Other.bAllowMultipleConnections &&
 		bAllowMultipleData == Other.bAllowMultipleData &&
 		Usage == Other.Usage &&
-		bAdvancedPin == Other.bAdvancedPin &&
+		PinStatus == Other.PinStatus &&
 		bInvisiblePin == Other.bInvisiblePin;
 }
 
+void FPCGPinProperties::PostSerialize(const FArchive& Ar)
+{
+#if WITH_EDITOR
+	if (Ar.IsLoading() && Ar.IsPersistent() && !Ar.HasAnyPortFlags(PPF_Duplicate | PPF_DuplicateForPIE))
+	{
+		if (bAdvancedPin_DEPRECATED)
+		{
+			PinStatus = EPCGPinStatus::Advanced;
+		}
+
+		bAdvancedPin_DEPRECATED = false;
+	}
+#endif // WITH_EDITOR
+}
+
 UPCGPin::UPCGPin(const FObjectInitializer& ObjectInitializer)
 	: Super(ObjectInitializer)
 {
diff --git a/Engine/Plugins/PCG/Source/PCG/Private/PCGSettings.cpp b/Engine/Plugins/PCG/Source/PCG/Private/PCGSettings.cpp
--- a/Engine/Plugins/PCG/Source/PCG/Private/PCGSettings.cpp
+++ b/Engine/Plugins/PCG/Source/PCG/Private/PCGSettings.cpp
@@ -345,7 +345,8 @@ UObject* UPCGSettings::GetJumpTargetForDoubleClick() const
 TArray<FPCGPinProperties> UPCGSettings::InputPinProperties() const
 {
 	TArray<FPCGPinProperties> PinProperties;
-	PinProperties.Emplace(PCGPinConstants::DefaultInputLabel, EPCGDataType::Any);
+	FPCGPinProperties& InputPinProperty = PinProperties.Emplace_GetRef(PCGPinConstants::DefaultInputLabel, EPCGDataType::Any);
+	InputPinProperty.SetRequiredPin();
 
 	return PinProperties;
 }
@@ -408,7 +409,7 @@ void UPCGSettings::FillOverridableParamsPins(TArray<FPCGPinProperties>& OutPins)
 	else
 	{
 		FPCGPinProperties& ParamPin = OutPins.Emplace_GetRef(PCGPinConstants::DefaultParamsLabel, EPCGDataType::Param, /*bInAllowMultipleConnections=*/ true, /*bAllowMultipleData=*/ true);
-		ParamPin.bAdvancedPin = true;
+		ParamPin.SetAdvancedPin();
 
 #if WITH_EDITOR
 		ParamPin.Tooltip = LOCTEXT("GlobalParamPinTooltip", "Atribute Set containing multiple parameters to override. Names must match perfectly.");
@@ -431,7 +432,7 @@ void UPCGSettings::FillOverridableParamsPins(TArray<FPCGPinProperties>& OutPins)
 		InputPinsLabelsAndTypes.Emplace(OverridableParam.Label, EPCGDataType::Param);
 
 		FPCGPinProperties& ParamPin = OutPins.Emplace_GetRef(OverridableParam.Label, EPCGDataType::Param, /*bInAllowMultipleConnections=*/ false, /*bAllowMultipleData=*/ false);
-		ParamPin.bAdvancedPin = true;
+		ParamPin.SetAdvancedPin();
 #if WITH_EDITOR
 
 		if (!OverridableParam.Properties.IsEmpty())
@@ -825,7 +826,8 @@ void UPCGSettings::CacheCrc()
 TArray<FPCGPinProperties> UPCGSettings::DefaultPointInputPinProperties() const
 {
 	TArray<FPCGPinProperties> Properties;
-	Properties.Emplace(PCGPinConstants::DefaultInputLabel, EPCGDataType::Point);
+	FPCGPinProperties& InputPinProperty = Properties.Emplace_GetRef(PCGPinConstants::DefaultInputLabel, EPCGDataType::Point);
+	InputPinProperty.SetRequiredPin();
 	return Properties;
 }
 
diff --git a/Engine/Plugins/PCG/Source/PCG/Private/PCGSubgraph.cpp b/Engine/Plugins/PCG/Source/PCG/Private/PCGSubgraph.cpp
--- a/Engine/Plugins/PCG/Source/PCG/Private/PCGSubgraph.cpp
+++ b/Engine/Plugins/PCG/Source/PCG/Private/PCGSubgraph.cpp
@@ -30,9 +30,9 @@ namespace PCGSubgraphSettings
 			{
 				PinProperties.bInvisiblePin = false;
 
-				if (Pin->IsConnected())
+				if (Pin->IsConnected() && PinProperties.IsAdvancedPin())
 				{
-					PinProperties.bAdvancedPin = false;
+					PinProperties.SetNormalPin();
 				}
 			}
 		}
diff --git a/Engine/Plugins/PCG/Source/PCG/Private/Tests/Elements/PCGMultiSelectTest.cpp b/Engine/Plugins/PCG/Source/PCG/Private/Tests/Elements/PCGMultiSelectTest.cpp
--- a/Engine/Plugins/PCG/Source/PCG/Private/Tests/Elements/PCGMultiSelectTest.cpp
+++ b/Engine/Plugins/PCG/Source/PCG/Private/Tests/Elements/PCGMultiSelectTest.cpp
@@ -45,7 +45,7 @@ protected:
 		TArray<FPCGPinProperties> SelectElementPinProperties = TestSettings->AllInputPinProperties();
 		for (int i = 0; i < InputNum; ++i)
 		{
-			if (!SelectElementPinProperties[i].bAdvancedPin && SelectElementPinProperties[i].Label != DefaultPathPinLabel)
+			if (!SelectElementPinProperties[i].IsAdvancedPin() && SelectElementPinProperties[i].Label != DefaultPathPinLabel)
 			{
 				FPCGTaggedData& TaggedData = TestData.InputData.TaggedData.Emplace_GetRef();
 				TaggedData.Data = PCGTestsCommon::CreateEmptyPointData();
diff --git a/Engine/Plugins/PCG/Source/PCG/Public/Elements/ControlFlow/PCGBranch.h b/Engine/Plugins/PCG/Source/PCG/Public/Elements/ControlFlow/PCGBranch.h
--- a/Engine/Plugins/PCG/Source/PCG/Public/Elements/ControlFlow/PCGBranch.h
+++ b/Engine/Plugins/PCG/Source/PCG/Public/Elements/ControlFlow/PCGBranch.h
@@ -24,7 +24,7 @@ public:
 	virtual bool HasDynamicPins() const override { return true; }
 	virtual bool OutputPinsCanBeDeactivated() const override { return true; }
 	virtual bool IsPinStaticallyActive(const FName& PinLabel) const override;
-	virtual bool IsInputPinRequiredByExecution(const UPCGPin* InPin) const override { return !InPin->Properties.bAdvancedPin; }
+	
 	virtual bool IsPinUsedByNodeExecution(const UPCGPin* InPin) const override;
 
 protected:
diff --git a/Engine/Plugins/PCG/Source/PCG/Public/Elements/ControlFlow/PCGSwitch.h b/Engine/Plugins/PCG/Source/PCG/Public/Elements/ControlFlow/PCGSwitch.h
--- a/Engine/Plugins/PCG/Source/PCG/Public/Elements/ControlFlow/PCGSwitch.h
+++ b/Engine/Plugins/PCG/Source/PCG/Public/Elements/ControlFlow/PCGSwitch.h
@@ -33,7 +33,7 @@ public:
 	virtual bool OutputPinsCanBeDeactivated() const override { return true; }
 	virtual bool IsPinStaticallyActive(const FName& PinLabel) const override;
 	virtual bool IsPinUsedByNodeExecution(const UPCGPin* InPin) const override;
-	virtual bool IsInputPinRequiredByExecution(const UPCGPin* InPin) const override { return !InPin->Properties.bAdvancedPin; }
+	
 	virtual FString GetAdditionalTitleInformation() const override;
 	virtual bool HasFlippedTitleLines() const override { return true; }
 
diff --git a/Engine/Plugins/PCG/Source/PCG/Public/Elements/Metadata/PCGMetadataBreakTransform.h b/Engine/Plugins/PCG/Source/PCG/Public/Elements/Metadata/PCGMetadataBreakTransform.h
--- a/Engine/Plugins/PCG/Source/PCG/Public/Elements/Metadata/PCGMetadataBreakTransform.h
+++ b/Engine/Plugins/PCG/Source/PCG/Public/Elements/Metadata/PCGMetadataBreakTransform.h
@@ -22,7 +22,7 @@ public:
 	virtual FName GetDefaultNodeName() const override;
 	virtual FText GetDefaultNodeTitle() const override;
 #endif
-	virtual bool IsInputPinRequiredByExecution(const UPCGPin* InPin) const override { return !InPin->Properties.bAdvancedPin; }
+	
 
 	virtual FPCGAttributePropertyInputSelector GetInputSource(uint32 Index) const override;
 
diff --git a/Engine/Plugins/PCG/Source/PCG/Public/Elements/Metadata/PCGMetadataBreakVector.h b/Engine/Plugins/PCG/Source/PCG/Public/Elements/Metadata/PCGMetadataBreakVector.h
--- a/Engine/Plugins/PCG/Source/PCG/Public/Elements/Metadata/PCGMetadataBreakVector.h
+++ b/Engine/Plugins/PCG/Source/PCG/Public/Elements/Metadata/PCGMetadataBreakVector.h
@@ -30,7 +30,7 @@ public:
 	virtual FName GetDefaultNodeName() const override;
 	virtual FText GetDefaultNodeTitle() const override;
 #endif
-	virtual bool IsInputPinRequiredByExecution(const UPCGPin* InPin) const override { return !InPin->Properties.bAdvancedPin; }
+	
 
 	virtual FPCGAttributePropertyInputSelector GetInputSource(uint32 Index) const override;
 
diff --git a/Engine/Plugins/PCG/Source/PCG/Public/Elements/Metadata/PCGMetadataElement.h b/Engine/Plugins/PCG/Source/PCG/Public/Elements/Metadata/PCGMetadataElement.h
--- a/Engine/Plugins/PCG/Source/PCG/Public/Elements/Metadata/PCGMetadataElement.h
+++ b/Engine/Plugins/PCG/Source/PCG/Public/Elements/Metadata/PCGMetadataElement.h
@@ -33,7 +33,6 @@ public:
 	virtual FText GetDefaultNodeTitle() const override { return NSLOCTEXT("PCGMetadataOperationSettings", "NodeTitle", "Copy Attribute"); }
 	virtual FText GetNodeTooltipText() const;
 	virtual EPCGSettingsType GetType() const override { return EPCGSettingsType::Metadata; }
-	virtual bool IsInputPinRequiredByExecution(const UPCGPin* InPin) const override { return InPin->Properties.Label == PCGPinConstants::DefaultInputLabel; }
 	virtual bool IsPinUsedByNodeExecution(const UPCGPin* InPin) const override;
 	virtual void ApplyDeprecation(UPCGNode* InOutNode) override;
 #endif
diff --git a/Engine/Plugins/PCG/Source/PCG/Public/Elements/Metadata/PCGMetadataPartition.h b/Engine/Plugins/PCG/Source/PCG/Public/Elements/Metadata/PCGMetadataPartition.h
--- a/Engine/Plugins/PCG/Source/PCG/Public/Elements/Metadata/PCGMetadataPartition.h
+++ b/Engine/Plugins/PCG/Source/PCG/Public/Elements/Metadata/PCGMetadataPartition.h
@@ -24,11 +24,10 @@ public:
 	virtual EPCGSettingsType GetType() const override { return EPCGSettingsType::Metadata; }
 	virtual bool HasDynamicPins() const override { return true; }
 #endif
-	virtual bool IsInputPinRequiredByExecution(const UPCGPin* InPin) const override { return !InPin->Properties.bAdvancedPin; }
+	
 	virtual FString GetAdditionalTitleInformation() const override;
 
 protected:
-	virtual TArray<FPCGPinProperties> InputPinProperties() const override;
 	virtual TArray<FPCGPinProperties> OutputPinProperties() const override;
 	virtual FPCGElementPtr CreateElement() const override;
 	//~End UPCGSettings interface
diff --git a/Engine/Plugins/PCG/Source/PCG/Public/Elements/Metadata/PCGMetadataRenameElement.h b/Engine/Plugins/PCG/Source/PCG/Public/Elements/Metadata/PCGMetadataRenameElement.h
--- a/Engine/Plugins/PCG/Source/PCG/Public/Elements/Metadata/PCGMetadataRenameElement.h
+++ b/Engine/Plugins/PCG/Source/PCG/Public/Elements/Metadata/PCGMetadataRenameElement.h
@@ -19,11 +19,10 @@ public:
 	virtual EPCGSettingsType GetType() const override { return EPCGSettingsType::Metadata; }
 	virtual bool HasDynamicPins() const override { return true; }
 #endif
-	virtual bool IsInputPinRequiredByExecution(const UPCGPin* InPin) const override { return !InPin->Properties.bAdvancedPin; }
+	
 	virtual FString GetAdditionalTitleInformation() const override;
 
 protected:
-	virtual TArray<FPCGPinProperties> InputPinProperties() const override;
 	virtual TArray<FPCGPinProperties> OutputPinProperties() const override;
 	virtual FPCGElementPtr CreateElement() const override;
 	//~End UPCGSettings interface
diff --git a/Engine/Plugins/PCG/Source/PCG/Public/Elements/PCGAddTag.h b/Engine/Plugins/PCG/Source/PCG/Public/Elements/PCGAddTag.h
--- a/Engine/Plugins/PCG/Source/PCG/Public/Elements/PCGAddTag.h
+++ b/Engine/Plugins/PCG/Source/PCG/Public/Elements/PCGAddTag.h
@@ -23,10 +23,9 @@ public:
 	virtual EPCGSettingsType GetType() const override { return EPCGSettingsType::Generic; }
 	virtual bool HasDynamicPins() const override { return true; }
 #endif
-	virtual bool IsInputPinRequiredByExecution(const UPCGPin* InPin) const override { return !InPin->Properties.bAdvancedPin; }
+	
 
 protected:
-	virtual TArray<FPCGPinProperties> InputPinProperties() const override;
 	virtual TArray<FPCGPinProperties> OutputPinProperties() const override;
 	virtual FPCGElementPtr CreateElement() const override;
 	//~End UPCGSettings interface
diff --git a/Engine/Plugins/PCG/Source/PCG/Public/Elements/PCGApplyScaleToBounds.h b/Engine/Plugins/PCG/Source/PCG/Public/Elements/PCGApplyScaleToBounds.h
--- a/Engine/Plugins/PCG/Source/PCG/Public/Elements/PCGApplyScaleToBounds.h
+++ b/Engine/Plugins/PCG/Source/PCG/Public/Elements/PCGApplyScaleToBounds.h
@@ -22,7 +22,7 @@ public:
 	virtual FText GetNodeTooltipText() const override { return NSLOCTEXT("PCGApplyScaleToBoundsElement", "NodeTooltip", "Applies the scale of each point to its bounds and resets the scale."); }
 	virtual EPCGSettingsType GetType() const override { return EPCGSettingsType::PointOps; }
 #endif
-	virtual bool IsInputPinRequiredByExecution(const UPCGPin* InPin) const override { return !InPin->Properties.bAdvancedPin; }
+	
 
 protected:
 	virtual TArray<FPCGPinProperties> InputPinProperties() const override { return Super::DefaultPointInputPinProperties(); }
diff --git a/Engine/Plugins/PCG/Source/PCG/Public/Elements/PCGAttributeFilter.h b/Engine/Plugins/PCG/Source/PCG/Public/Elements/PCGAttributeFilter.h
--- a/Engine/Plugins/PCG/Source/PCG/Public/Elements/PCGAttributeFilter.h
+++ b/Engine/Plugins/PCG/Source/PCG/Public/Elements/PCGAttributeFilter.h
@@ -81,7 +81,6 @@ public:
 	virtual EPCGSettingsType GetType() const override { return EPCGSettingsType::Filter; }
 	virtual bool HasDynamicPins() const override { return true; }
 #endif
-	virtual bool IsInputPinRequiredByExecution(const UPCGPin* InPin) const override { return InPin->Properties.Label == PCGPinConstants::DefaultInputLabel; }
 	virtual FString GetAdditionalTitleInformation() const override;
 	virtual EPCGDataType GetCurrentPinTypes(const UPCGPin* InPin) const override;
 
@@ -149,7 +148,6 @@ public:
 	virtual EPCGSettingsType GetType() const override { return EPCGSettingsType::Filter; }
 	virtual bool HasDynamicPins() const override { return true; }
 #endif
-	virtual bool IsInputPinRequiredByExecution(const UPCGPin* InPin) const override { return InPin->Properties.Label == PCGPinConstants::DefaultInputLabel; }
 	virtual FString GetAdditionalTitleInformation() const override;
 	virtual EPCGDataType GetCurrentPinTypes(const UPCGPin* InPin) const override;
 
diff --git a/Engine/Plugins/PCG/Source/PCG/Public/Elements/PCGAttributeGetFromIndexElement.h b/Engine/Plugins/PCG/Source/PCG/Public/Elements/PCGAttributeGetFromIndexElement.h
--- a/Engine/Plugins/PCG/Source/PCG/Public/Elements/PCGAttributeGetFromIndexElement.h
+++ b/Engine/Plugins/PCG/Source/PCG/Public/Elements/PCGAttributeGetFromIndexElement.h
@@ -19,7 +19,7 @@ public:
 	virtual FText GetNodeTooltipText() const override;
 	virtual EPCGSettingsType GetType() const override { return EPCGSettingsType::Param; }
 #endif
-	virtual bool IsInputPinRequiredByExecution(const UPCGPin* InPin) const override { return !InPin->Properties.bAdvancedPin; }
+	
 
 protected:
 	virtual TArray<FPCGPinProperties> InputPinProperties() const override;
diff --git a/Engine/Plugins/PCG/Source/PCG/Public/Elements/PCGAttributeGetFromPointIndexElement.h b/Engine/Plugins/PCG/Source/PCG/Public/Elements/PCGAttributeGetFromPointIndexElement.h
--- a/Engine/Plugins/PCG/Source/PCG/Public/Elements/PCGAttributeGetFromPointIndexElement.h
+++ b/Engine/Plugins/PCG/Source/PCG/Public/Elements/PCGAttributeGetFromPointIndexElement.h
@@ -38,7 +38,7 @@ public:
 	virtual EPCGSettingsType GetType() const override { return EPCGSettingsType::Metadata; }
 	virtual void ApplyDeprecation(UPCGNode* InOutNode) override;
 #endif
-	virtual bool IsInputPinRequiredByExecution(const UPCGPin* InPin) const override { return !InPin->Properties.bAdvancedPin; }
+	
 
 protected:
 	virtual TArray<FPCGPinProperties> InputPinProperties() const override;
diff --git a/Engine/Plugins/PCG/Source/PCG/Public/Elements/PCGAttributeNoise.h b/Engine/Plugins/PCG/Source/PCG/Public/Elements/PCGAttributeNoise.h
--- a/Engine/Plugins/PCG/Source/PCG/Public/Elements/PCGAttributeNoise.h
+++ b/Engine/Plugins/PCG/Source/PCG/Public/Elements/PCGAttributeNoise.h
@@ -48,7 +48,7 @@ public:
 	virtual EPCGSettingsType GetType() const override { return EPCGSettingsType::Metadata; }
 #endif
 	virtual bool HasDynamicPins() const override { return true; }
-	virtual bool IsInputPinRequiredByExecution(const UPCGPin* InPin) const override { return !InPin->Properties.bAdvancedPin; }
+	
 
 protected:
 	virtual TArray<FPCGPinProperties> InputPinProperties() const override;
diff --git a/Engine/Plugins/PCG/Source/PCG/Public/Elements/PCGAttributeSelectElement.h b/Engine/Plugins/PCG/Source/PCG/Public/Elements/PCGAttributeSelectElement.h
--- a/Engine/Plugins/PCG/Source/PCG/Public/Elements/PCGAttributeSelectElement.h
+++ b/Engine/Plugins/PCG/Source/PCG/Public/Elements/PCGAttributeSelectElement.h
@@ -62,7 +62,7 @@ public:
 	virtual EPCGSettingsType GetType() const override { return EPCGSettingsType::Metadata; }
 	virtual void ApplyDeprecation(UPCGNode* InOutNode) override;
 #endif
-	virtual bool IsInputPinRequiredByExecution(const UPCGPin* InPin) const override { return !InPin->Properties.bAdvancedPin; }
+	
 	virtual FString GetAdditionalTitleInformation() const override;
 
 protected:
diff --git a/Engine/Plugins/PCG/Source/PCG/Public/Elements/PCGAttributeTransferElement.h b/Engine/Plugins/PCG/Source/PCG/Public/Elements/PCGAttributeTransferElement.h
--- a/Engine/Plugins/PCG/Source/PCG/Public/Elements/PCGAttributeTransferElement.h
+++ b/Engine/Plugins/PCG/Source/PCG/Public/Elements/PCGAttributeTransferElement.h
@@ -44,7 +44,6 @@ public:
 	virtual bool HasDynamicPins() const override { return true; }
 	virtual void ApplyDeprecation(UPCGNode* InOutNode) override;
 #endif
-	virtual bool IsInputPinRequiredByExecution(const UPCGPin* InPin) const override { return InPin->Properties.Label == PCGAttributeTransferConstants::TargetLabel; }
 	virtual EPCGDataType GetCurrentPinTypes(const UPCGPin* InPin) const override;
 	virtual FString GetAdditionalTitleInformation() const override;
 
diff --git a/Engine/Plugins/PCG/Source/PCG/Public/Elements/PCGBoundsModifier.h b/Engine/Plugins/PCG/Source/PCG/Public/Elements/PCGBoundsModifier.h
--- a/Engine/Plugins/PCG/Source/PCG/Public/Elements/PCGBoundsModifier.h
+++ b/Engine/Plugins/PCG/Source/PCG/Public/Elements/PCGBoundsModifier.h
@@ -34,7 +34,7 @@ public:
 	virtual FText GetNodeTooltipText() const override;
 	virtual EPCGSettingsType GetType() const override { return EPCGSettingsType::PointOps; }
 #endif
-	virtual bool IsInputPinRequiredByExecution(const UPCGPin* InPin) const override { return !InPin->Properties.bAdvancedPin; }
+	
 
 protected:
 	virtual TArray<FPCGPinProperties> InputPinProperties() const override { return Super::DefaultPointInputPinProperties(); }
diff --git a/Engine/Plugins/PCG/Source/PCG/Public/Elements/PCGCollapseElement.h b/Engine/Plugins/PCG/Source/PCG/Public/Elements/PCGCollapseElement.h
--- a/Engine/Plugins/PCG/Source/PCG/Public/Elements/PCGCollapseElement.h
+++ b/Engine/Plugins/PCG/Source/PCG/Public/Elements/PCGCollapseElement.h
@@ -21,7 +21,7 @@ public:
 	virtual bool ShouldDrawNodeCompact() const override { return true; }
 	virtual bool GetCompactNodeIcon(FName& OutCompactNodeIcon) const override;
 #endif
-	virtual bool IsInputPinRequiredByExecution(const UPCGPin* InPin) const override { return !InPin->Properties.bAdvancedPin; }
+	
 
 protected:
 	virtual TArray<FPCGPinProperties> InputPinProperties() const override;
diff --git a/Engine/Plugins/PCG/Source/PCG/Public/Elements/PCGCombinePoints.h b/Engine/Plugins/PCG/Source/PCG/Public/Elements/PCGCombinePoints.h
--- a/Engine/Plugins/PCG/Source/PCG/Public/Elements/PCGCombinePoints.h
+++ b/Engine/Plugins/PCG/Source/PCG/Public/Elements/PCGCombinePoints.h
@@ -23,7 +23,7 @@ public:
 	virtual FText GetNodeTooltipText() const override { return NSLOCTEXT("PCGCombinePointsElement", "NodeTooltip", "Combines each point to share a singular bound extent."); }
 	virtual EPCGSettingsType GetType() const override { return EPCGSettingsType::PointOps; }
 #endif
-	virtual bool IsInputPinRequiredByExecution(const UPCGPin* InPin) const override { return !InPin->Properties.bAdvancedPin; }
+	
 
 protected:
 	virtual TArray<FPCGPinProperties> InputPinProperties() const override { return Super::DefaultPointInputPinProperties(); }
diff --git a/Engine/Plugins/PCG/Source/PCG/Public/Elements/PCGConvertToAttributeSet.h b/Engine/Plugins/PCG/Source/PCG/Public/Elements/PCGConvertToAttributeSet.h
--- a/Engine/Plugins/PCG/Source/PCG/Public/Elements/PCGConvertToAttributeSet.h
+++ b/Engine/Plugins/PCG/Source/PCG/Public/Elements/PCGConvertToAttributeSet.h
@@ -22,7 +22,7 @@ public:
 	virtual bool ShouldDrawNodeCompact() const override { return true; }
 	virtual bool GetCompactNodeIcon(FName& OutCompactNodeIcon) const override;
 #endif
-	virtual bool IsInputPinRequiredByExecution(const UPCGPin* InPin) const override { return !InPin->Properties.bAdvancedPin; }
+	
 
 protected:
 	virtual TArray<FPCGPinProperties> InputPinProperties() const override;
diff --git a/Engine/Plugins/PCG/Source/PCG/Public/Elements/PCGConvexHull2D.h b/Engine/Plugins/PCG/Source/PCG/Public/Elements/PCGConvexHull2D.h
--- a/Engine/Plugins/PCG/Source/PCG/Public/Elements/PCGConvexHull2D.h
+++ b/Engine/Plugins/PCG/Source/PCG/Public/Elements/PCGConvexHull2D.h
@@ -22,7 +22,7 @@ public:
 	virtual FText GetNodeTooltipText() const override { return NSLOCTEXT("PCGConvexHullElement", "NodeTooltip", "Return the 2D convex hull of a set of points on the XY plane."); }
 	virtual EPCGSettingsType GetType() const override { return EPCGSettingsType::Spatial; }
 #endif
-	virtual bool IsInputPinRequiredByExecution(const UPCGPin* InPin) const override { return !InPin->Properties.bAdvancedPin; }
+	
 
 protected:
 	virtual TArray<FPCGPinProperties> InputPinProperties() const override { return Super::DefaultPointInputPinProperties(); }
diff --git a/Engine/Plugins/PCG/Source/PCG/Public/Elements/PCGCopyPoints.h b/Engine/Plugins/PCG/Source/PCG/Public/Elements/PCGCopyPoints.h
--- a/Engine/Plugins/PCG/Source/PCG/Public/Elements/PCGCopyPoints.h
+++ b/Engine/Plugins/PCG/Source/PCG/Public/Elements/PCGCopyPoints.h
@@ -51,7 +51,7 @@ public:
 	virtual FText GetNodeTooltipText() const override;
 	virtual EPCGSettingsType GetType() const override { return EPCGSettingsType::Sampler; }
 #endif
-	virtual bool IsInputPinRequiredByExecution(const UPCGPin* InPin) const override { return !InPin->Properties.bAdvancedPin; }
+	
 
 protected:
 	virtual TArray<FPCGPinProperties> InputPinProperties() const override;
diff --git a/Engine/Plugins/PCG/Source/PCG/Public/Elements/PCGCreateAttribute.h b/Engine/Plugins/PCG/Source/PCG/Public/Elements/PCGCreateAttribute.h
--- a/Engine/Plugins/PCG/Source/PCG/Public/Elements/PCGCreateAttribute.h
+++ b/Engine/Plugins/PCG/Source/PCG/Public/Elements/PCGCreateAttribute.h
@@ -39,7 +39,6 @@ public:
 	virtual bool IsPinUsedByNodeExecution(const UPCGPin* InPin) const override;
 	virtual bool CanEditChange(const FProperty* InProperty) const override;
 #endif // WITH_EDITOR
-	virtual bool IsInputPinRequiredByExecution(const UPCGPin* InPin) const override { return InPin->Properties.Label == PCGPinConstants::DefaultInputLabel; }
 	virtual EPCGDataType GetCurrentPinTypes(const UPCGPin* InPin) const override;
 	virtual FString GetAdditionalTitleInformation() const override;
 
diff --git a/Engine/Plugins/PCG/Source/PCG/Public/Elements/PCGCreateSpline.h b/Engine/Plugins/PCG/Source/PCG/Public/Elements/PCGCreateSpline.h
--- a/Engine/Plugins/PCG/Source/PCG/Public/Elements/PCGCreateSpline.h
+++ b/Engine/Plugins/PCG/Source/PCG/Public/Elements/PCGCreateSpline.h
@@ -32,7 +32,7 @@ public:
 	virtual FText GetNodeTooltipText() const override;
 	virtual EPCGSettingsType GetType() const override { return EPCGSettingsType::Spatial; }
 #endif
-	virtual bool IsInputPinRequiredByExecution(const UPCGPin* InPin) const override { return !InPin->Properties.bAdvancedPin; }
+	
 
 	virtual TArray<FPCGPinProperties> InputPinProperties() const override { return Super::DefaultPointInputPinProperties(); }
 	virtual TArray<FPCGPinProperties> OutputPinProperties() const override;
diff --git a/Engine/Plugins/PCG/Source/PCG/Public/Elements/PCGCullPointsOutsideActorBounds.h b/Engine/Plugins/PCG/Source/PCG/Public/Elements/PCGCullPointsOutsideActorBounds.h
--- a/Engine/Plugins/PCG/Source/PCG/Public/Elements/PCGCullPointsOutsideActorBounds.h
+++ b/Engine/Plugins/PCG/Source/PCG/Public/Elements/PCGCullPointsOutsideActorBounds.h
@@ -22,7 +22,7 @@ public:
 	virtual FText GetNodeTooltipText() const override;
 	virtual EPCGSettingsType GetType() const override { return EPCGSettingsType::Spatial; }
 #endif
-	virtual bool IsInputPinRequiredByExecution(const UPCGPin* InPin) const override { return !InPin->Properties.bAdvancedPin; }
+	
 
 protected:
 	virtual TArray<FPCGPinProperties> InputPinProperties() const override;
diff --git a/Engine/Plugins/PCG/Source/PCG/Public/Elements/PCGDataNum.h b/Engine/Plugins/PCG/Source/PCG/Public/Elements/PCGDataNum.h
--- a/Engine/Plugins/PCG/Source/PCG/Public/Elements/PCGDataNum.h
+++ b/Engine/Plugins/PCG/Source/PCG/Public/Elements/PCGDataNum.h
@@ -20,7 +20,7 @@ public:
 	virtual FText GetNodeTooltipText() const override;
 	virtual EPCGSettingsType GetType() const override { return EPCGSettingsType::Param; }
 #endif
-	virtual bool IsInputPinRequiredByExecution(const UPCGPin* InPin) const override { return !InPin->Properties.bAdvancedPin; }
+	
 
 protected:
 	virtual TArray<FPCGPinProperties> OutputPinProperties() const override;
diff --git a/Engine/Plugins/PCG/Source/PCG/Public/Elements/PCGDebugElement.h b/Engine/Plugins/PCG/Source/PCG/Public/Elements/PCGDebugElement.h
--- a/Engine/Plugins/PCG/Source/PCG/Public/Elements/PCGDebugElement.h
+++ b/Engine/Plugins/PCG/Source/PCG/Public/Elements/PCGDebugElement.h
@@ -25,7 +25,7 @@ public:
 	virtual FText GetDefaultNodeTitle() const override { return NSLOCTEXT("PCGDebugSettings", "NodeTitle", "Debug"); }
 	virtual EPCGSettingsType GetType() const override { return EPCGSettingsType::Debug; }
 #endif
-	virtual bool IsInputPinRequiredByExecution(const UPCGPin* InPin) const override { return !InPin->Properties.bAdvancedPin; }
+	
 
 protected:
 	virtual TArray<FPCGPinProperties> InputPinProperties() const override;
diff --git a/Engine/Plugins/PCG/Source/PCG/Public/Elements/PCGDeleteAttributesElement.h b/Engine/Plugins/PCG/Source/PCG/Public/Elements/PCGDeleteAttributesElement.h
--- a/Engine/Plugins/PCG/Source/PCG/Public/Elements/PCGDeleteAttributesElement.h
+++ b/Engine/Plugins/PCG/Source/PCG/Public/Elements/PCGDeleteAttributesElement.h
@@ -39,14 +39,13 @@ public:
 	virtual EPCGSettingsType GetType() const override { return EPCGSettingsType::Metadata; }
 	virtual bool HasDynamicPins() const override { return true; }
 #endif
-	virtual bool IsInputPinRequiredByExecution(const UPCGPin* InPin) const override { return !InPin->Properties.bAdvancedPin; }
+	
 	virtual FString GetAdditionalTitleInformation() const override;
 
 protected:
 #if WITH_EDITOR
 	virtual EPCGChangeType GetChangeTypeForProperty(const FName& InPropertyName) const override { return Super::GetChangeTypeForProperty(InPropertyName) | EPCGChangeType::Cosmetic; }
 #endif
-	virtual TArray<FPCGPinProperties> InputPinProperties() const override;
 	virtual TArray<FPCGPinProperties> OutputPinProperties() const override;
 	virtual FPCGElementPtr CreateElement() const override;
 	//~End UPCGSettings interface
diff --git a/Engine/Plugins/PCG/Source/PCG/Public/Elements/PCGDeleteTags.h b/Engine/Plugins/PCG/Source/PCG/Public/Elements/PCGDeleteTags.h
--- a/Engine/Plugins/PCG/Source/PCG/Public/Elements/PCGDeleteTags.h
+++ b/Engine/Plugins/PCG/Source/PCG/Public/Elements/PCGDeleteTags.h
@@ -27,7 +27,7 @@ public:
 	virtual EPCGSettingsType GetType() const override { return EPCGSettingsType::Generic; }
 	virtual bool HasDynamicPins() const override { return true; }
 #endif
-	virtual bool IsInputPinRequiredByExecution(const UPCGPin* InPin) const override { return !InPin->Properties.bAdvancedPin; }
+	
 	virtual FString GetAdditionalTitleInformation() const override;
 
 protected:
diff --git a/Engine/Plugins/PCG/Source/PCG/Public/Elements/PCGDensityFilter.h b/Engine/Plugins/PCG/Source/PCG/Public/Elements/PCGDensityFilter.h
--- a/Engine/Plugins/PCG/Source/PCG/Public/Elements/PCGDensityFilter.h
+++ b/Engine/Plugins/PCG/Source/PCG/Public/Elements/PCGDensityFilter.h
@@ -18,7 +18,7 @@ public:
 	virtual FText GetDefaultNodeTitle() const override { return NSLOCTEXT("PCGDensityFilterSettings", "NodeTitle", "Density Filter"); }
 	virtual EPCGSettingsType GetType() const override { return EPCGSettingsType::Filter; }
 #endif
-	virtual bool IsInputPinRequiredByExecution(const UPCGPin* InPin) const override { return !InPin->Properties.bAdvancedPin; }
+	
 
 protected:
 	virtual TArray<FPCGPinProperties> InputPinProperties() const override { return Super::DefaultPointInputPinProperties(); }
diff --git a/Engine/Plugins/PCG/Source/PCG/Public/Elements/PCGDensityRemapElement.h b/Engine/Plugins/PCG/Source/PCG/Public/Elements/PCGDensityRemapElement.h
--- a/Engine/Plugins/PCG/Source/PCG/Public/Elements/PCGDensityRemapElement.h
+++ b/Engine/Plugins/PCG/Source/PCG/Public/Elements/PCGDensityRemapElement.h
@@ -20,7 +20,7 @@ public:
 	virtual FText GetDefaultNodeTitle() const override { return NSLOCTEXT("PCGDensityRemapSettings", "NodeTitle", "Density Remap"); }
 	virtual EPCGSettingsType GetType() const override { return EPCGSettingsType::Density; }
 #endif
-	virtual bool IsInputPinRequiredByExecution(const UPCGPin* InPin) const override { return !InPin->Properties.bAdvancedPin; }
+	
 
 protected:
 	virtual TArray<FPCGPinProperties> InputPinProperties() const override { return Super::DefaultPointInputPinProperties(); }
diff --git a/Engine/Plugins/PCG/Source/PCG/Public/Elements/PCGDifferenceElement.h b/Engine/Plugins/PCG/Source/PCG/Public/Elements/PCGDifferenceElement.h
--- a/Engine/Plugins/PCG/Source/PCG/Public/Elements/PCGDifferenceElement.h
+++ b/Engine/Plugins/PCG/Source/PCG/Public/Elements/PCGDifferenceElement.h
@@ -27,7 +27,6 @@ public:
 	virtual EPCGSettingsType GetType() const override { return EPCGSettingsType::Spatial; }
 	virtual void ApplyStructuralDeprecation(UPCGNode* InOutNode) override;
 #endif
-	virtual bool IsInputPinRequiredByExecution(const UPCGPin* InPin) const override { return InPin->Properties.Label == PCGDifferenceConstants::SourceLabel; }
 	virtual FString GetAdditionalTitleInformation() const override;
 
 protected:
diff --git a/Engine/Plugins/PCG/Source/PCG/Public/Elements/PCGDistance.h b/Engine/Plugins/PCG/Source/PCG/Public/Elements/PCGDistance.h
--- a/Engine/Plugins/PCG/Source/PCG/Public/Elements/PCGDistance.h
+++ b/Engine/Plugins/PCG/Source/PCG/Public/Elements/PCGDistance.h
@@ -41,7 +41,6 @@ public:
 	virtual FText GetNodeTooltipText() const override;
 	virtual EPCGSettingsType GetType() const override { return EPCGSettingsType::Spatial; }
 #endif
-	virtual bool IsInputPinRequiredByExecution(const UPCGPin* InPin) const override { return InPin->Properties.Label == PCGDistance::SourceLabel; }
 
 protected:
 	virtual TArray<FPCGPinProperties> InputPinProperties() const override;
diff --git a/Engine/Plugins/PCG/Source/PCG/Public/Elements/PCGDuplicatePoint.h b/Engine/Plugins/PCG/Source/PCG/Public/Elements/PCGDuplicatePoint.h
--- a/Engine/Plugins/PCG/Source/PCG/Public/Elements/PCGDuplicatePoint.h
+++ b/Engine/Plugins/PCG/Source/PCG/Public/Elements/PCGDuplicatePoint.h
@@ -22,7 +22,7 @@ public:
 	virtual FText GetNodeTooltipText() const override { return NSLOCTEXT("PCGDuplicatePointElement", "NodeTooltip", "Creates duplicates of each point with optional transform offsets."); }
 	virtual EPCGSettingsType GetType() const override { return EPCGSettingsType::PointOps; }
 #endif
-	virtual bool IsInputPinRequiredByExecution(const UPCGPin* InPin) const override { return !InPin->Properties.bAdvancedPin; }
+	
 
 protected:
 	virtual TArray<FPCGPinProperties> InputPinProperties() const override { return Super::DefaultPointInputPinProperties(); }
diff --git a/Engine/Plugins/PCG/Source/PCG/Public/Elements/PCGFilterByAttribute.h b/Engine/Plugins/PCG/Source/PCG/Public/Elements/PCGFilterByAttribute.h
--- a/Engine/Plugins/PCG/Source/PCG/Public/Elements/PCGFilterByAttribute.h
+++ b/Engine/Plugins/PCG/Source/PCG/Public/Elements/PCGFilterByAttribute.h
@@ -19,7 +19,7 @@ public:
 	virtual FText GetDefaultNodeTitle() const override;
 	virtual FText GetNodeTooltipText() const override { return NSLOCTEXT("PCGFilterByAttributeElement", "NodeTooltip", "Separates input data by whether they have the specified attribute or not."); }
 #endif
-	virtual bool IsInputPinRequiredByExecution(const UPCGPin* InPin) const override { return !InPin->Properties.bAdvancedPin; }
+	
 	virtual FString GetAdditionalTitleInformation() const override;
 
 protected:
diff --git a/Engine/Plugins/PCG/Source/PCG/Public/Elements/PCGFilterByIndex.h b/Engine/Plugins/PCG/Source/PCG/Public/Elements/PCGFilterByIndex.h
--- a/Engine/Plugins/PCG/Source/PCG/Public/Elements/PCGFilterByIndex.h
+++ b/Engine/Plugins/PCG/Source/PCG/Public/Elements/PCGFilterByIndex.h
@@ -19,7 +19,7 @@ public:
 	virtual FText GetDefaultNodeTitle() const override;
 	virtual FText GetNodeTooltipText() const override;
 #endif
-	virtual bool IsInputPinRequiredByExecution(const UPCGPin* InPin) const override { return !InPin->Properties.bAdvancedPin; }
+	
 	virtual FString GetAdditionalTitleInformation() const override;
 
 protected:
diff --git a/Engine/Plugins/PCG/Source/PCG/Public/Elements/PCGFilterByTag.h b/Engine/Plugins/PCG/Source/PCG/Public/Elements/PCGFilterByTag.h
--- a/Engine/Plugins/PCG/Source/PCG/Public/Elements/PCGFilterByTag.h
+++ b/Engine/Plugins/PCG/Source/PCG/Public/Elements/PCGFilterByTag.h
@@ -26,7 +26,7 @@ public:
 	virtual FText GetDefaultNodeTitle() const override;
 	virtual FText GetNodeTooltipText() const override;
 #endif
-	virtual bool IsInputPinRequiredByExecution(const UPCGPin* InPin) const override { return !InPin->Properties.bAdvancedPin; }
+	
 	virtual FString GetAdditionalTitleInformation() const override;
 
 protected:
diff --git a/Engine/Plugins/PCG/Source/PCG/Public/Elements/PCGFilterByType.h b/Engine/Plugins/PCG/Source/PCG/Public/Elements/PCGFilterByType.h
--- a/Engine/Plugins/PCG/Source/PCG/Public/Elements/PCGFilterByType.h
+++ b/Engine/Plugins/PCG/Source/PCG/Public/Elements/PCGFilterByType.h
@@ -21,7 +21,7 @@ public:
 	virtual bool ShouldDrawNodeCompact() const override { return true; }
 	virtual bool GetCompactNodeIcon(FName& OutCompactNodeIcon) const override;
 #endif
-	virtual bool IsInputPinRequiredByExecution(const UPCGPin* InPin) const override { return !InPin->Properties.bAdvancedPin; }
+	
 	virtual EPCGDataType GetCurrentPinTypes(const UPCGPin* InPin) const override;
 
 protected:
diff --git a/Engine/Plugins/PCG/Source/PCG/Public/Elements/PCGFilterDataBase.h b/Engine/Plugins/PCG/Source/PCG/Public/Elements/PCGFilterDataBase.h
--- a/Engine/Plugins/PCG/Source/PCG/Public/Elements/PCGFilterDataBase.h
+++ b/Engine/Plugins/PCG/Source/PCG/Public/Elements/PCGFilterDataBase.h
@@ -20,7 +20,6 @@ public:
 #endif // WITH_EDITOR
 
 protected:
-	virtual TArray<FPCGPinProperties> InputPinProperties() const override;
 	virtual TArray<FPCGPinProperties> OutputPinProperties() const override;
 	//~End UPCGSettings interface
 };
\ No newline at end of file
diff --git a/Engine/Plugins/PCG/Source/PCG/Public/Elements/PCGGather.h b/Engine/Plugins/PCG/Source/PCG/Public/Elements/PCGGather.h
--- a/Engine/Plugins/PCG/Source/PCG/Public/Elements/PCGGather.h
+++ b/Engine/Plugins/PCG/Source/PCG/Public/Elements/PCGGather.h
@@ -22,7 +22,6 @@ public:
 	virtual EPCGSettingsType GetType() const override { return EPCGSettingsType::Spatial; }
 	virtual bool HasDynamicPins() const override { return true; }
 #endif
-	virtual bool IsInputPinRequiredByExecution(const UPCGPin* InPin) const override { return InPin->Properties.Label == PCGPinConstants::DefaultInputLabel; }
 
 protected:
 	virtual TArray<FPCGPinProperties> InputPinProperties() const override;
diff --git a/Engine/Plugins/PCG/Source/PCG/Public/Elements/PCGGetBounds.h b/Engine/Plugins/PCG/Source/PCG/Public/Elements/PCGGetBounds.h
--- a/Engine/Plugins/PCG/Source/PCG/Public/Elements/PCGGetBounds.h
+++ b/Engine/Plugins/PCG/Source/PCG/Public/Elements/PCGGetBounds.h
@@ -22,7 +22,7 @@ public:
 	virtual FText GetNodeTooltipText() const override;
 	virtual EPCGSettingsType GetType() const override { return EPCGSettingsType::Spatial; }
 #endif
-	virtual bool IsInputPinRequiredByExecution(const UPCGPin* InPin) const override { return !InPin->Properties.bAdvancedPin; }
+	
 
 protected:
 	virtual TArray<FPCGPinProperties> InputPinProperties() const override;
diff --git a/Engine/Plugins/PCG/Source/PCG/Public/Elements/PCGInnerIntersectionElement.h b/Engine/Plugins/PCG/Source/PCG/Public/Elements/PCGInnerIntersectionElement.h
--- a/Engine/Plugins/PCG/Source/PCG/Public/Elements/PCGInnerIntersectionElement.h
+++ b/Engine/Plugins/PCG/Source/PCG/Public/Elements/PCGInnerIntersectionElement.h
@@ -20,7 +20,7 @@ public:
 	virtual FText GetNodeTooltipText() const override;
 	virtual EPCGSettingsType GetType() const override { return EPCGSettingsType::Spatial; }
 #endif
-	virtual bool IsInputPinRequiredByExecution(const UPCGPin* InPin) const override { return !InPin->Properties.bAdvancedPin; }
+	
 
 	virtual TArray<FPCGPinProperties> InputPinProperties() const override;
 	virtual TArray<FPCGPinProperties> OutputPinProperties() const override;
diff --git a/Engine/Plugins/PCG/Source/PCG/Public/Elements/PCGMakeConcreteElement.h b/Engine/Plugins/PCG/Source/PCG/Public/Elements/PCGMakeConcreteElement.h
--- a/Engine/Plugins/PCG/Source/PCG/Public/Elements/PCGMakeConcreteElement.h
+++ b/Engine/Plugins/PCG/Source/PCG/Public/Elements/PCGMakeConcreteElement.h
@@ -21,7 +21,7 @@ public:
 	virtual EPCGSettingsType GetType() const override { return EPCGSettingsType::Spatial; }
 	virtual bool ShouldDrawNodeCompact() const override { return true; }
 #endif
-	virtual bool IsInputPinRequiredByExecution(const UPCGPin* InPin) const override { return !InPin->Properties.bAdvancedPin; }
+	
 
 protected:
 	virtual TArray<FPCGPinProperties> OutputPinProperties() const override;
diff --git a/Engine/Plugins/PCG/Source/PCG/Public/Elements/PCGMatchAndSetAttributes.h b/Engine/Plugins/PCG/Source/PCG/Public/Elements/PCGMatchAndSetAttributes.h
--- a/Engine/Plugins/PCG/Source/PCG/Public/Elements/PCGMatchAndSetAttributes.h
+++ b/Engine/Plugins/PCG/Source/PCG/Public/Elements/PCGMatchAndSetAttributes.h
@@ -55,7 +55,6 @@ public:
 	virtual EPCGSettingsType GetType() const override { return EPCGSettingsType::Metadata; }
 	virtual bool HasDynamicPins() const { return true; }
 #endif // WITH_EDITOR
-	virtual bool IsInputPinRequiredByExecution(const UPCGPin* InPin) const override { return InPin->Properties.Label == PCGPinConstants::DefaultInputLabel; }
 
 protected:
 	virtual TArray<FPCGPinProperties> InputPinProperties() const override;
diff --git a/Engine/Plugins/PCG/Source/PCG/Public/Elements/PCGMergeAttributes.h b/Engine/Plugins/PCG/Source/PCG/Public/Elements/PCGMergeAttributes.h
--- a/Engine/Plugins/PCG/Source/PCG/Public/Elements/PCGMergeAttributes.h
+++ b/Engine/Plugins/PCG/Source/PCG/Public/Elements/PCGMergeAttributes.h
@@ -20,7 +20,7 @@ public:
 	virtual FText GetNodeTooltipText() const override;
 	virtual EPCGSettingsType GetType() const override { return EPCGSettingsType::Metadata; }
 #endif // WITH_EDITOR
-	virtual bool IsInputPinRequiredByExecution(const UPCGPin* InPin) const override { return !InPin->Properties.bAdvancedPin; }
+	
 
 protected:
 	virtual TArray<FPCGPinProperties> InputPinProperties() const override;
diff --git a/Engine/Plugins/PCG/Source/PCG/Public/Elements/PCGMergeElement.h b/Engine/Plugins/PCG/Source/PCG/Public/Elements/PCGMergeElement.h
--- a/Engine/Plugins/PCG/Source/PCG/Public/Elements/PCGMergeElement.h
+++ b/Engine/Plugins/PCG/Source/PCG/Public/Elements/PCGMergeElement.h
@@ -20,7 +20,7 @@ public:
 	virtual FText GetNodeTooltipText() const override;
 	virtual EPCGSettingsType GetType() const override { return EPCGSettingsType::Spatial; }
 #endif
-	virtual bool IsInputPinRequiredByExecution(const UPCGPin* InPin) const override { return !InPin->Properties.bAdvancedPin; }
+	
 
 	virtual TArray<FPCGPinProperties> InputPinProperties() const override;
 	virtual TArray<FPCGPinProperties> OutputPinProperties() const override;
diff --git a/Engine/Plugins/PCG/Source/PCG/Public/Elements/PCGMutateSeed.h b/Engine/Plugins/PCG/Source/PCG/Public/Elements/PCGMutateSeed.h
--- a/Engine/Plugins/PCG/Source/PCG/Public/Elements/PCGMutateSeed.h
+++ b/Engine/Plugins/PCG/Source/PCG/Public/Elements/PCGMutateSeed.h
@@ -25,7 +25,7 @@ public:
 	virtual FText GetNodeTooltipText() const override { return NSLOCTEXT("PCGMutateSeedElement", "NodeTooltip", "Applies a new random seed from point input."); }
 	virtual EPCGSettingsType GetType() const override { return EPCGSettingsType::Spatial; }
 #endif
-	virtual bool IsInputPinRequiredByExecution(const UPCGPin* InPin) const override { return !InPin->Properties.bAdvancedPin; }
+	
 
 protected:
 	virtual TArray<FPCGPinProperties> InputPinProperties() const override { return Super::DefaultPointInputPinProperties(); }
diff --git a/Engine/Plugins/PCG/Source/PCG/Public/Elements/PCGNormalToDensity.h b/Engine/Plugins/PCG/Source/PCG/Public/Elements/PCGNormalToDensity.h
--- a/Engine/Plugins/PCG/Source/PCG/Public/Elements/PCGNormalToDensity.h
+++ b/Engine/Plugins/PCG/Source/PCG/Public/Elements/PCGNormalToDensity.h
@@ -33,7 +33,7 @@ public:
 	virtual FText GetDefaultNodeTitle() const override { return NSLOCTEXT("PCGNormalToDensitySettings", "NodeTitle", "Normal To Density"); }
 	virtual EPCGSettingsType GetType() const override { return EPCGSettingsType::Spatial; }
 #endif
-	virtual bool IsInputPinRequiredByExecution(const UPCGPin* InPin) const override { return !InPin->Properties.bAdvancedPin; }
+	
 
 protected:
 	virtual TArray<FPCGPinProperties> InputPinProperties() const override { return Super::DefaultPointInputPinProperties(); }
diff --git a/Engine/Plugins/PCG/Source/PCG/Public/Elements/PCGNumberOfElements.h b/Engine/Plugins/PCG/Source/PCG/Public/Elements/PCGNumberOfElements.h
--- a/Engine/Plugins/PCG/Source/PCG/Public/Elements/PCGNumberOfElements.h
+++ b/Engine/Plugins/PCG/Source/PCG/Public/Elements/PCGNumberOfElements.h
@@ -21,7 +21,7 @@ class UPCGNumberOfElementsBaseSettings : public UPCGSettings
 	GENERATED_BODY()
 
 	//~Begin UPCGSettings interface
-	virtual bool IsInputPinRequiredByExecution(const UPCGPin* InPin) const override { return !InPin->Properties.bAdvancedPin; }
+	
 
 protected:
 	virtual TArray<FPCGPinProperties> OutputPinProperties() const override;
diff --git a/Engine/Plugins/PCG/Source/PCG/Public/Elements/PCGOuterIntersectionElement.h b/Engine/Plugins/PCG/Source/PCG/Public/Elements/PCGOuterIntersectionElement.h
--- a/Engine/Plugins/PCG/Source/PCG/Public/Elements/PCGOuterIntersectionElement.h
+++ b/Engine/Plugins/PCG/Source/PCG/Public/Elements/PCGOuterIntersectionElement.h
@@ -30,7 +30,6 @@ public:
 	virtual EPCGSettingsType GetType() const override { return EPCGSettingsType::Spatial; }
 	virtual bool HasDynamicPins() const override { return true; }
 #endif // WITH_EDITOR
-	virtual bool IsInputPinRequiredByExecution(const UPCGPin* InPin) const override { return InPin->Properties.Label == PCGIntersectionConstants::PrimaryLabel; }
 	// If node disabled, don't intersect - pass through all primary edges
 	virtual bool OnlyPassThroughOneEdgeWhenDisabled() const override { return false; }
 
diff --git a/Engine/Plugins/PCG/Source/PCG/Public/Elements/PCGPointExtentsModifier.h b/Engine/Plugins/PCG/Source/PCG/Public/Elements/PCGPointExtentsModifier.h
--- a/Engine/Plugins/PCG/Source/PCG/Public/Elements/PCGPointExtentsModifier.h
+++ b/Engine/Plugins/PCG/Source/PCG/Public/Elements/PCGPointExtentsModifier.h
@@ -31,7 +31,7 @@ public:
 	virtual FText GetDefaultNodeTitle() const override { return NSLOCTEXT("PCGPointExtentsModifierSettings", "NodeTitle", "Extents Modifier"); }
 	virtual EPCGSettingsType GetType() const override { return EPCGSettingsType::PointOps; }
 #endif
-	virtual bool IsInputPinRequiredByExecution(const UPCGPin* InPin) const override { return !InPin->Properties.bAdvancedPin; }
+	
 
 protected:
 	virtual TArray<FPCGPinProperties> InputPinProperties() const override { return Super::DefaultPointInputPinProperties(); }
diff --git a/Engine/Plugins/PCG/Source/PCG/Public/Elements/PCGPointMatchAndSet.h b/Engine/Plugins/PCG/Source/PCG/Public/Elements/PCGPointMatchAndSet.h
--- a/Engine/Plugins/PCG/Source/PCG/Public/Elements/PCGPointMatchAndSet.h
+++ b/Engine/Plugins/PCG/Source/PCG/Public/Elements/PCGPointMatchAndSet.h
@@ -30,7 +30,7 @@ public:
 	virtual EPCGSettingsType GetType() const override { return EPCGSettingsType::Metadata; }
 	virtual void ApplyDeprecation(UPCGNode* InOutNode) override;
 #endif // WITH_EDITOR
-	virtual bool IsInputPinRequiredByExecution(const UPCGPin* InPin) const override { return !InPin->Properties.bAdvancedPin; }
+	
 
 protected:
 	virtual TArray<FPCGPinProperties> InputPinProperties() const override;
diff --git a/Engine/Plugins/PCG/Source/PCG/Public/Elements/PCGPointNeighborhood.h b/Engine/Plugins/PCG/Source/PCG/Public/Elements/PCGPointNeighborhood.h
--- a/Engine/Plugins/PCG/Source/PCG/Public/Elements/PCGPointNeighborhood.h
+++ b/Engine/Plugins/PCG/Source/PCG/Public/Elements/PCGPointNeighborhood.h
@@ -29,7 +29,7 @@ public:
 	virtual FText GetNodeTooltipText() const override { return NSLOCTEXT("PCGPointNeighborhoodElement", "NodeTooltip", "Computes quantities from nearby neighbor points, such as average density, color, and position."); }
 	virtual EPCGSettingsType GetType() const override { return EPCGSettingsType::Spatial; }
 #endif
-	virtual bool IsInputPinRequiredByExecution(const UPCGPin* InPin) const override { return !InPin->Properties.bAdvancedPin; }
+	
 
 protected:
 	virtual TArray<FPCGPinProperties> InputPinProperties() const override { return Super::DefaultPointInputPinProperties(); }
diff --git a/Engine/Plugins/PCG/Source/PCG/Public/Elements/PCGProjectionElement.h b/Engine/Plugins/PCG/Source/PCG/Public/Elements/PCGProjectionElement.h
--- a/Engine/Plugins/PCG/Source/PCG/Public/Elements/PCGProjectionElement.h
+++ b/Engine/Plugins/PCG/Source/PCG/Public/Elements/PCGProjectionElement.h
@@ -31,7 +31,7 @@ public:
 	virtual bool HasDynamicPins() const override { return true; }
 #endif
 	virtual EPCGDataType GetCurrentPinTypes(const UPCGPin* InPin) const override;
-	virtual bool IsInputPinRequiredByExecution(const UPCGPin* InPin) const override { return !InPin->Properties.bAdvancedPin; }
+	
 
 protected:
 	virtual TArray<FPCGPinProperties> InputPinProperties() const override;
diff --git a/Engine/Plugins/PCG/Source/PCG/Public/Elements/PCGReplaceTags.h b/Engine/Plugins/PCG/Source/PCG/Public/Elements/PCGReplaceTags.h
--- a/Engine/Plugins/PCG/Source/PCG/Public/Elements/PCGReplaceTags.h
+++ b/Engine/Plugins/PCG/Source/PCG/Public/Elements/PCGReplaceTags.h
@@ -20,7 +20,7 @@ public:
 	virtual EPCGSettingsType GetType() const override { return EPCGSettingsType::Generic; }
 	virtual bool HasDynamicPins() const override { return true; }
 #endif
-	virtual bool IsInputPinRequiredByExecution(const UPCGPin* InPin) const override { return !InPin->Properties.bAdvancedPin; }
+	
 	virtual FString GetAdditionalTitleInformation() const override;
 
 protected:
diff --git a/Engine/Plugins/PCG/Source/PCG/Public/Elements/PCGResetPointCenter.h b/Engine/Plugins/PCG/Source/PCG/Public/Elements/PCGResetPointCenter.h
--- a/Engine/Plugins/PCG/Source/PCG/Public/Elements/PCGResetPointCenter.h
+++ b/Engine/Plugins/PCG/Source/PCG/Public/Elements/PCGResetPointCenter.h
@@ -23,7 +23,7 @@ public:
 	virtual FText GetNodeTooltipText() const override { return NSLOCTEXT("PCGResetPointCenterElement", "NodeTooltip", "Modify the position of a point within its bounds, while keeping its bounds the same."); }
 	virtual EPCGSettingsType GetType() const override { return EPCGSettingsType::PointOps; }
 #endif
-	virtual bool IsInputPinRequiredByExecution(const UPCGPin* InPin) const override { return !InPin->Properties.bAdvancedPin; }
+	
 
 protected:
 	virtual TArray<FPCGPinProperties> InputPinProperties() const override { return Super::DefaultPointInputPinProperties(); }
diff --git a/Engine/Plugins/PCG/Source/PCG/Public/Elements/PCGSelectPoints.h b/Engine/Plugins/PCG/Source/PCG/Public/Elements/PCGSelectPoints.h
--- a/Engine/Plugins/PCG/Source/PCG/Public/Elements/PCGSelectPoints.h
+++ b/Engine/Plugins/PCG/Source/PCG/Public/Elements/PCGSelectPoints.h
@@ -21,7 +21,7 @@ public:
 	virtual FText GetNodeTooltipText() const override;
 	virtual EPCGSettingsType GetType() const override { return EPCGSettingsType::Sampler; }
 #endif
-	virtual bool IsInputPinRequiredByExecution(const UPCGPin* InPin) const override { return !InPin->Properties.bAdvancedPin; }
+	
 
 protected:
 	virtual TArray<FPCGPinProperties> InputPinProperties() const override { return Super::DefaultPointInputPinProperties(); }
diff --git a/Engine/Plugins/PCG/Source/PCG/Public/Elements/PCGSelfPruning.h b/Engine/Plugins/PCG/Source/PCG/Public/Elements/PCGSelfPruning.h
--- a/Engine/Plugins/PCG/Source/PCG/Public/Elements/PCGSelfPruning.h
+++ b/Engine/Plugins/PCG/Source/PCG/Public/Elements/PCGSelfPruning.h
@@ -98,7 +98,7 @@ public:
 	virtual FText GetDefaultNodeTitle() const override { return NSLOCTEXT("PCGSelfPruningSettings", "NodeTitle", "Self Pruning"); }
 	virtual EPCGSettingsType GetType() const override { return EPCGSettingsType::Filter; }
 #endif
-	virtual bool IsInputPinRequiredByExecution(const UPCGPin* InPin) const override { return !InPin->Properties.bAdvancedPin; }
+	
 
 protected:
 	virtual TArray<FPCGPinProperties> InputPinProperties() const override { return Super::DefaultPointInputPinProperties(); }
diff --git a/Engine/Plugins/PCG/Source/PCG/Public/Elements/PCGSortAttributes.h b/Engine/Plugins/PCG/Source/PCG/Public/Elements/PCGSortAttributes.h
--- a/Engine/Plugins/PCG/Source/PCG/Public/Elements/PCGSortAttributes.h
+++ b/Engine/Plugins/PCG/Source/PCG/Public/Elements/PCGSortAttributes.h
@@ -32,7 +32,7 @@ public:
 	virtual EPCGSettingsType GetType() const override { return EPCGSettingsType::Generic; }
 	virtual bool HasDynamicPins() const override { return true; }
 #endif
-	virtual bool IsInputPinRequiredByExecution(const UPCGPin* InPin) const override { return !InPin->Properties.bAdvancedPin; }
+	
 
 protected:
 	virtual TArray<FPCGPinProperties> OutputPinProperties() const override;
diff --git a/Engine/Plugins/PCG/Source/PCG/Public/Elements/PCGSpatialNoise.h b/Engine/Plugins/PCG/Source/PCG/Public/Elements/PCGSpatialNoise.h
--- a/Engine/Plugins/PCG/Source/PCG/Public/Elements/PCGSpatialNoise.h
+++ b/Engine/Plugins/PCG/Source/PCG/Public/Elements/PCGSpatialNoise.h
@@ -62,7 +62,7 @@ public:
 	virtual FText GetDefaultNodeTitle() const override { return NSLOCTEXT("PCGSpatialNoise", "NodeTitle", "Spatial Noise"); }
 	virtual EPCGSettingsType GetType() const override { return EPCGSettingsType::Spatial; }
 #endif
-	virtual bool IsInputPinRequiredByExecution(const UPCGPin* InPin) const override { return !InPin->Properties.bAdvancedPin; }
+	
 
 protected:
 	virtual TArray<FPCGPinProperties> InputPinProperties() const override;
diff --git a/Engine/Plugins/PCG/Source/PCG/Public/Elements/PCGSpawnActor.h b/Engine/Plugins/PCG/Source/PCG/Public/Elements/PCGSpawnActor.h
--- a/Engine/Plugins/PCG/Source/PCG/Public/Elements/PCGSpawnActor.h
+++ b/Engine/Plugins/PCG/Source/PCG/Public/Elements/PCGSpawnActor.h
@@ -110,7 +110,7 @@ public:
 #endif // WITH_EDITOR
 
 	//~Begin UCPGSettings interface
-	virtual bool IsInputPinRequiredByExecution(const UPCGPin* InPin) const override { return !InPin->Properties.bAdvancedPin; }
+	
 	virtual UPCGNode* CreateNode() const override;
 
 #if WITH_EDITOR	
diff --git a/Engine/Plugins/PCG/Source/PCG/Public/Elements/PCGSplineSampler.h b/Engine/Plugins/PCG/Source/PCG/Public/Elements/PCGSplineSampler.h
--- a/Engine/Plugins/PCG/Source/PCG/Public/Elements/PCGSplineSampler.h
+++ b/Engine/Plugins/PCG/Source/PCG/Public/Elements/PCGSplineSampler.h
@@ -225,7 +225,6 @@ public:
 	virtual void ApplyDeprecation(UPCGNode* InOutNode) override;
 	virtual void ApplyDeprecationBeforeUpdatePins(UPCGNode* InOutNode, TArray<TObjectPtr<UPCGPin>>& InputPins, TArray<TObjectPtr<UPCGPin>>& OutputPins) override;
 #endif
-	virtual bool IsInputPinRequiredByExecution(const UPCGPin* InPin) const override { return InPin->Properties.Label == PCGSplineSamplerConstants::SplineLabel; }
 
 protected:
 	virtual TArray<FPCGPinProperties> InputPinProperties() const override;
diff --git a/Engine/Plugins/PCG/Source/PCG/Public/Elements/PCGSplitPoints.h b/Engine/Plugins/PCG/Source/PCG/Public/Elements/PCGSplitPoints.h
--- a/Engine/Plugins/PCG/Source/PCG/Public/Elements/PCGSplitPoints.h
+++ b/Engine/Plugins/PCG/Source/PCG/Public/Elements/PCGSplitPoints.h
@@ -30,7 +30,7 @@ public:
 	virtual FText GetNodeTooltipText() const override { return NSLOCTEXT("PCGSplitPointsElement", "NodeTooltip", "Splits each input point into two separate points and sets bounds based on the position and axis of the cut."); }
 	virtual EPCGSettingsType GetType() const override { return EPCGSettingsType::PointOps; }
 #endif
-	virtual bool IsInputPinRequiredByExecution(const UPCGPin* InPin) const override { return !InPin->Properties.bAdvancedPin; }
+	
 
 protected:
 	virtual TArray<FPCGPinProperties> InputPinProperties() const override { return Super::DefaultPointInputPinProperties();  }
diff --git a/Engine/Plugins/PCG/Source/PCG/Public/Elements/PCGStaticMeshSpawner.h b/Engine/Plugins/PCG/Source/PCG/Public/Elements/PCGStaticMeshSpawner.h
--- a/Engine/Plugins/PCG/Source/PCG/Public/Elements/PCGStaticMeshSpawner.h
+++ b/Engine/Plugins/PCG/Source/PCG/Public/Elements/PCGStaticMeshSpawner.h
@@ -53,7 +53,7 @@ public:
 	virtual EPCGSettingsType GetType() const override { return EPCGSettingsType::Spawner; }
 	virtual void ApplyDeprecation(UPCGNode* InOutNode) override;
 #endif
-	virtual bool IsInputPinRequiredByExecution(const UPCGPin* InPin) const override { return !InPin->Properties.bAdvancedPin; }
+	
 
 protected:
 	virtual TArray<FPCGPinProperties> InputPinProperties() const override { return Super::DefaultPointInputPinProperties(); }
diff --git a/Engine/Plugins/PCG/Source/PCG/Public/Elements/PCGSurfaceSampler.h b/Engine/Plugins/PCG/Source/PCG/Public/Elements/PCGSurfaceSampler.h
--- a/Engine/Plugins/PCG/Source/PCG/Public/Elements/PCGSurfaceSampler.h
+++ b/Engine/Plugins/PCG/Source/PCG/Public/Elements/PCGSurfaceSampler.h
@@ -96,7 +96,6 @@ public:
 	virtual void ApplyDeprecationBeforeUpdatePins(UPCGNode* InOutNode, TArray<TObjectPtr<UPCGPin>>& InputPins, TArray<TObjectPtr<UPCGPin>>& OutputPins) override;
 	virtual void ApplyDeprecation(UPCGNode* InOutNode) override;
 #endif
-	virtual bool IsInputPinRequiredByExecution(const UPCGPin* InPin) const override { return InPin->Properties.Label == PCGSurfaceSamplerConstants::SurfaceLabel; }
 
 protected:
 	virtual TArray<FPCGPinProperties> InputPinProperties() const override;
diff --git a/Engine/Plugins/PCG/Source/PCG/Public/Elements/PCGTransformPoints.h b/Engine/Plugins/PCG/Source/PCG/Public/Elements/PCGTransformPoints.h
--- a/Engine/Plugins/PCG/Source/PCG/Public/Elements/PCGTransformPoints.h
+++ b/Engine/Plugins/PCG/Source/PCG/Public/Elements/PCGTransformPoints.h
@@ -20,7 +20,7 @@ public:
 	virtual FText GetDefaultNodeTitle() const override { return NSLOCTEXT("PCGTransformPointsSettings", "NodeTitle", "Transform Points"); }
 	virtual EPCGSettingsType GetType() const override { return EPCGSettingsType::PointOps; }
 #endif
-	virtual bool IsInputPinRequiredByExecution(const UPCGPin* InPin) const override { return !InPin->Properties.bAdvancedPin; }
+	
 
 protected:
 	virtual TArray<FPCGPinProperties> InputPinProperties() const override { return Super::DefaultPointInputPinProperties(); }
diff --git a/Engine/Plugins/PCG/Source/PCG/Public/Elements/PCGUnionElement.h b/Engine/Plugins/PCG/Source/PCG/Public/Elements/PCGUnionElement.h
--- a/Engine/Plugins/PCG/Source/PCG/Public/Elements/PCGUnionElement.h
+++ b/Engine/Plugins/PCG/Source/PCG/Public/Elements/PCGUnionElement.h
@@ -19,7 +19,7 @@ public:
 	virtual FText GetDefaultNodeTitle() const override { return NSLOCTEXT("PCGUnionSettings", "NodeTitle", "Union"); }
 	virtual EPCGSettingsType GetType() const override { return EPCGSettingsType::Spatial; }
 #endif
-	virtual bool IsInputPinRequiredByExecution(const UPCGPin* InPin) const override { return !InPin->Properties.bAdvancedPin; }
+	
 
 protected:
 	virtual TArray<FPCGPinProperties> InputPinProperties() const override;
diff --git a/Engine/Plugins/PCG/Source/PCG/Public/Elements/PCGVolumeSampler.h b/Engine/Plugins/PCG/Source/PCG/Public/Elements/PCGVolumeSampler.h
--- a/Engine/Plugins/PCG/Source/PCG/Public/Elements/PCGVolumeSampler.h
+++ b/Engine/Plugins/PCG/Source/PCG/Public/Elements/PCGVolumeSampler.h
@@ -83,7 +83,6 @@ public:
 	virtual EPCGSettingsType GetType() const override { return EPCGSettingsType::Sampler; }
 	virtual void ApplyDeprecationBeforeUpdatePins(UPCGNode* InOutNode, TArray<TObjectPtr<UPCGPin>>& InputPins, TArray<TObjectPtr<UPCGPin>>& OutputPins) override;
 #endif
-	virtual bool IsInputPinRequiredByExecution(const UPCGPin* InPin) const override { return InPin->Properties.Label == PCGVolumeSamplerConstants::VolumeLabel; }
 
 protected:
 	virtual TArray<FPCGPinProperties> InputPinProperties() const override;
diff --git a/Engine/Plugins/PCG/Source/PCG/Public/PCGPin.h b/Engine/Plugins/PCG/Source/PCG/Public/PCGPin.h
--- a/Engine/Plugins/PCG/Source/PCG/Public/PCGPin.h
+++ b/Engine/Plugins/PCG/Source/PCG/Public/PCGPin.h
@@ -14,12 +14,26 @@ class UPCGEdge;
 UENUM(BlueprintType)
 enum class EPCGPinUsage : uint8
 {
-	Normal = 0, // Normal usage pin, will pass all data as is.
-	Loop, // When used in a loop subgraph node, will separate each data from that pin into separate subgraph executions.
-	Feedback, // When used in a loop subgraph node, will pass data on the feedback pins to the next iteration only if the data is passed from a previous iteration (or the original subgraph call).
+	/** Normal usage pin, will pass all data as is. */
+	Normal = 0,
+	/** When used in a loop subgraph node, will separate each data from that pin into separate subgraph executions. */
+	Loop,
+	/** When used in a loop subgraph node, will pass data on the feedback pins to the next iteration only if the data is passed from a previous iteration (or the original subgraph call). */
+	Feedback,
 	DependencyOnly UMETA(Hidden)
 };
 
+UENUM(BlueprintType)
+enum class EPCGPinStatus : uint8
+{
+	/** Normal usage pin. */
+	Normal = 0,
+	/** Only for input pins, mark this pin as required.If set on an output pin, behave as Normal. */
+	Required,
+	/** Advanced pin will be hidden by default in the UIand will be shown only if the user extends the node(in the UI) to see advanced pins.Pins can't be required and advanced at the same time. */
+	Advanced
+};
+
 USTRUCT(BlueprintType, meta=(HasNativeBreak="/Script/PCG.PCGBlueprintPinHelpers.BreakPinProperty", HasNativeMake="/Script/PCG.PCGBlueprintPinHelpers.MakePinProperty"))
 struct PCG_API FPCGPinProperties
 {
@@ -42,9 +56,9 @@ struct PCG_API FPCGPinProperties
 	UPROPERTY(BlueprintReadWrite, EditAnywhere, Category = Settings)
 	bool bAllowMultipleData = true;
 
-	/* Advanced pin will be hidden by default in the UI and will be shown only if the user extend the node (in the UI) to see advanced pins. */
+	/** Define the status of a given pin (Normal, Required or Advanced) */
 	UPROPERTY(BlueprintReadWrite, EditAnywhere, Category = Settings)
-	bool bAdvancedPin = false;
+	EPCGPinStatus PinStatus = EPCGPinStatus::Normal;
 
 	UPROPERTY(BlueprintReadWrite, Category = Settings)
 	bool bInvisiblePin = false;
@@ -60,11 +74,38 @@ struct PCG_API FPCGPinProperties
 	// Allowing multiple connections will automatically enable multi data.
 	void SetAllowMultipleConnections(bool bInAllowMultipleConnectons);
 
+	bool IsAdvancedPin() const { return PinStatus == EPCGPinStatus::Advanced; }
+	void SetAdvancedPin() { PinStatus = EPCGPinStatus::Advanced; }
+
+	bool IsRequiredPin() const { return PinStatus == EPCGPinStatus::Required; }
+	void SetRequiredPin() { PinStatus = EPCGPinStatus::Required; }
+
+	bool IsNormalPin() const { return PinStatus == EPCGPinStatus::Normal; }
+	void SetNormalPin() { PinStatus = EPCGPinStatus::Normal; }
+
+	// Convert the bIsAdvanced boolean to PinStatus for deprecation purposes.
+	void PostSerialize(const FArchive& Ar);
+
 private:
+#if WITH_EDITORONLY_DATA
+	/* Advanced pin will be hidden by default in the UI and will be shown only if the user extend the node (in the UI) to see advanced pins. */
+	UPROPERTY(meta=(DeprecatedProperty, DeprecationMessage = "Use IsAdvancedPin function or PinStatus property."))
+	bool bAdvancedPin_DEPRECATED = false;
+#endif // WITH_EDITORONLY_DATA
+
 	UPROPERTY(BlueprintReadWrite, EditAnywhere, Category = Settings, meta = (AllowPrivateAccess = "true", EditCondition = "bAllowMultipleData", DisplayAfter = "bAllowMultipleData"))
 	bool bAllowMultipleConnections = true;
 };
 
+template<>
+struct TStructOpsTypeTraits<FPCGPinProperties> : public TStructOpsTypeTraitsBase2<FPCGPinProperties>
+{
+	enum
+	{
+		WithPostSerialize = true,
+	};
+};
+
 UENUM()
 enum class EPCGTypeConversion : uint8
 {
@@ -83,7 +124,7 @@ class PCG_API UPCGBlueprintPinHelpers : public UBlueprintFunctionLibrary
 
 public:
 	UFUNCTION(BlueprintPure, Category="PCG|Pins", meta = (NativeBreakFunc))
-	static void BreakPinProperty(const FPCGPinProperties& PinProperty, FName& Label, bool& bAllowMultipleData, bool& bAllowMultipleConnections, bool& bAdvancedPin, EPCGExclusiveDataType& AllowedType)
+	static void BreakPinProperty(const FPCGPinProperties& PinProperty, FName& Label, bool& bAllowMultipleData, bool& bAllowMultipleConnections, bool& bIsAdvancedPin, EPCGExclusiveDataType& AllowedType)
 	{
 		Label = PinProperty.Label;
 
@@ -108,11 +149,11 @@ public:
 
 		bAllowMultipleData = PinProperty.bAllowMultipleData;
 		bAllowMultipleConnections = PinProperty.AllowsMultipleConnections();
-		bAdvancedPin = PinProperty.bAdvancedPin;
+		bIsAdvancedPin = PinProperty.IsAdvancedPin();
 	}
 
 	UFUNCTION(BlueprintPure, Category="PCG|Pins", meta = (NativeMakeFunc))
-	static FPCGPinProperties MakePinProperty(FName Label, bool bAllowMultipleData, bool bAllowMultipleConnections, bool bAdvancedPin, EPCGExclusiveDataType AllowedType = EPCGExclusiveDataType::Any)
+	static FPCGPinProperties MakePinProperty(FName Label, bool bAllowMultipleData, bool bAllowMultipleConnections, bool bIsAdvancedPin, EPCGExclusiveDataType AllowedType = EPCGExclusiveDataType::Any)
 	{
 		const UEnum* DataTypeEnum = StaticEnum<EPCGDataType>();
 		const UEnum* ExclusiveDataTypeEnum = StaticEnum<EPCGExclusiveDataType>();
@@ -133,7 +174,10 @@ public:
 		}
 
 		FPCGPinProperties PinProperties = FPCGPinProperties(Label, DataType, bAllowMultipleConnections, bAllowMultipleData);
-		PinProperties.bAdvancedPin = bAdvancedPin;
+		if (bIsAdvancedPin)
+		{
+			PinProperties.SetAdvancedPin();
+		}
 
 		return PinProperties;
 	}
@@ -204,6 +248,40 @@ namespace PCGPinPropertiesHelpers
 }
 #endif // WITH_EDITOR
 
+/**
+* Helper class to allow the BP to call the custom functions on FPCGPinProperties.
+*/
+UCLASS()
+class PCG_API UPCGPinPropertiesBlueprintHelpers : public UBlueprintFunctionLibrary
+{
+	GENERATED_BODY()
+
+public:
+	UFUNCTION(BlueprintCallable, Category = "PCG|PinProperties", meta = (ScriptMethod))
+	static bool AllowsMultipleConnections(UPARAM(ref) const FPCGPinProperties& PinProperties) { return PinProperties.AllowsMultipleConnections(); }
+
+	UFUNCTION(BlueprintCallable, Category = "PCG|PinProperties", meta = (ScriptMethod))
+	static void SetAllowMultipleConnections(UPARAM(ref) FPCGPinProperties& PinProperties, bool bAllowMultipleConnections) { PinProperties.SetAllowMultipleConnections(bAllowMultipleConnections); }
+
+	UFUNCTION(BlueprintCallable, Category = "PCG|PinProperties", meta = (ScriptMethod))
+	static bool IsAdvancedPin(UPARAM(ref) const FPCGPinProperties& PinProperties) { return PinProperties.IsAdvancedPin(); }
+
+	UFUNCTION(BlueprintCallable, Category = "PCG|PinProperties", meta = (ScriptMethod))
+	static void SetAdvancedPin(UPARAM(ref) FPCGPinProperties& PinProperties) { PinProperties.SetAdvancedPin(); }
+
+	UFUNCTION(BlueprintCallable, Category = "PCG|PinProperties", meta = (ScriptMethod))
+	static bool IsRequiredPin(UPARAM(ref) const FPCGPinProperties& PinProperties) { return PinProperties.IsRequiredPin(); }
+
+	UFUNCTION(BlueprintCallable, Category = "PCG|PinProperties", meta = (ScriptMethod))
+	static void SetRequiredPin(UPARAM(ref) FPCGPinProperties& PinProperties) { PinProperties.SetRequiredPin(); }
+
+	UFUNCTION(BlueprintCallable, Category = "PCG|PinProperties", meta = (ScriptMethod))
+	static bool IsNormalPin(UPARAM(ref) const FPCGPinProperties& PinProperties) { return PinProperties.IsNormalPin(); }
+
+	UFUNCTION(BlueprintCallable, Category = "PCG|PinProperties", meta = (ScriptMethod))
+	static void SetNormalPin(UPARAM(ref) FPCGPinProperties& PinProperties) { PinProperties.SetNormalPin(); }
+};
+
 #if UE_ENABLE_INCLUDE_ORDER_DEPRECATED_IN_5_2
 #include "CoreMinimal.h"
 #endif
diff --git a/Engine/Plugins/PCG/Source/PCG/Public/PCGSettings.h b/Engine/Plugins/PCG/Source/PCG/Public/PCGSettings.h
--- a/Engine/Plugins/PCG/Source/PCG/Public/PCGSettings.h
+++ b/Engine/Plugins/PCG/Source/PCG/Public/PCGSettings.h
@@ -342,9 +342,9 @@ public:
 	virtual FName AdditionalTaskName() const { return FName(GetAdditionalTitleInformation()); }
 
 	/** By default a node does not specify any pin requirements, and will execute if it has no non-advanced pins or if it has any active
-	* connection to any pin. Override this function to specify which pins *must* have active connections to avoid getting culled.
+	* connection to any pin. Set the Required status on pin properties to specify which pins *must* have active connections to avoid getting culled.
 	*/
-	virtual bool IsInputPinRequiredByExecution(const UPCGPin* InPin) const { return false; }
+	virtual bool IsInputPinRequiredByExecution(const UPCGPin* InPin) const { return InPin && InPin->Properties.IsRequiredPin(); }
 
 	/** Returns true if InPin is in use by node (assuming node enabled). Can be used to communicate when a pin is not in use to user. */
 	virtual bool IsPinUsedByNodeExecution(const UPCGPin* InPin) const { return true; }
diff --git a/Engine/Plugins/PCG/Source/PCGEditor/Private/PCGEditorGraphNodeBase.cpp b/Engine/Plugins/PCG/Source/PCGEditor/Private/PCGEditorGraphNodeBase.cpp
--- a/Engine/Plugins/PCG/Source/PCGEditor/Private/PCGEditorGraphNodeBase.cpp
+++ b/Engine/Plugins/PCG/Source/PCGEditor/Private/PCGEditorGraphNodeBase.cpp
@@ -1054,7 +1054,7 @@ void UPCGEditorGraphNodeBase::CreatePins(const TArray<UPCGPin*>& InInputPins, co
 
 		UEdGraphPin* Pin = CreatePin(EEdGraphPinDirection::EGPD_Input, GetPinType(InputPin), InputPin->Properties.Label);
 		Pin->PinFriendlyName = GetPinFriendlyName(InputPin);
-		Pin->bAdvancedView = InputPin->Properties.bAdvancedPin;
+		Pin->bAdvancedView = InputPin->Properties.IsAdvancedPin();
 		bHasAdvancedPin |= Pin->bAdvancedView;
 	}
 
@@ -1067,7 +1067,7 @@ void UPCGEditorGraphNodeBase::CreatePins(const TArray<UPCGPin*>& InInputPins, co
 
 		UEdGraphPin* Pin = CreatePin(EEdGraphPinDirection::EGPD_Output, GetPinType(OutputPin), OutputPin->Properties.Label);
 		Pin->PinFriendlyName = GetPinFriendlyName(OutputPin);
-		Pin->bAdvancedView = OutputPin->Properties.bAdvancedPin;
+		Pin->bAdvancedView = OutputPin->Properties.IsAdvancedPin();
 		bHasAdvancedPin |= Pin->bAdvancedView;
 	}
 
