[PCG] Enable conversion from Spatial to Point|Param pin.

#rnx

diff --git a/Engine/Plugins/PCG/Source/PCG/Private/PCGPin.cpp b/Engine/Plugins/PCG/Source/PCG/Private/PCGPin.cpp
--- a/Engine/Plugins/PCG/Source/PCG/Private/PCGPin.cpp
+++ b/Engine/Plugins/PCG/Source/PCG/Private/PCGPin.cpp
@@ -363,10 +363,10 @@ bool UPCGPin::IsCompatible(const UPCGPin* OtherPin) const
 		return true;
 	}
 
-	// Anything spatial can collapse to point
-	const bool bUpstreamInSpatial = !(UpstreamTypes & ~EPCGDataType::Spatial);
-	const bool bDownstreamIsPoint = (DownstreamTypes == EPCGDataType::Point);
-	if (bUpstreamInSpatial && bDownstreamIsPoint)
+	// Anything spatial can collapse to point as long as the downstream pin does not support any other spatial types. For example, Point | Param is collapsible, but Point | Surface is not
+	const bool bUpstreamIsSpatial = !(UpstreamTypes & ~EPCGDataType::Spatial);
+	const bool bDownstreamIsPoint = (DownstreamTypes & EPCGDataType::Spatial) == EPCGDataType::Point;
+	if (bUpstreamIsSpatial && bDownstreamIsPoint)
 	{
 		return true;
 	}
@@ -427,10 +427,13 @@ EPCGTypeConversion UPCGPin::GetRequiredTypeConversion(const UPCGPin* InOtherPin)
 	}
 
 	// Spatial -> Point - "To Point" conversion.
-	const bool bUpstreamInSpatial = !(UpstreamTypes & ~EPCGDataType::Spatial);
-	if (bUpstreamInSpatial && DownstreamTypes == EPCGDataType::Point)
+	// Anything spatial can collapse to point as long as the downstream pin does not support any other spatial types. For example, Point | Param is collapsible, but Point | Surface is not.
+	const bool bUpstreamIsSpatial = !(UpstreamTypes & ~EPCGDataType::Spatial);
+	const bool bUpstreamIsPoint = UpstreamTypes == EPCGDataType::Point;
+	const bool bDownstreamIsSpatialAndPointOnly = (DownstreamTypes & EPCGDataType::Spatial) == EPCGDataType::Point;
+	if (bUpstreamIsSpatial && bDownstreamIsSpatialAndPointOnly)
 	{
-		return EPCGTypeConversion::CollapseToPoint;
+		return bUpstreamIsPoint ? EPCGTypeConversion::NoConversionRequired : EPCGTypeConversion::CollapseToPoint;
 	}
 
 	// Spline -> Surface conversion
diff --git a/Engine/Plugins/PCG/Source/PCG/Private/Tests/GraphAuthoring/PCGGraphAuthoringTest.cpp b/Engine/Plugins/PCG/Source/PCG/Private/Tests/GraphAuthoring/PCGGraphAuthoringTest.cpp
--- a/Engine/Plugins/PCG/Source/PCG/Private/Tests/GraphAuthoring/PCGGraphAuthoringTest.cpp
+++ b/Engine/Plugins/PCG/Source/PCG/Private/Tests/GraphAuthoring/PCGGraphAuthoringTest.cpp
@@ -92,6 +92,13 @@ bool FPCGGraphAuthoringTypeSystemTest::RunTest(const FString& Parameters)
 	bTestPassed &= ValidateConversion("Collapse: Concrete to Point", EPCGDataType::Concrete, EPCGDataType::Point, true, EPCGTypeConversion::CollapseToPoint);
 	bTestPassed &= ValidateConversion("Collapse: Spatial to Point", EPCGDataType::Spatial, EPCGDataType::Point, true, EPCGTypeConversion::CollapseToPoint);
 
+	// Collapse, with Point | Param
+	bTestPassed &= ValidateConversion("Collapse: Spline to Point", EPCGDataType::Spline, EPCGDataType::Point | EPCGDataType::Param, true, EPCGTypeConversion::CollapseToPoint);
+	bTestPassed &= ValidateConversion("Collapse: Landscape to Point", EPCGDataType::Landscape, EPCGDataType::Point | EPCGDataType::Param, true, EPCGTypeConversion::CollapseToPoint);
+	bTestPassed &= ValidateConversion("Collapse: Surface to Point", EPCGDataType::Surface, EPCGDataType::Point | EPCGDataType::Param, true, EPCGTypeConversion::CollapseToPoint);
+	bTestPassed &= ValidateConversion("Collapse: Concrete to Point", EPCGDataType::Concrete, EPCGDataType::Point | EPCGDataType::Param, true, EPCGTypeConversion::CollapseToPoint);
+	bTestPassed &= ValidateConversion("Collapse: Spatial to Point", EPCGDataType::Spatial, EPCGDataType::Point | EPCGDataType::Param, true, EPCGTypeConversion::CollapseToPoint);
+
 	// Filter
 	bTestPassed &= ValidateConversion("Filter: Any to Point", EPCGDataType::Any, EPCGDataType::Point, true, EPCGTypeConversion::Filter);
 	bTestPassed &= ValidateConversion("Filter: Any to Landscape", EPCGDataType::Any, EPCGDataType::Landscape, true, EPCGTypeConversion::Filter);
