[PCG] Reduce unnecessary graph re-executions while editing, by demoting changes that do not clearly affect compiled tasks to Cosmetic.

Side effects include cleaning up some change notification senders/receivers slightly, and removing a notification path which is no longer needed which should help editor performance.

diff --git a/Engine/Plugins/PCG/Source/PCG/Private/Graph/PCGGraphCompiler.cpp b/Engine/Plugins/PCG/Source/PCG/Private/Graph/PCGGraphCompiler.cpp
--- a/Engine/Plugins/PCG/Source/PCG/Private/Graph/PCGGraphCompiler.cpp
+++ b/Engine/Plugins/PCG/Source/PCG/Private/Graph/PCGGraphCompiler.cpp
@@ -162,7 +162,7 @@ TArray<FPCGGraphTask> FPCGGraphCompiler::CompileGraph(UPCGGraph* InGraph, FPCGTa
 			PostTask.StackIndex = InOutStackContext.GetCurrentStackIndex();
 			// Implementation note: since we`ve already executed the node once, we normally don`t need to execute it a second time
 			// especially since we cannot distinguish between the pre and post during execution so any data filtering related to pins is bound to fail.
-			PostTask.Element = MakeShared<FPCGTrivialElement>();
+			PostTask.Element = GetSharedTrivialElement();
 
 			// Add execution-only dependency on pre-task, without this post task can be scheduled concurrently with pre-task, and concurrently
 			// with something that might become inactive and would then fail to dynamically cull this already-scheduled task.
@@ -402,11 +402,13 @@ void FPCGGraphCompiler::CreateGridLinkages(
 						DownstreamNode,
 						static_cast<FPCGGridLinkageContext*>(InContext));
 				};
+
 				FPCGGenericElement::FContextAllocator ContextAllocator = [](const FPCGDataCollection&, TWeakObjectPtr<UPCGComponent>, const UPCGNode*)
 				{
 					return new FPCGGridLinkageContext();
 				};
-				LinkTask.Element = MakeShared<FPCGGenericElement>(GridLinkageOperation, ContextAllocator);
+
+				LinkTask.Element = MakeShared<PCGGraphExecutor::FPCGGridLinkageElement>(GridLinkageOperation, ContextAllocator, FromGrid, ToGrid, ResourceKey);
 
 				// Now splice in the new task - redirect the downstream task to grab its input from the link task.
 				TaskInput.TaskId = LinkTask.NodeId;
@@ -962,6 +964,8 @@ void FPCGGraphCompiler::CompileTopGraph(UPCGGraph* InGraph, uint32 GenerationGri
 		return;
 	}
 
+	UE_LOG(LogPCG, Verbose, TEXT("FPCGGraphCompiler::CompileTopGraph '%s' grid: %u"), *InGraph->GetName(), GenerationGridSize);
+
 	// Build from non-top tasks
 	FPCGStackContext StackContext;
 	TArray<FPCGGraphTask> CompiledTasks = GetCompiledTasks(InGraph, GenerationGridSize, StackContext, /*bIsTopGraph=*/false);
@@ -1026,7 +1030,7 @@ void FPCGGraphCompiler::CompileTopGraph(UPCGGraph* InGraph, uint32 GenerationGri
 	const FPCGTaskId PostExecuteTaskId = PreExecuteTaskId + 1;
 
 	FPCGGraphTask& PreExecuteTask = CompiledTasks.Emplace_GetRef();
-	PreExecuteTask.Element = MakeShared<FPCGTrivialElement>();
+	PreExecuteTask.Element = GetSharedTrivialElement();
 	PreExecuteTask.NodeId = PreExecuteTaskId;
 
 	for (int TaskIndex = 0; TaskIndex < TaskNum; ++TaskIndex)
@@ -1041,7 +1045,7 @@ void FPCGGraphCompiler::CompileTopGraph(UPCGGraph* InGraph, uint32 GenerationGri
 	}
 
 	FPCGGraphTask& PostExecuteTask = CompiledTasks.Emplace_GetRef();
-	PostExecuteTask.Element = MakeShared<FPCGTrivialElement>();
+	PostExecuteTask.Element = GetSharedTrivialElement();
 	PostExecuteTask.NodeId = PostExecuteTaskId;
 
 	// Find end nodes, e.g. all nodes that have no successors.
@@ -1085,6 +1089,27 @@ void FPCGGraphCompiler::CompileTopGraph(UPCGGraph* InGraph, uint32 GenerationGri
 	GraphToTaskMapLock.WriteUnlock();
 }
 
+FPCGElementPtr FPCGGraphCompiler::GetSharedTrivialElement()
+{
+	{
+		FReadScopeLock Lock(SharedTrivialElementLock);
+
+		if (SharedTrivialElement)
+		{
+			return SharedTrivialElement;
+		}
+	}
+
+	FWriteScopeLock Lock(SharedTrivialElementLock);
+
+	if (!SharedTrivialElement)
+	{
+		SharedTrivialElement = MakeShared<FPCGTrivialElement>();
+	}
+
+	return SharedTrivialElement;
+}
+
 void FPCGGraphCompiler::ClearCache()
 {
 	FWriteScopeLock Lock(GraphToTaskMapLock);
@@ -1095,16 +1120,42 @@ void FPCGGraphCompiler::ClearCache()
 }
 
 #if WITH_EDITOR
-void FPCGGraphCompiler::NotifyGraphChanged(UPCGGraph* InGraph)
+void FPCGGraphCompiler::NotifyGraphChanged(UPCGGraph* InGraph, EPCGChangeType ChangeType)
 {
-	if (InGraph)
+	if (InGraph && (ChangeType != EPCGChangeType::Cosmetic))
 	{
 		RemoveFromCache(InGraph);
 	}
 }
 
+bool FPCGGraphCompiler::Recompile(UPCGGraph* InGraph, uint32 GenerationGridSize, bool bIsTopGraph)
+{
+	FPCGStackContext StackContextBefore;
+	const TArray<FPCGGraphTask> TasksBefore = GetPrecompiledTasks(InGraph, GenerationGridSize, StackContextBefore, bIsTopGraph);
+
+	// Need to manually purge as the graph compiler will not have gotten the change notification yet. Editor only.
+	RemoveFromCache(InGraph);
+
+	FPCGStackContext StackContextAfter;
+	const TArray<FPCGGraphTask> TasksAfter = GetCompiledTasks(InGraph, GenerationGridSize, StackContextAfter, bIsTopGraph);
+
+	bool bAllTasksEqual = TasksBefore.Num() == TasksAfter.Num();
+	if (bAllTasksEqual)
+	{
+		for (int I = 0; I < TasksBefore.Num(); ++I)
+		{
+			bAllTasksEqual = bAllTasksEqual && TasksAfter[I].IsApproximatelyEqual(TasksBefore[I]);
+		}
+	}
+
+	// Compiled result is compiled tasks + associated stacks. Compare both and return true if compiled result changes.
+	return !bAllTasksEqual || (StackContextBefore != StackContextAfter);
+}
+
 void FPCGGraphCompiler::RemoveFromCache(UPCGGraph* InGraph)
 {
+	UE_LOG(LogPCG, Verbose, TEXT("FPCGGraphCompiler::RemoveFromCache '%s'"), *InGraph->GetName());
+
 	check(InGraph);
 	RemoveFromCacheRecursive(InGraph);
 }
diff --git a/Engine/Plugins/PCG/Source/PCG/Private/Graph/PCGGraphCompiler.h b/Engine/Plugins/PCG/Source/PCG/Private/Graph/PCGGraphCompiler.h
--- a/Engine/Plugins/PCG/Source/PCG/Private/Graph/PCGGraphCompiler.h
+++ b/Engine/Plugins/PCG/Source/PCG/Private/Graph/PCGGraphCompiler.h
@@ -5,9 +5,12 @@
 #include "PCGCommon.h"
 #include "Graph/PCGStackContext.h"
 
+class IPCGElement;
 class UPCGGraph;
 struct FPCGGraphTask;
 
+typedef TSharedPtr<IPCGElement, ESPMode::ThreadSafe> FPCGElementPtr;
+
 /** 
 * FPCGGraphCompiler
 * This class compiles a graph into tasks and keeps an internal cache
@@ -22,7 +25,14 @@ public:
 	static void OffsetNodeIds(TArray<FPCGGraphTask>& Tasks, FPCGTaskId Offset, FPCGTaskId ParentId);
 
 #if WITH_EDITOR
-	void NotifyGraphChanged(UPCGGraph* InGraph);
+	/** Checks whether the compiled result changes when recompiling InGraph. Note that this incurs the cost of a graph compilation each time it is called.
+	* The cache will be updated with the latest compiled result. Returns true if the compiled result changes;
+	*/
+	bool Recompile(UPCGGraph* InGraph, uint32 GenerationGridSize, bool bIsTopGraph = true);
+
+	void RemoveFromCache(UPCGGraph* InGraph);
+
+	void NotifyGraphChanged(UPCGGraph* InGraph, EPCGChangeType ChangeType);
 #endif
 
 	/** Flush all cached compiled graphs. */
@@ -34,6 +44,9 @@ private:
 	/** Compiles the top graph and applies culling optimizations if a non-unitialized grid size is provided. */
 	void CompileTopGraph(UPCGGraph* InGraph, uint32 GenerationGridSize);
 
+	/** Returns the trivial element object shared by all tasks that need it. */
+	FPCGElementPtr GetSharedTrivialElement();
+
 	/** Propagates grid sizes through a graph's compiled tasks. */
 	static void ResolveGridSizes(
 		EPCGHiGenGrid GenerationGrid,
@@ -80,8 +93,10 @@ private:
 	TMap<UPCGGraph*, TMap<uint32, TArray<FPCGGraphTask>>> TopGraphToTaskMap;
 	TMap<UPCGGraph*, TMap<uint32, FPCGStackContext>> TopGraphToStackContextMap;
 
+	FPCGElementPtr SharedTrivialElement;
+	mutable FRWLock SharedTrivialElementLock;
+
 #if WITH_EDITOR
-	void RemoveFromCache(UPCGGraph* InGraph);
 	void RemoveFromCacheRecursive(UPCGGraph* InGraph);
 
 	FCriticalSection GraphDependenciesLock;
diff --git a/Engine/Plugins/PCG/Source/PCG/Private/Graph/PCGGraphExecutor.cpp b/Engine/Plugins/PCG/Source/PCG/Private/Graph/PCGGraphExecutor.cpp
--- a/Engine/Plugins/PCG/Source/PCG/Private/Graph/PCGGraphExecutor.cpp
+++ b/Engine/Plugins/PCG/Source/PCG/Private/Graph/PCGGraphExecutor.cpp
@@ -68,6 +68,48 @@ namespace PCGGraphExecutor
 		TEXT("Controls whether tasks are culled at execution time, for example in response to an deactivated dynamic branch pin"));
 }
 
+#if WITH_EDITOR
+bool FPCGGraphTaskInput::operator==(const FPCGGraphTaskInput& Other) const
+{
+	return (TaskId == Other.TaskId)
+		&& (InPin == Other.InPin)
+		&& (OutPin == Other.OutPin)
+		&& (bProvideData == Other.bProvideData);
+}
+
+bool FPCGGraphTask::IsApproximatelyEqual(const FPCGGraphTask& Other) const
+{
+	// Do trivial pointer comparisons first, then run == operator to determine equivalence.
+	bool bElementsMatch = (Element == Other.Element);
+	if (!bElementsMatch && Element && Other.Element)
+	{
+		if (Element->IsGridLinkage() && Other.Element->IsGridLinkage())
+		{
+			const PCGGraphExecutor::FPCGGridLinkageElement& LinkageElement = static_cast<const PCGGraphExecutor::FPCGGridLinkageElement&>(*Element);
+			const PCGGraphExecutor::FPCGGridLinkageElement& OtherLinkageElement = static_cast<const PCGGraphExecutor::FPCGGridLinkageElement&>(*Other.Element);
+			bElementsMatch = (LinkageElement == OtherLinkageElement);
+		}
+		else
+		{
+			ensureMsgf(false, TEXT("Graph compilation emitted an element type that is not a trivial element or a grid linkage element. Element comparison will fail. ")
+				TEXT("Equivalence operator needs to be implemented for this new element."));
+		}
+	}
+
+	return (Inputs == Other.Inputs)
+		&& (Node == Other.Node)
+		&& (SourceComponent == Other.SourceComponent)
+		&& bElementsMatch
+		&& (Context == Other.Context)
+		&& (NodeId == Other.NodeId)
+		&& (CompiledTaskId == Other.CompiledTaskId)
+		&& (ParentId == Other.ParentId)
+		&& (PinDependency == Other.PinDependency)
+		&& (StackIndex == Other.StackIndex)
+		&& (StackContext == Other.StackContext);
+}
+#endif // WITH_EDITOR
+
 FPCGGraphExecutor::FPCGGraphExecutor()
 #if WITH_EDITOR
 	: GenerationProgressNotification(GetNotificationTextFormat())
@@ -1626,9 +1668,9 @@ void FPCGGraphExecutor::ReleaseUnusedActors()
 #endif
 }
 
-void FPCGGraphExecutor::NotifyGraphChanged(UPCGGraph* InGraph)
+void FPCGGraphExecutor::NotifyGraphChanged(UPCGGraph* InGraph, EPCGChangeType ChangeType)
 {
-	GraphCompiler.NotifyGraphChanged(InGraph);
+	GraphCompiler.NotifyGraphChanged(InGraph, ChangeType);
 }
 
 void FPCGGraphExecutor::UpdateGenerationNotification()
@@ -1781,6 +1823,13 @@ void FPCGGenericElement::AbortInternal(FPCGContext* Context) const
 
 namespace PCGGraphExecutor
 {
+#if WITH_EDITOR
+	bool FPCGGridLinkageElement::operator==(const FPCGGridLinkageElement& Other) const
+	{
+		return FromGrid == Other.FromGrid && ToGrid == Other.ToGrid && ResourceKey == Other.ResourceKey;
+	}
+#endif
+
 	bool ExecuteGridLinkage(
 		EPCGHiGenGrid InGenerationGrid,
 		EPCGHiGenGrid InFromGrid,
diff --git a/Engine/Plugins/PCG/Source/PCG/Private/Graph/PCGGraphExecutor.h b/Engine/Plugins/PCG/Source/PCG/Private/Graph/PCGGraphExecutor.h
--- a/Engine/Plugins/PCG/Source/PCG/Private/Graph/PCGGraphExecutor.h
+++ b/Engine/Plugins/PCG/Source/PCG/Private/Graph/PCGGraphExecutor.h
@@ -46,6 +46,10 @@ struct FPCGGraphTaskInput
 	{
 	}
 
+#if WITH_EDITOR
+	bool operator==(const FPCGGraphTaskInput& Other) const;
+#endif
+
 	FPCGTaskId TaskId;
 
 	/** The upstream output pin from which the input data comes. */
@@ -60,6 +64,11 @@ struct FPCGGraphTaskInput
 
 struct FPCGGraphTask
 {
+#if WITH_EDITOR
+	/** Approximate equivalence. Does not deeply check node settings, nor does it do a deep comparison of the element. */
+	bool IsApproximatelyEqual(const FPCGGraphTask& Other) const;
+#endif
+
 	TArray<FPCGGraphTaskInput> Inputs;
 	const UPCGNode* Node = nullptr;
 	TWeakObjectPtr<UPCGComponent> SourceComponent = nullptr;
@@ -172,7 +181,7 @@ public:
 	void AddToUnusedActors(const TSet<FWorldPartitionReference>& UnusedActors);
 
 	/** Notify compiler that graph has changed so it'll be removed from the cache */
-	void NotifyGraphChanged(UPCGGraph* InGraph);
+	void NotifyGraphChanged(UPCGGraph* InGraph, EPCGChangeType ChangeType);
 
 	/** Returns the number of entries currently in the cache for InElement. */
 	uint32 GetGraphCacheEntryCount(IPCGElement* InElement) const { return GraphCache.GetGraphCacheEntryCount(InElement); }
@@ -327,6 +336,34 @@ struct FPCGGridLinkageContext : public FPCGContext
 
 namespace PCGGraphExecutor
 {
+	/** Marshals data across grid sizes at execution time. */
+	class FPCGGridLinkageElement : public FPCGGenericElement
+	{
+	public:
+		FPCGGridLinkageElement(TFunction<bool(FPCGContext*)> InOperation, const FContextAllocator& InContextAllocator, EPCGHiGenGrid InFromGrid, EPCGHiGenGrid InToGrid, const FString& InResourceKey)
+			: FPCGGenericElement(InOperation, InContextAllocator)
+			, FromGrid(InFromGrid)
+			, ToGrid(InToGrid)
+			, ResourceKey(InResourceKey)
+		{
+		}
+
+#if WITH_EDITOR
+		//~Begin IPCGElement interface
+		virtual bool IsGridLinkage() const override { return true; }
+		//~End IPCGElement interface
+
+		/** Return true if the grid sizes & path match. */
+		bool operator==(const FPCGGridLinkageElement& Other) const;
+#endif
+
+	private:
+		// These values are stored here so that we can compare two grid linkage elements for equivalence.
+		EPCGHiGenGrid FromGrid = EPCGHiGenGrid::Uninitialized;
+		EPCGHiGenGrid ToGrid = EPCGHiGenGrid::Uninitialized;
+		FString ResourceKey;
+	};
+
 	/** Compares InFromGrid and InToGrid and performs data storage/retrieval as necessary to marshal data across execution grids. */
 	bool ExecuteGridLinkage(
 		EPCGHiGenGrid InGenerationGrid,
diff --git a/Engine/Plugins/PCG/Source/PCG/Private/Graph/PCGPinDependencyExpression.h b/Engine/Plugins/PCG/Source/PCG/Private/Graph/PCGPinDependencyExpression.h
--- a/Engine/Plugins/PCG/Source/PCG/Private/Graph/PCGPinDependencyExpression.h
+++ b/Engine/Plugins/PCG/Source/PCG/Private/Graph/PCGPinDependencyExpression.h
@@ -32,6 +32,8 @@ struct FPCGPinDependencyExpression
 	/** Applies given node ID offset to all pin IDs in expression (pin ID is combination of node ID and pin index). */
 	void OffsetNodeIds(uint64 NodeIdOffset);
 
+	bool operator==(const FPCGPinDependencyExpression& Other) const { return (Expression == Other.Expression); }
+
 #if WITH_EDITOR
 	FString ToString() const;
 #endif
diff --git a/Engine/Plugins/PCG/Source/PCG/Private/Graph/PCGStackContext.cpp b/Engine/Plugins/PCG/Source/PCG/Private/Graph/PCGStackContext.cpp
--- a/Engine/Plugins/PCG/Source/PCG/Private/Graph/PCGStackContext.cpp
+++ b/Engine/Plugins/PCG/Source/PCG/Private/Graph/PCGStackContext.cpp
@@ -286,3 +286,8 @@ void FPCGStackContext::PrependParentStack(const FPCGStack* InParentStack)
 		Stack.StackFrames.Insert(InParentStack->StackFrames, 0);
 	}
 }
+
+bool FPCGStackContext::operator==(const FPCGStackContext& Other) const
+{
+	return (CurrentStackIndex == Other.CurrentStackIndex) && (Stacks == Other.Stacks);
+}
diff --git a/Engine/Plugins/PCG/Source/PCG/Private/PCGGraph.cpp b/Engine/Plugins/PCG/Source/PCG/Private/PCGGraph.cpp
--- a/Engine/Plugins/PCG/Source/PCG/Private/PCGGraph.cpp
+++ b/Engine/Plugins/PCG/Source/PCG/Private/PCGGraph.cpp
@@ -13,6 +13,8 @@
 #include "PCGSubgraph.h"
 #include "Elements/PCGHiGenGridSize.h"
 #include "Elements/PCGUserParameterGet.h"
+#include "Graph/PCGGraphCompiler.h"
+#include "Graph/PCGGraphExecutor.h"
 
 #if WITH_EDITOR
 #include "CoreGlobals.h"
@@ -522,7 +524,7 @@ void UPCGGraph::BeginDestroy()
 	{
 		if (UPCGSubsystem* PCGSubsystem = UPCGSubsystem::GetInstance(GEditor->GetEditorWorldContext().World()))
 		{
-			PCGSubsystem->NotifyGraphChanged(this);
+			PCGSubsystem->NotifyGraphChanged(this, EPCGChangeType::Structural | EPCGChangeType::GenerationGrid);
 		}
 	}
 
@@ -667,7 +669,7 @@ void UPCGGraph::OnNodesAdded(TArrayView<UPCGNode*> InNodes)
 		}
 	}
 
-	NotifyGraphChanged(ChangeType);
+	NotifyGraphStructureChanged(ChangeType);
 #endif
 }
 
@@ -691,7 +693,7 @@ void UPCGGraph::OnNodesRemoved(TArrayView<UPCGNode*> InNodes)
 		}
 	}
 
-	NotifyGraphChanged(bAnyGridSizeNodes ? (EPCGChangeType::Structural | EPCGChangeType::GenerationGrid) : EPCGChangeType::Structural);
+	NotifyGraphStructureChanged(bAnyGridSizeNodes ? (EPCGChangeType::Structural | EPCGChangeType::GenerationGrid) : EPCGChangeType::Structural);
 #endif
 }
 
@@ -748,7 +750,7 @@ bool UPCGGraph::AddLabeledEdge(UPCGNode* From, const FName& FromPinLabel, UPCGNo
 	// After all nodes are notified, re-enable graph notifications and send graph change notification.
 	EnableNotificationsForEditor();
 
-	NotifyGraphChanged(ChangeType);
+	NotifyGraphStructureChanged(ChangeType);
 #endif
 
 	return bToPinBrokeOtherEdges;
@@ -893,7 +895,7 @@ bool UPCGGraph::RemoveEdge(UPCGNode* From, const FName& FromLabel, UPCGNode* To,
 
 	if (TouchedNodes.Num() > 0)
 	{
-		NotifyGraphChanged(ChangeType);
+		NotifyGraphStructureChanged(ChangeType);
 	}
 #endif
 
@@ -976,7 +978,7 @@ bool UPCGGraph::RemoveInboundEdges(UPCGNode* InNode, const FName& InboundLabel)
 
 	if (TouchedNodes.Num() > 0)
 	{
-		NotifyGraphChanged(ChangeType);
+		NotifyGraphStructureChanged(ChangeType);
 	}
 #endif
 
@@ -1006,7 +1008,7 @@ bool UPCGGraph::RemoveOutboundEdges(UPCGNode* InNode, const FName& OutboundLabel
 
 	if (TouchedNodes.Num() > 0)
 	{
-		NotifyGraphChanged(ChangeType);
+		NotifyGraphStructureChanged(ChangeType);
 	}
 #endif
 
@@ -1140,6 +1142,41 @@ void UPCGGraph::RemoveExtraEditorNode(const UObject* InNode)
 	ExtraEditorNodes.Remove(const_cast<UObject*>(InNode));
 }
 
+bool UPCGGraph::PrimeGraphCompilationCache()
+{
+	UPCGSubsystem* Subsystem = UPCGSubsystem::GetActiveEditorInstance();
+	FPCGGraphCompiler* GraphCompiler = Subsystem ? Subsystem->GetGraphCompiler() : nullptr;
+	if (!Subsystem || !GraphCompiler)
+	{
+		return false;
+	}
+
+	const UPCGComponent* InspectedComponent = InspectedStack.GetRootComponent();
+
+	FPCGStackContext StackContext;
+	GraphCompiler->GetCompiledTasks(this, PCGHiGenGrid::UninitializedGridSize(), StackContext, /*bIsTopGraph=*/true);
+
+	UE_LOG(LogPCG, Verbose, TEXT("UPCGGraph::PrimeGraphCompilationCache '%s' %u"), *this->GetName(), PCGHiGenGrid::UninitializedGridSize());
+
+	return true;
+}
+
+bool UPCGGraph::Recompile()
+{
+	UPCGSubsystem* Subsystem = UPCGSubsystem::GetActiveEditorInstance();
+	FPCGGraphCompiler* GraphCompiler = Subsystem ? Subsystem->GetGraphCompiler() : nullptr;
+	if (!Subsystem || !GraphCompiler)
+	{
+		return true;
+	}
+
+	const bool bChanged = GraphCompiler->Recompile(this, PCGHiGenGrid::UninitializedGridSize(), /*bIsTopGraph=*/true);
+
+	UE_LOG(LogPCG, Verbose, TEXT("UPCGGraph::Recompile '%s' grid: %u changed: %d"), *this->GetName(), PCGHiGenGrid::UninitializedGridSize(), bChanged ? 1 : 0);
+
+	return bChanged;
+}
+
 FPCGSelectionKeyToSettingsMap UPCGGraph::GetTrackedActorKeysToSettings() const
 {
 	FPCGSelectionKeyToSettingsMap TagsToSettings;
@@ -1184,6 +1221,32 @@ void UPCGGraph::GetTrackedActorKeysToSettings(FPCGSelectionKeyToSettingsMap& Out
 	}
 }
 
+void UPCGGraph::NotifyGraphStructureChanged(EPCGChangeType ChangeType, bool bForce)
+{
+	bool bExecutionAffected = true;
+
+	// If settings were not changed, we can gate the change notification based on whether compiled graph output changed. This compilation check
+	// does not support settings changes.
+	if (!bForce && !(ChangeType & EPCGChangeType::Settings))
+	{
+		bExecutionAffected = Recompile();
+	}
+
+	if (!bExecutionAffected)
+	{
+		// If compiled tasks are unchanged and the settings have not changed, then we can demote the change to avoid unnecessary graph executions.
+		// * Structural and GenerationGrid are frequent change types that trigger generation
+		// * Node because edges can change pin types on nodes.
+		// * Settings because changing a subgraph currently changes settings on subgraph node.
+		ChangeType &= ~(EPCGChangeType::Structural | EPCGChangeType::GenerationGrid | EPCGChangeType::Node | EPCGChangeType::Settings);
+
+		// Positively flagging as cosmetic is required because downstream things specifically test for this currently.
+		ChangeType |= EPCGChangeType::Cosmetic;
+	}
+
+	NotifyGraphChanged(ChangeType);
+}
+
 void UPCGGraph::NotifyGraphChanged(EPCGChangeType ChangeType)
 {
 	TRACE_CPUPROFILER_EVENT_SCOPE(UPCGGraph::NotifyGraphChanged);
@@ -1211,13 +1274,13 @@ void UPCGGraph::NotifyGraphChanged(EPCGChangeType ChangeType)
 		{
 			if (UPCGSubsystem* PCGPIESubsystem = UPCGSubsystem::GetInstance(GEditor->PlayWorld.Get()))
 			{
-				PCGPIESubsystem->NotifyGraphChanged(this);
+				PCGPIESubsystem->NotifyGraphChanged(this, ChangeType);
 			}
 		}
 
 		if (UPCGSubsystem* PCGSubsystem = UPCGSubsystem::GetInstance(GEditor->GetEditorWorldContext().World()))
 		{
-			PCGSubsystem->NotifyGraphChanged(this);
+			PCGSubsystem->NotifyGraphChanged(this, ChangeType);
 		}
 	}
 
@@ -1272,14 +1335,11 @@ void UPCGGraph::OnNodeChanged(UPCGNode* InNode, EPCGChangeType ChangeType)
 			FWriteScopeLock Lock(NodeToGridSizeLock);
 			NodeToGridSize.Reset();
 		}
-
-		// Broadcast so that grid size visualization can be updated editor-side.
-		OnGraphStructureChangedDelegate.Broadcast(this);
 	}
 
 	if ((ChangeType & ~EPCGChangeType::Cosmetic) != EPCGChangeType::None)
 	{
-		NotifyGraphChanged(ChangeType);
+		NotifyGraphStructureChanged(ChangeType);
 	}
 }
 
@@ -1727,27 +1787,22 @@ void UPCGGraphInstance::OnGraphChanged(UPCGGraphInterface* InGraph, EPCGChangeTy
 {
 	if (InGraph == Graph)
 	{
-		// Also notify other systems that this graph changed, only if the owner is not a PCG Component nor PCG Subgraph.
-		// They already have their own system to trigger a refresh.
-		const UObject* Outer = GetOuter();
-		if (!Outer || !(Outer->IsA<UPCGComponent>() || Outer->IsA<UPCGSubgraphSettings>()))
+		if (ChangeType != EPCGChangeType::Cosmetic)
 		{
-			FPropertyChangedEvent EmptyEvent{ nullptr };
-			FCoreUObjectDelegates::OnObjectPropertyChanged.Broadcast(this, EmptyEvent);
+			// Also notify other systems that this graph changed, only if the owner is not a PCG Component nor PCG Subgraph.
+			// They already have their own system to trigger a refresh.
+			const UObject* Outer = GetOuter();
+			if (!Outer || !(Outer->IsA<UPCGComponent>() || Outer->IsA<UPCGSubgraphSettings>()))
+			{
+				FPropertyChangedEvent EmptyEvent{ nullptr };
+				FCoreUObjectDelegates::OnObjectPropertyChanged.Broadcast(this, EmptyEvent);
+			}
 		}
 
 		OnGraphChangedDelegate.Broadcast(this, ChangeType);
 	}
 }
 
-void UPCGGraphInstance::OnGraphStructureChanged(UPCGGraphInterface* InGraph)
-{
-	if (InGraph == Graph)
-	{
-		OnGraphStructureChangedDelegate.Broadcast(this);
-	}
-}
-
 bool UPCGGraphInstance::CanEditChange(const FProperty* InProperty) const
 {
 	// Graph can only be changed if it is not in a local PCGComponent
@@ -1767,7 +1822,6 @@ void UPCGGraphInstance::TeardownCallbacks()
 	if (Graph)
 	{
 		Graph->OnGraphChangedDelegate.RemoveAll(this);
-		Graph->OnGraphStructureChangedDelegate.RemoveAll(this);
 		Graph->OnGraphParametersChangedDelegate.RemoveAll(this);
 	}
 }
@@ -1777,7 +1831,6 @@ void UPCGGraphInstance::SetupCallbacks()
 	if (Graph && !Graph->OnGraphChangedDelegate.IsBoundToObject(this))
 	{
 		Graph->OnGraphChangedDelegate.AddUObject(this, &UPCGGraphInstance::OnGraphChanged);
-		Graph->OnGraphStructureChangedDelegate.AddUObject(this, &UPCGGraphInstance::OnGraphStructureChanged);
 		Graph->OnGraphParametersChangedDelegate.AddUObject(this, &UPCGGraphInstance::OnGraphParametersChanged);
 	}
 }
diff --git a/Engine/Plugins/PCG/Source/PCG/Private/PCGSubgraph.cpp b/Engine/Plugins/PCG/Source/PCG/Private/PCGSubgraph.cpp
--- a/Engine/Plugins/PCG/Source/PCG/Private/PCGSubgraph.cpp
+++ b/Engine/Plugins/PCG/Source/PCG/Private/PCGSubgraph.cpp
@@ -185,7 +185,13 @@ void UPCGBaseSubgraphSettings::OnSubgraphChanged(UPCGGraphInterface* InGraph, EP
 {
 	if (InGraph == GetSubgraphInterface())
 	{
-		OnSettingsChangedDelegate.Broadcast(this, (ChangeType | EPCGChangeType::Settings));
+		// Only add settings if not cosmetic - we don't want to promote a cosmetic change to something deeper.
+		if (ChangeType != EPCGChangeType::Cosmetic)
+		{
+			ChangeType |= EPCGChangeType::Settings;
+		}
+
+		OnSettingsChangedDelegate.Broadcast(this, ChangeType);
 
 		// Also rebuild the overrides
 		InitializeCachedOverridableParams(/*bReset=*/true);
diff --git a/Engine/Plugins/PCG/Source/PCG/Private/PCGSubsystem.cpp b/Engine/Plugins/PCG/Source/PCG/Private/PCGSubsystem.cpp
--- a/Engine/Plugins/PCG/Source/PCG/Private/PCGSubsystem.cpp
+++ b/Engine/Plugins/PCG/Source/PCG/Private/PCGSubsystem.cpp
@@ -1277,11 +1277,11 @@ void UPCGSubsystem::DeleteSerializedPartitionActors(bool bOnlyDeleteUnused, bool
 	}
 }
 
-void UPCGSubsystem::NotifyGraphChanged(UPCGGraph* InGraph)
+void UPCGSubsystem::NotifyGraphChanged(UPCGGraph* InGraph, EPCGChangeType ChangeType)
 {
 	if (GraphExecutor)
 	{
-		GraphExecutor->NotifyGraphChanged(InGraph);
+		GraphExecutor->NotifyGraphChanged(InGraph, ChangeType);
 	}
 }
 
@@ -1345,7 +1345,16 @@ bool UPCGSubsystem::GetStackContext(const UPCGComponent* InComponent, FPCGStackC
 {
 	if (InComponent && InComponent->GetGraph())
 	{
-		GetGraphCompiler()->GetCompiledTasks(InComponent->GetGraph(), InComponent->GetGenerationGridSize(), OutStackContext);
+		uint32 GenerationGridSize = PCGHiGenGrid::UninitializedGridSize();
+		if (InComponent->GetGraph()->IsHierarchicalGenerationEnabled())
+		{
+			if (InComponent->IsLocalComponent() || InComponent->IsPartitioned())
+			{
+				GenerationGridSize = InComponent->GetGenerationGridSize();
+			}
+		}
+
+		GetGraphCompiler()->GetCompiledTasks(InComponent->GetGraph(), GenerationGridSize, OutStackContext);
 		return true;
 	}
 
diff --git a/Engine/Plugins/PCG/Source/PCG/Public/Graph/PCGStackContext.h b/Engine/Plugins/PCG/Source/PCG/Public/Graph/PCGStackContext.h
--- a/Engine/Plugins/PCG/Source/PCG/Public/Graph/PCGStackContext.h
+++ b/Engine/Plugins/PCG/Source/PCG/Public/Graph/PCGStackContext.h
@@ -136,6 +136,8 @@ public:
 
 	TArray<FPCGStack>& GetStacksMutable() { return Stacks; }
 
+	bool operator==(const FPCGStackContext& Other) const;
+
 private:
 	/** List of all stacks encountered top graph and all (nested) subgraphs. Order is simply order of encountering during compilation. */
 	TArray<FPCGStack> Stacks;
diff --git a/Engine/Plugins/PCG/Source/PCG/Public/PCGElement.h b/Engine/Plugins/PCG/Source/PCG/Public/PCGElement.h
--- a/Engine/Plugins/PCG/Source/PCG/Public/PCGElement.h
+++ b/Engine/Plugins/PCG/Source/PCG/Public/PCGElement.h
@@ -99,8 +99,11 @@ public:
 	/** Public function called when an element is cancelled, passing its current context if any. */
 	void Abort(FPCGContext* Context) const;
 
-	/** Note: the following methods must be called from the main thread */
 #if WITH_EDITOR
+	/** Is this element used for marshalling data across higen grids. Used as rudimentary RTTI for element object comparisons (editor only). */
+	virtual bool IsGridLinkage() const { return false; }
+
+	/** Note: must be called from the main thread. */
 	void DebugDisplay(FPCGContext* Context) const;
 #endif
 
diff --git a/Engine/Plugins/PCG/Source/PCG/Public/PCGGraph.h b/Engine/Plugins/PCG/Source/PCG/Public/PCGGraph.h
--- a/Engine/Plugins/PCG/Source/PCG/Public/PCGGraph.h
+++ b/Engine/Plugins/PCG/Source/PCG/Public/PCGGraph.h
@@ -114,9 +114,6 @@ public:
 #if WITH_EDITOR
 	FOnPCGGraphChanged OnGraphChangedDelegate;
 	FOnPCGGraphParametersChanged OnGraphParametersChangedDelegate;
-
-	/** A structural change has been made such as changing a static branch or a higen grid size node. */
-	FOnPCGGraphStructureChanged OnGraphStructureChangedDelegate;
 #endif // WITH_EDITOR
 
 	template <typename T>
@@ -312,10 +309,16 @@ public:
 	void SetExtraEditorNodes(const TArray<TObjectPtr<const UObject>>& InNodes);
 
 	bool IsInspecting() const { return bIsInspecting; }
-	void EnableInspection() { bIsInspecting = true; }
-	void DisableInspection() { bIsInspecting = false; }
+	void EnableInspection(const FPCGStack& InInspectedStack) { bIsInspecting = true; InspectedStack = InInspectedStack; }
+	void DisableInspection() { bIsInspecting = false; InspectedStack = FPCGStack(); }
 	bool DebugFlagAppliesToIndividualComponents() const { return bDebugFlagAppliesToIndividualComponents; }
 	void RemoveExtraEditorNode(const UObject* InNode);
+
+	/** Instruct the graph compiler to cache the relevant permutations of this graph. */
+	bool PrimeGraphCompilationCache();
+
+	/** Trigger a recompilation of the relevant permutations of this graph and check for change in the compiled tasks. */
+	bool Recompile();
 #endif
 
 #if WITH_EDITOR
@@ -414,7 +417,11 @@ protected:
 
 #if WITH_EDITOR
 private:
+	/** Sends a change notification. Demotes change if the compiled tasks are not significantly changed. */
+	void NotifyGraphStructureChanged(EPCGChangeType ChangeType, bool bForce = false);
+
 	void NotifyGraphChanged(EPCGChangeType ChangeType);
+
 	void OnNodeChanged(UPCGNode* InNode, EPCGChangeType ChangeType);
 
 	void NotifyGraphParametersChanged(EPCGGraphParameterEvent InChangeType, FName InChangedPropertyName);
@@ -431,6 +438,7 @@ private:
 	bool bIsNotifying = false;
 	bool bUserPausedNotificationsInGraphEditor = false;
 	bool bIsInspecting = false;
+	FPCGStack InspectedStack;
 #endif // WITH_EDITOR
 };
 
@@ -466,7 +474,6 @@ public:
 protected:
 #if WITH_EDITOR
 	void OnGraphChanged(UPCGGraphInterface* InGraph, EPCGChangeType ChangeType);
-	void OnGraphStructureChanged(UPCGGraphInterface* InGraph);
 	void NotifyGraphParametersChanged(EPCGGraphParameterEvent InChangeType, FName InChangedPropertyName);
 #endif
 	void OnGraphParametersChanged(UPCGGraphInterface* InGraph, EPCGGraphParameterEvent InChangeType, FName InChangedPropertyName);
diff --git a/Engine/Plugins/PCG/Source/PCG/Public/PCGSubsystem.h b/Engine/Plugins/PCG/Source/PCG/Public/PCGSubsystem.h
--- a/Engine/Plugins/PCG/Source/PCG/Public/PCGSubsystem.h
+++ b/Engine/Plugins/PCG/Source/PCG/Public/PCGSubsystem.h
@@ -237,7 +237,7 @@ public:
 	void DeleteSerializedPartitionActors(bool bOnlyDeleteUnused, bool bOnlyChildren = false);
 
 	/** Propagate to the graph compiler graph changes */
-	void NotifyGraphChanged(UPCGGraph* InGraph);
+	void NotifyGraphChanged(UPCGGraph* InGraph, EPCGChangeType ChangeType);
 
 	/** Update the tracking on a given component. */
 	void UpdateComponentTracking(UPCGComponent* InComponent, bool bShouldDirtyActors, const TArray<FPCGSelectionKey>* OptionalChangedKeys = nullptr) { ActorAndComponentMapping.UpdateTracking(InComponent, bShouldDirtyActors, OptionalChangedKeys); }
diff --git a/Engine/Plugins/PCG/Source/PCGEditor/Private/PCGEditor.cpp b/Engine/Plugins/PCG/Source/PCGEditor/Private/PCGEditor.cpp
--- a/Engine/Plugins/PCG/Source/PCGEditor/Private/PCGEditor.cpp
+++ b/Engine/Plugins/PCG/Source/PCGEditor/Private/PCGEditor.cpp
@@ -110,11 +110,6 @@ void FPCGEditor::Initialize(const EToolkitMode::Type InMode, const TSharedPtr<cl
 {
 	PCGGraphBeingEdited = InPCGGraph;
 
-	if (PCGGraphBeingEdited)
-	{
-		PCGGraphBeingEdited->OnGraphStructureChangedDelegate.AddRaw(this, &FPCGEditor::OnGraphStructureChanged);
-	}
-
 	if (UPCGSubsystem* Subsystem = GetSubsystem())
 	{
 		Subsystem->OnComponentGenerationCompleteOrCancelled.AddRaw(this, &FPCGEditor::OnComponentGenerationCompleteOrCancelled);
@@ -272,7 +267,7 @@ void FPCGEditor::SetStackBeingInspected(const FPCGStack& FullStack)
 
 		if (PCGGraphBeingEdited)
 		{
-			PCGGraphBeingEdited->EnableInspection();
+			PCGGraphBeingEdited->EnableInspection(StackBeingInspected);
 		}
 	}
 
@@ -1891,6 +1886,9 @@ void FPCGEditor::DeleteSelectedNodes()
 		UPCGGraph* PCGGraph = PCGEditorGraph->GetPCGGraph();
 		check(PCGEditorGraph && PCGGraph);
 
+		// DeleteSelectedNodes is called directly from UI command 
+		PCGGraph->PrimeGraphCompilationCache();
+
 		bool bChanged = false;
 		TArray<UPCGNode*> NodesToRemove;
 		{
@@ -2375,8 +2373,6 @@ void FPCGEditor::OnClose()
 		{
 			PCGGraphBeingEdited->ToggleUserPausedNotificationsForEditor();
 		}
-
-		PCGGraphBeingEdited->OnGraphStructureChangedDelegate.RemoveAll(this);
 	}
 
 	if (UPCGSubsystem* Subsystem = GetSubsystem())
@@ -2530,13 +2526,6 @@ void FPCGEditor::JumpToDefinition(const UClass* Class) const
 	}
 }
 
-void FPCGEditor::OnGraphStructureChanged(UPCGGraphInterface* InGraph)
-{
-	check(PCGEditorGraph);
-
-	PCGEditorGraph->UpdateStructuralVisualization(GetPCGComponentBeingInspected(), GetStackBeingInspected());
-}
-
 void FPCGEditor::OnComponentGenerationCompleteOrCancelled()
 {
 	DebugObjectTreeWidget->RequestRefresh();
diff --git a/Engine/Plugins/PCG/Source/PCGEditor/Private/PCGEditor.h b/Engine/Plugins/PCG/Source/PCGEditor/Private/PCGEditor.h
--- a/Engine/Plugins/PCG/Source/PCGEditor/Private/PCGEditor.h
+++ b/Engine/Plugins/PCG/Source/PCGEditor/Private/PCGEditor.h
@@ -299,8 +299,6 @@ private:
 	/** To be called everytime we need to replicate our extra nodes to the underlying PCGGraph */
 	void ReplicateExtraNodes() const;
 
-	void OnGraphStructureChanged(UPCGGraphInterface* InGraph);
-
 	/** Called when a component finishes executing. Useful for updating debugging tools/UIs. */
 	void OnComponentGenerationCompleteOrCancelled();
 
diff --git a/Engine/Plugins/PCG/Source/PCGEditor/Private/PCGEditorGraphNodeBase.cpp b/Engine/Plugins/PCG/Source/PCGEditor/Private/PCGEditorGraphNodeBase.cpp
--- a/Engine/Plugins/PCG/Source/PCGEditor/Private/PCGEditorGraphNodeBase.cpp
+++ b/Engine/Plugins/PCG/Source/PCGEditor/Private/PCGEditorGraphNodeBase.cpp
@@ -412,8 +412,6 @@ void UPCGEditorGraphNodeBase::OnNodeChanged(UPCGNode* InNode, EPCGChangeType Cha
 			StackBeingDebugged = Editor ? Editor->GetStackBeingInspected() : nullptr;
 		}
 
-		ChangeType |= UpdateStructuralVisualization(ComponentBeingDebugged, StackBeingDebugged);
-
 		if (!!(ChangeType & (EPCGChangeType::Structural | EPCGChangeType::Node | EPCGChangeType::Edge | EPCGChangeType::Cosmetic)))
 		{
 			ReconstructNodeOnChange();
diff --git a/Engine/Plugins/PCG/Source/PCGEditor/Private/PCGEditorGraphSchema.cpp b/Engine/Plugins/PCG/Source/PCGEditor/Private/PCGEditorGraphSchema.cpp
--- a/Engine/Plugins/PCG/Source/PCGEditor/Private/PCGEditorGraphSchema.cpp
+++ b/Engine/Plugins/PCG/Source/PCGEditor/Private/PCGEditorGraphSchema.cpp
@@ -197,6 +197,9 @@ bool UPCGEditorGraphSchema::TryCreateConnectionInternal(UEdGraphPin* InA, UEdGra
 	UPCGGraph* PCGGraph = PCGNodeA->GetGraph();
 	check(PCGGraph);
 
+	// UPCGEditorGraphSchema::TryCreateConnectionInternal is called directly by FDragConnection::DroppedOnPin
+	PCGGraph->PrimeGraphCompilationCache();
+
 	// Creates a connection via an intermediate conversion node.
 	auto ConnectViaIntermediate = [this, PCGGraph, NodeA, NodeB, A, B](UPCGNode* IntermediateNode)
 	{
diff --git a/Engine/Plugins/PCG/Source/PCGEditor/Private/PCGEditorGraphSchemaActions.cpp b/Engine/Plugins/PCG/Source/PCGEditor/Private/PCGEditorGraphSchemaActions.cpp
--- a/Engine/Plugins/PCG/Source/PCGEditor/Private/PCGEditorGraphSchemaActions.cpp
+++ b/Engine/Plugins/PCG/Source/PCGEditor/Private/PCGEditorGraphSchemaActions.cpp
@@ -69,6 +69,9 @@ UEdGraphNode* FPCGEditorGraphSchemaAction_NewNativeElement::PerformAction(UEdGra
 	// Important - do not reconstruct the editor graph node/pins midway through this function as this will invalidate FromPin.
 	const FPCGDeferNodeReconstructScope DisableReconstruct(FromPin);
 
+	// Ensure compilation cache is populated before changing graph, so we can compare before/after compiled tasks later to detect change.
+	PCGGraph->PrimeGraphCompilationCache();
+
 	const FScopedTransaction Transaction(*FPCGEditorCommon::ContextIdentifier, LOCTEXT("PCGEditorNewNativeElement", "PCG Editor: New Native Element"), nullptr);
 	EditorGraph->Modify();
 
@@ -145,6 +148,9 @@ UEdGraphNode* FPCGEditorGraphSchemaAction_NewSettingsElement::PerformAction(UEdG
 	// Important - do not reconstruct the editor graph node/pins midway through this function as this will invalidate FromPin.
 	const FPCGDeferNodeReconstructScope DisableReconstruct(FromPin);
 
+	// Ensure compilation cache is populated before changing graph, so we can compare before/after compiled tasks later to detect change.
+	PCGGraph->PrimeGraphCompilationCache();
+
 	bool bCreateInstance = false;
 
 	if(Behavior != EPCGEditorNewSettingsBehavior::Normal)
@@ -320,6 +326,9 @@ UEdGraphNode* FPCGEditorGraphSchemaAction_NewLoadAssetElement::PerformAction(UEd
 	// Important - do not reconstruct the editor graph node/pins midway through this function as this will invalidate FromPin.
 	const FPCGDeferNodeReconstructScope DisableReconstruct(FromPin);
 
+	// Ensure compilation cache is populated before changing graph, so we can compare before/after compiled tasks later to detect change.
+	PCGGraph->PrimeGraphCompilationCache();
+
 	const FScopedTransaction Transaction(*FPCGEditorCommon::ContextIdentifier, LOCTEXT("PCGEditorNewBlueprintELement", "PCG Editor: New Blueprint Element"), nullptr);
 	EditorGraph->Modify();
 
@@ -363,6 +372,9 @@ UEdGraphNode* FPCGEditorGraphSchemaAction_NewBlueprintElement::PerformAction(UEd
 	// Important - do not reconstruct the editor graph node/pins midway through this function as this will invalidate FromPin.
 	const FPCGDeferNodeReconstructScope DisableReconstruct(FromPin);
 
+	// Ensure compilation cache is populated before changing graph, so we can compare before/after compiled tasks later to detect change.
+	PCGGraph->PrimeGraphCompilationCache();
+
 	const FScopedTransaction Transaction(*FPCGEditorCommon::ContextIdentifier, LOCTEXT("PCGEditorNewBlueprintELement", "PCG Editor: New Blueprint Element"), nullptr);
 	EditorGraph->Modify();
 
@@ -414,6 +426,9 @@ UEdGraphNode* FPCGEditorGraphSchemaAction_NewSubgraphElement::PerformAction(UEdG
 	// Important - do not reconstruct the editor graph node/pins midway through this function as this will invalidate FromPin.
 	const FPCGDeferNodeReconstructScope DisableReconstruct(FromPin);
 
+	// Ensure compilation cache is populated before changing graph, so we can compare before/after compiled tasks later to detect change.
+	PCGGraph->PrimeGraphCompilationCache();
+
 	UPCGGraphInterface* Subgraph = CastChecked<UPCGGraphInterface>(SubgraphObjectPath.TryLoad());
 
 	const FScopedTransaction Transaction(*FPCGEditorCommon::ContextIdentifier, LOCTEXT("PCGEditorNewSubgraphElement", "PCG Editor: New Subgraph Element"), nullptr);
@@ -493,6 +508,9 @@ UEdGraphNode* FPCGEditorGraphSchemaAction_NewReroute::PerformAction(class UEdGra
 	// Important - do not reconstruct the editor graph node/pins midway through this function as this will invalidate FromPin.
 	const FPCGDeferNodeReconstructScope DisableReconstruct(FromPin);
 
+	// Ensure compilation cache is populated before changing graph, so we can compare before/after compiled tasks later to detect change.
+	PCGGraph->PrimeGraphCompilationCache();
+
 	const FScopedTransaction Transaction(*FPCGEditorCommon::ContextIdentifier, LOCTEXT("PCGEditorNewReroute", "PCG Editor: New Reroute Node"), nullptr);
 	EditorGraph->Modify();
 
@@ -538,6 +556,9 @@ UEdGraphNode* FPCGEditorGraphSchemaAction_NewNamedRerouteUsage::PerformAction(cl
 	const FScopedTransaction Transaction(*FPCGEditorCommon::ContextIdentifier, LOCTEXT("PCGEditorNewNamedRerouteUsage", "PCG Editor: New Named Reroute Node Usage"), nullptr);
 	EditorGraph->Modify();
 
+	// Ensure compilation cache is populated before changing graph, so we can compare before/after compiled tasks later to detect change.
+	PCGGraph->PrimeGraphCompilationCache();
+
 	PCGGraph->DisableNotificationsForEditor();
 
 	ON_SCOPE_EXIT
@@ -596,6 +617,9 @@ UEdGraphNode* FPCGEditorGraphSchemaAction_NewNamedRerouteDeclaration::PerformAct
 	// Important - do not reconstruct the editor graph node/pins midway through this function as this will invalidate FromPin.
 	const FPCGDeferNodeReconstructScope DisableReconstruct(FromPin);
 
+	// Ensure compilation cache is populated before changing graph, so we can compare before/after compiled tasks later to detect change.
+	PCGGraph->PrimeGraphCompilationCache();
+
 	const FScopedTransaction Transaction(*FPCGEditorCommon::ContextIdentifier, LOCTEXT("PCGEditorNewNamedRerouteDeclaration", "PCG Editor: New Named Reroute Node Declaration"), nullptr);
 	EditorGraph->Modify();
 
diff --git a/Engine/Plugins/PCG/Source/PCGEditor/Private/PCGEditorGraphSchemaActions.h b/Engine/Plugins/PCG/Source/PCGEditor/Private/PCGEditorGraphSchemaActions.h
--- a/Engine/Plugins/PCG/Source/PCGEditor/Private/PCGEditorGraphSchemaActions.h
+++ b/Engine/Plugins/PCG/Source/PCGEditor/Private/PCGEditorGraphSchemaActions.h
@@ -198,7 +198,7 @@ struct FPCGEditorGraphSchemaAction_NewComment : public FEdGraphSchemaAction
 
 	// FEdGraphSchemaAction interface
 	virtual FName GetTypeId() const override { return StaticGetTypeId(); }
-	virtual UEdGraphNode* PerformAction(class UEdGraph* ParentGraph, UEdGraphPin* FromPin, const FVector2D Location, bool bSelectNewNode = true) override;
+	virtual UEdGraphNode* PerformAction(UEdGraph* ParentGraph, UEdGraphPin* FromPin, const FVector2D Location, bool bSelectNewNode = true) override;
 	// End of FEdGraphSchemaAction interface
 };
 
@@ -219,7 +219,7 @@ struct FPCGEditorGraphSchemaAction_NewReroute : public FEdGraphSchemaAction
 
 	// FEdGraphSchemaAction interface
 	virtual FName GetTypeId() const override { return StaticGetTypeId(); }
-	virtual UEdGraphNode* PerformAction(class UEdGraph* ParentGraph, UEdGraphPin* FromPin, const FVector2D Location, bool bSelectNewNode = true) override;
+	virtual UEdGraphNode* PerformAction(UEdGraph* ParentGraph, UEdGraphPin* FromPin, const FVector2D Location, bool bSelectNewNode = true) override;
 	// End of FEdGraphSchemaAction interface
 };
 
@@ -244,7 +244,7 @@ struct FPCGEditorGraphSchemaAction_NewNamedRerouteUsage : public FEdGraphSchemaA
 
 	// FEdGraphSchemaAction interface
 	virtual FName GetTypeId() const override { return StaticGetTypeId(); }
-	virtual UEdGraphNode* PerformAction(class UEdGraph* ParentGraph, UEdGraphPin* FromPin, const FVector2D Location, bool bSelectNewNode = true) override;
+	virtual UEdGraphNode* PerformAction(UEdGraph* ParentGraph, UEdGraphPin* FromPin, const FVector2D Location, bool bSelectNewNode = true) override;
 	// End of FEdGraphSchemaAction interface
 };
 
@@ -265,6 +265,6 @@ struct FPCGEditorGraphSchemaAction_NewNamedRerouteDeclaration : public FEdGraphS
 
 	// FEdGraphSchemaAction interface
 	virtual FName GetTypeId() const override { return StaticGetTypeId(); }
-	virtual UEdGraphNode* PerformAction(class UEdGraph* ParentGraph, UEdGraphPin* FromPin, const FVector2D Location, bool bSelectNewNode = true) override;
+	virtual UEdGraphNode* PerformAction(UEdGraph* ParentGraph, UEdGraphPin* FromPin, const FVector2D Location, bool bSelectNewNode = true) override;
 	// End of FEdGraphSchemaAction interface
-};
\ No newline at end of file
+};
