Added async loading capability to PCG asset loading element.

#rnx

diff --git a/Engine/Plugins/PCG/Source/PCG/Private/Elements/IO/PCGLoadAssetElement.cpp b/Engine/Plugins/PCG/Source/PCG/Private/Elements/IO/PCGLoadAssetElement.cpp
--- a/Engine/Plugins/PCG/Source/PCG/Private/Elements/IO/PCGLoadAssetElement.cpp
+++ b/Engine/Plugins/PCG/Source/PCG/Private/Elements/IO/PCGLoadAssetElement.cpp
@@ -3,6 +3,7 @@
 #include "Elements/IO/PCGLoadAssetElement.h"
 
 #include "PCGModule.h"
+#include "Helpers/PCGDynamicTrackingHelpers.h"
 
 #include "AssetRegistry/AssetData.h"
 
@@ -32,7 +33,7 @@ void UPCGLoadDataAssetSettings::PostEditChangeProperty(FPropertyChangedEvent& Pr
 
 FPCGElementPtr UPCGLoadDataAssetSettings::CreateElement() const
 {
-	return MakeShared<FPCGDataAssetElement>();
+	return MakeShared<FPCGLoadDataAssetElement>();
 }
 
 FString UPCGLoadDataAssetSettings::GetAdditionalTitleInformation() const
@@ -103,7 +104,33 @@ void UPCGLoadDataAssetSettings::UpdateFromData()
 	}
 }
 
-bool FPCGDataAssetElement::ExecuteInternal(FPCGContext* Context) const
+bool FPCGLoadDataAssetElement::PrepareDataInternal(FPCGContext* InContext) const
+{
+	TRACE_CPUPROFILER_EVENT_SCOPE(FPCGLoadDataAssetElement::PrepareData);
+
+	check(InContext);
+	FPCGLoadDataAssetContext* Context = static_cast<FPCGLoadDataAssetContext*>(InContext);
+
+	const UPCGLoadDataAssetSettings* Settings = Context->GetInputSettings<UPCGLoadDataAssetSettings>();
+	check(Settings);
+
+	if (Settings->Asset.IsNull())
+	{
+		return true;
+	}
+
+	// Request load, return false if we need to wait, otherwise continue
+	if (!Context->WasLoadRequested())
+	{
+		return !Context->RequestResourceLoad(Context, { Settings->Asset.ToSoftObjectPath() }, !Settings->bSynchronousLoad);
+	}
+	else
+	{
+		return true;
+	}
+}
+
+bool FPCGLoadDataAssetElement::ExecuteInternal(FPCGContext* Context) const
 {
 	TRACE_CPUPROFILER_EVENT_SCOPE(FPCGDataAssetElement::Execute);
 
@@ -111,6 +138,14 @@ bool FPCGDataAssetElement::ExecuteInternal(FPCGContext* Context) const
 	const UPCGLoadDataAssetSettings* Settings = Context->GetInputSettings<UPCGLoadDataAssetSettings>();
 	check(Settings);
 
+#if WITH_EDITOR
+	if (Context->IsValueOverriden(GET_MEMBER_NAME_CHECKED(UPCGLoadDataAssetSettings, Asset)))
+	{
+		FPCGDynamicTrackingHelper::AddSingleDynamicTrackingKey(Context, FPCGSelectionKey::CreateFromPath(Settings->Asset.ToSoftObjectPath()), /*bIsCulled=*/false);
+	}
+#endif
+
+	// At this point, the data should already be loaded
 	if (UPCGDataAsset* AssetData = Settings->Asset.LoadSynchronous())
 	{
 		Context->OutputData = AssetData->Data;
diff --git a/Engine/Plugins/PCG/Source/PCG/Public/Elements/IO/PCGLoadAssetElement.h b/Engine/Plugins/PCG/Source/PCG/Public/Elements/IO/PCGLoadAssetElement.h
--- a/Engine/Plugins/PCG/Source/PCG/Public/Elements/IO/PCGLoadAssetElement.h
+++ b/Engine/Plugins/PCG/Source/PCG/Public/Elements/IO/PCGLoadAssetElement.h
@@ -5,6 +5,7 @@
 #include "PCGContext.h"
 #include "PCGDataAsset.h"
 #include "PCGSettings.h"
+#include "Async/PCGAsyncLoadingContext.h"
 
 struct FAssetData;
 
@@ -72,10 +73,17 @@ public:
 
 	UPROPERTY(BlueprintReadWrite, EditAnywhere, Category = Settings)
 	bool bWarnIfNoAsset = true;
+
+	/** By default, data table loading is asynchronous, can force it synchronous if needed. */
+	UPROPERTY(BlueprintReadWrite, EditAnywhere, Category = "Settings|Debug")
+	bool bSynchronousLoad = false;
 };
 
-class PCG_API FPCGDataAssetElement : public IPCGElement
+struct FPCGLoadDataAssetContext : public FPCGContext, public IPCGAsyncLoadingContext {};
+
+class PCG_API FPCGLoadDataAssetElement : public IPCGElementWithCustomContext<FPCGLoadDataAssetContext>
 {
 protected:
+	virtual bool PrepareDataInternal(FPCGContext* Context) const override;
 	virtual bool ExecuteInternal(FPCGContext* Context) const override;
 };
