Added similar methods on named reroute nodes like in the material editor. Improved naming a bit and color matching and layout.
Deleting a declaration now will make the usages log errors to make it super obvious it doesn't work.

#rnx

diff --git a/Engine/Plugins/PCG/Source/PCG/Private/Elements/PCGReroute.cpp b/Engine/Plugins/PCG/Source/PCG/Private/Elements/PCGReroute.cpp
--- a/Engine/Plugins/PCG/Source/PCG/Private/Elements/PCGReroute.cpp
+++ b/Engine/Plugins/PCG/Source/PCG/Private/Elements/PCGReroute.cpp
@@ -4,6 +4,7 @@
 
 #include "PCGContext.h"
 #include "PCGEdge.h"
+#include "PCGModule.h"
 #include "PCGNode.h"
 #include "PCGPin.h"
 
@@ -83,7 +84,11 @@ EPCGDataType UPCGNamedRerouteUsageSettings::GetCurrentPinTypes(const UPCGPin* In
 
 bool FPCGRerouteElement::ExecuteInternal(FPCGContext* Context) const
 {
-	ensureMsgf(false, TEXT("Reroute elements are not supposed to execute - reroutes should be culled during graph compilation."));
+	check(Context);
+
+	// Reroute elements are culled during graph compilation unless they have no inbound edge.
+	// In such as case, this is a good place to log an error for user to deal with.
+	PCGE_LOG(Error, GraphAndLog, NSLOCTEXT("PCGRerouteSettings", "DetachedReroute", "Reroute is not linked to anything. Reconnect to recreate to fix the error."));
 	
 	Context->OutputData = Context->InputData;
 	for (FPCGTaggedData& Output : Context->OutputData.TaggedData)
diff --git a/Engine/Plugins/PCG/Source/PCG/Private/Graph/PCGGraphCompiler.cpp b/Engine/Plugins/PCG/Source/PCG/Private/Graph/PCGGraphCompiler.cpp
--- a/Engine/Plugins/PCG/Source/PCG/Private/Graph/PCGGraphCompiler.cpp
+++ b/Engine/Plugins/PCG/Source/PCG/Private/Graph/PCGGraphCompiler.cpp
@@ -979,7 +979,7 @@ void FPCGGraphCompiler::CompileTopGraph(UPCGGraph* InGraph, uint32 GenerationGri
 	if (PCGGraphCompiler::CVarEnableTaskStaticCulling.GetValueOnAnyThread())
 	{
 		// Remove reroute nodes before execution grid setup, as grid linkages need final nodes to connect from/to.
-		CullTasks(CompiledTasks, /*bAddPassthroughWires=*/true, [](const FPCGGraphTask& InTask) { return InTask.Node && Cast<UPCGRerouteSettings>(InTask.Node->GetSettings()); });
+		CullTasks(CompiledTasks, /*bAddPassthroughWires=*/true, [](const FPCGGraphTask& InTask) { return InTask.Node && Cast<UPCGRerouteSettings>(InTask.Node->GetSettings()) && InTask.Node->HasInboundEdges(); });
 
 		// Cull inactive branches downstream of branch nodes with static selection values.
 		CullTasksStaticInactive(CompiledTasks);
diff --git a/Engine/Plugins/PCG/Source/PCG/Public/Elements/PCGReroute.h b/Engine/Plugins/PCG/Source/PCG/Public/Elements/PCGReroute.h
--- a/Engine/Plugins/PCG/Source/PCG/Public/Elements/PCGReroute.h
+++ b/Engine/Plugins/PCG/Source/PCG/Public/Elements/PCGReroute.h
@@ -73,6 +73,8 @@ protected:
 	virtual TArray<FPCGPinProperties> InputPinProperties() const override;
 
 public:
+	/** Very counter-intuitive but reroute nodes are normally culled by other means, if they aren't we want to make sure they log errors. */
+	virtual bool CanCullTaskIfUnwired() const { return false; }
 	EPCGDataType GetCurrentPinTypes(const UPCGPin* InPin) const override;
 
 public:
diff --git a/Engine/Plugins/PCG/Source/PCGEditor/Private/PCGEditor.cpp b/Engine/Plugins/PCG/Source/PCGEditor/Private/PCGEditor.cpp
--- a/Engine/Plugins/PCG/Source/PCGEditor/Private/PCGEditor.cpp
+++ b/Engine/Plugins/PCG/Source/PCGEditor/Private/PCGEditor.cpp
@@ -12,6 +12,7 @@
 #include "PCGInputOutputSettings.h"
 #include "PCGPin.h"
 #include "PCGSubsystem.h"
+#include "Elements/PCGReroute.h"
 #include "Helpers/PCGSubgraphHelpers.h"
 #include "Rendering/SlateRenderer.h"
 #include "Tests/Determinism/PCGDeterminismNativeTests.h"
@@ -384,24 +385,33 @@ void FPCGEditor::JumpToNode(const UEdGraphNode* InNode)
 	}
 }
 
-void FPCGEditor::JumpToNode(const UPCGNode* InNode)
+UPCGEditorGraphNodeBase* FPCGEditor::GetEditorNode(const UPCGNode* InNode)
 {
-	if (!ensure(PCGEditorGraph))
+	if (!ensure(PCGEditorGraph) || !InNode)
 	{
-		return;
+		return nullptr;
 	}
 
-	for (const UEdGraphNode* EdGraphNode : PCGEditorGraph->Nodes)
+	for (UEdGraphNode* EdGraphNode : PCGEditorGraph->Nodes)
 	{
-		if (const UPCGEditorGraphNodeBase* PCGEdGraphNode = Cast<UPCGEditorGraphNodeBase>(EdGraphNode))
+		if (UPCGEditorGraphNodeBase* PCGEdGraphNode = Cast<UPCGEditorGraphNodeBase>(EdGraphNode))
 		{
 			if (PCGEdGraphNode->GetPCGNode() == InNode)
 			{
-				JumpToNode(EdGraphNode);
-				break;
+				return PCGEdGraphNode;
 			}
 		}
 	}
+
+	return nullptr;
+}
+
+void FPCGEditor::JumpToNode(const UPCGNode* InNode)
+{
+	if (const UPCGEditorGraphNodeBase* EditorNode = GetEditorNode(InNode))
+	{
+		JumpToNode(EditorNode);
+	}
 }
 
 void FPCGEditor::RegisterTabSpawners(const TSharedRef<FTabManager>& InTabManager)
@@ -726,6 +736,26 @@ void FPCGEditor::BindCommands()
 		PCGEditorCommands.RenameNode,
 		FExecuteAction::CreateSP(this, &FPCGEditor::OnRenameNode),
 		FCanExecuteAction::CreateSP(this, &FPCGEditor::CanRenameNode));
+
+	GraphEditorCommands->MapAction(
+		PCGEditorCommands.ConvertNamedRerouteToReroute,
+		FExecuteAction::CreateSP(this, &FPCGEditor::OnConvertNamedRerouteToReroute),
+		FCanExecuteAction::CreateSP(this, &FPCGEditor::CanConvertNamedRerouteToReroute));
+
+	GraphEditorCommands->MapAction(
+		PCGEditorCommands.SelectNamedRerouteUsages,
+		FExecuteAction::CreateSP(this, &FPCGEditor::OnSelectNamedRerouteUsages),
+		FCanExecuteAction::CreateSP(this, &FPCGEditor::CanSelectNamedRerouteUsages));
+
+	GraphEditorCommands->MapAction(
+		PCGEditorCommands.SelectNamedRerouteDeclaration,
+		FExecuteAction::CreateSP(this, &FPCGEditor::OnSelectNamedRerouteDeclaration),
+		FCanExecuteAction::CreateSP(this, &FPCGEditor::CanSelectNamedRerouteDeclaration));
+
+	GraphEditorCommands->MapAction(
+		PCGEditorCommands.ConvertRerouteToNamedReroute,
+		FExecuteAction::CreateSP(this, &FPCGEditor::OnConvertRerouteToNamedReroute),
+		FCanExecuteAction::CreateSP(this, &FPCGEditor::CanConvertRerouteToNamedReroute));
 }
 
 void FPCGEditor::OnFind()
@@ -1171,25 +1201,398 @@ bool FPCGEditor::CanRenameNode() const
 	}
 }
 
-void FPCGEditor::OnCollapseNodesInSubgraph()
+bool FPCGEditor::InternalValidationOnAction()
 {
 	if (!GraphEditorWidget.IsValid() || PCGEditorGraph == nullptr)
 	{
 		UE_LOG(LogPCGEditor, Error, TEXT("GraphEditorWidget or PCGEditorGraph is null, aborting"));
-		return;
+		return false;
 	}
 
 	UPCGGraph* PCGGraph = PCGEditorGraph->GetPCGGraph();
 	if (PCGGraph == nullptr)
 	{
 		UE_LOG(LogPCGEditor, Error, TEXT("PCGGraph is null, aborting"));
+		return false;
+	}
+
+	return true;
+}
+
+void FPCGEditor::OnConvertNamedRerouteToReroute()
+{
+	if (!InternalValidationOnAction())
+	{
+		return;
+	}
+
+	UPCGGraph* PCGGraph = PCGEditorGraph->GetPCGGraph();
+	check(PCGGraph);
+
+	TArray<UPCGNode*> DeclarationsToConvert;
+	for (UObject* Object : GraphEditorWidget->GetSelectedNodes())
+	{
+		if (!Object)
+		{
+			continue;
+		}
+
+		if (UPCGEditorGraphNodeNamedRerouteDeclaration* DeclarationNode = Cast<UPCGEditorGraphNodeNamedRerouteDeclaration>(Object))
+		{
+			check(DeclarationNode->GetPCGNode());
+			DeclarationsToConvert.Add(DeclarationNode->GetPCGNode());
+		}
+	}
+
+	if (DeclarationsToConvert.IsEmpty())
+	{
+		return;
+	}
+
+	// Disable graph notifications
+	PCGGraph->DisableNotificationsForEditor();
+
+	// For all found declarations, replace by a normal reroute node and forward the connection from the usages to this reroute instead.
+	TArray<UPCGNode*> NodesToDelete;
+
+	for (UPCGNode* Declaration : DeclarationsToConvert)
+	{
+		NodesToDelete.Add(Declaration);
+
+		// Create reroute node
+		UPCGSettings* RerouteNodeSettings = nullptr;
+		UPCGNode* RerouteNode = PCGGraph->AddNodeOfType(UPCGRerouteSettings::StaticClass(), RerouteNodeSettings);
+
+		// Copy any non-functional information
+		Declaration->TransferEditorProperties(RerouteNode);
+
+		// Force a graph refresh - needed for the other operations/editor callbacks to go through, a bit unfortunate.
+		// TODO: improve this
+		PCGEditorGraph->ReconstructGraph();
+
+		UPCGPin* RerouteInput = RerouteNode->GetInputPin(PCGPinConstants::DefaultInputLabel);
+		UPCGPin* RerouteOutput = RerouteNode->GetOutputPin(PCGPinConstants::DefaultOutputLabel);
+
+		// Copy input edge to the declaration to the new reroute node
+		check(Declaration->GetInputPin(PCGPinConstants::DefaultInputLabel));
+		
+		TArray<UPCGEdge*> InputEdges = Declaration->GetInputPin(PCGPinConstants::DefaultInputLabel)->Edges;
+		ensure(InputEdges.Num() <= 1);
+
+		if(InputEdges.Num() >= 1)
+		{
+			PCGGraph->AddEdge(InputEdges[0]->InputPin->Node, InputEdges[0]->InputPin->Properties.Label, RerouteNode, RerouteInput->Properties.Label);
+		}
+
+		// The reroute declaration has edges to all its usages, and potentially other edges through its out pin.
+		// Keep track of all usages to be deleted but also forward edges from the usages to the newly created reroute.
+		if (UPCGPin* OutputPin = Declaration->GetOutputPin(PCGPinConstants::DefaultOutputLabel))
+		{
+			TArray<UPCGEdge*> Edges = OutputPin->Edges;
+			for(UPCGEdge* Edge : Edges)
+			{
+				PCGGraph->AddEdge(RerouteNode, RerouteOutput->Properties.Label, Edge->OutputPin->Node, Edge->OutputPin->Properties.Label);
+			}
+		}
+
+		if (UPCGPin* InvisiblePin = Declaration->GetOutputPin(PCGNamedRerouteConstants::InvisiblePinLabel))
+		{
+			TArray<UPCGEdge*> Edges = InvisiblePin->Edges;
+			for(UPCGEdge* Edge : Edges)
+			{
+				check(Edge->OutputPin && Edge->OutputPin->Node);
+				UPCGNode* Usage = Edge->OutputPin->Node;
+				NodesToDelete.Add(Usage);
+
+				check(Usage->GetOutputPin(PCGPinConstants::DefaultOutputLabel));
+				TArray<UPCGEdge*> UsageEdges = Usage->GetOutputPin(PCGPinConstants::DefaultOutputLabel)->Edges;
+				for(UPCGEdge* UsageEdge : UsageEdges)
+				{
+					PCGGraph->AddEdge(RerouteNode, RerouteOutput->Properties.Label, UsageEdge->OutputPin->Node, UsageEdge->OutputPin->Properties.Label);
+				}
+			}
+		}
+	}
+
+	// Delete removed ndoes
+	PCGGraph->RemoveNodes(NodesToDelete);
+
+	// Re-enable graph notifications
+	PCGGraph->EnableNotificationsForEditor();
+
+	// Force a graph refresh
+	PCGEditorGraph->ReconstructGraph();
+
+	// Notify the widget
+	GraphEditorWidget->NotifyGraphChanged();	
+}
+
+bool FPCGEditor::CanConvertNamedRerouteToReroute() const
+{
+	if (!GraphEditorWidget)
+	{
+		return false;
+	}
+
+	for (UObject* Object : GraphEditorWidget->GetSelectedNodes())
+	{
+		if (UPCGEditorGraphNodeNamedRerouteDeclaration* DeclarationNode = Cast<UPCGEditorGraphNodeNamedRerouteDeclaration>(Object))
+		{
+			return true;
+		}
+	}
+
+	return false;
+}
+
+void FPCGEditor::OnSelectNamedRerouteUsages()
+{
+	if (!InternalValidationOnAction())
+	{
+		return;
+	}
+
+	const FGraphPanelSelectionSet SelectedNodes = GraphEditorWidget->GetSelectedNodes();
+
+	if (SelectedNodes.Num() != 1)
+	{
 		return;
 	}
 
+	const UPCGEditorGraphNodeNamedRerouteDeclaration* DeclarationNode = nullptr;
+
+	for (const UObject* Object : SelectedNodes)
+	{
+		DeclarationNode = Cast<UPCGEditorGraphNodeNamedRerouteDeclaration>(Object);
+	}
+
+	if (!DeclarationNode || !DeclarationNode->GetPCGNode())
+	{
+		return;
+	}
+
+	GraphEditorWidget->ClearSelectionSet();
+
+	// Some assumptions below - that only usages are connected to the invisible pin.
+	if (const UPCGPin* InvisiblePin = DeclarationNode->GetPCGNode()->GetOutputPin(PCGNamedRerouteConstants::InvisiblePinLabel))
+	{
+		for (const UPCGEdge* Edge : InvisiblePin->Edges)
+		{
+			if (const UPCGNode* Usage = Edge->OutputPin->Node)
+			{
+				GraphEditorWidget->SetNodeSelection(GetEditorNode(Usage), true);
+			}
+		}
+	}
+
+	GraphEditorWidget->ZoomToFit(true);
+}
+
+bool FPCGEditor::CanSelectNamedRerouteUsages() const
+{
+	if (!GraphEditorWidget || GraphEditorWidget->GetSelectedNodes().Num() != 1)
+	{
+		return false;
+	}
+
+	for (const UObject* Object : GraphEditorWidget->GetSelectedNodes())
+	{
+		return Object && Object->IsA<UPCGEditorGraphNodeNamedRerouteDeclaration>();
+	}
+
+	return false;
+}
+
+void FPCGEditor::OnSelectNamedRerouteDeclaration()
+{
+	if (!InternalValidationOnAction())
+	{
+		return;
+	}
+
+	const FGraphPanelSelectionSet SelectedNodes = GraphEditorWidget->GetSelectedNodes();
+	
+	if (SelectedNodes.Num() != 1)
+	{
+		return;
+	}
+
+	for (const UObject* Object : SelectedNodes)
+	{
+		const UPCGEditorGraphNodeNamedRerouteUsage* UsageNode = Cast<UPCGEditorGraphNodeNamedRerouteUsage>(Object);
+
+		if (!UsageNode)
+		{
+			continue;
+		}
+
+		GraphEditorWidget->ClearSelectionSet();
+
+		if (!UsageNode->GetPCGNode())
+		{
+			continue;
+		}
+
+		// Find the declaration node that matches the settings in the Usage node.
+		if (UPCGNamedRerouteUsageSettings* UsageSettings = Cast<UPCGNamedRerouteUsageSettings>(UsageNode->GetPCGNode()->GetSettings()))
+		{
+			if (UsageSettings->Declaration && UsageSettings->Declaration->GetOuter() && UsageSettings->Declaration->GetOuter()->IsA<UPCGNode>())
+			{
+				JumpToNode(Cast<UPCGNode>(UsageSettings->Declaration->GetOuter()));
+				break;
+			}
+		}
+	}
+}
+
+bool FPCGEditor::CanSelectNamedRerouteDeclaration() const
+{
+	if (!GraphEditorWidget || GraphEditorWidget->GetSelectedNodes().Num() != 1)
+	{
+		return false;
+	}
+
+	for (const UObject* Object : GraphEditorWidget->GetSelectedNodes())
+	{
+		return Object && Object->IsA<UPCGEditorGraphNodeNamedRerouteUsage>();
+	}
+
+	return false;
+}
+
+void FPCGEditor::OnConvertRerouteToNamedReroute()
+{
+	if (!InternalValidationOnAction())
+	{
+		return;
+	}
+
+	UPCGGraph* PCGGraph = PCGEditorGraph->GetPCGGraph();
+	check(PCGGraph);
+
+	TArray<UPCGNode*> ReroutesToConvert;
+	for (UObject* Object : GraphEditorWidget->GetSelectedNodes())
+	{
+		if(UPCGEditorGraphNodeReroute* RerouteNode = Cast<UPCGEditorGraphNodeReroute>(Object))
+		{
+			// Skip named reroutes
+			if (RerouteNode->IsA<UPCGEditorGraphNodeNamedRerouteBase>())
+			{
+				continue;
+			}
+
+			check(RerouteNode->GetPCGNode());
+			ReroutesToConvert.Add(RerouteNode->GetPCGNode());
+		}
+	}
+
+	if (ReroutesToConvert.IsEmpty())
+	{
+		return;
+	}
+
+	// Disable graph notifications
+	PCGGraph->DisableNotificationsForEditor();
+
+	// For all found reroutes, replace by a named reroute pair and forward inputs to the declaration and outputs to the usage.
+	TArray<UPCGNode*> NodesToDelete;
+
+	for (UPCGNode* Reroute : ReroutesToConvert)
+	{
+		NodesToDelete.Add(Reroute);
+
+		// Create named reroute declaration
+		UPCGSettings* DeclarationSettings = nullptr;
+		UPCGNode* Declaration = PCGGraph->AddNodeOfType(UPCGNamedRerouteDeclarationSettings::StaticClass(), DeclarationSettings);
+
+		// Create named reroute usage
+		UPCGSettings* UsageSettings = nullptr;
+		UPCGNode* Usage = PCGGraph->AddNodeOfType(UPCGNamedRerouteUsageSettings::StaticClass(), UsageSettings);
+		Cast<UPCGNamedRerouteUsageSettings>(UsageSettings)->Declaration = Cast<UPCGNamedRerouteDeclarationSettings>(DeclarationSettings);
+
+		// Setup non-fonctional information
+		Reroute->TransferEditorProperties(Declaration);
+		Reroute->TransferEditorProperties(Usage);
+		constexpr float PositionOffsetIncrementY = 50.f;
+		Usage->PositionY += PositionOffsetIncrementY;
+
+		// Force a graph refresh - needed for the other operations to go through which a bit unfortunate
+		// TODO improve this
+		PCGEditorGraph->ReconstructGraph();
+
+		// Copy all inputs to the reroute to the declaration inputs
+		check(Reroute->GetInputPin(PCGPinConstants::DefaultInputLabel));
+		TArray<UPCGEdge*> RerouteInputEdges = Reroute->GetInputPin(PCGPinConstants::DefaultInputLabel)->Edges;
+		ensure(RerouteInputEdges.Num() <= 1);
+
+		if(RerouteInputEdges.Num() >= 1)
+		{
+			UPCGEdge* InputEdge = RerouteInputEdges[0];
+			PCGGraph->AddEdge(InputEdge->InputPin->Node, InputEdge->InputPin->Properties.Label, Declaration, PCGPinConstants::DefaultInputLabel);
+			CastChecked<UPCGEditorGraphNodeNamedRerouteDeclaration>(GetEditorNode(Declaration))->SetNodeName(InputEdge->InputPin->Node, InputEdge->InputPin->Properties.Label);
+		}
+
+		// Add invisible edge between declaration and usage
+		PCGGraph->AddEdge(Declaration, PCGNamedRerouteConstants::InvisiblePinLabel, Usage, PCGPinConstants::DefaultInputLabel);
+
+		// Copy all outputs from the reroute to the usage outputs
+		check(Reroute->GetOutputPin(PCGPinConstants::DefaultOutputLabel));
+		TArray<UPCGEdge*> OutputEdges = Reroute->GetOutputPin(PCGPinConstants::DefaultOutputLabel)->Edges;
+		for (UPCGEdge* OutputEdge : OutputEdges)
+		{
+			PCGGraph->AddEdge(Usage, PCGPinConstants::DefaultOutputLabel, OutputEdge->OutputPin->Node, OutputEdge->OutputPin->Properties.Label);
+		}
+	}
+
+	// Delete removed ndoes
+	PCGGraph->RemoveNodes(NodesToDelete);
+
+	// Re-enable graph notifications
+	PCGGraph->EnableNotificationsForEditor();
+
+	// Force a graph refresh
+	PCGEditorGraph->ReconstructGraph();
+
+	// Notify the widget
+	GraphEditorWidget->NotifyGraphChanged();
+}
+
+bool FPCGEditor::CanConvertRerouteToNamedReroute() const
+{
+	if (!GraphEditorWidget)
+	{
+		return false;
+	}
+
+	for (UObject* Object : GraphEditorWidget->GetSelectedNodes())
+	{
+		if (UPCGEditorGraphNodeReroute* RerouteNode = Cast<UPCGEditorGraphNodeReroute>(Object))
+		{
+			if (!RerouteNode->IsA<UPCGEditorGraphNodeNamedRerouteBase>())
+			{
+				return true;
+			}
+		}
+	}
+
+	return false;
+}
+
+void FPCGEditor::OnCollapseNodesInSubgraph()
+{
+	if (!InternalValidationOnAction())
+	{
+		return;
+	}
+
+	UPCGGraph* PCGGraph = PCGEditorGraph->GetPCGGraph();
+	check(PCGGraph);
+
 	// Gather all nodes that will be included in the subgraph, and the extra nodes
 	TArray<UPCGNode*> NodesToCollapse;
 	TArray<UObject*> ExtraNodesToCollapse;
 
+	check(GraphEditorWidget);
 	for (UObject* Object : GraphEditorWidget->GetSelectedNodes())
 	{
 		check(Object);
diff --git a/Engine/Plugins/PCG/Source/PCGEditor/Private/PCGEditor.h b/Engine/Plugins/PCG/Source/PCGEditor/Private/PCGEditor.h
--- a/Engine/Plugins/PCG/Source/PCGEditor/Private/PCGEditor.h
+++ b/Engine/Plugins/PCG/Source/PCGEditor/Private/PCGEditor.h
@@ -241,6 +241,32 @@ private:
 	/** Whether the user can rename the selected node */
 	bool CanRenameNode() const;
 
+	/** Converts a named reroute to reroute, rewiring its edges to a newly created reroute */
+	void OnConvertNamedRerouteToReroute();
+	/** Whether the user can convert the selection to a reroute */
+	bool CanConvertNamedRerouteToReroute() const;
+
+	/** Selects the associated usages of a given reroute declaration */
+	void OnSelectNamedRerouteUsages();
+	/** Whether the user can find the usages from the selection */
+	bool CanSelectNamedRerouteUsages() const;
+
+	/** Selects the associated declaration of a given reroute usage */
+	void OnSelectNamedRerouteDeclaration();
+	/** Whether the user can find the declaration from the selection */
+	bool CanSelectNamedRerouteDeclaration() const;
+
+	/** Converts a selected reroute to a named reroute pair; the inputs will go to the declaration and the outputs will go to the usage */
+	void OnConvertRerouteToNamedReroute();
+	/** Whether the user can convert to a named reroute from the selection */
+	bool CanConvertRerouteToNamedReroute() const;
+
+	/** Internal method that validates a few things (& logs errors) prior to executing actions. */
+	bool InternalValidationOnAction();
+
+	/** Finds editor graph node that matches the provided PCG node */
+	UPCGEditorGraphNodeBase* GetEditorNode(const UPCGNode* InNode);
+
 	void OnAlignTop();
 	void OnAlignMiddle();
 	void OnAlignBottom();
diff --git a/Engine/Plugins/PCG/Source/PCGEditor/Private/PCGEditorCommands.cpp b/Engine/Plugins/PCG/Source/PCGEditor/Private/PCGEditorCommands.cpp
--- a/Engine/Plugins/PCG/Source/PCGEditor/Private/PCGEditorCommands.cpp
+++ b/Engine/Plugins/PCG/Source/PCGEditor/Private/PCGEditorCommands.cpp
@@ -36,6 +36,10 @@ void FPCGEditorCommands::RegisterCommands()
 	UI_COMMAND(ToggleInspect, "Toggle Inspection", "Toggle node inspection for selected node", EUserInterfaceActionType::ToggleButton, FInputChord(EKeys::A));
 	UI_COMMAND(AddSourcePin, "Add Source Pin", "Add new source pin to the current node", EUserInterfaceActionType::Button, FInputChord());
 	UI_COMMAND(RenameNode, "Rename Node", "Rename the selected node", EUserInterfaceActionType::Button, FInputChord(EKeys::F2));
+	UI_COMMAND(ConvertNamedRerouteToReroute, "Convert Named Reroute to Reroute", "Convers the selected named reroute declaration to a normal reroute node", EUserInterfaceActionType::Button, FInputChord());
+	UI_COMMAND(SelectNamedRerouteUsages, "Select Named Reroute Usages", "Selects all usages of this Named Reroute Declaration", EUserInterfaceActionType::Button, FInputChord());
+	UI_COMMAND(SelectNamedRerouteDeclaration, "Select Named Reroute Declaration", "Selects the associated Named Reroute Declaration matching this Named Reroute Usage", EUserInterfaceActionType::Button, FInputChord());
+	UI_COMMAND(ConvertRerouteToNamedReroute, "Convert Reroute to Named Reroute", "Convers the selected reroute node to a Named reroute pair (declaration & usage)", EUserInterfaceActionType::Button, FInputChord());
 }
 
 #undef LOCTEXT_NAMESPACE
diff --git a/Engine/Plugins/PCG/Source/PCGEditor/Private/PCGEditorCommands.h b/Engine/Plugins/PCG/Source/PCGEditor/Private/PCGEditorCommands.h
--- a/Engine/Plugins/PCG/Source/PCGEditor/Private/PCGEditorCommands.h
+++ b/Engine/Plugins/PCG/Source/PCGEditor/Private/PCGEditorCommands.h
@@ -32,4 +32,8 @@ public:
 	TSharedPtr<FUICommandInfo> ToggleInspect;
 	TSharedPtr<FUICommandInfo> AddSourcePin;
 	TSharedPtr<FUICommandInfo> RenameNode;
+	TSharedPtr<FUICommandInfo> ConvertNamedRerouteToReroute;
+	TSharedPtr<FUICommandInfo> SelectNamedRerouteUsages;
+	TSharedPtr<FUICommandInfo> SelectNamedRerouteDeclaration;
+	TSharedPtr<FUICommandInfo> ConvertRerouteToNamedReroute;
 };
diff --git a/Engine/Plugins/PCG/Source/PCGEditor/Private/PCGEditorGraphNodeBase.cpp b/Engine/Plugins/PCG/Source/PCGEditor/Private/PCGEditorGraphNodeBase.cpp
--- a/Engine/Plugins/PCG/Source/PCGEditor/Private/PCGEditorGraphNodeBase.cpp
+++ b/Engine/Plugins/PCG/Source/PCGEditor/Private/PCGEditorGraphNodeBase.cpp
@@ -195,6 +195,24 @@ void UPCGEditorGraphNodeBase::GetNodeContextMenuActions(UToolMenu* Menu, class U
 		Section.AddMenuEntry(FPCGEditorCommands::Get().ExportNodes);
 		Section.AddMenuEntry(FPCGEditorCommands::Get().ConvertToStandaloneNodes);
 		Section.AddMenuEntry(FPCGEditorCommands::Get().RenameNode, LOCTEXT("RenameNode", "Rename"));
+
+		// Special nodes operations
+		if (PCGNode && PCGNode->GetSettings())
+		{
+			if (PCGNode->GetSettings()->IsA<UPCGNamedRerouteDeclarationSettings>())
+			{
+				Section.AddMenuEntry(FPCGEditorCommands::Get().ConvertNamedRerouteToReroute);
+				Section.AddMenuEntry(FPCGEditorCommands::Get().SelectNamedRerouteUsages);
+			}
+			else if (PCGNode->GetSettings()->IsA<UPCGNamedRerouteUsageSettings>())
+			{
+				Section.AddMenuEntry(FPCGEditorCommands::Get().SelectNamedRerouteDeclaration);
+			}
+			else if (PCGNode->GetSettings()->IsA<UPCGRerouteSettings>())
+			{
+				Section.AddMenuEntry(FPCGEditorCommands::Get().ConvertRerouteToNamedReroute);
+			}
+		}
 	}
 
 	{
@@ -204,7 +222,8 @@ void UPCGEditorGraphNodeBase::GetNodeContextMenuActions(UToolMenu* Menu, class U
 			LOCTEXT("PCGNode_SetColor", "Set Node Color"),
 			LOCTEXT("PCGNode_SetColorTooltip", "Sets a specific color on the given node. Note that white maps to the default value"),
 			FSlateIcon(FAppStyle::GetAppStyleSetName(), "ColorPicker.Mode"),
-			FUIAction(FExecuteAction::CreateUObject(const_cast<UPCGEditorGraphNodeBase*>(this), &UPCGEditorGraphNodeBase::OnPickColor)));
+			FUIAction(FExecuteAction::CreateUObject(const_cast<UPCGEditorGraphNodeBase*>(this), &UPCGEditorGraphNodeBase::OnPickColor),
+				FCanExecuteAction::CreateUObject(const_cast<UPCGEditorGraphNodeBase*>(this), &UPCGEditorGraphNodeBase::CanPickColor)));
 
 		Section.AddSubMenu("Alignment", LOCTEXT("AlignmentHeader", "Alignment"), FText(), FNewToolMenuDelegate::CreateLambda([](UToolMenu* AlignmentMenu)
 		{
diff --git a/Engine/Plugins/PCG/Source/PCGEditor/Private/PCGEditorGraphNodeBase.h b/Engine/Plugins/PCG/Source/PCGEditor/Private/PCGEditorGraphNodeBase.h
--- a/Engine/Plugins/PCG/Source/PCGEditor/Private/PCGEditorGraphNodeBase.h
+++ b/Engine/Plugins/PCG/Source/PCGEditor/Private/PCGEditorGraphNodeBase.h
@@ -55,7 +55,7 @@ public:
 	const UPCGNode* GetPCGNode() const { return PCGNode; }
 	void PostCopy();
 	void RebuildAfterPaste();
-	void PostPaste();
+	virtual void PostPaste();
 
 	void SetInspected(bool InIsInspecting) { bIsInspected = InIsInspecting; }
 	bool GetInspected() const { return bIsInspected; }
@@ -125,7 +125,8 @@ protected:
 
 	void OnNodeChanged(UPCGNode* InNode, EPCGChangeType ChangeType);
 	void OnPickColor();
-	void OnColorPicked(FLinearColor NewColor);
+	virtual bool CanPickColor() const { return true; }
+	virtual void OnColorPicked(FLinearColor NewColor);
 	void UpdateCommentBubblePinned();
 	void UpdatePosition();
 
diff --git a/Engine/Plugins/PCG/Source/PCGEditor/Private/PCGEditorGraphNodeReroute.cpp b/Engine/Plugins/PCG/Source/PCGEditor/Private/PCGEditorGraphNodeReroute.cpp
--- a/Engine/Plugins/PCG/Source/PCGEditor/Private/PCGEditorGraphNodeReroute.cpp
+++ b/Engine/Plugins/PCG/Source/PCGEditor/Private/PCGEditorGraphNodeReroute.cpp
@@ -173,6 +173,28 @@ void UPCGEditorGraphNodeNamedRerouteDeclaration::ReconstructNodeOnChange()
 	});
 }
 
+void UPCGEditorGraphNodeNamedRerouteDeclaration::SetNodeName(const UPCGNode* FromNode, FName FromPinName)
+{
+	FString NewName;
+
+	if (FromNode)
+	{
+		NewName = FromNode->GetNodeTitle(EPCGNodeTitleType::ListView).ToString() + " " + FromPinName.ToString();
+	}
+	else if (FromPinName != NAME_None)
+	{
+		NewName = FromPinName.ToString();
+	}
+	else
+	{
+		NewName = TEXT("Reroute");
+	}
+
+	bCanRenameNode = true;
+	OnRenameNode(GetCollisionFreeNodeName(NewName));
+	bCanRenameNode = false;
+}
+
 void UPCGEditorGraphNodeNamedRerouteDeclaration::OnRenameNode(const FString& NewName)
 {
 	Super::OnRenameNode(NewName);
@@ -186,4 +208,67 @@ void UPCGEditorGraphNodeNamedRerouteDeclaration::OnRenameNode(const FString& New
 	});
 
 	ReconstructNodeOnChange();
+}
+
+void UPCGEditorGraphNodeNamedRerouteDeclaration::OnColorPicked(FLinearColor NewColor)
+{
+	Super::OnColorPicked(NewColor);
+
+	// Propagate change to downstream usage nodes
+	ApplyToUsageNodes([&NewColor](UPCGEditorGraphNodeNamedRerouteUsage* RerouteNode)
+	{
+		RerouteNode->OnColorPicked(NewColor);
+	});
+}
+
+void UPCGEditorGraphNodeNamedRerouteDeclaration::PostPaste()
+{
+	Super::PostPaste();
+	FixNodeNameCollision();
+}
+
+void UPCGEditorGraphNodeNamedRerouteDeclaration::FixNodeNameCollision()
+{
+	const FString BaseName = GetNodeTitle(ENodeTitleType::ListView).ToString();
+
+	bCanRenameNode = true;
+	OnRenameNode(GetCollisionFreeNodeName(BaseName));
+	bCanRenameNode = false;
+}
+
+FString UPCGEditorGraphNodeNamedRerouteDeclaration::GetCollisionFreeNodeName(const FString& BaseName) const
+{
+	// If there is another declaration in the graph with the same name, append a _N to the name here.
+	UPCGEditorGraph* EditorGraph = Cast<UPCGEditorGraph>(GetGraph());
+
+	if (!EditorGraph)
+	{
+		return BaseName;
+	}
+
+	FString TentativeName = BaseName;
+	bool bHasNameCollision = true;
+	int32 NameSuffix = 1;
+
+	while (bHasNameCollision)
+	{
+		bHasNameCollision = false;
+
+		for (const TObjectPtr<UEdGraphNode>& EdGraphNode : EditorGraph->Nodes)
+		{
+			if (EdGraphNode && EdGraphNode != this && EdGraphNode->GetNodeTitle(ENodeTitleType::ListView).ToString() == TentativeName)
+			{
+				bHasNameCollision = true;
+				break;
+			}
+		}
+
+		if (bHasNameCollision)
+		{
+			// Implementation note: since the node title is a display string, we need to craft the tentative name to be "stable" against it, otherwise the comparisons will fail.
+			TentativeName = FString::Printf(TEXT("%s %d"), *BaseName, NameSuffix++);
+		}
+	}
+
+	return TentativeName;
 }
\ No newline at end of file
diff --git a/Engine/Plugins/PCG/Source/PCGEditor/Private/PCGEditorGraphNodeReroute.h b/Engine/Plugins/PCG/Source/PCGEditor/Private/PCGEditorGraphNodeReroute.h
--- a/Engine/Plugins/PCG/Source/PCGEditor/Private/PCGEditorGraphNodeReroute.h
+++ b/Engine/Plugins/PCG/Source/PCGEditor/Private/PCGEditorGraphNodeReroute.h
@@ -4,6 +4,8 @@
 
 #include "PCGEditorGraphNode.h"
 
+class UPCGNode;
+
 #include "PCGEditorGraphNodeReroute.generated.h"
 
 UCLASS()
@@ -45,6 +47,7 @@ class UPCGEditorGraphNodeNamedRerouteUsage : public UPCGEditorGraphNodeNamedRero
 
 protected:
 	virtual void RebuildEdgesFromPins_Internal();
+	virtual bool CanPickColor() const override { return false; }
 	virtual FText GetPinFriendlyName(const UPCGPin* InPin) const override;
 };
 
@@ -55,9 +58,15 @@ class UPCGEditorGraphNodeNamedRerouteDeclaration : public UPCGEditorGraphNodeNam
 
 public:
 	virtual void OnRenameNode(const FString& NewName) override;
+	virtual void PostPaste() override;
+
+	void SetNodeName(const UPCGNode* FromNode, FName FromPinName);
+	void FixNodeNameCollision();
+	FString GetCollisionFreeNodeName(const FString& BaseName) const;
 
 protected:
 	virtual FText GetPinFriendlyName(const UPCGPin* InPin) const override;
+	virtual void OnColorPicked(FLinearColor NewColor) override;
 	virtual void ReconstructNodeOnChange() override;
 
 	void ApplyToUsageNodes(TFunctionRef<void(UPCGEditorGraphNodeNamedRerouteUsage*)> Action);
diff --git a/Engine/Plugins/PCG/Source/PCGEditor/Private/PCGEditorGraphSchemaActions.cpp b/Engine/Plugins/PCG/Source/PCGEditor/Private/PCGEditorGraphSchemaActions.cpp
--- a/Engine/Plugins/PCG/Source/PCGEditor/Private/PCGEditorGraphSchemaActions.cpp
+++ b/Engine/Plugins/PCG/Source/PCGEditor/Private/PCGEditorGraphSchemaActions.cpp
@@ -581,6 +581,7 @@ UEdGraphNode* FPCGEditorGraphSchemaAction_NewNamedRerouteUsage::PerformAction(cl
 
 	CastChecked<UPCGNamedRerouteUsageSettings>(DefaultNodeSettings)->Declaration = Declaration;
 	NewPCGNode->NodeTitle = DeclarationNode->NodeTitle;
+	NewPCGNode->NodeTitleColor = DeclarationNode->NodeTitleColor;
 
 	// Create edge from the declaration to the new node
 	PCGGraph->AddEdge(const_cast<UPCGNode*>(DeclarationNode.Get()), PCGNamedRerouteConstants::InvisiblePinLabel, NewPCGNode, PCGPinConstants::DefaultInputLabel);
@@ -644,27 +645,16 @@ UEdGraphNode* FPCGEditorGraphSchemaAction_NewNamedRerouteDeclaration::PerformAct
 	NewPCGNode->PositionX = Location.X;
 	NewPCGNode->PositionY = Location.Y;
 
-	NewPCGNode->NodeTitle = TEXT("Reroute");
+	UPCGNode* FromNode = nullptr;
 
 	if (FromPin)
 	{
-		if (UPCGNamedRerouteDeclarationSettings* Declaration = Cast<UPCGNamedRerouteDeclarationSettings>(DefaultNodeSettings))
-		{
-			UPCGNode* FromNode = CastChecked<UPCGEditorGraphNodeBase>(FromPin->GetOwningNode())->GetPCGNode();
-
-			if (FromNode)
-			{
-				NewPCGNode->NodeTitle = FName(FromNode->GetNodeTitle(EPCGNodeTitleType::ListView).ToString() + " " + FromPin->PinName.ToString());
-			}
-			else
-			{
-				NewPCGNode->NodeTitle = FromPin->PinName;
-			}
-		}
-
+		FromNode = CastChecked<UPCGEditorGraphNodeBase>(FromPin->GetOwningNode())->GetPCGNode();
 		NewNode->AutowireNewNode(FromPin);
 	}
 
+	NewNode->SetNodeName(FromNode, FromPin ? FromPin->PinName : NAME_None);
+
 	return NewNode;
 }
 
diff --git a/Engine/Plugins/PCG/Source/PCGEditor/Private/SPCGEditorGraphVarNode.cpp b/Engine/Plugins/PCG/Source/PCGEditor/Private/SPCGEditorGraphVarNode.cpp
--- a/Engine/Plugins/PCG/Source/PCGEditor/Private/SPCGEditorGraphVarNode.cpp
+++ b/Engine/Plugins/PCG/Source/PCGEditor/Private/SPCGEditorGraphVarNode.cpp
@@ -44,17 +44,21 @@ void SPCGEditorGraphVarNode::UpdateGraphNode()
 
 	SetupErrorReporting();
 
-	bool bNeedsTitle = (Cast<UPCGEditorGraphNodeNamedRerouteBase>(GraphNode) != nullptr);
+	const bool bNeedsTitle = (Cast<UPCGEditorGraphNodeNamedRerouteBase>(GraphNode) != nullptr);
 
 	TSharedPtr<SWidget> TitleArea = SNullWidget::NullWidget.ToSharedPtr();
 
 	if (bNeedsTitle)
 	{
+		const bool bNeedsPadding = (Cast<UPCGEditorGraphNodeNamedRerouteUsage>(GraphNode) != nullptr);
+		FMargin CustomPadding = FMargin(Settings->PaddingLeftOfOutput, Settings->PaddingAbovePin, 0, Settings->PaddingBelowPin);
+
 		TSharedPtr<SNodeTitle> NodeTitle = SNew(SNodeTitle, GraphNode);
 		SAssignNew(TitleArea, SVerticalBox)
 			+ SVerticalBox::Slot()
 			.AutoHeight()
 			.VAlign(VAlign_Center)
+			.Padding(bNeedsPadding ? CustomPadding : FMargin())
 			[
 				CreateTitleWidget(NodeTitle)
 			]
