[PCG] Added the ability to write to Graph Parameter directly from the Getter

#rnx

diff --git a/Engine/Plugins/PCG/Source/PCG/Public/PCGGraph.h b/Engine/Plugins/PCG/Source/PCG/Public/PCGGraph.h
--- a/Engine/Plugins/PCG/Source/PCG/Public/PCGGraph.h
+++ b/Engine/Plugins/PCG/Source/PCG/Public/PCGGraph.h
@@ -160,8 +160,9 @@ public:
 
 	EPropertyBagResult SetGraphParameter(const FName PropertyName, const uint64 Value, const UEnum* Enum);
 
-protected:
 	virtual void OnGraphParametersChanged(EPCGGraphParameterEvent InChangeType, FName InChangedPropertyName) PURE_VIRTUAL(UPCGGraphInterface::OnGraphParametersChanged, )
+
+protected:
 	virtual FInstancedPropertyBag* GetMutableUserParametersStruct() PURE_VIRTUAL(UPCGGraphInterface::GetMutableUserParametersStruct, return nullptr;)
 
 	/** Detecting if we need to refresh the graph depending on the type of change in the Graph Parameter. */
@@ -412,7 +413,6 @@ public:
 	UPROPERTY(BlueprintReadWrite, EditAnywhere, Category = "Runtime Generation")
 	FPCGRuntimeGenerationRadii GenerationRadii;
 
-protected:
 	virtual void OnGraphParametersChanged(EPCGGraphParameterEvent InChangeType, FName InChangedPropertyName) override;
 
 #if WITH_EDITOR
@@ -471,13 +471,14 @@ public:
 	void TeardownCallbacks();
 #endif
 
+	virtual void OnGraphParametersChanged(EPCGGraphParameterEvent InChangeType, FName InChangedPropertyName) override;
+
 protected:
 #if WITH_EDITOR
 	void OnGraphChanged(UPCGGraphInterface* InGraph, EPCGChangeType ChangeType);
 	void NotifyGraphParametersChanged(EPCGGraphParameterEvent InChangeType, FName InChangedPropertyName);
 #endif
 	void OnGraphParametersChanged(UPCGGraphInterface* InGraph, EPCGGraphParameterEvent InChangeType, FName InChangedPropertyName);
-	virtual void OnGraphParametersChanged(EPCGGraphParameterEvent InChangeType, FName InChangedPropertyName) override;
 	void RefreshParameters(EPCGGraphParameterEvent InChangeType, FName InChangedPropertyName = NAME_None);
 	virtual FInstancedPropertyBag* GetMutableUserParametersStruct() override;
 
diff --git a/Engine/Plugins/PCG/Source/PCGEditor/Private/Details/PCGEditableUserParameterDetails.cpp b/Engine/Plugins/PCG/Source/PCGEditor/Private/Details/PCGEditableUserParameterDetails.cpp
new file mode 100644
--- /dev/null
+++ b/Engine/Plugins/PCG/Source/PCGEditor/Private/Details/PCGEditableUserParameterDetails.cpp
@@ -0,0 +1,93 @@
+// Copyright Epic Games, Inc. All Rights Reserved.
+
+#include "Details/PCGEditableUserParameterDetails.h"
+
+#include "PCGCommon.h"
+#include "PCGGraph.h"
+#include "Elements/PCGUserParameterGet.h"
+
+#include "DetailCategoryBuilder.h"
+#include "DetailLayoutBuilder.h"
+#include "Editor.h"
+#include "InstancedPropertyBagStructureDataProvider.h"
+#include "UObject/StructOnScope.h"
+
+#define LOCTEXT_NAMESPACE "PCGEditableUserParameterDetails"
+
+namespace PCGEditableUserParameterDetailsConstants
+{
+	FName UserParametersCategory = TEXT("Instance");
+}
+
+TSharedRef<IDetailCustomization> FPCGEditableUserParameterDetails::MakeInstance()
+{
+	return MakeShareable(new FPCGEditableUserParameterDetails);
+}
+
+void FPCGEditableUserParameterDetails::CustomizeDetails(IDetailLayoutBuilder& DetailBuilder)
+{
+	TArray<TSharedPtr<FStructOnScope>> CustomizedStructs;
+	DetailBuilder.GetStructsBeingCustomized(CustomizedStructs);
+
+	TArray<TWeakObjectPtr<UObject>> SelectedObjects = DetailBuilder.GetSelectedObjects();
+	if (!ensure(!SelectedObjects.IsEmpty()))
+	{
+		return;
+	}
+
+	TWeakObjectPtr<UObject> SettingsObject = SelectedObjects[0];
+	if (!ensure(SettingsObject.IsValid()))
+	{
+		return;
+	}
+
+	if (UPCGUserParameterGetSettings* Settings = Cast<UPCGUserParameterGetSettings>(SettingsObject.Get()))
+	{
+		if (UObject* NodeObject = Settings->GetOuter())
+		{
+			if (UPCGGraphInterface* GraphInterface = Cast<UPCGGraphInterface>(NodeObject->GetOuter()))
+			{
+				// It is safe, because we hook pre/post edit changes that will trigger the callbacks
+				if (FInstancedPropertyBag* UserParameters = GraphInterface->GetMutableUserParametersStruct_Unsafe())
+				{
+					IDetailCategoryBuilder& CategoryBuilder = DetailBuilder.EditCategory(PCGEditableUserParameterDetailsConstants::UserParametersCategory);
+					IDetailPropertyRow* DetailPropertyRow = CategoryBuilder.AddExternalStructureProperty(MakeShared<FInstancePropertyBagStructureDataProvider>(*UserParameters), Settings->PropertyName);
+					TSharedPtr<IPropertyHandle> RowPropertyHandle = DetailPropertyRow->GetPropertyHandle();
+
+					if (RowPropertyHandle.IsValid())
+					{
+						RowPropertyHandle->SetOnPropertyValuePreChange(FSimpleDelegate::CreateLambda([this, GraphInterface = MakeWeakObjectPtr(GraphInterface), PropertyName = Settings->PropertyName]()
+						{
+							check(GEditor);
+							if (GEditor->CanTransact())
+							{
+								GEditor->BeginTransaction(LOCTEXT("EditGraphParameter", "Edit Graph Parameter"));
+							}
+
+							if (GraphInterface.IsValid())
+							{
+								GraphInterface->Modify();
+							}
+						}));
+
+						RowPropertyHandle->SetOnPropertyValueChanged(FSimpleDelegate::CreateLambda([this, GraphInterface = MakeWeakObjectPtr(GraphInterface), PropertyName = Settings->PropertyName]()
+						{
+							if (GraphInterface.IsValid())
+							{
+								GraphInterface->OnGraphParametersChanged(EPCGGraphParameterEvent::ValueModifiedLocally, PropertyName);
+							}
+
+							check(GEditor);
+							if (GEditor->IsTransactionActive())
+							{
+								GEditor->EndTransaction();
+							}
+						}));
+					}
+				}
+			}
+		}
+	}
+}
+
+#undef LOCTEXT_NAMESPACE
\ No newline at end of file
diff --git a/Engine/Plugins/PCG/Source/PCGEditor/Private/PCGEditorModule.cpp b/Engine/Plugins/PCG/Source/PCGEditor/Private/PCGEditorModule.cpp
--- a/Engine/Plugins/PCG/Source/PCGEditor/Private/PCGEditorModule.cpp
+++ b/Engine/Plugins/PCG/Source/PCGEditor/Private/PCGEditorModule.cpp
@@ -24,6 +24,7 @@
 #include "Details/EnumSelectorDetails.h"
 #include "Details/PCGAttributePropertySelectorDetails.h"
 #include "Details/PCGBlueprintSettingsDetails.h"
+#include "Details/PCGEditableUserParameterDetails.h"
 #include "Details/PCGGraphDetails.h"
 #include "Details/PCGGraphInstanceDetails.h"
 #include "Details/PCGInstancedPropertyBagOverrideDetails.h"
@@ -115,6 +116,7 @@ void FPCGEditorModule::RegisterDetailsCustomizations()
 	PropertyEditor.RegisterCustomClassLayout("PCGGraph", FOnGetDetailCustomizationInstance::CreateStatic(&FPCGGraphDetails::MakeInstance));
 	PropertyEditor.RegisterCustomClassLayout("PCGGraphInstance", FOnGetDetailCustomizationInstance::CreateStatic(&FPCGGraphInstanceDetails::MakeInstance));
 	PropertyEditor.RegisterCustomClassLayout("PCGVolume", FOnGetDetailCustomizationInstance::CreateStatic(&FPCGVolumeDetails::MakeInstance));
+	PropertyEditor.RegisterCustomClassLayout("PCGUserParameterGetSettings", FOnGetDetailCustomizationInstance::CreateStatic(&FPCGEditableUserParameterDetails::MakeInstance));
 
 	PropertyEditor.RegisterCustomPropertyTypeLayout("PCGAttributePropertySelector", FOnGetPropertyTypeCustomizationInstance::CreateStatic(&FPCGAttributePropertySelectorDetails::MakeInstance));
 	PropertyEditor.RegisterCustomPropertyTypeLayout("PCGAttributePropertyInputSelector", FOnGetPropertyTypeCustomizationInstance::CreateStatic(&FPCGAttributePropertySelectorDetails::MakeInstance));
diff --git a/Engine/Plugins/PCG/Source/PCGEditor/Public/Details/PCGEditableUserParameterDetails.h b/Engine/Plugins/PCG/Source/PCGEditor/Public/Details/PCGEditableUserParameterDetails.h
new file mode 100644
--- /dev/null
+++ b/Engine/Plugins/PCG/Source/PCGEditor/Public/Details/PCGEditableUserParameterDetails.h
@@ -0,0 +1,16 @@
+// Copyright Epic Games, Inc. All Rights Reserved.
+
+#pragma once
+
+#include "IDetailCustomization.h"
+
+/**
+* A simple Details Customization that will provide a single property from the graph's FInstancePropertyBag
+*/
+class FPCGEditableUserParameterDetails : public IDetailCustomization
+{
+public:
+	static TSharedRef<IDetailCustomization> MakeInstance();
+
+	virtual void CustomizeDetails(IDetailLayoutBuilder& DetailBuilder) override;
+};
\ No newline at end of file
