[PCG] Made some improvements to the Visualize Attribute's default properties

#rnx

diff --git a/Engine/Plugins/PCG/Source/PCG/Public/Elements/PCGVisualizeAttribute.h b/Engine/Plugins/PCG/Source/PCG/Public/Elements/PCGVisualizeAttribute.h
--- a/Engine/Plugins/PCG/Source/PCG/Public/Elements/PCGVisualizeAttribute.h
+++ b/Engine/Plugins/PCG/Source/PCG/Public/Elements/PCGVisualizeAttribute.h
@@ -43,11 +43,11 @@ public:
 
 	/** Prefix the printed value with the point's index. */
 	UPROPERTY(BlueprintReadWrite, EditAnywhere, Category = Settings, meta = (PCG_Overridable))
-	bool bPrefixWithIndex = false;
+	bool bPrefixWithIndex = true;
 
 	/** Prefix the printed value with the attribute's name. */
 	UPROPERTY(BlueprintReadWrite, EditAnywhere, Category = Settings, meta = (PCG_Overridable))
-	bool bPrefixWithAttributeName = true;
+	bool bPrefixWithAttributeName = false;
 
 	/** A local offset from the point's location to draw the text. */
 	UPROPERTY(BlueprintReadWrite, EditAnywhere, Category = Settings, meta = (PCG_Overridable))
@@ -63,7 +63,7 @@ public:
 
 	/** The limit of points to draw debug messages. */
 	UPROPERTY(BlueprintReadWrite, EditAnywhere, Category = Settings, meta = (PCG_Overridable, ClampMin = "1", ClampMax="4096"))
-	int32 PointLimit = 256;
+	int32 PointLimit = 4096;
 
 	/** The visualizer is enabled. Useful for dynamically overriding. */
 	UPROPERTY(BlueprintReadWrite, EditAnywhere, Category = Settings, meta = (PCG_Overridable))
