[PCG] Runtime Quality Branch/Select nodes
* Added 'pcg.Quality' cvar ranging from 0 to 4 (Low, Medium, High, Epic, Cinematic)

diff --git a/Engine/Plugins/PCG/Source/PCG/Private/Elements/ControlFlow/PCGQualityBranch.cpp b/Engine/Plugins/PCG/Source/PCG/Private/Elements/ControlFlow/PCGQualityBranch.cpp
new file mode 100644
--- /dev/null
+++ b/Engine/Plugins/PCG/Source/PCG/Private/Elements/ControlFlow/PCGQualityBranch.cpp
@@ -0,0 +1,202 @@
+// Copyright Epic Games, Inc. All Rights Reserved.
+
+#include "Elements/ControlFlow/PCGQualityBranch.h"
+
+#include "PCGCommon.h"
+#include "PCGComponent.h"
+#include "PCGContext.h"
+#include "PCGModule.h"
+#include "PCGPin.h"
+#include "PCGSubsystem.h"
+#include "Elements/PCGGather.h"
+
+#define LOCTEXT_NAMESPACE "FPCGQualityBranchElement"
+
+#if WITH_EDITOR
+FText UPCGQualityBranchSettings::GetNodeTooltipText() const
+{
+	return LOCTEXT("NodeTooltip", "Control flow node that dynamically routes input data based on 'pcg.Quality' setting.");
+}
+#endif // WITH_EDITOR
+
+FString UPCGQualityBranchSettings::GetAdditionalTitleInformation() const
+{
+	return PCGQualityHelpers::GetQualityPinLabel().ToString();
+}
+
+bool UPCGQualityBranchSettings::IsPinStaticallyActive(const FName& PinLabel) const
+{
+	// TODO: If we decide to have 'realtime' graphs in the future (e.g. for microscattering) where we do not recompile/regen when quality
+	// level changes, we may not be able to statically evaluate which pin is active anymore.
+
+	// If disabled, passthrough to Default pin.
+	if (!bEnabled)
+	{
+		return PinLabel == PCGQualityHelpers::PinLabelDefault;
+	}
+
+	const int32 QualityLevel = UPCGSubsystem::GetPCGQualityLevel();
+	bool bIsQualityPinActive = false;
+
+	// If the Pin matches the QualityLevel, check if the pin is enabled.
+	switch (QualityLevel)
+	{
+		case 0:
+			bIsQualityPinActive = bUseLowPin;
+			break;
+		case 1:
+			bIsQualityPinActive = bUseMediumPin;
+			break;
+		case 2:
+			bIsQualityPinActive = bUseHighPin;
+			break;
+		case 3:
+			bIsQualityPinActive = bUseEpicPin;
+			break;
+		case 4:
+			bIsQualityPinActive = bUseCinematicPin;
+			break;
+		default:
+			break;
+	}
+
+	if (bIsQualityPinActive)
+	{
+		return PinLabel == PCGQualityHelpers::GetQualityPinLabel();
+	}
+	else
+	{
+		return PinLabel == PCGQualityHelpers::PinLabelDefault;
+	}
+
+}
+
+bool UPCGQualityBranchSettings::IsPinUsedByNodeExecution(const UPCGPin* InPin) const
+{
+	if (!InPin->IsOutputPin())
+	{
+		return Super::IsPinUsedByNodeExecution(InPin);
+	}
+
+	return IsPinStaticallyActive(InPin->Properties.Label);
+}
+
+TArray<FPCGPinProperties> UPCGQualityBranchSettings::OutputPinProperties() const
+{
+	TArray<FPCGPinProperties> PinProperties;
+
+	PinProperties.Emplace(PCGQualityHelpers::PinLabelDefault,
+		EPCGDataType::Any,
+		/*bInAllowMultipleConnections=*/true,
+		/*bAllowMultipleData=*/true,
+		LOCTEXT("OutputPinTooltipDefault", "All unusued quality pins funnel through this default pin."));
+
+	if (bUseLowPin)
+	{
+		PinProperties.Emplace(PCGQualityHelpers::PinLabelLow,
+			EPCGDataType::Any,
+			/*bInAllowMultipleConnections=*/true,
+			/*bAllowMultipleData=*/true,
+			LOCTEXT("OutputPinTooltipLow", "Active if pcg.Quality setting is set to 0."));
+	}
+
+	if (bUseMediumPin)
+	{
+		PinProperties.Emplace(PCGQualityHelpers::PinLabelMedium,
+			EPCGDataType::Any,
+			/*bInAllowMultipleConnections=*/true,
+			/*bAllowMultipleData=*/true,
+			LOCTEXT("OutputPinTooltipMedium", "Active if pcg.Quality setting is set to 1."));
+	}
+
+	if (bUseHighPin)
+	{
+		PinProperties.Emplace(PCGQualityHelpers::PinLabelHigh,
+			EPCGDataType::Any,
+			/*bInAllowMultipleConnections=*/true,
+			/*bAllowMultipleData=*/true,
+			LOCTEXT("OutputPinTooltipHigh", "Active if pcg.Quality setting is set to 2."));
+	}
+
+	if (bUseEpicPin)
+	{
+		PinProperties.Emplace(PCGQualityHelpers::PinLabelEpic,
+			EPCGDataType::Any,
+			/*bInAllowMultipleConnections=*/true,
+			/*bAllowMultipleData=*/true,
+			LOCTEXT("OutputPinTooltipEpic", "Active if pcg.Quality setting is set to 3."));
+	}
+
+	if (bUseCinematicPin)
+	{
+		PinProperties.Emplace(PCGQualityHelpers::PinLabelCinematic,
+			EPCGDataType::Any,
+			/*bInAllowMultipleConnections=*/true,
+			/*bAllowMultipleData=*/true,
+			LOCTEXT("OutputPinTooltipCinematic", "Active if pcg.Quality setting is set to 4."));
+	}
+
+	return PinProperties;
+}
+
+FPCGElementPtr UPCGQualityBranchSettings::CreateElement() const
+{
+	return MakeShared<FPCGQualityBranchElement>();
+}
+
+void FPCGQualityBranchElement::GetDependenciesCrc(const FPCGDataCollection& InInput, const UPCGSettings* InSettings, UPCGComponent* InComponent, FPCGCrc& OutCrc) const
+{
+	FPCGCrc Crc;
+	IPCGElement::GetDependenciesCrc(InInput, InSettings, InComponent, Crc);
+
+	// Add the quality level to CRC so that we can fully refresh when quality level changes.
+	Crc.Combine(UPCGSubsystem::GetPCGQualityLevel());
+
+	OutCrc = Crc;
+}
+
+bool FPCGQualityBranchElement::ExecuteInternal(FPCGContext* Context) const
+{
+	TRACE_CPUPROFILER_EVENT_SCOPE(FPCGQualityBranchElement::ExecuteInternal);
+
+	const UPCGQualityBranchSettings* Settings = Context->GetInputSettings<UPCGQualityBranchSettings>();
+	check(Settings);
+
+	if (Context->SourceComponent.IsValid() && !Context->SourceComponent->IsManagedByRuntimeGenSystem())
+	{
+		PCGLog::LogWarningOnGraph(LOCTEXT("UsedOutsideOfRuntimeGen", "Using Runtime Quality Branch node in a non-RuntimeGen execution domain. Set your PCG component's 'Generation Trigger' to 'Generate At Runtime'."), Context);
+	}
+
+	const int32 QualityLevel = UPCGSubsystem::GetPCGQualityLevel();
+
+	FName ActivePinNames[PCGQualityHelpers::NumPins];
+	int32 NumActivePins = 0;
+	int32 ActivePinIndex = 0;
+
+	auto ProcessQualityPin = [QualityLevel, &ActivePinIndex, &ActivePinNames, &NumActivePins](int32 PinQualityLevel, bool bIsPinActive, FName PinLabel)
+	{
+		if (bIsPinActive)
+		{
+			if (QualityLevel == PinQualityLevel)
+			{
+				ActivePinIndex = NumActivePins;
+			}
+
+			ActivePinNames[NumActivePins++] = PinLabel;
+		}
+	};
+
+	ActivePinNames[NumActivePins++] = PCGQualityHelpers::PinLabelDefault;
+	ProcessQualityPin(/*PinQualityLevel=*/0, Settings->bUseLowPin, PCGQualityHelpers::PinLabelLow);
+	ProcessQualityPin(/*PinQualityLevel=*/1, Settings->bUseMediumPin, PCGQualityHelpers::PinLabelMedium);
+	ProcessQualityPin(/*PinQualityLevel=*/2, Settings->bUseHighPin, PCGQualityHelpers::PinLabelHigh);
+	ProcessQualityPin(/*PinQualityLevel=*/3, Settings->bUseEpicPin, PCGQualityHelpers::PinLabelEpic);
+	ProcessQualityPin(/*PinQualityLevel=*/4, Settings->bUseCinematicPin, PCGQualityHelpers::PinLabelCinematic);
+
+	// Reuse the functionality of the Gather Node
+	Context->OutputData = PCGGather::GatherDataForPin(Context->InputData, PCGPinConstants::DefaultInputLabel, ActivePinNames[ActivePinIndex]);
+
+	return true;
+}
+
+#undef LOCTEXT_NAMESPACE
diff --git a/Engine/Plugins/PCG/Source/PCG/Private/Elements/ControlFlow/PCGQualitySelect.cpp b/Engine/Plugins/PCG/Source/PCG/Private/Elements/ControlFlow/PCGQualitySelect.cpp
new file mode 100644
--- /dev/null
+++ b/Engine/Plugins/PCG/Source/PCG/Private/Elements/ControlFlow/PCGQualitySelect.cpp
@@ -0,0 +1,165 @@
+// Copyright Epic Games, Inc. All Rights Reserved.
+
+#include "Elements/ControlFlow/PCGQualitySelect.h"
+
+#include "PCGCommon.h"
+#include "PCGComponent.h"
+#include "PCGContext.h"
+#include "PCGModule.h"
+#include "PCGPin.h"
+#include "PCGSubsystem.h"
+#include "Elements/PCGGather.h"
+
+#define LOCTEXT_NAMESPACE "FPCGQualitySelectElement"
+
+#if WITH_EDITOR
+FText UPCGQualitySelectSettings::GetNodeTooltipText() const
+{
+	return LOCTEXT("NodeTooltip", "Selects from input pins based on 'pcg.Quality' setting.");
+}
+#endif // WITH_EDITOR
+
+FString UPCGQualitySelectSettings::GetAdditionalTitleInformation() const
+{
+	return PCGQualityHelpers::GetQualityPinLabel().ToString();
+}
+
+EPCGDataType UPCGQualitySelectSettings::GetCurrentPinTypes(const UPCGPin* InPin) const
+{
+	check(InPin);
+
+	if (!InPin->IsOutputPin())
+	{
+		return Super::GetCurrentPinTypes(InPin);
+	}
+
+	const EPCGDataType InputType = GetTypeUnionOfIncidentEdges(PCGQualityHelpers::GetQualityPinLabel());
+	return (InputType != EPCGDataType::None) ? InputType : EPCGDataType::Any;
+}
+
+TArray<FPCGPinProperties> UPCGQualitySelectSettings::InputPinProperties() const
+{
+	TArray<FPCGPinProperties> PinProperties;
+
+	PinProperties.Emplace(PCGQualityHelpers::PinLabelDefault,
+		EPCGDataType::Any,
+		/*bInAllowMultipleConnections=*/true,
+		/*bAllowMultipleData=*/true,
+		LOCTEXT("InputPinTooltipDefault", "This default pin will be used for any quality levels that do not have enabled pins."));
+
+	if (bUseLowPin)
+	{
+		PinProperties.Emplace(PCGQualityHelpers::PinLabelLow,
+			EPCGDataType::Any,
+			/*bInAllowMultipleConnections=*/true,
+			/*bAllowMultipleData=*/true,
+			LOCTEXT("InputPinTooltipLow", "Consumed if 'pcg.Quality' setting is set to 0."));
+	}
+
+	if (bUseMediumPin)
+	{
+		PinProperties.Emplace(PCGQualityHelpers::PinLabelMedium,
+			EPCGDataType::Any,
+			/*bInAllowMultipleConnections=*/true,
+			/*bAllowMultipleData=*/true,
+			LOCTEXT("InputPinTooltipMedium", "Consumed if 'pcg.Quality' setting is set to 1."));
+	}
+
+	if (bUseHighPin)
+	{
+		PinProperties.Emplace(PCGQualityHelpers::PinLabelHigh,
+			EPCGDataType::Any,
+			/*bInAllowMultipleConnections=*/true,
+			/*bAllowMultipleData=*/true,
+			LOCTEXT("InputPinTooltipHigh", "Consumed if 'pcg.Quality' setting is set to 2."));
+	}
+
+	if (bUseEpicPin)
+	{
+		PinProperties.Emplace(PCGQualityHelpers::PinLabelEpic,
+			EPCGDataType::Any,
+			/*bInAllowMultipleConnections=*/true,
+			/*bAllowMultipleData=*/true,
+			LOCTEXT("InputPinTooltipEpic", "Consumed if 'pcg.Quality' setting is set to 3."));
+	}
+
+	if (bUseCinematicPin)
+	{
+		PinProperties.Emplace(PCGQualityHelpers::PinLabelCinematic,
+			EPCGDataType::Any,
+			/*bInAllowMultipleConnections=*/true,
+			/*bAllowMultipleData=*/true,
+			LOCTEXT("InputPinTooltipCinematic", "Consumed if 'pcg.Quality' setting is set to 4."));
+	}
+
+	return PinProperties;
+}
+
+TArray<FPCGPinProperties> UPCGQualitySelectSettings::OutputPinProperties() const
+{
+	TArray<FPCGPinProperties> PinProperties;
+	PinProperties.Emplace(PCGPinConstants::DefaultOutputLabel, EPCGDataType::Any);
+
+	return PinProperties;
+}
+
+FPCGElementPtr UPCGQualitySelectSettings::CreateElement() const
+{
+	return MakeShared<FPCGQualitySelectElement>();
+}
+
+void FPCGQualitySelectElement::GetDependenciesCrc(const FPCGDataCollection& InInput, const UPCGSettings* InSettings, UPCGComponent* InComponent, FPCGCrc& OutCrc) const
+{
+	FPCGCrc Crc;
+	IPCGElement::GetDependenciesCrc(InInput, InSettings, InComponent, Crc);
+
+	// Add the quality level to CRC so that we can fully refresh when quality level changes.
+	Crc.Combine(UPCGSubsystem::GetPCGQualityLevel());
+
+	OutCrc = Crc;
+}
+
+bool FPCGQualitySelectElement::ExecuteInternal(FPCGContext* Context) const
+{
+	TRACE_CPUPROFILER_EVENT_SCOPE(FPCGQualitySelectElement::ExecuteInternal);
+
+	const UPCGQualitySelectSettings* Settings = Context->GetInputSettings<UPCGQualitySelectSettings>();
+	check(Settings);
+
+	if (Context->SourceComponent.IsValid() && !Context->SourceComponent->IsManagedByRuntimeGenSystem())
+	{
+		PCGLog::LogWarningOnGraph(LOCTEXT("UsedOutsideOfRuntimeGen", "Using Runtime Quality Select node in a non-RuntimeGen execution domain. Set your PCG component's 'Generation Trigger' to 'Generate At Runtime'."), Context);
+	}
+
+	const int32 QualityLevel = UPCGSubsystem::GetPCGQualityLevel();
+
+	FName ActivePinNames[PCGQualityHelpers::NumPins];
+	int32 NumActivePins = 0;
+	int32 ActivePinIndex = 0;
+
+	auto ProcessQualityPin = [QualityLevel, &ActivePinIndex, &ActivePinNames, &NumActivePins](int32 PinQualityLevel, bool bIsPinActive, FName PinLabel)
+	{
+		if (bIsPinActive)
+		{
+			if (QualityLevel == PinQualityLevel)
+			{
+				ActivePinIndex = NumActivePins;
+			}
+
+			ActivePinNames[NumActivePins++] = PinLabel;
+		}
+	};
+
+	ActivePinNames[NumActivePins++] = PCGQualityHelpers::PinLabelDefault;
+	ProcessQualityPin(/*PinQualityLevel=*/0, Settings->bUseLowPin, PCGQualityHelpers::PinLabelLow);
+	ProcessQualityPin(/*PinQualityLevel=*/1, Settings->bUseMediumPin, PCGQualityHelpers::PinLabelMedium);
+	ProcessQualityPin(/*PinQualityLevel=*/2, Settings->bUseHighPin, PCGQualityHelpers::PinLabelHigh);
+	ProcessQualityPin(/*PinQualityLevel=*/3, Settings->bUseEpicPin, PCGQualityHelpers::PinLabelEpic);
+	ProcessQualityPin(/*PinQualityLevel=*/4, Settings->bUseCinematicPin, PCGQualityHelpers::PinLabelCinematic);
+
+	Context->OutputData = PCGGather::GatherDataForPin(Context->InputData, ActivePinNames[ActivePinIndex], PCGPinConstants::DefaultOutputLabel);
+
+	return true;
+}
+
+#undef LOCTEXT_NAMESPACE
diff --git a/Engine/Plugins/PCG/Source/PCG/Private/PCGActorAndComponentMapping.cpp b/Engine/Plugins/PCG/Source/PCG/Private/PCGActorAndComponentMapping.cpp
--- a/Engine/Plugins/PCG/Source/PCG/Private/PCGActorAndComponentMapping.cpp
+++ b/Engine/Plugins/PCG/Source/PCG/Private/PCGActorAndComponentMapping.cpp
@@ -507,6 +507,19 @@ void FPCGActorAndComponentMapping::ForAllIntersectingPartitionedComponents(const
 	});
 }
 
+void FPCGActorAndComponentMapping::ForAllOriginalComponents(TFunctionRef<void(UPCGComponent*)> InFunc)
+{
+	for (UPCGComponent* Component : PartitionedOctree.GetAllComponents())
+	{
+		InFunc(Component);
+	}
+
+	for (UPCGComponent* Component : NonPartitionedOctree.GetAllComponents())
+	{
+		InFunc(Component);
+	}
+}
+
 TArray<UPCGComponent*> FPCGActorAndComponentMapping::GetAllIntersectingComponents(const FBoxCenterAndExtent& InBounds) const
 {
 	TArray<UPCGComponent*> Result;
diff --git a/Engine/Plugins/PCG/Source/PCG/Private/PCGCommon.cpp b/Engine/Plugins/PCG/Source/PCG/Private/PCGCommon.cpp
--- a/Engine/Plugins/PCG/Source/PCG/Private/PCGCommon.cpp
+++ b/Engine/Plugins/PCG/Source/PCG/Private/PCGCommon.cpp
@@ -1,10 +1,10 @@
 // Copyright Epic Games, Inc. All Rights Reserved.
 
 #include "PCGCommon.h"
+#include "PCGSubsystem.h"
 
 namespace PCGFeatureSwitches
 {
-
 	TAutoConsoleVariable<bool> CVarCheckSamplerMemory{
 		TEXT("pcg.CheckSamplerMemory"),
 		true,
@@ -127,3 +127,35 @@ namespace PCGPinIdHelpers
 		return PinId % PCGPinIdHelpers::PinActiveBitmaskSize;
 	}
 }
+
+namespace PCGQualityHelpers
+{
+	FName GetQualityPinLabel()
+	{
+		const int32 QualityLevel = UPCGSubsystem::GetPCGQualityLevel();
+		FName SelectedPinLabel = NAME_None;
+
+		switch (QualityLevel)
+		{
+			case 0: // Low Quality
+				SelectedPinLabel = PinLabelLow;
+				break;
+			case 1: // Medium Quality
+				SelectedPinLabel = PinLabelMedium;
+				break;
+			case 2: // High Quality
+				SelectedPinLabel = PinLabelHigh;
+				break;
+			case 3: // Epic Quality
+				SelectedPinLabel = PinLabelEpic;
+				break;
+			case 4: // Cinematic Quality
+				SelectedPinLabel = PinLabelCinematic;
+				break;
+			default: // Default to Low Quality if we don't have a valid quality level
+				SelectedPinLabel = PinLabelDefault;
+		}
+
+		return SelectedPinLabel;
+	}
+}
diff --git a/Engine/Plugins/PCG/Source/PCG/Private/PCGGraph.cpp b/Engine/Plugins/PCG/Source/PCG/Private/PCGGraph.cpp
--- a/Engine/Plugins/PCG/Source/PCG/Private/PCGGraph.cpp
+++ b/Engine/Plugins/PCG/Source/PCG/Private/PCGGraph.cpp
@@ -13,6 +13,8 @@
 #include "PCGSubgraph.h"
 #include "Elements/PCGHiGenGridSize.h"
 #include "Elements/PCGUserParameterGet.h"
+#include "Elements/ControlFlow/PCGQualityBranch.h"
+#include "Elements/ControlFlow/PCGQualitySelect.h"
 #include "Graph/PCGGraphCompiler.h"
 #include "Graph/PCGGraphExecutor.h"
 
@@ -1177,6 +1179,32 @@ bool UPCGGraph::Recompile()
 	return bChanged;
 }
 
+void UPCGGraph::OnPCGQualityLevelChanged()
+{
+	bool bContainsQualityControlFlowNode = false;
+
+	for (UPCGNode* Node : Nodes)
+	{
+		if (!Node)
+		{
+			continue;
+		}
+
+		const UPCGSettings* Settings = Node->GetSettings();
+
+		if (Settings && (Settings->IsA<UPCGQualityBranchSettings>() || Settings->IsA<UPCGQualitySelectSettings>()))
+		{
+			Node->OnNodeChangedDelegate.Broadcast(Node, EPCGChangeType::Cosmetic);
+			bContainsQualityControlFlowNode = true;
+		}
+	}
+
+	if (bContainsQualityControlFlowNode)
+	{
+		OnGraphChangedDelegate.Broadcast(this, EPCGChangeType::Structural | EPCGChangeType::GenerationGrid);
+	}
+}
+
 FPCGSelectionKeyToSettingsMap UPCGGraph::GetTrackedActorKeysToSettings() const
 {
 	FPCGSelectionKeyToSettingsMap TagsToSettings;
diff --git a/Engine/Plugins/PCG/Source/PCG/Private/PCGSubsystem.cpp b/Engine/Plugins/PCG/Source/PCG/Private/PCGSubsystem.cpp
--- a/Engine/Plugins/PCG/Source/PCG/Private/PCGSubsystem.cpp
+++ b/Engine/Plugins/PCG/Source/PCG/Private/PCGSubsystem.cpp
@@ -43,6 +43,23 @@ namespace PCGSubsystemConsole
 				}
 			}));
 
+	TAutoConsoleVariable<int32> CVarPCGQuality(
+		TEXT("pcg.Quality"), 2,
+		TEXT("Selects the quality permutation of PCG which impacts Runtime Quality Branch/Select nodes.\n")
+		TEXT(" 0: Low\n")
+		TEXT(" 1: Medium\n")
+		TEXT(" 2: High\n")
+		TEXT(" 3: Epic\n")
+		TEXT(" 4: Cinematic\n"),
+		FConsoleVariableDelegate::CreateLambda([](IConsoleVariable* CVar)
+		{
+			if (UPCGSubsystem* Subsystem = UPCGSubsystem::GetSubsystemForCurrentWorld())
+			{
+				Subsystem->OnPCGQualityLevelChanged();
+			}
+		}),
+		ECVF_Scalability);
+
 #if WITH_EDITOR
 	static FAutoConsoleCommand CommandBuildLandscapeCache(
 		TEXT("pcg.BuildLandscapeCache"),
@@ -259,6 +276,31 @@ APCGWorldActor* UPCGSubsystem::FindPCGWorldActor()
 	return PCGWorldActor;
 }
 
+int32 UPCGSubsystem::GetPCGQualityLevel()
+{
+	return PCGSubsystemConsole::CVarPCGQuality.GetValueOnAnyThread();
+}
+
+void UPCGSubsystem::OnPCGQualityLevelChanged()
+{
+	// Trigger deep refresh of all runtime generated components.
+	ActorAndComponentMapping.ForAllOriginalComponents([this](UPCGComponent* InComponent)
+	{
+		if (InComponent && InComponent->IsManagedByRuntimeGenSystem())
+		{
+#if WITH_EDITOR
+			if (UPCGGraph* Graph = InComponent->GetGraph())
+			{
+				Graph->OnPCGQualityLevelChanged();
+			}
+#else
+			RefreshRuntimeGenComponent(InComponent, EPCGChangeType::GenerationGrid);
+#endif
+
+		}
+	});
+}
+
 #if WITH_EDITOR
 void UPCGSubsystem::DestroyCurrentPCGWorldActor()
 {
diff --git a/Engine/Plugins/PCG/Source/PCG/Public/Elements/ControlFlow/PCGQualityBranch.h b/Engine/Plugins/PCG/Source/PCG/Public/Elements/ControlFlow/PCGQualityBranch.h
new file mode 100644
--- /dev/null
+++ b/Engine/Plugins/PCG/Source/PCG/Public/Elements/ControlFlow/PCGQualityBranch.h
@@ -0,0 +1,59 @@
+// Copyright Epic Games, Inc. All Rights Reserved.
+
+#pragma once
+
+#include "PCGSettings.h"
+
+#include "PCGQualityBranch.generated.h"
+
+/**
+ * Control flow node that dynamically routes input data based on 'pcg.Quality' setting.
+ */
+UCLASS(BlueprintType, ClassGroup = (Procedural), meta = (Keywords = "if bool switch quality"))
+class UPCGQualityBranchSettings : public UPCGSettings
+{
+	GENERATED_BODY()
+
+public:
+	//~Begin UPCGSettings interface
+#if WITH_EDITOR
+	virtual FName GetDefaultNodeName() const override { return FName(TEXT("RuntimeQualityBranch")); }
+	virtual FText GetDefaultNodeTitle() const override { return NSLOCTEXT("FPCGQualityBranchElement", "NodeTitle", "Runtime Quality Branch"); }
+	virtual FText GetNodeTooltipText() const override;
+	virtual EPCGSettingsType GetType() const override { return EPCGSettingsType::ControlFlow; }
+#endif
+	virtual FString GetAdditionalTitleInformation() const override;
+	virtual bool HasDynamicPins() const override { return true; }
+	virtual bool IsPinStaticallyActive(const FName& PinLabel) const override;
+	virtual bool IsPinUsedByNodeExecution(const UPCGPin* InPin) const override;
+
+protected:
+	virtual TArray<FPCGPinProperties> OutputPinProperties() const override;
+	virtual FPCGElementPtr CreateElement() const override;
+	//~End UPCGSettings interface
+
+public:
+	UPROPERTY(BlueprintReadWrite, EditAnywhere, Category = Settings, meta = (PCG_Overridable))
+	bool bUseLowPin = false;
+
+	UPROPERTY(BlueprintReadWrite, EditAnywhere, Category = Settings, meta = (PCG_Overridable))
+	bool bUseMediumPin = false;
+
+	UPROPERTY(BlueprintReadWrite, EditAnywhere, Category = Settings, meta = (PCG_Overridable))
+	bool bUseHighPin = false;
+
+	UPROPERTY(BlueprintReadWrite, EditAnywhere, Category = Settings, meta = (PCG_Overridable))
+	bool bUseEpicPin = false;
+
+	UPROPERTY(BlueprintReadWrite, EditAnywhere, Category = Settings, meta = (PCG_Overridable))
+	bool bUseCinematicPin = false;
+};
+
+class FPCGQualityBranchElement : public IPCGElement
+{
+public:
+	virtual void GetDependenciesCrc(const FPCGDataCollection& InInput, const UPCGSettings* InSettings, UPCGComponent* InComponent, FPCGCrc& OutCrc) const override;
+
+protected:
+	virtual bool ExecuteInternal(FPCGContext* Context) const override;
+};
diff --git a/Engine/Plugins/PCG/Source/PCG/Public/Elements/ControlFlow/PCGQualitySelect.h b/Engine/Plugins/PCG/Source/PCG/Public/Elements/ControlFlow/PCGQualitySelect.h
new file mode 100644
--- /dev/null
+++ b/Engine/Plugins/PCG/Source/PCG/Public/Elements/ControlFlow/PCGQualitySelect.h
@@ -0,0 +1,59 @@
+// Copyright Epic Games, Inc. All Rights Reserved.
+
+#pragma once
+
+#include "PCGSettings.h"
+
+#include "PCGQualitySelect.generated.h"
+
+/**
+ * Selects from input pins based on 'pcg.Quality' setting.
+ */
+UCLASS(BlueprintType, ClassGroup = (Procedural), meta = (Keywords = "if bool switch quality"))
+class UPCGQualitySelectSettings : public UPCGSettings
+{
+	GENERATED_BODY()
+
+public:
+	//~Begin UPCGSettings interface
+#if WITH_EDITOR
+	virtual FName GetDefaultNodeName() const override { return FName(TEXT("RuntimeQualitySelect")); }
+	virtual FText GetDefaultNodeTitle() const override { return NSLOCTEXT("FPCGQualitySelectElement", "NodeTitle", "Runtime Quality Select"); }
+	virtual FText GetNodeTooltipText() const override;
+	virtual EPCGSettingsType GetType() const override { return EPCGSettingsType::ControlFlow; }
+#endif
+	virtual FString GetAdditionalTitleInformation() const override;
+	virtual bool HasDynamicPins() const override { return true; }
+	virtual EPCGDataType GetCurrentPinTypes(const UPCGPin* InPin) const override;
+
+protected:
+	virtual TArray<FPCGPinProperties> InputPinProperties() const override;
+	virtual TArray<FPCGPinProperties> OutputPinProperties() const override;
+	virtual FPCGElementPtr CreateElement() const override;
+	//~End UPCGSettings interface
+
+public:
+	UPROPERTY(BlueprintReadWrite, EditAnywhere, Category = Settings, meta = (PCG_Overridable))
+	bool bUseLowPin = false;
+
+	UPROPERTY(BlueprintReadWrite, EditAnywhere, Category = Settings, meta = (PCG_Overridable))
+	bool bUseMediumPin = false;
+
+	UPROPERTY(BlueprintReadWrite, EditAnywhere, Category = Settings, meta = (PCG_Overridable))
+	bool bUseHighPin = false;
+
+	UPROPERTY(BlueprintReadWrite, EditAnywhere, Category = Settings, meta = (PCG_Overridable))
+	bool bUseEpicPin = false;
+
+	UPROPERTY(BlueprintReadWrite, EditAnywhere, Category = Settings, meta = (PCG_Overridable))
+	bool bUseCinematicPin = false;
+};
+
+class FPCGQualitySelectElement : public IPCGElement
+{
+public:
+	virtual void GetDependenciesCrc(const FPCGDataCollection& InInput, const UPCGSettings* InSettings, UPCGComponent* InComponent, FPCGCrc& OutCrc) const override;
+
+protected:
+	virtual bool ExecuteInternal(FPCGContext* Context) const override;
+};
diff --git a/Engine/Plugins/PCG/Source/PCG/Public/PCGActorAndComponentMapping.h b/Engine/Plugins/PCG/Source/PCG/Public/PCGActorAndComponentMapping.h
--- a/Engine/Plugins/PCG/Source/PCG/Public/PCGActorAndComponentMapping.h
+++ b/Engine/Plugins/PCG/Source/PCG/Public/PCGActorAndComponentMapping.h
@@ -114,6 +114,9 @@ private:
 	/** Call the InFunc function to all partitioned components which bounds intersect 'InBounds'. */
 	void ForAllIntersectingPartitionedComponents(const FBoxCenterAndExtent& InBounds, TFunctionRef<void(UPCGComponent*)> InFunc) const;
 
+	/** Call the InFunc function for all original components (regardless of partitioned or not). */
+	void ForAllOriginalComponents(TFunctionRef<void(UPCGComponent*)> InFunc);
+
 	/** Gather all the PCG components within some bounds. */
 	TArray<UPCGComponent*> GetAllIntersectingComponents(const FBoxCenterAndExtent& InBounds) const;
 
diff --git a/Engine/Plugins/PCG/Source/PCG/Public/PCGCommon.h b/Engine/Plugins/PCG/Source/PCG/Public/PCGCommon.h
--- a/Engine/Plugins/PCG/Source/PCG/Public/PCGCommon.h
+++ b/Engine/Plugins/PCG/Source/PCG/Public/PCGCommon.h
@@ -324,3 +324,17 @@ enum class EPCGNodeTitleType : uint8
 	/** More concise, single line title. */
 	ListView,
 };
+
+namespace PCGQualityHelpers
+{
+	constexpr int32 NumPins = 6;
+	const FName PinLabelDefault = TEXT("Default");
+	const FName PinLabelLow = TEXT("Low");
+	const FName PinLabelMedium = TEXT("Medium");
+	const FName PinLabelHigh = TEXT("High");
+	const FName PinLabelEpic = TEXT("Epic");
+	const FName PinLabelCinematic = TEXT("Cinematic");
+
+	/** Get the pin label associated with the current 'pcg.Quality' value. If the quality level is invalid, it will return the default pin label. */
+	FName GetQualityPinLabel();
+}
diff --git a/Engine/Plugins/PCG/Source/PCG/Public/PCGGraph.h b/Engine/Plugins/PCG/Source/PCG/Public/PCGGraph.h
--- a/Engine/Plugins/PCG/Source/PCG/Public/PCGGraph.h
+++ b/Engine/Plugins/PCG/Source/PCG/Public/PCGGraph.h
@@ -320,6 +320,8 @@ public:
 
 	/** Trigger a recompilation of the relevant permutations of this graph and check for change in the compiled tasks. */
 	bool Recompile();
+
+	void OnPCGQualityLevelChanged();
 #endif
 
 #if WITH_EDITOR
diff --git a/Engine/Plugins/PCG/Source/PCG/Public/PCGSubsystem.h b/Engine/Plugins/PCG/Source/PCG/Public/PCGSubsystem.h
--- a/Engine/Plugins/PCG/Source/PCG/Public/PCGSubsystem.h
+++ b/Engine/Plugins/PCG/Source/PCG/Public/PCGSubsystem.h
@@ -81,6 +81,11 @@ public:
 
 	APCGWorldActor* GetPCGWorldActor();
 	APCGWorldActor* FindPCGWorldActor();
+
+	/** Returns current quality level between Low (0) and Cinematic (4). */
+	static int32 GetPCGQualityLevel();
+	void OnPCGQualityLevelChanged();
+
 #if WITH_EDITOR
 	void DestroyAllPCGWorldActors();
 	void DestroyCurrentPCGWorldActor();
