[PCG] PCG partition actors grid size
- Make sure to set grid size on PCG spawned partition actors (we should extract spawning logic from FActorPartitionWorldPartition::GetActor and use it for PCG in the future).
- Fixup existing PCG partition actors to set GridSize to match PCGGridSize (we should remove PCGGridSize and use GridSize in the future).

diff --git a/Engine/Plugins/PCG/Source/PCG/Private/Grid/PCGPartitionActor.cpp b/Engine/Plugins/PCG/Source/PCG/Private/Grid/PCGPartitionActor.cpp
--- a/Engine/Plugins/PCG/Source/PCG/Private/Grid/PCGPartitionActor.cpp
+++ b/Engine/Plugins/PCG/Source/PCG/Private/Grid/PCGPartitionActor.cpp
@@ -57,6 +57,11 @@ void APCGPartitionActor::PostLoad()
 
 		LocalToOriginalMap_DEPRECATED.Reset();
 	}
+
+	if (GetGridSize() != PCGGridSize)
+	{
+		SetGridSize(PCGGridSize);
+	}
 #endif
 
 	// Safe guard if we ever load a local that was deleted but not removed (like if the user deleted themselves the component)
diff --git a/Engine/Plugins/PCG/Source/PCG/Private/PCGSubsystem.cpp b/Engine/Plugins/PCG/Source/PCG/Private/PCGSubsystem.cpp
--- a/Engine/Plugins/PCG/Source/PCG/Private/PCGSubsystem.cpp
+++ b/Engine/Plugins/PCG/Source/PCG/Private/PCGSubsystem.cpp
@@ -1121,6 +1121,7 @@ APCGPartitionActor* UPCGSubsystem::FindOrCreatePCGPartitionActor(const FGuid& Gu
 #if WITH_EDITOR
 	NewActor->SetLockLocation(true);
 	NewActor->SetActorLabel(SpawnParams.Name.ToString());
+	NewActor->SetGridSize(GridSize);
 #endif
 
 	// Empty GUID if runtime generated, since transient PAs don't need one.
