Fixed erroneous use of required pins in the proxy node when there are no settings provided.

#rnx

diff --git a/Engine/Plugins/PCG/Source/PCG/Private/Elements/PCGIndirectionElement.cpp b/Engine/Plugins/PCG/Source/PCG/Private/Elements/PCGIndirectionElement.cpp
--- a/Engine/Plugins/PCG/Source/PCG/Private/Elements/PCGIndirectionElement.cpp
+++ b/Engine/Plugins/PCG/Source/PCG/Private/Elements/PCGIndirectionElement.cpp
@@ -121,9 +121,18 @@ TArray<FPCGPinProperties> UPCGIndirectionSettings::InputPinProperties() const
 			checkNoEntry();
 	}
 
+	// If we couldn't get a source of pins, use the default pins instead.
 	if (!bSetProperties)
 	{
-		return Super::InputPinProperties();
+		InputProperties = Super::InputPinProperties();
+		// However, unlike the default pins, we don't want to make those as required if the dynamic dispatch does not use them.
+		for (FPCGPinProperties& PinProperty : InputProperties)
+		{
+			if (PinProperty.IsRequiredPin())
+			{
+				PinProperty.SetNormalPin();
+			}
+		}
 	}
 
 	return InputProperties;
diff --git a/Engine/Plugins/PCG/Source/PCG/Private/PCGSubgraph.cpp b/Engine/Plugins/PCG/Source/PCG/Private/PCGSubgraph.cpp
--- a/Engine/Plugins/PCG/Source/PCG/Private/PCGSubgraph.cpp
+++ b/Engine/Plugins/PCG/Source/PCG/Private/PCGSubgraph.cpp
@@ -201,16 +201,26 @@ void UPCGBaseSubgraphSettings::OnSubgraphChanged(UPCGGraphInterface* InGraph, EP
 
 TArray<FPCGPinProperties> UPCGBaseSubgraphSettings::InputPinProperties() const
 {
+	TArray<FPCGPinProperties> InputPins;
 	if (UPCGGraph* Subgraph = GetSubgraph())
 	{
-		TArray<FPCGPinProperties> InputPins = Subgraph->GetInputNode()->InputPinProperties();
+		InputPins = Subgraph->GetInputNode()->InputPinProperties();
 		PCGSubgraphSettings::RemoveAdvancedAndInvisibleOnConnectedPins(Subgraph, InputPins, /*bIsInput=*/true);
-		return InputPins;
 	}
 	else
 	{
-		return Super::InputPinProperties();
+		InputPins = Super::InputPinProperties();
+		// Considering this is likely a case where we'll have dynamic graph dispatch, don't make the default input pins required
+		for (FPCGPinProperties& PinProperties : InputPins)
+		{
+			if (PinProperties.IsRequiredPin())
+			{
+				PinProperties.SetNormalPin();
+			}
+		}
 	}
+
+	return InputPins;
 }
 
 TArray<FPCGPinProperties> UPCGBaseSubgraphSettings::OutputPinProperties() const
