[PCG] Pops up a window to create a subgraph node or loop node when dropping graph assests in graph

diff --git a/Engine/Plugins/PCG/Source/PCGEditor/Private/PCGEditorGraphSchema.cpp b/Engine/Plugins/PCG/Source/PCGEditor/Private/PCGEditorGraphSchema.cpp
--- a/Engine/Plugins/PCG/Source/PCGEditor/Private/PCGEditorGraphSchema.cpp
+++ b/Engine/Plugins/PCG/Source/PCGEditor/Private/PCGEditorGraphSchema.cpp
@@ -632,8 +632,11 @@ void UPCGEditorGraphSchema::DroppedAssetsOnGraph(const TArray<FAssetData>& Asset
 	constexpr float PositionOffsetIncrementY = 50.f;
 	UEdGraphPin* NullFromPin = nullptr;
 
+	TArray<FSoftObjectPath> SubgraphPaths;
+	TArray<FVector2D> SubgraphPositions;
+
 	TArray<FSoftObjectPath> SettingsPaths;
-	TArray<FVector2D> GraphPositions;
+	TArray<FVector2D> SettingsGraphPositions;
 
 	for (const FAssetData& AssetData : Assets)
 	{
@@ -641,9 +644,8 @@ void UPCGEditorGraphSchema::DroppedAssetsOnGraph(const TArray<FAssetData>& Asset
 		{
 			if(AssetClass->IsChildOf(UPCGGraphInterface::StaticClass()))
 			{
-				FPCGEditorGraphSchemaAction_NewSubgraphElement NewSubgraphAction;
-				NewSubgraphAction.SubgraphObjectPath = AssetData.GetSoftObjectPath();
-				NewSubgraphAction.PerformAction(Graph, NullFromPin, GraphPositionOffset);
+				SubgraphPaths.Add(AssetData.GetSoftObjectPath());
+				SubgraphPositions.Add(GraphPositionOffset);
 				GraphPositionOffset.Y += PositionOffsetIncrementY;
 			}
 			else if (PCGEditorUtils::IsAssetPCGBlueprint(AssetData))
@@ -667,22 +669,28 @@ void UPCGEditorGraphSchema::DroppedAssetsOnGraph(const TArray<FAssetData>& Asset
 			{
 				// Delay creation so we can open a menu, once, if needed.
 				SettingsPaths.Add(AssetData.GetSoftObjectPath());
-				GraphPositions.Add(GraphPositionOffset);
+				SettingsGraphPositions.Add(GraphPositionOffset);
 				GraphPositionOffset.Y += PositionOffsetIncrementY;
 			}
 		}
 	}
 
-	// If we've dragged settings assets, we might want to open a menu (ergo this call)
+	UPCGEditorGraph* EditorGraph = CastChecked<UPCGEditorGraph>(Graph);
+
+	TSharedPtr<SGraphEditor> GraphEditor = SGraphEditor::FindGraphEditorForGraph(EditorGraph);
+	const FVector2D MouseCursorLocation = FSlateApplication::Get().GetCursorPos();
+
+	// If we've dragged settings assets or a graph, we might want to open a menu (ergo this call)
 	if (!SettingsPaths.IsEmpty())
 	{
-		UPCGEditorGraph* EditorGraph = CastChecked<UPCGEditorGraph>(Graph);
-		
-		TSharedPtr<SGraphEditor> GraphEditor = SGraphEditor::FindGraphEditorForGraph(EditorGraph);
-		const FVector2D MouseCursorLocation = FSlateApplication::Get().GetCursorPos();
+		check(SettingsPaths.Num() == SettingsGraphPositions.Num());
+		FPCGEditorGraphSchemaAction_NewSettingsElement::MakeSettingsNodesOrContextualMenu(GraphEditor->GetGraphPanel()->AsShared(), MouseCursorLocation, Graph, SettingsPaths, SettingsGraphPositions, /*bSelectNewNodes=*/true);
+	}
 
-		check(SettingsPaths.Num() == GraphPositions.Num());
-		FPCGEditorGraphSchemaAction_NewSettingsElement::MakeSettingsNodesOrContextualMenu(GraphEditor->GetGraphPanel()->AsShared(), MouseCursorLocation, Graph, SettingsPaths, GraphPositions, /*bSelectNewNodes=*/true);
+	if (!SubgraphPaths.IsEmpty())
+	{
+		check(SubgraphPaths.Num() == SubgraphPositions.Num());
+		FPCGEditorGraphSchemaAction_NewSubgraphElement::MakeGraphNodesOrContextualMenu(GraphEditor->GetGraphPanel()->AsShared(), MouseCursorLocation, Graph, SubgraphPaths, SubgraphPositions, /*bSelectNewNodes=*/true);
 	}
 }
 
diff --git a/Engine/Plugins/PCG/Source/PCGEditor/Private/PCGEditorGraphSchemaActions.cpp b/Engine/Plugins/PCG/Source/PCGEditor/Private/PCGEditorGraphSchemaActions.cpp
--- a/Engine/Plugins/PCG/Source/PCGEditor/Private/PCGEditorGraphSchemaActions.cpp
+++ b/Engine/Plugins/PCG/Source/PCGEditor/Private/PCGEditorGraphSchemaActions.cpp
@@ -14,6 +14,7 @@
 #include "PCGSubgraph.h"
 #include "Elements/IO/PCGLoadAssetElement.h"
 #include "Elements/PCGExecuteBlueprint.h"
+#include "Elements/PCGLoopElement.h"
 #include "Elements/PCGReroute.h"
 #include "Elements/PCGUserParameterGet.h"
 
@@ -423,37 +424,165 @@ UEdGraphNode* FPCGEditorGraphSchemaAction_NewSubgraphElement::PerformAction(UEdG
 		return nullptr;
 	}
 
+	UPCGGraphInterface* Graph = CastChecked<UPCGGraphInterface>(SubgraphObjectPath.TryLoad());
+	if (!Graph)
+	{
+		UE_LOG(LogPCGEditor, Error, TEXT("Could not load Graph"));
+		return nullptr;
+	}
+
 	// Important - do not reconstruct the editor graph node/pins midway through this function as this will invalidate FromPin.
 	const FPCGDeferNodeReconstructScope DisableReconstruct(FromPin);
 
 	// Ensure compilation cache is populated before changing graph, so we can compare before/after compiled tasks later to detect change.
 	PCGGraph->PrimeGraphCompilationCache();
+	bool bCreateLoop = false;
+
+	if (Behavior != EPCGEditorNewPCGGraphBehavior::Normal)
+	{
+		bCreateLoop = (Behavior == EPCGEditorNewPCGGraphBehavior::LoopNode);
+	}
+	else
+	{
+		FModifierKeysState ModifierKeys = FSlateApplication::Get().GetModifierKeys();
+		bCreateLoop = ModifierKeys.IsAltDown();
+	}
+
+	return MakeGraphNode(EditorGraph, FromPin, Graph, Location, bSelectNewNode, bCreateLoop);
+}
+
+void FPCGEditorGraphSchemaAction_NewSubgraphElement::MakeGraphNodesOrContextualMenu(const TSharedRef<class SWidget>& InPanel, const FVector2D& InScreenPosition, UEdGraph* InGraph, const TArray<FSoftObjectPath>& InGraphPaths, const TArray<FVector2D>& InLocations, bool bInSelectNewNodes)
+{
+	UPCGEditorGraph* EditorGraph = Cast<UPCGEditorGraph>(InGraph);
+
+	if (!EditorGraph)
+	{
+		return;
+	}
+
+	TArray<UPCGGraphInterface*> Graph;
+	TArray<FVector2D> GraphLocations;
+	check(InGraphPaths.Num() == InLocations.Num());
+
+	for (int32 PathIndex = 0; PathIndex < InGraphPaths.Num(); ++PathIndex)
+	{
+		const FSoftObjectPath& GraphPath = InGraphPaths[PathIndex];
+		if (UPCGGraphInterface* LoadedGraph = Cast<UPCGGraphInterface>(GraphPath.TryLoad()))
+		{
+			Graph.Add(LoadedGraph);
+			GraphLocations.Add(InLocations[PathIndex]);
+		}
+	}
+
+	FModifierKeysState ModifierKeys = FSlateApplication::Get().GetModifierKeys();
+	const bool bModifiedKeysActive = ModifierKeys.IsControlDown() || ModifierKeys.IsAltDown();
+
+	if (bModifiedKeysActive)
+	{
+		MakeGraphNodes(EditorGraph, Graph, GraphLocations, bInSelectNewNodes, ModifierKeys.IsAltDown());
+	}
+	else if (!Graph.IsEmpty())
+	{
+		FMenuBuilder MenuBuilder(true, nullptr);
+		const FText GraphTextName = ((Graph.Num() == 1) ? FText::FromName(Graph[0]->GetFName()) : LOCTEXT("MultipleGraphs", "Multiple Graphs selected"));
+
+		MenuBuilder.BeginSection("GraphDroppedOn", GraphTextName);
+
+		MenuBuilder.AddMenuEntry(
+			FText::Format(LOCTEXT("CreateSubgraphNode", "Create {0} Subgraph Node"), GraphTextName),
+			FText::Format(LOCTEXT("CreateSubgraphToolTip", "Creates a subgraph for Graph asset {0} \n(Ctrl-drop to automatically create a subgraph node)"), GraphTextName),
+			FSlateIcon(),
+			FUIAction(
+				FExecuteAction::CreateStatic(&FPCGEditorGraphSchemaAction_NewSubgraphElement::MakeGraphNodes, EditorGraph, Graph, GraphLocations, bInSelectNewNodes, false)
+			)
+		);
+
+		MenuBuilder.AddMenuEntry(
+			FText::Format(LOCTEXT("CreateLoopNode", "Create {0} Loop Node"), GraphTextName),
+			FText::Format(LOCTEXT("CreateLoopNodeToolTip", "Creates a loop node for Graph asset {0}\n(Alt-drop to automatically create a loop node)"), GraphTextName),
+			FSlateIcon(),
+			FUIAction(
+				FExecuteAction::CreateStatic(&FPCGEditorGraphSchemaAction_NewSubgraphElement::MakeGraphNodes, EditorGraph, Graph, GraphLocations, bInSelectNewNodes, true)
+			)
+		);
+
+		TSharedRef<SWidget> PanelWidget = InPanel;
+		// Show dialog to choose getter vs setter
+		FSlateApplication::Get().PushMenu(
+			PanelWidget,
+			FWidgetPath(),
+			MenuBuilder.MakeWidget(),
+			InScreenPosition,
+			FPopupTransitionEffect(FPopupTransitionEffect::ContextMenu)
+		);
+
+		MenuBuilder.EndSection();
+	}
+}
+
+void FPCGEditorGraphSchemaAction_NewSubgraphElement::MakeGraphNodes(UPCGEditorGraph* InEditorGraph, TArray<UPCGGraphInterface*> InGraph, TArray<FVector2D> InGraphLocations, bool bInSelectNewNodes, bool bInCreateLoop)
+{
+	check(InGraph.Num() == InGraphLocations.Num());
+	for (int32 GraphIndex = 0; GraphIndex < InGraph.Num(); ++GraphIndex)
+	{
+		FPCGEditorGraphSchemaAction_NewSubgraphElement::MakeGraphNode(InEditorGraph, nullptr, InGraph[GraphIndex], InGraphLocations[GraphIndex], bInSelectNewNodes, bInCreateLoop);
+	}
+}
+
+UEdGraphNode* FPCGEditorGraphSchemaAction_NewSubgraphElement::MakeGraphNode(UPCGEditorGraph* InEditorGraph, UEdGraphPin* InFromPin, UPCGGraphInterface* InGraph, const FVector2D& InLocation, bool bInSelectNewNode, bool bInCreateLoop)
+{
+	if (!InEditorGraph || !InGraph)
+	{
+		return nullptr;
+	}
+
+	UPCGGraph* PCGGraph = InEditorGraph->GetPCGGraph();
+	if (!PCGGraph)
+	{
+		return nullptr;
+	}
 
-	UPCGGraphInterface* Subgraph = CastChecked<UPCGGraphInterface>(SubgraphObjectPath.TryLoad());
+	// Important - do not reconstruct the editor graph node/pins midway through this function as this will invalidate InFromPin.
+	const FPCGDeferNodeReconstructScope DisableReconstruct(InFromPin);
 
 	const FScopedTransaction Transaction(*FPCGEditorCommon::ContextIdentifier, LOCTEXT("PCGEditorNewSubgraphElement", "PCG Editor: New Subgraph Element"), nullptr);
-	EditorGraph->Modify();
+	InEditorGraph->Modify();
 
 	UPCGSettings* DefaultNodeSettings = nullptr;
-	UPCGNode* NewPCGNode = PCGGraph->AddNodeOfType(UPCGSubgraphSettings::StaticClass(), DefaultNodeSettings);
-	UPCGSubgraphSettings* DefaultSubgraphSettings = CastChecked<UPCGSubgraphSettings>(DefaultNodeSettings);
-	DefaultSubgraphSettings->SetSubgraph(Subgraph);
+	UPCGNode* NewPCGNode = nullptr;
+	
+	if (!bInCreateLoop)
+	{
+		NewPCGNode = PCGGraph->AddNodeOfType(UPCGSubgraphSettings::StaticClass(), DefaultNodeSettings);
+		UPCGSubgraphSettings* DefaultSubgraphSettings = CastChecked<UPCGSubgraphSettings>(DefaultNodeSettings);
+		DefaultSubgraphSettings->SetSubgraph(InGraph);
+	}
+	else
+	{
+		NewPCGNode = PCGGraph->AddNodeOfType(UPCGLoopSettings::StaticClass(), DefaultNodeSettings);
+		UPCGLoopSettings* DefaultLoopSettings = CastChecked<UPCGLoopSettings>(DefaultNodeSettings);
+		DefaultLoopSettings->SetSubgraph(InGraph);
+	}
 
-	NewPCGNode->UpdateAfterSettingsChangeDuringCreation();
+	if (!NewPCGNode)
+	{
+		UE_LOG(LogPCGEditor, Error, TEXT("Unable to create node"));
+		return nullptr;
+	}
 
-	FGraphNodeCreator<UPCGEditorGraphNode> NodeCreator(*EditorGraph);
-	UPCGEditorGraphNode* NewNode = NodeCreator.CreateUserInvokedNode(bSelectNewNode);
+	FGraphNodeCreator<UPCGEditorGraphNode> NodeCreator(*InEditorGraph);
+	UPCGEditorGraphNode* NewNode = NodeCreator.CreateUserInvokedNode(bInSelectNewNode);
 	NewNode->Construct(NewPCGNode);
-	NewNode->NodePosX = Location.X;
-	NewNode->NodePosY = Location.Y;
+	NewNode->NodePosX = InLocation.X;
+	NewNode->NodePosY = InLocation.Y;
 	NodeCreator.Finalize();
 
-	NewPCGNode->PositionX = Location.X;
-	NewPCGNode->PositionY = Location.Y;
+	NewPCGNode->PositionX = InLocation.X;
+	NewPCGNode->PositionY = InLocation.Y;
 
-	if (FromPin)
+	if (InFromPin)
 	{
-		NewNode->AutowireNewNode(FromPin);
+		NewNode->AutowireNewNode(InFromPin);
 	}
 
 	return NewNode;
diff --git a/Engine/Plugins/PCG/Source/PCGEditor/Private/PCGEditorGraphSchemaActions.h b/Engine/Plugins/PCG/Source/PCGEditor/Private/PCGEditorGraphSchemaActions.h
--- a/Engine/Plugins/PCG/Source/PCGEditor/Private/PCGEditorGraphSchemaActions.h
+++ b/Engine/Plugins/PCG/Source/PCGEditor/Private/PCGEditorGraphSchemaActions.h
@@ -4,6 +4,7 @@
 
 #include "EdGraph/EdGraphSchema.h"
 
+#include "PCGGraph.h"
 #include "PCGSettings.h"
 #include "Templates/SubclassOf.h"
 
@@ -23,6 +24,14 @@ enum class EPCGEditorNewSettingsBehavior : uint8
 	ForceInstance
 };
 
+UENUM()
+enum class EPCGEditorNewPCGGraphBehavior : uint8
+{
+	Normal = 0,
+	SubgraphNode,
+	LoopNode
+};
+
 USTRUCT()
 struct FPCGEditorGraphSchemaAction_NewNativeElement : public FEdGraphSchemaAction
 {
@@ -156,6 +165,7 @@ struct FPCGEditorGraphSchemaAction_NewLoadAssetElement : public FEdGraphSchemaAc
 	//~ End FEdGraphSchemaAction Interface
 };
 
+/** Creates a new subgraph element and allows choosing whether it can be a subgraph node or a loop node. */
 USTRUCT()
 struct FPCGEditorGraphSchemaAction_NewSubgraphElement : public FEdGraphSchemaAction
 {
@@ -174,10 +184,17 @@ struct FPCGEditorGraphSchemaAction_NewSubgraphElement : public FEdGraphSchemaAct
 	UPROPERTY()
 	FSoftObjectPath SubgraphObjectPath;
 
+	UPROPERTY()
+	EPCGEditorNewPCGGraphBehavior Behavior = EPCGEditorNewPCGGraphBehavior::Normal;
+
 	//~ Begin FEdGraphSchemaAction Interface
 	virtual FName GetTypeId() const override { return StaticGetTypeId(); }
 	virtual UEdGraphNode* PerformAction(UEdGraph* ParentGraph, UEdGraphPin* FromPin, const FVector2D Location, bool bSelectNewNode = true) override;
 	//~ End FEdGraphSchemaAction Interface
+
+	static void MakeGraphNodesOrContextualMenu(const TSharedRef<class SWidget>& InPanel, const FVector2D& InScreenPosition, UEdGraph* InGraph, const TArray<FSoftObjectPath>& InGraphPaths, const TArray<FVector2D>& InLocations, bool bInSelectNewNodes);
+	static void MakeGraphNodes(UPCGEditorGraph* InEditorGraph, TArray<UPCGGraphInterface*> InGraph, TArray<FVector2D> InGraphLocations, bool bInSelectNewNodes, bool bInCreateLoop);
+	static UEdGraphNode* MakeGraphNode(UPCGEditorGraph* InEditorGraph, UEdGraphPin* InFromPin, UPCGGraphInterface* InGraph, const FVector2D& InLocation, bool bInSelectNewNode, bool bInCreateLoop);
 };
 
 /** Action to add a 'comment' node to the graph */
