[PCG] Promote logs that describe automated build actions to Display category

diff --git a/Engine/Plugins/PCG/Source/PCG/Private/Utils/PCGGraphExecutionLogging.cpp b/Engine/Plugins/PCG/Source/PCG/Private/Utils/PCGGraphExecutionLogging.cpp
--- a/Engine/Plugins/PCG/Source/PCG/Private/Utils/PCGGraphExecutionLogging.cpp
+++ b/Engine/Plugins/PCG/Source/PCG/Private/Utils/PCGGraphExecutionLogging.cpp
@@ -129,7 +129,7 @@ namespace PCGGraphExecutionLogging
 			return;
 		}
 
-		UE_LOG(LogPCG, Log, TEXT("[%s/%s] --- SCHEDULE GRAPH ---"),
+		UE_LOG(LogPCG, Display, TEXT("[%s/%s] --- SCHEDULE GRAPH ---"),
 			(SourceComponent && SourceComponent->GetOwner()) ? *SourceComponent->GetOwner()->GetName() : TEXT("MISSINGCOMPONENT"),
 			(SourceComponent && SourceComponent->GetGraph()) ? *SourceComponent->GetGraph()->GetName() : TEXT("MISSINGGRAPH"));
 	}
@@ -141,7 +141,7 @@ namespace PCGGraphExecutionLogging
 			return;
 		}
 
-		UE_LOG(LogPCG, Log, TEXT("[%s/%s] --- SCHEDULE GRAPH FOR DEPENDENCY ---"),
+		UE_LOG(LogPCG, Display, TEXT("[%s/%s] --- SCHEDULE GRAPH FOR DEPENDENCY ---"),
 			(InComponent && InComponent->GetOwner()) ? *InComponent->GetOwner()->GetName() : TEXT("MISSINGCOMPONENT"),
 			(InComponent && InComponent->GetGraph()) ? *InComponent->GetGraph()->GetName() : TEXT("MISSINGGRAPH"));
 	}
@@ -179,7 +179,7 @@ namespace PCGGraphExecutionLogging
 			return;
 		}
 
-		UE_LOG(LogPCG, Log, TEXT("[%s/%s] UPCGComponent::PostProcessGraph"),
+		UE_LOG(LogPCG, Display, TEXT("[%s/%s] UPCGComponent::PostProcessGraph"),
 			(InSourceComponent && InSourceComponent->GetOwner()) ? *InSourceComponent->GetOwner()->GetName() : TEXT("MISSINGCOMPONENT"),
 			(InSourceComponent && InSourceComponent->GetGraph()) ? *InSourceComponent->GetGraph()->GetName() : TEXT("MISSINGGRAPH"));
 	}
@@ -193,7 +193,7 @@ namespace PCGGraphExecutionLogging
 
 		for (const UPCGComponent* Component : CancelledComponents)
 		{
-			UE_LOG(LogPCG, Log, TEXT("[%s/%s] Component cancelled"),
+			UE_LOG(LogPCG, Display, TEXT("[%s/%s] Component cancelled"),
 				(Component && Component->GetOwner()) ? *Component->GetOwner()->GetName() : TEXT("MISSINGCOMPONENT"),
 				(Component && Component->GetGraph()) ? *Component->GetGraph()->GetName() : TEXT("MISSINGGRAPH"));
 		}
diff --git a/Engine/Plugins/PCG/Source/PCGEditor/Private/WorldPartitionBuilder/PCGWorldPartitionBuilder.cpp b/Engine/Plugins/PCG/Source/PCGEditor/Private/WorldPartitionBuilder/PCGWorldPartitionBuilder.cpp
--- a/Engine/Plugins/PCG/Source/PCGEditor/Private/WorldPartitionBuilder/PCGWorldPartitionBuilder.cpp
+++ b/Engine/Plugins/PCG/Source/PCGEditor/Private/WorldPartitionBuilder/PCGWorldPartitionBuilder.cpp
@@ -165,15 +165,15 @@ bool UPCGWorldPartitionBuilder::RunInternal(UWorld* World, const FCellInfo& InCe
 	const int NumPackagesDirtied = PendingDirtyPackages.Num() - NumPendingDirtyPackagesBefore;
 	if (!bGeneratedAnyComponent)
 	{
-		UE_LOG(LogPCGWorldPartitionBuilder, Log, TEXT("No components found to generate. No packages will be saved."));
+		UE_LOG(LogPCGWorldPartitionBuilder, Display, TEXT("No components found to generate. No packages will be saved."));
 	}
 	else if (NumPackagesDirtied <= 0)
 	{
-		UE_LOG(LogPCGWorldPartitionBuilder, Log, TEXT("At least one component was generated but no additional packages were dirtied."));
+		UE_LOG(LogPCGWorldPartitionBuilder, Display, TEXT("At least one component was generated but no additional packages were dirtied."));
 	}
 	else
 	{
-		UE_LOG(LogPCGWorldPartitionBuilder, Log, TEXT("Generation complete, %d packages dirtied."), NumPackagesDirtied);
+		UE_LOG(LogPCGWorldPartitionBuilder, Display, TEXT("Generation complete, %d packages dirtied."), NumPackagesDirtied);
 	}
 
 	return bGeneratedAnyComponent;
@@ -202,7 +202,7 @@ bool UPCGWorldPartitionBuilder::PostRun(UWorld* World, FPackageSourceControlHelp
 		}
 	}
 
-	UE_LOG(LogPCGWorldPartitionBuilder, Log, TEXT("PostRun: %d packages modified, %d packages deleted."), DirtyPackages.Num(), PackagesToDelete.Num());
+	UE_LOG(LogPCGWorldPartitionBuilder, Display, TEXT("PostRun: %d packages modified, %d packages deleted."), DirtyPackages.Num(), PackagesToDelete.Num());
 
 	if (!SavePackages(DirtyPackages, PackageHelper))
 	{
@@ -286,7 +286,7 @@ bool PCGWorldPartitionBuilder::GenerateComponents(
 
 	auto WaitForComponentGeneration = [InWorld](const UPCGComponent* InComponent)
 	{
-		UE_LOG(LogPCGWorldPartitionBuilder, Log, TEXT("Completing generation on PCG component on actor '%s' label '%s', graph '%s'"),
+		UE_LOG(LogPCGWorldPartitionBuilder, Display, TEXT("Completing generation on PCG component on actor '%s' label '%s', graph '%s'"),
 			*InComponent->GetOwner()->GetName(),
 			*InComponent->GetOwner()->GetActorNameOrLabel(),
 			*InComponent->GetGraph()->GetName());
@@ -329,7 +329,7 @@ bool PCGWorldPartitionBuilder::GenerateComponents(
 		// Last minute validations, done here just prior to generation (after component has passed all previous filters) to minimize spam.
 		if (!Component->bActivated)
 		{
-			UE_LOG(LogPCGWorldPartitionBuilder, Log, TEXT("'Activated' toggle was set false on PCG component on actor '%s' label '%s' graph '%s'. Component skipped."),
+			UE_LOG(LogPCGWorldPartitionBuilder, Display, TEXT("'Activated' toggle was set false on PCG component on actor '%s' label '%s' graph '%s'. Component skipped."),
 				*Component->GetOwner()->GetName(),
 				*Component->GetOwner()->GetActorNameOrLabel(),
 				*Component->GetGraph()->GetName());
@@ -338,7 +338,7 @@ bool PCGWorldPartitionBuilder::GenerateComponents(
 
 		if (Component->IsManagedByRuntimeGenSystem())
 		{
-			UE_LOG(LogPCGWorldPartitionBuilder, Log, TEXT("PCG component generation trigger is set to run-time generation on actor '%s' label '%s' graph '%s'. Component skipped."),
+			UE_LOG(LogPCGWorldPartitionBuilder, Display, TEXT("PCG component generation trigger is set to run-time generation on actor '%s' label '%s' graph '%s'. Component skipped."),
 				*Component->GetOwner()->GetName(),
 				*Component->GetOwner()->GetActorNameOrLabel(),
 				*Component->GetGraph()->GetName());
@@ -350,7 +350,7 @@ bool PCGWorldPartitionBuilder::GenerateComponents(
 			PCGWorldPartitionBuilder::WaitForAllAsyncEditorProcesses(InWorld);
 		}
 
-		UE_LOG(LogPCGWorldPartitionBuilder, Log, TEXT("Generate PCG component on actor '%s' label '%s', graph '%s'"),
+		UE_LOG(LogPCGWorldPartitionBuilder, Display, TEXT("Generate PCG component on actor '%s' label '%s', graph '%s'"),
 			*Component->GetOwner()->GetName(),
 			*Component->GetOwner()->GetActorNameOrLabel(),
 			*Component->GetGraph()->GetName());
@@ -395,6 +395,11 @@ void PCGWorldPartitionBuilder::GenerateComponent(UPCGComponent* InComponent, UWo
 
 	if (InComponent->GetSerializedEditingMode() == EPCGEditorDirtyMode::LoadAsPreview)
 	{
+		UE_LOG(LogPCGWorldPartitionBuilder, Display, TEXT("Setting PCG editing mode to Load As Preview on actor '%s' label '%s' graph '%s'."),
+			*InComponent->GetOwner()->GetName(),
+			*InComponent->GetOwner()->GetActorNameOrLabel(),
+			*InComponent->GetGraph()->GetName());
+
 		InComponent->SetEditingMode(EPCGEditorDirtyMode::LoadAsPreview, InComponent->GetSerializedEditingMode());
 		InComponent->ChangeTransientState(EPCGEditorDirtyMode::LoadAsPreview);
 	}
