[PCG][Bug] Make Vector Attribute now properly works when X coordinate isn't provided

MetadataOps were taking an early out if Primary Input pin wasn't provided.

+ Added the ability for an ops settings to not require primary pin

#rnx

diff --git a/Engine/Plugins/PCG/Source/PCG/Private/Elements/Metadata/PCGMetadataOpElementBase.cpp b/Engine/Plugins/PCG/Source/PCG/Private/Elements/Metadata/PCGMetadataOpElementBase.cpp
--- a/Engine/Plugins/PCG/Source/PCG/Private/Elements/Metadata/PCGMetadataOpElementBase.cpp
+++ b/Engine/Plugins/PCG/Source/PCG/Private/Elements/Metadata/PCGMetadataOpElementBase.cpp
@@ -383,8 +383,8 @@ bool FPCGMetadataElementBase::PrepareDataInternal(FPCGContext* Context) const
 	const FName PrimaryPinLabel = Settings->GetInputPinLabel(Settings->GetInputPinToForward());
 	const TArray<FPCGTaggedData> PrimaryInputs = Context->InputData.GetInputsByPin(PrimaryPinLabel);
 
-	// There are no inputs on the primary pin, so passthrough inputs
-	if (PrimaryInputs.IsEmpty())
+	// There are no inputs on the primary pin, so pass-through inputs if the primary pin is required
+	if (Settings->IsPrimaryInputPinRequired() && PrimaryInputs.IsEmpty())
 	{
 		return true;
 	}
diff --git a/Engine/Plugins/PCG/Source/PCG/Public/Elements/Metadata/PCGMetadataMakeVector.h b/Engine/Plugins/PCG/Source/PCG/Public/Elements/Metadata/PCGMetadataMakeVector.h
--- a/Engine/Plugins/PCG/Source/PCG/Public/Elements/Metadata/PCGMetadataMakeVector.h
+++ b/Engine/Plugins/PCG/Source/PCG/Public/Elements/Metadata/PCGMetadataMakeVector.h
@@ -58,6 +58,7 @@ public:
 #endif // WITH_EDITOR
 	//~End UPCGSettings interface
 
+	//~Begin UPCGMetadataSettingsBase interface
 	virtual FPCGAttributePropertyInputSelector GetInputSource(uint32 Index) const override;
 
 	virtual FName GetInputPinLabel(uint32 Index) const override;
@@ -66,6 +67,10 @@ public:
 	virtual bool IsSupportedInputType(uint16 TypeId, uint32 InputIndex, bool& bHasSpecialRequirement) const override;
 	virtual uint16 GetOutputType(uint16 InputTypeId) const override;
 
+	/** If the primary pin is required and missing, some operations can be bypassed. Some, like Make Vector, can replace missing pins with default values. */
+	virtual bool IsPrimaryInputPinRequired() const override { return false; }
+	//~End UPCGMetadataSettingsBase interface
+
 protected:
 	//~Begin UPCGSettings interface
 	virtual FPCGElementPtr CreateElement() const override;
diff --git a/Engine/Plugins/PCG/Source/PCG/Public/Elements/Metadata/PCGMetadataOpElementBase.h b/Engine/Plugins/PCG/Source/PCG/Public/Elements/Metadata/PCGMetadataOpElementBase.h
--- a/Engine/Plugins/PCG/Source/PCG/Public/Elements/Metadata/PCGMetadataOpElementBase.h
+++ b/Engine/Plugins/PCG/Source/PCG/Public/Elements/Metadata/PCGMetadataOpElementBase.h
@@ -126,6 +126,9 @@ public:
 	/* Return the current input pin to forward to the output. */
 	uint32 GetInputPinToForward() const;
 
+	/** If the primary input is needed and empty, it will pass through data. */
+	virtual bool IsPrimaryInputPinRequired() const { return true; }
+
 	UPROPERTY(BlueprintReadWrite, EditAnywhere, Category = "Output", meta = (PCG_Overridable))
 	FPCGAttributePropertyOutputSelector OutputTarget;
 
