[PCG][Bug] Fixed a bug allowing property selector details to remain editable when they otherwise should not be

#rnx

diff --git a/Engine/Plugins/PCG/Source/PCGEditor/Private/Details/PCGAttributePropertySelectorDetails.cpp b/Engine/Plugins/PCG/Source/PCGEditor/Private/Details/PCGAttributePropertySelectorDetails.cpp
--- a/Engine/Plugins/PCG/Source/PCGEditor/Private/Details/PCGAttributePropertySelectorDetails.cpp
+++ b/Engine/Plugins/PCG/Source/PCGEditor/Private/Details/PCGAttributePropertySelectorDetails.cpp
@@ -132,10 +132,15 @@ EVisibility FPCGAttributePropertySelectorDetails::ExtraMenuVisibility() const
 
 bool FPCGAttributePropertySelectorDetails::IsEnabled() const
 {
+	if (!PropertyHandle->IsEditable())
+	{
+		return false;
+	}
+
 	FStructProperty* StructProperty = CastFieldChecked<FStructProperty>(PropertyHandle->GetProperty());
 	TArray<UObject*> Outers;
 	PropertyHandle->GetOuterObjects(Outers);
-	for (UObject* Outer : Outers)
+	for (const UObject* Outer : Outers)
 	{
 		if (Outer && !Outer->CanEditChange(StructProperty))
 		{
