[PCG] Disable warnings on grid size nodes if higen disabled or if component is not partitioned.

A higen graph should be able to run in other contexts (non-partitioned) without noisey warnings. Editor visualization added to make grid size node appear disabled when inactive.

diff --git a/Engine/Plugins/PCG/Source/PCG/Private/Elements/PCGHiGenGridSize.cpp b/Engine/Plugins/PCG/Source/PCG/Private/Elements/PCGHiGenGridSize.cpp
--- a/Engine/Plugins/PCG/Source/PCG/Private/Elements/PCGHiGenGridSize.cpp
+++ b/Engine/Plugins/PCG/Source/PCG/Private/Elements/PCGHiGenGridSize.cpp
@@ -31,7 +31,11 @@ FText UPCGHiGenGridSizeSettings::GetDefaultNodeTitle() const
 
 FText UPCGHiGenGridSizeSettings::GetNodeTooltipText() const
 {
-	return LOCTEXT("NodeTooltip", "Set the execution grid size for downstream nodes. Enables executing a single graph across a hierarchy of grids. Has no effect if generating component is not partitioned.");
+	return LOCTEXT("NodeTooltip", "Set the execution grid size for downstream nodes. Enables executing a single graph across a hierarchy of grids."
+		"\n\nHas no effect if any of the following are true:"
+		"\n\t* Generating PCG component is not set to Partitioned."
+		"\n\t* Hierarchical Generation is disabled in the graph settings."
+		"\n\t* Executed in a subgraph, as subgraphs are always invoked on parent grid level.");
 }
 #endif
 
@@ -113,19 +117,7 @@ EPCGChangeType UPCGHiGenGridSizeSettings::GetChangeTypeForProperty(const FName&
 bool FPCGHiGenGridSizeElement::ExecuteInternal(FPCGContext* Context) const
 {
 	TRACE_CPUPROFILER_EVENT_SCOPE(FPCGHiGenGridSizeElement::Execute);
-
-	// Validation
 	check(Context);
-	if (Context->Node && Context->Node->GetGraph() && !Context->Node->GetGraph()->IsHierarchicalGenerationEnabled())
-	{
-		PCGE_LOG(Warning, GraphAndLog, LOCTEXT("GridSizeUsedInNonHiGenGraph", "Grid Size node used in a non-hierarchical graph. Enable hierarchical generation in the graph settings or remove this node."));
-	}
-	else if (!Context->SourceComponent->IsPartitioned() && !Context->SourceComponent->IsLocalComponent())
-	{
-		// Warning if component is not partitioned (and not local component) as this node will otherwise be silently ignored.
-		// Also serves as a hint if the user forgot to enable higen for this graph.
-		PCGE_LOG(Warning, GraphAndLog, LOCTEXT("NonPartitionedComponent", "Grid Size node used on a non-partitioned component and will have no effect. Is Partitioned must be enabled on the component."));
-	}
 
 	// Trivial pass through. Will only execute on the prescribed grid.
 	Context->OutputData = Context->InputData;
diff --git a/Engine/Plugins/PCG/Source/PCGEditor/Private/PCGEditorGraphNodeBase.cpp b/Engine/Plugins/PCG/Source/PCGEditor/Private/PCGEditorGraphNodeBase.cpp
--- a/Engine/Plugins/PCG/Source/PCGEditor/Private/PCGEditorGraphNodeBase.cpp
+++ b/Engine/Plugins/PCG/Source/PCGEditor/Private/PCGEditorGraphNodeBase.cpp
@@ -10,6 +10,7 @@
 #include "PCGSettingsWithDynamicInputs.h"
 #include "PCGSubsystem.h"
 #include "PCGWorldActor.h"
+#include "Elements/PCGHiGenGridSize.h"
 #include "Elements/PCGReroute.h"
 
 #include "PCGEditor.h"
@@ -126,6 +127,11 @@ void UPCGEditorGraphNodeBase::Construct(UPCGNode* InPCGNode)
 
 		bCanUserAddRemoveSourcePins = (InPCGNode->GetSettings() && InPCGNode->GetSettings()->IsA<UPCGSettingsWithDynamicInputs>());
 	}
+
+	// Update to current graph/inspection state.
+	const UPCGEditorGraph* Graph = Cast<UPCGEditorGraph>(GetOuter());
+	const FPCGEditor* Editor = Graph ? Graph->GetEditor().Pin().Get() : nullptr;
+	UpdateStructuralVisualization(Editor ? Editor->GetPCGComponentBeingInspected() : nullptr, Editor ? Editor->GetStackBeingInspected() : nullptr, /*bNewlyPlaced=*/true);
 }
 
 void UPCGEditorGraphNodeBase::BeginDestroy()
@@ -499,7 +505,7 @@ EPCGChangeType UPCGEditorGraphNodeBase::UpdateErrorsAndWarnings()
 	return bStateChanged ? EPCGChangeType::Cosmetic : EPCGChangeType::None;
 }
 
-EPCGChangeType UPCGEditorGraphNodeBase::UpdateStructuralVisualization(UPCGComponent* InComponentBeingDebugged, const FPCGStack* InStackBeingInspected)
+EPCGChangeType UPCGEditorGraphNodeBase::UpdateStructuralVisualization(UPCGComponent* InComponentBeingDebugged, const FPCGStack* InStackBeingInspected, bool bNewlyPlaced)
 {
 	const UPCGGraph* Graph = PCGNode ? PCGNode->GetGraph() : nullptr;
 	if (!Graph)
@@ -507,9 +513,11 @@ EPCGChangeType UPCGEditorGraphNodeBase::UpdateStructuralVisualization(UPCGCompon
 		return EPCGChangeType::None;
 	}
 
+	const bool bInspecting = InComponentBeingDebugged && InStackBeingInspected && !InStackBeingInspected->GetStackFrames().IsEmpty();
+
 	EPCGChangeType ChangeType = EPCGChangeType::None;
 
-	const uint64 NewInactiveMask = (InComponentBeingDebugged && InStackBeingInspected) ? InComponentBeingDebugged->GetNodeInactivePinMask(PCGNode, *InStackBeingInspected) : 0;
+	const uint64 NewInactiveMask = bInspecting ? InComponentBeingDebugged->GetNodeInactivePinMask(PCGNode, *InStackBeingInspected) : 0;
 	if (NewInactiveMask != InactiveOutputPinMask)
 	{
 		InactiveOutputPinMask = NewInactiveMask;
@@ -517,13 +525,13 @@ EPCGChangeType UPCGEditorGraphNodeBase::UpdateStructuralVisualization(UPCGCompon
 	}
 
 	// Check top graph for higen enable - subgraphs always inherit higen state from the top graph.
-	const UPCGGraph* TopGraph = InComponentBeingDebugged ? InComponentBeingDebugged->GetGraph() : nullptr;
+	const UPCGGraph* TopGraph = bInspecting ? InStackBeingInspected->GetRootGraph() : Graph;
 	const bool HiGenEnabled = TopGraph && TopGraph->IsHierarchicalGenerationEnabled();
 
 	// Set the inspected grid size - this is used for grid size visualization.
 	uint32 InspectingGridSize = PCGHiGenGrid::UninitializedGridSize();
 	EPCGHiGenGrid InspectingGrid = EPCGHiGenGrid::Uninitialized;
-	if (TopGraph && TopGraph->IsHierarchicalGenerationEnabled())
+	if (TopGraph && TopGraph->IsHierarchicalGenerationEnabled() && InComponentBeingDebugged && (InComponentBeingDebugged->IsPartitioned() || InComponentBeingDebugged->IsLocalComponent()))
 	{
 		InspectingGridSize = InComponentBeingDebugged->GetGenerationGridSize();
 		InspectingGrid = InComponentBeingDebugged->GetGenerationGrid();
@@ -536,46 +544,63 @@ EPCGChangeType UPCGEditorGraphNodeBase::UpdateStructuralVisualization(UPCGCompon
 	}
 
 	bool bShouldDisplayAsDisabled = false;
-	bool bIsCulled = !PCGEditorGraphNodeBase::ShouldDisplayAsActive(this, InComponentBeingDebugged, InStackBeingInspected);
 
-	// Show grid size visualization if higen is enabled and if we're inspecting a specific grid, and we're inspecting a subgraph since subgraphs
-	// execute at the invoked grid level. 
-	if (HiGenEnabled && InStackBeingInspected && InStackBeingInspected->IsCurrentFrameInRootGraph() && InspectingGridSize != PCGHiGenGrid::UninitializedGridSize())
+	// Special treatment for higen grid sizes nodes which do nothing if higen is disabled.
+	// TODO: Drive this from an API on settings as we add more higen-specific functionality.
+	if (PCGNode && Cast<UPCGHiGenGridSizeSettings>(PCGNode->GetSettings()))
 	{
-		const uint32 DefaultGridSize = TopGraph->GetDefaultGridSize();
-		const uint32 NodeGridSize = Graph->GetNodeGenerationGridSize(PCGNode, DefaultGridSize);
+		// Higen must be enabled on graph, and we must be editing top graph.
+		bShouldDisplayAsDisabled = (Graph != TopGraph) || !TopGraph->IsHierarchicalGenerationEnabled();
 
-		if (NodeGridSize < InspectingGridSize)
+		// If we're inspecting a component, it must either be a partitioned OC or an LC (because higen requires partitioning).
+		if (!bShouldDisplayAsDisabled && InComponentBeingDebugged)
 		{
-			// Disable nodes that are on a smaller grid
-			bShouldDisplayAsDisabled |= NodeGridSize < InspectingGridSize;
-
-			// We don't know if the node was culled or not on that grid, disable visualization.
-			bIsCulled = false;
+			bShouldDisplayAsDisabled = !InComponentBeingDebugged->IsPartitioned() && !InComponentBeingDebugged->IsLocalComponent();
 		}
-		else if (NodeGridSize > InspectingGridSize)
+	}
+
+	// Don't do culling visualization on newly placed nodes. Let the execution complete notification update that.
+	bool bIsCulled = !bNewlyPlaced && !PCGEditorGraphNodeBase::ShouldDisplayAsActive(this, InComponentBeingDebugged, InStackBeingInspected);
+
+	EPCGHiGenGrid ThisGrid = EPCGHiGenGrid::Uninitialized;
+
+	// Show grid size visualization if higen is enabled and if we're inspecting a specific grid, and we're inspecting a subgraph since subgraphs
+	// execute at the invoked grid level.
+	if (HiGenEnabled && InspectingGridSize != PCGHiGenGrid::UninitializedGridSize())
+	{
+		if (InStackBeingInspected && InStackBeingInspected->IsCurrentFrameInRootGraph())
 		{
-			// We don't know if the node was culled or not on that grid, disable visualization.
-			bIsCulled = false;
-		}
+			const uint32 DefaultGridSize = TopGraph->GetDefaultGridSize();
+			const uint32 NodeGridSize = Graph->GetNodeGenerationGridSize(PCGNode, DefaultGridSize);
+
+			if (NodeGridSize < InspectingGridSize)
+			{
+				// Disable nodes that are on a smaller grid
+				bShouldDisplayAsDisabled |= NodeGridSize < InspectingGridSize;
+
+				// We don't know if the node was culled or not on that grid, disable visualization.
+				bIsCulled = false;
+			}
+			else if (NodeGridSize > InspectingGridSize)
+			{
+				// We don't know if the node was culled or not on that grid, disable visualization.
+				bIsCulled = false;
+			}
 
-		const EPCGHiGenGrid Grid = PCGHiGenGrid::GridSizeToGrid(NodeGridSize);
-		if (GenerationGrid != Grid)
+			ThisGrid = PCGHiGenGrid::GridSizeToGrid(NodeGridSize);
+		}
+		else
 		{
-			GenerationGrid = Grid;
-			ChangeType |= EPCGChangeType::Cosmetic;
+			// If higen is enabled then we are inspecting an invoked subgraph. Display the inspected grid size so that the user still
+			// gets the execution grid information.
+			ThisGrid = PCGHiGenGrid::GridSizeToGrid(InspectingGridSize);
 		}
 	}
-	else
+
+	if (GenerationGrid != ThisGrid)
 	{
-		// If higen is enabled then we are inspecting an invoked subgraph. Display the inspected grid size so that the user still
-		// gets the execution grid information.
-		const EPCGHiGenGrid Grid = HiGenEnabled ? PCGHiGenGrid::GridSizeToGrid(InspectingGridSize) : EPCGHiGenGrid::Uninitialized;
-		if (GenerationGrid != Grid)
-		{
-			GenerationGrid = Grid;
-			ChangeType |= EPCGChangeType::Cosmetic;
-		}
+		GenerationGrid = ThisGrid;
+		ChangeType |= EPCGChangeType::Cosmetic;
 	}
 
 	SetIsCulledFromExecution(bIsCulled);
diff --git a/Engine/Plugins/PCG/Source/PCGEditor/Private/PCGEditorGraphNodeBase.h b/Engine/Plugins/PCG/Source/PCGEditor/Private/PCGEditorGraphNodeBase.h
--- a/Engine/Plugins/PCG/Source/PCGEditor/Private/PCGEditorGraphNodeBase.h
+++ b/Engine/Plugins/PCG/Source/PCGEditor/Private/PCGEditorGraphNodeBase.h
@@ -73,7 +73,7 @@ public:
 	EPCGChangeType UpdateErrorsAndWarnings();
 
 	/** If the currently inspected grid size is smaller than the grid size of this node, display transparent. */
-	EPCGChangeType UpdateStructuralVisualization(class UPCGComponent* InComponentBeingDebugged, const FPCGStack* InStackBeingInspected);
+	EPCGChangeType UpdateStructuralVisualization(class UPCGComponent* InComponentBeingDebugged, const FPCGStack* InStackBeingInspected, bool bNewlyPlaced = false);
 
 	/** Puts node title on node body, reducing overall node size */
 	bool ShouldDrawCompact() const;
