[PCG] Fix non-partitioned octree failing to register PCG Components in standalone
* Caused RuntimeGenScheduler to fail finding NonPartitioned PCG Components.

#rnx

diff --git a/Engine/Plugins/PCG/Source/PCG/Private/PCGActorAndComponentMapping.cpp b/Engine/Plugins/PCG/Source/PCG/Private/PCGActorAndComponentMapping.cpp
--- a/Engine/Plugins/PCG/Source/PCG/Private/PCGActorAndComponentMapping.cpp
+++ b/Engine/Plugins/PCG/Source/PCG/Private/PCGActorAndComponentMapping.cpp
@@ -345,8 +345,6 @@ bool FPCGActorAndComponentMapping::RegisterOrUpdatePartitionedPCGComponent(UPCGC
 
 bool FPCGActorAndComponentMapping::RegisterOrUpdateNonPartitionedPCGComponent(UPCGComponent* InComponent)
 {
-	// Tracking is only done in Editor for now
-#if WITH_EDITOR
 	FBox Bounds(EForceInit::ForceInit);
 	bool bComponentHasChanged = false;
 	bool bComponentWasAdded = false;
@@ -354,9 +352,6 @@ bool FPCGActorAndComponentMapping::RegisterOrUpdateNonPartitionedPCGComponent(UP
 	NonPartitionedOctree.AddOrUpdateComponent(InComponent, Bounds, bComponentHasChanged, bComponentWasAdded);
 
 	return bComponentHasChanged;
-#else
-	return false;
-#endif // WITH_EDITOR
 }
 
 bool FPCGActorAndComponentMapping::RemapPCGComponent(const UPCGComponent* OldComponent, UPCGComponent* NewComponent, bool bDoActorMapping)
