[PCG] Make Quality Branch dynamic instead of static
* Fixes issue where graphs with quality branch nodes don't regenerate correctly when changing 'pcg.Quality' setting

#rnx

diff --git a/Engine/Plugins/PCG/Source/PCG/Private/Elements/ControlFlow/PCGQualityBranch.cpp b/Engine/Plugins/PCG/Source/PCG/Private/Elements/ControlFlow/PCGQualityBranch.cpp
--- a/Engine/Plugins/PCG/Source/PCG/Private/Elements/ControlFlow/PCGQualityBranch.cpp
+++ b/Engine/Plugins/PCG/Source/PCG/Private/Elements/ControlFlow/PCGQualityBranch.cpp
@@ -24,10 +24,14 @@ FString UPCGQualityBranchSettings::GetAdditionalTitleInformation() const
 	return PCGQualityHelpers::GetQualityPinLabel().ToString();
 }
 
-bool UPCGQualityBranchSettings::IsPinStaticallyActive(const FName& PinLabel) const
+bool UPCGQualityBranchSettings::IsPinUsedByNodeExecution(const UPCGPin* InPin) const
 {
-	// TODO: If we decide to have 'realtime' graphs in the future (e.g. for microscattering) where we do not recompile/regen when quality
-	// level changes, we may not be able to statically evaluate which pin is active anymore.
+	if (!InPin || !InPin->IsOutputPin())
+	{
+		return Super::IsPinUsedByNodeExecution(InPin);
+	}
+
+	const FName PinLabel = InPin->Properties.Label;
 
 	// If disabled, passthrough to Default pin.
 	if (!bEnabled)
@@ -68,17 +72,6 @@ bool UPCGQualityBranchSettings::IsPinStaticallyActive(const FName& PinLabel) con
 	{
 		return PinLabel == PCGQualityHelpers::PinLabelDefault;
 	}
-
-}
-
-bool UPCGQualityBranchSettings::IsPinUsedByNodeExecution(const UPCGPin* InPin) const
-{
-	if (!InPin->IsOutputPin())
-	{
-		return Super::IsPinUsedByNodeExecution(InPin);
-	}
-
-	return IsPinStaticallyActive(InPin->Properties.Label);
 }
 
 TArray<FPCGPinProperties> UPCGQualityBranchSettings::OutputPinProperties() const
@@ -196,6 +189,10 @@ bool FPCGQualityBranchElement::ExecuteInternal(FPCGContext* Context) const
 	// Reuse the functionality of the Gather Node
 	Context->OutputData = PCGGather::GatherDataForPin(Context->InputData, PCGPinConstants::DefaultInputLabel, ActivePinNames[ActivePinIndex]);
 
+	// Output bitmask of deactivated pins.
+	const uint64 AllPinMask = (1ULL << NumActivePins) - 1;
+	Context->OutputData.InactiveOutputPinBitmask = (~(1ULL << ActivePinIndex)) & AllPinMask;
+
 	return true;
 }
 
diff --git a/Engine/Plugins/PCG/Source/PCG/Public/Elements/ControlFlow/PCGQualityBranch.h b/Engine/Plugins/PCG/Source/PCG/Public/Elements/ControlFlow/PCGQualityBranch.h
--- a/Engine/Plugins/PCG/Source/PCG/Public/Elements/ControlFlow/PCGQualityBranch.h
+++ b/Engine/Plugins/PCG/Source/PCG/Public/Elements/ControlFlow/PCGQualityBranch.h
@@ -24,7 +24,7 @@ public:
 #endif
 	virtual FString GetAdditionalTitleInformation() const override;
 	virtual bool HasDynamicPins() const override { return true; }
-	virtual bool IsPinStaticallyActive(const FName& PinLabel) const override;
+	virtual bool OutputPinsCanBeDeactivated() const override { return true; }
 	virtual bool IsPinUsedByNodeExecution(const UPCGPin* InPin) const override;
 
 protected:
