[PCG] Pops up a window to create a subgraph node or loop node when dropping graph assests in graph

diff --git a/Engine/Plugins/PCG/Source/PCGEditor/Private/PCGEditorGraphSchemaActions.cpp b/Engine/Plugins/PCG/Source/PCGEditor/Private/PCGEditorGraphSchemaActions.cpp
--- a/Engine/Plugins/PCG/Source/PCGEditor/Private/PCGEditorGraphSchemaActions.cpp
+++ b/Engine/Plugins/PCG/Source/PCGEditor/Private/PCGEditorGraphSchemaActions.cpp
@@ -484,7 +484,7 @@ void FPCGEditorGraphSchemaAction_NewSubgraphElement::MakeGraphNodesOrContextualM
 	else if (!Graph.IsEmpty())
 	{
 		FMenuBuilder MenuBuilder(true, nullptr);
-		const FText GraphTextName = ((Graph.Num() == 1) ? FText::FromName(Graph[0]->GetFName()) : LOCTEXT("MultipleGraphs", "Multiple Graphs selected"));
+		const FText GraphTextName = ((Graph.Num() == 1) ? FText::FromName(Graph[0]->GetFName()) : LOCTEXT("MultipleGraphs", "Multiple Graphs"));
 
 		MenuBuilder.BeginSection("GraphDroppedOn", GraphTextName);
 
