[PCG] Fix crash in standalone PIE from unsafe use of GEditor in UPCGComponent::ChangeTransientState()

#rnx

diff --git a/Engine/Plugins/PCG/Source/PCG/Private/PCGComponent.cpp b/Engine/Plugins/PCG/Source/PCG/Private/PCGComponent.cpp
--- a/Engine/Plugins/PCG/Source/PCG/Private/PCGComponent.cpp
+++ b/Engine/Plugins/PCG/Source/PCG/Private/PCGComponent.cpp
@@ -3245,7 +3245,7 @@ void UPCGComponent::ChangeTransientState(EPCGEditorDirtyMode NewEditingMode)
 	}
 
 	// Changing the transient state can and will play with packages and is not meant to be undoable
-	if (GEditor->Trans)
+	if (GEditor && GEditor->Trans)
 	{
 		GEditor->Trans->Reset(LOCTEXT("ChangeEditingMode", "Changing Editing Mode"));
 	}
