[PCG] Add New PCG Editor Settings 'bDisableCPUThrottlingDuringGraphExecution' (defaults to false)
- Also remove registration of APCGVolume factory in the module startup as the GEditor wasn't initialized yet. This didn't matter because PCGVolume factory gets discovered on Editor init so instead validate that it exists in the PostEngineInit callback

#rnx

diff --git a/Engine/Plugins/PCG/Source/PCG/Private/Graph/PCGGraphExecutor.cpp b/Engine/Plugins/PCG/Source/PCG/Private/Graph/PCGGraphExecutor.cpp
--- a/Engine/Plugins/PCG/Source/PCG/Private/Graph/PCGGraphExecutor.cpp
+++ b/Engine/Plugins/PCG/Source/PCG/Private/Graph/PCGGraphExecutor.cpp
@@ -471,6 +471,22 @@ bool FPCGGraphExecutor::IsGraphCurrentlyExecuting(UPCGGraph* InGraph)
 	return bAnyPresent;
 }
 
+bool FPCGGraphExecutor::IsAnyGraphCurrentlyExecuting() const
+{
+	if (GetNonScheduledRemainingTaskCount() > 0)
+	{
+		return true;
+	}
+
+	// No need for locking here as we don't need the value to be precise here and it can change right after a lock anyways.
+	return ScheduledTasks.Num() > 0;
+}
+
+int32 FPCGGraphExecutor::GetNonScheduledRemainingTaskCount() const
+{
+	return Tasks.Num() + ReadyTasks.Num() + ActiveTasks.Num() + SleepingTasks.Num();
+}
+
 FPCGTaskId FPCGGraphExecutor::ScheduleGeneric(TFunction<bool()> InOperation, UPCGComponent* InSourceComponent, const TArray<FPCGTaskId>& TaskExecutionDependencies)
 {
 	return ScheduleGeneric(
@@ -1069,7 +1085,7 @@ void FPCGGraphExecutor::Execute()
 		// TODO: this is fine and will make sure any intermediate data is properly
 		// garbage collected, however, this goes a bit against our goals if we want to
 		// keep a cache of intermediate results.
-		if (ReadyTasks.Num() == 0 && ActiveTasks.Num() == 0 && SleepingTasks.Num() == 0 && Tasks.Num() == 0)
+		if (GetNonScheduledRemainingTaskCount() == 0)
 		{
 			if (!ensure(TaskSuccessors.IsEmpty()))
 			{
@@ -1675,7 +1691,7 @@ void FPCGGraphExecutor::NotifyGraphChanged(UPCGGraph* InGraph, EPCGChangeType Ch
 
 void FPCGGraphExecutor::UpdateGenerationNotification()
 {
-	const int32 RemainingTaskNum = Tasks.Num() + ReadyTasks.Num() + ActiveTasks.Num() + SleepingTasks.Num();
+	const int32 RemainingTaskNum = GetNonScheduledRemainingTaskCount();
 
 	if (RemainingTaskNum > 0)
 	{
diff --git a/Engine/Plugins/PCG/Source/PCG/Private/Graph/PCGGraphExecutor.h b/Engine/Plugins/PCG/Source/PCG/Private/Graph/PCGGraphExecutor.h
--- a/Engine/Plugins/PCG/Source/PCG/Private/Graph/PCGGraphExecutor.h
+++ b/Engine/Plugins/PCG/Source/PCG/Private/Graph/PCGGraphExecutor.h
@@ -144,6 +144,9 @@ public:
 	/** Returns true if any task is scheduled or executing for the given graph. */
 	bool IsGraphCurrentlyExecuting(UPCGGraph* InGraph);
 
+	/** Returns true if any task is scheduled or executing for any graph */
+	bool IsAnyGraphCurrentlyExecuting() const;
+	
 	// Back compatibility function. Use ScheduleGenericWithContext
 	FPCGTaskId ScheduleGeneric(TFunction<bool()> InOperation, UPCGComponent* InSourceComponent, const TArray<FPCGTaskId>& TaskExecutionDependencies);
 	FPCGTaskId ScheduleGeneric(TFunction<bool()> InOperation, TFunction<void()> InAbortOperation, UPCGComponent* InSourceComponent, const TArray<FPCGTaskId>& TaskExecutionDependencies);
@@ -223,6 +226,8 @@ private:
 
 	void LogTaskState() const;
 
+	int32 GetNonScheduledRemainingTaskCount() const;
+
 #if WITH_EDITOR
 	/** Notify the component that the given pins were deactivated during execution. */
 	void SendInactivePinNotification(const UPCGNode* InNode, const FPCGStack* InStack, uint64 InactiveOutputPinBitmask);
diff --git a/Engine/Plugins/PCG/Source/PCG/Private/PCGSubsystem.cpp b/Engine/Plugins/PCG/Source/PCG/Private/PCGSubsystem.cpp
--- a/Engine/Plugins/PCG/Source/PCG/Private/PCGSubsystem.cpp
+++ b/Engine/Plugins/PCG/Source/PCG/Private/PCGSubsystem.cpp
@@ -877,6 +877,11 @@ bool UPCGSubsystem::IsGraphCurrentlyExecuting(UPCGGraph* Graph)
 	return GraphExecutor->IsGraphCurrentlyExecuting(Graph);
 }
 
+bool UPCGSubsystem::IsAnyGraphCurrentlyExecuting() const
+{
+	return GraphExecutor && GraphExecutor->IsAnyGraphCurrentlyExecuting();
+}
+
 void UPCGSubsystem::ForAllRegisteredLocalComponents(UPCGComponent* InOriginalComponent, const TFunctionRef<void(UPCGComponent*)>& InFunc) const
 {
 	auto WrapperFunc = [&InFunc](UPCGComponent* Component) -> FPCGTaskId
diff --git a/Engine/Plugins/PCG/Source/PCG/Public/PCGSubsystem.h b/Engine/Plugins/PCG/Source/PCG/Public/PCGSubsystem.h
--- a/Engine/Plugins/PCG/Source/PCG/Public/PCGSubsystem.h
+++ b/Engine/Plugins/PCG/Source/PCG/Public/PCGSubsystem.h
@@ -165,6 +165,9 @@ public:
 	/** Returns true if there are any tasks for this graph currently scheduled or executing. */
 	bool IsGraphCurrentlyExecuting(UPCGGraph* Graph);
 
+	/** Returns true if any task is scheduled or executing for any graph */
+	bool IsAnyGraphCurrentlyExecuting() const;
+
 	/** Cancels everything running */
 	void CancelAllGeneration();
 
diff --git a/Engine/Plugins/PCG/Source/PCGEditor/Private/PCGEditorModule.cpp b/Engine/Plugins/PCG/Source/PCGEditor/Private/PCGEditorModule.cpp
--- a/Engine/Plugins/PCG/Source/PCGEditor/Private/PCGEditorModule.cpp
+++ b/Engine/Plugins/PCG/Source/PCGEditor/Private/PCGEditorModule.cpp
@@ -46,15 +46,7 @@ void FPCGEditorModule::StartupModule()
 	GraphNodeFactory = MakeShareable(new FPCGEditorGraphNodeFactory());
 	FEdGraphUtilities::RegisterVisualNodeFactory(GraphNodeFactory);
 
-	if (GEditor)
-	{
-		GEditor->ActorFactories.Add(NewObject<UPCGVolumeFactory>());
-	}
-
-	if (!IsRunningCommandlet())
-	{
-		FCoreDelegates::OnPostEngineInit.AddRaw(this, &FPCGEditorModule::RegisterOnEditorModeChange);
-	}
+	FCoreDelegates::OnPostEngineInit.AddRaw(this, &FPCGEditorModule::OnPostEngineInit);
 }
 
 void FPCGEditorModule::ShutdownModule()
@@ -71,11 +63,16 @@ void FPCGEditorModule::ShutdownModule()
 	if (GEditor)
 	{
 		GEditor->ActorFactories.RemoveAll([](const UActorFactory* ActorFactory) { return ActorFactory->IsA<UPCGVolumeFactory>(); });
+
+		GEditor->ShouldDisableCPUThrottlingDelegates.RemoveAll([this](const UEditorEngine::FShouldDisableCPUThrottling& Delegate)
+		{
+			return Delegate.GetHandle() == ShouldDisableCPUThrottlingDelegateHandle;
+		});
 	}
 
+	FCoreDelegates::OnPostEngineInit.RemoveAll(this);
 	if (!IsRunningCommandlet())
 	{
-		FCoreDelegates::OnPostEngineInit.RemoveAll(this);
 		if (GLevelEditorModeToolsIsValid())
 		{
 			GLevelEditorModeTools().OnEditorModeIDChanged().RemoveAll(this);
@@ -83,6 +80,36 @@ void FPCGEditorModule::ShutdownModule()
 	}
 }
 
+void FPCGEditorModule::OnPostEngineInit()
+{
+	RegisterOnEditorModeChange();
+
+	if (GEditor)
+	{
+		// Factory should be auto-discovered by UEditorEngine::InitEditor
+		check(GEditor->ActorFactories.FindItemByClass<UPCGVolumeFactory>());
+
+		if (!IsRunningCommandlet())
+		{
+			GEditor->ShouldDisableCPUThrottlingDelegates.Add(UEditorEngine::FShouldDisableCPUThrottling::CreateRaw(this, &FPCGEditorModule::ShouldDisableCPUThrottling));
+			ShouldDisableCPUThrottlingDelegateHandle = GEditor->ShouldDisableCPUThrottlingDelegates.Last().GetHandle();
+		}
+	}
+}
+
+bool FPCGEditorModule::ShouldDisableCPUThrottling()
+{
+	if (const UPCGEditorSettings* EditorSettings = GetDefault<UPCGEditorSettings>(); EditorSettings && EditorSettings->bDisableCPUThrottlingDuringGraphExecution)
+	{
+		if (UPCGSubsystem* PCGSubsystem = UPCGSubsystem::GetSubsystemForCurrentWorld())
+		{
+			return PCGSubsystem->IsAnyGraphCurrentlyExecuting();
+		}
+	}
+
+	return false;
+}
+
 void FPCGEditorModule::OnEditorModeIDChanged(const FEditorModeID& EditorModeID, bool bIsEntering)
 {
 	if (EditorModeID == FBuiltinEditorModes::EM_Landscape && !bIsEntering)
diff --git a/Engine/Plugins/PCG/Source/PCGEditor/Private/PCGEditorModule.h b/Engine/Plugins/PCG/Source/PCGEditor/Private/PCGEditorModule.h
--- a/Engine/Plugins/PCG/Source/PCGEditor/Private/PCGEditorModule.h
+++ b/Engine/Plugins/PCG/Source/PCGEditor/Private/PCGEditorModule.h
@@ -32,9 +32,12 @@ protected:
 	void RegisterSettings();
 	void UnregisterSettings();
 
+	void OnPostEngineInit();
 	void RegisterOnEditorModeChange();
 	void OnEditorModeIDChanged(const FEditorModeID& EditorModeID, bool bIsEntering);
 
+	bool ShouldDisableCPUThrottling();
+
 	TArray<TSharedRef<IAssetTypeActions>> RegisteredAssetTypeActions;
 	static EAssetTypeCategories::Type PCGAssetCategory;
 
@@ -42,4 +45,6 @@ protected:
 	TSharedPtr<FExtensibilityManager> ToolBarExtensibilityManager;
 
 	TSharedPtr<FPCGEditorGraphNodeFactory> GraphNodeFactory;
+
+	FDelegateHandle ShouldDisableCPUThrottlingDelegateHandle;
 };
diff --git a/Engine/Plugins/PCG/Source/PCGEditor/Private/PCGEditorSettings.cpp b/Engine/Plugins/PCG/Source/PCGEditor/Private/PCGEditorSettings.cpp
--- a/Engine/Plugins/PCG/Source/PCGEditor/Private/PCGEditorSettings.cpp
+++ b/Engine/Plugins/PCG/Source/PCGEditor/Private/PCGEditorSettings.cpp
@@ -42,8 +42,6 @@ UPCGEditorSettings::UPCGEditorSettings(const FObjectInitializer& ObjectInitializ
 
 	ParamDataPinColor = FLinearColor(1.0f, 0.38f, 0.02f);
 	UnknownDataPinColor = FLinearColor(0.3f, 0.3f, 0.3f);
-
-	bEnableNavigateToNativeNodes = true;
 }
 
 FLinearColor UPCGEditorSettings::GetColor(UPCGSettings* Settings) const
diff --git a/Engine/Plugins/PCG/Source/PCGEditor/Public/PCGEditorSettings.h b/Engine/Plugins/PCG/Source/PCGEditor/Public/PCGEditorSettings.h
--- a/Engine/Plugins/PCG/Source/PCGEditor/Public/PCGEditorSettings.h
+++ b/Engine/Plugins/PCG/Source/PCGEditor/Public/PCGEditorSettings.h
@@ -163,7 +163,11 @@ public:
 
 	/** Specify if we want to jump to definition in case of double click on native PCG Nodes */
 	UPROPERTY(EditAnywhere, config, Category = Workflow)
-	bool bEnableNavigateToNativeNodes;
+	bool bEnableNavigateToNativeNodes = true;
+
+	/** Specify if we want to disable CPU Throttling when a PCG Graph is executing, this will improve execution time when app is out of focus/minimized */
+	UPROPERTY(EditAnywhere, config, Category = "Editor Performance", meta = (DisplayName = "Disable CPU throttling during graph execution"))
+	bool bDisableCPUThrottlingDuringGraphExecution = false;
 };
 
 #if UE_ENABLE_INCLUDE_ORDER_DEPRECATED_IN_5_2
