[PCG] PCGGraph marks ObjectInitializer to not create nested default sub objects (DefaultInputNode.DefaultNodeSettings and DefaultOutputNode.DefaultNodeSettings) as they will be overriden right after by UPCGGRaphInputOutputSettings.

Fixes a fatal log error with PCGGraph CDO
Also prevents useless creation/renaming/markasgarbage for all graphs.

#rnx

diff --git a/Engine/Plugins/PCG/Source/PCG/Private/PCGGraph.cpp b/Engine/Plugins/PCG/Source/PCG/Private/PCGGraph.cpp
--- a/Engine/Plugins/PCG/Source/PCG/Private/PCGGraph.cpp
+++ b/Engine/Plugins/PCG/Source/PCG/Private/PCGGraph.cpp
@@ -248,7 +248,7 @@ EPCGChangeType UPCGGraphInterface::GetChangeTypeForGraphParameterChange(EPCGGrap
 ****************************/
 
 UPCGGraph::UPCGGraph(const FObjectInitializer& ObjectInitializer)
-	: Super(ObjectInitializer)
+	: Super(ObjectInitializer.DoNotCreateNestedDefaultSubobject({ TEXT("DefaultInputNode"), TEXT("DefaultNodeSettings") }).DoNotCreateNestedDefaultSubobject({ TEXT("DefaultOutputNode"), TEXT("DefaultNodeSettings") }))
 {
 	InputNode = ObjectInitializer.CreateDefaultSubobject<UPCGNode>(this, TEXT("DefaultInputNode"));
 	InputNode->SetFlags(RF_Transactional);
diff --git a/Engine/Plugins/PCG/Source/PCG/Private/PCGNode.cpp b/Engine/Plugins/PCG/Source/PCG/Private/PCGNode.cpp
--- a/Engine/Plugins/PCG/Source/PCG/Private/PCGNode.cpp
+++ b/Engine/Plugins/PCG/Source/PCG/Private/PCGNode.cpp
@@ -40,7 +40,7 @@ namespace PCGNodeHelpers
 UPCGNode::UPCGNode(const FObjectInitializer& ObjectInitializer)
 	: Super(ObjectInitializer)
 {
-	SettingsInterface = ObjectInitializer.CreateDefaultSubobject<UPCGTrivialSettings>(this, TEXT("DefaultNodeSettings"));
+	SettingsInterface = ObjectInitializer.CreateOptionalDefaultSubobject<UPCGTrivialSettings>(this, TEXT("DefaultNodeSettings"));
 }
 
 void UPCGNode::PostLoad()
