[PCG] Fix PCG Graph losing nodes on Delete -> Undo
* For users this problem manifested as phantom nodes/edges being restored on undo/redo, but they did not actually function since they didn't exist in the PCG Graph.
* This was introduced at the same time as support for bulk node deletion.

#rnx

diff --git a/Engine/Plugins/PCG/Source/PCGEditor/Private/PCGEditor.cpp b/Engine/Plugins/PCG/Source/PCGEditor/Private/PCGEditor.cpp
--- a/Engine/Plugins/PCG/Source/PCGEditor/Private/PCGEditor.cpp
+++ b/Engine/Plugins/PCG/Source/PCGEditor/Private/PCGEditor.cpp
@@ -2293,10 +2293,12 @@ void FPCGEditor::DeleteSelectedNodes()
 		PCGGraph->PrimeGraphCompilationCache();
 
 		bool bChanged = false;
-		TArray<UPCGNode*> NodesToRemove;
+
 		{
 			const FScopedTransaction Transaction(*FPCGEditorCommon::ContextIdentifier, LOCTEXT("PCGEditorDeleteTransactionMessage", "PCG Editor: Delete"), nullptr);
 			PCGEditorGraph->Modify();
+		
+			TArray<UPCGNode*> NodesToRemove;
 
 			for (UObject* Object : GraphEditorWidget->GetSelectedNodes())
 			{
@@ -2324,11 +2326,15 @@ void FPCGEditor::DeleteSelectedNodes()
 			}
 
 			PCGEditorGraph->Modify();
+
+			if (bChanged)
+			{
+				PCGGraph->RemoveNodes(NodesToRemove);
+			}
 		}
 
 		if (bChanged)
 		{
-			PCGGraph->RemoveNodes(NodesToRemove);
 			GraphEditorWidget->ClearSelectionSet();
 			GraphEditorWidget->NotifyGraphChanged();
 		}
