[PCG] Fixed "set debug object from level editor" button on debug object tree view not working for dynamic graph invocations (dynamic SG, loop)

diff --git a/Engine/Plugins/PCG/Source/PCGEditor/Private/SPCGEditorGraphDebugObjectTree.cpp b/Engine/Plugins/PCG/Source/PCGEditor/Private/SPCGEditorGraphDebugObjectTree.cpp
--- a/Engine/Plugins/PCG/Source/PCGEditor/Private/SPCGEditorGraphDebugObjectTree.cpp
+++ b/Engine/Plugins/PCG/Source/PCGEditor/Private/SPCGEditorGraphDebugObjectTree.cpp
@@ -482,6 +482,8 @@ bool SPCGEditorGraphDebugObjectTree::IsSetDebugObjectFromSelectionButtonEnabled(
 		return false;
 	}
 
+	UPCGSubsystem* Subsystem = PCGEditor.Pin() ? PCGEditor.Pin()->GetSubsystem() : nullptr;
+
 	for (FSelectionIterator It(*SelectedActors); It; ++It)
 	{
 		const AActor* SelectedActor = Cast<AActor>(*It);
@@ -500,20 +502,23 @@ bool SPCGEditorGraphDebugObjectTree::IsSetDebugObjectFromSelectionButtonEnabled(
 				continue;
 			}
 
+			// Look for graph in static stacks.
 			FPCGStackContext StackContext;
-			if (!PCGComponent->GetStackContext(StackContext))
+			if (PCGComponent->GetStackContext(StackContext))
 			{
-				continue;
+				if (Algo::AnyOf(StackContext.GetStacks(), [&PCGGraph](const FPCGStack& InStack) { return InStack.GetGraphForCurrentFrame() == PCGGraph; }))
+				{
+					return true;
+				}
 			}
 
-			const bool bGraphFound = Algo::AnyOf(StackContext.GetStacks(), [&PCGGraph](const FPCGStack& InStack)
+			// Look for graph in dynamic stacks.
+			if (Subsystem)
 			{
-				return Cast<const UPCGGraph>(InStack.GetStackFrames().Top().Object) == PCGGraph;
-			});
-
-			if(bGraphFound)
-			{
-				return true;
+				if (Algo::AnyOf(Subsystem->GetExecutedStacks(PCGComponent, PCGGraph), [&PCGGraph](const FPCGStack& InStack) { return InStack.GetGraphForCurrentFrame() == PCGGraph; }))
+				{
+					return true;
+				}
 			}
 		}
 	}
