[PCG] Reset metadata entry on Copy Points if we chose to not propagate any metadata

* This was producing a warning about invalid metadata.

#rnx

diff --git a/Engine/Plugins/PCG/Source/PCG/Private/Elements/PCGCopyPoints.cpp b/Engine/Plugins/PCG/Source/PCG/Private/Elements/PCGCopyPoints.cpp
--- a/Engine/Plugins/PCG/Source/PCG/Private/Elements/PCGCopyPoints.cpp
+++ b/Engine/Plugins/PCG/Source/PCG/Private/Elements/PCGCopyPoints.cpp
@@ -290,6 +290,11 @@ bool FPCGCopyPointsElement::ExecuteInternal(FPCGContext* Context) const
 						}
 					}
 				}
+				else
+				{
+					// Reset the metadata entry if we have no metadata.
+					OutPoint.MetadataEntry = PCGInvalidEntryKey;
+				}
 
 				return true;
 			});
