[PCG] Transform old Add Attribute nodes that has no pin connected to Create Attribute when deprecating

#rnx

diff --git a/Engine/Plugins/PCG/Source/PCG/Private/Elements/PCGCreateAttribute.cpp b/Engine/Plugins/PCG/Source/PCG/Private/Elements/PCGCreateAttribute.cpp
--- a/Engine/Plugins/PCG/Source/PCG/Private/Elements/PCGCreateAttribute.cpp
+++ b/Engine/Plugins/PCG/Source/PCG/Private/Elements/PCGCreateAttribute.cpp
@@ -80,6 +80,23 @@ bool UPCGAddAttributeSettings::CanEditChange(const FProperty* InProperty) const
 
 	return true;
 }
+
+void UPCGAddAttributeSettings::ApplyStructuralDeprecation(UPCGNode* InOutNode)
+{
+	check(InOutNode);
+	// Arbitrary version that approximatively matches the time when Add/Create attributes changed.
+	// It will convert any add attributes that have nothing connected to it to a create attribute.
+	if (DataVersion < FPCGCustomVersion::SupportPartitionedComponentsInNonPartitionedLevels)
+	{
+		if (!InOutNode->IsInputPinConnected(PCGPinConstants::DefaultInputLabel))
+		{
+			UPCGCreateAttributeSetSettings* NewSettings = NewObject<UPCGCreateAttributeSetSettings>(InOutNode);
+			NewSettings->OutputTarget.ImportFromOtherSelector(OutputTarget);
+			NewSettings->AttributeTypes = AttributeTypes;
+			InOutNode->SetSettingsInterface(NewSettings);
+		}
+	}
+}
 #endif // WITH_EDITOR
 
 FPCGElementPtr UPCGAddAttributeSettings::CreateElement() const
diff --git a/Engine/Plugins/PCG/Source/PCG/Public/Elements/PCGCreateAttribute.h b/Engine/Plugins/PCG/Source/PCG/Public/Elements/PCGCreateAttribute.h
--- a/Engine/Plugins/PCG/Source/PCG/Public/Elements/PCGCreateAttribute.h
+++ b/Engine/Plugins/PCG/Source/PCG/Public/Elements/PCGCreateAttribute.h
@@ -38,6 +38,7 @@ public:
 	virtual EPCGSettingsType GetType() const override { return EPCGSettingsType::Param; }
 	virtual bool IsPinUsedByNodeExecution(const UPCGPin* InPin) const override;
 	virtual bool CanEditChange(const FProperty* InProperty) const override;
+	virtual void ApplyStructuralDeprecation(UPCGNode* InOutNode) override;
 #endif // WITH_EDITOR
 	virtual EPCGDataType GetCurrentPinTypes(const UPCGPin* InPin) const override;
 	virtual FString GetAdditionalTitleInformation() const override;
