Improved structure extraction by allowing to target members in arrays, which was preventing finer grained control.

#rnx

diff --git a/Engine/Plugins/PCG/Source/PCG/Private/Helpers/PCGPropertyHelpers.cpp b/Engine/Plugins/PCG/Source/PCG/Private/Helpers/PCGPropertyHelpers.cpp
--- a/Engine/Plugins/PCG/Source/PCG/Private/Helpers/PCGPropertyHelpers.cpp
+++ b/Engine/Plugins/PCG/Source/PCG/Private/Helpers/PCGPropertyHelpers.cpp
@@ -33,17 +33,50 @@ namespace PCGPropertyHelpers
 	}
 
 	/**
-	* Recursive function to go down the property chain to find the property and its container address.
+	* Expands container locations to their contents when the property passed in is an array.
+	* This is useful to allow extraction downstream of properties inside of arrays and also to generate the list of addresses/values to look at
+	* when extracting the values to the attribute set.
+	* 
+	* @param InArrayProperty   Property that drives the container expansion. Can be null, in which case we'll copy the container locations directly.
+	* @param InContainers      Container locations to expand
+	* @param OutContainers     Expanded container locations. Expected to be a different array than InContainers.
+	*/
+	template<typename FirstArrayType, typename SecondArrayType>
+	void ExpandContainers(const FArrayProperty* InArrayProperty, const FirstArrayType& InContainers, SecondArrayType& OutContainers)
+	{
+		check(OutContainers.IsEmpty());
+
+		if (InArrayProperty)
+		{
+			for (const void* Container : InContainers)
+			{
+				FScriptArrayHelper_InContainer Helper(InArrayProperty, Container);
+				int32 Offset = OutContainers.Num();
+				OutContainers.SetNumUninitialized(OutContainers.Num() + Helper.Num());
+				for (int32 DynamicIndex = 0; DynamicIndex < Helper.Num(); ++DynamicIndex)
+				{
+					OutContainers[Offset + DynamicIndex] = Helper.GetRawPtr(DynamicIndex);
+				}
+			}
+		}
+		else
+		{
+			OutContainers = InContainers;
+		}
+	}
+
+	/**
+	* Recursive function to go down the property chain to find the property and its container addresses.
 	* @param CurrentClass            Struct/Class for the current container
 	* @param CurrentName             Property name to look for in the container class.
 	* @param NextNames               List of property names to continue extracting at a deeper level.
 	* @param bNeedsToBeVisible       Discard properties that are not visibile in Blueprint
-	* @param OutContainer            Raw address for the current container. Will be write to at each recursive call.
+	* @param OutContainers           Raw addresses for the containers. Will be written to at each recursive call.
 	* @param OptionalContext         Optional context used for logging.
 	* @param OptionalObjectTraversed Optional set to store all object that we traversed, to be able to react to those objects changes.
 	* @returns                 The last property of the chain (and its container address is in OutContainer)
 	*/
-	const FProperty* ExtractPropertyChain(const UStruct* CurrentClass, const FName CurrentName, TArrayView<const FString> NextNames, const bool bNeedsToBeVisible, const void*& OutContainer, FPCGContext* OptionalContext, TSet<FSoftObjectPath>* OptionalObjectTraversed)
+	const FProperty* ExtractPropertyChain(const UStruct* CurrentClass, const FName CurrentName, TArrayView<const FString> NextNames, const bool bNeedsToBeVisible, TArray<const void*>& OutContainers, FPCGContext* OptionalContext, TSet<FSoftObjectPath>* OptionalObjectTraversed)
 	{
 		check(CurrentClass);
 
@@ -82,24 +115,60 @@ namespace PCGPropertyHelpers
 		if (!NextNames.IsEmpty())
 		{
 			UStruct* NextClass = nullptr;
+			bool bPropertyNotExtractable = false;
 
 			if (const FStructProperty* StructProperty = CastField<FStructProperty>(Property))
 			{
 				NextClass = StructProperty->Struct;
-				OutContainer = StructProperty->ContainerPtrToValuePtr<void>(OutContainer);
+				for (const void*& OutContainer : OutContainers)
+				{
+					OutContainer = StructProperty->ContainerPtrToValuePtr<void>(OutContainer);
+				}
 			}
 			else if (const FObjectProperty* ObjectProperty = CastField<FObjectProperty>(Property))
 			{
 				NextClass = ObjectProperty->PropertyClass;
-				OutContainer = ObjectProperty->GetObjectPropertyValue_InContainer(OutContainer);
+				for (const void*& OutContainer : OutContainers)
+				{
+					OutContainer = ObjectProperty->GetObjectPropertyValue_InContainer(OutContainer);
+				}
+			}
+			else if (const FArrayProperty* ArrayProperty = CastField<FArrayProperty>(Property))
+			{
+				const FProperty* InnerProperty = ArrayProperty->Inner;
+				if (const FStructProperty* InnerStructProperty = CastField<FStructProperty>(InnerProperty))
+				{
+					NextClass = InnerStructProperty->Struct;
+				}
+				else if (const FObjectProperty* InnerObjectProperty = CastField<FObjectProperty>(InnerProperty))
+				{
+					NextClass = InnerObjectProperty->PropertyClass;
+				}
+				else
+				{
+					bPropertyNotExtractable = true;
+				}
+
+				// If contents of the array are extractable, do so now by replacing the container entry (e.g. the array) with the pointer to its contents
+				if (!bPropertyNotExtractable)
+				{
+					TArray<const void*> Subcontainers;
+					ExpandContainers(ArrayProperty, OutContainers, Subcontainers);
+					OutContainers = MoveTemp(Subcontainers);
+				}
 			}
 			else
+			{
+				bPropertyNotExtractable = true;
+			}
+			
+			if(bPropertyNotExtractable)
 			{
 				LogError(FText::Format(LOCTEXT("PropertyIsNotExtractable", "Property '{0}' does exist in {1}, but is not extractable."), FText::FromName(CurrentName), FText::FromName(CurrentClass->GetFName())), OptionalContext);
 				return nullptr;
 			}
 
-			return ExtractPropertyChain(NextClass, FName(NextNames[0]), NextNames.RightChop(1), bNeedsToBeVisible, OutContainer, OptionalContext, OptionalObjectTraversed);
+			return ExtractPropertyChain(NextClass, FName(NextNames[0]), NextNames.RightChop(1), bNeedsToBeVisible, OutContainers, OptionalContext, OptionalObjectTraversed);
 		}
 		else
 		{
@@ -130,14 +199,14 @@ UPCGParamData* PCGPropertyHelpers::ExtractPropertyAsAttributeSet(const PCGProper
 {
 	check(Parameters.Container && Parameters.Class);
 
-	const void* Container = Parameters.Container;
+	TArray<const void*> Containers = { Parameters.Container };
 	const FProperty* Property = nullptr;
 	const FName PropertyName = Parameters.PropertySelector.GetName();
 	const bool ExtractRoot = (PropertyName == NAME_None);
 	// If Name is none, extract the container as-is, using Parameters.Class, otherwise, extract the chain.
 	if (!ExtractRoot)
 	{
-		Property = ExtractPropertyChain(Parameters.Class, PropertyName, Parameters.PropertySelector.GetExtraNames(), Parameters.bPropertyNeedsToBeVisible, Container, OptionalContext, OptionalObjectTraversed);
+		Property = ExtractPropertyChain(Parameters.Class, PropertyName, Parameters.PropertySelector.GetExtraNames(), Parameters.bPropertyNeedsToBeVisible, Containers, OptionalContext, OptionalObjectTraversed);
 		if (!Property)
 		{
 			return nullptr;
@@ -234,19 +303,7 @@ UPCGParamData* PCGPropertyHelpers::ExtractPropertyAsAttributeSet(const PCGProper
 
 	// Before we need to compute all the addresses for each entry in our array (or just a single entry if there is no array)
 	TArray<const void*, TInlineAllocator<16>> ElementAddresses;
-	if (ArrayProperty)
-	{
-		FScriptArrayHelper_InContainer Helper(ArrayProperty, Container);
-		ElementAddresses.Reserve(Helper.Num());
-		for (int32 DynamicIndex = 0; DynamicIndex < Helper.Num(); ++DynamicIndex)
-		{
-			ElementAddresses.Add(Helper.GetRawPtr(DynamicIndex));
-		}
-	}
-	else
-	{
-		ElementAddresses.Add(Container);
-	}
+	ExpandContainers(ArrayProperty, Containers, ElementAddresses);
 
 	// From there, we should be able to create the data.
 	UPCGParamData* ParamData = NewObject<UPCGParamData>();
