[PCG] Attribute Noise/Filter cleanup + fixing Point|Param pin types

* Make sure the input/output pin types have the supported types
* Deprecate the ToPointData done on spatial, so old version still work.
* Differentiate PointFilter/AttributeFilter and DensityNoise/AttributeNoise so they have different default
* Add option in the schema to not group preconfigured settings.
* Add PointOrParam pin type since bitfields are not serialized correctly
* Add deprecation fixes for impacted nodes (Attribute Noise and Match and Set. Attribute Filter is OK as we start using PointOrParam for this CL, it was Any before)

#rnx

diff --git a/Engine/Plugins/PCG/Source/PCG/Private/Elements/PCGAttributeFilter.cpp b/Engine/Plugins/PCG/Source/PCG/Private/Elements/PCGAttributeFilter.cpp
--- a/Engine/Plugins/PCG/Source/PCG/Private/Elements/PCGAttributeFilter.cpp
+++ b/Engine/Plugins/PCG/Source/PCG/Private/Elements/PCGAttributeFilter.cpp
@@ -249,6 +249,10 @@ void FPCGAttributeFilterThresholdSettings::OnPostLoad()
 }
 #endif
 
+////////////////////////////////////////
+// UPCGAttributeFilteringSettings
+////////////////////////////////////////
+
 EPCGDataType UPCGAttributeFilteringSettings::GetCurrentPinTypes(const UPCGPin* InPin) const
 {
 	check(InPin);
@@ -274,7 +278,7 @@ EPCGDataType UPCGAttributeFilteringSettings::GetCurrentPinTypes(const UPCGPin* I
 TArray<FPCGPinProperties> UPCGAttributeFilteringSettings::InputPinProperties() const
 {
 	TArray<FPCGPinProperties> PinProperties;
-	FPCGPinProperties& DataToFilterPinProperty = PinProperties.Emplace_GetRef(PCGAttributeFilterConstants::DataToFilterLabel, EPCGDataType::Any);
+	FPCGPinProperties& DataToFilterPinProperty = PinProperties.Emplace_GetRef(PCGAttributeFilterConstants::DataToFilterLabel, EPCGDataType::PointOrParam);
 	DataToFilterPinProperty.SetRequiredPin();
 
 #if WITH_EDITOR
@@ -296,8 +300,8 @@ TArray<FPCGPinProperties> UPCGAttributeFilteringSettings::InputPinProperties() c
 TArray<FPCGPinProperties> UPCGAttributeFilteringSettings::OutputPinProperties() const
 {
 	TArray<FPCGPinProperties> PinProperties;
-	PinProperties.Emplace(PCGPinConstants::DefaultInFilterLabel, EPCGDataType::Any);
-	PinProperties.Emplace(PCGPinConstants::DefaultOutFilterLabel, EPCGDataType::Any);
+	PinProperties.Emplace(PCGPinConstants::DefaultInFilterLabel, EPCGDataType::PointOrParam);
+	PinProperties.Emplace(PCGPinConstants::DefaultOutFilterLabel, EPCGDataType::PointOrParam);
 
 	return PinProperties;
 }
@@ -330,25 +334,84 @@ UPCGAttributeFilteringSettings::UPCGAttributeFilteringSettings()
 	ThresholdAttribute.SetPointProperty(EPCGPointProperties::Density);
 }
 
+#if WITH_EDITOR
+TArray<FPCGPreConfiguredSettingsInfo> UPCGAttributeFilteringSettings::GetPreconfiguredInfo() const
+{
+	TArray<FPCGPreConfiguredSettingsInfo> PreconfiguredInfo;
+	PreconfiguredInfo.Emplace(0, GetDefaultNodeTitle());
+	PreconfiguredInfo.Emplace(1, LOCTEXT("PointRangeNodeTitle", "Point Filter"));
+
+	return PreconfiguredInfo;
+}
+#endif
+
+void UPCGAttributeFilteringSettings::ApplyPreconfiguredSettings(const FPCGPreConfiguredSettingsInfo& PreconfigureInfo)
+{
+	// If index is 1, it is the default ($Density)
+	if (PreconfigureInfo.PreconfiguredIndex == 0)
+	{
+		TargetAttribute.SetAttributeName(PCGMetadataAttributeConstants::LastAttributeName);
+		ThresholdAttribute.SetAttributeName(PCGMetadataAttributeConstants::LastAttributeName);
+	}
+}
+
 void UPCGAttributeFilteringSettings::PostLoad()
 {
 	Super::PostLoad();
 
 #if WITH_EDITOR
 	AttributeTypes.OnPostLoad();
+
+	// Check for the data spatial to point gate version
+	if (DataVersion < FPCGCustomVersion::NoMoreSpatialDataConversionToPointDataByDefaultOnNonPointPins)
+	{
+		bHasSpatialToPointDeprecation = true;
+	}
 #endif // WITH_EDITOR
 }
 
-#if WITH_EDITOR
+////////////////////////////////////////
+// UPCGAttributeFilteringRangeSettings
+////////////////////////////////////////
+
 void UPCGAttributeFilteringRangeSettings::PostLoad()
 {
 	Super::PostLoad();
 
+#if WITH_EDITOR
 	MinThreshold.OnPostLoad();
 	MaxThreshold.OnPostLoad();
+
+	// Check for the data spatial to point gate version
+	if (DataVersion < FPCGCustomVersion::NoMoreSpatialDataConversionToPointDataByDefaultOnNonPointPins)
+	{
+		bHasSpatialToPointDeprecation = true;
+	}
+#endif
+}
+
+#if WITH_EDITOR
+TArray<FPCGPreConfiguredSettingsInfo> UPCGAttributeFilteringRangeSettings::GetPreconfiguredInfo() const
+{
+	TArray<FPCGPreConfiguredSettingsInfo> PreconfiguredInfo;
+	PreconfiguredInfo.Emplace(0, GetDefaultNodeTitle());
+	PreconfiguredInfo.Emplace(1, LOCTEXT("PointRangeNodeTitle", "Point Filter Range"));
+
+	return PreconfiguredInfo;
 }
 #endif
 
+void UPCGAttributeFilteringRangeSettings::ApplyPreconfiguredSettings(const FPCGPreConfiguredSettingsInfo& PreconfigureInfo)
+{
+	// If index is 1, it is the default ($Density)
+	if (PreconfigureInfo.PreconfiguredIndex == 0)
+	{
+		TargetAttribute.SetAttributeName(PCGMetadataAttributeConstants::LastAttributeName);
+		MinThreshold.ThresholdAttribute.SetAttributeName(PCGMetadataAttributeConstants::LastAttributeName);
+		MaxThreshold.ThresholdAttribute.SetAttributeName(PCGMetadataAttributeConstants::LastAttributeName);
+	}
+}
+
 FString UPCGAttributeFilteringRangeSettings::GetAdditionalTitleInformation() const
 {
 #if WITH_EDITOR
@@ -404,7 +467,7 @@ EPCGDataType UPCGAttributeFilteringRangeSettings::GetCurrentPinTypes(const UPCGP
 TArray<FPCGPinProperties> UPCGAttributeFilteringRangeSettings::InputPinProperties() const
 {
 	TArray<FPCGPinProperties> PinProperties;
-	FPCGPinProperties& DataToFilterPinProperty = PinProperties.Emplace_GetRef(PCGAttributeFilterConstants::DataToFilterLabel, EPCGDataType::Any);
+	FPCGPinProperties& DataToFilterPinProperty = PinProperties.Emplace_GetRef(PCGAttributeFilterConstants::DataToFilterLabel, EPCGDataType::PointOrParam);
 	DataToFilterPinProperty.SetRequiredPin();
 
 #if WITH_EDITOR
@@ -435,13 +498,17 @@ TArray<FPCGPinProperties> UPCGAttributeFilteringRangeSettings::InputPinPropertie
 TArray<FPCGPinProperties> UPCGAttributeFilteringRangeSettings::OutputPinProperties() const
 {
 	TArray<FPCGPinProperties> PinProperties;
-	PinProperties.Emplace(PCGPinConstants::DefaultInFilterLabel, EPCGDataType::Any);
-	PinProperties.Emplace(PCGPinConstants::DefaultOutFilterLabel, EPCGDataType::Any);
+	PinProperties.Emplace(PCGPinConstants::DefaultInFilterLabel, EPCGDataType::PointOrParam);
+	PinProperties.Emplace(PCGPinConstants::DefaultOutFilterLabel, EPCGDataType::PointOrParam);
 
 	return PinProperties;
 }
 
-bool FPCGAttributeFilterElementBase::DoFiltering(FPCGContext* Context, EPCGAttributeFilterOperator InOperation, const FPCGAttributePropertyInputSelector& InTargetAttribute, const FPCGAttributeFilterThresholdSettings& FirstThreshold, const FPCGAttributeFilterThresholdSettings* SecondThreshold) const
+////////////////////////////////////////
+// FPCGAttributeFilterElementBase
+////////////////////////////////////////
+
+bool FPCGAttributeFilterElementBase::DoFiltering(FPCGContext* Context, EPCGAttributeFilterOperator InOperation, const FPCGAttributePropertyInputSelector& InTargetAttribute, bool bHasSpatialToPointDeprecation, const FPCGAttributeFilterThresholdSettings& FirstThreshold, const FPCGAttributeFilterThresholdSettings* SecondThreshold) const
 {
 	TRACE_CPUPROFILER_EVENT_SCOPE(FPCGAttributeFilterElementBase::DoFiltering);
 	check(Context);
@@ -502,15 +569,29 @@ bool FPCGAttributeFilterElementBase::DoFiltering(FPCGContext* Context, EPCGAttri
 
 		if (const UPCGSpatialData* SpatialInput = Cast<const UPCGSpatialData>(OriginalData))
 		{
-			const UPCGPointData* OriginalPointData = SpatialInput->ToPointData(Context);
-			if (!OriginalPointData)
+			if (!SpatialInput->IsA<UPCGPointData>())
 			{
-				PCGE_LOG(Error, GraphAndLog, LOCTEXT("NoPointDataInInput", "Unable to get point data from input"));
-				continue;
-			}
+				const UPCGPointData* OriginalPointData = nullptr;
 
-			OperationData.bIsInputPointData = true;
-			OriginalData = OriginalPointData;
+				if (bHasSpatialToPointDeprecation)
+				{
+					OriginalPointData = SpatialInput->ToPointData(Context);
+				}
+
+				if (!OriginalPointData)
+				{
+					PCGE_LOG(Error, GraphAndLog, LOCTEXT("NoPointDataInInput", "Unable to get point data from input. Use a conversion node before this node to transform it to points."));
+					continue;
+				}
+
+				OperationData.bIsInputPointData = true;
+				OriginalData = OriginalPointData;
+			}
+			else
+			{
+				OperationData.bIsInputPointData = true;
+				OriginalData = SpatialInput;
+			}
 		}
 		else if (OriginalData->IsA<UPCGParamData>())
 		{
@@ -526,17 +607,21 @@ bool FPCGAttributeFilterElementBase::DoFiltering(FPCGContext* Context, EPCGAttri
 
 		// Helper lambdas to fail nicely and forward input to in/out filter pin
 		// If there is a problem with threshold -> forward to InFilter
-		auto ForwardInputToInFilterPin = [&Outputs, Input]()
+		auto ForwardInputToInFilterPin = [&Outputs, Input, OriginalData]()
 		{
 			FPCGTaggedData& InFilterOutput = Outputs.Add_GetRef(Input);
 			InFilterOutput.Pin = PCGPinConstants::DefaultInFilterLabel;
+			// Use original data because it could have been collapsed
+			InFilterOutput.Data = OriginalData;
 		};
 
 		// If there is a problem with target -> forward to OutFilter
-		auto ForwardInputToOutFilterPin = [&Outputs, Input]()
+		auto ForwardInputToOutFilterPin = [&Outputs, Input, OriginalData]()
 		{
 			FPCGTaggedData& OutFilterOutput = Outputs.Add_GetRef(Input);
 			OutFilterOutput.Pin = PCGPinConstants::DefaultOutFilterLabel;
+			// Use original data because it could have been collapsed
+			OutFilterOutput.Data = OriginalData;
 		};
 
 		FPCGAttributePropertyInputSelector TargetAttribute = InTargetAttribute.CopyAndFixLast(OriginalData);
@@ -766,6 +851,10 @@ bool FPCGAttributeFilterElementBase::DoFiltering(FPCGContext* Context, EPCGAttri
 	return true;
 }
 
+////////////////////////////////////////
+// FPCGAttributeFilterElement
+////////////////////////////////////////
+
 bool FPCGAttributeFilterElement::ExecuteInternal(FPCGContext* Context) const
 {
 	TRACE_CPUPROFILER_EVENT_SCOPE(FPCGAttributeFilterElement::Execute);
@@ -790,9 +879,13 @@ bool FPCGAttributeFilterElement::ExecuteInternal(FPCGContext* Context) const
 	ThresholdSettings.ThresholdAttribute = Settings->ThresholdAttribute;
 	ThresholdSettings.AttributeTypes = Settings->AttributeTypes;
 
-	return DoFiltering(Context, Settings->Operator, Settings->TargetAttribute, ThresholdSettings);
+	return DoFiltering(Context, Settings->Operator, Settings->TargetAttribute, Settings->bHasSpatialToPointDeprecation, ThresholdSettings);
 }
 
+////////////////////////////////////////
+// FPCGAttributeFilterRangeElement
+////////////////////////////////////////
+
 bool FPCGAttributeFilterRangeElement::ExecuteInternal(FPCGContext* Context) const
 {
 	TRACE_CPUPROFILER_EVENT_SCOPE(FPCGAttributeFilterRangeElement::Execute);
@@ -811,7 +904,7 @@ bool FPCGAttributeFilterRangeElement::ExecuteInternal(FPCGContext* Context) cons
 	const UPCGAttributeFilteringRangeSettings* Settings = Context->GetInputSettings<UPCGAttributeFilteringRangeSettings>();
 	check(Settings);
 
-	return DoFiltering(Context, EPCGAttributeFilterOperator::InRange, Settings->TargetAttribute, Settings->MinThreshold, &Settings->MaxThreshold);
+	return DoFiltering(Context, EPCGAttributeFilterOperator::InRange, Settings->TargetAttribute, Settings->bHasSpatialToPointDeprecation, Settings->MinThreshold, &Settings->MaxThreshold);
 }
 
 #undef LOCTEXT_NAMESPACE
diff --git a/Engine/Plugins/PCG/Source/PCG/Private/Elements/PCGAttributeNoise.cpp b/Engine/Plugins/PCG/Source/PCG/Private/Elements/PCGAttributeNoise.cpp
--- a/Engine/Plugins/PCG/Source/PCG/Private/Elements/PCGAttributeNoise.cpp
+++ b/Engine/Plugins/PCG/Source/PCG/Private/Elements/PCGAttributeNoise.cpp
@@ -4,6 +4,7 @@
 
 #include "PCGCommon.h"
 #include "PCGContext.h"
+#include "PCGCustomVersion.h"
 #include "Data/PCGPointData.h"
 #include "Helpers/PCGAsync.h"
 #include "Helpers/PCGHelpers.h"
@@ -96,6 +97,26 @@ UPCGAttributeNoiseSettings::UPCGAttributeNoiseSettings()
 	InputSource.SetPointProperty(EPCGPointProperties::Density);
 }
 
+#if WITH_EDITOR
+TArray<FPCGPreConfiguredSettingsInfo> UPCGAttributeNoiseSettings::GetPreconfiguredInfo() const
+{
+	TArray<FPCGPreConfiguredSettingsInfo> PreconfiguredInfo;
+	PreconfiguredInfo.Emplace(0, GetDefaultNodeTitle());
+	PreconfiguredInfo.Emplace(1, LOCTEXT("DensityNoiseNodeTitle", "Density Noise"));
+
+	return PreconfiguredInfo;
+}
+#endif
+
+void UPCGAttributeNoiseSettings::ApplyPreconfiguredSettings(const FPCGPreConfiguredSettingsInfo& PreconfigureInfo)
+{
+	// If index is 1, it is the default ($Density)
+	if (PreconfigureInfo.PreconfiguredIndex == 0)
+	{
+		InputSource.SetAttributeName(PCGMetadataAttributeConstants::LastAttributeName);
+	}
+}
+
 void UPCGAttributeNoiseSettings::PostLoad()
 {
 	Super::PostLoad();
@@ -124,6 +145,12 @@ void UPCGAttributeNoiseSettings::PostLoad()
 		bInvertSource = bInvertSourceDensity_DEPRECATED;
 		bInvertSourceDensity_DEPRECATED = false;
 	}
+
+	// Check for the data spatial to point gate version
+	if (DataVersion < FPCGCustomVersion::NoMoreSpatialDataConversionToPointDataByDefaultOnNonPointPins)
+	{
+		bHasSpatialToPointDeprecation = true;
+	}
 #endif // WITH_EDITOR
 }
 
@@ -145,6 +172,18 @@ void UPCGAttributeNoiseSettings::ApplyDeprecationBeforeUpdatePins(UPCGNode* InOu
 	{
 		InOutNode->RenameInputPin(OldToNew.Key, OldToNew.Value);
 	}
+
+	// Param | Point type was not explicitly defined in the data types, and therefore was not serialized correctly, resulting in an Input/Output pin serialized to None.
+	auto FixInvalidAllowedTypes = [](UPCGPin* InPin)
+	{
+		if (InPin && InPin->Properties.AllowedTypes == EPCGDataType::None)
+		{
+			InPin->Properties.AllowedTypes = EPCGDataType::PointOrParam;
+		}
+	};
+
+	FixInvalidAllowedTypes(InOutNode->GetInputPin(PCGPinConstants::DefaultInputLabel));
+	FixInvalidAllowedTypes(InOutNode->GetOutputPin(PCGPinConstants::DefaultOutputLabel));
 }
 #endif // WITH_EDITOR
 
@@ -156,7 +195,7 @@ FPCGElementPtr UPCGAttributeNoiseSettings::CreateElement() const
 TArray<FPCGPinProperties> UPCGAttributeNoiseSettings::InputPinProperties() const
 {
 	TArray<FPCGPinProperties> PinProperties;
-	FPCGPinProperties& InputPinProperty = PinProperties.Emplace_GetRef(PCGPinConstants::DefaultInputLabel, EPCGDataType::Point | EPCGDataType::Param);
+	FPCGPinProperties& InputPinProperty = PinProperties.Emplace_GetRef(PCGPinConstants::DefaultInputLabel, EPCGDataType::PointOrParam);
 	InputPinProperty.SetRequiredPin();
 	return PinProperties;
 }
@@ -164,11 +203,33 @@ TArray<FPCGPinProperties> UPCGAttributeNoiseSettings::InputPinProperties() const
 TArray<FPCGPinProperties> UPCGAttributeNoiseSettings::OutputPinProperties() const
 {
 	TArray<FPCGPinProperties> PinProperties;
-	PinProperties.Emplace(PCGPinConstants::DefaultOutputLabel, EPCGDataType::Point | EPCGDataType::Param);
+	PinProperties.Emplace(PCGPinConstants::DefaultOutputLabel, EPCGDataType::PointOrParam);
 
 	return PinProperties;
 }
 
+EPCGDataType UPCGAttributeNoiseSettings::GetCurrentPinTypes(const UPCGPin* InPin) const
+{
+	check(InPin);
+	if (!InPin->IsOutputPin())
+	{
+		return Super::GetCurrentPinTypes(InPin);
+	}
+
+	// Output pin narrows to union of inputs on first pin
+	const EPCGDataType InputTypeUnion = GetTypeUnionOfIncidentEdges(PCGPinConstants::DefaultInputLabel);
+
+	// Spatial is collapsed into points
+	if (InputTypeUnion != EPCGDataType::None && (InputTypeUnion & EPCGDataType::Spatial) == InputTypeUnion)
+	{
+		return EPCGDataType::Point;
+	}
+	else
+	{
+		return (InputTypeUnion != EPCGDataType::None) ? InputTypeUnion : (EPCGDataType::PointOrParam);
+	}
+}
+
 FPCGContext* FPCGAttributeNoiseElement::CreateContext()
 {
 	return new FPCGAttributeNoiseContext();
@@ -207,6 +268,15 @@ bool FPCGAttributeNoiseElement::ExecuteInternal(FPCGContext* InContext) const
 				continue;
 			}
 
+			// For deprecation
+			if (const UPCGSpatialData* InputSpatialData = Cast<const UPCGSpatialData>(InputData))
+			{
+				if (Settings->bHasSpatialToPointDeprecation)
+				{
+					InputData = InputSpatialData->ToPointData();
+				}
+			}
+
 			Context->InputSource = Settings->InputSource.CopyAndFixLast(InputData);
 
 			// Create a dummy accessor on the input before allocating the output, to avoid doing useless allocation.
diff --git a/Engine/Plugins/PCG/Source/PCG/Private/Elements/PCGMatchAndSetAttributes.cpp b/Engine/Plugins/PCG/Source/PCG/Private/Elements/PCGMatchAndSetAttributes.cpp
--- a/Engine/Plugins/PCG/Source/PCG/Private/Elements/PCGMatchAndSetAttributes.cpp
+++ b/Engine/Plugins/PCG/Source/PCG/Private/Elements/PCGMatchAndSetAttributes.cpp
@@ -44,6 +44,27 @@ FText UPCGMatchAndSetAttributesSettings::GetNodeTooltipText() const
 {
 	return LOCTEXT("NodeTooltip", "Matches or randomly assigns values from the Attribute Set to the input data");
 }
+
+void UPCGMatchAndSetAttributesSettings::ApplyDeprecationBeforeUpdatePins(UPCGNode* InOutNode, TArray<TObjectPtr<UPCGPin>>& InputPins, TArray<TObjectPtr<UPCGPin>>& OutputPins)
+{
+	Super::ApplyDeprecationBeforeUpdatePins(InOutNode, InputPins, OutputPins);
+
+	check(InOutNode);
+
+	// Param | Point type was not explicitly defined in the data types, and therefore was not serialized correctly, resulting in an Input/Output pin serialized to None.
+	// Restoring the right value here, before update pins.
+	auto FixInvalidAllowedTypes = [](UPCGPin* InPin)
+	{
+		if (InPin && InPin->Properties.AllowedTypes == EPCGDataType::None)
+		{
+			InPin->Properties.AllowedTypes = EPCGDataType::PointOrParam;
+		}
+	};
+
+	FixInvalidAllowedTypes(InOutNode->GetInputPin(PCGPinConstants::DefaultInputLabel));
+	FixInvalidAllowedTypes(InOutNode->GetInputPin(PCGMatchAndSetAttributesConstants::MaxDistanceLabel));
+	FixInvalidAllowedTypes(InOutNode->GetOutputPin(PCGPinConstants::DefaultOutputLabel));
+}
 #endif // WITH_EDITOR
 
 UPCGMatchAndSetAttributesSettings::UPCGMatchAndSetAttributesSettings()
@@ -55,7 +76,7 @@ UPCGMatchAndSetAttributesSettings::UPCGMatchAndSetAttributesSettings()
 TArray<FPCGPinProperties> UPCGMatchAndSetAttributesSettings::InputPinProperties() const
 {
 	TArray<FPCGPinProperties> PinProperties;
-	FPCGPinProperties& InputPinProperty = PinProperties.Emplace_GetRef(PCGPinConstants::DefaultInputLabel, EPCGDataType::Point | EPCGDataType::Param);
+	FPCGPinProperties& InputPinProperty = PinProperties.Emplace_GetRef(PCGPinConstants::DefaultInputLabel, EPCGDataType::PointOrParam);
 	InputPinProperty.SetRequiredPin();
 
 	PinProperties.Emplace(PCGMatchAndSetAttributesConstants::MatchDataLabel, 
@@ -68,7 +89,7 @@ TArray<FPCGPinProperties> UPCGMatchAndSetAttributesSettings::InputPinProperties(
 	if (bFindNearest && MaxDistanceMode == EPCGMatchMaxDistanceMode::AttributeMaxDistance)
 	{
 		PinProperties.Emplace(PCGMatchAndSetAttributesConstants::MaxDistanceLabel,
-			EPCGDataType::Point | EPCGDataType::Param,
+			EPCGDataType::PointOrParam,
 			/*bAllowMultipleConnections=*/false,
 			/*bAllowMultipleData=*/true,
 			LOCTEXT("MaxDistanceTooltip", "Input containing the maximum distance allowed for nearest search, selected by the Max Distance Attribute.")
@@ -81,7 +102,7 @@ TArray<FPCGPinProperties> UPCGMatchAndSetAttributesSettings::InputPinProperties(
 TArray<FPCGPinProperties> UPCGMatchAndSetAttributesSettings::OutputPinProperties() const
 {
 	TArray<FPCGPinProperties> PinProperties;
-	PinProperties.Emplace(PCGPinConstants::DefaultOutputLabel, EPCGDataType::Point | EPCGDataType::Param);
+	PinProperties.Emplace(PCGPinConstants::DefaultOutputLabel, EPCGDataType::PointOrParam);
 
 	return PinProperties;
 }
diff --git a/Engine/Plugins/PCG/Source/PCG/Private/Tests/GraphAuthoring/PCGGraphAuthoringTest.cpp b/Engine/Plugins/PCG/Source/PCG/Private/Tests/GraphAuthoring/PCGGraphAuthoringTest.cpp
--- a/Engine/Plugins/PCG/Source/PCG/Private/Tests/GraphAuthoring/PCGGraphAuthoringTest.cpp
+++ b/Engine/Plugins/PCG/Source/PCG/Private/Tests/GraphAuthoring/PCGGraphAuthoringTest.cpp
@@ -93,11 +93,11 @@ bool FPCGGraphAuthoringTypeSystemTest::RunTest(const FString& Parameters)
 	bTestPassed &= ValidateConversion("Collapse: Spatial to Point", EPCGDataType::Spatial, EPCGDataType::Point, true, EPCGTypeConversion::CollapseToPoint);
 
 	// Collapse, with Point | Param
-	bTestPassed &= ValidateConversion("Collapse: Spline to Point", EPCGDataType::Spline, EPCGDataType::Point | EPCGDataType::Param, true, EPCGTypeConversion::CollapseToPoint);
-	bTestPassed &= ValidateConversion("Collapse: Landscape to Point", EPCGDataType::Landscape, EPCGDataType::Point | EPCGDataType::Param, true, EPCGTypeConversion::CollapseToPoint);
-	bTestPassed &= ValidateConversion("Collapse: Surface to Point", EPCGDataType::Surface, EPCGDataType::Point | EPCGDataType::Param, true, EPCGTypeConversion::CollapseToPoint);
-	bTestPassed &= ValidateConversion("Collapse: Concrete to Point", EPCGDataType::Concrete, EPCGDataType::Point | EPCGDataType::Param, true, EPCGTypeConversion::CollapseToPoint);
-	bTestPassed &= ValidateConversion("Collapse: Spatial to Point", EPCGDataType::Spatial, EPCGDataType::Point | EPCGDataType::Param, true, EPCGTypeConversion::CollapseToPoint);
+	bTestPassed &= ValidateConversion("Collapse: Spline to Point", EPCGDataType::Spline, EPCGDataType::PointOrParam, true, EPCGTypeConversion::CollapseToPoint);
+	bTestPassed &= ValidateConversion("Collapse: Landscape to Point", EPCGDataType::Landscape, EPCGDataType::PointOrParam, true, EPCGTypeConversion::CollapseToPoint);
+	bTestPassed &= ValidateConversion("Collapse: Surface to Point", EPCGDataType::Surface, EPCGDataType::PointOrParam, true, EPCGTypeConversion::CollapseToPoint);
+	bTestPassed &= ValidateConversion("Collapse: Concrete to Point", EPCGDataType::Concrete, EPCGDataType::PointOrParam, true, EPCGTypeConversion::CollapseToPoint);
+	bTestPassed &= ValidateConversion("Collapse: Spatial to Point", EPCGDataType::Spatial, EPCGDataType::PointOrParam, true, EPCGTypeConversion::CollapseToPoint);
 
 	// Filter
 	bTestPassed &= ValidateConversion("Filter: Any to Point", EPCGDataType::Any, EPCGDataType::Point, true, EPCGTypeConversion::Filter);
diff --git a/Engine/Plugins/PCG/Source/PCG/Public/Elements/PCGAttributeFilter.h b/Engine/Plugins/PCG/Source/PCG/Public/Elements/PCGAttributeFilter.h
--- a/Engine/Plugins/PCG/Source/PCG/Public/Elements/PCGAttributeFilter.h
+++ b/Engine/Plugins/PCG/Source/PCG/Public/Elements/PCGAttributeFilter.h
@@ -77,10 +77,15 @@ public:
 #if WITH_EDITOR
 	virtual FName GetDefaultNodeName() const override { return FName(TEXT("AttributeFilter")); }
 	virtual FText GetDefaultNodeTitle() const override { return NSLOCTEXT("PCGAttributeFilteringElement", "NodeTitle", "Attribute Filter"); }
-	virtual TArray<FText> GetNodeTitleAliases() const override { return { NSLOCTEXT("PCGAttributeFilteringElement", "AliasNodeTitle", "Point Filter") }; }
 	virtual EPCGSettingsType GetType() const override { return EPCGSettingsType::Filter; }
 	virtual bool HasDynamicPins() const override { return true; }
+
+	// Expose 2 nodes: Attribute filter and Point Filter that will not have the same defaults.
+	virtual TArray<FPCGPreConfiguredSettingsInfo> GetPreconfiguredInfo() const override;
+	virtual bool OnlyExposePreconfiguredSettings() const override { return true; }
+	virtual bool GroupPreconfiguredSettings() const override { return false; }
 #endif
+	virtual void ApplyPreconfiguredSettings(const FPCGPreConfiguredSettingsInfo& PreconfigureInfo) override;
 	virtual FString GetAdditionalTitleInformation() const override;
 	virtual EPCGDataType GetCurrentPinTypes(const UPCGPin* InPin) const override;
 
@@ -111,6 +116,10 @@ public:
 
 	UPROPERTY(BlueprintReadWrite, EditAnywhere, Category = Settings, meta = (EditCondition = "bUseConstantThreshold", EditConditionHides, ShowOnlyInnerProperties, DisplayAfter = "bUseConstantThreshold", PCG_NotOverridable))
 	FPCGMetadataTypesConstantStruct AttributeTypes;
+
+	// Hidden value to indicate that Spatial -> Point deprecation is on where pins are not explicitly points.
+	UPROPERTY()
+	bool bHasSpatialToPointDeprecation = false;
 };
 
 
@@ -135,19 +144,22 @@ public:
 	UPCGAttributeFilteringRangeSettings();
 
 	//~Begin UObject interface
-#if WITH_EDITOR
 	virtual void PostLoad() override;
-#endif
 	//~End UObject interface
 
 	//~Begin UPCGSettings interface
 #if WITH_EDITOR
 	virtual FName GetDefaultNodeName() const override { return FName(TEXT("AttributeFilterRange")); }
 	virtual FText GetDefaultNodeTitle() const override { return NSLOCTEXT("PCGAttributeFilteringElement", "NodeTitleRange", "Attribute Filter Range"); }
-	virtual TArray<FText> GetNodeTitleAliases() const override { return { NSLOCTEXT("PCGAttributeFilteringElement", "AliasNodeTitleRange", "Point Filter Range") }; }
 	virtual EPCGSettingsType GetType() const override { return EPCGSettingsType::Filter; }
 	virtual bool HasDynamicPins() const override { return true; }
+
+	// Expose 2 nodes: "Attribute Filter Range" and "Point Filter Range" that will not have the same defaults.
+	virtual TArray<FPCGPreConfiguredSettingsInfo> GetPreconfiguredInfo() const override;
+	virtual bool OnlyExposePreconfiguredSettings() const override { return true; }
+	virtual bool GroupPreconfiguredSettings() const override { return false; }
 #endif
+	virtual void ApplyPreconfiguredSettings(const FPCGPreConfiguredSettingsInfo& PreconfigureInfo) override;
 	virtual FString GetAdditionalTitleInformation() const override;
 	virtual EPCGDataType GetCurrentPinTypes(const UPCGPin* InPin) const override;
 
@@ -168,12 +180,16 @@ public:
 
 	UPROPERTY(BlueprintReadWrite, EditAnywhere, Category = Settings, meta = (PCG_Overridable))
 	FPCGAttributeFilterThresholdSettings MaxThreshold;
+
+	// Hidden value to indicate that Spatial -> Point deprecation is on where pins are not explicitly points.
+	UPROPERTY()
+	bool bHasSpatialToPointDeprecation = false;
 };
 
 class FPCGAttributeFilterElementBase : public IPCGElement
 {
 protected:
-	bool DoFiltering(FPCGContext* Context, EPCGAttributeFilterOperator InOperation, const FPCGAttributePropertyInputSelector& TargetAttribute, const FPCGAttributeFilterThresholdSettings& FirstThreshold, const FPCGAttributeFilterThresholdSettings* SecondThreshold = nullptr) const;
+	bool DoFiltering(FPCGContext* Context, EPCGAttributeFilterOperator InOperation, const FPCGAttributePropertyInputSelector& TargetAttribute, bool bHasSpatialToPointDeprecation, const FPCGAttributeFilterThresholdSettings& FirstThreshold, const FPCGAttributeFilterThresholdSettings* SecondThreshold = nullptr) const;
 };
 
 class FPCGAttributeFilterElement : public FPCGAttributeFilterElementBase
diff --git a/Engine/Plugins/PCG/Source/PCG/Public/Elements/PCGAttributeNoise.h b/Engine/Plugins/PCG/Source/PCG/Public/Elements/PCGAttributeNoise.h
--- a/Engine/Plugins/PCG/Source/PCG/Public/Elements/PCGAttributeNoise.h
+++ b/Engine/Plugins/PCG/Source/PCG/Public/Elements/PCGAttributeNoise.h
@@ -44,9 +44,15 @@ public:
 	virtual void ApplyDeprecationBeforeUpdatePins(UPCGNode* InOutNode, TArray<TObjectPtr<UPCGPin>>& InputPins, TArray<TObjectPtr<UPCGPin>>& OutputPins);
 	virtual FName GetDefaultNodeName() const override { return FName(TEXT("AttributeNoise")); }
 	virtual FText GetDefaultNodeTitle() const override { return NSLOCTEXT("PCGAttributeNoiseSettings", "NodeTitle", "Attribute Noise"); }
-	virtual TArray<FText> GetNodeTitleAliases() const { return { NSLOCTEXT("PCGAttributeNoiseSettings", "DensityNoiseAlias", "Density Noise") }; }
 	virtual EPCGSettingsType GetType() const override { return EPCGSettingsType::Metadata; }
+
+	// Expose 2 nodes: Density noise and Attribute Noise that will not have the same defaults.
+	virtual TArray<FPCGPreConfiguredSettingsInfo> GetPreconfiguredInfo() const override;
+	virtual bool OnlyExposePreconfiguredSettings() const override { return true; }
+	virtual bool GroupPreconfiguredSettings() const override { return false; }
 #endif
+	virtual void ApplyPreconfiguredSettings(const FPCGPreConfiguredSettingsInfo& PreconfigureInfo) override;
+	virtual EPCGDataType GetCurrentPinTypes(const UPCGPin* InPin) const override;
 	virtual bool HasDynamicPins() const override { return true; }
 	
 
@@ -98,6 +104,10 @@ public:
 	UPROPERTY()
 	bool bOutputTargetDifferentFromInputSource_DEPRECATED = false;
 #endif // WITH_EDITORDATA_ONLY
+
+	// Hidden value to indicate that Spatial -> Point deprecation is on where pins are not explicitly points.
+	UPROPERTY()
+	bool bHasSpatialToPointDeprecation = false;
 };
 
 struct FPCGAttributeNoiseContext : public FPCGContext
diff --git a/Engine/Plugins/PCG/Source/PCG/Public/Elements/PCGMatchAndSetAttributes.h b/Engine/Plugins/PCG/Source/PCG/Public/Elements/PCGMatchAndSetAttributes.h
--- a/Engine/Plugins/PCG/Source/PCG/Public/Elements/PCGMatchAndSetAttributes.h
+++ b/Engine/Plugins/PCG/Source/PCG/Public/Elements/PCGMatchAndSetAttributes.h
@@ -54,6 +54,7 @@ public:
 	virtual FText GetNodeTooltipText() const override;
 	virtual EPCGSettingsType GetType() const override { return EPCGSettingsType::Metadata; }
 	virtual bool HasDynamicPins() const { return true; }
+	virtual void ApplyDeprecationBeforeUpdatePins(UPCGNode* InOutNode, TArray<TObjectPtr<UPCGPin>>& InputPins, TArray<TObjectPtr<UPCGPin>>& OutputPins) override;
 #endif // WITH_EDITOR
 
 protected:
diff --git a/Engine/Plugins/PCG/Source/PCG/Public/PCGCommon.h b/Engine/Plugins/PCG/Source/PCG/Public/PCGCommon.h
--- a/Engine/Plugins/PCG/Source/PCG/Public/PCGCommon.h
+++ b/Engine/Plugins/PCG/Source/PCG/Public/PCGCommon.h
@@ -54,6 +54,7 @@ ENUM_CLASS_FLAGS(EPCGChangeType);
 // enums have to be uint8 for blueprint, and we already use more than 8 bits in the bitmask.
 // This is why we have a parallel enum just below that must match on a name basis 1:1 to allow the make/break functions to work properly
 // in blueprint.
+// WARNING: Please be mindful that combination of flags that are not explicitly defined there won't be serialized correctly, inducing data loss.
 UENUM(meta = (Bitflags))
 enum class EPCGDataType : uint32
 {
@@ -83,6 +84,10 @@ enum class EPCGDataType : uint32
 	Spatial = Composite | Concrete,
 
 	Param = 1 << 27 UMETA(DisplayName = "Attribute Set"),
+
+	// Combination of Param and Point, necessary for named-based serialization of enums.
+	PointOrParam = Point | Param,
+
 	Settings = 1 << 28 UMETA(Hidden),
 	Other = 1 << 29,
 	Any = (1 << 30) - 1
@@ -111,7 +116,8 @@ enum class EPCGExclusiveDataType : uint8
 	Param UMETA(DisplayName = "Attribute Set"),
 	Settings UMETA(Hidden),
 	Other,
-	Any
+	Any,
+	PointOrParam
 };
 
 namespace PCGPinConstants
diff --git a/Engine/Plugins/PCG/Source/PCG/Public/PCGCustomVersion.h b/Engine/Plugins/PCG/Source/PCG/Public/PCGCustomVersion.h
--- a/Engine/Plugins/PCG/Source/PCG/Public/PCGCustomVersion.h
+++ b/Engine/Plugins/PCG/Source/PCG/Public/PCGCustomVersion.h
@@ -74,6 +74,9 @@ struct PCG_API FPCGCustomVersion
 		// Supporting partitioned components in non-partitioned levels
 		SupportPartitionedComponentsInNonPartitionedLevels = 21,
 
+		// New gate for new data, so any node that has a non Point pin don't do any ToPointData by default.
+		NoMoreSpatialDataConversionToPointDataByDefaultOnNonPointPins = 22,
+
 		// -----<new versions can be added above this line>-------------------------------------------------
 		VersionPlusOne,
 		LatestVersion = VersionPlusOne - 1
diff --git a/Engine/Plugins/PCG/Source/PCG/Public/PCGSettings.h b/Engine/Plugins/PCG/Source/PCG/Public/PCGSettings.h
--- a/Engine/Plugins/PCG/Source/PCG/Public/PCGSettings.h
+++ b/Engine/Plugins/PCG/Source/PCG/Public/PCGSettings.h
@@ -314,6 +314,9 @@ public:
 	/** If there are preconfigured info, we can skip the default settings and only expose pre-configured actions in the editor palette */
 	virtual bool OnlyExposePreconfiguredSettings() const { return false; }
 
+	/** If there are preconfigured info, decide if they are grouped in the palette in a folder with the node name, or if they are not grouped. */
+	virtual bool GroupPreconfiguredSettings() const { return true; }
+
 	/** Perform post-operations when an editor node is copied */
 	virtual void PostPaste();
 
@@ -430,14 +433,14 @@ protected:
 	TArray<FPCGPinProperties> DefaultPointInputPinProperties() const;
 	TArray<FPCGPinProperties> DefaultPointOutputPinProperties() const;
 
-#if WITH_EDITOR
 public:
+#if WITH_EDITORONLY_DATA
 	/** The version number of the data after load and after any data migration. */
 	int32 DataVersion = -1;
 
 	/** If a custom version guid was provided through GetUserCustomVersionGuid(), this field will hold the version number after load and after any data migration. */
 	int32 UserDataVersion = -1;
-#endif
+#endif // WITH_EDITORONLY_DATA
 
 private:
 	mutable FPCGElementPtr CachedElement;
diff --git a/Engine/Plugins/PCG/Source/PCGEditor/Private/PCGEditorGraphSchema.cpp b/Engine/Plugins/PCG/Source/PCGEditor/Private/PCGEditorGraphSchema.cpp
--- a/Engine/Plugins/PCG/Source/PCGEditor/Private/PCGEditorGraphSchema.cpp
+++ b/Engine/Plugins/PCG/Source/PCGEditor/Private/PCGEditorGraphSchema.cpp
@@ -400,7 +400,7 @@ void UPCGEditorGraphSchema::GetNativeElementActions(FGraphActionMenuBuilder& Act
 				}
 
 				// Also add preconfigured settings
-				const FText NewCategory = FText::Format(LOCTEXT("PreconfiguredSettingsCategory", "{0}|{1}"), Category, MenuDesc);
+				const FText NewCategory = PCGSettings->GroupPreconfiguredSettings() ? FText::Format(LOCTEXT("PreconfiguredSettingsCategory", "{0}|{1}"), Category, MenuDesc) : Category;
 
 				for (FPCGPreConfiguredSettingsInfo PreconfiguredInfo : AllPreconfiguredInfo)
 				{
