[PCG] Some pins are required while they should not be and some nodes are culled when they should not

* Pins issue: Blueprint
* Cull issue: Get Object from Property Path

#rnx

diff --git a/Engine/Plugins/PCG/Source/PCG/Private/Elements/PCGExecuteBlueprint.cpp b/Engine/Plugins/PCG/Source/PCG/Private/Elements/PCGExecuteBlueprint.cpp
--- a/Engine/Plugins/PCG/Source/PCG/Private/Elements/PCGExecuteBlueprint.cpp
+++ b/Engine/Plugins/PCG/Source/PCG/Private/Elements/PCGExecuteBlueprint.cpp
@@ -587,7 +587,8 @@ TArray<FPCGPinProperties> UPCGBlueprintSettings::InputPinProperties() const
 
 	if (!BlueprintElementInstance || BlueprintElementInstance->bHasDefaultInPin)
 	{
-		PinProperties.Append(Super::InputPinProperties());
+		// Here we do not want the base class implementation as it makes the input pin required.
+		PinProperties.Emplace(PCGPinConstants::DefaultInputLabel, EPCGDataType::Any);
 	}
 
 	if (BlueprintElementInstance)
diff --git a/Engine/Plugins/PCG/Source/PCG/Public/Elements/PCGGetPropertyFromObjectPath.h b/Engine/Plugins/PCG/Source/PCG/Public/Elements/PCGGetPropertyFromObjectPath.h
--- a/Engine/Plugins/PCG/Source/PCG/Public/Elements/PCGGetPropertyFromObjectPath.h
+++ b/Engine/Plugins/PCG/Source/PCG/Public/Elements/PCGGetPropertyFromObjectPath.h
@@ -27,6 +27,7 @@ public:
 	virtual void GetStaticTrackedKeys(FPCGSelectionKeyToSettingsMap& OutKeysToSettings, TArray<TObjectPtr<const UPCGGraph>>& OutVisitedGraphs) const override;
 	virtual bool CanDynamicalyTrackKeys() const override { return true; }
 #endif
+	virtual bool CanCullTaskIfUnwired() const override { return false; }
 	virtual FString GetAdditionalTitleInformation() const override;
 
 protected:
