[PCG] Make all new actor spawning nodes attached InFolder by default.

#rnx

diff --git a/Engine/Plugins/PCG/Source/PCG/Private/Elements/PCGCreateSpline.cpp b/Engine/Plugins/PCG/Source/PCG/Private/Elements/PCGCreateSpline.cpp
--- a/Engine/Plugins/PCG/Source/PCG/Private/Elements/PCGCreateSpline.cpp
+++ b/Engine/Plugins/PCG/Source/PCG/Private/Elements/PCGCreateSpline.cpp
@@ -21,6 +21,15 @@
 
 #define LOCTEXT_NAMESPACE "PCGCreateSpline"
 
+UPCGCreateSplineSettings::UPCGCreateSplineSettings(const FObjectInitializer& ObjectInitializer)
+	: UPCGSettings(ObjectInitializer)
+{
+	if (PCGHelpers::IsNewObjectAndNotDefault(this))
+	{
+		AttachOptions = EPCGAttachOptions::InFolder;
+	}
+}
+
 #if WITH_EDITOR
 FText UPCGCreateSplineSettings::GetNodeTooltipText() const
 {
diff --git a/Engine/Plugins/PCG/Source/PCG/Private/Elements/PCGCreateTargetActor.cpp b/Engine/Plugins/PCG/Source/PCG/Private/Elements/PCGCreateTargetActor.cpp
--- a/Engine/Plugins/PCG/Source/PCG/Private/Elements/PCGCreateTargetActor.cpp
+++ b/Engine/Plugins/PCG/Source/PCG/Private/Elements/PCGCreateTargetActor.cpp
@@ -37,6 +37,15 @@ namespace PCGCreateTargetActorConstants
 	const FText ActorPropertyOverridesTooltip = LOCTEXT("ActorOverrideToolTip", "Provide property overrides for the created target actor. The attribute name must match the InputSource name in the actor property override description.");
 }
 
+UPCGCreateTargetActor::UPCGCreateTargetActor(const FObjectInitializer& ObjectInitializer)
+	: UPCGSettings(ObjectInitializer)
+{
+	if (PCGHelpers::IsNewObjectAndNotDefault(this))
+	{
+		AttachOptions = EPCGAttachOptions::InFolder;
+	}
+}
+
 FPCGElementPtr UPCGCreateTargetActor::CreateElement() const
 {
 	return MakeShared<FPCGCreateTargetActorElement>();
diff --git a/Engine/Plugins/PCG/Source/PCG/Private/Elements/PCGSpawnActor.cpp b/Engine/Plugins/PCG/Source/PCG/Private/Elements/PCGSpawnActor.cpp
--- a/Engine/Plugins/PCG/Source/PCG/Private/Elements/PCGSpawnActor.cpp
+++ b/Engine/Plugins/PCG/Source/PCG/Private/Elements/PCGSpawnActor.cpp
@@ -127,6 +127,15 @@ public:
 	FName SpawnAttribute = NAME_None;
 };
 
+UPCGSpawnActorSettings::UPCGSpawnActorSettings(const FObjectInitializer& ObjectInitializer)
+	: UPCGBaseSubgraphSettings(ObjectInitializer)
+{
+	if (PCGHelpers::IsNewObjectAndNotDefault(this))
+	{
+		AttachOptions = EPCGAttachOptions::InFolder;
+	}
+}
+
 UPCGNode* UPCGSpawnActorSettings::CreateNode() const
 {
 	return NewObject<UPCGSpawnActorNode>();
diff --git a/Engine/Plugins/PCG/Source/PCG/Public/Elements/PCGCreateSpline.h b/Engine/Plugins/PCG/Source/PCG/Public/Elements/PCGCreateSpline.h
--- a/Engine/Plugins/PCG/Source/PCG/Public/Elements/PCGCreateSpline.h
+++ b/Engine/Plugins/PCG/Source/PCG/Public/Elements/PCGCreateSpline.h
@@ -25,6 +25,8 @@ class PCG_API UPCGCreateSplineSettings : public UPCGSettings
 	GENERATED_BODY()
 
 public:
+	UPCGCreateSplineSettings(const FObjectInitializer& ObjectInitializer);
+
 	//~Begin UPCGSettings interface
 #if WITH_EDITOR
 	virtual FName GetDefaultNodeName() const override { return FName(TEXT("CreateSpline")); }
@@ -66,7 +68,7 @@ public:
 	TSoftObjectPtr<AActor> TargetActor;
 
 	UPROPERTY(BlueprintReadWrite, EditAnywhere, Category = Settings, meta = (EditCondition = "Mode==EPCGCreateSplineMode::CreateNewActor", EditConditionHides))
-	EPCGAttachOptions AttachOptions = EPCGAttachOptions::Attached;
+	EPCGAttachOptions AttachOptions = EPCGAttachOptions::Attached; // Note that this is no longer the default value for new nodes, it is now EPCGAttachOptions::InFolder
 
 	/** Specify a list of functions to be called on the target actor after spline creation. Functions need to be parameter-less and with "CallInEditor" flag enabled. */
 	UPROPERTY(BlueprintReadWrite, EditAnywhere, Category = Settings)
diff --git a/Engine/Plugins/PCG/Source/PCG/Public/Elements/PCGCreateTargetActor.h b/Engine/Plugins/PCG/Source/PCG/Public/Elements/PCGCreateTargetActor.h
--- a/Engine/Plugins/PCG/Source/PCG/Public/Elements/PCGCreateTargetActor.h
+++ b/Engine/Plugins/PCG/Source/PCG/Public/Elements/PCGCreateTargetActor.h
@@ -15,6 +15,8 @@ class PCG_API UPCGCreateTargetActor : public UPCGSettings
 	GENERATED_BODY()
 
 public:
+	UPCGCreateTargetActor(const FObjectInitializer& ObjectInitializer);
+
 	//~Begin UObject interface
 	virtual void PostLoad() override;
 	virtual void BeginDestroy() override;
@@ -61,14 +63,17 @@ public:
 	bool bAllowTemplateActorEditing = false;
 
 	UPROPERTY(BlueprintReadWrite, EditAnywhere, Category = Settings)
-	EPCGAttachOptions AttachOptions = EPCGAttachOptions::Attached;
+	EPCGAttachOptions AttachOptions = EPCGAttachOptions::Attached; // Note that this is no longer the default value for new nodes, it is now EPCGAttachOptions::InFolder
 
+	/** Actor to attach to if the option is Attached. Default to the Component owner. */
 	UPROPERTY(meta = (PCG_Overridable))
 	TSoftObjectPtr<AActor> RootActor;
 
+	/** Name of the actor that will be created. */
 	UPROPERTY(meta = (PCG_Overridable))
 	FString ActorLabel;
 
+	/** Transform of the actor that will be created. */
 	UPROPERTY(meta = (PCG_Overridable))
 	FTransform ActorPivot;
 
diff --git a/Engine/Plugins/PCG/Source/PCG/Public/Elements/PCGSpawnActor.h b/Engine/Plugins/PCG/Source/PCG/Public/Elements/PCGSpawnActor.h
--- a/Engine/Plugins/PCG/Source/PCG/Public/Elements/PCGSpawnActor.h
+++ b/Engine/Plugins/PCG/Source/PCG/Public/Elements/PCGSpawnActor.h
@@ -41,6 +41,8 @@ class PCG_API UPCGSpawnActorSettings : public UPCGBaseSubgraphSettings
 	GENERATED_BODY()
 
 public:
+	UPCGSpawnActorSettings(const FObjectInitializer& ObjectInitializer);
+
 	UPROPERTY(BlueprintReadOnly, EditAnywhere, Category = Settings, meta = (OnlyPlaceable, DisallowCreateNew))
 	TSubclassOf<AActor> TemplateActorClass = nullptr;
 
@@ -80,7 +82,7 @@ public:
 	TSoftObjectPtr<AActor> RootActor;
 
 	UPROPERTY(BlueprintReadWrite, EditAnywhere, Category = Settings, meta = (EditCondition = "Option != EPCGSpawnActorOption::CollapseActors", EditConditionHides))
-	EPCGAttachOptions AttachOptions = EPCGAttachOptions::Attached;
+	EPCGAttachOptions AttachOptions = EPCGAttachOptions::Attached; // Note that this is no longer the default value for new nodes, it is now EPCGAttachOptions::InFolder
 
 	UPROPERTY(BlueprintReadWrite, EditAnywhere, Category = Settings, meta = (InlineEditConditionToggle))
 	bool bSpawnByAttribute = false;
