[PCG] Add error capture to PCG world partition builder and fail build if errors occurred, with argument to disable this behaviour.

diff --git a/Engine/Plugins/PCG/Source/PCGEditor/Private/WorldPartitionBuilder/PCGWorldPartitionBuilder.cpp b/Engine/Plugins/PCG/Source/PCGEditor/Private/WorldPartitionBuilder/PCGWorldPartitionBuilder.cpp
--- a/Engine/Plugins/PCG/Source/PCGEditor/Private/WorldPartitionBuilder/PCGWorldPartitionBuilder.cpp
+++ b/Engine/Plugins/PCG/Source/PCGEditor/Private/WorldPartitionBuilder/PCGWorldPartitionBuilder.cpp
@@ -8,6 +8,7 @@
 #include "AssetCompilingManager.h"
 #include "Editor.h"
 #include "FileHelpers.h"
+#include "Misc/OutputDevice.h"
 #include "Modules/ModuleManager.h"
 #include "UObject/Linker.h"
 #include "WorldPartition/IWorldPartitionEditorModule.h"
@@ -21,8 +22,8 @@ namespace PCGWorldPartitionBuilder
 	void CollectComponentsToGenerate(UWorld* InWorld, TFunctionRef<bool(const UPCGComponent*)> ComponentFilter, TArray<TWeakObjectPtr<UPCGComponent>>& OutComponents);
 	
 	/** Generate the given components. Optionally generate only one component at a time with a wait on async processes after each. Optionally apply given filter to select components. */
-	bool GenerateComponents(TArray<TWeakObjectPtr<UPCGComponent>>& Components, UWorld* InWorld, bool bOneComponentAtATime);
-	bool GenerateComponents(TArray<TWeakObjectPtr<UPCGComponent>>& Components, UWorld* InWorld, bool bOneComponentAtATime, TFunctionRef<bool(const UPCGComponent*)> ComponentFilter);
+	bool GenerateComponents(TArray<TWeakObjectPtr<UPCGComponent>>& Components, UWorld* InWorld, bool bOneComponentAtATime, bool& bOutGenerationErrors);
+	bool GenerateComponents(TArray<TWeakObjectPtr<UPCGComponent>>& Components, UWorld* InWorld, bool bOneComponentAtATime, TFunctionRef<bool(const UPCGComponent*)> ComponentFilter, bool& bOutGenerationErrors);
 
 	/** Generate a component. Applies correct editing mode if necessary. */
 	void GenerateComponent(UPCGComponent* InComponent, UWorld* InWorld);
@@ -35,10 +36,54 @@ namespace PCGWorldPartitionBuilder
 
 	static FAutoConsoleCommand CommandBuildComponents(
 		TEXT("pcg.BuildComponents"),
-		TEXT("Runs PCG world builder on current world. Args: [-GenerateComponentEditingModeNormal] [-IncludeGraphNames=PCG_GraphA;PCG_GraphB] [-IncludeActorIDs=MyActor1_UID1234678;MyActor2_UID1234678]"),
+		TEXT("Runs PCG world builder on PCG components in current world. Arguments (multiple values separated with ';'):\n"
+			"\t[-IncludeGraphNames=PCG_GraphA;PCG_GraphB]\n"
+			"\t[-GenerateComponentEditingModeNormal]\n"
+			"\t[-GenerateComponentEditingModePreview]\n"
+			"\t[-IgnoreGenerationErrors]\n"
+			"\t[-IncludeActorIDs=MyActor1_UID1234678;MyActor2_UID1234678]\n"
+			"\t[-OneComponentAtATime]\n"),
 		FConsoleCommandWithArgsDelegate::CreateLambda([](const TArray<FString>& Args) { PCGWorldPartitionBuilder::Build(Args); }));
 };
 
+/** Output device to capture presence of errors during generation. */
+struct FPCGDetectErrorsInScope : public FOutputDevice
+{
+	FPCGDetectErrorsInScope()
+	{
+		if (GLog)
+		{
+			GLog->AddOutputDevice(this);
+		}
+	}
+
+	virtual ~FPCGDetectErrorsInScope()
+	{
+		if (GLog)
+		{
+			GLog->RemoveOutputDevice(this);
+		}
+	}
+
+	//~Begin FOutputDevice interface
+	virtual bool IsMemoryOnly() const override { return true; }
+	virtual bool CanBeUsedOnMultipleThreads() const override { return true; }
+	virtual bool CanBeUsedOnAnyThread() const override { return true; }
+	virtual void Serialize(const TCHAR* V, ELogVerbosity::Type Verbosity, const class FName& Category)
+	{
+		if (Verbosity <= ELogVerbosity::Error)
+		{
+			bErrorOccurred = true;
+		}
+	}
+	//~End FOutputDevice interface
+
+	bool GetErrorOccurred() const { return bErrorOccurred; }
+
+private:
+	std::atomic<bool> bErrorOccurred = false;
+};
+
 UPCGWorldPartitionBuilder::UPCGWorldPartitionBuilder(const FObjectInitializer& ObjectInitializer)
 	: Super(ObjectInitializer)
 {
@@ -50,6 +95,8 @@ UPCGWorldPartitionBuilder::UPCGWorldPartitionBuilder(const FObjectInitializer& O
 
 		bOneComponentAtATime = HasParam("OneComponentAtATime");
 
+		bIgnoreGenerationErrors = HasParam("IgnoreGenerationErrors");
+
 		FString IncludeGraphNamesValue;
 		if (GetParamValue("IncludeGraphNames=", IncludeGraphNamesValue) && !IncludeGraphNamesValue.IsEmpty())
 		{
@@ -140,12 +187,18 @@ bool UPCGWorldPartitionBuilder::RunInternal(UWorld* World, const FCellInfo& InCe
 				return Graph && Graph->GetName() == IncludeGraphName;
 			};
 
-			bGeneratedAnyComponent |= PCGWorldPartitionBuilder::GenerateComponents(ComponentsToGenerate, World, bOneComponentAtATime, FilterOnGraphName);
+			bool bErrorsOccurred = false;
+			bGeneratedAnyComponent |= PCGWorldPartitionBuilder::GenerateComponents(ComponentsToGenerate, World, bOneComponentAtATime, FilterOnGraphName, bErrorsOccurred);
+
+			bErrorOccurredWhileGenerating |= bErrorsOccurred;
 		}
 	}
 	else
 	{
-		bGeneratedAnyComponent |= PCGWorldPartitionBuilder::GenerateComponents(ComponentsToGenerate, World, bOneComponentAtATime);
+		bool bErrorsOccurred = false;
+		bGeneratedAnyComponent |= PCGWorldPartitionBuilder::GenerateComponents(ComponentsToGenerate, World, bOneComponentAtATime, bErrorsOccurred);
+
+		bErrorOccurredWhileGenerating |= bErrorsOccurred;
 	}
 
 	// 4. Get packages that were dirtied during generation and record them for saving later.
@@ -181,11 +234,29 @@ bool UPCGWorldPartitionBuilder::RunInternal(UWorld* World, const FCellInfo& InCe
 
 bool UPCGWorldPartitionBuilder::PostRun(UWorld* World, FPackageSourceControlHelper& PackageHelper, const bool bInRunSuccess)
 {
+	// bInRunSuccess is return value of Run() which is true if any component was scheduled.
 	if (!bInRunSuccess)
 	{
+		UE_LOG(LogPCGWorldPartitionBuilder, Display, TEXT("Dirty package detection and save skipped due to trivial run"));
+
 		return true;
 	}
 
+	// Check whether an error was thrown while generating components and fail the builder if the ignore argument is not provided.
+	if (bErrorOccurredWhileGenerating)
+	{
+		if (!bIgnoreGenerationErrors)
+		{
+			UE_LOG(LogPCGWorldPartitionBuilder, Display, TEXT("Dirty package detection and save skipped due to errors during generation."));
+
+			return !bErrorOccurredWhileGenerating;
+		}
+		else
+		{
+			UE_LOG(LogPCGWorldPartitionBuilder, Display, TEXT("Generation errors ignored, dirty packages will be saved."));
+		}
+	}
+
 	// Save/delete pending packages.
 
 	TArray<UPackage*> DirtyPackages;
@@ -268,22 +339,25 @@ void PCGWorldPartitionBuilder::CollectComponentsToGenerate(
 	}
 }
 
-bool PCGWorldPartitionBuilder::GenerateComponents(TArray<TWeakObjectPtr<UPCGComponent>>& Components, UWorld* InWorld, bool bOneComponentAtATime)
+bool PCGWorldPartitionBuilder::GenerateComponents(TArray<TWeakObjectPtr<UPCGComponent>>& Components, UWorld* InWorld, bool bOneComponentAtATime, bool& bOutGenerationErrors)
 {
-	return PCGWorldPartitionBuilder::GenerateComponents(Components, InWorld, bOneComponentAtATime, [](const UPCGComponent*) { return true; });
+	return PCGWorldPartitionBuilder::GenerateComponents(Components, InWorld, bOneComponentAtATime, [](const UPCGComponent*) { return true; }, bOutGenerationErrors);
 }
 
 bool PCGWorldPartitionBuilder::GenerateComponents(
 	TArray<TWeakObjectPtr<UPCGComponent>>& Components,
 	UWorld* InWorld,
 	bool bOneComponentAtATime,
-	TFunctionRef<bool(const UPCGComponent*)> ComponentFilter)
+	TFunctionRef<bool(const UPCGComponent*)> ComponentFilter,
+	bool& bOutGenerationErrors)
 {
 	if (!bOneComponentAtATime)
 	{
 		PCGWorldPartitionBuilder::WaitForAllAsyncEditorProcesses(InWorld);
 	}
 
+	const FPCGDetectErrorsInScope DetectErrors;
+
 	auto WaitForComponentGeneration = [InWorld](const UPCGComponent* InComponent)
 	{
 		UE_LOG(LogPCGWorldPartitionBuilder, Display, TEXT("Completing generation on PCG component on actor '%s' label '%s', graph '%s'"),
@@ -379,6 +453,8 @@ bool PCGWorldPartitionBuilder::GenerateComponents(
 		}
 	}
 
+	bOutGenerationErrors = DetectErrors.GetErrorOccurred();
+
 	return !GeneratedComponents.IsEmpty();
 }
 
diff --git a/Engine/Plugins/PCG/Source/PCGEditor/Private/WorldPartitionBuilder/PCGWorldPartitionBuilder.h b/Engine/Plugins/PCG/Source/PCGEditor/Private/WorldPartitionBuilder/PCGWorldPartitionBuilder.h
--- a/Engine/Plugins/PCG/Source/PCGEditor/Private/WorldPartitionBuilder/PCGWorldPartitionBuilder.h
+++ b/Engine/Plugins/PCG/Source/PCGEditor/Private/WorldPartitionBuilder/PCGWorldPartitionBuilder.h
@@ -53,4 +53,10 @@ private:
 
 	/** If non empty, only components on actors with given ID(s) will be generated. */
 	TArray<FString> IncludeActorIDs;
+
+	/** Submit dirty files even if errors occurred during generation. */
+	bool bIgnoreGenerationErrors = false;
+
+	/** Flag to register if any error occurred during generation. */
+	bool bErrorOccurredWhileGenerating = false;
 };
