[PCG] Fix loc and unittest

* Seed was changed after I ran the unitest, and I didn't re-run them, shame on me.

#rnx

diff --git a/Engine/Plugins/PCG/Source/PCG/Private/Elements/PCGAttributeFilter.cpp b/Engine/Plugins/PCG/Source/PCG/Private/Elements/PCGAttributeFilter.cpp
--- a/Engine/Plugins/PCG/Source/PCG/Private/Elements/PCGAttributeFilter.cpp
+++ b/Engine/Plugins/PCG/Source/PCG/Private/Elements/PCGAttributeFilter.cpp
@@ -339,7 +339,7 @@ TArray<FPCGPreConfiguredSettingsInfo> UPCGAttributeFilteringSettings::GetPreconf
 {
 	TArray<FPCGPreConfiguredSettingsInfo> PreconfiguredInfo;
 	PreconfiguredInfo.Emplace(0, GetDefaultNodeTitle());
-	PreconfiguredInfo.Emplace(1, LOCTEXT("PointRangeNodeTitle", "Point Filter"));
+	PreconfiguredInfo.Emplace(1, LOCTEXT("PointNodeTitle", "Point Filter"));
 
 	return PreconfiguredInfo;
 }
diff --git a/Engine/Plugins/PCG/Source/PCG/Private/Tests/Elements/PCGAttributeNoiseTest.cpp b/Engine/Plugins/PCG/Source/PCG/Private/Tests/Elements/PCGAttributeNoiseTest.cpp
--- a/Engine/Plugins/PCG/Source/PCG/Private/Tests/Elements/PCGAttributeNoiseTest.cpp
+++ b/Engine/Plugins/PCG/Source/PCG/Private/Tests/Elements/PCGAttributeNoiseTest.cpp
@@ -335,7 +335,7 @@ bool FPCGAttributeNoiseAttributeSetTest::RunTest(const FString& Parameters)
 	UTEST_NOT_NULL("Output attrbiute exists", OutAttribute);
 
 	// Taken from execution
-	const float ExpectedValues[NbElements] = { 0.382462f, 0.193192f, 0.134616f, 0.782528f, 0.084569f };
+	const float ExpectedValues[NbElements] = { 0.193192f, 0.134616f, 0.782528f, 0.084569f, 0.868019700f };
 	for (int32 i = 0; i < NbElements; ++i)
 	{
 		float Value = OutAttribute->GetValueFromItemKey(PCGMetadataEntryKey(i));
@@ -384,7 +384,7 @@ bool FPCGAttributeNoiseOutputAttributeExistingTest::RunTest(const FString& Param
 	UTEST_NOT_NULL("Output attrbiute exists", OutAttribute);
 
 	// Taken from execution
-	const double ExpectedValues[NbElements] = { 0.382462, 0.193192, 0.134616, 0.782528, 0.084569 };
+	const double ExpectedValues[NbElements] = { 0.193192, 0.134616, 0.782528, 0.084569, 0.868019700 };
 	for (int32 i = 0; i < NbElements; ++i)
 	{
 		FVector Value = OutAttribute->GetValueFromItemKey(PCGMetadataEntryKey(i));
