[PCG] Fix issue with Get User Parameters in subgraph that could fail in some cases

* Filter any UPCGUserParametersData that are coming from a parent subgraph, so we only have 1 data of that type.
* Also add an EditorOnly ensure to verify that statement.

#rnx

diff --git a/Engine/Plugins/PCG/Source/PCG/Private/Elements/PCGUserParameterGet.cpp b/Engine/Plugins/PCG/Source/PCG/Private/Elements/PCGUserParameterGet.cpp
--- a/Engine/Plugins/PCG/Source/PCG/Private/Elements/PCGUserParameterGet.cpp
+++ b/Engine/Plugins/PCG/Source/PCG/Private/Elements/PCGUserParameterGet.cpp
@@ -36,6 +36,10 @@ namespace PCGUserParameterGetSettings
 		TArray<FPCGTaggedData> UserParameterData = InContext.InputData.GetTaggedTypedInputs<UPCGUserParametersData>(PCGBaseSubgraphConstants::UserParameterTagData);
 		if (!UserParameterData.IsEmpty())
 		{
+#if WITH_EDITOR
+			// Safe guard to make sure we always have one and only one data of this type.
+			ensure(UserParameterData.Num() == 1);
+#endif // WITH_EDITOR
 			if (const UPCGUserParametersData* OverrideParametersData = CastChecked<UPCGUserParametersData>(UserParameterData[0].Data))
 			{
 				GraphInterface = OverrideParametersData->OriginalGraph;
diff --git a/Engine/Plugins/PCG/Source/PCG/Private/PCGSubgraph.cpp b/Engine/Plugins/PCG/Source/PCG/Private/PCGSubgraph.cpp
--- a/Engine/Plugins/PCG/Source/PCG/Private/PCGSubgraph.cpp
+++ b/Engine/Plugins/PCG/Source/PCG/Private/PCGSubgraph.cpp
@@ -540,12 +540,19 @@ void FPCGSubgraphElement::PrepareSubgraphData(const UPCGSubgraphSettings* Settin
 	{
 		OutputData = InputData;
 	}
+
+	// We also need to make sure we are not forwarding any UserParameterData from previous parents. So we remove all of them.
+	OutputData.TaggedData.RemoveAllSwap([](const FPCGTaggedData& OutData) { return Cast<UPCGUserParametersData>(OutData.Data); });
+
+	// Note for the future: Dynamic subgraphs are forwarding any data coming from the PreTask of their parent, so if you ever get data that should not be there,
+	// you should probably do some filtering, like the UPCGUserParametersData.
 }
 
 void FPCGSubgraphElement::PrepareSubgraphUserParameters(const UPCGSubgraphSettings* Settings, FPCGSubgraphContext* Context, FPCGDataCollection& OutputData) const
 {
 	// Also create a new data containing information about the original subgraph and the parameters override
 	// It is used mainly by the UserParameterGetElement to access the correct value.
+	// By construction, there should be one and only one of this data. (Filtering of previous data is done in PrepareSubgraphData)
 	if (const UPCGGraphInterface* SubgraphInterface = Settings->GetSubgraphInterface())
 	{
 		UPCGUserParametersData* UserParamData = NewObject<UPCGUserParametersData>();
