[PCG] Fix partitioned PCG components leaking cells when moved
* Make RuntimeGenScheduler robust against leaking delayed refresh components
* Fix inconsistent logic when updating the ACM component mapping that could leak non-RuntimeGen components and potentially crash

#rnx

diff --git a/Engine/Plugins/PCG/Source/PCG/Private/PCGActorAndComponentMapping.cpp b/Engine/Plugins/PCG/Source/PCG/Private/PCGActorAndComponentMapping.cpp
--- a/Engine/Plugins/PCG/Source/PCG/Private/PCGActorAndComponentMapping.cpp
+++ b/Engine/Plugins/PCG/Source/PCG/Private/PCGActorAndComponentMapping.cpp
@@ -716,8 +716,6 @@ void FPCGActorAndComponentMapping::UpdateMappingPCGComponentPartitionActor(UPCGC
 		return;
 	}
 
-	TSet<TObjectPtr<APCGPartitionActor>> RemovedActors;
-
 	if (const APCGWorldActor* WorldActor = PCGSubsystem->GetPCGWorldActor())
 	{
 		// Get the generation grids as a bitflag to compare against.
@@ -732,53 +730,60 @@ void FPCGActorAndComponentMapping::UpdateMappingPCGComponentPartitionActor(UPCGC
 			ValidGrids |= GridSize;
 		}
 
-		auto UpdateMapping = [this, InComponent, &Bounds, &RemovedActors, WorldActor, ValidGrids](TMap<const UPCGComponent*, TSet<TObjectPtr<APCGPartitionActor>>>& Map, FRWLock& Lock)
+		auto UpdateMapping = [this, InComponent, &Bounds, WorldActor, ValidGrids](bool bIsRuntimeGenerated)
 		{
-			FWriteScopeLock WriteLock(Lock);
-			TSet<TObjectPtr<APCGPartitionActor>>* PartitionActorsPtr = Map.Find(InComponent);
+			TMap<const UPCGComponent*, TSet<TObjectPtr<APCGPartitionActor>>>& Map = bIsRuntimeGenerated ? ComponentToRuntimeGenPartitionActorsMap : ComponentToPartitionActorsMap;
+			TSet<TObjectPtr<APCGPartitionActor>> RemovedActors;
 
-			if (!PartitionActorsPtr)
 			{
-				// Does not yet exists, add it
-				PartitionActorsPtr = &Map.Emplace(InComponent);
-				check(PartitionActorsPtr);
-			}
+				FWriteScopeLock WriteLock(bIsRuntimeGenerated ? ComponentToRuntimeGenPartitionActorsMapLock : ComponentToPartitionActorsMapLock);
+				TSet<TObjectPtr<APCGPartitionActor>>* PartitionActorsPtr = Map.Find(InComponent);
 
-			TSet<TObjectPtr<APCGPartitionActor>> NewMapping;
-			ForAllIntersectingPartitionActors(Bounds, [&NewMapping, InComponent, WorldActor, ValidGrids](APCGPartitionActor* Actor)
-			{
-				if (!Actor)
+				if (!PartitionActorsPtr)
 				{
-					return;
+					// Does not yet exists, add it
+					PartitionActorsPtr = &Map.Emplace(InComponent);
+					check(PartitionActorsPtr);
 				}
 
-				// TODO: This will need to be revisited when execution domains become a broader concept.
-				const bool bSameDomain = Actor->IsRuntimeGenerated() == InComponent->IsManagedByRuntimeGenSystem();
-
-				// Only add a graph instance to partition actors that are:
-				// * In the same execution domain as the original component.
-				// * On a valid grid for the original component.
-				if (bSameDomain && (ValidGrids & Actor->GetPCGGridSize()))
+				TSet<TObjectPtr<APCGPartitionActor>> NewMapping;
+				ForAllIntersectingPartitionActors(Bounds, [&NewMapping, InComponent, WorldActor, ValidGrids, bIsRuntimeGenerated](APCGPartitionActor* Actor)
 				{
-					Actor->AddGraphInstance(InComponent);
-					NewMapping.Add(Actor);
-				}
-			});
+					if (!Actor)
+					{
+						return;
+					}
 
-			// Find the ones that were removed
-			RemovedActors = PartitionActorsPtr->Difference(NewMapping);
+					// TODO: This will need to be revisited when execution domains become a broader concept.
+					const bool bSameDomain = Actor->IsRuntimeGenerated() == bIsRuntimeGenerated;
 
-			*PartitionActorsPtr = MoveTemp(NewMapping);
-		};
+					// Only add a graph instance to partition actors that are:
+					// * In the same execution domain as the original component.
+					// * On a valid grid for the original component.
+					if (bSameDomain && (ValidGrids & Actor->GetPCGGridSize()))
+					{
+						Actor->AddGraphInstance(InComponent);
+						NewMapping.Add(Actor);
+					}
+				});
 
-		UpdateMapping(ComponentToPartitionActorsMap, ComponentToPartitionActorsMapLock);
-		UpdateMapping(ComponentToRuntimeGenPartitionActorsMap, ComponentToRuntimeGenPartitionActorsMapLock);
-	}
+				// Find the ones that were removed
+				RemovedActors = PartitionActorsPtr->Difference(NewMapping);
 
-	// No need to be locked to do this.
-	for (APCGPartitionActor* RemovedActor : RemovedActors)
-	{
-		RemovedActor->RemoveGraphInstance(InComponent);
+				*PartitionActorsPtr = MoveTemp(NewMapping);
+			}
+
+			// No need to be locked to do this.
+			for (APCGPartitionActor* RemovedActor : RemovedActors)
+			{
+				if (RemovedActor)
+				{
+					RemovedActor->RemoveGraphInstance(InComponent);
+				}
+			}
+		};
+
+		UpdateMapping(InComponent->IsManagedByRuntimeGenSystem());
 	}
 }
 
diff --git a/Engine/Plugins/PCG/Source/PCG/Private/RuntimeGen/PCGRuntimeGenScheduler.cpp b/Engine/Plugins/PCG/Source/PCG/Private/RuntimeGen/PCGRuntimeGenScheduler.cpp
--- a/Engine/Plugins/PCG/Source/PCG/Private/RuntimeGen/PCGRuntimeGenScheduler.cpp
+++ b/Engine/Plugins/PCG/Source/PCG/Private/RuntimeGen/PCGRuntimeGenScheduler.cpp
@@ -105,9 +105,7 @@ void FPCGRuntimeGenScheduler::Tick(APCGWorldActor* InPCGWorldActor)
 		}
 	}
 
-	// Remove any generation keys that have been registered for deferred removal.
-	GeneratedComponents = GeneratedComponents.Difference(GeneratedComponentsToRemove);
-	GeneratedComponentsToRemove.Empty();
+	CleanupDelayedRefreshComponents();
 
 	// 1. Queue nearby components for generation.
 	
@@ -896,6 +894,52 @@ void FPCGRuntimeGenScheduler::CleanupComponent(const FGridGenerationKey& Generat
 	GeneratedComponents.Remove(GenerationKey);
 }
 
+void FPCGRuntimeGenScheduler::CleanupDelayedRefreshComponents()
+{
+	TRACE_CPUPROFILER_EVENT_SCOPE(FPCGRuntimeGenScheduler::CleanupDelayedRefreshComponents);
+
+	check(ActorAndComponentMapping);
+
+	// Check that each refreshed local component is still inside its original component.
+	// If it is not, it would be leaked instead of refreshed, so we should force a full cleanup.
+	for (const FGridGenerationKey& GenerationKey : GeneratedComponentsToRemove)
+	{
+		const uint32 GridSize = GenerationKey.GetGridSize();
+		const EPCGHiGenGrid Grid = PCGHiGenGrid::GridSizeToGrid(GridSize);
+
+		// The unbounded grid level will always lie inside the original component, so we can skip it.
+		if (Grid == EPCGHiGenGrid::Unbounded)
+		{
+			continue;
+		}
+
+		const UPCGComponent* OriginalComponent = GenerationKey.GetOriginalComponent();
+		const FIntVector& GridCoords = GenerationKey.GetGridCoords();
+
+		UPCGComponent* LocalComponent = OriginalComponent ? ActorAndComponentMapping->GetLocalComponent(GridSize, GridCoords, OriginalComponent, /*bRuntimeGenerated=*/true) : nullptr;
+
+		if (LocalComponent)
+		{
+			const FBox OriginalBounds = OriginalComponent->GetGridBounds();
+			const FBox LocalBounds = LocalComponent->GetGridBounds();
+
+			if (!LocalBounds.IsInsideOrOn(OriginalBounds))
+			{
+				CleanupComponent(GenerationKey, LocalComponent);
+			}
+		}
+		else
+		{
+			// If the original or local component no longer exists, just clean up.
+			CleanupComponent(GenerationKey, /*GenerationKey=*/nullptr);
+		}
+	}
+
+	// Remove any remaining generation keys that have been registered for deferred removal.
+	GeneratedComponents = GeneratedComponents.Difference(GeneratedComponentsToRemove);
+	GeneratedComponentsToRemove.Empty();
+}
+
 void FPCGRuntimeGenScheduler::RefreshComponent(UPCGComponent* InComponent, bool bRemovePartitionActors)
 {
 	if (!InComponent || !ensure(IsInGameThread()))
diff --git a/Engine/Plugins/PCG/Source/PCG/Public/RuntimeGen/PCGRuntimeGenScheduler.h b/Engine/Plugins/PCG/Source/PCG/Public/RuntimeGen/PCGRuntimeGenScheduler.h
--- a/Engine/Plugins/PCG/Source/PCG/Public/RuntimeGen/PCGRuntimeGenScheduler.h
+++ b/Engine/Plugins/PCG/Source/PCG/Public/RuntimeGen/PCGRuntimeGenScheduler.h
@@ -79,6 +79,9 @@ protected:
 	/** Cleanup a component and remove it from the GeneratedComponents set. */
 	void CleanupComponent(const FGridGenerationKey& GenerationKey, UPCGComponent* GeneratedComponent);
 
+	/** Remove components from the GeneratedComponents set that have been marked for delayed refresh. Fully cleanup any that would be leaked otherwise. */
+	void CleanupDelayedRefreshComponents();
+
 	/** Refresh a generated component. bRemovePartitionActors will also perform a full cleanup of PAs and local components. */
 	void RefreshComponent(UPCGComponent* InComponent, bool bRemovePartitionActors = false);
 	
