[PCG] Fix RuntimeGen leaving original component generated when toggling IsPartitioned

#rnx

diff --git a/Engine/Plugins/PCG/Source/PCG/Private/PCGComponent.cpp b/Engine/Plugins/PCG/Source/PCG/Private/PCGComponent.cpp
--- a/Engine/Plugins/PCG/Source/PCG/Private/PCGComponent.cpp
+++ b/Engine/Plugins/PCG/Source/PCG/Private/PCGComponent.cpp
@@ -1580,7 +1580,14 @@ void UPCGComponent::PreEditChange(FProperty* PropertyAboutToChange)
 {
 	Super::PreEditChange(PropertyAboutToChange);
 
-	if (PropertyAboutToChange && PropertyAboutToChange->GetFName() == GET_MEMBER_NAME_CHECKED(UPCGComponent, GenerationTrigger))
+	if (!PropertyAboutToChange)
+	{
+		return;
+	}
+
+	const FName PropName = PropertyAboutToChange->GetFName();
+
+	if (PropName == GET_MEMBER_NAME_CHECKED(UPCGComponent, GenerationTrigger))
 	{
 		if (IsManagedByRuntimeGenSystem())
 		{
@@ -1594,6 +1601,17 @@ void UPCGComponent::PreEditChange(FProperty* PropertyAboutToChange)
 			SetEditingMode(PreviousEditingMode, SerializedEditingMode);
 		}
 	}
+	else if (PropName == GET_MEMBER_NAME_CHECKED(UPCGComponent, bIsComponentPartitioned))
+	{
+		if (IsManagedByRuntimeGenSystem())
+		{
+			if (UPCGSubsystem* Subsystem = GetSubsystem())
+			{
+				// When toggling IsPartitioned, we should proactively flush the RuntimeGenScheduler.
+				Subsystem->RefreshRuntimeGenComponent(this, EPCGChangeType::GenerationGrid);
+			}
+		}
+	}
 }
 
 void UPCGComponent::PostEditChangeProperty(FPropertyChangedEvent& PropertyChangedEvent)
