[PCG] Move UpdateGraphNotifications() call before StartTime is captured so that call doesn't eat up in the task budget (causing increase in execution wall time)

#rnx

diff --git a/Engine/Plugins/PCG/Source/PCG/Private/Graph/PCGGraphExecutor.cpp b/Engine/Plugins/PCG/Source/PCG/Private/Graph/PCGGraphExecutor.cpp
--- a/Engine/Plugins/PCG/Source/PCG/Private/Graph/PCGGraphExecutor.cpp
+++ b/Engine/Plugins/PCG/Source/PCG/Private/Graph/PCGGraphExecutor.cpp
@@ -637,6 +637,11 @@ void FPCGGraphExecutor::Execute()
 	bool bAnyTaskEnded = false;
 	bool bHasAlreadyCheckedSleepingTasks = false;
 
+#if WITH_EDITOR
+	// Update Notifications before capturing StartTime so that this call doesn't eat up task budget
+	UpdateGenerationNotification();
+#endif
+
 	const double StartTime = FPlatformTime::Seconds();
 
 	double VarTimePerFrame = PCGGraphExecutor::CVarTimePerFrame.GetValueOnAnyThread() / 1000.0;
@@ -654,10 +659,6 @@ void FPCGGraphExecutor::Execute()
 	const bool bAllowMultiDispatch = PCGGraphExecutor::CVarGraphMultithreading.GetValueOnAnyThread();
 	const bool bDynamicTaskCulling = PCGGraphExecutor::CVarDynamicTaskCulling.GetValueOnAnyThread();
 
-#if WITH_EDITOR
-	UpdateGenerationNotification();
-#endif
-
 	while (ReadyTasks.Num() > 0 || ActiveTasks.Num() > 0 || (!bHasAlreadyCheckedSleepingTasks && SleepingTasks.Num() > 0))
 	{
 		// If we only have sleeping tasks, we will go through this loop only once. If all tasks are still sleeping after one iteration,
