[PCG][Bug] Fixed deprecation bug with the Metadata Partition node

#rnx

diff --git a/Engine/Plugins/PCG/Source/PCG/Public/Elements/Metadata/PCGMetadataPartition.h b/Engine/Plugins/PCG/Source/PCG/Public/Elements/Metadata/PCGMetadataPartition.h
--- a/Engine/Plugins/PCG/Source/PCG/Public/Elements/Metadata/PCGMetadataPartition.h
+++ b/Engine/Plugins/PCG/Source/PCG/Public/Elements/Metadata/PCGMetadataPartition.h
@@ -36,7 +36,7 @@ public:
 	// TODO: Should be overridable once array override is supported
 	/** The data will be partitioned on these selected attributes. */
 	UPROPERTY(BlueprintReadWrite, EditAnywhere, Category = Settings)
-	TArray<FPCGAttributePropertyInputSelector> PartitionAttributeSelectors;
+	TArray<FPCGAttributePropertyInputSelector> PartitionAttributeSelectors = {FPCGAttributePropertyInputSelector()};
 
 	// TODO: Should be deprecated once array override is supported
 	UPROPERTY(meta = (PCG_Overridable))
