[PCG][Bug] Fixed a bug that would allow an editable graph parameter to be edited after it was renamed or deleted, which would cause an editor crash.

* Also addresses a bug from copy and pasting getter nodes across graphs
+ If the property no longer matches, it will force a refresh
+ If the property returns (undo/redo, added back, etc) a timer will kick off for reset

#rnx

diff --git a/Engine/Plugins/PCG/Source/PCGEditor/Private/Details/PCGEditableUserParameterDetails.cpp b/Engine/Plugins/PCG/Source/PCGEditor/Private/Details/PCGEditableUserParameterDetails.cpp
--- a/Engine/Plugins/PCG/Source/PCGEditor/Private/Details/PCGEditableUserParameterDetails.cpp
+++ b/Engine/Plugins/PCG/Source/PCGEditor/Private/Details/PCGEditableUserParameterDetails.cpp
@@ -17,6 +17,8 @@
 namespace PCGEditableUserParameterDetailsConstants
 {
 	FName UserParametersCategory = TEXT("Instance");
+	// Periodic duration for a forced refresh of the details.
+	static constexpr float RefreshTime = 1.f;
 }
 
 TSharedRef<IDetailCustomization> FPCGEditableUserParameterDetails::MakeInstance()
@@ -26,6 +28,11 @@ TSharedRef<IDetailCustomization> FPCGEditableUserParameterDetails::MakeInstance(
 
 void FPCGEditableUserParameterDetails::CustomizeDetails(IDetailLayoutBuilder& DetailBuilder)
 {
+	// Reset cached values
+	ParentDetailsView = TWeakPtr<IDetailsView>();
+	CachedGraphInterface = TWeakObjectPtr<UPCGGraphInterface>();
+	CachedPropertyDesc = FPropertyBagPropertyDesc();
+
 	TArray<TSharedPtr<FStructOnScope>> CustomizedStructs;
 	DetailBuilder.GetStructsBeingCustomized(CustomizedStructs);
 
@@ -47,16 +54,36 @@ void FPCGEditableUserParameterDetails::CustomizeDetails(IDetailLayoutBuilder& De
 		{
 			if (UPCGGraphInterface* GraphInterface = Cast<UPCGGraphInterface>(NodeObject->GetOuter()))
 			{
+				// Cache the interface and view to be used to verify on tick
+				CachedGraphInterface = MakeWeakObjectPtr(GraphInterface);
+				ParentDetailsView = StaticCastWeakPtr<IDetailsView>(DetailBuilder.GetDetailsView()->AsWeak());
+
 				// It is safe, because we hook pre/post edit changes that will trigger the callbacks
 				if (FInstancedPropertyBag* UserParameters = GraphInterface->GetMutableUserParametersStruct_Unsafe())
 				{
+					const UPropertyBag* PropertyBag = UserParameters->GetPropertyBagStruct();
+					if (!IsValid(PropertyBag))
+					{
+						return;
+					}
+
+					// Cache the property description to compare later in the tick against our property
+					if (const FPropertyBagPropertyDesc* PropertyDesc = PropertyBag->FindPropertyDescByName(Settings->PropertyName))
+					{
+						CachedPropertyDesc = *PropertyDesc;
+					}
+					else
+					{
+						return;
+					}
+
 					IDetailCategoryBuilder& CategoryBuilder = DetailBuilder.EditCategory(PCGEditableUserParameterDetailsConstants::UserParametersCategory);
-					IDetailPropertyRow* DetailPropertyRow = CategoryBuilder.AddExternalStructureProperty(MakeShared<FInstancePropertyBagStructureDataProvider>(*UserParameters), Settings->PropertyName);
+					IDetailPropertyRow* DetailPropertyRow = CategoryBuilder.AddExternalStructureProperty(MakeShared<FInstancePropertyBagStructureDataProvider>(*UserParameters), CachedPropertyDesc.Name);
 					TSharedPtr<IPropertyHandle> RowPropertyHandle = DetailPropertyRow->GetPropertyHandle();
 
 					if (RowPropertyHandle.IsValid())
 					{
-						RowPropertyHandle->SetOnPropertyValuePreChange(FSimpleDelegate::CreateLambda([this, GraphInterface = MakeWeakObjectPtr(GraphInterface), PropertyName = Settings->PropertyName]()
+						RowPropertyHandle->SetOnPropertyValuePreChange(FSimpleDelegate::CreateLambda([this]()
 						{
 							check(GEditor);
 							if (GEditor->CanTransact())
@@ -64,20 +91,30 @@ void FPCGEditableUserParameterDetails::CustomizeDetails(IDetailLayoutBuilder& De
 								GEditor->BeginTransaction(LOCTEXT("EditGraphParameter", "Edit Graph Parameter"));
 							}
 
-							if (GraphInterface.IsValid())
+							// TODO: Structs and Arrays don't transact
+							if (CachedGraphInterface.IsValid())
 							{
-								GraphInterface->Modify();
+								CachedGraphInterface->Modify();
 							}
 						}));
 
-						RowPropertyHandle->SetOnPropertyValueChanged(FSimpleDelegate::CreateLambda([this, GraphInterface = MakeWeakObjectPtr(GraphInterface), PropertyName = Settings->PropertyName]()
+						RowPropertyHandle->SetOnPropertyValueChanged(FSimpleDelegate::CreateLambda([this]()
 						{
-							if (GraphInterface.IsValid())
+							check(GEditor);
+							if (CachedGraphInterface.IsValid())
 							{
-								GraphInterface->OnGraphParametersChanged(EPCGGraphParameterEvent::ValueModifiedLocally, PropertyName);
+								const FInstancedPropertyBag* CurrentUserParameters = CachedGraphInterface->GetMutableUserParametersStruct_Unsafe();
+								if (CurrentUserParameters && CurrentUserParameters->FindPropertyDescByName(CachedPropertyDesc.Name))
+								{
+									CachedGraphInterface->OnGraphParametersChanged(EPCGGraphParameterEvent::ValueModifiedLocally, CachedPropertyDesc.Name);
+								}
+								else if (GEditor->IsTransactionActive())
+								{
+									GEditor->CancelTransaction(0);
+									return;
+								}
 							}
 
-							check(GEditor);
 							if (GEditor->IsTransactionActive())
 							{
 								GEditor->EndTransaction();
@@ -90,4 +127,52 @@ void FPCGEditableUserParameterDetails::CustomizeDetails(IDetailLayoutBuilder& De
 	}
 }
 
+void FPCGEditableUserParameterDetails::Tick(float DeltaTime)
+{
+	if (!ensure(ParentDetailsView.IsValid()))
+	{
+		return;
+	}
+
+	// If nothing cached is valid, don't do anything
+	if (!CachedGraphInterface.IsValid() || !CachedPropertyDesc.CachedProperty)
+	{
+		// If the property is deleted and this details panel is alive, it will remove it and cease checking
+		// So, periodically force a refresh, in case the user issues a redo command or adds the same property back
+		TimeUntilRefresh -= DeltaTime;
+		if (TimeUntilRefresh < 0.f)
+		{
+			ParentDetailsView.Pin()->ForceRefresh();
+			TimeUntilRefresh = PCGEditableUserParameterDetailsConstants::RefreshTime;
+		}
+
+		return;
+	}
+
+	// Check the property bag for the property. If it is no longer valid, refresh the details view
+	if (const FInstancedPropertyBag* UserParameters = CachedGraphInterface->GetMutableUserParametersStruct_Unsafe())
+	{
+		const UPropertyBag* PropertyBag = UserParameters->GetPropertyBagStruct();
+		if (IsValid(PropertyBag))
+		{
+			const FPropertyBagPropertyDesc* PropertyDesc = PropertyBag->FindPropertyDescByName(CachedPropertyDesc.Name);
+			// Validate that these are the same property. If they are not, refresh
+			if (!PropertyDesc ||
+				PropertyDesc->ID != CachedPropertyDesc.ID ||
+				!PropertyDesc->CompatibleType(CachedPropertyDesc) ||
+				!PropertyDesc->CachedProperty ||
+				!CachedPropertyDesc.CachedProperty ||
+				PropertyDesc->CachedProperty->GetOffset_ForInternal() != CachedPropertyDesc.CachedProperty->GetOffset_ForInternal())
+			{
+				ParentDetailsView.Pin()->ForceRefresh();
+			}
+		}
+	}
+}
+
+TStatId FPCGEditableUserParameterDetails::GetStatId() const
+{
+	RETURN_QUICK_DECLARE_CYCLE_STAT(FPCGEditableUserParameterDetails, STATGROUP_Tickables);
+}
+
 #undef LOCTEXT_NAMESPACE
\ No newline at end of file
diff --git a/Engine/Plugins/PCG/Source/PCGEditor/Public/Details/PCGEditableUserParameterDetails.h b/Engine/Plugins/PCG/Source/PCGEditor/Public/Details/PCGEditableUserParameterDetails.h
--- a/Engine/Plugins/PCG/Source/PCGEditor/Public/Details/PCGEditableUserParameterDetails.h
+++ b/Engine/Plugins/PCG/Source/PCGEditor/Public/Details/PCGEditableUserParameterDetails.h
@@ -3,14 +3,32 @@
 #pragma once
 
 #include "IDetailCustomization.h"
+#include "PropertyBag.h"
+#include "TickableEditorObject.h"
+
+class IDetailsView;
+class UPCGGraphInterface;
 
 /**
 * A simple Details Customization that will provide a single property from the graph's FInstancePropertyBag
 */
-class FPCGEditableUserParameterDetails : public IDetailCustomization
+class FPCGEditableUserParameterDetails : public IDetailCustomization, public FTickableEditorObject
 {
 public:
 	static TSharedRef<IDetailCustomization> MakeInstance();
 
 	virtual void CustomizeDetails(IDetailLayoutBuilder& DetailBuilder) override;
+
+	//~Begin FTickableEditorObject interface
+	virtual void Tick(float DeltaTime) override;
+	virtual TStatId GetStatId() const override;
+	//~End FTickableEditorObject interface
+
+private:
+	TWeakPtr<IDetailsView> ParentDetailsView;
+	FPropertyBagPropertyDesc CachedPropertyDesc;
+	TWeakObjectPtr<UPCGGraphInterface> CachedGraphInterface;
+
+	// Time left until the next details refresh.
+	float TimeUntilRefresh = -1.f;
 };
\ No newline at end of file
