[PCG] Updated Attribute Partition unit test to respect new default value

#rnx

diff --git a/Engine/Plugins/PCG/Source/PCG/Private/Tests/Elements/PCGAttributePartitionTest.cpp b/Engine/Plugins/PCG/Source/PCG/Private/Tests/Elements/PCGAttributePartitionTest.cpp
--- a/Engine/Plugins/PCG/Source/PCG/Private/Tests/Elements/PCGAttributePartitionTest.cpp
+++ b/Engine/Plugins/PCG/Source/PCG/Private/Tests/Elements/PCGAttributePartitionTest.cpp
@@ -19,7 +19,10 @@ bool FPCGAttributePartition_Points::RunTest(const FString& Parameters)
 	UPCGMetadataPartitionSettings* Settings = PCGTestsCommon::GenerateSettings<UPCGMetadataPartitionSettings>(TestData);
 	check(Settings);
 
-	Settings->PartitionAttributeSelectors.Emplace_GetRef().SetPointProperty(EPCGPointProperties::Density);
+	// By default there should be one selector with @Last attribute
+	UTEST_EQUAL("There is one Partition Attribute Selector by default", Settings->PartitionAttributeSelectors.Num(), 1);
+
+	Settings->PartitionAttributeSelectors[0].SetPointProperty(EPCGPointProperties::Density);
 
 	UPCGPointData* InputPointData = NewObject<UPCGPointData>();
 	TArray<FPCGPoint>& Points = InputPointData->GetMutablePoints();
@@ -85,8 +88,11 @@ bool FPCGAttributePartition_Order::RunTest(const FString& Parameters)
 	UPCGMetadataPartitionSettings* Settings = PCGTestsCommon::GenerateSettings<UPCGMetadataPartitionSettings>(TestData);
 	check(Settings);
 
+	// By default there should be one selector with @Last attribute
+	UTEST_EQUAL("There is one Partition Attribute Selector by default", Settings->PartitionAttributeSelectors.Num(), 1);
+
 	const FName AttributeName = TEXT("Attr");
-	Settings->PartitionAttributeSelectors.Emplace_GetRef().SetAttributeName(AttributeName);
+	Settings->PartitionAttributeSelectors[0].SetAttributeName(AttributeName);
 
 	UPCGPointData* InputPointData = NewObject<UPCGPointData>();
 	FPCGMetadataAttribute<int>* Attribute = InputPointData->Metadata->CreateAttribute<int>(AttributeName, 0, false, false);
@@ -152,8 +158,11 @@ bool FPCGAttributePartition_AttributeSet::RunTest(const FString& Parameters)
 	UPCGMetadataPartitionSettings* Settings = PCGTestsCommon::GenerateSettings<UPCGMetadataPartitionSettings>(TestData);
 	check(Settings);
 
+	// By default there should be one selector with @Last attribute
+	UTEST_EQUAL("There is one Partition Attribute Selector by default", Settings->PartitionAttributeSelectors.Num(), 1);
+
 	const FName InputAttributeName = TEXT("Double");
-	Settings->PartitionAttributeSelectors.Emplace_GetRef().SetAttributeName(InputAttributeName);
+	Settings->PartitionAttributeSelectors[0].SetAttributeName(InputAttributeName);
 
 	UPCGParamData* InputParam = NewObject<UPCGParamData>();
 	FPCGMetadataAttribute<double>* Attribute = InputParam->Metadata->CreateAttribute<double>(InputAttributeName, 0.0, true, false);
