[PCG] Remove compilation of graphs on PostLoad and BeginDestroy, which was producing garbage

* OnNodeAdded and OnNodeRemove was triggering (a lot) of compilations on PostLoad and BeginDestroy, and the compiler was not happy.

#rnx

diff --git a/Engine/Plugins/PCG/Source/PCG/Private/PCGGraph.cpp b/Engine/Plugins/PCG/Source/PCG/Private/PCGGraph.cpp
--- a/Engine/Plugins/PCG/Source/PCG/Private/PCGGraph.cpp
+++ b/Engine/Plugins/PCG/Source/PCG/Private/PCGGraph.cpp
@@ -398,12 +398,11 @@ void UPCGGraph::PostLoad()
 		{
 			Nodes.RemoveAtSwap(i);
 		}
-		else
-		{
-			OnNodeAdded(Nodes[i]);
-		}
 	}
 
+	// Nodes is an array of TObjectPtr so we need this trick to convert to a UPCGNode*.
+	OnNodesAdded(static_cast<TArray<UPCGNode*>>(MutableView(Nodes)), /*bNotify=*/false);
+
 	if (bHasInvalidNode || PCGGraph::CVarFixInvalidEdgesOnPostLoad.GetValueOnAnyThread())
 	{
 		FixInvalidEdges();
@@ -506,10 +505,9 @@ void UPCGGraph::DeclareConstructClasses(TArray<FTopLevelAssetPath>& OutConstruct
 void UPCGGraph::BeginDestroy()
 {
 #if WITH_EDITOR
-	for (UPCGNode* Node : Nodes)
-	{
-		OnNodeRemoved(Node);
-	}
+	// Nodes is an array of TObjectPtr so we need this trick to convert to a UPCGNode*.
+	// We don't need to notify that nodes were removed when the graph dies.
+	OnNodesRemoved(static_cast<TArray<UPCGNode*>>(MutableView(Nodes)), /*bNotify=*/false);
 
 	if (OutputNode)
 	{
@@ -522,12 +520,9 @@ void UPCGGraph::BeginDestroy()
 	}
 
 	// Notify the compiler to remove this graph from its cache
-	if (GEditor)
+	if (UPCGSubsystem* PCGSubsystem = UPCGSubsystem::GetActiveEditorInstance())
 	{
-		if (UPCGSubsystem* PCGSubsystem = UPCGSubsystem::GetInstance(GEditor->GetEditorWorldContext().World()))
-		{
-			PCGSubsystem->NotifyGraphChanged(this, EPCGChangeType::Structural | EPCGChangeType::GenerationGrid);
-		}
+		PCGSubsystem->NotifyGraphChanged(this, EPCGChangeType::Structural | EPCGChangeType::GenerationGrid);
 	}
 
 #endif
@@ -648,12 +643,12 @@ UPCGNode* UPCGGraph::AddNodeCopy(UPCGSettings* InSettings, UPCGSettings*& Defaul
 	return NewNode;
 }
 
-void UPCGGraph::OnNodeAdded(UPCGNode* InNode)
+void UPCGGraph::OnNodeAdded(UPCGNode* InNode, bool bNotify)
 {
-	OnNodesAdded(MakeArrayView<UPCGNode*>(&InNode, 1));
+	OnNodesAdded(MakeArrayView<UPCGNode*>(&InNode, 1), bNotify);
 }
 
-void UPCGGraph::OnNodesAdded(TArrayView<UPCGNode*> InNodes)
+void UPCGGraph::OnNodesAdded(TArrayView<UPCGNode*> InNodes, bool bNotify)
 {
 #if WITH_EDITOR
 	EPCGChangeType ChangeType = EPCGChangeType::Structural;
@@ -671,16 +666,19 @@ void UPCGGraph::OnNodesAdded(TArrayView<UPCGNode*> InNodes)
 		}
 	}
 
-	NotifyGraphStructureChanged(ChangeType);
+	if (bNotify)
+	{
+		NotifyGraphStructureChanged(ChangeType);
+	}
 #endif
 }
 
-void UPCGGraph::OnNodeRemoved(UPCGNode* InNode)
+void UPCGGraph::OnNodeRemoved(UPCGNode* InNode, bool bNotify)
 {
-	OnNodesRemoved(MakeArrayView<UPCGNode*>(&InNode, 1));
+	OnNodesRemoved(MakeArrayView<UPCGNode*>(&InNode, 1), bNotify);
 }
 
-void UPCGGraph::OnNodesRemoved(TArrayView<UPCGNode*> InNodes)
+void UPCGGraph::OnNodesRemoved(TArrayView<UPCGNode*> InNodes, bool bNotify)
 {
 #if WITH_EDITOR
 	bool bAnyGridSizeNodes = false;
@@ -695,7 +693,10 @@ void UPCGGraph::OnNodesRemoved(TArrayView<UPCGNode*> InNodes)
 		}
 	}
 
-	NotifyGraphStructureChanged(bAnyGridSizeNodes ? (EPCGChangeType::Structural | EPCGChangeType::GenerationGrid) : EPCGChangeType::Structural);
+	if (bNotify)
+	{
+		NotifyGraphStructureChanged(bAnyGridSizeNodes ? (EPCGChangeType::Structural | EPCGChangeType::GenerationGrid) : EPCGChangeType::Structural);
+	}
 #endif
 }
 
diff --git a/Engine/Plugins/PCG/Source/PCG/Public/PCGGraph.h b/Engine/Plugins/PCG/Source/PCG/Public/PCGGraph.h
--- a/Engine/Plugins/PCG/Source/PCG/Public/PCGGraph.h
+++ b/Engine/Plugins/PCG/Source/PCG/Public/PCGGraph.h
@@ -333,10 +333,11 @@ public:
 	uint32 GetNodeGenerationGridSize(const UPCGNode* InNode, uint32 InDefaultGridSize) const;
 
 protected:
-	void OnNodeAdded(UPCGNode* InNode);
-	void OnNodesAdded(TArrayView<UPCGNode*> InNodes);
-	void OnNodeRemoved(UPCGNode* InNode);
-	void OnNodesRemoved(TArrayView<UPCGNode*> InNodes);
+	/** Internal function to react to add/remove nodes. bNotify can be set to false to not notify the world. */
+	void OnNodeAdded(UPCGNode* InNode, bool bNotify = true);
+	void OnNodesAdded(TArrayView<UPCGNode*> InNodes, bool bNotify = true);
+	void OnNodeRemoved(UPCGNode* InNode, bool bNotify = true);
+	void OnNodesRemoved(TArrayView<UPCGNode*> InNodes, bool bNotify = true);
 
 	void RemoveNodes_Internal(TArrayView<UPCGNode*> InNodes);
 	void AddNodes_Internal(TArrayView<UPCGNode*> InNodes);
