[PCG] Missing call on Subgraph placement makes it not populate the pins

#rnx

diff --git a/Engine/Plugins/PCG/Source/PCGEditor/Private/PCGEditorGraphSchemaActions.cpp b/Engine/Plugins/PCG/Source/PCGEditor/Private/PCGEditorGraphSchemaActions.cpp
--- a/Engine/Plugins/PCG/Source/PCGEditor/Private/PCGEditorGraphSchemaActions.cpp
+++ b/Engine/Plugins/PCG/Source/PCGEditor/Private/PCGEditorGraphSchemaActions.cpp
@@ -570,6 +570,8 @@ UEdGraphNode* FPCGEditorGraphSchemaAction_NewSubgraphElement::MakeGraphNode(UPCG
 		return nullptr;
 	}
 
+	NewPCGNode->UpdateAfterSettingsChangeDuringCreation();
+
 	FGraphNodeCreator<UPCGEditorGraphNode> NodeCreator(*InEditorGraph);
 	UPCGEditorGraphNode* NewNode = NodeCreator.CreateUserInvokedNode(bInSelectNewNode);
 	NewNode->Construct(NewPCGNode);
