[PCG] Make a pass on all nodes that needs to be run on the main thread

* Also very small pass on Spatial Noise to make sure noise comments are associated with the right enum value and allow the use of seed.

#rnx

diff --git a/Engine/Plugins/PCG/Source/PCG/Private/Elements/PCGSpatialNoise.cpp b/Engine/Plugins/PCG/Source/PCG/Private/Elements/PCGSpatialNoise.cpp
--- a/Engine/Plugins/PCG/Source/PCG/Private/Elements/PCGSpatialNoise.cpp
+++ b/Engine/Plugins/PCG/Source/PCG/Private/Elements/PCGSpatialNoise.cpp
@@ -23,6 +23,7 @@
 
 UPCGSpatialNoiseSettings::UPCGSpatialNoiseSettings()
 {
+	bUseSeed = true;
 	ValueTarget.SetPointProperty(EPCGPointProperties::Density);
 }
 
diff --git a/Engine/Plugins/PCG/Source/PCG/Public/Elements/IO/PCGExternalData.h b/Engine/Plugins/PCG/Source/PCG/Public/Elements/IO/PCGExternalData.h
--- a/Engine/Plugins/PCG/Source/PCG/Public/Elements/IO/PCGExternalData.h
+++ b/Engine/Plugins/PCG/Source/PCG/Public/Elements/IO/PCGExternalData.h
@@ -35,6 +35,10 @@ public:
 
 class PCG_API FPCGExternalDataElement : public IPCGElement
 {
+public:
+	// Loading needs to be done on the main thread and accessing objects outside of PCG might not be thread safe, so taking the safe approach
+	virtual bool CanExecuteOnlyOnMainThread(FPCGContext* Context) const override { return true; }
+
 protected:
 	virtual FPCGContext* CreateContext() override;
 	virtual bool PrepareDataInternal(FPCGContext* InContext) const override;
diff --git a/Engine/Plugins/PCG/Source/PCG/Public/Elements/IO/PCGLoadAssetElement.h b/Engine/Plugins/PCG/Source/PCG/Public/Elements/IO/PCGLoadAssetElement.h
--- a/Engine/Plugins/PCG/Source/PCG/Public/Elements/IO/PCGLoadAssetElement.h
+++ b/Engine/Plugins/PCG/Source/PCG/Public/Elements/IO/PCGLoadAssetElement.h
@@ -83,6 +83,10 @@ struct FPCGLoadDataAssetContext : public FPCGContext, public IPCGAsyncLoadingCon
 
 class PCG_API FPCGLoadDataAssetElement : public IPCGElementWithCustomContext<FPCGLoadDataAssetContext>
 {
+public:
+	// Loading needs to be done on the main thread and accessing objects outside of PCG might not be thread safe, so taking the safe approach
+	virtual bool CanExecuteOnlyOnMainThread(FPCGContext* Context) const override { return true; }
+
 protected:
 	virtual bool PrepareDataInternal(FPCGContext* Context) const override;
 	virtual bool ExecuteInternal(FPCGContext* Context) const override;
diff --git a/Engine/Plugins/PCG/Source/PCG/Public/Elements/PCGApplyOnActor.h b/Engine/Plugins/PCG/Source/PCG/Public/Elements/PCGApplyOnActor.h
--- a/Engine/Plugins/PCG/Source/PCG/Public/Elements/PCGApplyOnActor.h
+++ b/Engine/Plugins/PCG/Source/PCG/Public/Elements/PCGApplyOnActor.h
@@ -48,6 +48,8 @@ public:
 class FPCGApplyOnActorElement : public IPCGElement
 {
 public:
+	// Calling a function on an actor might not be threadsafe, so taking the safe approach.
+	virtual bool CanExecuteOnlyOnMainThread(FPCGContext* Context) const override { return true; }
 	virtual bool IsCacheable(const UPCGSettings* InSettings) const override { return false; }
 
 protected:
diff --git a/Engine/Plugins/PCG/Source/PCG/Public/Elements/PCGGetPropertyFromObjectPath.h b/Engine/Plugins/PCG/Source/PCG/Public/Elements/PCGGetPropertyFromObjectPath.h
--- a/Engine/Plugins/PCG/Source/PCG/Public/Elements/PCGGetPropertyFromObjectPath.h
+++ b/Engine/Plugins/PCG/Source/PCG/Public/Elements/PCGGetPropertyFromObjectPath.h
@@ -80,6 +80,10 @@ struct FPCGGetPropertyFromObjectPathContext : public FPCGContext, public IPCGAsy
 
 class FPCGGetPropertyFromObjectPathElement : public IPCGElement
 {
+public:
+	// Loading needs to be done on the main thread and accessing objects outside of PCG might not be thread safe, so taking the safe approach
+	virtual bool CanExecuteOnlyOnMainThread(FPCGContext* Context) const override { return true; }
+
 protected:
 	virtual FPCGContext* CreateContext() override;
 	virtual bool PrepareDataInternal(FPCGContext* Context) const override;
diff --git a/Engine/Plugins/PCG/Source/PCG/Public/Elements/PCGPointFromMeshElement.h b/Engine/Plugins/PCG/Source/PCG/Public/Elements/PCGPointFromMeshElement.h
--- a/Engine/Plugins/PCG/Source/PCG/Public/Elements/PCGPointFromMeshElement.h
+++ b/Engine/Plugins/PCG/Source/PCG/Public/Elements/PCGPointFromMeshElement.h
@@ -50,6 +50,10 @@ struct FPCGPointFromMeshContext : public FPCGContext, public IPCGAsyncLoadingCon
 
 class FPCGPointFromMeshElement : public IPCGElementWithCustomContext<FPCGPointFromMeshContext>
 {
+public:
+	// Loading needs to be done on the main thread and accessing objects outside of PCG might not be thread safe, so taking the safe approach
+	virtual bool CanExecuteOnlyOnMainThread(FPCGContext* Context) const override { return true; }
+
 protected:
 	virtual bool PrepareDataInternal(FPCGContext* Context) const override;
 	virtual bool ExecuteInternal(FPCGContext* Context) const override;
diff --git a/Engine/Plugins/PCG/Source/PCG/Public/Elements/PCGPrintElement.h b/Engine/Plugins/PCG/Source/PCG/Public/Elements/PCGPrintElement.h
--- a/Engine/Plugins/PCG/Source/PCG/Public/Elements/PCGPrintElement.h
+++ b/Engine/Plugins/PCG/Source/PCG/Public/Elements/PCGPrintElement.h
@@ -122,6 +122,8 @@ public:
 class FPCGPrintElement : public IPCGElement
 {
 public:
+	// Print to screen is perhaps safe to be called from outside of the main thread, but taking the safe approach here.
+	virtual bool CanExecuteOnlyOnMainThread(FPCGContext* Context) const override { return true; }
 	virtual bool IsCacheable(const UPCGSettings* InSettings) const override { return false; }
 
 protected:
diff --git a/Engine/Plugins/PCG/Source/PCG/Public/Elements/PCGSampleTexture.h b/Engine/Plugins/PCG/Source/PCG/Public/Elements/PCGSampleTexture.h
--- a/Engine/Plugins/PCG/Source/PCG/Public/Elements/PCGSampleTexture.h
+++ b/Engine/Plugins/PCG/Source/PCG/Public/Elements/PCGSampleTexture.h
@@ -57,6 +57,10 @@ public:
 
 class FPCGSampleTextureElement : public IPCGElement
 {
+public:
+	// Loading needs to be done on the main thread and accessing objects outside of PCG might not be thread safe, so taking the safe approach
+	virtual bool CanExecuteOnlyOnMainThread(FPCGContext* Context) const override { return true; }
+
 protected:
 	virtual bool ExecuteInternal(FPCGContext* Context) const override;
 };
diff --git a/Engine/Plugins/PCG/Source/PCG/Public/Elements/PCGSpatialNoise.h b/Engine/Plugins/PCG/Source/PCG/Public/Elements/PCGSpatialNoise.h
--- a/Engine/Plugins/PCG/Source/PCG/Public/Elements/PCGSpatialNoise.h
+++ b/Engine/Plugins/PCG/Source/PCG/Public/Elements/PCGSpatialNoise.h
@@ -11,19 +11,27 @@
 UENUM()
 enum class PCGSpatialNoiseMode
 {
-	Perlin2D, // your classic perlin noise
-	Caustic2D, // based on underwater fake custic rendering, gives swirly look
-	Voronoi2D, // voronoi noise, result a the distance to edge and cell ID
-	FractionalBrownian2D, // based on fractional brownian motion
-	EdgeMask2D, // used to create masks to blend out edges
+	/** Your classic perlin noise. */
+	Perlin2D,
+	/** Based on underwater fake caustic rendering, gives swirly look. */
+	Caustic2D,
+	/** Voronoi noise, result a the distance to edge and cell ID. */
+	Voronoi2D,
+	/** Based on fractional brownian motion. */
+	FractionalBrownian2D,
+	/** Used to create masks to blend out edges. */
+	EdgeMask2D,
 };
 
 UENUM()
 enum class PCGSpatialNoiseMask2DMode
 {
-	Perlin, // your classic perlin noise
-	Caustic, // based on underwater fake custic rendering, gives swirly look
-	FractionalBrownian, // based on fractional brownian motion
+	/** Your classic perlin noise. */
+	Perlin,
+	/** Based on underwater fake caustic rendering, gives swirly look. */
+	Caustic,
+	/** Based on fractional brownian motion. */
+	FractionalBrownian,
 };
 
 namespace PCGSpatialNoise
diff --git a/Engine/Plugins/PCG/Source/PCG/Public/Elements/PCGTextureSampler.h b/Engine/Plugins/PCG/Source/PCG/Public/Elements/PCGTextureSampler.h
--- a/Engine/Plugins/PCG/Source/PCG/Public/Elements/PCGTextureSampler.h
+++ b/Engine/Plugins/PCG/Source/PCG/Public/Elements/PCGTextureSampler.h
@@ -127,6 +127,8 @@ class FPCGTextureSamplerElement : public IPCGElement
 {
 public:
 	virtual void GetDependenciesCrc(const FPCGDataCollection& InInput, const UPCGSettings* InSettings, UPCGComponent* InComponent, FPCGCrc& OutCrc) const override;
+	// Loading needs to be done on the main thread and accessing objects outside of PCG might not be thread safe, so taking the safe approach
+	virtual bool CanExecuteOnlyOnMainThread(FPCGContext* Context) const override { return true; }
 
 protected:
 	virtual FPCGContext* CreateContext() override;
diff --git a/Engine/Plugins/PCG/Source/PCG/Public/Elements/PCGUserParameterGet.h b/Engine/Plugins/PCG/Source/PCG/Public/Elements/PCGUserParameterGet.h
--- a/Engine/Plugins/PCG/Source/PCG/Public/Elements/PCGUserParameterGet.h
+++ b/Engine/Plugins/PCG/Source/PCG/Public/Elements/PCGUserParameterGet.h
@@ -55,6 +55,8 @@ protected:
 class PCG_API FPCGUserParameterGetElement : public IPCGElement
 {
 public:
+	// Graph parameters access is not thread safe
+	virtual bool CanExecuteOnlyOnMainThread(FPCGContext* Context) const override { return true; }
 	virtual bool IsCacheable(const UPCGSettings* InSettings) const override { return false; }
 
 protected:
diff --git a/Engine/Plugins/PCG/Source/PCG/Public/Elements/PCGVisualizeAttribute.h b/Engine/Plugins/PCG/Source/PCG/Public/Elements/PCGVisualizeAttribute.h
--- a/Engine/Plugins/PCG/Source/PCG/Public/Elements/PCGVisualizeAttribute.h
+++ b/Engine/Plugins/PCG/Source/PCG/Public/Elements/PCGVisualizeAttribute.h
@@ -73,6 +73,8 @@ public:
 class FPCGVisualizeAttribute : public IPCGElement
 {
 public:
+	// Creating components (like the debug component for this node) is not thread safe.
+	virtual bool CanExecuteOnlyOnMainThread(FPCGContext* Context) const override { return true; }
 	virtual bool IsCacheable(const UPCGSettings* InSettings) const override { return false; }
 
 protected:
diff --git a/Engine/Plugins/PCGGeometryScriptInterop/Source/PCGGeometryScriptInterop/Public/Elements/PCGMeshSampler.h b/Engine/Plugins/PCGGeometryScriptInterop/Source/PCGGeometryScriptInterop/Public/Elements/PCGMeshSampler.h
--- a/Engine/Plugins/PCGGeometryScriptInterop/Source/PCGGeometryScriptInterop/Public/Elements/PCGMeshSampler.h
+++ b/Engine/Plugins/PCGGeometryScriptInterop/Source/PCGGeometryScriptInterop/Public/Elements/PCGMeshSampler.h
@@ -199,6 +199,10 @@ public:
 
 class FPCGMeshSamplerElement : public IPCGElement
 {
+public:
+	// Loading needs to be done on the main thread and accessing objects outside of PCG might not be thread safe, so taking the safe approach
+	virtual bool CanExecuteOnlyOnMainThread(FPCGContext* Context) const override { return true; }
+
 protected:
 	virtual FPCGContext* CreateContext() override;
 	virtual bool PrepareDataInternal(FPCGContext* InContext) const override;
