[PCG] Fix UI bug with Distance Node settings

#rnx

diff --git a/Engine/Plugins/PCG/Source/PCG/Public/Elements/PCGDistance.h b/Engine/Plugins/PCG/Source/PCG/Public/Elements/PCGDistance.h
--- a/Engine/Plugins/PCG/Source/PCG/Public/Elements/PCGDistance.h
+++ b/Engine/Plugins/PCG/Source/PCG/Public/Elements/PCGDistance.h
@@ -54,11 +54,11 @@ public:
 	// ~End UObject interface
 
 	/** Output the distance or distance vector to an attribute. */
-	UPROPERTY(BlueprintReadWrite, EditAnywhere, Category = Settings, meta = (InlineEditConditionToggle, PCG_Overridable))
+	UPROPERTY(BlueprintReadWrite, EditAnywhere, Category = Settings, meta = (PCG_Overridable))
 	bool bOutputToAttribute = true;
 
 	/** The attribute output for the resulting distance value. */
-	UPROPERTY(BlueprintReadWrite, EditAnywhere, Category = Settings, meta = (EditCondition = "bOutputToAttribute", PCG_Overridable))
+	UPROPERTY(BlueprintReadWrite, EditAnywhere, Category = Settings, meta = (EditCondition = "bOutputToAttribute", EditConditionHides, PCG_Overridable))
 	FPCGAttributePropertySelector OutputAttribute = FPCGAttributePropertySelector::CreateAttributeSelector(PCGDistanceConstants::DefaultOutputAttributeName);
 
 	/** Controls whether the attribute will be a scalar or a vector */
