*PCG Asset Exporter
Generalized PCG Asset Exporter to make it reusable with other source data
Added some implementer friendly functions to make it easier to create new exporters and reuse the same process
Added an update mechanism that's aware of its exporter & metadata.

*PCG Level to asset exporter
Improved tag parsing on level-to-asset so it's in line with GetActorData
Adjusted metadata type sto be more efficient in the level-to-asset exporter
Added some robustness around reserved tags

*PCG Load Asset element
Fixed tracking updates missing - only on save at the moment
Fixed an issue with pins & data tagging when executing with no template in the PCG Load Asset element - also present in the Proxy element.

*PCG Alembic loader
Implemented the alembic exporter as a PCG Asset Exporter
Updated the EUW to use the new code

#rnx

diff --git a/Engine/Plugins/PCG/Source/PCG/Private/Elements/IO/PCGLoadAssetElement.cpp b/Engine/Plugins/PCG/Source/PCG/Private/Elements/IO/PCGLoadAssetElement.cpp
--- a/Engine/Plugins/PCG/Source/PCG/Private/Elements/IO/PCGLoadAssetElement.cpp
+++ b/Engine/Plugins/PCG/Source/PCG/Private/Elements/IO/PCGLoadAssetElement.cpp
@@ -7,6 +7,12 @@
 
 #include "AssetRegistry/AssetData.h"
 
+UPCGLoadDataAssetSettings::UPCGLoadDataAssetSettings()
+{
+	Pins = Super::OutputPinProperties();
+	bTagOutputsBasedOnOutputPins = true;
+}
+
 #if WITH_EDITOR
 void UPCGLoadDataAssetSettings::GetStaticTrackedKeys(FPCGSelectionKeyToSettingsMap& OutKeysToSettings, TArray<TObjectPtr<const UPCGGraph>>& OutVisitedGraphs) const
 {
@@ -85,6 +91,7 @@ void UPCGLoadDataAssetSettings::UpdateFromData()
 		}
 
 		Pins = NewPins;
+		bTagOutputsBasedOnOutputPins = false;
 
 		// Update rest of cached data (name, tooltip, color, ...)
 		AssetName = AssetData->Name;
@@ -95,7 +102,9 @@ void UPCGLoadDataAssetSettings::UpdateFromData()
 	}
 	else
 	{
-		Pins = TArray<FPCGPinProperties>();
+		Pins = Super::OutputPinProperties();
+		bTagOutputsBasedOnOutputPins = true;
+
 		AssetName = FString();
 #if WITH_EDITOR
 		AssetDescription = FText::GetEmpty();
@@ -149,6 +158,17 @@ bool FPCGLoadDataAssetElement::ExecuteInternal(FPCGContext* Context) const
 	if (UPCGDataAsset* AssetData = Settings->Asset.LoadSynchronous())
 	{
 		Context->OutputData = AssetData->Data;
+
+		if (Settings->bTagOutputsBasedOnOutputPins)
+		{
+			for (FPCGTaggedData& TaggedData : Context->OutputData.TaggedData)
+			{
+				if (TaggedData.Pin != NAME_None)
+				{
+					TaggedData.Tags.Add(TaggedData.Pin.ToString());
+				}
+			}
+		}
 	}
 	else if (!Settings->Asset.IsNull() || Settings->bWarnIfNoAsset)
 	{
diff --git a/Engine/Plugins/PCG/Source/PCG/Private/Elements/PCGIndirectionElement.cpp b/Engine/Plugins/PCG/Source/PCG/Private/Elements/PCGIndirectionElement.cpp
--- a/Engine/Plugins/PCG/Source/PCG/Private/Elements/PCGIndirectionElement.cpp
+++ b/Engine/Plugins/PCG/Source/PCG/Private/Elements/PCGIndirectionElement.cpp
@@ -318,10 +318,9 @@ bool FPCGIndirectionElement::ExecuteInternal(FPCGContext* InContext) const
 	{
 		for (FPCGTaggedData& TaggedData : Context->OutputData.TaggedData)
 		{
-			if (TaggedData.Pin != NAME_None && TaggedData.Pin != PCGPinConstants::DefaultOutputLabel)
+			if (TaggedData.Pin != NAME_None)
 			{
 				TaggedData.Tags.Add(TaggedData.Pin.ToString());
-				TaggedData.Pin = NAME_None;
 			}
 		}
 	}
diff --git a/Engine/Plugins/PCG/Source/PCG/Private/PCGActorAndComponentMapping.cpp b/Engine/Plugins/PCG/Source/PCG/Private/PCGActorAndComponentMapping.cpp
--- a/Engine/Plugins/PCG/Source/PCG/Private/PCGActorAndComponentMapping.cpp
+++ b/Engine/Plugins/PCG/Source/PCG/Private/PCGActorAndComponentMapping.cpp
@@ -3,6 +3,7 @@
 #include "PCGActorAndComponentMapping.h"
 
 #include "PCGComponent.h"
+#include "PCGDataAsset.h"
 #include "PCGGraph.h"
 #include "PCGModule.h"
 #include "PCGSubsystem.h"
@@ -1509,10 +1510,10 @@ void FPCGActorAndComponentMapping::OnObjectSaved(UObject* InObject, FObjectPreSa
 		return;
 	}
 
-	// Only trigger a refresh a new user data and if it is a data table
-	// We only track data table because we probably will catch other changes with OnObjectPropertyChanged.
-	// To avoid to force the check multiple times (on change + on save)
-	if (!InObjectSaveContext.IsProceduralSave() && Cast<UDataTable>(InObject))
+	// Only trigger a refresh on save for limited data classes.
+	// At this point in time, We only track data tables and PCG data assets because in most cases we probably will catch other changes with OnObjectPropertyChanged.
+	// This is especially important to make sure we don't trigger refresh multiple times (less a problem for PCG assets, but a big problem for data tables).
+	if (!InObjectSaveContext.IsProceduralSave() && (Cast<UDataTable>(InObject) || Cast<UPCGDataAsset>(InObject)))
 	{
 		FPropertyChangedEvent Event{ nullptr };
 		OnObjectPropertyChanged(InObject, Event);
diff --git a/Engine/Plugins/PCG/Source/PCG/Private/PCGAssetExporter.cpp b/Engine/Plugins/PCG/Source/PCG/Private/PCGAssetExporter.cpp
new file mode 100644
--- /dev/null
+++ b/Engine/Plugins/PCG/Source/PCG/Private/PCGAssetExporter.cpp
@@ -0,0 +1,145 @@
+// Copyright Epic Games, Inc. All Rights Reserved.
+
+#include "PCGAssetExporter.h"
+
+#include "PCGModule.h"
+
+#include "AssetRegistry/AssetData.h"
+#include "Misc/Base64.h"
+#include "Serialization/MemoryReader.h"
+#include "Serialization/MemoryWriter.h"
+#include "UObject/UObjectGlobals.h"
+#include "UObject/Package.h"
+
+TSubclassOf<UPCGDataAsset> UPCGAssetExporter::BP_GetAssetType_Implementation() const
+{
+	return UPCGDataAsset::StaticClass();
+}
+
+TSubclassOf<UPCGDataAsset> UPCGAssetExporter::GetAssetType() const
+{
+	return BP_GetAssetType();
+}
+
+bool UPCGAssetExporter::BP_ExportToAsset_Implementation(UPCGDataAsset* Asset)
+{
+	// TODO in BP derived exporter classes
+	return false;
+}
+
+UPackage* UPCGAssetExporter::Update(const FAssetData& PCGAsset)
+{
+#if WITH_EDITOR
+	SerializeMetadataFromAsset(PCGAsset);
+	if (UPackage* Package = UpdateAsset(PCGAsset))
+	{
+		SerializeMetadataToAsset(PCGAsset);
+		Package->MarkPackageDirty();
+		return Package;
+	}
+	else
+	{
+		UE_LOG(LogPCG, Warning, TEXT("Updating the '%s' PCG asset failed."), *PCGAsset.AssetName.ToString())
+		return nullptr;
+	}
+#else
+	UE_LOG(LogPCG, Error, TEXT("PCG Asset update is not supported in non-editor builds."));
+	return nullptr;
+#endif
+}
+
+UPackage* UPCGAssetExporter::UpdateAsset(const FAssetData& PCGAsset)
+{
+	// Default implementation will load the asset and call the BP update method
+	UPCGDataAsset* LoadedAsset = Cast<UPCGDataAsset>(PCGAsset.GetAsset());
+
+	if (LoadedAsset && BP_ExportToAsset(LoadedAsset))
+	{
+		return LoadedAsset->GetPackage();
+	}
+	else
+	{
+		return nullptr;
+	}
+}
+
+bool UPCGAssetExporter::Export(const FString& PackageName, UPCGDataAsset* Asset)
+{
+#if WITH_EDITOR
+	if (Asset)
+	{
+		Asset->ExporterClass = GetClass();
+		if (ExportAsset(PackageName, Asset))
+		{
+			SerializeMetadataToAsset(Asset);
+			Asset->MarkPackageDirty();
+			return true;
+		}
+		else
+		{
+			return false;
+		}
+	}
+	else
+	{
+		return false;
+	}
+#else
+	UE_LOG(LogPCG, Error, TEXT("PCG Asset export should not be used in non-editor builds."));
+	return false;
+#endif
+}
+
+bool UPCGAssetExporter::ExportAsset(const FString& PackageName, UPCGDataAsset* Asset)
+{
+	return BP_ExportToAsset(Asset);
+}
+
+void UPCGAssetExporter::SerializeMetadataFromAsset(const FAssetData& PCGAsset)
+{
+#if WITH_EDITOR
+	const FString ExporterMetadata = PCGAsset.GetTagValueRef<FString>(GET_MEMBER_NAME_CHECKED(UPCGDataAsset, ExporterMetadata));
+	
+	TArray<uint8> PayloadData;
+	if (FBase64::Decode(ExporterMetadata, PayloadData))
+	{
+		FMemoryReader PayloadAr(PayloadData, /*bIsPersistent=*/true);
+
+		// Read custom versions
+		FCustomVersionContainer CustomVersions;
+		CustomVersions.Serialize(PayloadAr);
+		PayloadAr.SetCustomVersions(CustomVersions);
+
+		SerializeMetadata(PayloadAr);
+	}
+#endif // WITH_EDITOR
+}
+
+void UPCGAssetExporter::SerializeMetadataToAsset(const FAssetData& PCGAsset)
+{
+#if WITH_EDITOR
+	// Asset is expected to be in memory at this point, but this might trigger loading otherwise
+	UPCGDataAsset* Asset = Cast<UPCGDataAsset>(PCGAsset.GetAsset());
+
+	if (!Asset)
+	{
+		return;
+	}
+
+	TArray<uint8> PayloadData;
+	FMemoryWriter PayloadAr(PayloadData, /*bIsPersistent=*/true);
+
+	SerializeMetadata(PayloadAr);
+
+	// Serialize custom versions
+	TArray<uint8> HeaderData;
+	FMemoryWriter HeaderAr(HeaderData);
+	FCustomVersionContainer CustomVersions = PayloadAr.GetCustomVersions();
+	CustomVersions.Serialize(HeaderAr);
+
+	TArray<uint8> SerializedMetadata = MoveTemp(HeaderData);
+	SerializedMetadata.Append(PayloadData);
+
+	Asset->ExporterMetadata = FBase64::Encode(SerializedMetadata);
+#endif // WITH_EDITOR
+}
\ No newline at end of file
diff --git a/Engine/Plugins/PCG/Source/PCG/Public/Elements/IO/PCGLoadAssetElement.h b/Engine/Plugins/PCG/Source/PCG/Public/Elements/IO/PCGLoadAssetElement.h
--- a/Engine/Plugins/PCG/Source/PCG/Public/Elements/IO/PCGLoadAssetElement.h
+++ b/Engine/Plugins/PCG/Source/PCG/Public/Elements/IO/PCGLoadAssetElement.h
@@ -18,6 +18,8 @@ class PCG_API UPCGLoadDataAssetSettings : public UPCGSettings
 	GENERATED_BODY()
 
 public:
+	UPCGLoadDataAssetSettings();
+
 	//~UObject interface implementation
 #if WITH_EDITOR
 	virtual void PostEditChangeProperty(FPropertyChangedEvent& PropertyChangedEvent) override;
@@ -77,6 +79,10 @@ public:
 	/** By default, data table loading is asynchronous, can force it synchronous if needed. */
 	UPROPERTY(BlueprintReadWrite, EditAnywhere, Category = "Settings|Debug")
 	bool bSynchronousLoad = false;
+
+	/** Controls whether the data output from the loaded asset will be passed to the default pin with tags or on the proper pins. */
+	UPROPERTY(BlueprintReadOnly, VisibleAnywhere, Category = "Settings", meta = (NoResetToDefaeult))
+	bool bTagOutputsBasedOnOutputPins = true;
 };
 
 struct FPCGLoadDataAssetContext : public FPCGContext, public IPCGAsyncLoadingContext {};
diff --git a/Engine/Plugins/PCG/Source/PCG/Public/PCGAssetExporter.h b/Engine/Plugins/PCG/Source/PCG/Public/PCGAssetExporter.h
new file mode 100644
--- /dev/null
+++ b/Engine/Plugins/PCG/Source/PCG/Public/PCGAssetExporter.h
@@ -0,0 +1,71 @@
+// Copyright Epic Games, Inc. All Rights Reserved.
+
+#pragma once
+
+#include "PCGDataAsset.h"
+
+class UPackage;
+
+#include "PCGAssetExporter.generated.h"
+
+/** Common structure to hold saving options required to export or update PCG assets. */
+USTRUCT(BlueprintType)
+struct PCG_API FPCGAssetExporterParameters
+{
+	GENERATED_BODY()
+
+	/** Controls whether we will open a Save... dialog, works only when a single level is exported. Overrides update anywhere. */
+	UPROPERTY(BlueprintReadWrite, EditAnywhere, Category = Settings)
+	bool bOpenSaveDialog = true;
+
+	/** Target asset path name */
+	UPROPERTY(BlueprintReadWrite, EditAnywhere, Category = Settings)
+	FString AssetName;
+
+	/** Target asset path to write the PCG assets to. */
+	UPROPERTY(BlueprintReadWrite, EditAnywhere, Category = Settings)
+	FString AssetPath;
+
+	/** Controls whether the assets will be saved at the end of the process or not. */
+	bool bSaveOnExportEnded = true;
+};
+
+/** Base class for asset exporters. Can be extended either natively or through blueprint. Not intended to be used in non-editor builds. */
+UCLASS(BlueprintType, Abstract, Blueprintable)
+class PCG_API UPCGAssetExporter : public UObject
+{
+	GENERATED_BODY()
+
+public:
+	/** Drives the asset creation for the data parsed from the level, will return UPCGDataAsset by default. */
+	virtual TSubclassOf<UPCGDataAsset> GetAssetType() const;
+
+	/** Exports data to an asset based on the exporter's metadata */
+	bool Export(const FString& PackageName, UPCGDataAsset* Asset);
+
+	/** Updates the asset from its source data, when possible */
+	UPackage* Update(const FAssetData& PCGAsset);
+
+protected:
+	/** Returns the subtype of PCG data asset generated by the Level To Asset process. Override this in the instances where you would not use the default base class. */
+	UFUNCTION(BlueprintNativeEvent, meta = (DisplayName = "Get Asset Type", ForceAsFunction))
+	TSubclassOf<UPCGDataAsset> BP_GetAssetType() const;
+
+	UFUNCTION(BlueprintNativeEvent, meta = (DisplayName = "Update Asset", ForceAsFunction))
+	bool BP_ExportToAsset(UPCGDataAsset* Asset);
+
+	/** Loads exporter metadata from the asset */
+	void SerializeMetadataFromAsset(const FAssetData& PCGAsset);
+
+	/** Saves exporter metadata to the asset */
+	void SerializeMetadataToAsset(const FAssetData& PCGAsset);
+
+	/** Exports data to an asset based on the exporter's metadata */
+	virtual bool ExportAsset(const FString& PackageName, UPCGDataAsset* Asset);
+
+	/** Updates the asset from its source data, when possible */
+	virtual UPackage* UpdateAsset(const FAssetData& PCGAsset);
+
+	/** Sets up the exporter prior to performing update */
+	virtual void SerializeMetadata(FArchive& Ar) {}
+};
\ No newline at end of file
diff --git a/Engine/Plugins/PCG/Source/PCG/Public/PCGDataAsset.h b/Engine/Plugins/PCG/Source/PCG/Public/PCGDataAsset.h
--- a/Engine/Plugins/PCG/Source/PCG/Public/PCGDataAsset.h
+++ b/Engine/Plugins/PCG/Source/PCG/Public/PCGDataAsset.h
@@ -28,13 +28,21 @@ public:
 	FString Name;
 
 #if WITH_EDITORONLY_DATA
+	/** Custom exporter class that was used to create this PCG asset. Should derive from UPCGAssetExporter. Can be left empty. */
+	UPROPERTY(BlueprintReadOnly, VisibleAnywhere, Category = Data, AssetRegistrySearchable, meta = (NoResetToDefault))
+	TSoftClassPtr<UObject> ExporterClass;
+
+	/** Custom exporter metadata to be able to update data without user intervention */
+	UPROPERTY(AssetRegistrySearchable)
+	FString ExporterMetadata;
+
 	/** Custom class to create settings/node in the graph when dragged in the editor. Can be left empty. */
-	UPROPERTY(BlueprintReadOnly, VisibleAnywhere, Category = Data, AssetRegistrySearchable, meta = (NoResetToDefault));
+	UPROPERTY(BlueprintReadOnly, VisibleAnywhere, Category = Data, AssetRegistrySearchable, meta = (NoResetToDefault))
 	TSubclassOf<UPCGSettings> SettingsClass;
 
-	/** Reference to originating level */
+	/** Reference to originating object (often will be a level) */
 	UPROPERTY(BlueprintReadOnly, VisibleAnywhere, Category = Data, AssetRegistrySearchable, meta = (NoResetToDefault))
-	FSoftObjectPath LevelPath;
+	FSoftObjectPath ObjectPath;
 
 	/** Custom node color to use for this specific asset */
 	UPROPERTY(BlueprintReadWrite, EditAnywhere, Category = Data)
diff --git a/Engine/Plugins/PCG/Source/PCGEditor/Private/PCGAssetExporterUtils.cpp b/Engine/Plugins/PCG/Source/PCGEditor/Private/PCGAssetExporterUtils.cpp
new file mode 100644
--- /dev/null
+++ b/Engine/Plugins/PCG/Source/PCGEditor/Private/PCGAssetExporterUtils.cpp
@@ -0,0 +1,171 @@
+// Copyright Epic Games, Inc. All Rights Reserved.
+
+#include "PCGAssetExporterUtils.h"
+
+#include "PCGEditorModule.h"
+
+#include "ContentBrowserModule.h"
+#include "FileHelpers.h"
+#include "IContentBrowserSingleton.h"
+#include "AssetRegistry/AssetRegistryModule.h"
+
+UPackage* UPCGAssetExporterUtils::CreateAsset(UPCGAssetExporter* Exporter, FPCGAssetExporterParameters Parameters)
+{
+	if (!Exporter || !Exporter->GetAssetType())
+	{
+		UE_LOG(LogPCGEditor, Error, TEXT("Unable to create asset without an exporter, or exporter is not setup properly."));
+		return nullptr;
+	}
+
+	const bool bExporterNeedsRooting = !Exporter->IsRooted();
+	if (bExporterNeedsRooting)
+	{
+		Exporter->AddToRoot();
+	}
+
+	ON_SCOPE_EXIT
+	{
+		if (bExporterNeedsRooting)
+		{
+			Exporter->RemoveFromRoot();
+		}
+	};
+
+	FString AssetName = Parameters.AssetName;
+	FString AssetPath = Parameters.AssetPath;
+	FString PackageName = FPaths::Combine(AssetPath, AssetName);
+
+	if (Parameters.bOpenSaveDialog)
+	{
+		FSaveAssetDialogConfig SaveAssetDialogConfig;
+		SaveAssetDialogConfig.DefaultPath = AssetPath;
+		SaveAssetDialogConfig.DefaultAssetName = AssetName;
+		SaveAssetDialogConfig.AssetClassNames.Add(Exporter->GetAssetType()->GetClassPathName());
+		SaveAssetDialogConfig.ExistingAssetPolicy = ESaveAssetDialogExistingAssetPolicy::AllowButWarn;
+		SaveAssetDialogConfig.DialogTitleOverride = NSLOCTEXT("PCGAssetExporter", "SaveAssetToFileDialogTitle", "Save PCG Asset");
+
+		FContentBrowserModule& ContentBrowserModule = FModuleManager::LoadModuleChecked<FContentBrowserModule>("ContentBrowser");
+		FString SaveObjectPath = ContentBrowserModule.Get().CreateModalSaveAssetDialog(SaveAssetDialogConfig);
+		if (!SaveObjectPath.IsEmpty())
+		{
+			AssetName = FPackageName::ObjectPathToObjectName(SaveObjectPath);
+			AssetPath = FString(); // not going to be reused
+			PackageName = FPackageName::ObjectPathToPackageName(SaveObjectPath);
+		}
+		else
+		{
+			return nullptr;
+		}
+	}
+
+	UPackage* Package = FPackageName::DoesPackageExist(PackageName) ? LoadPackage(nullptr, *PackageName, LOAD_None) : nullptr;
+
+	UPCGDataAsset* Asset = nullptr;
+	bool NewAssetCreated = false;
+
+	if (Package)
+	{
+		UObject* Object = FindObjectFast<UObject>(Package, *AssetName);
+		if (Object && Object->GetClass() != Exporter->GetAssetType())
+		{
+			Object->SetFlags(RF_Transient);
+			Object->Rename(nullptr, GetTransientPackage(), REN_DontCreateRedirectors | REN_NonTransactional);
+			NewAssetCreated = true;
+		}
+		else
+		{
+			Asset = Cast<UPCGDataAsset>(Object);
+		}
+	}
+	else
+	{
+		Package = CreatePackage(*PackageName);
+		NewAssetCreated = true;
+	}
+
+	if (!Asset)
+	{
+		const EObjectFlags Flags = RF_Public | RF_Standalone | RF_Transactional;
+		Asset = NewObject<UPCGDataAsset>(Package, Exporter->GetAssetType(), FName(*AssetName), Flags);
+	}
+
+	if (Asset)
+	{
+		if (NewAssetCreated)
+		{
+			// Notify the asset registry
+			FAssetRegistryModule::AssetCreated(Asset);
+		}
+
+		if (Exporter->Export(PackageName, Asset))
+		{
+			// Make sure everybody knows we changed the asset
+			if (!NewAssetCreated)
+			{
+				FCoreUObjectDelegates::BroadcastOnObjectModified(Asset);
+			}
+		}
+		else
+		{
+			return nullptr;
+		}
+	}
+
+	// Save the file
+	if (Package && Parameters.bSaveOnExportEnded)
+	{
+		FEditorFileUtils::PromptForCheckoutAndSave({ Package }, /*bCheckDirty=*/false, /*bPromptToSave=*/false);
+	}
+
+	return Package;
+}
+
+void UPCGAssetExporterUtils::UpdateAssets(const TArray<FAssetData>& PCGAssets, FPCGAssetExporterParameters InParameters)
+{
+	TArray<UPackage*> PackagesToSave;
+	FPCGAssetExporterParameters Parameters = InParameters;
+
+	// Never open a dialog when updating pre-existing assets
+	Parameters.bOpenSaveDialog = false;
+
+	for (const FAssetData& PCGAsset : PCGAssets)
+	{
+		TSubclassOf<UPCGAssetExporter> ExporterSubclass;
+
+		const FString ExporterClassString = PCGAsset.GetTagValueRef<FString>(GET_MEMBER_NAME_CHECKED(UPCGDataAsset, ExporterClass));
+		TSoftClassPtr<> ExporterClass(ExporterClassString);
+		if (ExporterClass.IsValid() && ExporterClass->IsChildOf<UPCGAssetExporter>())
+		{
+			ExporterSubclass = ExporterClass.Get();
+		}
+		else
+		{
+			UE_LOG(LogPCGEditor, Error, TEXT("Unable to update asset '%s' because exporter isn't valid."), *PCGAsset.AssetName.ToString());
+			continue;
+		}
+
+		UPCGAssetExporter* Exporter = NewObject<UPCGAssetExporter>(GetTransientPackage(), ExporterSubclass);
+
+		if (!Exporter)
+		{
+			UE_LOG(LogPCGEditor, Error, TEXT("Unable to create exporter for asset '%s' during update process."), *PCGAsset.AssetName.ToString());
+			continue;
+		}
+
+		Exporter->AddToRoot();
+
+		if (UPackage* Package = Exporter->Update(PCGAsset))
+		{
+			FCoreUObjectDelegates::BroadcastOnObjectModified(PCGAsset.GetAsset());
+			PackagesToSave.Add(Package);
+		}
+
+		Exporter->RemoveFromRoot();
+	}
+
+	// Save the file(s)
+	if (!PackagesToSave.IsEmpty() && Parameters.bSaveOnExportEnded)
+	{
+		FEditorFileUtils::PromptForCheckoutAndSave(PackagesToSave, /*bCheckDirty=*/false, /*bPromptToSave=*/false);
+	}
+}
\ No newline at end of file
diff --git a/Engine/Plugins/PCG/Source/PCGEditor/Private/PCGEditorMenuUtils.cpp b/Engine/Plugins/PCG/Source/PCGEditor/Private/PCGEditorMenuUtils.cpp
--- a/Engine/Plugins/PCG/Source/PCGEditor/Private/PCGEditorMenuUtils.cpp
+++ b/Engine/Plugins/PCG/Source/PCGEditor/Private/PCGEditorMenuUtils.cpp
@@ -2,6 +2,7 @@
 
 #include "PCGEditorMenuUtils.h"
 
+#include "PCGAssetExporterUtils.h"
 #include "PCGLevelToAsset.h"
 
 #include "Engine/World.h"
@@ -24,37 +25,65 @@ namespace PCGEditorMenuUtils
 		return Section;
 	}
 
-	void CreateOrUpdatePCGAssetFromMenu(UToolMenu* Menu, TArray<FAssetData>& InWorldAssets)
+	void CreateOrUpdatePCGAssetFromMenu(UToolMenu* Menu, TArray<FAssetData>& InAssets)
 	{
 		TArray<FAssetData> TempWorldAssets;
-		for (const FAssetData& Asset : InWorldAssets)
+		TArray<FAssetData> TempPCGAssets;
+
+		for (const FAssetData& Asset : InAssets)
 		{
 			if (Asset.IsInstanceOf<UWorld>())
 			{
 				TempWorldAssets.Add(Asset);
 			}
+
+			if (Asset.IsInstanceOf<UPCGDataAsset>())
+			{
+				TempPCGAssets.Add(Asset);
+			}
 		}
 
-		if (TempWorldAssets.IsEmpty())
+		if (TempWorldAssets.IsEmpty() && TempPCGAssets.IsEmpty())
 		{
 			return;
 		}
 
 		FToolMenuSection& Section = CreatePCGSection(Menu);
 
-		FToolUIAction UIAction;
-		UIAction.ExecuteAction.BindLambda([WorldAssets = MoveTemp(TempWorldAssets)](const FToolMenuContext& MenuContext)
+		if (!TempWorldAssets.IsEmpty())
 		{
-			FScopedSlowTask SlowTask(0.0f, LOCTEXT("CreateOrUpdatePCGAssetsInProgress", "Creating PCG Assets..."));
-			UPCGLevelToAsset::CreateOrUpdatePCGAssets(WorldAssets);
-		});
-
-		Section.AddMenuEntry(
-			"CreateOrUpdatePCGAssetFromMenu",
-			LOCTEXT("CreateOrUpdatePCGAssetFromMenu", "Create or Update PCG Asset from Level(s)"),
-			TAttribute<FText>(),
-			FSlateIcon(),
-			UIAction);
+			FToolUIAction UIAction;
+			UIAction.ExecuteAction.BindLambda([WorldAssets = MoveTemp(TempWorldAssets)](const FToolMenuContext& MenuContext)
+			{
+				FScopedSlowTask SlowTask(0.0f, LOCTEXT("CreateOrUpdatePCGAssetsInProgress", "Creating PCG Assets from Level(s)..."));
+				UPCGLevelToAsset::CreateOrUpdatePCGAssets(WorldAssets);
+			});
+
+			Section.AddMenuEntry(
+				"CreatePCGAssetFromMenu",
+				LOCTEXT("CreatePCGAssetFromMenu", "Create PCG Assets from Level(s)"),
+				TAttribute<FText>(),
+				FSlateIcon(),
+				UIAction);
+		}
+
+		if (!TempPCGAssets.IsEmpty())
+		{
+			FToolUIAction UIAction;
+			UIAction.ExecuteAction.BindLambda([PCGAssets = MoveTemp(TempPCGAssets)](const FToolMenuContext& MenuContext)
+			{
+				FScopedSlowTask SlowTask(0.0f, LOCTEXT("UpdatePCGAssetsInProgress", "Updating PCG Assets..."));
+				UPCGAssetExporterUtils::UpdateAssets(PCGAssets);
+			});
+
+			Section.AddMenuEntry(
+				"UpdatePCGAssetFromMenu",
+				LOCTEXT("UpdatePCGAssetFromMenu", "Update PCG Assets"),
+				TAttribute<FText>(),
+				FSlateIcon(),
+				UIAction);
+		}
+
 	}
 }
 
diff --git a/Engine/Plugins/PCG/Source/PCGEditor/Private/PCGLevelToAsset.cpp b/Engine/Plugins/PCG/Source/PCGEditor/Private/PCGLevelToAsset.cpp
--- a/Engine/Plugins/PCG/Source/PCGEditor/Private/PCGLevelToAsset.cpp
+++ b/Engine/Plugins/PCG/Source/PCGEditor/Private/PCGLevelToAsset.cpp
@@ -3,47 +3,55 @@
 #include "PCGLevelToAsset.h"
 
 #include "PCGEditorModule.h"
+#include "PCGEditorUtils.h"
+#include "PCGAssetExporterUtils.h"
 
 #include "Data/PCGPointData.h"
 #include "Helpers/PCGActorHelpers.h"
 #include "Helpers/PCGHelpers.h"
 #include "Metadata/PCGMetadata.h"
 
+#include "ContentBrowserModule.h"
 #include "FileHelpers.h"
+#include "IContentBrowserSingleton.h"
 #include "AssetRegistry/AssetRegistryModule.h"
 #include "Components/InstancedStaticMeshComponent.h"
 #include "UObject/UObjectGlobals.h"
 #include "UObject/Package.h"
 
-void UPCGLevelToAsset::CreateOrUpdatePCGAssets(const TArray<FAssetData>& WorldAssets, TSubclassOf<UPCGLevelToAsset> ExporterSubclass)
+void UPCGLevelToAsset::CreateOrUpdatePCGAssets(const TArray<FAssetData>& WorldAssets, const FPCGAssetExporterParameters& InParameters, TSubclassOf<UPCGLevelToAsset> ExporterSubclass)
 {
 	TArray<UPackage*> PackagesToSave;
+	FPCGAssetExporterParameters Parameters = InParameters;
+
+	if (WorldAssets.Num() > 1)
+	{
+		Parameters.bOpenSaveDialog = false;
+	}
 
 	for (const FAssetData& WorldAsset : WorldAssets)
 	{
-		if (UPackage* Package = CreateOrUpdatePCGAsset(TSoftObjectPtr<UWorld>(WorldAsset.GetSoftObjectPath()), ExporterSubclass))
+		if (UPackage* Package = CreateOrUpdatePCGAsset(TSoftObjectPtr<UWorld>(WorldAsset.GetSoftObjectPath()), Parameters, ExporterSubclass))
 		{
 			PackagesToSave.Add(Package);
-			// TODO: consider if we should garbage collect?
 		}
 	}
 
 	// Save the file(s)
-	// TODO: check if we should just dirty and not save (could be an option here)
-	if (!PackagesToSave.IsEmpty())
+	if (!PackagesToSave.IsEmpty() && Parameters.bSaveOnExportEnded)
 	{
 		FEditorFileUtils::PromptForCheckoutAndSave(PackagesToSave, false, false);
 	}
 }
 
-UPackage* UPCGLevelToAsset::CreateOrUpdatePCGAsset(TSoftObjectPtr<UWorld> WorldPath, TSubclassOf<UPCGLevelToAsset> ExporterSubclass)
+UPackage* UPCGLevelToAsset::CreateOrUpdatePCGAsset(TSoftObjectPtr<UWorld> WorldPath, const FPCGAssetExporterParameters& Parameters, TSubclassOf<UPCGLevelToAsset> ExporterSubclass)
 {
-	return CreateOrUpdatePCGAsset(WorldPath.LoadSynchronous(), ExporterSubclass);
+	return CreateOrUpdatePCGAsset(WorldPath.LoadSynchronous(), Parameters, ExporterSubclass);
 }
 
-UPackage* UPCGLevelToAsset::CreateOrUpdatePCGAsset(UWorld* Level, TSubclassOf<UPCGLevelToAsset> ExporterSubclass)
+UPackage* UPCGLevelToAsset::CreateOrUpdatePCGAsset(UWorld* World, const FPCGAssetExporterParameters& InParameters, TSubclassOf<UPCGLevelToAsset> ExporterSubclass)
 {
-	if (!Level)
+	if (!World)
 	{
 		return nullptr;
 	}
@@ -64,93 +72,77 @@ UPackage* UPCGLevelToAsset::CreateOrUpdatePCGAsset(UWorld* Level, TSubclassOf<UP
 		return nullptr;
 	}
 
-	const FString AssetName = Level->GetName() + TEXT("_PCG");
-	// TODO: since we store the level SOP in the PCG data assets, we can technically search if there is already an asset matching this level and select this one instead.
-	const FString PackageName = FPaths::Combine(FPackageName::GetLongPackagePath(Level->GetPackage()->GetName()), AssetName);
-
-	// Implementation note: this will cause a warning if the file can't be loaded (if the file doesn't exist, for example). It doesn't seem possible to quiet this.
-	UPackage* Package = LoadPackage(nullptr, *PackageName, LOAD_None);
+	Exporter->WorldToExport = World;
 
-	UPCGDataAsset* Asset = nullptr;
-	bool NewAssetCreated = false;
+	FPCGAssetExporterParameters Parameters = InParameters;
+	Parameters.AssetName = World->GetName() + TEXT("_PCG");
 
-	if (Package)
+	if (InParameters.AssetPath.IsEmpty() && World->GetPackage())
 	{
-		UObject* Object = FindObjectFast<UObject>(Package, *AssetName);
-		if (Object && Object->GetClass() != Exporter->GetAssetType())
-		{
-			Object->SetFlags(RF_Transient);
-			Object->Rename(nullptr, GetTransientPackage(), REN_DontCreateRedirectors | REN_NonTransactional);
-			NewAssetCreated = true;
-		}
-		else
-		{
-			Asset = Cast<UPCGDataAsset>(Object);
-		}
-	}
-	else
-	{
-		Package = CreatePackage(*PackageName);
-		NewAssetCreated = true;
+		Parameters.AssetPath = FPackageName::GetLongPackagePath(World->GetPackage()->GetName());
 	}
 
+	return UPCGAssetExporterUtils::CreateAsset(Exporter, Parameters);
+}
+
+UPackage* UPCGLevelToAsset::UpdateAsset(const FAssetData& PCGAsset)
+{
+	UPCGDataAsset* Asset = Cast<UPCGDataAsset>(PCGAsset.GetAsset());
 	if (!Asset)
 	{
-		const EObjectFlags Flags = RF_Public | RF_Standalone | RF_Transactional;
-		Asset = NewObject<UPCGDataAsset>(Package, Exporter->GetAssetType(), FName(*AssetName), Flags);
+		UE_LOG(LogPCG, Error, TEXT("Asset '%s' isn't a PCG data asset or could not be properly loaded."), *PCGAsset.GetObjectPathString());
+		return nullptr;
 	}
 
-	if (Asset)
+	UPackage* Package = Asset->GetPackage();
+	if (!Package)
 	{
-		if (NewAssetCreated)
-		{
-			// Notify the asset registry
-			FAssetRegistryModule::AssetCreated(Asset);
-		}
-
-		Exporter->ExportLevel(Level, PackageName, Asset);
+		UE_LOG(LogPCG, Error, TEXT("Unable to retrieve package from Asset '%s'."), *PCGAsset.GetObjectPathString());
+		return nullptr;
+	}
 
-		// Mark the package dirty...
-		Package->MarkPackageDirty();
+	TSoftObjectPtr<UWorld> WorldPtr(Asset->ObjectPath);
+	UWorld* World = WorldPtr.LoadSynchronous();
 
-		// Make sure everybody knows we changed these settings.
-		if (!NewAssetCreated)
-		{
-			FCoreUObjectDelegates::BroadcastOnObjectModified(Asset);
-		}
+	if (!World)
+	{
+		UE_LOG(LogPCGEditor, Error, TEXT("PCG asset was unable to load world '%s'."), *Asset->ObjectPath.ToString());
+		return nullptr;
 	}
 
-	return Package;
-}
+	WorldToExport = World;
 
-TSubclassOf<UPCGDataAsset> UPCGLevelToAsset::BP_GetAssetType_Implementation() const
-{
-	return UPCGDataAsset::StaticClass();
-}
-
-TSubclassOf<UPCGDataAsset> UPCGLevelToAsset::GetAssetType() const
-{
-	return BP_GetAssetType();
+	if (ExportAsset(Package->GetPathName(), Asset))
+	{
+		FCoreUObjectDelegates::BroadcastOnObjectModified(Asset);
+		return Package;
+	}
+	else
+	{
+		return nullptr;
+	}
 }
 
-bool UPCGLevelToAsset::ExportLevel(UWorld* Level, const FString& PackageName, UPCGDataAsset* Asset)
+bool UPCGLevelToAsset::ExportAsset(const FString& PackageName, UPCGDataAsset* Asset)
 {
-	return BP_ExportLevel(Level, PackageName, Asset);
+	return BP_ExportWorld(WorldToExport, PackageName, Asset);
 }
 
-bool UPCGLevelToAsset::BP_ExportLevel_Implementation(UWorld* Level, const FString& PackageName, UPCGDataAsset* Asset)
+bool UPCGLevelToAsset::BP_ExportWorld_Implementation(UWorld* World, const FString& PackageName, UPCGDataAsset* Asset)
 {
-	check(Level && Asset);
-	Asset->LevelPath = FSoftObjectPath(Level);
-	Asset->Description = FText::Format(NSLOCTEXT("PCGLevelToAsset", "DefaultDescriptionOnExportedLevel", "Generated from level: {0}"), FText::FromString(Level->GetName()));
+	check(World && Asset);
+	Asset->ObjectPath = FSoftObjectPath(World);
+	Asset->Description = FText::Format(NSLOCTEXT("PCGLevelToAsset", "DefaultDescriptionOnExportedLevel", "Generated from world: {0}"), FText::FromString(World->GetName()));
+	Asset->ExporterClass = GetClass();
 
 	FPCGDataCollection& DataCollection = Asset->Data;
+	DataCollection.TaggedData.Reset();
 
 	// Create Root Data
 	UPCGPointData* RootPointData = NewObject<UPCGPointData>(Asset);
 	UPCGMetadata* RootMetadata = RootPointData->MutableMetadata();
 	TArray<FPCGPoint>& Roots = RootPointData->GetMutablePoints();
-	RootMetadata->CreateAttribute<FString>(TEXT("Name"), Level->GetName(), /*bAllowsInterpolation=*/false, /*bOverrideParent=*/true);
+	RootMetadata->CreateAttribute<FString>(TEXT("Name"), World->GetName(), /*bAllowsInterpolation=*/false, /*bOverrideParent=*/true);
 	RootMetadata->CreateAttribute<FSoftObjectPath>(TEXT("Source"), PackageName, /*bAllowsInterpolation=*/false, /*bOverrideParent=*/true);
 
 	// Add to data collection
@@ -175,15 +167,30 @@ bool UPCGLevelToAsset::BP_ExportLevel_Implementation(UWorld* Level, const FStrin
 	// Common data shared across steps
 	FBox AllActorBounds(EForceInit::ForceInit);
 
-	// Attribute setup on the points
-	FPCGMetadataAttribute<FSoftObjectPath>* MaterialAttribute = PointMetadata->CreateAttribute<FSoftObjectPath>(TEXT("Material"), FSoftObjectPath(), /*bAllowsInterpolation=*/false, /*bOverrideParent=*/true);
-	FPCGMetadataAttribute<FSoftObjectPath>* MeshAttribute = PointMetadata->CreateAttribute<FSoftObjectPath>(TEXT("Mesh"), FSoftObjectPath(), /*bAllowsInterpolation=*/false, /*bOverrideParent=*/true);
-	FPCGMetadataAttribute<int64>* HierarchyDepthAttribute = PointMetadata->CreateAttribute<int64>(TEXT("HierarchyDepth"), 0, /*bAllowsInterpolation=*/false, /*bOverrideParent=*/true);
-	FPCGMetadataAttribute<int64>* ActorIndexAttribute = PointMetadata->CreateAttribute<int64>(TEXT("ActorIndex"), -1, /*bAllowsInterpolation=*/false, /*bOverrideParent=*/true);
-	FPCGMetadataAttribute<int64>* ParentIndexAttribute = PointMetadata->CreateAttribute<int64>(TEXT("ParentIndex"), -1, /*bAllowsInterpolation=*/false, /*bOverrideParent=*/true);
-	FPCGMetadataAttribute<FTransform>* RelativeTransformAttribute = PointMetadata->CreateAttribute<FTransform>(TEXT("RelativeTransform"), FTransform::Identity, /*bAllowsInterpolation=*/false, /*bOverrideParent=*/true);
+	// Hardcoded attributes
+	const FName MaterialAttributeName = TEXT("Material");
+	const FName MeshAttributeName = TEXT("Mesh");
+	const FName HierarchyDepthAttributeName = TEXT("HierarchyDepth");
+	const FName ActorIndexAttributeName = TEXT("ActorIndex");
+	const FName ParentIndexAttributeName = TEXT("ParentIndex");
+	const FName RelativeTransformAttributeName = TEXT("RelativeTransform");
 
-	TMap<FName, FPCGMetadataAttribute<int64>*> TagToAttributeMap;
+	// Attribute setup on the points
+	FPCGMetadataAttribute<FSoftObjectPath>* MaterialAttribute = PointMetadata->CreateAttribute<FSoftObjectPath>(MaterialAttributeName, FSoftObjectPath(), /*bAllowsInterpolation=*/false, /*bOverrideParent=*/true);
+	FPCGMetadataAttribute<FSoftObjectPath>* MeshAttribute = PointMetadata->CreateAttribute<FSoftObjectPath>(MeshAttributeName, FSoftObjectPath(), /*bAllowsInterpolation=*/false, /*bOverrideParent=*/true);
+	FPCGMetadataAttribute<int64>* HierarchyDepthAttribute = PointMetadata->CreateAttribute<int64>(HierarchyDepthAttributeName, 0, /*bAllowsInterpolation=*/false, /*bOverrideParent=*/true);
+	FPCGMetadataAttribute<int64>* ActorIndexAttribute = PointMetadata->CreateAttribute<int64>(ActorIndexAttributeName, -1, /*bAllowsInterpolation=*/false, /*bOverrideParent=*/true);
+	FPCGMetadataAttribute<int64>* ParentIndexAttribute = PointMetadata->CreateAttribute<int64>(ParentIndexAttributeName, -1, /*bAllowsInterpolation=*/false, /*bOverrideParent=*/true);
+	FPCGMetadataAttribute<FTransform>* RelativeTransformAttribute = PointMetadata->CreateAttribute<FTransform>(RelativeTransformAttributeName, FTransform::Identity, /*bAllowsInterpolation=*/false, /*bOverrideParent=*/true);
+
+	TMap<FName, FPCGMetadataAttributeBase*> TagToAttributeMap;
+	TSet<FName> ReservedTags;
+	ReservedTags.Add(MaterialAttributeName);
+	ReservedTags.Add(MeshAttributeName);
+	ReservedTags.Add(HierarchyDepthAttributeName);
+	ReservedTags.Add(ActorIndexAttributeName);
+	ReservedTags.Add(ParentIndexAttributeName);
+	ReservedTags.Add(RelativeTransformAttributeName);
 
 	// Hierarchy root point
 	{
@@ -198,14 +205,14 @@ bool UPCGLevelToAsset::BP_ExportLevel_Implementation(UWorld* Level, const FStrin
 	// Build actor-index map
 	TMap<AActor*, int> ActorIndexMap;
 	int LastActorIndex = 1; // Since the root is the "first" point we'll have, we'll have the map start from 1.
-	UPCGActorHelpers::ForEachActorInWorld(Level, AActor::StaticClass(), [&ActorIndexMap, &LastActorIndex](AActor* Actor)
+	UPCGActorHelpers::ForEachActorInWorld(World, AActor::StaticClass(), [&ActorIndexMap, &LastActorIndex](AActor* Actor)
 	{
 		ActorIndexMap.Add(Actor, LastActorIndex++);
 		return true;
 	});
 
 	// Create points
-	UPCGActorHelpers::ForEachActorInWorld(Level, AActor::StaticClass(), [&](AActor* Actor)
+	UPCGActorHelpers::ForEachActorInWorld(World, AActor::StaticClass(), [&](AActor* Actor)
 	{
 		// TODO Actor-level decisions if any; if the actor is "consumed" at this step, make sure to update AllActorBounds as well.
 
@@ -223,9 +230,44 @@ bool UPCGLevelToAsset::BP_ExportLevel_Implementation(UWorld* Level, const FStrin
 
 		for (FName ActorTag : Actor->Tags)
 		{
-			if (!TagToAttributeMap.Contains(ActorTag))
+			if (ReservedTags.Contains(ActorTag) || TagToAttributeMap.Contains(ActorTag))
+			{
+				continue;
+			}
+
+			const FString TagString = ActorTag.ToString();
+			int32 DividerPosition = INDEX_NONE;
+			
+			if (TagString.FindChar(':', DividerPosition))
+			{
+				const FString LeftSide = TagString.Left(DividerPosition);
+				const FString RightSide = TagString.RightChop(DividerPosition + 1);
+
+				if (LeftSide.IsEmpty() || RightSide.IsEmpty())
+				{
+					continue;
+				}
+
+				const FName AttributeName = FName(LeftSide);
+				if (ReservedTags.Contains(AttributeName) || TagToAttributeMap.Contains(AttributeName))
+				{
+					continue;
+				}
+
+				// Otherwise, create the attribute based on the type of the data after the colon.
+				if (RightSide.IsNumeric())
+				{
+					TagToAttributeMap.Add(AttributeName, PointMetadata->CreateAttribute<double>(AttributeName, 0.0f, /*bAllowsInterpolation=*/false, /*bOverrideParent=*/true));
+				}
+				else
+				{
+					TagToAttributeMap.Add(AttributeName, PointMetadata->CreateAttribute<FString>(AttributeName, FString(), /*bAllowsInterpolation=*/false, /*bOverrideParent=*/true));
+				}
+			}
+			else
 			{
-				TagToAttributeMap.Add(ActorTag, PointMetadata->CreateAttribute<int64>(ActorTag, 0, /*bAllowsInterpolation=*/false, /*bOverrideParent=*/true));
+				// Simple boolean attribute
+				TagToAttributeMap.Add(ActorTag, PointMetadata->CreateAttribute<bool>(ActorTag, false, /*bAllowsInterpolation=*/false, /*bOverrideParent=*/true));
 			}
 		}
 
@@ -264,10 +306,44 @@ bool UPCGLevelToAsset::BP_ExportLevel_Implementation(UWorld* Level, const FStrin
 			RelativeTransformAttribute->SetValue(Point.MetadataEntry, RelativeTransform);
 			HierarchyDepthAttribute->SetValue(Point.MetadataEntry, HierarchyDepth);
 
-			// For all tags, set attribute value to 1.
+			// For all tags, set attribute value to true.
 			for (FName ActorTag : Actor->Tags)
 			{
-				TagToAttributeMap[ActorTag]->SetValue(Point.MetadataEntry, 1);
+				const FString TagString = ActorTag.ToString();
+				int32 DividerPosition = INDEX_NONE;
+
+				if (TagString.FindChar(':', DividerPosition))
+				{
+					const FString LeftSide = TagString.Left(DividerPosition);
+					const FString RightSide = TagString.Right(DividerPosition + 1);
+
+					if (LeftSide.IsEmpty() || RightSide.IsEmpty())
+					{
+						continue;
+					}
+
+					const FName AttributeName = FName(LeftSide);
+					if (FPCGMetadataAttributeBase** Attribute = TagToAttributeMap.Find(AttributeName))
+					{
+						check(*Attribute);
+						if (RightSide.IsNumeric() && (*Attribute)->GetTypeId() == PCG::Private::MetadataTypes<double>::Id)
+						{
+							double RightSideValue = FCString::Atod(*RightSide);
+							static_cast<FPCGMetadataAttribute<double>*>(*Attribute)->SetValue(Point.MetadataEntry, RightSideValue);
+						}
+						else if ((*Attribute)->GetTypeId() == PCG::Private::MetadataTypes<FString>::Id)
+						{
+							static_cast<FPCGMetadataAttribute<FString>*>(*Attribute)->SetValue(Point.MetadataEntry, RightSide);
+						}
+					}
+				}
+				else if (FPCGMetadataAttributeBase** Attribute = TagToAttributeMap.Find(ActorTag))
+				{
+					if ((*Attribute)->GetTypeId() == PCG::Private::MetadataTypes<bool>::Id)
+					{
+						static_cast<FPCGMetadataAttribute<bool>*>(*Attribute)->SetValue(Point.MetadataEntry, true);
+					}
+				}
 			}
 		};
 
diff --git a/Engine/Plugins/PCG/Source/PCGEditor/Public/PCGAssetExporterUtils.h b/Engine/Plugins/PCG/Source/PCGEditor/Public/PCGAssetExporterUtils.h
new file mode 100644
--- /dev/null
+++ b/Engine/Plugins/PCG/Source/PCGEditor/Public/PCGAssetExporterUtils.h
@@ -0,0 +1,22 @@
+// Copyright Epic Games, Inc. All Rights Reserved.
+
+#pragma once
+
+#include "PCGAssetExporter.h"
+
+#include "AssetRegistry/AssetData.h"
+
+#include "PCGAssetExporterUtils.generated.h"
+
+UCLASS()
+class PCGEDITOR_API UPCGAssetExporterUtils : public UBlueprintFunctionLibrary
+{
+	GENERATED_BODY()
+
+public:
+	UFUNCTION(BlueprintCallable, Category = "PCG|IO")
+	static UPackage* CreateAsset(UPCGAssetExporter* Exporter, FPCGAssetExporterParameters Parameters = FPCGAssetExporterParameters());
+
+	UFUNCTION(BlueprintCallable, Category = "PCG|IO")
+	static void UpdateAssets(const TArray<FAssetData>& PCGAssets, FPCGAssetExporterParameters Parameters = FPCGAssetExporterParameters());
+};
\ No newline at end of file
diff --git a/Engine/Plugins/PCG/Source/PCGEditor/Public/PCGEditorMenuUtils.h b/Engine/Plugins/PCG/Source/PCGEditor/Public/PCGEditorMenuUtils.h
--- a/Engine/Plugins/PCG/Source/PCGEditor/Public/PCGEditorMenuUtils.h
+++ b/Engine/Plugins/PCG/Source/PCGEditor/Public/PCGEditorMenuUtils.h
@@ -8,5 +8,5 @@
 namespace PCGEditorMenuUtils
 {
 	PCGEDITOR_API FToolMenuSection& CreatePCGSection(UToolMenu* Menu);
-	PCGEDITOR_API void CreateOrUpdatePCGAssetFromMenu(UToolMenu* Menu, TArray<FAssetData>& WorldAssets);
+	PCGEDITOR_API void CreateOrUpdatePCGAssetFromMenu(UToolMenu* Menu, TArray<FAssetData>& Assets);
 }
\ No newline at end of file
diff --git a/Engine/Plugins/PCG/Source/PCGEditor/Public/PCGLevelToAsset.h b/Engine/Plugins/PCG/Source/PCGEditor/Public/PCGLevelToAsset.h
--- a/Engine/Plugins/PCG/Source/PCGEditor/Public/PCGLevelToAsset.h
+++ b/Engine/Plugins/PCG/Source/PCGEditor/Public/PCGLevelToAsset.h
@@ -2,39 +2,38 @@
 
 #pragma once
 
-#include "PCGDataAsset.h"
+#include "PCGAssetExporter.h"
+
+#include "Engine/World.h"
 
 #include "PCGLevelToAsset.generated.h"
 
 class UPackage;
 
 UCLASS(BlueprintType, Blueprintable, meta = (ShowWorldContextPin))
-class PCGEDITOR_API UPCGLevelToAsset : public UObject
+class PCGEDITOR_API UPCGLevelToAsset : public UPCGAssetExporter
 {
 	GENERATED_BODY()
 
 public:
-	/** Creates/Updates on PCG Asset per given world. Allows exporter subclassing by passing in a Subclass. */
-	static void CreateOrUpdatePCGAssets(const TArray<FAssetData>& WorldAssets, TSubclassOf<UPCGLevelToAsset> Subclass = {});
-
-	/** Creates/Updates on PCG Asset for a specific world. Allows exporter subclassing by passing in a Subclass. Will return null if it fails, or the package that was modified otherwise. */
-	static UPackage* CreateOrUpdatePCGAsset(TSoftObjectPtr<UWorld> World, TSubclassOf<UPCGLevelToAsset> Subclass = {});
+	/** Creates/updates a PCG Asset per given world. Allows exporter subclassing by passing in a Subclass. */
+	static void CreateOrUpdatePCGAssets(const TArray<FAssetData>& WorldAssets, const FPCGAssetExporterParameters& Parameters = FPCGAssetExporterParameters(), TSubclassOf<UPCGLevelToAsset> Subclass = {});
 
-	/** Creates/Updates a PCG Asset for a specific world, allowing exporter subclassing (and settings creation by extension). Will return null if it fails, or the package that was modified otherwise. */
-	static UPackage* CreateOrUpdatePCGAsset(UWorld* World, TSubclassOf<UPCGLevelToAsset> Subclass = {});
+	/** Creates/Updates a PCG Asset for a specific world. Allows exporter subclassing by passing in a Subclass. Will return null if it fails, or the package that was modified on success. */
+	static UPackage* CreateOrUpdatePCGAsset(TSoftObjectPtr<UWorld> World, const FPCGAssetExporterParameters& Parameters = FPCGAssetExporterParameters(), TSubclassOf<UPCGLevelToAsset> Subclass = {});
 
-	/** Parses the level and fills in the provided data asset. Implement this in BP to drive the generation in a custom manner. */
-	UFUNCTION(BlueprintNativeEvent, meta = (DisplayName = "Export Level", ForceAsFunction))
-	bool BP_ExportLevel(UWorld* Level, const FString& PackageName, UPCGDataAsset* Asset);
+	/** Creates/Updates a PCG Asset for a specific world. Allows exporter subclassing (and settings creation by extension). Will return null if it fails, or the package that was modified on success. */
+	static UPackage* CreateOrUpdatePCGAsset(UWorld* World, const FPCGAssetExporterParameters& Parameters = FPCGAssetExporterParameters(), TSubclassOf<UPCGLevelToAsset> Subclass = {});
 
-	/** Returns the subtype of PCG data asset generated by the Level To Asset process. Override this in the instances where you would not use the default base class. */
-	UFUNCTION(BlueprintNativeEvent, meta = (DisplayName = "Get Asset Type", ForceAsFunction))
-	TSubclassOf<UPCGDataAsset> BP_GetAssetType() const;
+	/** Parses the world and fills in the provided data asset. Implement this in BP to drive the generation in a custom manner. */
+	UFUNCTION(BlueprintNativeEvent, meta = (DisplayName = "Export World", ForceAsFunction))
+	bool BP_ExportWorld(UWorld* World, const FString& PackageName, UPCGDataAsset* Asset);
 
 protected:
-	/** Drives the asset creation for the data parsed from the level, will return UPCGDataAsset by default. */
-	virtual TSubclassOf<UPCGDataAsset> GetAssetType() const;
+	//~Being UPCGAssetExporter interface
+	virtual bool ExportAsset(const FString& PackageName, UPCGDataAsset* Asset) override;
+	virtual UPackage* UpdateAsset(const FAssetData& PCGAsset) override;
+	//~End UPCGAssetExporter interface
 
-	/** Parses the level and pushes data to the asset. */
-	virtual bool ExportLevel(UWorld* Level, const FString& PackageName, UPCGDataAsset* Asset);
+	UWorld* WorldToExport = nullptr;
 };
\ No newline at end of file
diff --git a/Engine/Plugins/PCGExternalDataInterop/Source/PCGExternalDataInteropEditor/PCGExternalDataInteropEditor.Build.cs b/Engine/Plugins/PCGExternalDataInterop/Source/PCGExternalDataInteropEditor/PCGExternalDataInteropEditor.Build.cs
--- a/Engine/Plugins/PCGExternalDataInterop/Source/PCGExternalDataInteropEditor/PCGExternalDataInteropEditor.Build.cs
+++ b/Engine/Plugins/PCGExternalDataInterop/Source/PCGExternalDataInteropEditor/PCGExternalDataInteropEditor.Build.cs
@@ -22,6 +22,7 @@ namespace UnrealBuildTool.Rules
 				PrivateDependencyModuleNames.AddRange(
 					new string[]
 					{
+						"PCGEditor",
 						"PCGExternalDataInterop",
 						"UnrealEd",
 						"AlembicLib",
diff --git a/Engine/Plugins/PCGExternalDataInterop/Source/PCGExternalDataInteropEditor/Private/PCGLoadAlembic.cpp b/Engine/Plugins/PCGExternalDataInterop/Source/PCGExternalDataInteropEditor/Private/PCGLoadAlembic.cpp
--- a/Engine/Plugins/PCGExternalDataInterop/Source/PCGExternalDataInteropEditor/Private/PCGLoadAlembic.cpp
+++ b/Engine/Plugins/PCGExternalDataInterop/Source/PCGExternalDataInteropEditor/Private/PCGLoadAlembic.cpp
@@ -1,10 +1,59 @@
 // Copyright Epic Games, Inc. All Rights Reserved.
 
 #include "PCGLoadAlembic.h"
+#include "PCGAssetExporterUtils.h"
+#include "PCGModule.h"
 
 #include "Elements/PCGLoadAlembicElement.h"
 
+void UPCGAlembicToPCGAssetExporter::SerializeMetadata(FArchive& Ar)
+{
+	FPCGLoadAlembicBPData::StaticStruct()->SerializeItem(Ar, &LoadSettings, nullptr);
+}
+
+bool UPCGAlembicToPCGAssetExporter::ExportAsset(const FString& PackageName, UPCGDataAsset* Asset)
+{
+	check(Asset);
+	UPCGLoadAlembicFunctionLibrary::LoadAlembicFileToPCG(LoadSettings, Asset->Data, Asset);
+	return true;
+}
+
+UPackage* UPCGAlembicToPCGAssetExporter::UpdateAsset(const FAssetData& PCGAsset)
+{
+	UPCGDataAsset* Asset = Cast<UPCGDataAsset>(PCGAsset.GetAsset());
+	if (!Asset)
+	{
+		UE_LOG(LogPCG, Error, TEXT("Asset '%s' isn't a PCG data asset or could not be properly loaded."), *PCGAsset.GetObjectPathString());
+		return nullptr;
+	}
+
+	UPackage* Package = Asset->GetPackage();
+	if (!Package)
+	{
+		UE_LOG(LogPCG, Error, TEXT("Unable to retrieve package from Asset '%s'."), *PCGAsset.GetObjectPathString());
+		return nullptr;
+	}
+
+	UPCGLoadAlembicFunctionLibrary::LoadAlembicFileToPCG(LoadSettings, Asset->Data, Asset);
+
+	return Package;
+}
+
+void UPCGLoadAlembicFunctionLibrary::ExportAlembicFileToPCG(const FPCGLoadAlembicBPData& InSettings, FPCGAssetExporterParameters Parameters)
+{
+	UPCGAlembicToPCGAssetExporter* Exporter = NewObject<UPCGAlembicToPCGAssetExporter>(GetTransientPackage());
+	check(Exporter);
+
+	Exporter->LoadSettings = InSettings;
+	UPCGAssetExporterUtils::CreateAsset(Exporter, Parameters);
+}
+
 void UPCGLoadAlembicFunctionLibrary::LoadAlembicFileToPCG(const FPCGLoadAlembicBPData& InSettings, FPCGDataCollection& Data, UObject* TargetOuter)
+{
+	LoadAlembicFileToPCGInternal(InSettings, Data, TargetOuter);
+}
+
+void UPCGLoadAlembicFunctionLibrary::LoadAlembicFileToPCGInternal(const FPCGLoadAlembicBPData& InSettings, FPCGDataCollection& Data, UObject* TargetOuter)
 {
 	// Simulate execution of the FPCGLoadAlembicElement.
 	UPCGLoadAlembicSettings* Settings = NewObject<UPCGLoadAlembicSettings>();
diff --git a/Engine/Plugins/PCGExternalDataInterop/Source/PCGExternalDataInteropEditor/Public/PCGLoadAlembic.h b/Engine/Plugins/PCGExternalDataInterop/Source/PCGExternalDataInteropEditor/Public/PCGLoadAlembic.h
--- a/Engine/Plugins/PCGExternalDataInterop/Source/PCGExternalDataInteropEditor/Public/PCGLoadAlembic.h
+++ b/Engine/Plugins/PCGExternalDataInterop/Source/PCGExternalDataInteropEditor/Public/PCGLoadAlembic.h
@@ -2,6 +2,7 @@
 
 #pragma once
 
+#include "PCGAssetExporter.h"
 #include "Elements/PCGLoadAlembicElement.h"
 #include "Metadata/PCGAttributePropertySelector.h"
 
@@ -28,14 +29,39 @@ struct PCGEXTERNALDATAINTEROPEDITOR_API FPCGLoadAlembicBPData
 	TMap<FString, FPCGAttributePropertyInputSelector> AttributeMapping;
 };
 
+UCLASS()
+class UPCGAlembicToPCGAssetExporter : public UPCGAssetExporter
+{
+	GENERATED_BODY()
+
+public:
+	//~Begin UPCGAssetExporter interface
+	virtual bool ExportAsset(const FString& PackageName, UPCGDataAsset* Asset);
+	virtual UPackage* UpdateAsset(const FAssetData& PCGAsset) override;
+
+protected:
+	virtual void SerializeMetadata(FArchive& Ar) override;
+	//~End UPCGAssetExporter interface
+
+public:
+	FPCGLoadAlembicBPData LoadSettings;
+};
+
 UCLASS()
 class PCGEXTERNALDATAINTEROPEDITOR_API UPCGLoadAlembicFunctionLibrary : public UBlueprintFunctionLibrary
 {
 	GENERATED_BODY()
 
+public:
 	UFUNCTION(BlueprintCallable, Category = "PCG|IO")
+	static void ExportAlembicFileToPCG(const FPCGLoadAlembicBPData& Settings, FPCGAssetExporterParameters Parameters = FPCGAssetExporterParameters());
+
+	UFUNCTION(BlueprintCallable, Category = "PCG|IO", meta = (DeprecatedFunction, DeprecationMessage="The LoadAlembicFileToPCG function has been replaced by ExportAlembicFileToPCG"))
 	static void LoadAlembicFileToPCG(const FPCGLoadAlembicBPData& Settings, FPCGDataCollection& Data, UObject* TargetOuter);
 
 	UFUNCTION(BlueprintCallable, Category = "PCG|IO", meta = (ScriptMethod))
 	static void SetupFromStandard(UPARAM(ref) FPCGLoadAlembicBPData& Data, EPCGLoadAlembicStandardSetup InSetup);
+
+protected:
+	static void LoadAlembicFileToPCGInternal(const FPCGLoadAlembicBPData& Settings, FPCGDataCollection& Data, UObject* TargetOuter);
 };
\ No newline at end of file
