[PCG] Add the ability to silence node warnings for missing property override attributes on the Static Mesh Spawner if applicable

#rnx

diff --git a/Engine/Plugins/PCG/Source/PCG/Private/MeshSelectors/PCGMeshSelectorByAttribute.cpp b/Engine/Plugins/PCG/Source/PCG/Private/MeshSelectors/PCGMeshSelectorByAttribute.cpp
--- a/Engine/Plugins/PCG/Source/PCG/Private/MeshSelectors/PCGMeshSelectorByAttribute.cpp
+++ b/Engine/Plugins/PCG/Source/PCG/Private/MeshSelectors/PCGMeshSelectorByAttribute.cpp
@@ -213,7 +213,7 @@ bool UPCGMeshSelectorByAttribute::SelectInstances(
 		// If there are valid overrides, partition the points on those attributes so that an instance can be created for each
 		if (!ValidSelectorOverrides.IsEmpty())
 		{
-			Context.AttributeOverridePartition = PCGMetadataPartitionCommon::AttributeGenericPartition(InPointData, ValidSelectorOverrides, &Context);
+			Context.AttributeOverridePartition = PCGMetadataPartitionCommon::AttributeGenericPartition(InPointData, ValidSelectorOverrides, &Context, Settings->bSilenceOverrideAttributeNotFoundErrors);
 		}
 
 		// Set the descriptors to match the partition count. Uninitialized, because it will be copied from the template below
diff --git a/Engine/Plugins/PCG/Source/PCG/Private/Metadata/PCGMetadataPartitionCommon.cpp b/Engine/Plugins/PCG/Source/PCG/Private/Metadata/PCGMetadataPartitionCommon.cpp
--- a/Engine/Plugins/PCG/Source/PCG/Private/Metadata/PCGMetadataPartitionCommon.cpp
+++ b/Engine/Plugins/PCG/Source/PCG/Private/Metadata/PCGMetadataPartitionCommon.cpp
@@ -169,7 +169,7 @@ namespace PCGMetadataPartitionCommon
 	/**
 	* Dispatch the partition according to the data and selector.
 	*/
-	TArray<TArray<int32>> AttributeGenericPartition(const UPCGData* InData, const FPCGAttributePropertySelector& InSelector, FPCGContext* InOptionalContext)
+	TArray<TArray<int32>> AttributeGenericPartition(const UPCGData* InData, const FPCGAttributePropertySelector& InSelector, FPCGContext* InOptionalContext, bool bSilenceMissingAttributeErrors)
 	{
 		TRACE_CPUPROFILER_EVENT_SCOPE(PCGMetadataPartitionCommon::AttributeGenericPartition::SingleSelector);
 		if (!InData)
@@ -196,7 +196,11 @@ namespace PCGMetadataPartitionCommon
 			const FPCGMetadataAttributeBase* Attribute = Metadata->GetConstAttribute(InSelector.GetName());
 			if (!Attribute)
 			{
-				PCGLog::LogErrorOnGraph(FText::Format(LOCTEXT("InvalidAttribute", "Attribute {0} not found"), InSelector.GetDisplayText()), InOptionalContext);
+				if (!bSilenceMissingAttributeErrors)
+				{
+					PCGLog::LogErrorOnGraph(FText::Format(LOCTEXT("InvalidAttribute", "Attribute {0} not found"), InSelector.GetDisplayText()), InOptionalContext);
+				}
+
 				return {};
 			}
 
@@ -207,7 +211,11 @@ namespace PCGMetadataPartitionCommon
 			TUniquePtr<const IPCGAttributeAccessor> Accessor = PCGAttributeAccessorHelpers::CreateConstAccessor(InData, InSelector);
 			if (!Accessor.IsValid())
 			{
-				PCGLog::LogErrorOnGraph(FText::Format(LOCTEXT("InvalidAccessor", "Attribute {0} not found"), InSelector.GetDisplayText()), InOptionalContext);
+				if (!bSilenceMissingAttributeErrors)
+				{
+					PCGLog::LogErrorOnGraph(FText::Format(LOCTEXT("InvalidAccessor", "Attribute {0} not found"), InSelector.GetDisplayText()), InOptionalContext);
+				}
+
 				return {};
 			}
 
@@ -245,14 +253,14 @@ namespace PCGMetadataPartitionCommon
 	 *  3  b  c  c                         Final Partition (A&B&C)->[0],[1],[2],[3,4]
 	 *  4  b  c  c
 	 */
-	TArray<TArray<int32>> AttributeGenericPartition(const UPCGData* InData, const TArrayView<const FPCGAttributePropertySelector>& InSelectorArrayView, FPCGContext* InOptionalContext)
+	TArray<TArray<int32>> AttributeGenericPartition(const UPCGData* InData, const TArrayView<const FPCGAttributePropertySelector>& InSelectorArrayView, FPCGContext* InOptionalContext, bool bSilenceMissingAttributeErrors)
 	{
 		TRACE_CPUPROFILER_EVENT_SCOPE(PCGMetadataPartitionCommon::AttributeGenericPartition::MultiSelector);
 
 		// Small optimization to partition on a single attribute
 		if (InSelectorArrayView.Num() == 1)
 		{
-			return AttributeGenericPartition(InData, InSelectorArrayView[0], InOptionalContext);
+			return AttributeGenericPartition(InData, InSelectorArrayView[0], InOptionalContext, bSilenceMissingAttributeErrors);
 		}
 
 		if (!InData || InSelectorArrayView.IsEmpty() || !InData->ConstMetadata())
@@ -292,8 +300,8 @@ namespace PCGMetadataPartitionCommon
 				BitPartition& CurrentBitPartition = BitPartitions[I];
 
 				// TODO: Ideally, refactor AttributeGenericPartition to return directly into BitArray format instead to avoid conversion
-				// Partition once for each attribute
-				CurrentIndexPartition = AttributeGenericPartition(InData, InSelectorArrayView[I], InOptionalContext);
+				// Partition once for each attribute. It is okay if this to be empty, as it will be skipped later during the iterative partition
+				CurrentIndexPartition = AttributeGenericPartition(InData, InSelectorArrayView[I], InOptionalContext, bSilenceMissingAttributeErrors);
 				// The bit partitions will match the index partitions
 				CurrentBitPartition.SetNum(CurrentIndexPartition.Num());
 
@@ -361,9 +369,9 @@ namespace PCGMetadataPartitionCommon
 	/**
 	* Do a partition on the given point data for the selector
 	*/
-	TArray<UPCGData*> AttributePointPartition(const UPCGPointData* InData, const TArrayView<const FPCGAttributePropertySelector>& InSelectorArrayView, FPCGContext* InOptionalContext)
+	TArray<UPCGData*> AttributePointPartition(const UPCGPointData* InData, const TArrayView<const FPCGAttributePropertySelector>& InSelectorArrayView, FPCGContext* InOptionalContext, bool bSilenceMissingAttributeErrors)
 	{
-		TArray<TArray<int32>> Partition = AttributeGenericPartition(InData, InSelectorArrayView, InOptionalContext);
+		TArray<TArray<int32>> Partition = AttributeGenericPartition(InData, InSelectorArrayView, InOptionalContext, bSilenceMissingAttributeErrors);
 		if (Partition.IsEmpty())
 		{
 			return {};
@@ -395,7 +403,7 @@ namespace PCGMetadataPartitionCommon
 		return PartitionedData;
 	}
 
-	TArray<UPCGData*> AttributeParamSpatialPartition(const UPCGData* InData, const TArrayView<const FPCGAttributePropertySelector>& InSelectorArray, FPCGContext* InOptionalContext)
+	TArray<UPCGData*> AttributeParamSpatialPartition(const UPCGData* InData, const TArrayView<const FPCGAttributePropertySelector>& InSelectorArray, FPCGContext* InOptionalContext, bool bSilenceMissingAttributeErrors)
 	{
 		if (!InData->IsA<UPCGSpatialData>() && !InData->IsA<UPCGParamData>())
 		{
@@ -403,7 +411,7 @@ namespace PCGMetadataPartitionCommon
 			return {};
 		}
 
-		const TArray<TArray<int32>> Partition = AttributeGenericPartition(InData, InSelectorArray, InOptionalContext);
+		const TArray<TArray<int32>> Partition = AttributeGenericPartition(InData, InSelectorArray, InOptionalContext, bSilenceMissingAttributeErrors);
 
 		if (Partition.IsEmpty())
 		{
@@ -472,28 +480,28 @@ namespace PCGMetadataPartitionCommon
 		return PartitionedData;
 	}
 
-	TArray<UPCGData*> AttributePartition(const UPCGData* InData, const FPCGAttributePropertySelector& InSelector, FPCGContext* InOptionalContext)
+	TArray<UPCGData*> AttributePartition(const UPCGData* InData, const FPCGAttributePropertySelector& InSelector, FPCGContext* InOptionalContext, bool bSilenceMissingAttributeErrors)
 	{
 		const TArrayView<const FPCGAttributePropertySelector> ArrayView(&InSelector, 1);
 		if (const UPCGPointData* InPointData = Cast<UPCGPointData>(InData))
 		{
-			return AttributePointPartition(InPointData, ArrayView, InOptionalContext);
+			return AttributePointPartition(InPointData, ArrayView, InOptionalContext, bSilenceMissingAttributeErrors);
 		}
 		else
 		{
-			return AttributeParamSpatialPartition(InData, ArrayView, InOptionalContext);
+			return AttributeParamSpatialPartition(InData, ArrayView, InOptionalContext, bSilenceMissingAttributeErrors);
 		}
 	}
 
-	TArray<UPCGData*> AttributePartition(const UPCGData* InData, const TArrayView<const FPCGAttributePropertySelector>& InSelectorArrayView, FPCGContext* InOptionalContext)
+	TArray<UPCGData*> AttributePartition(const UPCGData* InData, const TArrayView<const FPCGAttributePropertySelector>& InSelectorArrayView, FPCGContext* InOptionalContext, bool bSilenceMissingAttributeErrors)
 	{
 		if (const UPCGPointData* InPointData = Cast<UPCGPointData>(InData))
 		{
-			return AttributePointPartition(InPointData, InSelectorArrayView, InOptionalContext);
+			return AttributePointPartition(InPointData, InSelectorArrayView, InOptionalContext, bSilenceMissingAttributeErrors);
 		}
 		else
 		{
-			return AttributeParamSpatialPartition(InData, InSelectorArrayView, InOptionalContext);
+			return AttributeParamSpatialPartition(InData, InSelectorArrayView, InOptionalContext, bSilenceMissingAttributeErrors);
 		}
 	}
 }
diff --git a/Engine/Plugins/PCG/Source/PCG/Public/Elements/PCGStaticMeshSpawner.h b/Engine/Plugins/PCG/Source/PCG/Public/Elements/PCGStaticMeshSpawner.h
--- a/Engine/Plugins/PCG/Source/PCG/Public/Elements/PCGStaticMeshSpawner.h
+++ b/Engine/Plugins/PCG/Source/PCG/Public/Elements/PCGStaticMeshSpawner.h
@@ -119,6 +119,10 @@ public:
 	UPROPERTY(BlueprintReadWrite, EditAnywhere, Category = "Settings|Debug")
 	bool bSynchronousLoad = false;
 
+	/** Opt-in option to silence errors when the property override attributes are not found. */
+	UPROPERTY(BlueprintReadWrite, EditAnywhere, Category = "Settings|Advanced")
+	bool bSilenceOverrideAttributeNotFoundErrors = false;
+
 protected:
 	void RefreshMeshSelector();
 	void RefreshInstancePacker();
diff --git a/Engine/Plugins/PCG/Source/PCG/Public/Metadata/PCGMetadataPartitionCommon.h b/Engine/Plugins/PCG/Source/PCG/Public/Metadata/PCGMetadataPartitionCommon.h
--- a/Engine/Plugins/PCG/Source/PCG/Public/Metadata/PCGMetadataPartitionCommon.h
+++ b/Engine/Plugins/PCG/Source/PCG/Public/Metadata/PCGMetadataPartitionCommon.h
@@ -17,9 +17,10 @@ namespace PCGMetadataPartitionCommon
 	* @param InData - Data to partition.
 	* @param InSelector - Selector on the attribute to partition.
 	* @param InOptionalContext - Optional context for logging.
+	* @param bSilenceMissingAttributeErrors - Do not log errors to the context or log.
 	* @returns Array of data.
 	*/
-	PCG_API TArray<UPCGData*> AttributePartition(const UPCGData* InData, const FPCGAttributePropertySelector& InSelector, FPCGContext* InOptionalContext = nullptr);
+	PCG_API TArray<UPCGData*> AttributePartition(const UPCGData* InData, const FPCGAttributePropertySelector& InSelector, FPCGContext* InOptionalContext = nullptr, bool bSilenceMissingAttributeErrors = false);
 
 	/**
 	* Partition the incoming data on the given attribute.
@@ -27,9 +28,10 @@ namespace PCGMetadataPartitionCommon
 	* @param InData - Data to partition.
 	* @param InSelectorArrayView - ArrayView of selectors on the attribute to partition.
 	* @param InOptionalContext - Optional context for logging.
+	* @param bSilenceMissingAttributeErrors - Do not log errors to the context or log.
 	* @returns Array of data.
 	*/
-	PCG_API TArray<UPCGData*> AttributePartition(const UPCGData* InData, const TArrayView<const FPCGAttributePropertySelector>& InSelectorArrayView, FPCGContext* InOptionalContext = nullptr);
+	PCG_API TArray<UPCGData*> AttributePartition(const UPCGData* InData, const TArrayView<const FPCGAttributePropertySelector>& InSelectorArrayView, FPCGContext* InOptionalContext = nullptr, bool bSilenceMissingAttributeErrors = false);
 
 	/**
 	* Generic partition for the incoming data on the given attribute.
@@ -37,9 +39,10 @@ namespace PCGMetadataPartitionCommon
 	* @param InData - Data to partition, need to support attributes (spatial data or attribute set).
 	* @param InSelector - Selector on the attribute to partition.
 	* @param InOptionalContext - Optional context for logging.
+	* @param bSilenceMissingAttributeErrors - Do not log errors to the context or log.
 	* @returns Array of bucket indices.
 	*/
-	PCG_API TArray<TArray<int32>> AttributeGenericPartition(const UPCGData* InData, const FPCGAttributePropertySelector& InSelector, FPCGContext* InOptionalContext = nullptr);
+	PCG_API TArray<TArray<int32>> AttributeGenericPartition(const UPCGData* InData, const FPCGAttributePropertySelector& InSelector, FPCGContext* InOptionalContext = nullptr, bool bSilenceMissingAttributeErrors = false);
 
 	/**
 	* Generic partition for the incoming data on the given array of attributes.
@@ -47,7 +50,8 @@ namespace PCGMetadataPartitionCommon
 	* @param InData - Data to partition, need to support attributes (spatial data or attribute set).
 	* @param InSelectorArrayView - ArrayView of selectors on the attribute to partition.
 	* @param InOptionalContext - Optional context for logging.
+	* @param bSilenceMissingAttributeErrors - Do not log errors to the context or log.
 	* @returns Array of bucket indices.
 	*/
-	PCG_API TArray<TArray<int32>> AttributeGenericPartition(const UPCGData* InData, const TArrayView<const FPCGAttributePropertySelector>& InSelectorArrayView, FPCGContext* InOptionalContext = nullptr);
+	PCG_API TArray<TArray<int32>> AttributeGenericPartition(const UPCGData* InData, const TArrayView<const FPCGAttributePropertySelector>& InSelectorArrayView, FPCGContext* InOptionalContext = nullptr, bool bSilenceMissingAttributeErrors = false);
 }
\ No newline at end of file
