[PCG] Fix builder leaking artifacts due to deleted packages getting lost during GCs.

Fix is to track actor deletions directly and keep a hold of these packages.

diff --git a/Engine/Plugins/PCG/Source/PCG/Private/Elements/PCGDataFromActor.cpp b/Engine/Plugins/PCG/Source/PCG/Private/Elements/PCGDataFromActor.cpp
--- a/Engine/Plugins/PCG/Source/PCG/Private/Elements/PCGDataFromActor.cpp
+++ b/Engine/Plugins/PCG/Source/PCG/Private/Elements/PCGDataFromActor.cpp
@@ -485,13 +485,13 @@ void FPCGDataFromActorElement::GatherWaitTasks(AActor* FoundActor, FPCGContext*
 			const FPCGTaskId GenerateTask = Component->GenerateLocalGetTaskId(EPCGComponentGenerationTrigger::GenerateOnDemand, /*bForce=*/false);
 			if (GenerateTask != InvalidPCGTaskId)
 			{
-				PCGGraphExecutionLogging::LogGraphScheduleDependency(Component);
+				PCGGraphExecutionLogging::LogGraphScheduleDependency(Component, Context->Stack);
 
 				OutWaitTasks.Add(GenerateTask);
 			}
 			else
 			{
-				PCGGraphExecutionLogging::LogGraphScheduleDependencyFailed(Component);
+				PCGGraphExecutionLogging::LogGraphScheduleDependencyFailed(Component, Context->Stack);
 			}
 		}
 	}
diff --git a/Engine/Plugins/PCG/Source/PCG/Private/Graph/PCGGraphExecutor.cpp b/Engine/Plugins/PCG/Source/PCG/Private/Graph/PCGGraphExecutor.cpp
--- a/Engine/Plugins/PCG/Source/PCG/Private/Graph/PCGGraphExecutor.cpp
+++ b/Engine/Plugins/PCG/Source/PCG/Private/Graph/PCGGraphExecutor.cpp
@@ -154,7 +154,7 @@ FPCGTaskId FPCGGraphExecutor::Schedule(
 {
 	check(SourceComponent);
 
-	PCGGraphExecutionLogging::LogGraphSchedule(SourceComponent);
+	PCGGraphExecutionLogging::LogGraphSchedule(SourceComponent, Graph);
 	
 	FPCGTaskId ScheduledId = InvalidPCGTaskId;
 
diff --git a/Engine/Plugins/PCG/Source/PCG/Private/PCGComponent.cpp b/Engine/Plugins/PCG/Source/PCG/Private/PCGComponent.cpp
--- a/Engine/Plugins/PCG/Source/PCG/Private/PCGComponent.cpp
+++ b/Engine/Plugins/PCG/Source/PCG/Private/PCGComponent.cpp
@@ -145,7 +145,7 @@ void UPCGComponent::SetGraphInterfaceLocal(UPCGGraphInterface* InGraphInterface)
 
 void UPCGComponent::AddToManagedResources(UPCGManagedResource* InResource)
 {
-	PCGGeneratedResourcesLogging::LogAddToManagedResources(InResource);
+	PCGGeneratedResourcesLogging::LogAddToManagedResources(this, InResource);
 
 	if (InResource)
 	{
@@ -613,7 +613,7 @@ FPCGTaskId UPCGComponent::CleanupInternal(bool bRemoveComponents, const TArray<F
 		return InvalidPCGTaskId;
 	}
 
-	PCGGeneratedResourcesLogging::LogCleanupInternal(bRemoveComponents);
+	PCGGeneratedResourcesLogging::LogCleanupInternal(this, bRemoveComponents);
 
 	Modify(!IsInPreviewMode());
 
@@ -859,7 +859,7 @@ bool UPCGComponent::MoveResourcesToNewActor(AActor* InNewActor, bool bCreateChil
 
 void UPCGComponent::CleanupLocalImmediate(bool bRemoveComponents, bool bCleanupLocalComponents)
 {
-	PCGGeneratedResourcesLogging::LogCleanupLocalImmediate(bRemoveComponents, GeneratedResources);
+	PCGGeneratedResourcesLogging::LogCleanupLocalImmediate(this, bRemoveComponents, GeneratedResources);
 
 	UPCGSubsystem* Subsystem = GetSubsystem();
 
@@ -894,7 +894,7 @@ void UPCGComponent::CleanupLocalImmediate(bool bRemoveComponents, bool bCleanupL
 				// Note: resources can be null here in some loading + bp object cases
 				UPCGManagedResource* Resource = GeneratedResources[ResourceIndex];
 
-				PCGGeneratedResourcesLogging::LogCleanupLocalImmediateResource(Resource);
+				PCGGeneratedResourcesLogging::LogCleanupLocalImmediateResource(this, Resource);
 
 				if (!Resource || Resource->Release(bRemoveComponents, ActorsToDelete))
 				{
@@ -933,7 +933,7 @@ void UPCGComponent::CleanupLocalImmediate(bool bRemoveComponents, bool bCleanupL
 		Subsystem->CleanupLocalComponentsImmediate(this, bRemoveComponents);
 	}
 
-	PCGGeneratedResourcesLogging::LogCleanupLocalImmediateFinished(GeneratedResources);
+	PCGGeneratedResourcesLogging::LogCleanupLocalImmediateFinished(this, GeneratedResources);
 }
 
 FPCGTaskId UPCGComponent::CreateCleanupTask(bool bRemoveComponents, const TArray<FPCGTaskId>& Dependencies)
@@ -948,7 +948,7 @@ FPCGTaskId UPCGComponent::CreateCleanupTask(bool bRemoveComponents, const TArray
 		return InvalidPCGTaskId;
 	}
 
-	PCGGeneratedResourcesLogging::LogCreateCleanupTask(bRemoveComponents);
+	PCGGeneratedResourcesLogging::LogCreateCleanupTask(this, bRemoveComponents);
 
 	// Keep track of all the dependencies
 	TArray<FPCGTaskId> AdditionalDependencies;
@@ -1003,7 +1003,7 @@ FPCGTaskId UPCGComponent::CreateCleanupTask(bool bRemoveComponents, const TArray
 					UE_LOG(LogPCG, Error, TEXT("[UPCGComponent::CreateCleanupTask] Null generated resource encountered on actor \"%s\"."), *ThisComponent->GetOwner()->GetFName().ToString());
 				}
 
-				PCGGeneratedResourcesLogging::LogCreateCleanupTaskResource(Resource);
+				PCGGeneratedResourcesLogging::LogCreateCleanupTaskResource(ThisComponentWeakPtr.Get(), Resource);
 
 				if (!Resource || Resource->Release(bRemoveComponents, Context->ActorsToDelete))
 				{
@@ -1099,7 +1099,7 @@ FPCGTaskId UPCGComponent::CreateCleanupTask(bool bRemoveComponents, const TArray
 			}
 #endif
 
-			PCGGeneratedResourcesLogging::LogCreateCleanupTaskFinished(ThisComponentWeakPtr->GeneratedResources);
+			PCGGeneratedResourcesLogging::LogCreateCleanupTaskFinished(ThisComponentWeakPtr.Get(), ThisComponentWeakPtr.IsValid() ? &ThisComponentWeakPtr->GeneratedResources : nullptr);
 		}
 
 		return true;
@@ -1110,7 +1110,7 @@ FPCGTaskId UPCGComponent::CreateCleanupTask(bool bRemoveComponents, const TArray
 
 void UPCGComponent::CleanupUnusedManagedResources()
 {
-	PCGGeneratedResourcesLogging::LogCleanupUnusedManagedResources(GeneratedResources);
+	PCGGeneratedResourcesLogging::LogCleanupUnusedManagedResources(this, GeneratedResources);
 
 	TSet<TSoftObjectPtr<AActor>> ActorsToDelete;
 
@@ -1121,7 +1121,7 @@ void UPCGComponent::CleanupUnusedManagedResources()
 		{
 			UPCGManagedResource* Resource = GetValid(GeneratedResources[ResourceIndex]);
 
-			PCGGeneratedResourcesLogging::LogCleanupUnusedManagedResourcesResource(Resource);
+			PCGGeneratedResourcesLogging::LogCleanupUnusedManagedResourcesResource(this, Resource);
 
 			if (!Resource && GetOwner())
 			{
@@ -1151,7 +1151,7 @@ void UPCGComponent::CleanupUnusedManagedResources()
 
 	UPCGActorHelpers::DeleteActors(GetWorld(), ActorsToDelete.Array());
 
-	PCGGeneratedResourcesLogging::LogCleanupUnusedManagedResourcesFinished(GeneratedResources);
+	PCGGeneratedResourcesLogging::LogCleanupUnusedManagedResourcesFinished(this, GeneratedResources);
 }
 
 void UPCGComponent::BeginPlay()
diff --git a/Engine/Plugins/PCG/Source/PCG/Private/PCGManagedResource.cpp b/Engine/Plugins/PCG/Source/PCG/Private/PCGManagedResource.cpp
--- a/Engine/Plugins/PCG/Source/PCG/Private/PCGManagedResource.cpp
+++ b/Engine/Plugins/PCG/Source/PCG/Private/PCGManagedResource.cpp
@@ -88,7 +88,7 @@ bool UPCGManagedActors::Release(bool bHardRelease, TSet<TSoftObjectPtr<AActor>>&
 
 	if (!Super::Release(bHardRelease, OutActorsToDelete))
 	{
-		PCGGeneratedResourcesLogging::LogManagedActorsSoftRelease(GeneratedActors);
+		PCGGeneratedResourcesLogging::LogManagedActorsRelease(this, GeneratedActors, bHardRelease, /*bOnlyMarkedForCleanup=*/true);
 
 		// Mark actors as potentially-to-be-cleaned-up
 		for (TSoftObjectPtr<AActor> GeneratedActor : GeneratedActors)
@@ -123,7 +123,7 @@ bool UPCGManagedActors::Release(bool bHardRelease, TSet<TSoftObjectPtr<AActor>>&
 		OutActorsToDelete.Append(GeneratedActors);
 	}
 
-	PCGGeneratedResourcesLogging::LogManagedActorsHardRelease(GeneratedActors);
+	PCGGeneratedResourcesLogging::LogManagedActorsRelease(this, GeneratedActors, bHardRelease, /*bOnlyMarkedForCleanup=*/false);
 
 	// Cleanup recursively
 	TInlineComponentArray<UPCGComponent*, 1> ComponentsToCleanup;
@@ -344,6 +344,8 @@ bool UPCGManagedComponent::Release(bool bHardRelease, TSet<TSoftObjectPtr<AActor
 #if WITH_EDITOR
 		if (bMarkedTransientOnLoad)
 		{
+			PCGGeneratedResourcesLogging::LogManagedComponentHidden(this);
+
 			HideComponent();
 			bIsMarkedUnused = true;
 		}
@@ -352,11 +354,15 @@ bool UPCGManagedComponent::Release(bool bHardRelease, TSet<TSoftObjectPtr<AActor
 		{
 			if (bDeleteComponent)
 			{
+				PCGGeneratedResourcesLogging::LogManagedResourceHardRelease(this);
+
 				GeneratedComponent->DestroyComponent();
 				ForgetComponent();
 			}
 			else
 			{
+				PCGGeneratedResourcesLogging::LogManagedResourceSoftRelease(this);
+
 				// We can only mark it unused if we can reset the component.
 				bIsMarkedUnused = true;
 				GeneratedComponent->ComponentTags.Add(PCGHelpers::MarkedForCleanupPCGTag);
@@ -365,6 +371,8 @@ bool UPCGManagedComponent::Release(bool bHardRelease, TSet<TSoftObjectPtr<AActor
 	}
 	else
 	{
+		PCGGeneratedResourcesLogging::LogManagedComponentDeleteNull(this);
+
 		// Dead component reference - clear it out.
 		bDeleteComponent = true;
 	}
diff --git a/Engine/Plugins/PCG/Source/PCG/Private/Utils/PCGGeneratedResourcesLogging.cpp b/Engine/Plugins/PCG/Source/PCG/Private/Utils/PCGGeneratedResourcesLogging.cpp
--- a/Engine/Plugins/PCG/Source/PCG/Private/Utils/PCGGeneratedResourcesLogging.cpp
+++ b/Engine/Plugins/PCG/Source/PCG/Private/Utils/PCGGeneratedResourcesLogging.cpp
@@ -3,26 +3,27 @@
 #include "PCGGeneratedResourcesLogging.h"
 
 #include "PCGComponent.h"
+#include "PCGGraph.h"
 #include "PCGManagedResource.h"
 #include "PCGModule.h"
 
 #include "GameFramework/Actor.h"
 #include "HAL/IConsoleManager.h"
 
+namespace PCGGeneratedResourcesLogging
+{
 #if WITH_EDITOR
-static TAutoConsoleVariable<bool> CVarGenResourcesLoggingEnable(
-	TEXT("pcg.GeneratedResourcesLogging.Enable"),
-	false,
-	TEXT("Enables fine grained log of generated resources management"));
-
-static TAutoConsoleVariable<int32> CVarGenResourcesLoggingMaxPrintCount(
-	TEXT("pcg.GeneratedResourcesLogging.MaxElementPrintCount"),
-	3,
-	TEXT("Enables fine grained log of generated resources management"));
+	static TAutoConsoleVariable<bool> CVarGenResourcesLoggingEnable(
+		TEXT("pcg.ManagedResourcesLogging.Enable"),
+		false,
+		TEXT("Enables fine grained log of generated resources management"));
+
+	static TAutoConsoleVariable<int32> CVarGenResourcesLoggingMaxPrintCount(
+		TEXT("pcg.ManagedResourcesLogging.MaxElementPrintCount"),
+		3,
+		TEXT("Sets how many entries to display for resources that are arrays of objects"));
 #endif
 
-namespace PCGGeneratedResourcesLogging
-{
 	bool LogEnabled()
 	{
 #if WITH_EDITOR
@@ -32,17 +33,15 @@ namespace PCGGeneratedResourcesLogging
 #endif
 	}
 
-	void LogAddToManagedResources(const UPCGManagedResource* Resource)
-	{
 #if WITH_EDITOR
-		if (!LogEnabled())
-		{
-			return;
-		}
-
-		if (const UPCGManagedActors* ManagedActors = Cast<UPCGManagedActors>(Resource))
+	void LogResource(const UPCGComponent* InComponent, UPCGManagedResource* InResource)
+	{
+		if (const UPCGManagedActors* ManagedActors = Cast<UPCGManagedActors>(InResource))
 		{
-			UE_LOG(LogPCG, Log, TEXT("[PCGMANAGEDRESOURCES] UPCGComponent::AddToManagedResources, actor count %d"), ManagedActors->GeneratedActors.Num());
+			UE_LOG(LogPCG, Log, TEXT("[PCGMANAGEDRESOURCES] [%s/%s]         Managed actors (%d):"),
+				(InComponent && InComponent->GetOwner()) ? *InComponent->GetOwner()->GetName() : TEXT("MISSINGCOMPONENT"),
+				(InComponent && InComponent->GetGraph()) ? *InComponent->GetGraph()->GetName() : TEXT("MISSINGGRAPH"),
+				ManagedActors->GeneratedActors.Num());
 
 			const int32 MaxPrint = CVarGenResourcesLoggingMaxPrintCount.GetValueOnAnyThread();
 			int32 Count = 0;
@@ -55,40 +54,64 @@ namespace PCGGeneratedResourcesLogging
 
 				if (Actor.Get())
 				{
-					UE_LOG(LogPCG, Log, TEXT("[PCGMANAGEDRESOURCES]         |- '%s' (%d/%d)"), *Actor->GetFName().ToString(), Count, ManagedActors->GeneratedActors.Num());
+					UE_LOG(LogPCG, Log, TEXT("[PCGMANAGEDRESOURCES] [%s/%s]                 |- '%s' (%d/%d)"),
+						(InComponent && InComponent->GetOwner()) ? *InComponent->GetOwner()->GetName() : TEXT("MISSINGCOMPONENT"),
+						(InComponent && InComponent->GetGraph()) ? *InComponent->GetGraph()->GetName() : TEXT("MISSINGGRAPH"),
+						*Actor->GetFName().ToString(),
+						Count, ManagedActors->GeneratedActors.Num());
 				}
 				else
 				{
-					UE_LOG(LogPCG, Log, TEXT("[PCGMANAGEDRESOURCES]         |- NULL (%d/%d)"), Count, ManagedActors->GeneratedActors.Num());
+					UE_LOG(LogPCG, Log, TEXT("[PCGMANAGEDRESOURCES] [%s/%s]                 |- NULL (%d/%d)"),
+						(InComponent && InComponent->GetOwner()) ? *InComponent->GetOwner()->GetName() : TEXT("MISSINGCOMPONENT"),
+						(InComponent && InComponent->GetGraph()) ? *InComponent->GetGraph()->GetName() : TEXT("MISSINGGRAPH"),
+						Count, ManagedActors->GeneratedActors.Num());
 				}
 			}
 		}
-		else if (const UPCGManagedComponent* ManagedComponent = Cast<UPCGManagedComponent>(Resource))
+		else if (const UPCGManagedComponent* ManagedComponent = Cast<UPCGManagedComponent>(InResource))
 		{
 			if (ManagedComponent->GeneratedComponent)
 			{
-				UE_LOG(LogPCG, Log, TEXT("[PCGMANAGEDRESOURCES] UPCGComponent::AddToManagedResources, managed component on actor '%s'"), *ManagedComponent->GeneratedComponent->GetReadableName());
+				UE_LOG(LogPCG, Log, TEXT("[PCGMANAGEDRESOURCES] [%s/%s]         Managed component: '%s'"),
+					(InComponent && InComponent->GetOwner()) ? *InComponent->GetOwner()->GetName() : TEXT("MISSINGCOMPONENT"),
+					(InComponent && InComponent->GetGraph()) ? *InComponent->GetGraph()->GetName() : TEXT("MISSINGGRAPH"),
+					*ManagedComponent->GeneratedComponent->GetReadableName());
 			}
 			else
 			{
-				UE_LOG(LogPCG, Warning, TEXT("[PCGMANAGEDRESOURCES] UPCGComponent::AddToManagedResources, managed component null or no owner"));
+				UE_LOG(LogPCG, Warning, TEXT("[PCGMANAGEDRESOURCES] [%s/%s]         NULL managed component or no owner"),
+					(InComponent && InComponent->GetOwner()) ? *InComponent->GetOwner()->GetName() : TEXT("MISSINGCOMPONENT"),
+					(InComponent && InComponent->GetGraph()) ? *InComponent->GetGraph()->GetName() : TEXT("MISSINGGRAPH"));
 			}
 		}
 		else
 		{
-			if (Resource)
+			if (InResource)
 			{
-				UE_LOG(LogPCG, Warning, TEXT("[PCGMANAGEDRESOURCES] UPCGComponent::AddToManagedResources, unidentified managed resource index"));
+				UE_LOG(LogPCG, Warning, TEXT("[PCGMANAGEDRESOURCES] [%s/%s]         Unidentified managed resource"),
+					(InComponent && InComponent->GetOwner()) ? *InComponent->GetOwner()->GetName() : TEXT("MISSINGCOMPONENT"),
+					(InComponent && InComponent->GetGraph()) ? *InComponent->GetGraph()->GetName() : TEXT("MISSINGGRAPH"));
 			}
 			else
 			{
-				UE_LOG(LogPCG, Warning, TEXT("[PCGMANAGEDRESOURCES] UPCGComponent::AddToManagedResources, encountered NULL managed resource index"));
+				UE_LOG(LogPCG, Warning, TEXT("[PCGMANAGEDRESOURCES] [%s/%s]         NULL unidentified resource"),
+					(InComponent && InComponent->GetOwner()) ? *InComponent->GetOwner()->GetName() : TEXT("MISSINGCOMPONENT"),
+					(InComponent && InComponent->GetGraph()) ? *InComponent->GetGraph()->GetName() : TEXT("MISSINGGRAPH"));
 			}
 		}
-#endif
 	}
 
-	void LogCleanupInternal(bool bRemoveComponents)
+	void LogResources(const UPCGComponent* InComponent, const TArray<UPCGManagedResource*>& Resources)
+	{
+		for (UPCGManagedResource* Resource : Resources)
+		{
+			LogResource(InComponent, Resource);
+		}
+	}
+#endif // WITH_EDITOR
+
+	void LogAddToManagedResources(const UPCGComponent* InComponent, UPCGManagedResource* InResource)
 	{
 #if WITH_EDITOR
 		if (!LogEnabled())
@@ -96,12 +119,15 @@ namespace PCGGeneratedResourcesLogging
 			return;
 		}
 
-		UE_LOG(LogPCG, Log, TEXT("[PCGMANAGEDRESOURCES]")); // Blank line
-		UE_LOG(LogPCG, Log, TEXT("[PCGMANAGEDRESOURCES] UPCGComponent::CleanupInternal, bRemoveComponents: %d"), bRemoveComponents);
+		UE_LOG(LogPCG, Log, TEXT("[PCGMANAGEDRESOURCES] [%s/%s] UPCGComponent::AddToManagedResources:"),
+			(InComponent && InComponent->GetOwner()) ? *InComponent->GetOwner()->GetName() : TEXT("MISSINGCOMPONENT"),
+			(InComponent && InComponent->GetGraph()) ? *InComponent->GetGraph()->GetName() : TEXT("MISSINGGRAPH"));
+
+		LogResource(InComponent, InResource);
 #endif
 	}
 
-	void LogCleanupLocalImmediate(bool bHardRelease, const TArray<UPCGManagedResource*>& GeneratedResources)
+	void LogCleanupInternal(const UPCGComponent* InComponent, bool bRemoveComponents)
 	{
 #if WITH_EDITOR
 		if (!LogEnabled())
@@ -110,11 +136,14 @@ namespace PCGGeneratedResourcesLogging
 		}
 
 		UE_LOG(LogPCG, Log, TEXT("[PCGMANAGEDRESOURCES]")); // Blank line
-		UE_LOG(LogPCG, Log, TEXT("[PCGMANAGEDRESOURCES] UPCGComponent::CleanupLocalImmediate BEGIN, bHardRelease: %d, GeneratedResources.Num() = %d"), bHardRelease, GeneratedResources.Num());
+		UE_LOG(LogPCG, Log, TEXT("[PCGMANAGEDRESOURCES] [%s/%s] UPCGComponent::CleanupInternal, bRemoveComponents: %d"),
+			(InComponent && InComponent->GetOwner()) ? *InComponent->GetOwner()->GetName() : TEXT("MISSINGCOMPONENT"),
+			(InComponent && InComponent->GetGraph()) ? *InComponent->GetGraph()->GetName() : TEXT("MISSINGGRAPH"),
+			bRemoveComponents);
 #endif
 	}
 
-	void LogCleanupLocalImmediateResource(const UPCGManagedResource* Resource)
+	void LogCleanupLocalImmediate(const UPCGComponent* InComponent, bool bHardRelease, const TArray<UPCGManagedResource*>& GeneratedResources)
 	{
 #if WITH_EDITOR
 		if (!LogEnabled())
@@ -122,55 +151,34 @@ namespace PCGGeneratedResourcesLogging
 			return;
 		}
 
-		if (const UPCGManagedActors* ManagedActors = Cast<UPCGManagedActors>(Resource))
-		{
-			UE_LOG(LogPCG, Log, TEXT("[PCGMANAGEDRESOURCES] UPCGComponent::CleanupLocalImmediate, actor count %d"), ManagedActors->GeneratedActors.Num());
+		UE_LOG(LogPCG, Log, TEXT("[PCGMANAGEDRESOURCES]")); // Blank line
+		UE_LOG(LogPCG, Log, TEXT("[PCGMANAGEDRESOURCES] [%s/%s] UPCGComponent::CleanupLocalImmediate BEGIN, bHardRelease: %d, GeneratedResources.Num() = %d"),
+			(InComponent && InComponent->GetOwner()) ? *InComponent->GetOwner()->GetName() : TEXT("MISSINGCOMPONENT"),
+			(InComponent && InComponent->GetGraph()) ? *InComponent->GetGraph()->GetName() : TEXT("MISSINGGRAPH"),
+			bHardRelease,
+			GeneratedResources.Num());
 
-			const int32 MaxPrint = CVarGenResourcesLoggingMaxPrintCount.GetValueOnAnyThread();
-			int32 Count = 0;
-			for (const TSoftObjectPtr<AActor>& Actor : ManagedActors->GeneratedActors)
-			{
-				if (Count++ >= MaxPrint)
-				{
-					break;
-				}
+		LogResources(InComponent, GeneratedResources);
+#endif
+	}
 
-				if (Actor.Get())
-				{
-					UE_LOG(LogPCG, Log, TEXT("[PCGMANAGEDRESOURCES]         |- '%s' (%d/%d)"), *Actor->GetFName().ToString(), Count, ManagedActors->GeneratedActors.Num());
-				}
-				else
-				{
-					UE_LOG(LogPCG, Log, TEXT("[PCGMANAGEDRESOURCES]         |- NULL (%d/%d)"), Count, ManagedActors->GeneratedActors.Num());
-				}
-			}
-		}
-		else if (const UPCGManagedComponent* ManagedComponent = Cast<UPCGManagedComponent>(Resource))
-		{
-			if (ManagedComponent->GeneratedComponent)
-			{
-				UE_LOG(LogPCG, Log, TEXT("[PCGMANAGEDRESOURCES] UPCGComponent::CleanupLocalImmediate, managed component on actor '%s'"), *ManagedComponent->GeneratedComponent->GetReadableName());
-			}
-			else
-			{
-				UE_LOG(LogPCG, Warning, TEXT("[PCGMANAGEDRESOURCES] UPCGComponent::CleanupLocalImmediate, managed component null or no owner"));
-			}
-		}
-		else
+	void LogCleanupLocalImmediateResource(const UPCGComponent* InComponent, UPCGManagedResource* Resource)
+	{
+#if WITH_EDITOR
+		if (!LogEnabled())
 		{
-			if (Resource)
-			{
-				UE_LOG(LogPCG, Warning, TEXT("[PCGMANAGEDRESOURCES] UPCGComponent::CleanupLocalImmediate, unidentified managed resource"));
-			}
-			else
-			{
-				UE_LOG(LogPCG, Warning, TEXT("[PCGMANAGEDRESOURCES] UPCGComponent::CleanupLocalImmediate, encountered NULL managed resource"));
-			}
+			return;
 		}
+
+		UE_LOG(LogPCG, Log, TEXT("[PCGMANAGEDRESOURCES] [%s/%s] UPCGComponent::CleanupLocalImmediate:"),
+			(InComponent && InComponent->GetOwner()) ? *InComponent->GetOwner()->GetName() : TEXT("MISSINGCOMPONENT"),
+			(InComponent && InComponent->GetGraph()) ? *InComponent->GetGraph()->GetName() : TEXT("MISSINGGRAPH"));
+		
+		LogResource(InComponent, Resource);
 #endif
 	}
 
-	void LogCleanupLocalImmediateFinished(const TArray<UPCGManagedResource*>& GeneratedResources)
+	void LogCleanupLocalImmediateFinished(const UPCGComponent* InComponent, const TArray<UPCGManagedResource*>& GeneratedResources)
 	{
 #if WITH_EDITOR
 		if (!LogEnabled())
@@ -178,11 +186,16 @@ namespace PCGGeneratedResourcesLogging
 			return;
 		}
 
-		UE_LOG(LogPCG, Log, TEXT("[PCGMANAGEDRESOURCES] UPCGComponent::CleanupLocalImmediate FINISHED, GeneratedResources.Num() = %d"), GeneratedResources.Num());
+		UE_LOG(LogPCG, Log, TEXT("[PCGMANAGEDRESOURCES] [%s/%s] UPCGComponent::CleanupLocalImmediate FINISHED, Final GeneratedResources (%d):"),
+			(InComponent && InComponent->GetOwner()) ? *InComponent->GetOwner()->GetName() : TEXT("MISSINGCOMPONENT"),
+			(InComponent && InComponent->GetGraph()) ? *InComponent->GetGraph()->GetName() : TEXT("MISSINGGRAPH"),
+			GeneratedResources.Num());
+
+		LogResources(InComponent, GeneratedResources);
 #endif
 	}
 
-	void LogCreateCleanupTask(bool bRemoveComponents)
+	void LogCreateCleanupTask(const UPCGComponent* InComponent, bool bRemoveComponents)
 	{
 #if WITH_EDITOR
 		if (!LogEnabled())
@@ -190,11 +203,14 @@ namespace PCGGeneratedResourcesLogging
 			return;
 		}
 
-		UE_LOG(LogPCG, Log, TEXT("[PCGMANAGEDRESOURCES] UPCGComponent::CreateCleanupTask, bRemoveComponents: %d"), bRemoveComponents);
+		UE_LOG(LogPCG, Log, TEXT("[PCGMANAGEDRESOURCES] [%s/%s] UPCGComponent::CreateCleanupTask, bRemoveComponents: %d"),
+			(InComponent && InComponent->GetOwner()) ? *InComponent->GetOwner()->GetName() : TEXT("MISSINGCOMPONENT"),
+			(InComponent && InComponent->GetGraph()) ? *InComponent->GetGraph()->GetName() : TEXT("MISSINGGRAPH"),
+			bRemoveComponents);
 #endif
 	}
 
-	void LogCreateCleanupTaskResource(const UPCGManagedResource* Resource/*, TArray<UPCGManagedResource*>& GeneratedResources*/)
+	void LogCreateCleanupTaskResource(const UPCGComponent* InComponent, UPCGManagedResource* Resource)
 	{
 #if WITH_EDITOR
 		if (!LogEnabled())
@@ -202,49 +218,30 @@ namespace PCGGeneratedResourcesLogging
 			return;
 		}
 
-		if (Resource)
+		UE_LOG(LogPCG, Log, TEXT("[PCGMANAGEDRESOURCES] [%s/%s] UPCGComponent::CreateCleanupTask::CleanupTask:"),
+			(InComponent && InComponent->GetOwner()) ? *InComponent->GetOwner()->GetName() : TEXT("MISSINGCOMPONENT"),
+			(InComponent && InComponent->GetGraph()) ? *InComponent->GetGraph()->GetName() : TEXT("MISSINGGRAPH"));
+
+		LogResource(InComponent, Resource);
+#endif
+	}
+
+	void LogCreateCleanupTaskFinished(const UPCGComponent* InComponent, const TArray<TObjectPtr<UPCGManagedResource>>* InGeneratedResources)
+	{
+#if WITH_EDITOR
+		if (!LogEnabled())
 		{
-			if (const UPCGManagedActors* ManagedActors = Cast<UPCGManagedActors>(Resource))
-			{
-				const int32 MaxPrint = CVarGenResourcesLoggingMaxPrintCount.GetValueOnAnyThread();
-				int32 Count = 0;
-				for (const TSoftObjectPtr<AActor>& Actor : ManagedActors->GeneratedActors)
-				{
-					if (Count++ >= MaxPrint)
-					{
-						break;
-					}
-
-					if (Actor.Get())
-					{
-						UE_LOG(LogPCG, Log, TEXT("[PCGMANAGEDRESOURCES]         |- '%s' (%d/%d)"), *Actor->GetFName().ToString(), Count, ManagedActors->GeneratedActors.Num());
-					}
-					else
-					{
-						UE_LOG(LogPCG, Log, TEXT("[PCGMANAGEDRESOURCES]         |- NULL (%d/%d)"), Count, ManagedActors->GeneratedActors.Num());
-					}
-				}
-			}
-			else if (const UPCGManagedComponent* ManagedComponent = Cast<UPCGManagedComponent>(Resource))
-			{
-				if (ManagedComponent->GeneratedComponent)
-				{
-					UE_LOG(LogPCG, Log, TEXT("[PCGMANAGEDRESOURCES] UPCGComponent::CreateCleanupTask::CleanupTask, managed component on actor '%s'"), *ManagedComponent->GeneratedComponent->GetReadableName());
-				}
-				else
-				{
-					UE_LOG(LogPCG, Warning, TEXT("[PCGMANAGEDRESOURCES] UPCGComponent::CreateCleanupTask::CleanupTask, managed component null or no owner"));
-				}
-			}
-			else
-			{
-				UE_LOG(LogPCG, Log, TEXT("[PCGMANAGEDRESOURCES] UPCGComponent::CreateCleanupTask::CleanupTask, unidentified resource or null"));
-			}
+			return;
 		}
+
+		UE_LOG(LogPCG, Log, TEXT("[PCGMANAGEDRESOURCES] [%s/%s] UPCGComponent::CreateCleanupTask::CleanupTask FINISHED, GeneratedResources.Num() = %d"),
+			(InComponent && InComponent->GetOwner()) ? *InComponent->GetOwner()->GetName() : TEXT("MISSINGCOMPONENT"),
+			(InComponent && InComponent->GetGraph()) ? *InComponent->GetGraph()->GetName() : TEXT("MISSINGGRAPH"),
+			InGeneratedResources ? InGeneratedResources->Num() : -1);
 #endif
 	}
 
-	void LogCreateCleanupTaskFinished(const TArray<UPCGManagedResource*>& GeneratedResources)
+	void LogCreateCleanupTaskFinished(const TArray<UPCGManagedResource*>* InGeneratedResources)
 	{
 #if WITH_EDITOR
 		if (!LogEnabled())
@@ -252,11 +249,12 @@ namespace PCGGeneratedResourcesLogging
 			return;
 		}
 
-		UE_LOG(LogPCG, Log, TEXT("[PCGMANAGEDRESOURCES] UPCGComponent::CreateCleanupTask::CleanupTask FINISHED, GeneratedResources.Num() = %d"), GeneratedResources.Num());
+		UE_LOG(LogPCG, Log, TEXT("[PCGMANAGEDRESOURCES] UPCGComponent::CreateCleanupTask::CleanupTask FINISHED, GeneratedResources.Num() = %d"),
+			InGeneratedResources ? InGeneratedResources->Num() : -1);
 #endif
 	}
 
-	void LogCleanupUnusedManagedResources(const TArray<UPCGManagedResource*>& GeneratedResources)
+	void LogCleanupUnusedManagedResources(const UPCGComponent* InComponent, const TArray<UPCGManagedResource*>& GeneratedResources)
 	{
 #if WITH_EDITOR
 		if (!LogEnabled())
@@ -264,11 +262,14 @@ namespace PCGGeneratedResourcesLogging
 			return;
 		}
 
-		UE_LOG(LogPCG, Log, TEXT("[PCGMANAGEDRESOURCES] UPCGComponent::CleanupUnusedManagedResources BEGIN, GeneratedResources.Num() = %d"), GeneratedResources.Num());
+		UE_LOG(LogPCG, Log, TEXT("[PCGMANAGEDRESOURCES] [%s/%s] UPCGComponent::CleanupUnusedManagedResources BEGIN, GeneratedResources.Num() = %d"),
+			(InComponent && InComponent->GetOwner()) ? *InComponent->GetOwner()->GetName() : TEXT("MISSINGCOMPONENT"),
+			(InComponent && InComponent->GetGraph()) ? *InComponent->GetGraph()->GetName() : TEXT("MISSINGGRAPH"),
+			GeneratedResources.Num());
 #endif
 	}
 
-	void LogCleanupUnusedManagedResourcesResource(const UPCGManagedResource* Resource)
+	void LogCleanupUnusedManagedResourcesResource(const UPCGComponent* InComponent, UPCGManagedResource* InResource)
 	{
 #if WITH_EDITOR
 		if (!LogEnabled())
@@ -276,49 +277,32 @@ namespace PCGGeneratedResourcesLogging
 			return;
 		}
 
-		if (Resource)
+		UE_LOG(LogPCG, Log, TEXT("[PCGMANAGEDRESOURCES] [%s/%s] UPCGComponent::CleanupUnusedManagedResources:"),
+			(InComponent && InComponent->GetOwner()) ? *InComponent->GetOwner()->GetName() : TEXT("MISSINGCOMPONENT"),
+			(InComponent && InComponent->GetGraph()) ? *InComponent->GetGraph()->GetName() : TEXT("MISSINGGRAPH"));
+
+		LogResource(InComponent, InResource);
+#endif
+	}
+
+	void LogCleanupUnusedManagedResourcesFinished(const UPCGComponent* InComponent, const TArray<UPCGManagedResource*>& GeneratedResources)
+	{
+#if WITH_EDITOR
+		if (!LogEnabled())
 		{
-			if (const UPCGManagedActors* ManagedActors = Cast<UPCGManagedActors>(Resource))
-			{
-				const int32 MaxPrint = CVarGenResourcesLoggingMaxPrintCount.GetValueOnAnyThread();
-				int32 Count = 0;
-				for (const TSoftObjectPtr<AActor>& Actor : ManagedActors->GeneratedActors)
-				{
-					if (Count++ >= MaxPrint)
-					{
-						break;
-					}
-
-					if (Actor.Get())
-					{
-						UE_LOG(LogPCG, Log, TEXT("[PCGMANAGEDRESOURCES]         |- '%s' (%d/%d)"), *Actor->GetFName().ToString(), Count, ManagedActors->GeneratedActors.Num());
-					}
-					else
-					{
-						UE_LOG(LogPCG, Log, TEXT("[PCGMANAGEDRESOURCES]         |- NULL (%d/%d)"), Count, ManagedActors->GeneratedActors.Num());
-					}
-				}
-			}
-			else if (const UPCGManagedComponent* ManagedComponent = Cast<UPCGManagedComponent>(Resource))
-			{
-				if (ManagedComponent->GeneratedComponent)
-				{
-					UE_LOG(LogPCG, Log, TEXT("[PCGMANAGEDRESOURCES] UPCGComponent::CleanupUnusedManagedResources, managed component on actor '%s'"), *ManagedComponent->GeneratedComponent->GetReadableName());
-				}
-				else
-				{
-					UE_LOG(LogPCG, Warning, TEXT("[PCGMANAGEDRESOURCES] UPCGComponent::CleanupUnusedManagedResources, managed component null or no owner"));
-				}
-			}
-			else
-			{
-				UE_LOG(LogPCG, Log, TEXT("[PCGMANAGEDRESOURCES] UPCGComponent::CleanupUnusedManagedResources, unidentified resource or null"));
-			}
+			return;
 		}
+
+		UE_LOG(LogPCG, Log, TEXT("[PCGMANAGEDRESOURCES] [%s/%s] UPCGComponent::CleanupUnusedManagedResources, FINISHED:"),
+			(InComponent && InComponent->GetOwner()) ? *InComponent->GetOwner()->GetName() : TEXT("MISSINGCOMPONENT"),
+			(InComponent && InComponent->GetGraph()) ? *InComponent->GetGraph()->GetName() : TEXT("MISSINGGRAPH"),
+			GeneratedResources.Num());
+
+		LogResources(InComponent, GeneratedResources);
 #endif
 	}
 
-	void LogCleanupUnusedManagedResourcesFinished(const TArray<UPCGManagedResource*>& GeneratedResources)
+	void LogManagedResourceSoftRelease(UPCGManagedResource* InResource)
 	{
 #if WITH_EDITOR
 		if (!LogEnabled())
@@ -326,11 +310,17 @@ namespace PCGGeneratedResourcesLogging
 			return;
 		}
 
-		UE_LOG(LogPCG, Log, TEXT("[PCGMANAGEDRESOURCES] UPCGComponent::CleanupUnusedManagedResources, FINISHED, GeneratedResources.Num() = %d"), GeneratedResources.Num());
+		const UPCGComponent* Component = InResource ? Cast<UPCGComponent>(InResource->GetOuter()) : nullptr;
+
+		UE_LOG(LogPCG, Log, TEXT("[PCGMANAGEDRESOURCES] [%s/%s] UPCGManagedResource::Release, SOFT release:"),
+			(Component && Component->GetOwner()) ? *Component->GetOwner()->GetName() : TEXT("MISSINGCOMPONENT"),
+			(Component && Component->GetGraph()) ? *Component->GetGraph()->GetName() : TEXT("MISSINGGRAPH"));
+
+		LogResource(Component, InResource);
 #endif
 	}
 
-	void LogManagedActorsSoftRelease(const TSet<TSoftObjectPtr<AActor>>& GeneratedActors)
+	void LogManagedResourceHardRelease(UPCGManagedResource* InResource)
 	{
 #if WITH_EDITOR
 		if (!LogEnabled())
@@ -338,30 +328,17 @@ namespace PCGGeneratedResourcesLogging
 			return;
 		}
 
-		UE_LOG(LogPCG, Log, TEXT("[PCGMANAGEDRESOURCES] UPCGManagedActors::Release, soft release, actor count %d"), GeneratedActors.Num());
+		const UPCGComponent* Component = InResource ? Cast<UPCGComponent>(InResource->GetOuter()) : nullptr;
 
-		const int32 MaxPrint = CVarGenResourcesLoggingMaxPrintCount.GetValueOnAnyThread();
-		int32 Count = 0;
-		for (const TSoftObjectPtr<AActor>& Actor : GeneratedActors)
-		{
-			if (Count++ >= MaxPrint)
-			{
-				break;
-			}
+		UE_LOG(LogPCG, Log, TEXT("[PCGMANAGEDRESOURCES] [%s/%s] UPCGManagedResource::Release, HARD release:"),
+			(Component && Component->GetOwner()) ? *Component->GetOwner()->GetName() : TEXT("MISSINGCOMPONENT"),
+			(Component && Component->GetGraph()) ? *Component->GetGraph()->GetName() : TEXT("MISSINGGRAPH"));
 
-			if (!Actor.Get())
-			{
-				UE_LOG(LogPCG, Log, TEXT("[PCGMANAGEDRESOURCES]         |- NULL (%d/%d)"), Count, GeneratedActors.Num());
-			}
-			else
-			{
-				UE_LOG(LogPCG, Log, TEXT("[PCGMANAGEDRESOURCES]         |- '%s' (%d/%d)"), *Actor->GetFName().ToString(), Count, GeneratedActors.Num());
-			}
-		}
+		LogResource(Component, InResource);
 #endif
 	}
 
-	void LogManagedActorsHardRelease(const TSet<TSoftObjectPtr<AActor>>& ActorsToDelete)
+	void LogManagedActorsRelease(const UPCGManagedResource* InResource, const TSet<TSoftObjectPtr<AActor>>& ActorsToDelete, bool bHardRelease, bool bOnlyMarkedForCleanup)
 	{
 #if WITH_EDITOR
 		if (!LogEnabled())
@@ -369,7 +346,16 @@ namespace PCGGeneratedResourcesLogging
 			return;
 		}
 
-		UE_LOG(LogPCG, Log, TEXT("[PCGMANAGEDRESOURCES] UPCGManagedActors::Release, hard release, actor count %d scheduled for delete"), ActorsToDelete.Num());
+		const UPCGComponent* Component = InResource ? Cast<UPCGComponent>(InResource->GetOuter()) : nullptr;
+		const bool bMarkedTransientOnLoad = InResource && InResource->IsMarkedTransientOnLoad();
+
+		UE_LOG(LogPCG, Log, TEXT("[PCGMANAGEDRESOURCES] [%s/%s] UPCGManagedActors::Release, %s release, bMarkedTransientOnLoad: %d, bOnlyMarkedForCleanup: %d, actor count %d scheduled for delete"),
+			(Component && Component->GetOwner()) ? *Component->GetOwner()->GetName() : TEXT("MISSINGCOMPONENT"),
+			(Component && Component->GetGraph()) ? *Component->GetGraph()->GetName() : TEXT("MISSINGGRAPH"),
+			bHardRelease ? TEXT("HARD") : TEXT("SOFT"),
+			bMarkedTransientOnLoad,
+			bOnlyMarkedForCleanup,
+			ActorsToDelete.Num());
 
 		const uint32 MaxPrint = 3;
 		uint32 Count = 0;
@@ -382,13 +368,56 @@ namespace PCGGeneratedResourcesLogging
 
 			if (ActorToDelete)
 			{
-				UE_LOG(LogPCG, Log, TEXT("[PCGMANAGEDRESOURCES]         |- '%s' (%d/%d)"), *ActorToDelete->GetFName().ToString(), Count, ActorsToDelete.Num());
+				UE_LOG(LogPCG, Log, TEXT("[PCGMANAGEDRESOURCES] [%s/%s]                 |- '%s' (%d/%d)"),
+					(Component && Component->GetOwner()) ? *Component->GetOwner()->GetName() : TEXT("MISSINGCOMPONENT"),
+					(Component && Component->GetGraph()) ? *Component->GetGraph()->GetName() : TEXT("MISSINGGRAPH"),
+					*ActorToDelete->GetFName().ToString(),
+					Count,
+					ActorsToDelete.Num());
 			}
 			else
 			{
-				UE_LOG(LogPCG, Log, TEXT("[PCGMANAGEDRESOURCES]         |- NULL (%d/%d)"), Count, ActorsToDelete.Num());
+				UE_LOG(LogPCG, Log, TEXT("[PCGMANAGEDRESOURCES] [%s/%s]                 |- NULL (%d/%d)"),
+					(Component && Component->GetOwner()) ? *Component->GetOwner()->GetName() : TEXT("MISSINGCOMPONENT"),
+					(Component && Component->GetGraph()) ? *Component->GetGraph()->GetName() : TEXT("MISSINGGRAPH"),
+					Count,
+					ActorsToDelete.Num());
 			}
 		}
+#endif
+	}
+
+	void LogManagedComponentHidden(UPCGManagedComponent* InResource)
+	{
+#if WITH_EDITOR
+		if (!LogEnabled())
+		{
+			return;
+		}
+
+		const UPCGComponent* Component = InResource ? Cast<UPCGComponent>(InResource->GetOuter()) : nullptr;
+
+		UE_LOG(LogPCG, Log, TEXT("[PCGMANAGEDRESOURCES] [%s/%s] UPCGManagedComponent::Release, hidden, component:"),
+			(Component && Component->GetOwner()) ? *Component->GetOwner()->GetName() : TEXT("MISSINGCOMPONENT"),
+			(Component && Component->GetGraph()) ? *Component->GetGraph()->GetName() : TEXT("MISSINGGRAPH"));
+
+		LogResource(Component, InResource);
+#endif
+	}
+
+	void LogManagedComponentDeleteNull(UPCGManagedComponent* InResource)
+	{
+#if WITH_EDITOR
+		if (!LogEnabled())
+		{
+			return;
+		}
+
+		const UPCGComponent* Component = InResource ? Cast<UPCGComponent>(InResource->GetOuter()) : nullptr;
+
+		UE_LOG(LogPCG, Log, TEXT("[PCGMANAGEDRESOURCES] [%s/%s] UPCGManagedComponent::Release, delete null component"),
+			(Component && Component->GetOwner()) ? *Component->GetOwner()->GetName() : TEXT("MISSINGCOMPONENT"),
+			(Component && Component->GetGraph()) ? *Component->GetGraph()->GetName() : TEXT("MISSINGGRAPH"));
 #endif
 	}
 }
diff --git a/Engine/Plugins/PCG/Source/PCG/Private/Utils/PCGGeneratedResourcesLogging.h b/Engine/Plugins/PCG/Source/PCG/Private/Utils/PCGGeneratedResourcesLogging.h
--- a/Engine/Plugins/PCG/Source/PCG/Private/Utils/PCGGeneratedResourcesLogging.h
+++ b/Engine/Plugins/PCG/Source/PCG/Private/Utils/PCGGeneratedResourcesLogging.h
@@ -7,27 +7,32 @@
 #include "UObject/SoftObjectPtr.h"
 
 class UPCGComponent;
+class UPCGManagedComponent;
 class UPCGManagedResource;
 class AActor;
 
 namespace PCGGeneratedResourcesLogging
 {
-	void LogAddToManagedResources(const UPCGManagedResource* Resource);
+	void LogAddToManagedResources(const UPCGComponent* InComponent, UPCGManagedResource* InResource);
 	
-	void LogCleanupInternal(bool bRemoveComponents);
+	void LogCleanupInternal(const UPCGComponent* InComponent, bool bRemoveComponents);
 	
-	void LogCleanupLocalImmediate(bool bHardRelease, const TArray<UPCGManagedResource*>& GeneratedResources);
-	void LogCleanupLocalImmediateResource(const UPCGManagedResource* Resource);
-	void LogCleanupLocalImmediateFinished(const TArray<UPCGManagedResource*>& GeneratedResources);
+	void LogCleanupLocalImmediate(const UPCGComponent* InComponent, bool bHardRelease, const TArray<UPCGManagedResource*>& GeneratedResources);
+	void LogCleanupLocalImmediateResource(const UPCGComponent* InComponent, UPCGManagedResource* InResource);
+	void LogCleanupLocalImmediateFinished(const UPCGComponent* InComponent, const TArray<UPCGManagedResource*>& GeneratedResources);
 	
-	void LogCreateCleanupTask(bool bRemoveComponents);
-	void LogCreateCleanupTaskResource(const UPCGManagedResource* Resource);
-	void LogCreateCleanupTaskFinished(const TArray<UPCGManagedResource*>& GeneratedResources);
+	void LogCreateCleanupTask(const UPCGComponent* InComponent, bool bRemoveComponents);
+	void LogCreateCleanupTaskResource(const UPCGComponent* InComponent, UPCGManagedResource* InResource);
+	void LogCreateCleanupTaskFinished(const UPCGComponent* InComponent, const TArray<TObjectPtr<UPCGManagedResource>>* InGeneratedResources);
 
-	void LogCleanupUnusedManagedResources(const TArray<UPCGManagedResource*>& GeneratedResources);
-	void LogCleanupUnusedManagedResourcesResource(const UPCGManagedResource* Resource);
-	void LogCleanupUnusedManagedResourcesFinished(const TArray<UPCGManagedResource*>& GeneratedResources);
+	void LogCleanupUnusedManagedResources(const UPCGComponent* InComponent, const TArray<UPCGManagedResource*>& GeneratedResources);
+	void LogCleanupUnusedManagedResourcesResource(const UPCGComponent* InComponent, UPCGManagedResource* Resource);
+	void LogCleanupUnusedManagedResourcesFinished(const UPCGComponent* InComponent, const TArray<UPCGManagedResource*>& GeneratedResources);
 
-	void LogManagedActorsSoftRelease(const TSet<TSoftObjectPtr<AActor>>& GeneratedActors);
-	void LogManagedActorsHardRelease(const TSet<TSoftObjectPtr<AActor>>& ActorsToDelete);
+	void LogManagedActorsRelease(const UPCGManagedResource* InResource, const TSet<TSoftObjectPtr<AActor>>& ActorsToDelete, bool bHardRelease, bool bOnlyMarkedForCleanup);
+
+	void LogManagedResourceSoftRelease(UPCGManagedResource* InResource);
+	void LogManagedResourceHardRelease(UPCGManagedResource* InResource);
+	void LogManagedComponentHidden(UPCGManagedComponent* InResource);
+	void LogManagedComponentDeleteNull(UPCGManagedComponent* InResource);
 }
diff --git a/Engine/Plugins/PCG/Source/PCG/Private/Utils/PCGGraphExecutionLogging.cpp b/Engine/Plugins/PCG/Source/PCG/Private/Utils/PCGGraphExecutionLogging.cpp
--- a/Engine/Plugins/PCG/Source/PCG/Private/Utils/PCGGraphExecutionLogging.cpp
+++ b/Engine/Plugins/PCG/Source/PCG/Private/Utils/PCGGraphExecutionLogging.cpp
@@ -13,7 +13,7 @@
 
 namespace PCGGraphExecutionLogging
 {
-#if WITH_EDITOR
+#if !(UE_BUILD_SHIPPING || UE_BUILD_TEST) || USE_LOGGING_IN_SHIPPING
 	static TAutoConsoleVariable<bool> CVarGraphExecutionLoggingEnable(
 		TEXT("pcg.GraphExecution.EnableLogging"),
 		false,
@@ -27,7 +27,7 @@ namespace PCGGraphExecutionLogging
 
 	bool LogEnabled()
 	{
-#if WITH_EDITOR
+#if !(UE_BUILD_SHIPPING || UE_BUILD_TEST) || USE_LOGGING_IN_SHIPPING
 		return CVarGraphExecutionLoggingEnable.GetValueOnAnyThread();
 #else
 		return false;
@@ -36,16 +36,29 @@ namespace PCGGraphExecutionLogging
 
 	bool CullingLogEnabled()
 	{
-#if WITH_EDITOR
+#if !(UE_BUILD_SHIPPING || UE_BUILD_TEST) || USE_LOGGING_IN_SHIPPING
 		return CVarGraphExecutionCullingLoggingEnable.GetValueOnAnyThread();
 #else
 		return false;
 #endif
 	}
 
-	void LogGraphTask(FPCGTaskId TaskId, const FPCGGraphTask& Task, const TSet<FPCGTaskId>* SuccessorIds)
+#if !(UE_BUILD_SHIPPING || UE_BUILD_TEST) || USE_LOGGING_IN_SHIPPING
+	FString GetOwnerName(const UPCGComponent* InComponent)
 	{
+		return (InComponent && InComponent->GetOwner()) ?
 #if WITH_EDITOR
+			InComponent->GetOwner()->GetActorLabel() :
+#else
+			InComponent->GetOwner()->GetName() :
+#endif
+			FString(TEXT("MISSINGOWNER"));
+	}
+#endif // !(UE_BUILD_SHIPPING || UE_BUILD_TEST) || USE_LOGGING_IN_SHIPPING
+
+	void LogGraphTask(FPCGTaskId TaskId, const FPCGGraphTask& Task, const TSet<FPCGTaskId>* SuccessorIds)
+	{
+#if !(UE_BUILD_SHIPPING || UE_BUILD_TEST) || USE_LOGGING_IN_SHIPPING
 		if (!LogEnabled())
 		{
 			return;
@@ -81,20 +94,27 @@ namespace PCGGraphExecutionLogging
 			}
 		}
 
+		const FString PinDependencyString =
+#if WITH_EDITOR
+			Task.PinDependency.ToString();
+#else
+			TEXT("MISSINGPINDEPS");
+#endif
+
 		UE_LOG(LogPCG, Log, TEXT("\t\tID: %u\tParent: %u\tNode: %s\tInputs: %s\tPinDeps: %s\tSuccessors: %s"),
 			TaskId,
 			Task.ParentId != InvalidPCGTaskId ? Task.ParentId : 0,
 			Task.Node ? (*Task.Node->GetNodeTitle(EPCGNodeTitleType::ListView).ToString()) : TEXT("NULL"),
 			*GenerateInputsString(Task.Inputs),
-			*Task.PinDependency.ToString(),
+			*PinDependencyString,
 			*SuccessorsString
 		);
-#endif
+#endif // !(UE_BUILD_SHIPPING || UE_BUILD_TEST) || USE_LOGGING_IN_SHIPPING
 	}
 
 	void LogGraphTasks(const TMap<FPCGTaskId, FPCGGraphTask>& Tasks, const TMap<FPCGTaskId, TSet<FPCGTaskId>>* TaskSuccessors)
 	{
-#if WITH_EDITOR
+#if !(UE_BUILD_SHIPPING || UE_BUILD_TEST) || USE_LOGGING_IN_SHIPPING
 		if (!LogEnabled())
 		{
 			return;
@@ -104,12 +124,12 @@ namespace PCGGraphExecutionLogging
 		{
 			PCGGraphExecutionLogging::LogGraphTask(TaskIdAndTask.Key, TaskIdAndTask.Value, TaskSuccessors ? TaskSuccessors->Find(TaskIdAndTask.Value.NodeId) : nullptr);
 		}
-#endif
+#endif // !(UE_BUILD_SHIPPING || UE_BUILD_TEST) || USE_LOGGING_IN_SHIPPING
 	}
 
 	void LogGraphTasks(const TArray<FPCGGraphTask>& Tasks)
 	{
-#if WITH_EDITOR
+#if !(UE_BUILD_SHIPPING || UE_BUILD_TEST) || USE_LOGGING_IN_SHIPPING
 		if (!LogEnabled())
 		{
 			return;
@@ -119,48 +139,69 @@ namespace PCGGraphExecutionLogging
 		{
 			LogGraphTask(Task.NodeId, Task);
 		}
-#endif
+#endif // !(UE_BUILD_SHIPPING || UE_BUILD_TEST) || USE_LOGGING_IN_SHIPPING
 	}
 
-	void LogGraphSchedule(const UPCGComponent* SourceComponent)
+	void LogGraphSchedule(const UPCGComponent* SourceComponent, const UPCGGraph* InScheduledGraph)
 	{
+#if !(UE_BUILD_SHIPPING || UE_BUILD_TEST) || USE_LOGGING_IN_SHIPPING
 		if (!LogEnabled() && !IsRunningCommandlet())
 		{
 			return;
 		}
 
-		UE_LOG(LogPCG, Display, TEXT("[%s/%s] --- SCHEDULE GRAPH ---"),
+		UE_LOG(LogPCG, Display, TEXT("[%s/%s] --- SCHEDULE GRAPH %s ---"),
 			(SourceComponent && SourceComponent->GetOwner()) ? *SourceComponent->GetOwner()->GetName() : TEXT("MISSINGCOMPONENT"),
-			(SourceComponent && SourceComponent->GetGraph()) ? *SourceComponent->GetGraph()->GetName() : TEXT("MISSINGGRAPH"));
+			(SourceComponent && SourceComponent->GetGraph()) ? *SourceComponent->GetGraph()->GetName() : TEXT("MISSINGGRAPH"),
+			InScheduledGraph ? *InScheduledGraph->GetName() : TEXT("MISSINGGRAPH)"));
+#endif // !(UE_BUILD_SHIPPING || UE_BUILD_TEST) || USE_LOGGING_IN_SHIPPING
 	}
 
-	void LogGraphScheduleDependency(const UPCGComponent* InComponent)
+	void LogGraphScheduleDependency(const UPCGComponent* InComponent, const FPCGStack* InFromStack)
 	{
+#if !(UE_BUILD_SHIPPING || UE_BUILD_TEST) || USE_LOGGING_IN_SHIPPING
 		if (!LogEnabled() && !IsRunningCommandlet())
 		{
 			return;
 		}
 
-		UE_LOG(LogPCG, Display, TEXT("[%s/%s] --- SCHEDULE GRAPH FOR DEPENDENCY ---"),
+		FString FromStackString;
+		if (InFromStack)
+		{
+			InFromStack->CreateStackFramePath(FromStackString);
+		}
+
+		UE_LOG(LogPCG, Display, TEXT("[%s/%s] --- SCHEDULE GRAPH FOR DEPENDENCY, from stack: %s"),
 			(InComponent && InComponent->GetOwner()) ? *InComponent->GetOwner()->GetName() : TEXT("MISSINGCOMPONENT"),
-			(InComponent && InComponent->GetGraph()) ? *InComponent->GetGraph()->GetName() : TEXT("MISSINGGRAPH"));
+			(InComponent && InComponent->GetGraph()) ? *InComponent->GetGraph()->GetName() : TEXT("MISSINGGRAPH"),
+			*FromStackString);
+#endif // !(UE_BUILD_SHIPPING || UE_BUILD_TEST) || USE_LOGGING_IN_SHIPPING
 	}
 
-	void LogGraphScheduleDependencyFailed(const UPCGComponent* InComponent)
+	void LogGraphScheduleDependencyFailed(const UPCGComponent* InComponent, const FPCGStack* InFromStack)
 	{
+#if !(UE_BUILD_SHIPPING || UE_BUILD_TEST) || USE_LOGGING_IN_SHIPPING
 		if (!LogEnabled() && !IsRunningCommandlet())
 		{
 			return;
 		}
 
-		UE_LOG(LogPCG, Warning, TEXT("[%s/%s] Failed to schedule dependency"),
+		FString FromStackString;
+		if (InFromStack)
+		{
+			InFromStack->CreateStackFramePath(FromStackString);
+		}
+
+		UE_LOG(LogPCG, Warning, TEXT("[%s/%s] Failed to schedule dependency, from stack: %s"),
 			(InComponent && InComponent->GetOwner()) ? *InComponent->GetOwner()->GetName() : TEXT("MISSINGCOMPONENT"),
-			(InComponent && InComponent->GetGraph()) ? *InComponent->GetGraph()->GetName() : TEXT("MISSINGGRAPH"));
+			(InComponent && InComponent->GetGraph()) ? *InComponent->GetGraph()->GetName() : TEXT("MISSINGGRAPH"),
+			*FromStackString);
+#endif // !(UE_BUILD_SHIPPING || UE_BUILD_TEST) || USE_LOGGING_IN_SHIPPING
 	}
 	
 	void LogGraphPostSchedule(const TMap<FPCGTaskId, FPCGGraphTask>& Tasks, const TMap<FPCGTaskId, TSet<FPCGTaskId>>& TaskSuccessors)
 	{
-#if WITH_EDITOR
+#if !(UE_BUILD_SHIPPING || UE_BUILD_TEST) || USE_LOGGING_IN_SHIPPING
 		if (!CullingLogEnabled())
 		{
 			return;
@@ -169,11 +210,12 @@ namespace PCGGraphExecutionLogging
 		UE_LOG(LogPCG, Log, TEXT("POST SCHEDULE:"));
 
 		LogGraphTasks(Tasks, &TaskSuccessors);
-#endif
+#endif // !(UE_BUILD_SHIPPING || UE_BUILD_TEST) || USE_LOGGING_IN_SHIPPING
 	}
 
 	void LogPostProcessGraph(const UPCGComponent* InSourceComponent)
 	{
+#if !(UE_BUILD_SHIPPING || UE_BUILD_TEST) || USE_LOGGING_IN_SHIPPING
 		if (!LogEnabled() && !IsRunningCommandlet())
 		{
 			return;
@@ -182,10 +224,12 @@ namespace PCGGraphExecutionLogging
 		UE_LOG(LogPCG, Display, TEXT("[%s/%s] UPCGComponent::PostProcessGraph"),
 			(InSourceComponent && InSourceComponent->GetOwner()) ? *InSourceComponent->GetOwner()->GetName() : TEXT("MISSINGCOMPONENT"),
 			(InSourceComponent && InSourceComponent->GetGraph()) ? *InSourceComponent->GetGraph()->GetName() : TEXT("MISSINGGRAPH"));
+#endif // !(UE_BUILD_SHIPPING || UE_BUILD_TEST) || USE_LOGGING_IN_SHIPPING
 	}
 
 	void LogComponentCancellation(const TSet<UPCGComponent*>& CancelledComponents)
 	{
+#if !(UE_BUILD_SHIPPING || UE_BUILD_TEST) || USE_LOGGING_IN_SHIPPING
 		if (!LogEnabled() && !IsRunningCommandlet())
 		{
 			return;
@@ -197,11 +241,12 @@ namespace PCGGraphExecutionLogging
 				(Component && Component->GetOwner()) ? *Component->GetOwner()->GetName() : TEXT("MISSINGCOMPONENT"),
 				(Component && Component->GetGraph()) ? *Component->GetGraph()->GetName() : TEXT("MISSINGGRAPH"));
 		}
+#endif // !(UE_BUILD_SHIPPING || UE_BUILD_TEST) || USE_LOGGING_IN_SHIPPING
 	}
 
 	void LogChangeOriginIgnoredForComponent(const UObject* InObject, const UPCGComponent* InComponent)
 	{
-#if WITH_EDITOR
+#if !(UE_BUILD_SHIPPING || UE_BUILD_TEST) || USE_LOGGING_IN_SHIPPING
 		if (!LogEnabled())
 		{
 			return;
@@ -211,22 +256,22 @@ namespace PCGGraphExecutionLogging
 			(InComponent && InComponent->GetOwner()) ? *InComponent->GetOwner()->GetName() : TEXT("MISSINGCOMPONENT"),
 			(InComponent && InComponent->GetGraph()) ? *InComponent->GetGraph()->GetName() : TEXT("MISSINGGRAPH"),
 			InObject ? *InObject->GetName() : TEXT("MISSINGOBJECT"));
-#endif
+#endif // !(UE_BUILD_SHIPPING || UE_BUILD_TEST) || USE_LOGGING_IN_SHIPPING
 	}
 
 	void LogGraphExecuteFrameFinished()
 	{
-#if WITH_EDITOR
+#if !(UE_BUILD_SHIPPING || UE_BUILD_TEST) || USE_LOGGING_IN_SHIPPING
 		if (!LogEnabled())
 		{
 			return;
 		}
 
 		UE_LOG(LogPCG, Log, TEXT("--- FINISH FPCGGRAPHEXECUTOR::EXECUTE ---"));
-#endif
+#endif // !(UE_BUILD_SHIPPING || UE_BUILD_TEST) || USE_LOGGING_IN_SHIPPING
 	}
 
-#if WITH_EDITOR
+#if !(UE_BUILD_SHIPPING || UE_BUILD_TEST) || USE_LOGGING_IN_SHIPPING
 	FString GetPinsToDeactivateString(const TArray<FPCGPinId>& PinIdsToDeactivate)
 	{
 		FString PinIdsToDeactivateString;
@@ -242,11 +287,11 @@ namespace PCGGraphExecutionLogging
 
 		return PinIdsToDeactivateString;
 	}
-#endif
+#endif // !(UE_BUILD_SHIPPING || UE_BUILD_TEST) || USE_LOGGING_IN_SHIPPING
 
 	void LogTaskExecute(const FPCGGraphTask& Task)
 	{
-#if WITH_EDITOR
+#if !(UE_BUILD_SHIPPING || UE_BUILD_TEST) || USE_LOGGING_IN_SHIPPING
 		if (!LogEnabled() || !Task.SourceComponent.Get())
 		{
 			return;
@@ -256,12 +301,12 @@ namespace PCGGraphExecutionLogging
 			*Task.SourceComponent->GetOwner()->GetName(),
 			Task.SourceComponent->GetGraph() ? *Task.SourceComponent->GetGraph()->GetName() : TEXT("MISSINGGRAPH"),
 			*FString::Printf(TEXT("%u'%s'"), Task.NodeId, Task.Node ? *Task.Node->GetNodeTitle(EPCGNodeTitleType::ListView).ToString() : TEXT("")));
-#endif
+#endif // !(UE_BUILD_SHIPPING || UE_BUILD_TEST) || USE_LOGGING_IN_SHIPPING
 	}
 
 	void LogTaskExecuteCachingDisabled(const FPCGGraphTask& Task)
 	{
-#if WITH_EDITOR
+#if !(UE_BUILD_SHIPPING || UE_BUILD_TEST) || USE_LOGGING_IN_SHIPPING
 		if (!LogEnabled() || !Task.SourceComponent.Get())
 		{
 			return;
@@ -271,12 +316,12 @@ namespace PCGGraphExecutionLogging
 			*Task.SourceComponent->GetOwner()->GetName(),
 			Task.SourceComponent->GetGraph() ? *Task.SourceComponent->GetGraph()->GetName() : TEXT("MISSINGGRAPH"),
 			*FString::Printf(TEXT("%u'%s'"), Task.NodeId, Task.Node ? *Task.Node->GetNodeTitle(EPCGNodeTitleType::ListView).ToString() : TEXT("")));
-#endif
+#endif // !(UE_BUILD_SHIPPING || UE_BUILD_TEST) || USE_LOGGING_IN_SHIPPING
 	}
 
 	void LogTaskCullingBegin(FPCGTaskId CompletedTaskId, uint64 InactiveOutputPinBitmask, const TArray<FPCGPinId>& PinIdsToDeactivate)
 	{
-#if WITH_EDITOR
+#if !(UE_BUILD_SHIPPING || UE_BUILD_TEST) || USE_LOGGING_IN_SHIPPING
 		if (!CullingLogEnabled())
 		{
 			return;
@@ -284,151 +329,172 @@ namespace PCGGraphExecutionLogging
 
 		UE_LOG(LogPCG, Log, TEXT("BEGIN CullInactiveDownstreamNodes, CompletedTaskId: %u, InactiveOutputPinBitmask: %u, Deactivating pin IDs: %s"),
 			CompletedTaskId, InactiveOutputPinBitmask, *GetPinsToDeactivateString(PinIdsToDeactivate));
-#endif
+#endif // !(UE_BUILD_SHIPPING || UE_BUILD_TEST) || USE_LOGGING_IN_SHIPPING
 	}
 
 	void LogTaskCullingBeginLoop(FPCGTaskId PinTaskId, uint64 PinIndex, const TArray<FPCGPinId>& PinIdsToDeactivate)
 	{
-#if WITH_EDITOR
+#if !(UE_BUILD_SHIPPING || UE_BUILD_TEST) || USE_LOGGING_IN_SHIPPING
 		if (!CullingLogEnabled())
 		{
 			return;
 		}
 
 		UE_LOG(LogPCG, Log, TEXT("LOOP: DEACTIVATE %u_%u, remaining IDs: %s"), PinTaskId, PinIndex, *GetPinsToDeactivateString(PinIdsToDeactivate));
-#endif
+#endif // !(UE_BUILD_SHIPPING || UE_BUILD_TEST) || USE_LOGGING_IN_SHIPPING
 	}
 
 	void LogTaskCullingUpdatedPinDeps(FPCGTaskId TaskId, const FPCGPinDependencyExpression& PinDependency, bool bDependencyExpressionBecameFalse)
 	{
-#if WITH_EDITOR
+#if !(UE_BUILD_SHIPPING || UE_BUILD_TEST) || USE_LOGGING_IN_SHIPPING
 		if (!CullingLogEnabled())
 		{
 			return;
 		}
 
-		UE_LOG(LogPCG, Log, TEXT("UPDATED PIN DEP EXPRESSION (task ID %u): %s"), TaskId, *PinDependency.ToString());
+		const FString PinDependencyString =
+#if WITH_EDITOR
+			PinDependency.ToString();
+#else
+			TEXT("MISSINGPINDEPS");
+#endif
+
+		UE_LOG(LogPCG, Log, TEXT("UPDATED PIN DEP EXPRESSION (task ID %u): %s"), TaskId, *PinDependencyString);
 
 		if (bDependencyExpressionBecameFalse)
 		{
 			UE_LOG(LogPCG, Log, TEXT("CULL task ID %u"), TaskId);
 		}
-#endif
+#endif // !(UE_BUILD_SHIPPING || UE_BUILD_TEST) || USE_LOGGING_IN_SHIPPING
 	}
 
 	void LogGridLinkageTaskExecuteStore(const FPCGContext* InContext, EPCGHiGenGrid InGenerationGrid, int32 InFromGridSize, int32 InToGridSize, const FString& InResourcePath)
 	{
-#if WITH_EDITOR
+#if !(UE_BUILD_SHIPPING || UE_BUILD_TEST) || USE_LOGGING_IN_SHIPPING
 		if (!LogEnabled())
 		{
 			return;
 		}
 		check(InContext);
 
-		const FString OwnerName = InContext->SourceComponent.Get() ? (InContext->SourceComponent->GetOwner() ? InContext->SourceComponent->GetOwner()->GetActorLabel() : FString()) : FString();
-		UE_LOG(LogPCG, Log, TEXT("[GRIDLINKING] [%s] STORE. GenerationGridSize=%d, FromGridSize=%d, ToGridSize=%d, Path=%s"), *OwnerName, PCGHiGenGrid::GridToGridSize(InGenerationGrid), InFromGridSize, InToGridSize, *InResourcePath);
-#endif
+		UE_LOG(LogPCG, Log, TEXT("[GRIDLINKING] [%s] STORE. GenerationGridSize=%d, FromGridSize=%d, ToGridSize=%d, Path=%s"),
+			*GetOwnerName(InContext->SourceComponent.Get()),
+			PCGHiGenGrid::GridToGridSize(InGenerationGrid),
+			InFromGridSize,
+			InToGridSize,
+			*InResourcePath);
+#endif // !(UE_BUILD_SHIPPING || UE_BUILD_TEST) || USE_LOGGING_IN_SHIPPING
 	}
 
 	void LogGridLinkageTaskExecuteRetrieve(const FPCGContext* InContext, EPCGHiGenGrid InGenerationGrid, int32 InFromGridSize, int32 InToGridSize, const FString& InResourcePath)
 	{
-#if WITH_EDITOR
+#if !(UE_BUILD_SHIPPING || UE_BUILD_TEST) || USE_LOGGING_IN_SHIPPING
 		if (!LogEnabled())
 		{
 			return;
 		}
 		check(InContext);
 
-		const FString OwnerName = InContext->SourceComponent.Get() ? (InContext->SourceComponent->GetOwner() ? InContext->SourceComponent->GetOwner()->GetActorLabel() : FString()) : FString();
-		UE_LOG(LogPCG, Log, TEXT("[GRIDLINKING] [%s] RETRIEVE. GenerationGridSize=%d, FromGridSize=%d, ToGridSize=%d, Path=%s"), *OwnerName, PCGHiGenGrid::GridToGridSize(InGenerationGrid), InFromGridSize, InToGridSize, *InResourcePath);
-#endif
+		UE_LOG(LogPCG, Log, TEXT("[GRIDLINKING] [%s] RETRIEVE. GenerationGridSize=%d, FromGridSize=%d, ToGridSize=%d, Path=%s"),
+			*GetOwnerName(InContext->SourceComponent.Get()),
+			PCGHiGenGrid::GridToGridSize(InGenerationGrid),
+			InFromGridSize,
+			InToGridSize,
+			*InResourcePath);
+#endif // !(UE_BUILD_SHIPPING || UE_BUILD_TEST) || USE_LOGGING_IN_SHIPPING
 	}
 
 	void LogGridLinkageTaskExecuteRetrieveSuccess(const FPCGContext* InContext, const UPCGComponent* InComponent, const FString& InResourcePath, int32 InDataItemCount)
 	{
-#if WITH_EDITOR
+#if !(UE_BUILD_SHIPPING || UE_BUILD_TEST) || USE_LOGGING_IN_SHIPPING
 		if (!LogEnabled())
 		{
 			return;
 		}
 		check(InContext);
 
-		const FString OwnerName = InContext->SourceComponent.Get() ? (InContext->SourceComponent->GetOwner() ? InContext->SourceComponent->GetOwner()->GetActorLabel() : FString()) : FString();
-		UE_LOG(LogPCG, Log, TEXT("[GRIDLINKING] [%s] RETRIEVE: SUCCESS. Component=%s Path=%s DataItems=%d"), *OwnerName, *InComponent->GetOwner()->GetActorLabel(), *InResourcePath, InDataItemCount);
-#endif
+		UE_LOG(LogPCG, Log, TEXT("[GRIDLINKING] [%s] RETRIEVE: SUCCESS. Path=%s DataItems=%d"),
+			*GetOwnerName(InContext->SourceComponent.Get()),
+			*InResourcePath,
+			InDataItemCount);
+#endif // !(UE_BUILD_SHIPPING || UE_BUILD_TEST) || USE_LOGGING_IN_SHIPPING
 	}
 
 	void LogGridLinkageTaskExecuteRetrieveScheduleGraph(const FPCGContext* InContext, const UPCGComponent* InScheduledComponent, const FString& InResourcePath)
 	{
-#if WITH_EDITOR
+#if !(UE_BUILD_SHIPPING || UE_BUILD_TEST) || USE_LOGGING_IN_SHIPPING
 		if (!LogEnabled())
 		{
 			return;
 		}
 		check(InContext);
 
-		const FString OwnerName = (InContext->SourceComponent.Get() && InContext->SourceComponent->GetOwner()) ? InContext->SourceComponent->GetOwner()->GetActorLabel() : FString();
-		const FString OtherOwnerName = (InScheduledComponent && InScheduledComponent->GetOwner()) ? InScheduledComponent->GetOwner()->GetActorLabel() : FString();
-		UE_LOG(LogPCG, Log, TEXT("[GRIDLINKING] [%s] RETRIEVE: SCHEDULE GRAPH. Component=%s Path=%s"), *OwnerName, *OtherOwnerName, *InResourcePath);
-#endif
+		UE_LOG(LogPCG, Log, TEXT("[GRIDLINKING] [%s] RETRIEVE: SCHEDULE GRAPH. Component=%s Path=%s"),
+			*GetOwnerName(InContext->SourceComponent.Get()),
+			*GetOwnerName(InScheduledComponent),
+			*InResourcePath);
+#endif // !(UE_BUILD_SHIPPING || UE_BUILD_TEST) || USE_LOGGING_IN_SHIPPING
 	}
 
 	void LogGridLinkageTaskExecuteRetrieveWaitOnScheduledGraph(const FPCGContext* InContext, const UPCGComponent* InWaitOnComponent, const FString& InResourcePath)
 	{
-#if WITH_EDITOR
+#if !(UE_BUILD_SHIPPING || UE_BUILD_TEST) || USE_LOGGING_IN_SHIPPING
 		if (!LogEnabled())
 		{
 			return;
 		}
 		check(InContext);
 
-		const FString OwnerName = (InContext->SourceComponent.Get() && InContext->SourceComponent->GetOwner()) ? InContext->SourceComponent->GetOwner()->GetActorLabel() : FString();
-		const FString OtherOwnerName = (InWaitOnComponent && InWaitOnComponent->GetOwner()) ? InWaitOnComponent->GetOwner()->GetActorLabel() : FString();
-		UE_LOG(LogPCG, Log, TEXT("[GRIDLINKING] [%s] RETRIEVE: WAIT FOR SCHEDULED GRAPH. Component=%s Path=%s"), *OwnerName, *OtherOwnerName, *InResourcePath);
-#endif
+		UE_LOG(LogPCG, Log, TEXT("[GRIDLINKING] [%s] RETRIEVE: WAIT FOR SCHEDULED GRAPH. Component=%s Path=%s"),
+			*GetOwnerName(InContext->SourceComponent.Get()),
+			*GetOwnerName(InWaitOnComponent),
+			*InResourcePath);
+#endif // !(UE_BUILD_SHIPPING || UE_BUILD_TEST) || USE_LOGGING_IN_SHIPPING
 	}
 	
 	void LogGridLinkageTaskExecuteRetrieveWakeUp(const FPCGContext* InContext, const UPCGComponent* InWokenBy)
 	{
-#if WITH_EDITOR
+#if !(UE_BUILD_SHIPPING || UE_BUILD_TEST) || USE_LOGGING_IN_SHIPPING
 		if (!LogEnabled())
 		{
 			return;
 		}
 		check(InContext);
 
-		const FString OwnerName = (InContext->SourceComponent.Get() && InContext->SourceComponent->GetOwner()) ? InContext->SourceComponent->GetOwner()->GetActorLabel() : FString();
-		const FString OtherOwnerName = (InWokenBy && InWokenBy->GetOwner()) ? InWokenBy->GetOwner()->GetActorLabel() : FString();
-		UE_LOG(LogPCG, Log, TEXT("[GRIDLINKING] [%s] RETRIEVE: WOKEN BY Component=%s"), *OwnerName, *OtherOwnerName);
-#endif
+		UE_LOG(LogPCG, Log, TEXT("[GRIDLINKING] [%s] RETRIEVE: WOKEN BY Component=%s"),
+			*GetOwnerName(InContext->SourceComponent.Get()),
+			*GetOwnerName(InWokenBy));
+#endif // !(UE_BUILD_SHIPPING || UE_BUILD_TEST) || USE_LOGGING_IN_SHIPPING
 	}
 
 	void LogGridLinkageTaskExecuteRetrieveNoLocalComponent(const FPCGContext* InContext, const FString& InResourcePath)
 	{
-#if WITH_EDITOR
+#if !(UE_BUILD_SHIPPING || UE_BUILD_TEST) || USE_LOGGING_IN_SHIPPING
 		if (!LogEnabled())
 		{
 			return;
 		}
 		check(InContext);
 
-		const FString OwnerName = (InContext->SourceComponent.Get() && InContext->SourceComponent->GetOwner()) ? InContext->SourceComponent->GetOwner()->GetActorLabel() : FString();
-		UE_LOG(LogPCG, Warning, TEXT("[GRIDLINKING] [%s] RETRIEVE: FAILED: No overlapping local component found. This may be expected. Path=%s"), *OwnerName, *InResourcePath);
-#endif
+		UE_LOG(LogPCG, Warning, TEXT("[GRIDLINKING] [%s] RETRIEVE: FAILED: No overlapping local component found. This may be expected. Path=%s"),
+			*GetOwnerName(InContext->SourceComponent.Get()),
+			*InResourcePath);
+#endif // !(UE_BUILD_SHIPPING || UE_BUILD_TEST) || USE_LOGGING_IN_SHIPPING
 	}
 
 	void LogGridLinkageTaskExecuteRetrieveNoData(const FPCGContext* InContext, const UPCGComponent* InComponent, const FString& InResourcePath)
 	{
-#if WITH_EDITOR
+#if !(UE_BUILD_SHIPPING || UE_BUILD_TEST) || USE_LOGGING_IN_SHIPPING
 		if (!LogEnabled())
 		{
 			return;
 		}
 		check(InContext);
 
-		const FString OwnerName = (InContext->SourceComponent.Get() && InContext->SourceComponent->GetOwner()) ? InContext->SourceComponent->GetOwner()->GetActorLabel() : FString();
-		const FString OtherOwnerName = (InComponent && InComponent->GetOwner()) ? InComponent->GetOwner()->GetActorLabel() : FString();
-		UE_LOG(LogPCG, Warning, TEXT("[GRIDLINKING] [%s] RETRIEVE: FAILED: No data found on local component. Component=%s, Path=%s"), *OwnerName, *OtherOwnerName, *InResourcePath);
-#endif
+		UE_LOG(LogPCG, Warning, TEXT("[GRIDLINKING] [%s] RETRIEVE: FAILED: No data found on local component. Component=%s, Path=%s"),
+			*GetOwnerName(InContext->SourceComponent.Get()),
+			*GetOwnerName(InComponent),
+			*InResourcePath);
+#endif // !(UE_BUILD_SHIPPING || UE_BUILD_TEST) || USE_LOGGING_IN_SHIPPING
 	}
 }
diff --git a/Engine/Plugins/PCG/Source/PCG/Private/Utils/PCGGraphExecutionLogging.h b/Engine/Plugins/PCG/Source/PCG/Private/Utils/PCGGraphExecutionLogging.h
--- a/Engine/Plugins/PCG/Source/PCG/Private/Utils/PCGGraphExecutionLogging.h
+++ b/Engine/Plugins/PCG/Source/PCG/Private/Utils/PCGGraphExecutionLogging.h
@@ -8,11 +8,13 @@
 #include "HAL/Platform.h"
 
 class UPCGComponent;
+class UPCGGraph;
 enum class EPCGHiGenGrid : uint32;
 struct FPCGContext;
 struct FPCGGraphActiveTask;
 struct FPCGGraphTask;
 struct FPCGPinDependencyExpression;
+struct FPCGStack;
 
 namespace PCGGraphExecutionLogging
 {
@@ -23,9 +25,9 @@ namespace PCGGraphExecutionLogging
 	void LogGraphTasks(const TMap<FPCGTaskId, FPCGGraphTask>& Tasks, const TMap<FPCGTaskId, TSet<FPCGTaskId>>* TaskSuccessors = nullptr);
 	void LogGraphTasks(const TArray<FPCGGraphTask>& Tasks);
 
-	void LogGraphSchedule(const UPCGComponent* SourceComponent);
-	void LogGraphScheduleDependency(const UPCGComponent* InComponent);
-	void LogGraphScheduleDependencyFailed(const UPCGComponent* InComponent);
+	void LogGraphSchedule(const UPCGComponent* SourceComponent, const UPCGGraph* InScheduledGraph);
+	void LogGraphScheduleDependency(const UPCGComponent* InComponent, const FPCGStack* InFromStack);
+	void LogGraphScheduleDependencyFailed(const UPCGComponent* InComponent, const FPCGStack* InFromStack);
 
 	void LogGraphPostSchedule(const TMap<FPCGTaskId, FPCGGraphTask>& Tasks, const TMap<FPCGTaskId, TSet<FPCGTaskId>>& TaskSuccessors);
 	void LogPostProcessGraph(const UPCGComponent* InSourceComponent);
diff --git a/Engine/Plugins/PCG/Source/PCGEditor/Private/WorldPartitionBuilder/PCGWorldPartitionBuilder.cpp b/Engine/Plugins/PCG/Source/PCGEditor/Private/WorldPartitionBuilder/PCGWorldPartitionBuilder.cpp
--- a/Engine/Plugins/PCG/Source/PCGEditor/Private/WorldPartitionBuilder/PCGWorldPartitionBuilder.cpp
+++ b/Engine/Plugins/PCG/Source/PCGEditor/Private/WorldPartitionBuilder/PCGWorldPartitionBuilder.cpp
@@ -22,8 +22,14 @@ namespace PCGWorldPartitionBuilder
 	void CollectComponentsToGenerate(UWorld* InWorld, TFunctionRef<bool(const UPCGComponent*)> ComponentFilter, TArray<TWeakObjectPtr<UPCGComponent>>& OutComponents);
 	
 	/** Generate the given components. Optionally generate only one component at a time with a wait on async processes after each. Optionally apply given filter to select components. */
-	bool GenerateComponents(TArray<TWeakObjectPtr<UPCGComponent>>& Components, UWorld* InWorld, bool bOneComponentAtATime, bool& bOutGenerationErrors);
-	bool GenerateComponents(TArray<TWeakObjectPtr<UPCGComponent>>& Components, UWorld* InWorld, bool bOneComponentAtATime, TFunctionRef<bool(const UPCGComponent*)> ComponentFilter, bool& bOutGenerationErrors);
+	bool GenerateComponents(TArray<TWeakObjectPtr<UPCGComponent>>& Components, UWorld* InWorld, bool bOneComponentAtATime, TArray<UPackage*>& InOutDeletedActorPackages, bool& bOutGenerationErrors);
+	bool GenerateComponents(
+		TArray<TWeakObjectPtr<UPCGComponent>>& Components,
+		UWorld* InWorld,
+		bool bOneComponentAtATime,
+		TFunctionRef<bool(const UPCGComponent*)> ComponentFilter,
+		TArray<UPackage*>& InOutDeletedActorPackages,
+		bool& bOutGenerationErrors);
 
 	/** Generate a component. Applies correct editing mode if necessary. */
 	void GenerateComponent(UPCGComponent* InComponent, UWorld* InWorld);
@@ -41,7 +47,7 @@ namespace PCGWorldPartitionBuilder
 			"\t[-GenerateComponentEditingModeNormal]\n"
 			"\t[-GenerateComponentEditingModePreview]\n"
 			"\t[-IgnoreGenerationErrors]\n"
-			"\t[-IncludeActorIDs=MyActor1_UID1234678;MyActor2_UID1234678]\n"
+			"\t[-IncludeActorIDs=MyActor1_UAID1234678;MyActor2_UAID1234678]\n"
 			"\t[-OneComponentAtATime]\n"),
 		FConsoleCommandWithArgsDelegate::CreateLambda([](const TArray<FString>& Args) { PCGWorldPartitionBuilder::Build(Args); }));
 };
@@ -188,7 +194,7 @@ bool UPCGWorldPartitionBuilder::RunInternal(UWorld* World, const FCellInfo& InCe
 			};
 
 			bool bErrorsOccurred = false;
-			bGeneratedAnyComponent |= PCGWorldPartitionBuilder::GenerateComponents(ComponentsToGenerate, World, bOneComponentAtATime, FilterOnGraphName, bErrorsOccurred);
+			bGeneratedAnyComponent |= PCGWorldPartitionBuilder::GenerateComponents(ComponentsToGenerate, World, bOneComponentAtATime, FilterOnGraphName, DeletedActorPackages, bErrorsOccurred);
 
 			bErrorOccurredWhileGenerating |= bErrorsOccurred;
 		}
@@ -196,7 +202,7 @@ bool UPCGWorldPartitionBuilder::RunInternal(UWorld* World, const FCellInfo& InCe
 	else
 	{
 		bool bErrorsOccurred = false;
-		bGeneratedAnyComponent |= PCGWorldPartitionBuilder::GenerateComponents(ComponentsToGenerate, World, bOneComponentAtATime, bErrorsOccurred);
+		bGeneratedAnyComponent |= PCGWorldPartitionBuilder::GenerateComponents(ComponentsToGenerate, World, bOneComponentAtATime, DeletedActorPackages, bErrorsOccurred);
 
 		bErrorOccurredWhileGenerating |= bErrorsOccurred;
 	}
@@ -273,7 +279,19 @@ bool UPCGWorldPartitionBuilder::PostRun(UWorld* World, FPackageSourceControlHelp
 		}
 	}
 
-	UE_LOG(LogPCGWorldPartitionBuilder, Display, TEXT("PostRun: %d packages modified, %d packages deleted."), DirtyPackages.Num(), PackagesToDelete.Num());
+	UE_LOG(LogPCGWorldPartitionBuilder, Display, TEXT("PostRun: %d packages modified, %d actor packages deleted, %d empty packages marked for delete."),
+		DirtyPackages.Num(),
+		DeletedActorPackages.Num(),
+		PackagesToDelete.Num());
+
+	// Combine the empty-deleted packages with the actor-deleted packages.
+	for (UPackage* DeletedPackage : DeletedActorPackages)
+	{
+		PackagesToDelete.AddUnique(DeletedPackage);
+	}
+
+	// Log final changes after combining packages (there may have been duplicates in the two deleted package lists).
+	UE_LOG(LogPCGWorldPartitionBuilder, Display, TEXT("PostRun: Final package changes: %d modified, %d deleted."), DirtyPackages.Num(), PackagesToDelete.Num());
 
 	if (!SavePackages(DirtyPackages, PackageHelper))
 	{
@@ -339,9 +357,9 @@ void PCGWorldPartitionBuilder::CollectComponentsToGenerate(
 	}
 }
 
-bool PCGWorldPartitionBuilder::GenerateComponents(TArray<TWeakObjectPtr<UPCGComponent>>& Components, UWorld* InWorld, bool bOneComponentAtATime, bool& bOutGenerationErrors)
+bool PCGWorldPartitionBuilder::GenerateComponents(TArray<TWeakObjectPtr<UPCGComponent>>& Components, UWorld* InWorld, bool bOneComponentAtATime, TArray<UPackage*>& InOutDeletedActorPackages, bool& bOutGenerationErrors)
 {
-	return PCGWorldPartitionBuilder::GenerateComponents(Components, InWorld, bOneComponentAtATime, [](const UPCGComponent*) { return true; }, bOutGenerationErrors);
+	return PCGWorldPartitionBuilder::GenerateComponents(Components, InWorld, bOneComponentAtATime, [](const UPCGComponent*) { return true; }, InOutDeletedActorPackages, bOutGenerationErrors);
 }
 
 bool PCGWorldPartitionBuilder::GenerateComponents(
@@ -349,6 +367,7 @@ bool PCGWorldPartitionBuilder::GenerateComponents(
 	UWorld* InWorld,
 	bool bOneComponentAtATime,
 	TFunctionRef<bool(const UPCGComponent*)> ComponentFilter,
+	TArray<UPackage*>& InOutDeletedActorPackages,
 	bool& bOutGenerationErrors)
 {
 	if (!bOneComponentAtATime)
@@ -374,6 +393,28 @@ bool PCGWorldPartitionBuilder::GenerateComponents(
 		FWorldPartitionHelpers::FakeEngineTick(InWorld);
 	};
 
+	// Hook actor deleted events and track any corresponding deleted packages, as the packages can be GC'd if PCG triggers a GC
+	// before generation.
+	FDelegateHandle ActorDeletedHandle = GEngine->OnLevelActorDeleted().AddLambda([&InOutDeletedActorPackages](AActor* InActor)
+	{
+		if (InActor && InActor->IsPackageExternal())
+		{
+			if (UPackage* ActorPackage = InActor->GetPackage())
+			{
+				UE_LOG(LogPCGWorldPartitionBuilder, Display, TEXT("Actor '%s' deleted, package '%s' added to delete list."),
+					*InActor->GetName(),
+					*ActorPackage->GetName());
+
+				InOutDeletedActorPackages.AddUnique(ActorPackage);
+			}
+		}
+	});
+
+	ON_SCOPE_EXIT
+	{
+		GEngine->OnLevelActorDeleted().Remove(ActorDeletedHandle);
+	};
+
 	TSet<TObjectKey<UPCGComponent>> GeneratedComponents;
 
 	for (TWeakObjectPtr<UPCGComponent> ComponentWeakPtr : Components)
diff --git a/Engine/Plugins/PCG/Source/PCGEditor/Private/WorldPartitionBuilder/PCGWorldPartitionBuilder.h b/Engine/Plugins/PCG/Source/PCGEditor/Private/WorldPartitionBuilder/PCGWorldPartitionBuilder.h
--- a/Engine/Plugins/PCG/Source/PCGEditor/Private/WorldPartitionBuilder/PCGWorldPartitionBuilder.h
+++ b/Engine/Plugins/PCG/Source/PCGEditor/Private/WorldPartitionBuilder/PCGWorldPartitionBuilder.h
@@ -39,6 +39,10 @@ private:
 	UPROPERTY(Transient)
 	TMap<FGuid, UPackage*> PendingDirtyPackages;
 
+	/** Packages that were logged as deleted through the OnActorDeleted event. UPROP to prevent GC. */
+	UPROPERTY(Transient)
+	TArray<UPackage*> DeletedActorPackages;
+
 	/** Include components which have editing mode set to Normal. */
 	bool bGenerateEditingModeNormalComponents = false;
 
