[PCG] Always call CreateDefaultPins (not UpdatePins) from the UPCGGraph constructor to avoid unwanted Modify/MarkPackageDirty calls in a UPCGGraph constructor

- This was causing the Resave commandlet to always save PCG Graphs as they would dirty on load throughg UpdatePins()

#rnx

diff --git a/Engine/Plugins/PCG/Source/PCG/Private/PCGGraph.cpp b/Engine/Plugins/PCG/Source/PCG/Private/PCGGraph.cpp
--- a/Engine/Plugins/PCG/Source/PCG/Private/PCGGraph.cpp
+++ b/Engine/Plugins/PCG/Source/PCG/Private/PCGGraph.cpp
@@ -267,15 +267,8 @@ UPCGGraph::UPCGGraph(const FObjectInitializer& ObjectInitializer)
 	InputNode->SetSettingsInterface(InputSettings, /*bUpdatePins=*/false);
 
 	// Only allocate default pins if this is the default object
-	if (this->HasAnyFlags(RF_ClassDefaultObject))
-	{
-		InputNode->CreateDefaultPins(PinAllocator);
-	}
-	else
-	{
-		InputNode->UpdatePins();
-	}
-
+	InputNode->CreateDefaultPins(PinAllocator);
+	
 	OutputNode = ObjectInitializer.CreateDefaultSubobject<UPCGNode>(this, TEXT("DefaultOutputNode"));
 	OutputNode->SetFlags(RF_Transactional);
 
@@ -284,15 +277,8 @@ UPCGGraph::UPCGGraph(const FObjectInitializer& ObjectInitializer)
 	OutputNode->SetSettingsInterface(OutputSettings, /*bUpdatePins=*/false);
 
 	// Only allocate default pins if this is the default object
-	if (this->HasAnyFlags(RF_ClassDefaultObject))
-	{
-		OutputNode->CreateDefaultPins(PinAllocator);
-	}
-	else
-	{
-		OutputNode->UpdatePins();
-	}
-
+	OutputNode->CreateDefaultPins(PinAllocator);
+	
 #if WITH_EDITOR
 	OutputNode->PositionX = 200;
 #endif
diff --git a/Engine/Plugins/PCG/Source/PCG/Private/PCGInputOutputSettings.cpp b/Engine/Plugins/PCG/Source/PCG/Private/PCGInputOutputSettings.cpp
--- a/Engine/Plugins/PCG/Source/PCG/Private/PCGInputOutputSettings.cpp
+++ b/Engine/Plugins/PCG/Source/PCG/Private/PCGInputOutputSettings.cpp
@@ -158,15 +158,16 @@ void UPCGGraphInputOutputSettings::ApplyDeprecationBeforeUpdatePins(UPCGNode* In
 }
 #endif // WITH_EDITOR
 
-TArray<FPCGPinProperties> UPCGGraphInputOutputSettings::DefaultInputPinProperties() const
+TArray<FPCGPinProperties> UPCGGraphInputOutputSettings::DefaultPinProperties(bool bInvisiblePin) const
 {
 	// It is important for serialization that this is not modified, or it could break existing graphs.
 	TArray<FPCGPinProperties> PinProperties;
 	const EPCGDataType DefaultPinDataType = bIsInput ? EPCGDataType::Spatial : EPCGDataType::Any;
 	const TArray<PCGInputOutputPrivate::FLabelAndTooltip>& StaticLabels = bIsInput ? PCGInputOutputPrivate::StaticInLabels : PCGInputOutputPrivate::StaticOutLabels;
 
-	Algo::Transform(StaticLabels, PinProperties, [DefaultPinDataType](const PCGInputOutputPrivate::FLabelAndTooltip& InLabelAndTooltip) {
+	Algo::Transform(StaticLabels, PinProperties, [bInvisiblePin, DefaultPinDataType](const PCGInputOutputPrivate::FLabelAndTooltip& InLabelAndTooltip) {
 		FPCGPinProperties Res = FPCGPinProperties(InLabelAndTooltip.Label, DefaultPinDataType, /*bMultiConnections=*/true, /*bMultiData=*/true, InLabelAndTooltip.Tooltip);
+		Res.bInvisiblePin = bInvisiblePin;
 		Res.SetAdvancedPin();
 		return Res;
 	});
@@ -174,9 +175,14 @@ TArray<FPCGPinProperties> UPCGGraphInputOutputSettings::DefaultInputPinPropertie
 	return PinProperties;
 }
 
+TArray<FPCGPinProperties> UPCGGraphInputOutputSettings::DefaultInputPinProperties() const
+{
+	return DefaultPinProperties(/*bInvisiblePin=*/bIsInput);
+}
+
 TArray<FPCGPinProperties> UPCGGraphInputOutputSettings::DefaultOutputPinProperties() const
 {
-	return DefaultInputPinProperties();
+	return DefaultPinProperties(/*bInvisiblePin=*/!bIsInput);
 }
 
 const FPCGPinProperties& UPCGGraphInputOutputSettings::AddPin(const FPCGPinProperties& NewPinProperties)
diff --git a/Engine/Plugins/PCG/Source/PCG/Public/PCGInputOutputSettings.h b/Engine/Plugins/PCG/Source/PCG/Public/PCGInputOutputSettings.h
--- a/Engine/Plugins/PCG/Source/PCG/Public/PCGInputOutputSettings.h
+++ b/Engine/Plugins/PCG/Source/PCG/Public/PCGInputOutputSettings.h
@@ -58,6 +58,7 @@ public:
 	[[nodiscard]] const FPCGPinProperties& AddPin(const FPCGPinProperties& NewCustomPinProperties);
 
 protected:
+	TArray<FPCGPinProperties> DefaultPinProperties(bool bInvisiblePin) const;
 	virtual FPCGElementPtr CreateElement() const override { return MakeShared<FPCGInputOutputElement>(); }
 	// ~End UPCGSettings interface
 
