[PCG] Don't present original component / unbounded grid component as a debuggable object if it will not execute.

Debugging this would display whole graph as culled, creating confusion.

diff --git a/Engine/Plugins/PCG/Source/PCGEditor/Private/SPCGEditorGraphDebugObjectTree.cpp b/Engine/Plugins/PCG/Source/PCGEditor/Private/SPCGEditorGraphDebugObjectTree.cpp
--- a/Engine/Plugins/PCG/Source/PCGEditor/Private/SPCGEditorGraphDebugObjectTree.cpp
+++ b/Engine/Plugins/PCG/Source/PCGEditor/Private/SPCGEditorGraphDebugObjectTree.cpp
@@ -552,6 +552,28 @@ void SPCGEditorGraphDebugObjectTree::AddStacksToTree(const TArray<FPCGStack>& St
 			continue;
 		}
 
+		if (PCGComponent->IsPartitioned())
+		{
+			if (!GraphBeingEdited->IsHierarchicalGenerationEnabled())
+			{
+				// Don't show the original component if graph does not have higen enabled because this component
+				// will never do any interesting execution in this case.
+				continue;
+			}
+			else
+			{
+				// If higen is enabled, check if the UB grid level will execute. If not then, then there is no execution to inspect.
+				bool bHasUnbounded = false;
+				PCGHiGenGrid::FSizeArray GridSizes;
+				GraphBeingEdited->GetGridSizes(GridSizes, bHasUnbounded);
+
+				if (!bHasUnbounded)
+				{
+					continue;
+				}
+			}
+		}
+
 		// Prevent duplicate entries from the editor world while in PIE.
 		if (PCGHelpers::IsRuntimeOrPIE())
 		{
