[PCG][Bug] Modified disabled node behavior to always output the first input pin

* Fixed a bug causing Metadata Ops and other nodes to not passthrough Attribute Sets correctly when disabled

#rnx

diff --git a/Engine/Plugins/PCG/Source/PCG/Private/Elements/Metadata/PCGMetadataOpElementBase.cpp b/Engine/Plugins/PCG/Source/PCG/Private/Elements/Metadata/PCGMetadataOpElementBase.cpp
--- a/Engine/Plugins/PCG/Source/PCG/Private/Elements/Metadata/PCGMetadataOpElementBase.cpp
+++ b/Engine/Plugins/PCG/Source/PCG/Private/Elements/Metadata/PCGMetadataOpElementBase.cpp
@@ -188,6 +188,11 @@ void UPCGMetadataSettingsBase::ApplyDeprecation(UPCGNode* InOutNode)
 }
 #endif // WITH_EDITOR
 
+bool UPCGMetadataSettingsBase::DoesPinSupportPassThrough(UPCGPin* InPin) const
+{
+	return InPin && !InPin->IsOutputPin() && GetInputPinIndex(InPin->Properties.Label) == GetInputPinToForward();
+}
+
 TArray<FPCGPinProperties> UPCGMetadataSettingsBase::InputPinProperties() const
 {
 	TArray<FPCGPinProperties> PinProperties;
diff --git a/Engine/Plugins/PCG/Source/PCG/Private/PCGElement.cpp b/Engine/Plugins/PCG/Source/PCG/Private/PCGElement.cpp
--- a/Engine/Plugins/PCG/Source/PCG/Private/PCGElement.cpp
+++ b/Engine/Plugins/PCG/Source/PCG/Private/PCGElement.cpp
@@ -227,7 +227,7 @@ void IPCGElement::DisabledPassThroughData(FPCGContext* Context) const
 		Data.Pin = PassThroughOutputPin->Properties.Label;
 	}
 
-	// Pass through input data if it is not params, and if the output type supports it (e.g. if we have a incoming
+	// Pass through input data if both it and the output are params, or if the output type supports it (e.g. if we have a incoming
 	// surface connected to an input pin of type Any, do not pass the surface through to an output pin of type Point).
 	auto InputDataShouldPassThrough = [OutputType](const FPCGTaggedData& InData)
 	{
@@ -235,12 +235,12 @@ void IPCGElement::DisabledPassThroughData(FPCGContext* Context) const
 		const bool bInputTypeNotWiderThanOutputType = !(InputType & ~OutputType);
 
 		// Right now we allow edges from Spatial to Concrete. This can happen for example if a point processing node
-		// is receving a Spatial data, and the node is disabled, it will want to pass the Spatial data through. In the
+		// is receiving a Spatial data, and the node is disabled, it will want to pass the Spatial data through. In the
 		// future we will force collapses/conversions. For now, allow an incoming Spatial to pass out through a Concrete.
 		// TODO remove!
 		const bool bAllowSpatialToConcrete = !!(InputType & EPCGDataType::Spatial) && !!(OutputType & EPCGDataType::Concrete);
 
-		return InputType != EPCGDataType::Param && (bInputTypeNotWiderThanOutputType || bAllowSpatialToConcrete);
+		return (InputType != EPCGDataType::Param || OutputType == EPCGDataType::Param) && (bInputTypeNotWiderThanOutputType || bAllowSpatialToConcrete);
 	};
 
 	// Now remove any non-params edges, and if only one edge should come through, remove the others
diff --git a/Engine/Plugins/PCG/Source/PCG/Private/PCGNode.cpp b/Engine/Plugins/PCG/Source/PCG/Private/PCGNode.cpp
--- a/Engine/Plugins/PCG/Source/PCG/Private/PCGNode.cpp
+++ b/Engine/Plugins/PCG/Source/PCG/Private/PCGNode.cpp
@@ -590,45 +590,38 @@ const UPCGPin* UPCGNode::GetPassThroughInputPin() const
 		return nullptr;
 	}
 
-	const EPCGDataType OutputType = PassThroughOutput->GetCurrentTypes();
-	// We assume a node whose primary output is params is a param processing node.
-	const bool bNodeOutputsParams = (OutputType == EPCGDataType::Param);
+	const TArray<TObjectPtr<UPCGPin>>& AllInputPins = GetInputPins();
+	if (AllInputPins.IsEmpty())
+	{
+		return nullptr;
+	}
 
-	if (bNodeOutputsParams)
+	const UPCGSettings* Settings = GetSettings();
+	if (!Settings)
 	{
-		// If this node primarily processes params, then look for a params pin, if we find it then we will pass it through
-		auto FindParams = [](const TObjectPtr<UPCGPin>& InPin) { return InPin->GetCurrentTypes() == EPCGDataType::Param; };
-		const TObjectPtr<UPCGPin>* FirstParamsPinPtr = Algo::FindByPredicate(GetInputPins(), FindParams);
-		return FirstParamsPinPtr ? *FirstParamsPinPtr : nullptr;
+		return nullptr;
 	}
-	else
+
+	// Always take the first valid input pin as the pass-through pin
+	auto FindFirstPassThrough = [Settings](const TObjectPtr<UPCGPin>& InPin) { return Settings->DoesPinSupportPassThrough(InPin); };
+	if (const TObjectPtr<UPCGPin>* FirstPassThroughPin = Algo::FindByPredicate(AllInputPins, FindFirstPassThrough))
 	{
-		// 'Normal' node. Look for the first pin that is not of type Params to pass through. If the node is disabled, the params are unused.
-		// Params-only pins will be rejected/ignored.
-		auto FindNonParams = [](const TObjectPtr<UPCGPin>& InPin) { return InPin->GetCurrentTypes() != EPCGDataType::Param; };
-		const TObjectPtr<UPCGPin>* FirstNonParamsPinPtr = Algo::FindByPredicate(GetInputPins(), FindNonParams);
+		// Finally in order to be a candidate for passing through, data must be compatible.
+		const EPCGDataType InputType = FirstPassThroughPin->Get()->GetCurrentTypes();
+		const EPCGDataType OutputType = PassThroughOutput->GetCurrentTypes();
 
-		if (FirstNonParamsPinPtr)
+		// Misc note - it would be nice if we could be stricter, like line below this comment. However it would mean it will stop an Any
+		// input being passed through to a Point output, even if the incoming edge will be receiving points dynamically/during execution. If a
+		// user creates a BP node with an Any input, which they may do lazily or unknowingly, this blocks passthrough. So instead we'll indicate
+		// that this pin *may* be used as a passthrough, and during execution in DisabledPassThroughData() we check dynamic types.
+		//const bool bInputTypeNotWiderThanOutputType = !(InputType & ~OutputType);
+		if (!!(InputType & OutputType))
 		{
-			// Finally in order to be a candidate for passing through, data must be compatible.
-			const UPCGPin* InputPin = *FirstNonParamsPinPtr;
-			const EPCGDataType InputType = InputPin->GetCurrentTypes();
-			const bool bTypesOverlap = !!(InputType & OutputType);
-			
-			// Misc note - it would be nice if we could be stricter, like line below this comment. However it would mean it will stop an Any
-			// input being passed through to a Point output, even if the incoming edge will be receiving points dynamically/during execution. If a
-			// user creates a BP node with an Any input, which they may do lazily or unknowingly, this blocks passthrough. So instead we'll indicate
-			// that this pin *may* be used as a passthrough, and during execution in DisabledPassThroughData() we check dynamic types.
-			//const bool bInputTypeNotWiderThanOutputType = !(InputType & ~OutputType);
-			
-			if (bTypesOverlap)
-			{
-				return *FirstNonParamsPinPtr;
-			}
+			return FirstPassThroughPin->Get();
 		}
-
-		return nullptr;
 	}
+
+	return nullptr;
 }
 
 const UPCGPin* UPCGNode::GetPassThroughOutputPin() const
diff --git a/Engine/Plugins/PCG/Source/PCG/Private/PCGSettings.cpp b/Engine/Plugins/PCG/Source/PCG/Private/PCGSettings.cpp
--- a/Engine/Plugins/PCG/Source/PCG/Private/PCGSettings.cpp
+++ b/Engine/Plugins/PCG/Source/PCG/Private/PCGSettings.cpp
@@ -693,6 +693,11 @@ bool UPCGSettings::IsPropertyOverriddenByPin(const TArrayView<const FName>& Prop
 	return false;
 }
 
+bool UPCGSettings::DoesPinSupportPassThrough(UPCGPin* InPin) const
+{
+	return InPin && !InPin->Properties.IsAdvancedPin();
+}
+
 EPCGDataType UPCGSettings::GetCurrentPinTypes(const UPCGPin* InPin) const
 {
 	check(InPin);
diff --git a/Engine/Plugins/PCG/Source/PCG/Public/Elements/Metadata/PCGMetadataOpElementBase.h b/Engine/Plugins/PCG/Source/PCG/Public/Elements/Metadata/PCGMetadataOpElementBase.h
--- a/Engine/Plugins/PCG/Source/PCG/Public/Elements/Metadata/PCGMetadataOpElementBase.h
+++ b/Engine/Plugins/PCG/Source/PCG/Public/Elements/Metadata/PCGMetadataOpElementBase.h
@@ -94,6 +94,7 @@ public:
 #endif
 	virtual bool HasFlippedTitleLines() const override { return true; }
 	virtual EPCGDataType GetCurrentPinTypes(const UPCGPin* InPin) const override;
+	virtual bool DoesPinSupportPassThrough(UPCGPin* InPin) const override;
 
 protected:
 	virtual TArray<FPCGPinProperties> InputPinProperties() const override;
diff --git a/Engine/Plugins/PCG/Source/PCG/Public/PCGSettings.h b/Engine/Plugins/PCG/Source/PCG/Public/PCGSettings.h
--- a/Engine/Plugins/PCG/Source/PCG/Public/PCGSettings.h
+++ b/Engine/Plugins/PCG/Source/PCG/Public/PCGSettings.h
@@ -363,6 +363,9 @@ public:
 	/** Returns true if only the first input edge is used from the primary pin when the node is disabled. */
 	virtual bool OnlyPassThroughOneEdgeWhenDisabled() const { return false; }
 
+	/** By default, the first non-Advanced pin that supports pass-through will be selected. */
+	virtual bool DoesPinSupportPassThrough(UPCGPin* InPin) const;
+
 	/** Returns the current pin types, which can either be the static types from the pin properties, or a dynamic type based on connected edges.
 	* By default we set output pin types to the union of the default input pin incident edge types, if it is dynamic and the default input exists.
 	*/
