[PCG] Extra security to alleviate out of memory asserts in Surface and Volume Samplers

diff --git a/Engine/Plugins/PCG/Source/PCG/Private/Elements/PCGSurfaceSampler.cpp b/Engine/Plugins/PCG/Source/PCG/Private/Elements/PCGSurfaceSampler.cpp
--- a/Engine/Plugins/PCG/Source/PCG/Private/Elements/PCGSurfaceSampler.cpp
+++ b/Engine/Plugins/PCG/Source/PCG/Private/Elements/PCGSurfaceSampler.cpp
@@ -72,15 +72,9 @@ namespace PCGSurfaceSampler
 			}
 
 			const int64 CellCount64 = CellCountX * CellCountY;
-			if (CellCount64 <= 0 || 
-				CellCount64 >= MAX_int32 ||
-				(PCGFeatureSwitches::CVarCheckSamplerMemory.GetValueOnAnyThread() && FPlatformMemory::GetStats().AvailablePhysical < sizeof(FPCGPoint) * CellCount64))
+			if (CellCount64 <= 0 || CellCount64 >= MAX_int32)
 			{
-				if (Context)
-				{
-					PCGE_LOG_C(Error, GraphAndLog, Context, FText::Format(LOCTEXT("TooManyPoints", "Skipped - tried to generate too many points {0}"), CellCount64));
-				}
-
+				PCGLog::LogErrorOnGraph(FText::Format((LOCTEXT("InvalidCellCount", "Skipped - tried to generate too many points ({0}).")), CellCount64), Context);
 				return false;
 			}
 
@@ -103,6 +97,12 @@ namespace PCGSurfaceSampler
 			return false;
 		}
 
+		if (PCGFeatureSwitches::CVarCheckSamplerMemory.GetValueOnAnyThread() && (PCGFeatureSwitches::CVarSamplerMemoryThreshold.GetValueOnAnyThread() * FPlatformMemory::GetStats().AvailablePhysical) < sizeof(FPCGPoint) * TargetPointCount)
+		{
+			PCGLog::LogErrorOnGraph(FText::Format((LOCTEXT("TooManyPoints", "Skipped - tried to generate too many points ({0}).\nAdjust 'pcg.SamplerMemoryThreshold' if needed.")), TargetPointCount), Context);
+			return false;
+		}
+
 		InputBoundsMinZ = InputBounds.Min.Z;
 		InputBoundsMaxZ = InputBounds.Max.Z;
 
diff --git a/Engine/Plugins/PCG/Source/PCG/Private/Elements/PCGVolumeSampler.cpp b/Engine/Plugins/PCG/Source/PCG/Private/Elements/PCGVolumeSampler.cpp
--- a/Engine/Plugins/PCG/Source/PCG/Private/Elements/PCGVolumeSampler.cpp
+++ b/Engine/Plugins/PCG/Source/PCG/Private/Elements/PCGVolumeSampler.cpp
@@ -76,21 +76,18 @@ namespace PCGVolumeSampler
 				return true;
 			}
 
-			if (NumIterationsXY64 > 0 && 
-				NumIterationsXY64 < MAX_int32 && 
-				NumIterations64 > 0 && 
+			if (NumIterationsXY64 > 0 &&
+				NumIterationsXY64 < MAX_int32 &&
+				NumIterations64 > 0 &&
 				NumIterations64 < MAX_int32 &&
-				(!PCGFeatureSwitches::CVarCheckSamplerMemory.GetValueOnAnyThread() || FPlatformMemory::GetStats().AvailablePhysical >= sizeof(FPCGPoint) * NumIterations64))
+				(!PCGFeatureSwitches::CVarCheckSamplerMemory.GetValueOnAnyThread() ||
+					(PCGFeatureSwitches::CVarSamplerMemoryThreshold.GetValueOnAnyThread() * FPlatformMemory::GetStats().AvailablePhysical) >= sizeof(FPCGPoint) * NumIterations64))
 			{
 				NumIterations = static_cast<int32>(NumIterations64);
 			}
 			else
 			{
-				if (Context)
-				{
-					PCGE_LOG_C(Error, GraphAndLog, Context, FText::Format(FText::FromString(TEXT("Skipped - tried to generate too many points ({0} x {1} x {2})")), NumX, NumY, NumZ));
-				}
-
+				PCGLog::LogErrorOnGraph(FText::Format(LOCTEXT("TooManyPoints", "Skipped - tried to generate too many points ({0} x {1} x {2} = {3}).\nAdjust 'pcg.SamplerMemoryThreshold' if needed."), NumX, NumY, NumZ, NumIterations64), Context);
 				return true;
 			}
 		}
diff --git a/Engine/Plugins/PCG/Source/PCG/Private/PCGCommon.cpp b/Engine/Plugins/PCG/Source/PCG/Private/PCGCommon.cpp
--- a/Engine/Plugins/PCG/Source/PCG/Private/PCGCommon.cpp
+++ b/Engine/Plugins/PCG/Source/PCG/Private/PCGCommon.cpp
@@ -11,6 +11,18 @@ namespace PCGFeatureSwitches
 		TEXT("Checks expected memory size consumption prior to performing sampling operations")
 	};
 
+	TAutoConsoleVariable<float> CVarSamplerMemoryThreshold{
+		TEXT("pcg.SamplerMemoryThreshold"),
+		0.8f,
+		TEXT("Normalized threshold of remaining physical memory required to abort sampling operation."),
+		FConsoleVariableDelegate::CreateLambda([](IConsoleVariable* InVariable)
+		{
+			if (InVariable->GetFloat() < 0.f || InVariable->GetFloat() > 1.0)
+			{
+				InVariable->SetWithCurrentPriority(FMath::Clamp(InVariable->GetFloat(), 0.f, 1.f));
+			}
+		})
+	};
 }
 
 namespace PCGHiGenGrid
diff --git a/Engine/Plugins/PCG/Source/PCG/Public/PCGCommon.h b/Engine/Plugins/PCG/Source/PCG/Public/PCGCommon.h
--- a/Engine/Plugins/PCG/Source/PCG/Public/PCGCommon.h
+++ b/Engine/Plugins/PCG/Source/PCG/Public/PCGCommon.h
@@ -183,6 +183,7 @@ namespace PCGObjectMetadata
 namespace PCGFeatureSwitches
 {
 	extern PCG_API TAutoConsoleVariable<bool> CVarCheckSamplerMemory;
+	extern PCG_API TAutoConsoleVariable<float> CVarSamplerMemoryThreshold;
 }
 
 /** Describes space referential for operations that create data */
