[PCG] Propagate graph change to instances correctly

*  On PostEditChange property we should do exactly the same thing than in SetGraph for a graph instance.

#rnx

diff --git a/Engine/Plugins/PCG/Source/PCG/Private/PCGGraph.cpp b/Engine/Plugins/PCG/Source/PCG/Private/PCGGraph.cpp
--- a/Engine/Plugins/PCG/Source/PCG/Private/PCGGraph.cpp
+++ b/Engine/Plugins/PCG/Source/PCG/Private/PCGGraph.cpp
@@ -1764,7 +1764,7 @@ void UPCGGraphInstance::PostEditChangeProperty(FPropertyChangedEvent& PropertyCh
 		// No need to refresh if it is the same graph
 		if (Graph != PreGraphCache)
 		{
-			RefreshParameters(EPCGGraphParameterEvent::GraphChanged);
+			OnGraphParametersChanged(Graph, EPCGGraphParameterEvent::GraphChanged, NAME_None);
 		}
 	}
 	else if (PropertyName == GET_MEMBER_NAME_CHECKED(FPCGOverrideInstancedPropertyBag, PropertiesIDsOverridden))
