[PCG] Fix: Debug object tree objects can now only be selected one at a time

#rnx

diff --git a/Engine/Plugins/PCG/Source/PCGEditor/Private/SPCGEditorGraphDebugObjectTree.cpp b/Engine/Plugins/PCG/Source/PCGEditor/Private/SPCGEditorGraphDebugObjectTree.cpp
--- a/Engine/Plugins/PCG/Source/PCGEditor/Private/SPCGEditorGraphDebugObjectTree.cpp
+++ b/Engine/Plugins/PCG/Source/PCGEditor/Private/SPCGEditorGraphDebugObjectTree.cpp
@@ -259,6 +259,7 @@ void SPCGEditorGraphDebugObjectTree::Construct(const FArguments& InArgs, TShared
 		.OnExpansionChanged(this, &SPCGEditorGraphDebugObjectTree::OnExpansionChanged)
 		.OnSetExpansionRecursive(this, &SPCGEditorGraphDebugObjectTree::OnSetExpansionRecursive)
 		.ItemHeight(18)
+		.SelectionMode(ESelectionMode::Single)
 		.AllowOverscroll(EAllowOverscroll::No)
 		.ExternalScrollbar(VerticalScrollBar)
 		.ConsumeMouseWheel(EConsumeMouseWheel::Always)
