[PCG] Add tooltip in Get Volume Data to indicate to use Get Primitive Data if they want the primitives.

#rnx

diff --git a/Engine/Plugins/PCG/Source/PCG/Private/Elements/PCGTypedGetter.cpp b/Engine/Plugins/PCG/Source/PCG/Private/Elements/PCGTypedGetter.cpp
--- a/Engine/Plugins/PCG/Source/PCG/Private/Elements/PCGTypedGetter.cpp
+++ b/Engine/Plugins/PCG/Source/PCG/Private/Elements/PCGTypedGetter.cpp
@@ -173,7 +173,9 @@ TArray<FPCGPinProperties> UPCGGetVolumeSettings::OutputPinProperties() const
 #if WITH_EDITOR
 FText UPCGGetVolumeSettings::GetNodeTooltipText() const
 {
-	return LOCTEXT("GetVolumeTooltip", "Builds a collection of volumes from the selected actors.");
+	return LOCTEXT("GetVolumeTooltip", "Builds a collection of volumes from the selected actors.\n"
+		"AVolume or APCGPartitionActor produce volume data.\n"
+		"Use GetPrimitiveData for primitive components (i.e like Box, Sphere or Static Mesh collisions).");
 }
 #endif
 
