[PCG] Making sure Infinite bounds for 2D grid on Partition Actor is not producing wrong results

* Streaming bounds are now using the default one.
* Bounds Component on the PA is now a flat box in 2D, and a default subobject in editor only.
* Fix WorldRayHit query default ray origin for parititioned components

#rnx

diff --git a/Engine/Plugins/PCG/Source/PCG/Private/Elements/PCGWorldQuery.cpp b/Engine/Plugins/PCG/Source/PCG/Private/Elements/PCGWorldQuery.cpp
--- a/Engine/Plugins/PCG/Source/PCG/Private/Elements/PCGWorldQuery.cpp
+++ b/Engine/Plugins/PCG/Source/PCG/Private/Elements/PCGWorldQuery.cpp
@@ -93,11 +93,16 @@ bool FPCGWorldRayHitQueryElement::ExecuteInternal(FPCGContext* Context) const
 	
 	// TODO: Support params pin + Apply param data
 
-	// Compute default parameters based on owner component - raycast down local Z axis
 	if (!QueryParams.bOverrideDefaultParams)
 	{
-		AActor* Owner = Context->SourceComponent->GetOwner();
-		const FTransform& Transform = Owner->GetTransform();
+		// Compute default parameters based on original owner component - raycast down local Z axis
+		// TODO: Might want to revisit this for 3D partitioning. The reasoning here is that ray origin should be the same
+		// if we are partitioned or non-partitioned. But with 3D partitioning we might want something different.
+		// Since it is not yet widely used, we'll stick with same behavior for all cases.
+		UPCGComponent* SourceComponent = Context->SourceComponent.Get();
+		UPCGComponent* OriginalComponent = SourceComponent ? SourceComponent->GetOriginalComponent() : nullptr;
+		AActor* Owner = OriginalComponent ? OriginalComponent->GetOwner() : nullptr;
+		const FTransform& Transform = Owner ? Owner->GetTransform() : FTransform::Identity;
 
 		const FBox LocalBounds = PCGHelpers::GetActorLocalBounds(Owner);
 		const FVector RayOrigin = Transform.TransformPosition(FVector(0, 0, LocalBounds.Max.Z));
diff --git a/Engine/Plugins/PCG/Source/PCG/Private/Grid/PCGPartitionActor.cpp b/Engine/Plugins/PCG/Source/PCG/Private/Grid/PCGPartitionActor.cpp
--- a/Engine/Plugins/PCG/Source/PCG/Private/Grid/PCGPartitionActor.cpp
+++ b/Engine/Plugins/PCG/Source/PCG/Private/Grid/PCGPartitionActor.cpp
@@ -27,7 +27,7 @@ APCGPartitionActor::APCGPartitionActor(const FObjectInitializer& ObjectInitializ
 
 #if WITH_EDITOR
 	// Setup bounds component
-	BoundsComponent = CreateDefaultSubobject<UBoxComponent>(TEXT("BoundsComponent"));
+	BoundsComponent = ObjectInitializer.CreateEditorOnlyDefaultSubobject<UBoxComponent>(this, TEXT("BoundsComponent"));
 	BoundsComponent->SetCollisionObjectType(ECC_WorldStatic);
 	BoundsComponent->SetCollisionResponseToAllChannels(ECR_Ignore);
 	BoundsComponent->SetCollisionEnabled(ECollisionEnabled::NoCollision);
@@ -107,6 +107,11 @@ void APCGPartitionActor::Serialize(FArchive& Ar)
 		LocalToOriginal.Reset();
 		for (const auto& It : LocalToOriginalCopy)
 		{
+			if (!ensure(It.Key))
+			{
+				continue;
+			}
+
 			if (!It.Key->HasAnyFlags(RF_Transient))
 			{
 				LocalToOriginal.Add(It);
@@ -162,10 +167,7 @@ void APCGPartitionActor::PostRegisterAllComponents()
 	Super::PostRegisterAllComponents();
 
 #if WITH_EDITOR
-	if (BoundsComponent)
-	{
-		BoundsComponent->SetBoxExtent(GetFixedBounds().GetExtent());
-	}
+	UpdateBoundsComponentExtents();
 #endif // WITH_EDITOR
 
 	// Reset the OriginalToLocal and build it from the local map
@@ -500,11 +502,6 @@ void APCGPartitionActor::RemoveLocalComponent(UPCGComponent* LocalComponent)
 }
 
 #if WITH_EDITOR
-FBox APCGPartitionActor::GetStreamingBounds() const
-{
-	return Super::GetStreamingBounds() + GetFixedBounds();
-}
-
 AActor* APCGPartitionActor::GetSceneOutlinerParent() const
 {
 	if (APCGWorldActor* PCGActor = PCGHelpers::GetPCGWorldActor(GetWorld()))
@@ -534,11 +531,7 @@ void APCGPartitionActor::PostCreation(const FGuid& InGridGUID, uint32 InGridSize
 	}
 
 #if WITH_EDITOR
-	// Since we have infinite bounds in 2D, we just disable the bounds
-	if (BoundsComponent && !bUse2DGrid)
-	{
-		BoundsComponent->SetBoxExtent(GetFixedBounds().GetExtent());
-	}
+	UpdateBoundsComponentExtents();
 #endif // WITH_EDITOR
 
 	RegisterPCG();
@@ -623,6 +616,15 @@ bool APCGPartitionActor::ChangeTransientState(UPCGComponent* OriginalComponent,
 
 	return OriginalToLocal.IsEmpty() && LoadedPreviewComponents.IsEmpty();
 }
+
+void APCGPartitionActor::UpdateBoundsComponentExtents()
+{
+	// Disable BoundsComponent if in 2D
+	if (BoundsComponent && !bUse2DGrid)
+	{
+		BoundsComponent->SetBoxExtent(GetFixedBounds().GetExtent());
+	}
+}
 #endif // WITH_EDITOR
 
 FString APCGPartitionActor::GetPCGPartitionActorName(uint32 GridSize, const FIntVector& GridCoords, bool bRuntimeGenerated)
@@ -642,4 +644,4 @@ FString APCGPartitionActor::GetPCGPartitionActorName(uint32 GridSize, const FInt
 	ActorNameBuilder += FString::Printf(TEXT("%d_%d_%d"), GridCoords.X, GridCoords.Y, GridCoords.Z);
 
 	return ActorNameBuilder.ToString();
-}
+}
\ No newline at end of file
diff --git a/Engine/Plugins/PCG/Source/PCG/Private/Grid/PCGPartitionActor.h b/Engine/Plugins/PCG/Source/PCG/Private/Grid/PCGPartitionActor.h
--- a/Engine/Plugins/PCG/Source/PCG/Private/Grid/PCGPartitionActor.h
+++ b/Engine/Plugins/PCG/Source/PCG/Private/Grid/PCGPartitionActor.h
@@ -39,7 +39,6 @@ public:
 	virtual void PostRegisterAllComponents() override;
 	virtual void PostUnregisterAllComponents() override;
 #if WITH_EDITOR
-	virtual FBox GetStreamingBounds() const override;
 	virtual AActor* GetSceneOutlinerParent() const override;
 #endif
 	//~End AActor Interface
@@ -111,6 +110,10 @@ public:
 private:
 	UPCGSubsystem* GetSubsystem() const;
 
+#if WITH_EDITOR
+	void UpdateBoundsComponentExtents();
+#endif // WITH_EDITOR
+
 	// Note: this map is not a property and not serialized since we will rebuild it from the LocalToOriginal
 	TMap<TObjectPtr<UPCGComponent>, TObjectPtr<UPCGComponent>> OriginalToLocal;
 
