[PCG] Generate on load wasn't working with Landscape surface sampling in cooked builds because LandscapeCache initialize was happening after first call to GraphExecutor->Execute

#rnx

diff --git a/Engine/Plugins/PCG/Source/PCG/Private/PCGSubsystem.cpp b/Engine/Plugins/PCG/Source/PCG/Private/PCGSubsystem.cpp
--- a/Engine/Plugins/PCG/Source/PCG/Private/PCGSubsystem.cpp
+++ b/Engine/Plugins/PCG/Source/PCG/Private/PCGSubsystem.cpp
@@ -237,14 +237,15 @@ void UPCGSubsystem::Tick(float DeltaSeconds)
 		bHasTickedOnce = true;
 	}
 
-	// If we have any tasks to execute, schedule some
-	GraphExecutor->Execute();
-
 	// Lose references to landscape cache as needed
-	if(PCGWorldActor && GetLandscapeCache())
+	// This will also initialize the cache if it isn't already so needs to happen before the GraphExecutor->Executor call
+	if (PCGWorldActor && GetLandscapeCache())
 	{
 		GetLandscapeCache()->Tick(DeltaSeconds);
 	}
+	
+	// If we have any tasks to execute, schedule some
+	GraphExecutor->Execute();
 
 	ActorAndComponentMapping.Tick();
 
