[PCG][Bug] OnPCGGraphStartGenerating, OnPCGGraphCancelled, OnPCGGraphGenerated, and OnPCGGraphCleaned events are now properly exposed to BP

diff --git a/Engine/Plugins/PCG/Source/PCG/Private/PCGComponent.cpp b/Engine/Plugins/PCG/Source/PCG/Private/PCGComponent.cpp
--- a/Engine/Plugins/PCG/Source/PCG/Private/PCGComponent.cpp
+++ b/Engine/Plugins/PCG/Source/PCG/Private/PCGComponent.cpp
@@ -33,6 +33,7 @@
 #include "Utils/PCGGeneratedResourcesLogging.h"
 #include "Utils/PCGGraphExecutionLogging.h"
 
+#include "CoreGlobals.h"
 #include "LandscapeComponent.h"
 #include "LandscapeProxy.h"
 #include "Algo/AllOf.h"
@@ -66,6 +67,15 @@ namespace PCGComponent
 		TEXT("pcg.GlobalDisableRefresh"),
 		false,
 		TEXT("Disable refresh for all PCG Components."));
+
+	template <typename DelegateType>
+	static void BroadcastDynamicDelegate(const DelegateType& Delegate, UPCGComponent* PCGComponent)
+	{
+#if WITH_EDITOR
+		const TGuardValue ScriptExecutionGuard(GAllowActorScriptExecutionInEditor, true);
+#endif // WITH_EDITOR
+		Delegate.Broadcast(PCGComponent);
+	}
 }
 
 UPCGComponent::UPCGComponent(const FObjectInitializer& InObjectInitializer)
@@ -326,13 +336,15 @@ FPCGTaskId UPCGComponent::GenerateInternal(bool bForce, EPCGHiGenGrid Grid, EPCG
 
 	CurrentGenerationTask = GetSubsystem()->ScheduleComponent(this, Grid, bForce, Dependencies);
 
-#if WITH_EDITOR
 	if (CurrentGenerationTask != InvalidPCGTaskId)
 	{
+#if WITH_EDITOR
 		OnPCGGraphStartGeneratingDelegate.Broadcast(this);
-	}
 #endif // WITH_EDITOR
 
+		PCGComponent::BroadcastDynamicDelegate(OnPCGGraphStartGeneratingExternal, this);
+	}
+
 	return CurrentGenerationTask;
 }
 
@@ -451,6 +463,8 @@ void UPCGComponent::PostProcessGraph(const FBox& InNewBounds, bool bInGenerated,
 
 		UpdateDynamicTracking();
 #endif // WITH_EDITOR
+
+		PCGComponent::BroadcastDynamicDelegate(OnPCGGraphGeneratedExternal, this);
 	}
 
 	// Trigger notification - will be used by other tracking mechanisms
@@ -544,6 +558,8 @@ void UPCGComponent::PostCleanupGraph()
 		FCoreUObjectDelegates::OnObjectPropertyChanged.Broadcast(this, GeneratedOutputChangedEvent);
 	}
 #endif
+
+	PCGComponent::BroadcastDynamicDelegate(OnPCGGraphCleanedExternal, this);
 }
 
 void UPCGComponent::OnProcessGraphAborted(bool bQuiet)
@@ -574,6 +590,8 @@ void UPCGComponent::OnProcessGraphAborted(bool bQuiet)
 
 	GetSubsystem()->OnComponentGenerationCompleteOrCancelled.Broadcast();
 #endif
+
+	PCGComponent::BroadcastDynamicDelegate(OnPCGGraphCancelledExternal, this);
 }
 
 void UPCGComponent::Cleanup()
@@ -3378,6 +3396,12 @@ void FPCGComponentInstanceData::ApplyToComponent(UActorComponent* Component, con
 			PCGComponent->LoadedPreviewResources = DuplicateLoadedPreviewResources;
 		}
 
+		// Move over invocation lists for dynamic delegates
+		PCGComponent->OnPCGGraphStartGeneratingExternal = SourceComponent->OnPCGGraphStartGeneratingExternal;
+		PCGComponent->OnPCGGraphCancelledExternal = SourceComponent->OnPCGGraphCancelledExternal;
+		PCGComponent->OnPCGGraphGeneratedExternal = SourceComponent->OnPCGGraphGeneratedExternal;
+		PCGComponent->OnPCGGraphCleanedExternal = SourceComponent->OnPCGGraphCleanedExternal;
+
 		// Reconnect callbacks
 		if (PCGComponent->GraphInstance)
 		{
diff --git a/Engine/Plugins/PCG/Source/PCG/Public/PCGComponent.h b/Engine/Plugins/PCG/Source/PCG/Public/PCGComponent.h
--- a/Engine/Plugins/PCG/Source/PCG/Public/PCGComponent.h
+++ b/Engine/Plugins/PCG/Source/PCG/Public/PCGComponent.h
@@ -41,6 +41,11 @@ struct FPCGContext;
 	DECLARE_MULTICAST_DELEGATE_OneParam(FOnPCGGraphCleaned, UPCGComponent*);
 #endif
 
+DECLARE_DYNAMIC_MULTICAST_DELEGATE_OneParam(FOnPCGGraphStartGeneratingExternal, UPCGComponent*, PCGComponent);
+DECLARE_DYNAMIC_MULTICAST_DELEGATE_OneParam(FOnPCGGraphCancelledExternal, UPCGComponent*, PCGComponent);
+DECLARE_DYNAMIC_MULTICAST_DELEGATE_OneParam(FOnPCGGraphGeneratedExternal, UPCGComponent*, PCGComponent);
+DECLARE_DYNAMIC_MULTICAST_DELEGATE_OneParam(FOnPCGGraphCleanedExternal, UPCGComponent*, PCGComponent);
+
 UENUM(Blueprintable)
 enum class EPCGComponentInput : uint8
 {
@@ -258,11 +263,24 @@ public:
 	UPROPERTY()
 	bool bForceGenerateOnBPAddedToWorld = false;
 
-	FOnPCGGraphGenerated OnPCGGraphGeneratedDelegate;
 	FOnPCGGraphStartGenerating OnPCGGraphStartGeneratingDelegate;
-	FOnPCGGraphCleaned OnPCGGraphCleanedDelegate;
 	FOnPCGGraphCancelled OnPCGGraphCancelledDelegate;
-#endif
+	FOnPCGGraphGenerated OnPCGGraphGeneratedDelegate;
+	FOnPCGGraphCleaned OnPCGGraphCleanedDelegate;
+#endif // WITH_EDITORONLY_DATA
+
+	/** Event dispatched when a graph begins generation on this component. */
+	UPROPERTY(BlueprintAssignable, meta = (DisplayName = "On Graph Started Generating"))
+	FOnPCGGraphStartGeneratingExternal OnPCGGraphStartGeneratingExternal;
+	/** Event dispatched when a graph cancels generation on this component. */
+	UPROPERTY(BlueprintAssignable, meta = (DisplayName = "On Graph Cancelled"))
+	FOnPCGGraphCancelledExternal OnPCGGraphCancelledExternal;
+	/** Event dispatched when a graph completes its generation on this component. */
+	UPROPERTY(BlueprintAssignable, meta = (DisplayName = "On Graph Generated"))
+	FOnPCGGraphGeneratedExternal OnPCGGraphGeneratedExternal;
+	/** Event dispatched when a graph cleans on this component. */
+	UPROPERTY(BlueprintAssignable, meta = (DisplayName = "On Graph Cleaned"))
+	FOnPCGGraphCleanedExternal OnPCGGraphCleanedExternal;
 
 	/** Flag to indicate whether this component has run in the editor. Note that for partitionable actors, this will always be false. */
 	UPROPERTY(BlueprintReadOnly, VisibleAnywhere, Category = Debug, NonTransactional, meta = (NoResetToDefault))
