[PCG][Bug] Fixed a crash in the PCG Enum Selector causing a crash when resetting to default

- Removed the Reset To Default option, as it doesn't make sense to have a "default" Enum Selection

#rnx

diff --git a/Engine/Plugins/PCG/Source/PCGEditor/Private/Details/EnumSelectorDetails.cpp b/Engine/Plugins/PCG/Source/PCGEditor/Private/Details/EnumSelectorDetails.cpp
--- a/Engine/Plugins/PCG/Source/PCGEditor/Private/Details/EnumSelectorDetails.cpp
+++ b/Engine/Plugins/PCG/Source/PCGEditor/Private/Details/EnumSelectorDetails.cpp
@@ -124,7 +124,8 @@ void FEnumSelectorDetails::CustomizeHeader(const TSharedRef<IPropertyHandle> InP
 					]
 				]
 			]
-		];
+		]
+		.OverrideResetToDefault(FResetToDefaultOverride::Hide());
 }
 
 FText FEnumSelectorDetails::StripEnumClassName(const FName EnumValueName)
