[PCG] Get Actor Bounds should gives the same result than Get Local Actor Bounds for Partition actors

* Also re-include child actors in bounds computation, as it was a wrong call to remove it and was not the issue.

#rnx

diff --git a/Engine/Plugins/PCG/Source/PCG/Private/Helpers/PCGHelpers.cpp b/Engine/Plugins/PCG/Source/PCG/Private/Helpers/PCGHelpers.cpp
--- a/Engine/Plugins/PCG/Source/PCG/Private/Helpers/PCGHelpers.cpp
+++ b/Engine/Plugins/PCG/Source/PCG/Private/Helpers/PCGHelpers.cpp
@@ -110,22 +110,28 @@ namespace PCGHelpers
 		// This is to ensure stable bounds and no timing issues (cleared ISMs, etc.)
 		FBox Box(EForceInit::ForceInit);
 
-		const bool bNonColliding = true;
-		// FIXME: This option was disabled because of wrong bounds with components attached to other components. (UE-205314)
-		const bool bIncludeFromChildActors = false;
-
 		if (InActor)
 		{
-			InActor->ForEachComponent<UPrimitiveComponent>(bIncludeFromChildActors, [bNonColliding, bIgnorePCGCreatedComponents, &Box](const UPrimitiveComponent* InPrimComp)
+			if (const APCGPartitionActor* PartitionActor = Cast<const APCGPartitionActor>(InActor))
+			{
+				// Skip per-component check, return fixed bounds.
+				Box = PartitionActor->GetFixedBounds();
+			}
+			else
 			{
-				// Note: we omit the IsRegistered check here (e.g. InPrimComp->IsRegistered() )
-				// since this can be called in a scope where the components are temporarily unregistered
-				if ((bNonColliding || InPrimComp->IsCollisionEnabled()) &&
-					(!bIgnorePCGCreatedComponents || !InPrimComp->ComponentTags.Contains(DefaultPCGTag)))
+				const bool bNonColliding = true;
+
+				InActor->ForEachComponent<UPrimitiveComponent>(/*bIncludeFromChildActors=*/true, [bNonColliding, bIgnorePCGCreatedComponents, &Box](const UPrimitiveComponent* InPrimComp)
 				{
-					Box += InPrimComp->Bounds.GetBox();
-				}
-			});
+					// Note: we omit the IsRegistered check here (e.g. InPrimComp->IsRegistered() )
+					// since this can be called in a scope where the components are temporarily unregistered
+					if ((bNonColliding || InPrimComp->IsCollisionEnabled()) &&
+						(!bIgnorePCGCreatedComponents || !InPrimComp->ComponentTags.Contains(DefaultPCGTag)))
+					{
+						Box += InPrimComp->Bounds.GetBox();
+					}
+				});
+			}
 		}
 		else
 		{
@@ -152,13 +158,11 @@ namespace PCGHelpers
 			else
 			{
 				const bool bNonColliding = true;
-				// FIXME: This option was disabled because of wrong bounds with components attached to other components. (UE-205314)
-				const bool bIncludeFromChildActors = false;
 
 				const FTransform& ActorToWorld = InActor->GetTransform();
 				const FTransform WorldToActor = ActorToWorld.Inverse();
 
-				InActor->ForEachComponent<UPrimitiveComponent>(bIncludeFromChildActors, [bNonColliding, bIgnorePCGCreatedComponents, &WorldToActor, &Box](const UPrimitiveComponent* InPrimComp)
+				InActor->ForEachComponent<UPrimitiveComponent>(/*bIncludeFromChildActors=*/true, [bNonColliding, bIgnorePCGCreatedComponents, &WorldToActor, &Box](const UPrimitiveComponent* InPrimComp)
 				{
 					if ((bNonColliding || InPrimComp->IsCollisionEnabled()) &&
 						(!bIgnorePCGCreatedComponents || !InPrimComp->ComponentTags.Contains(DefaultPCGTag)))
@@ -309,7 +313,7 @@ namespace PCGHelpers
 	{
 		UClass* ObjectClass = Object ? Object->GetClass() : nullptr;
 
-		if(!CanBeExpanded(ObjectClass))
+		if (!CanBeExpanded(ObjectClass))
 		{
 			return;
 		}
@@ -334,7 +338,7 @@ namespace PCGHelpers
 			if (Object && !OutDependencies.Contains(Object))
 			{
 				// If we explicitly don't want to track this object, early out.
-				if (!Object->GetClass() || 
+				if (!Object->GetClass() ||
 					!CanBeExpanded(Object->GetClass()) ||
 					Algo::AnyOf(InExcludedClasses, [InClass = Object->GetClass()](const UClass* ExcludedClass) { return InClass->IsChildOf(ExcludedClass); }))
 				{
@@ -391,11 +395,11 @@ namespace PCGHelpers
 			FScriptMapHelper_InContainer Helper(MapProperty, InContainer);
 			for (FScriptMapHelper::FIterator It(Helper); It; ++It)
 			{
-					// Key and Value are stored next to each other in memory.
-					// ValueProp has an offset, so we should use the same starting address for both.
-					const void* PairKeyValuePtr = Helper.GetKeyPtr(It);
-					GatherDependencies(MapProperty->KeyProp, PairKeyValuePtr, OutDependencies, MaxDepth, InExcludedClasses);
-					GatherDependencies(MapProperty->ValueProp, PairKeyValuePtr, OutDependencies, MaxDepth, InExcludedClasses);
+				// Key and Value are stored next to each other in memory.
+				// ValueProp has an offset, so we should use the same starting address for both.
+				const void* PairKeyValuePtr = Helper.GetKeyPtr(It);
+				GatherDependencies(MapProperty->KeyProp, PairKeyValuePtr, OutDependencies, MaxDepth, InExcludedClasses);
+				GatherDependencies(MapProperty->ValueProp, PairKeyValuePtr, OutDependencies, MaxDepth, InExcludedClasses);
 			}
 		}
 		else if (FSetProperty* SetProperty = CastField<FSetProperty>(Property))
@@ -479,7 +483,7 @@ namespace PCGHelpers
 			InActorToAttach->AttachToActor(InParent, FAttachmentTransformRules::KeepWorldTransform);
 		}
 #if WITH_EDITOR
-		else if(AttachOptions == EPCGAttachOptions::InFolder)
+		else if (AttachOptions == EPCGAttachOptions::InFolder)
 		{
 			FString DefaultFolderPath;
 
