[PCG] Fix crash when BoundsComponent cannot be created in PartitionActors

#rnx

diff --git a/Engine/Plugins/PCG/Source/PCG/Private/Grid/PCGPartitionActor.cpp b/Engine/Plugins/PCG/Source/PCG/Private/Grid/PCGPartitionActor.cpp
--- a/Engine/Plugins/PCG/Source/PCG/Private/Grid/PCGPartitionActor.cpp
+++ b/Engine/Plugins/PCG/Source/PCG/Private/Grid/PCGPartitionActor.cpp
@@ -28,12 +28,15 @@ APCGPartitionActor::APCGPartitionActor(const FObjectInitializer& ObjectInitializ
 #if WITH_EDITOR
 	// Setup bounds component
 	BoundsComponent = ObjectInitializer.CreateEditorOnlyDefaultSubobject<UBoxComponent>(this, TEXT("BoundsComponent"));
-	BoundsComponent->SetCollisionObjectType(ECC_WorldStatic);
-	BoundsComponent->SetCollisionResponseToAllChannels(ECR_Ignore);
-	BoundsComponent->SetCollisionEnabled(ECollisionEnabled::NoCollision);
-	BoundsComponent->SetGenerateOverlapEvents(false);
-	BoundsComponent->SetupAttachment(GetRootComponent());
-	BoundsComponent->bDrawOnlyIfSelected = true;
+	if (BoundsComponent)
+	{
+		BoundsComponent->SetCollisionObjectType(ECC_WorldStatic);
+		BoundsComponent->SetCollisionResponseToAllChannels(ECR_Ignore);
+		BoundsComponent->SetCollisionEnabled(ECollisionEnabled::NoCollision);
+		BoundsComponent->SetGenerateOverlapEvents(false);
+		BoundsComponent->SetupAttachment(GetRootComponent());
+		BoundsComponent->bDrawOnlyIfSelected = true;
+	}
 #endif // WITH_EDITOR
 }
 
