[PCG] Don't ignore changes to GraphInterface in OnObjectChanged now that we can react to object changes (and not just actor changes)

* Also add natvis visualization for selector keys for easier debug.

#rnx

diff --git a/Engine/Plugins/PCG/Source/PCG/Private/PCGActorAndComponentMapping.cpp b/Engine/Plugins/PCG/Source/PCG/Private/PCGActorAndComponentMapping.cpp
--- a/Engine/Plugins/PCG/Source/PCG/Private/PCGActorAndComponentMapping.cpp
+++ b/Engine/Plugins/PCG/Source/PCG/Private/PCGActorAndComponentMapping.cpp
@@ -1630,11 +1630,10 @@ void FPCGActorAndComponentMapping::OnObjectChanged(UObject* InObject, const FAct
 		return;
 	}
 
-	// Don't react to what the PCG Component is already reacting to (itself and graphs).
+	// Don't react to what the PCG Component is already reacting to.
 	static const TArray<const UClass*> ExcludedClasses =
 	{
-		UPCGComponent::StaticClass(),
-		UPCGGraphInterface::StaticClass()
+		UPCGComponent::StaticClass()
 	};
 
 	if (Algo::AnyOf(ExcludedClasses, [InObject](const UClass* Class) -> bool { return InObject->IsA(Class); }))
