[PCG] Fix incorrect logic in CleanupDelayedRefreshComponents() that caused RuntimeGen inspection to fail

#rnx

diff --git a/Engine/Plugins/PCG/Source/PCG/Private/RuntimeGen/PCGRuntimeGenScheduler.cpp b/Engine/Plugins/PCG/Source/PCG/Private/RuntimeGen/PCGRuntimeGenScheduler.cpp
--- a/Engine/Plugins/PCG/Source/PCG/Private/RuntimeGen/PCGRuntimeGenScheduler.cpp
+++ b/Engine/Plugins/PCG/Source/PCG/Private/RuntimeGen/PCGRuntimeGenScheduler.cpp
@@ -366,7 +366,7 @@ void FPCGRuntimeGenScheduler::TickQueueComponentsForGeneration(
 
 								// Verify the grid cell actually lies within the generation radius.
 								// TODO: this is no longer necessary if we rasterize the sphere instead.
-								const double LocalDistanceSquared = IntersectedBounds.ComputeSquaredDistanceToPoint(ModifiedGenSourcePosition);
+								const double LocalDistanceSquared = CellBounds.ComputeSquaredDistanceToPoint(ModifiedGenSourcePosition);
 								if (LocalDistanceSquared <= GenerationRadius * GenerationRadius)
 								{
 									AddComponentToGenerate(Key, GenSource, Policy, IntersectedBounds, InPCGWorldActor->bUse2DGrid);
@@ -498,14 +498,15 @@ void FPCGRuntimeGenScheduler::TickCleanup(const TSet<IPCGGenSourceBase*>& InGenS
 		else
 		{
 			UPCGComponent* LocalComponent = ActorAndComponentMapping->GetLocalComponent(GridSize, GridCoords, OriginalComponent, /*bRuntimeGenerated=*/true);
-			if (!LocalComponent)
+			APCGPartitionActor* PartitionActor = LocalComponent ? Cast<APCGPartitionActor>(LocalComponent->GetOwner()) : nullptr;
+			if (!PartitionActor)
 			{
 				// Attempt to clean even in failure case to avoid leaking resources.
 				ComponentsToClean.Add({ GenerationKey, LocalComponent });
 				continue;
 			}
 
-			const FBox GridBounds = LocalComponent->GetGridBounds();
+			const FBox GridBounds = PartitionActor->GetFixedBounds();
 
 			double MinSquaredDistanceToGenSource = UE_DOUBLE_BIG_NUMBER;
 			for (const IPCGGenSourceBase* GenSource : InGenSources)
@@ -900,7 +901,7 @@ void FPCGRuntimeGenScheduler::CleanupDelayedRefreshComponents()
 
 	check(ActorAndComponentMapping);
 
-	// Check that each refreshed local component is still inside its original component.
+	// Check that each refreshed local component is still intersecting its original component.
 	// If it is not, it would be leaked instead of refreshed, so we should force a full cleanup.
 	for (const FGridGenerationKey& GenerationKey : GeneratedComponentsToRemove)
 	{
@@ -917,20 +918,21 @@ void FPCGRuntimeGenScheduler::CleanupDelayedRefreshComponents()
 		const FIntVector& GridCoords = GenerationKey.GetGridCoords();
 
 		UPCGComponent* LocalComponent = OriginalComponent ? ActorAndComponentMapping->GetLocalComponent(GridSize, GridCoords, OriginalComponent, /*bRuntimeGenerated=*/true) : nullptr;
+		APCGPartitionActor* PartitionActor = LocalComponent ? Cast<APCGPartitionActor>(LocalComponent->GetOwner()) : nullptr;
 
-		if (LocalComponent)
+		if (LocalComponent && PartitionActor)
 		{
 			const FBox OriginalBounds = OriginalComponent->GetGridBounds();
-			const FBox LocalBounds = LocalComponent->GetGridBounds();
+			const FBox LocalBounds = PartitionActor->GetFixedBounds();
 
-			if (!LocalBounds.IsInsideOrOn(OriginalBounds))
+			if (!OriginalBounds.Intersect(LocalBounds))
 			{
 				CleanupComponent(GenerationKey, LocalComponent);
 			}
 		}
 		else
 		{
-			// If the original or local component no longer exists, just clean up.
+			// If the component or partition actor could not be recovered, just clean up.
 			CleanupComponent(GenerationKey, /*GenerationKey=*/nullptr);
 		}
 	}
