[PCG] Make sure the attribute list view window keeps any collapsed point data alive as the window holds raw pointers to the data.

#rnx

diff --git a/Engine/Plugins/PCG/Source/PCGEditor/Private/SPCGEditorGraphAttributeListView.cpp b/Engine/Plugins/PCG/Source/PCGEditor/Private/SPCGEditorGraphAttributeListView.cpp
--- a/Engine/Plugins/PCG/Source/PCGEditor/Private/SPCGEditorGraphAttributeListView.cpp
+++ b/Engine/Plugins/PCG/Source/PCGEditor/Private/SPCGEditorGraphAttributeListView.cpp
@@ -9,7 +9,6 @@
 #include "PCGParamData.h"
 #include "PCGPin.h"
 #include "PCGSubsystem.h"
-#include "Data/PCGPointData.h"
 #include "Data/PCGSpatialData.h"
 #include "Metadata/PCGAttributePropertySelector.h"
 #include "Metadata/Accessors/IPCGAttributeAccessor.h"
@@ -315,6 +314,8 @@ SPCGEditorGraphAttributeListView::~SPCGEditorGraphAttributeListView()
 	{
 		PCGEditorPtr.Pin()->OnInspectedStackChangedDelegate.RemoveAll(this);
 	}
+
+	CollapsedPointData.Reset();
 }
 
 void SPCGEditorGraphAttributeListView::Construct(const FArguments& InArgs, TSharedPtr<FPCGEditor> InPCGEditor)
@@ -679,6 +680,7 @@ void SPCGEditorGraphAttributeListView::RefreshAttributeList()
 	ListViewItems.Empty();
 	ListViewHeader->ClearColumns();
 	InfoTextBlock->SetText(FText::GetEmpty());
+	CollapsedPointData.Reset();
 
 	const FPCGDataCollection* InspectionData = GetInspectionData();
 	if (!InspectionData)
@@ -727,6 +729,8 @@ void SPCGEditorGraphAttributeListView::RefreshAttributeList()
 	{
 		if (const UPCGPointData* PCGPointData = PCGSpatialData->ToPointData())
 		{
+			CollapsedPointData.Reset(PCGPointData);
+
 			const UPCGMetadata* PCGMetadata = PCGPointData->ConstMetadata();
 
 			AddPointDataColumns(PCGPointData);
diff --git a/Engine/Plugins/PCG/Source/PCGEditor/Private/SPCGEditorGraphAttributeListView.h b/Engine/Plugins/PCG/Source/PCGEditor/Private/SPCGEditorGraphAttributeListView.h
--- a/Engine/Plugins/PCG/Source/PCGEditor/Private/SPCGEditorGraphAttributeListView.h
+++ b/Engine/Plugins/PCG/Source/PCGEditor/Private/SPCGEditorGraphAttributeListView.h
@@ -2,12 +2,14 @@
 
 #pragma once
 
+#include "Data/PCGPointData.h"
 #include "Graph/PCGStackContext.h"
 #include "Metadata/Accessors/IPCGAttributeAccessor.h"
 #include "Metadata/Accessors/PCGAttributeAccessorKeys.h"
 
 #include "Misc/TextFilterExpressionEvaluator.h"
 #include "Tasks/Task.h"
+#include "UObject/StrongObjectPtr.h"
 #include "Widgets/Views/SListView.h"
 #include "Widgets/Views/STableRow.h"
 
@@ -260,4 +262,7 @@ private:
 	bool bIsLocked : 1 = false;
 
 	TSharedPtr<FPCGListViewUpdater> CurrentUpdateTask = nullptr;
+
+	/** Used to ensure data collapsed for inspection is kept alive. */
+	TStrongObjectPtr<const UPCGPointData> CollapsedPointData = nullptr;
 };
