[PCG] Fix PAs ending up with broken bounds after PCG builder job. The PCG builder needs to save in Run() rather than in PostRun() as we need to save before unload.

#rnx

diff --git a/Engine/Plugins/PCG/Source/PCGEditor/Private/WorldPartitionBuilder/PCGWorldPartitionBuilder.cpp b/Engine/Plugins/PCG/Source/PCGEditor/Private/WorldPartitionBuilder/PCGWorldPartitionBuilder.cpp
--- a/Engine/Plugins/PCG/Source/PCGEditor/Private/WorldPartitionBuilder/PCGWorldPartitionBuilder.cpp
+++ b/Engine/Plugins/PCG/Source/PCGEditor/Private/WorldPartitionBuilder/PCGWorldPartitionBuilder.cpp
@@ -235,19 +235,18 @@ bool UPCGWorldPartitionBuilder::RunInternal(UWorld* World, const FCellInfo& InCe
 		UE_LOG(LogPCGWorldPartitionBuilder, Display, TEXT("Generation complete, %d packages dirtied."), NumPackagesDirtied);
 	}
 
-	return bGeneratedAnyComponent;
-}
-
-bool UPCGWorldPartitionBuilder::PostRun(UWorld* World, FPackageSourceControlHelper& PackageHelper, const bool bInRunSuccess)
-{
-	// bInRunSuccess is return value of Run() which is true if any component was scheduled.
-	if (!bInRunSuccess)
+	if (!bGeneratedAnyComponent)
 	{
 		UE_LOG(LogPCGWorldPartitionBuilder, Display, TEXT("Dirty package detection and save skipped due to trivial run"));
-
-		return true;
+		return false;
 	}
 
+	// TODO: Review the save flow when we have iterative loading.
+	return SaveDirtyPackages(World, PackageHelper);
+}
+
+bool UPCGWorldPartitionBuilder::SaveDirtyPackages(UWorld* World, FPackageSourceControlHelper& PackageHelper)
+{
 	// Check whether an error was thrown while generating components and fail the builder if the ignore argument is not provided.
 	if (bErrorOccurredWhileGenerating)
 	{
diff --git a/Engine/Plugins/PCG/Source/PCGEditor/Private/WorldPartitionBuilder/PCGWorldPartitionBuilder.h b/Engine/Plugins/PCG/Source/PCGEditor/Private/WorldPartitionBuilder/PCGWorldPartitionBuilder.h
--- a/Engine/Plugins/PCG/Source/PCGEditor/Private/WorldPartitionBuilder/PCGWorldPartitionBuilder.h
+++ b/Engine/Plugins/PCG/Source/PCGEditor/Private/WorldPartitionBuilder/PCGWorldPartitionBuilder.h
@@ -31,8 +31,10 @@ public:
 
 protected:
 	virtual bool RunInternal(UWorld* World, const FCellInfo& InCellInfo, FPackageSourceControlHelper& PackageHelper) override;
-	virtual bool PostRun(UWorld* World, FPackageSourceControlHelper& PackageHelper, const bool bInRunSuccess) override;
-	virtual bool CanProcessNonPartitionedWorlds() const { return true; }
+	virtual bool CanProcessNonPartitionedWorlds() const override { return true; }
+
+	/** Save all the pending dirty and deleted packages. */
+	virtual bool SaveDirtyPackages(UWorld* World, FPackageSourceControlHelper& PackageHelper);
 
 private:
 	/** The packages dirtied while generating all components. */
