[PCG] Fixed a bug with ISM Body Instance Customization not displaying the correct channels

diff --git a/Engine/Plugins/PCG/Source/PCG/Private/MeshSelectors/PCGMeshSelectorByAttribute.cpp b/Engine/Plugins/PCG/Source/PCG/Private/MeshSelectors/PCGMeshSelectorByAttribute.cpp
--- a/Engine/Plugins/PCG/Source/PCG/Private/MeshSelectors/PCGMeshSelectorByAttribute.cpp
+++ b/Engine/Plugins/PCG/Source/PCG/Private/MeshSelectors/PCGMeshSelectorByAttribute.cpp
@@ -99,6 +99,10 @@ void UPCGMeshSelectorByAttribute::PostLoad()
 		WorldPositionOffsetDisableDistance_DEPRECATED = 0;
 	}
 #endif
+
+	// TODO: Remove if/when FBodyInstance is updated or replaced
+	// Necessary to update the collision Response Container from the Response Array
+	TemplateDescriptor.PostLoadFixup(this);
 }
 
 bool UPCGMeshSelectorByAttribute::SelectInstances(
diff --git a/Engine/Plugins/PCG/Source/PCG/Private/MeshSelectors/PCGMeshSelectorWeighted.cpp b/Engine/Plugins/PCG/Source/PCG/Private/MeshSelectors/PCGMeshSelectorWeighted.cpp
--- a/Engine/Plugins/PCG/Source/PCG/Private/MeshSelectors/PCGMeshSelectorWeighted.cpp
+++ b/Engine/Plugins/PCG/Source/PCG/Private/MeshSelectors/PCGMeshSelectorWeighted.cpp
@@ -291,12 +291,18 @@ void UPCGMeshSelectorWeighted::PostLoad()
 {
 	Super::PostLoad();
 
-#if WITH_EDITOR
 	for (FPCGMeshSelectorWeightedEntry& Entry : MeshEntries)
 	{
+#if WITH_EDITOR
 		Entry.ApplyDeprecation();
+#endif // WITH_EDITOR
+
+		// TODO: Remove if/when FBodyInstance is updated or replaced
+		// Necessary to update the collision Response Container from the Response Array
+		Entry.Descriptor.PostLoadFixup(this);
 	}
 
+#if WITH_EDITOR
 	RefreshDisplayNames();
 #endif // WITH_EDITOR
 }
diff --git a/Engine/Plugins/PCG/Source/PCG/Private/MeshSelectors/PCGMeshSelectorWeightedByCategory.cpp b/Engine/Plugins/PCG/Source/PCG/Private/MeshSelectors/PCGMeshSelectorWeightedByCategory.cpp
--- a/Engine/Plugins/PCG/Source/PCG/Private/MeshSelectors/PCGMeshSelectorWeightedByCategory.cpp
+++ b/Engine/Plugins/PCG/Source/PCG/Private/MeshSelectors/PCGMeshSelectorWeightedByCategory.cpp
@@ -300,12 +300,21 @@ void UPCGMeshSelectorWeightedByCategory::PostLoad()
 {
 	Super::PostLoad();
 
-#if WITH_EDITOR
 	for (FPCGWeightedByCategoryEntryList& Entry : Entries)
 	{
+#if WITH_EDITOR
 		Entry.ApplyDeprecation();
+#endif // WITH_EDITOR
+
+		// TODO: Remove if/when FBodyInstance is updated or replaced
+		// Necessary to update the collision Response Container from the Response Array
+		for (FPCGMeshSelectorWeightedEntry& WeightedEntry : Entry.WeightedMeshEntries)
+		{
+			WeightedEntry.Descriptor.PostLoadFixup(this);
+		}
 	}
 
+#if WITH_EDITOR
 	RefreshDisplayNames();
 #endif // WITH_EDITOR
 }
diff --git a/Engine/Plugins/PCG/Source/PCG/Public/MeshSelectors/PCGMeshSelectorByAttribute.h b/Engine/Plugins/PCG/Source/PCG/Public/MeshSelectors/PCGMeshSelectorByAttribute.h
--- a/Engine/Plugins/PCG/Source/PCG/Public/MeshSelectors/PCGMeshSelectorByAttribute.h
+++ b/Engine/Plugins/PCG/Source/PCG/Public/MeshSelectors/PCGMeshSelectorByAttribute.h
@@ -13,7 +13,7 @@ class PCG_API UPCGMeshSelectorByAttribute : public UPCGMeshSelectorBase
 
 public:
 	// ~Begin UObject interface
-	void PostLoad() override;
+	virtual void PostLoad() override;
 	// ~End UObject interface
 
 	virtual bool SelectInstances(
diff --git a/Engine/Source/Runtime/Engine/Private/ISMPartition/ISMComponentDescriptor.cpp b/Engine/Source/Runtime/Engine/Private/ISMPartition/ISMComponentDescriptor.cpp
--- a/Engine/Source/Runtime/Engine/Private/ISMPartition/ISMComponentDescriptor.cpp
+++ b/Engine/Source/Runtime/Engine/Private/ISMPartition/ISMComponentDescriptor.cpp
@@ -165,6 +165,14 @@ void FSoftISMComponentDescriptor::InitFrom(const UStaticMeshComponent* Template,
 	Super::InitFrom(Template, bInitBodyInstance);
 }
 
+void FISMComponentDescriptorBase::PostLoadFixup(UObject* Loader)
+{
+	check(Loader);
+
+	// Necessary to update the collision Response Container from the array
+	BodyInstance.FixupData(Loader);
+}
+
 bool FISMComponentDescriptorBase::operator!=(const FISMComponentDescriptorBase& Other) const
 {
 	return !(*this == Other);
diff --git a/Engine/Source/Runtime/Engine/Public/ISMPartition/ISMComponentDescriptor.h b/Engine/Source/Runtime/Engine/Public/ISMPartition/ISMComponentDescriptor.h
--- a/Engine/Source/Runtime/Engine/Public/ISMPartition/ISMComponentDescriptor.h
+++ b/Engine/Source/Runtime/Engine/Public/ISMPartition/ISMComponentDescriptor.h
@@ -29,6 +29,8 @@ struct FISMComponentDescriptorBase
 	ENGINE_API virtual uint32 ComputeHash() const;
 	ENGINE_API virtual void InitComponent(UInstancedStaticMeshComponent* ISMComponent) const;
 
+	ENGINE_API void PostLoadFixup(UObject* Loader);
+
 	ENGINE_API bool operator!=(const FISMComponentDescriptorBase& Other) const;
 	ENGINE_API bool operator==(const FISMComponentDescriptorBase& Other) const;
 
