[PCG] Clamp grid size to graph default grid size to fix grid(s) not executing bug

#rnx

diff --git a/Engine/Plugins/PCG/Source/PCG/Private/Elements/PCGHiGenGridSize.cpp b/Engine/Plugins/PCG/Source/PCG/Private/Elements/PCGHiGenGridSize.cpp
--- a/Engine/Plugins/PCG/Source/PCG/Private/Elements/PCGHiGenGridSize.cpp
+++ b/Engine/Plugins/PCG/Source/PCG/Private/Elements/PCGHiGenGridSize.cpp
@@ -119,6 +119,15 @@ bool FPCGHiGenGridSizeElement::ExecuteInternal(FPCGContext* Context) const
 	TRACE_CPUPROFILER_EVENT_SCOPE(FPCGHiGenGridSizeElement::Execute);
 	check(Context);
 
+	const UPCGHiGenGridSizeSettings* Settings = Context->GetInputSettings<UPCGHiGenGridSizeSettings>();
+	check(Settings);
+
+	const UPCGGraph* Graph = Context->Node ? Context->Node->GetGraph() : nullptr;
+	if (Graph && Graph->IsHierarchicalGenerationEnabled() && (Graph->GetDefaultGrid() < Settings->GetGrid()))
+	{
+		PCGE_LOG(Warning, GraphAndLog, LOCTEXT("GridSizeLargerThanGraphGridSize", "Grid size is larger than graph default grid size and will be automatically clamped."));
+	}
+
 	// Trivial pass through. Will only execute on the prescribed grid.
 	Context->OutputData = Context->InputData;
 	for (FPCGTaggedData& Data : Context->OutputData.TaggedData)
diff --git a/Engine/Plugins/PCG/Source/PCG/Private/Graph/PCGGraphCompiler.cpp b/Engine/Plugins/PCG/Source/PCG/Private/Graph/PCGGraphCompiler.cpp
--- a/Engine/Plugins/PCG/Source/PCG/Private/Graph/PCGGraphCompiler.cpp
+++ b/Engine/Plugins/PCG/Source/PCG/Private/Graph/PCGGraphCompiler.cpp
@@ -445,7 +445,7 @@ EPCGHiGenGrid FPCGGraphCompiler::CalculateGridRecursive(
 	// Grid Size nodes in the top graph set the execution grid level.
 	if (GridSizeSettings && GridSizeSettings->bEnabled && bTopLevelGraph)
 	{
-		Grid = GridSizeSettings->GetGrid();
+		Grid = FMath::Min(GenerationDefaultGrid, GridSizeSettings->GetGrid());
 	}
 	else
 	{
diff --git a/Engine/Plugins/PCG/Source/PCG/Private/PCGGraph.cpp b/Engine/Plugins/PCG/Source/PCG/Private/PCGGraph.cpp
--- a/Engine/Plugins/PCG/Source/PCG/Private/PCGGraph.cpp
+++ b/Engine/Plugins/PCG/Source/PCG/Private/PCGGraph.cpp
@@ -1617,17 +1617,16 @@ uint32 UPCGGraph::CalculateNodeGridSizeRecursive_Unsafe(const UPCGNode* InNode,
 		return *CachedGridSize;
 	}
 
-	uint32 GridSize = PCGHiGenGrid::UninitializedGridSize();
+	uint32 GridSize = InDefaultGridSize;
 
 	const UPCGHiGenGridSizeSettings* GridSizeSettings = Cast<UPCGHiGenGridSizeSettings>(InNode->GetSettings());
 	if (GridSizeSettings && GridSizeSettings->bEnabled)
 	{
-		GridSize = GridSizeSettings->GetGridSize();
+		GridSize = FMath::Min(GridSize, GridSizeSettings->GetGridSize());
 	}
 	else
 	{
 		// Grid size for a node is the minimum of the grid sizes of connected upstream nodes.
-		uint32 MinGenerationSize = std::numeric_limits<uint32>::max();
 		for (const UPCGPin* Pin : InNode->GetInputPins())
 		{
 			if (Pin)
@@ -1640,14 +1639,12 @@ uint32 UPCGGraph::CalculateNodeGridSizeRecursive_Unsafe(const UPCGNode* InNode,
 						const uint32 InputGridSize = CalculateNodeGridSizeRecursive_Unsafe(OtherPin->Node, InDefaultGridSize);
 						if (PCGHiGenGrid::IsValidGridSize(InputGridSize))
 						{
-							MinGenerationSize = FMath::Min(MinGenerationSize, InputGridSize);
+							GridSize = FMath::Min(GridSize, InputGridSize);
 						}
 					}
 				}
 			}
 		}
-
-		GridSize = (MinGenerationSize == std::numeric_limits<uint32>::max()) ? InDefaultGridSize : MinGenerationSize;
 	}
 
 	if (GridSize != PCGHiGenGrid::UninitializedGridSize())
