[PCG] Remove check in SurfaceSampler initialize as values can be set by the user and we need to handle this gracefully if it is invalid and not crash.

#rnx

diff --git a/Engine/Plugins/PCG/Source/PCG/Private/Elements/PCGSurfaceSampler.cpp b/Engine/Plugins/PCG/Source/PCG/Private/Elements/PCGSurfaceSampler.cpp
--- a/Engine/Plugins/PCG/Source/PCG/Private/Elements/PCGSurfaceSampler.cpp
+++ b/Engine/Plugins/PCG/Source/PCG/Private/Elements/PCGSurfaceSampler.cpp
@@ -50,7 +50,12 @@ namespace PCGSurfaceSampler
 		InterstitialDistance = PointExtents * 2;
 		InnerCellSize = InterstitialDistance * Looseness;
 		CellSize = InterstitialDistance + InnerCellSize;
-		check(CellSize.X > 0 && CellSize.Y > 0);
+		if (CellSize.X <= 0 || CellSize.Y <= 0)
+		{
+			// PointExtents and Looseness are user overridable, if any of those values are 0 or negative, it's invalid, so we early out.
+			PCGLog::LogWarningOnGraph(LOCTEXT("InvalidParameters", "Skipped - Extents and/or Looseness are negative or zero."), Context);
+			return false;
+		}
 
 		// By using scaled indices in the world, we can easily make this process deterministic
 		CellMinX = FMath::CeilToInt((InputBounds.Min.X) / CellSize.X);
