[PCG] Add warning if we have multi entry attribute set on override + promote BuildTaskInput warning to the node.

* Also combine params on task input was not debugging correctly because it had no entry. Use the same method as AddAttributes (which is what we expect to happen anyway)
* Factorize GetTaskStack to allow to push visual log on node without any context

#rnx

diff --git a/Engine/Plugins/PCG/Source/PCG/Private/Graph/PCGGraphExecutor.cpp b/Engine/Plugins/PCG/Source/PCG/Private/Graph/PCGGraphExecutor.cpp
--- a/Engine/Plugins/PCG/Source/PCG/Private/Graph/PCGGraphExecutor.cpp
+++ b/Engine/Plugins/PCG/Source/PCG/Private/Graph/PCGGraphExecutor.cpp
@@ -68,7 +68,28 @@ namespace PCGGraphExecutor
 		TEXT("Controls whether tasks are culled at execution time, for example in response to an deactivated dynamic branch pin"));
 }
 
+const FPCGStack* FPCGGraphTask::GetStack() const
+{
+	return (StackContext && StackIndex != INDEX_NONE) ? StackContext->GetStack(StackIndex) : nullptr;
+}
+
 #if WITH_EDITOR
+void FPCGGraphTask::LogVisual(ELogVerbosity::Type InVerbosity, const FText& InMessage) const
+{
+	if (!SourceComponent.IsValid())
+	{
+		return;
+	}
+
+	if (UPCGSubsystem* Subsystem = UPCGSubsystem::GetInstance(SourceComponent->GetWorld()))
+	{
+		const FPCGStack* TaskStack = GetStack();
+		FPCGStack StackWithNode = TaskStack ? FPCGStack(*TaskStack) : FPCGStack();
+		StackWithNode.PushFrame(Node);
+		Subsystem->GetNodeVisualLogsMutable().Log(StackWithNode, InVerbosity, InMessage);
+	}
+}
+
 bool FPCGGraphTaskInput::operator==(const FPCGGraphTaskInput& Other) const
 {
 	return (TaskId == Other.TaskId)
@@ -802,7 +823,7 @@ void FPCGGraphExecutor::Execute()
 					{
 						if (Task.StackIndex != INDEX_NONE)
 						{
-							const FPCGStack* Stack = Task.StackContext->GetStack(Task.StackIndex);
+							const FPCGStack* Stack = Task.GetStack();
 							SourceComponent->StoreInspectionData(Stack, Task.Node, TaskInput, CachedOutput);
 						}
 					}
@@ -813,7 +834,7 @@ void FPCGGraphExecutor::Execute()
 						CullInactiveDownstreamNodes(Task.NodeId, CachedOutput.InactiveOutputPinBitmask);
 
 #if WITH_EDITOR
-						SendInactivePinNotification(Task.Node, Task.StackContext->GetStack(Task.StackIndex), CachedOutput.InactiveOutputPinBitmask);
+						SendInactivePinNotification(Task.Node, Task.GetStack(), CachedOutput.InactiveOutputPinBitmask);
 #endif
 					}
 
@@ -836,7 +857,7 @@ void FPCGGraphExecutor::Execute()
 					Task.Context->TaskId = Task.NodeId;
 					Task.Context->CompiledTaskId = Task.CompiledTaskId;
 					Task.Context->DependenciesCrc = DependenciesCrc;
-					Task.Context->Stack = (Task.StackContext && Task.StackIndex != INDEX_NONE) ? Task.StackContext->GetStack(Task.StackIndex) : nullptr;
+					Task.Context->Stack = Task.GetStack();
 				}
 
 				// Validate that we can start this task now
@@ -1243,11 +1264,19 @@ void FPCGGraphExecutor::BuildTaskInput(const FPCGGraphTask& Task, FPCGDataCollec
 
 	auto LogDiscardedData = [&Task](const UPCGPin* InPin)
 	{
-		// Log only - currently context has not yet been allocated when this is called
-		UE_LOG(LogPCG, Warning, TEXT("[%s] %s - BuildTaskInput - too many data items arriving on single data pin '%s', only first data item will be used"),
+#if !(UE_BUILD_SHIPPING || UE_BUILD_TEST) || USE_LOGGING_IN_SHIPPING
+		const FString Message = FString::Printf(
+			TEXT("[%s] %s - BuildTaskInput - too many data items arriving on single data pin '%s', only first data item will be used"),
 			(Task.SourceComponent.Get() && Task.SourceComponent->GetOwner()) ? *Task.SourceComponent->GetOwner()->GetName() : TEXT("MissingComponent"),
 			Task.Node ? *Task.Node->GetNodeTitle(EPCGNodeTitleType::ListView).ToString() : TEXT("MissingNode"),
 			InPin ? *InPin->Properties.Label.ToString() : TEXT("MissingPin"));
+
+#if WITH_EDITOR
+		Task.LogVisual(ELogVerbosity::Warning, FText::FromString(Message));
+#endif // WITH_EDITOR
+
+		UE_LOG(LogPCG, Warning, TEXT("%s"), *Message);
+#endif
 	};
 
 	// Initialize a Crc onto which each input Crc will be combined (using random prime number).
@@ -1345,11 +1374,25 @@ void FPCGGraphExecutor::CombineParams(FPCGTaskId InTaskId, FPCGDataCollection& I
 	TArray<FPCGTaggedData> AllParamsData = InTaskInput.GetParamsByPin(PCGPinConstants::DefaultParamsLabel);
 	if (AllParamsData.Num() > 1)
 	{
-		UPCGParamData* CombinedParamData = NewObject<UPCGParamData>();
+		UPCGParamData* CombinedParamData = nullptr;
+		bool bSuccess = true;
+
 		for (const FPCGTaggedData& TaggedDatum : AllParamsData)
 		{
 			const UPCGParamData* ParamData = CastChecked<UPCGParamData>(TaggedDatum.Data);
-			CombinedParamData->Metadata->AddAttributes(ParamData->Metadata);
+			if (!CombinedParamData)
+			{
+				CombinedParamData = ParamData->DuplicateData();
+			}
+			else
+			{
+				bSuccess &= PCGMetadataHelpers::CopyAllAttributes(ParamData, CombinedParamData, nullptr);
+			}
+		}
+
+		if (!bSuccess)
+		{
+			return;
 		}
 
 		const int32 NewNumberOfInputs = InTaskInput.TaggedData.Num() - AllParamsData.Num() + 1;
@@ -1485,7 +1528,7 @@ void FPCGGraphExecutor::CullInactiveDownstreamNodes(FPCGTaskId InCompletedTaskId
 					GetPinIdsToDeactivate(RemovedTaskId, InactiveOutputPinBitmask, PinIdsToDeactivate);
 
 #if WITH_EDITOR
-					SendInactivePinNotification(RemovedTask.Node, RemovedTask.StackContext->GetStack(RemovedTask.StackIndex), InactiveOutputPinBitmask);
+					SendInactivePinNotification(RemovedTask.Node, RemovedTask.GetStack(), InactiveOutputPinBitmask);
 #endif
 				}
 
diff --git a/Engine/Plugins/PCG/Source/PCG/Private/Graph/PCGGraphExecutor.h b/Engine/Plugins/PCG/Source/PCG/Private/Graph/PCGGraphExecutor.h
--- a/Engine/Plugins/PCG/Source/PCG/Private/Graph/PCGGraphExecutor.h
+++ b/Engine/Plugins/PCG/Source/PCG/Private/Graph/PCGGraphExecutor.h
@@ -67,8 +67,13 @@ struct FPCGGraphTask
 #if WITH_EDITOR
 	/** Approximate equivalence. Does not deeply check node settings, nor does it do a deep comparison of the element. */
 	bool IsApproximatelyEqual(const FPCGGraphTask& Other) const;
+
+	/** Because we might not already have a context, but still want to attach some logs to the node, use this utility function */
+	void LogVisual(ELogVerbosity::Type InVerbosity, const FText& InMessage) const;
 #endif
 
+	const FPCGStack* GetStack() const;
+
 	TArray<FPCGGraphTaskInput> Inputs;
 	const UPCGNode* Node = nullptr;
 	TWeakObjectPtr<UPCGComponent> SourceComponent = nullptr;
diff --git a/Engine/Plugins/PCG/Source/PCG/Private/Metadata/Accessors/PCGAttributeAccessorHelpers.cpp b/Engine/Plugins/PCG/Source/PCG/Private/Metadata/Accessors/PCGAttributeAccessorHelpers.cpp
--- a/Engine/Plugins/PCG/Source/PCG/Private/Metadata/Accessors/PCGAttributeAccessorHelpers.cpp
+++ b/Engine/Plugins/PCG/Source/PCG/Private/Metadata/Accessors/PCGAttributeAccessorHelpers.cpp
@@ -544,15 +544,21 @@ TUniquePtr<const IPCGAttributeAccessor> PCGAttributeAccessorHelpers::CreateConst
 {
 	bool bFromGlobalParamsPin = false;
 	TArray<FPCGTaggedData> InputParamData = InInputData.GetParamsByPin(InParam.Label);
-	const UPCGParamData* ParamData = !InputParamData.IsEmpty() ? CastChecked<UPCGParamData>(InputParamData[0].Data) : nullptr;
-
-	// If it is empty, try with the Overrides pin (Global Params)
-	if (!ParamData)
+	if (InputParamData.IsEmpty())
 	{
+		// If it is empty, try with the Overrides pin (Global Params)
 		bFromGlobalParamsPin = true;
-		ParamData = InInputData.GetFirstParamsOnParamsPin();
+		InputParamData = InInputData.GetParamsByPin(PCGPinConstants::DefaultParamsLabel);
 	}
-	else if (OutResult)
+
+	if (OutResult)
+	{
+		OutResult->bHasMultipleAttributeSetsOnOverridePin = InputParamData.Num() > 1;
+	}
+
+	const UPCGParamData* ParamData = !InputParamData.IsEmpty() ? CastChecked<UPCGParamData>(InputParamData[0].Data) : nullptr;
+
+	if (OutResult && ParamData && !bFromGlobalParamsPin)
 	{
 		OutResult->bPinConnected = true;
 	}
@@ -593,6 +599,11 @@ TUniquePtr<const IPCGAttributeAccessor> PCGAttributeAccessorHelpers::CreateConst
 			}
 		}
 
+		if (Result && OutResult && ParamData && ParamData->Metadata)
+		{
+			OutResult->bHasMultipleDataInAttributeSet = ParamData->Metadata->GetLocalItemCount() > 1;
+		}
+
 		return Result;
 	}
 
diff --git a/Engine/Plugins/PCG/Source/PCG/Private/PCGContext.cpp b/Engine/Plugins/PCG/Source/PCG/Private/PCGContext.cpp
--- a/Engine/Plugins/PCG/Source/PCG/Private/PCGContext.cpp
+++ b/Engine/Plugins/PCG/Source/PCG/Private/PCGContext.cpp
@@ -210,6 +210,22 @@ void FPCGContext::OverrideSettings()
 			PCGE_LOG_C(Warning, GraphAndLog, this, FText::Format(LOCTEXT("OverrideWithAlias", "Attribute '{0}' was not found, but one of its deprecated aliases ('{1}') was. Please update the name to the new value."), FText::FromName(AttributeName), FText::FromName(AccessorResult.AliasUsed)));
 		}
 
+		// Throw warnings if we have multiple data on override (multiple attribute sets or multi entry attribute set)
+		if (AccessorResult.bHasMultipleAttributeSetsOnOverridePin || AccessorResult.bHasMultipleDataInAttributeSet)
+		{
+			const FText OverridePinText = AccessorResult.bPinConnected ? FText::Format(LOCTEXT("OverridePinText", "override pin '{0}'"), FText::FromName(Param.Label)) : LOCTEXT("GlobalOverridePinText", "global override pin");
+
+			if (AccessorResult.bHasMultipleAttributeSetsOnOverridePin)
+			{
+				PCGE_LOG_C(Warning, GraphAndLog, this, FText::Format(LOCTEXT("HasMultipleAttributeSetsOnOverridePin", "Multiple attribute sets were found on the {0}. We will use the first one."), OverridePinText));
+			}
+
+			if (AccessorResult.bHasMultipleDataInAttributeSet)
+			{
+				PCGE_LOG_C(Warning, GraphAndLog, this, FText::Format(LOCTEXT("HasMultipleDataInAttributeSet", "Multi entry attribute set was found on the {0}. We will only use the first entry to override."), OverridePinText));
+			}
+		}
+
 		TUniquePtr<IPCGAttributeAccessor> PropertyAccessor = PCGAttributeAccessorHelpers::CreatePropertyChainAccessor(TArray<const FProperty*>(Param.Properties));
 		check(PropertyAccessor.IsValid());
 
diff --git a/Engine/Plugins/PCG/Source/PCG/Public/Metadata/Accessors/PCGAttributeAccessorHelpers.h b/Engine/Plugins/PCG/Source/PCG/Public/Metadata/Accessors/PCGAttributeAccessorHelpers.h
--- a/Engine/Plugins/PCG/Source/PCG/Public/Metadata/Accessors/PCGAttributeAccessorHelpers.h
+++ b/Engine/Plugins/PCG/Source/PCG/Public/Metadata/Accessors/PCGAttributeAccessorHelpers.h
@@ -46,6 +46,8 @@ namespace PCGAttributeAccessorHelpers
 		FName AliasUsed = NAME_None;
 		bool bUsedAliases = false;
 		bool bPinConnected = false;
+		bool bHasMultipleAttributeSetsOnOverridePin = false;
+		bool bHasMultipleDataInAttributeSet = false;
 	};
 
 	UE_DEPRECATED(5.3, "Use the CreateConstAccessorForOverrideParamWithResult version")
