Fixed static culling when using required pins on subgraphs (either static or dynamic w/ template).
Fixed missing input on static subgraph node inspection
Added missing full graph recompilation when pin status changes (esp. important for static culling).

#rnx

diff --git a/Engine/Plugins/PCG/Source/PCG/Private/Graph/PCGGraphCompiler.cpp b/Engine/Plugins/PCG/Source/PCG/Private/Graph/PCGGraphCompiler.cpp
--- a/Engine/Plugins/PCG/Source/PCG/Private/Graph/PCGGraphCompiler.cpp
+++ b/Engine/Plugins/PCG/Source/PCG/Private/Graph/PCGGraphCompiler.cpp
@@ -4,6 +4,7 @@
 
 #include "PCGEdge.h"
 #include "PCGGraph.h"
+#include "PCGInputOutputSettings.h"
 #include "PCGModule.h"
 #include "PCGPin.h"
 #include "PCGSubgraph.h"
@@ -166,6 +167,8 @@ TArray<FPCGGraphTask> FPCGGraphCompiler::CompileGraph(UPCGGraph* InGraph, FPCGTa
 
 			// Add execution-only dependency on pre-task, without this post task can be scheduled concurrently with pre-task, and concurrently
 			// with something that might become inactive and would then fail to dynamically cull this already-scheduled task.
+			// Additional implementation note: this first depedencency is critical in our ability to do static culling (see CalculateStaticallyActiveRecursive)
+			// and should not be changed here without changing the other.
 			PostTask.Inputs.Emplace(PreId, /*InInboundPin=*/nullptr, /*InOutboundPin=*/nullptr, /*bInProvideData=*/false);
 
 			// Add subgraph output node task as input to the post-task
@@ -493,25 +496,41 @@ bool FPCGGraphCompiler::CalculateStaticallyActiveRecursive(FPCGTaskId InTaskId,
 		return *bEntry;
 	}
 
+	const UPCGNode* Node = InCompiledTasks[InTaskId].Node;
+
 	// Nodes within subgraphs - if the subgraph node is inactive then all tasks within the subgraph are inactive.
 	if (InCompiledTasks[InTaskId].ParentId != InvalidPCGTaskId)
 	{
 		const bool bParentActive = CalculateStaticallyActiveRecursive(InCompiledTasks[InTaskId].ParentId, InCompiledTasks, InOutTaskIdToActiveFlag);
-		if (!bParentActive)
+		// With respect to the input node in a static subgraph, it has to use the same value as its parent since there is no direct edge between the subgraph node and this input node.
+		if (!bParentActive || (Node && Node->GetSettings() && Node->GetSettings()->IsA<UPCGGraphInputOutputSettings>()))
 		{
 			InOutTaskIdToActiveFlag.Add(InTaskId, bParentActive);
-
 			return bParentActive;
 		}
 	}
 
-	const UPCGNode* Node = InCompiledTasks[InTaskId].Node;
 	if (!Node)
 	{
 		InOutTaskIdToActiveFlag.Add(InTaskId, true);
 		return true;
 	}
 
+	// For static subgraphs, the second time this node is seen, it has a trivial task and no proper edge, which will trip the static validation below
+	// This should basically forward the same information as the original node - which is conveniently the first input on that task.
+	if (UPCGSubgraphSettings* SubgraphSettings = Cast<UPCGSubgraphSettings>(Node->GetSettings()))
+	{
+		FPCGTaskId TentativeSubgraphInputId = InCompiledTasks[InTaskId].Inputs.IsEmpty() ? InvalidPCGTaskId : InCompiledTasks[InTaskId].Inputs[0].TaskId;
+
+		if (!SubgraphSettings->IsDynamicGraph() && TentativeSubgraphInputId != InvalidPCGTaskId && InCompiledTasks[TentativeSubgraphInputId].Node == Node)
+		{
+			const bool bSubgraphInputActive = CalculateStaticallyActiveRecursive(TentativeSubgraphInputId, InCompiledTasks, InOutTaskIdToActiveFlag);
+
+			InOutTaskIdToActiveFlag.Add(InTaskId, bSubgraphInputActive);
+			return bSubgraphInputActive;
+		}
+	}
+
 	TArray<FName, TInlineAllocator<8>> PinsRequiringActiveConnection;
 
 	// Three relevant types of input pins - required, non-advanced and advanced. This tracks if we encounter the second category.
@@ -554,10 +573,13 @@ bool FPCGGraphCompiler::CalculateStaticallyActiveRecursive(FPCGTaskId InTaskId,
 		bool bInputActive = true;
 
 		// If we are connected to an upstream node, evaluate if the output pin is active.
-		const UPCGNode* UpstreamNode = InCompiledTasks[Input.TaskId].Node;
-		if (const UPCGSettings* UpstreamSettings = UpstreamNode ? UpstreamNode->GetSettings() : nullptr)
+		if (Input.InPin)
 		{
-			bInputActive &= UpstreamSettings->IsPinStaticallyActive(Input.InPin->Properties.Label);
+			const UPCGNode* UpstreamNode = InCompiledTasks[Input.TaskId].Node;
+			if (const UPCGSettings* UpstreamSettings = UpstreamNode ? UpstreamNode->GetSettings() : nullptr)
+			{
+				bInputActive &= UpstreamSettings->IsPinStaticallyActive(Input.InPin->Properties.Label);
+			}
 		}
 
 		if (bInputActive)
@@ -842,12 +864,15 @@ void FPCGGraphCompiler::CalculateDynamicActivePinDependencies(FPCGTaskId InTaskI
 	}
 	else
 	{
+		// In the case of output nodes, advanced pins shouldn't be ignored as it can and will prevent culling in some instances
+		const bool bTreatAdvancedPinsAsNormal = (Node && Cast<UPCGGraphInputOutputSettings>(Node->GetSettings()) && !Cast<UPCGGraphInputOutputSettings>(Node->GetSettings())->IsInput());
+
 		// If we don't have any dependent pins logged by now then there are no required pins (note that a node that does not have task
 		// inputs for required pins will be statically culled in an earlier compilation step). In which case we'll be active if *any* input
 		// is active. We build a disjunction that expresses this.
 		for (const FPCGGraphTaskInput& Input : InOutCompiledTasks[InTaskId].Inputs)
 		{
-			if (Input.OutPin && Input.OutPin->Properties.IsAdvancedPin())
+			if (Input.OutPin && Input.OutPin->Properties.IsAdvancedPin() && !bTreatAdvancedPinsAsNormal)
 			{
 				// Advanced input pins never participate in keeping node active.
 				continue;
diff --git a/Engine/Plugins/PCG/Source/PCG/Private/PCGComponent.cpp b/Engine/Plugins/PCG/Source/PCG/Private/PCGComponent.cpp
--- a/Engine/Plugins/PCG/Source/PCG/Private/PCGComponent.cpp
+++ b/Engine/Plugins/PCG/Source/PCG/Private/PCGComponent.cpp
@@ -2065,14 +2065,11 @@ void UPCGComponent::StoreInspectionData(const FPCGStack* InStack, const UPCGNode
 				FPCGDataCollection PinDataCollection;
 				InData.GetInputsAndCrcsByPin(Pin->Properties.Label, PinDataCollection.TaggedData, PinDataCollection.DataCrcs);
 
+				// Implementation note: since static subgraphs actually are visited twice and the second time the input doesn't match the input pins, we don't clear the data.
 				if (!PinDataCollection.TaggedData.IsEmpty())
 				{
 					InOutInspectionCache.Add(Stack, PinDataCollection);
 				}
-				else
-				{
-					InOutInspectionCache.Remove(Stack);
-				}
 			}
 		};
 
diff --git a/Engine/Plugins/PCG/Source/PCG/Private/PCGInputOutputSettings.cpp b/Engine/Plugins/PCG/Source/PCG/Private/PCGInputOutputSettings.cpp
--- a/Engine/Plugins/PCG/Source/PCG/Private/PCGInputOutputSettings.cpp
+++ b/Engine/Plugins/PCG/Source/PCG/Private/PCGInputOutputSettings.cpp
@@ -217,6 +217,24 @@ void UPCGGraphInputOutputSettings::PostEditChangeProperty(struct FPropertyChange
 
 	Super::PostEditChangeProperty(PropertyChangedEvent);
 }
+
+EPCGChangeType UPCGGraphInputOutputSettings::GetChangeTypeForProperty(FPropertyChangedEvent& PropertyChangedEvent) const
+{
+	EPCGChangeType ChangeType = Super::GetChangeTypeForProperty(PropertyChangedEvent);
+
+	if (bIsInput && PropertyChangedEvent.Property)
+	{
+		// If pins were removed or the required-ness of a pin has changed, this needs to be a structural change.
+		if (PropertyChangedEvent.Property->GetFName() == GET_MEMBER_NAME_CHECKED(UPCGGraphInputOutputSettings, Pins) ||
+			(PropertyChangedEvent.MemberProperty->GetFName() == GET_MEMBER_NAME_CHECKED(UPCGGraphInputOutputSettings, Pins) && 
+			PropertyChangedEvent.Property->GetFName() == GET_MEMBER_NAME_CHECKED(FPCGPinProperties, PinStatus)))
+		{
+			ChangeType |= EPCGChangeType::Structural;
+		}
+	}
+
+	return ChangeType;
+}
 #endif // WITH_EDITOR
 
 TArray<FPCGPinProperties> UPCGGraphInputOutputSettings::InputPinProperties() const
diff --git a/Engine/Plugins/PCG/Source/PCG/Private/PCGSettings.cpp b/Engine/Plugins/PCG/Source/PCG/Private/PCGSettings.cpp
--- a/Engine/Plugins/PCG/Source/PCG/Private/PCGSettings.cpp
+++ b/Engine/Plugins/PCG/Source/PCG/Private/PCGSettings.cpp
@@ -548,7 +548,7 @@ void UPCGSettings::PostEditChangeProperty(FPropertyChangedEvent& PropertyChanged
 		EPCGChangeType ChangeType = EPCGChangeType::Structural | EPCGChangeType::GenerationGrid;
 		if (PropertyName != NAME_None)
 		{
-			ChangeType = GetChangeTypeForProperty(PropertyChangedEvent.GetPropertyName());
+			ChangeType = GetChangeTypeForProperty(PropertyChangedEvent);
 		}
 
 		OnSettingsChangedDelegate.Broadcast(this, ChangeType);
@@ -557,6 +557,11 @@ void UPCGSettings::PostEditChangeProperty(FPropertyChangedEvent& PropertyChanged
 	CacheCrc();
 }
 
+EPCGChangeType UPCGSettings::GetChangeTypeForProperty(FPropertyChangedEvent& PropertyChangedEvent) const
+{
+	return GetChangeTypeForProperty(PropertyChangedEvent.GetPropertyName());
+}
+
 EPCGChangeType UPCGSettings::GetChangeTypeForProperty(const FName& InPropertyName) const
 {
 PRAGMA_DISABLE_DEPRECATION_WARNINGS
diff --git a/Engine/Plugins/PCG/Source/PCG/Public/PCGInputOutputSettings.h b/Engine/Plugins/PCG/Source/PCG/Public/PCGInputOutputSettings.h
--- a/Engine/Plugins/PCG/Source/PCG/Public/PCGInputOutputSettings.h
+++ b/Engine/Plugins/PCG/Source/PCG/Public/PCGInputOutputSettings.h
@@ -34,24 +34,31 @@ public:
 	// ~End UObject interface
 
 	// ~Begin UPCGSettings interface
+	virtual TArray<FPCGPinProperties> DefaultInputPinProperties() const override;
+	virtual TArray<FPCGPinProperties> DefaultOutputPinProperties() const override;
+
 #if WITH_EDITOR
 	virtual FName GetDefaultNodeName() const override { return bIsInput ? FName(TEXT("InputNode")) : FName(TEXT("OutputNode")); }
 	virtual FText GetDefaultNodeTitle() const override { return bIsInput ? NSLOCTEXT("PCGGraphInputOutputSettings", "InputNodeTitle", "Input Node") : NSLOCTEXT("PCGGraphInputOutputSettings", "OutputNodeTitle", "Output Node"); }
 	virtual EPCGSettingsType GetType() const override { return EPCGSettingsType::InputOutput; }
-	virtual void PostEditChangeProperty(struct FPropertyChangedEvent& PropertyChangedEvent) override;
 	virtual void ApplyDeprecationBeforeUpdatePins(UPCGNode* InOutNode, TArray<TObjectPtr<UPCGPin>>& InputPins, TArray<TObjectPtr<UPCGPin>>& OutputPins) override;
-	// The output node does not emit a task (input node does however). This prevents it displaying as culled when debugging.
-	virtual bool EmitsTaskForExecution() const override { return bIsInput; }
-#endif
+#endif // WITH_EDITOR
 
-	TArray<FPCGPinProperties> InputPinProperties() const override;
-	TArray<FPCGPinProperties> OutputPinProperties() const override;
+protected:
+	virtual FPCGElementPtr CreateElement() const override { return MakeShared<FPCGInputOutputElement>(); }
+	virtual TArray<FPCGPinProperties> InputPinProperties() const override;
+	virtual TArray<FPCGPinProperties> OutputPinProperties() const override;
 
-	virtual TArray<FPCGPinProperties> DefaultInputPinProperties() const override;
-	virtual TArray<FPCGPinProperties> DefaultOutputPinProperties() const override;
+#if WITH_EDITOR
+	virtual void PostEditChangeProperty(FPropertyChangedEvent& PropertyChangedEvent) override;
+	virtual EPCGChangeType GetChangeTypeForProperty(FPropertyChangedEvent& PropertyChangedEvent) const override;
+#endif // WITH_EDITOR
+	// ~End UPCGSettings interface
 
+public:
+	bool IsInput() const { return bIsInput; }
 	void SetInput(bool bInIsInput);
-
+	
 	// Add a new custom pin
 	// Note that you should use the return value of this function, since it can be different from
 	// the one passed as argument. It will change if its label collides with existing pins.
@@ -59,9 +66,6 @@ public:
 
 protected:
 	TArray<FPCGPinProperties> DefaultPinProperties(bool bInvisiblePin) const;
-	virtual FPCGElementPtr CreateElement() const override { return MakeShared<FPCGInputOutputElement>(); }
-	// ~End UPCGSettings interface
-
 	void FixPinProperties();
 
 protected:
diff --git a/Engine/Plugins/PCG/Source/PCG/Public/PCGSettings.h b/Engine/Plugins/PCG/Source/PCG/Public/PCGSettings.h
--- a/Engine/Plugins/PCG/Source/PCG/Public/PCGSettings.h
+++ b/Engine/Plugins/PCG/Source/PCG/Public/PCGSettings.h
@@ -319,9 +319,6 @@ public:
 
 	/** Perform post-operations when an editor node is copied */
 	virtual void PostPaste();
-
-	/** Whether the node will emit one or more tasks for execution. */
-	virtual bool EmitsTaskForExecution() const { return true; }
 #endif // WITH_EDITOR
 
 	/** Derived classes can implement this to expose additional information or context, such as an asset in use by the node. */
@@ -415,6 +412,9 @@ protected:
 	UE_DEPRECATED(5.4, "IsStructuralProperty is deprecated, return EPCGChangeType::Structural from GetChangeTypeForProperty instead.")
 	virtual bool IsStructuralProperty(const FName& InPropertyName) const { return false; }
 
+	/** Gets the change impact for a given property. Can be used to signal structural or cosmetic node changes for example - calls the named version by default. */
+	virtual EPCGChangeType GetChangeTypeForProperty(FPropertyChangedEvent& PropertyChangedEvent) const;
+
 	/** Gets the change impact for a given property. Can be used to signal structural or cosmetic node changes for example. */
 	virtual EPCGChangeType GetChangeTypeForProperty(const FName& InPropertyName) const;
 
diff --git a/Engine/Plugins/PCG/Source/PCGEditor/Private/PCGEditorGraphNodeBase.cpp b/Engine/Plugins/PCG/Source/PCGEditor/Private/PCGEditorGraphNodeBase.cpp
--- a/Engine/Plugins/PCG/Source/PCGEditor/Private/PCGEditorGraphNodeBase.cpp
+++ b/Engine/Plugins/PCG/Source/PCGEditor/Private/PCGEditorGraphNodeBase.cpp
@@ -62,8 +62,8 @@ namespace PCGEditorGraphNodeBase
 		{
 			if (!Settings || !Settings->IsA<UPCGRerouteSettings>())
 			{
-				// Task will be displayed as active if it was executed or if it does not produce tasks for execution.
-				return InComponentBeingDebugged->WasNodeExecuted(PCGNode, *InStackBeingInspected) || (Settings && !Settings->EmitsTaskForExecution());
+				// Task will be displayed as active if it was executed.
+				return InComponentBeingDebugged->WasNodeExecuted(PCGNode, *InStackBeingInspected);
 			}
 			else
 			{
