[PCG] Make sure to clean original component artifacts when enabling IsPartitioned, and make sure we process the correct grids when registering local components.

#rnx

diff --git a/Engine/Plugins/PCG/Source/PCG/Private/PCGActorAndComponentMapping.cpp b/Engine/Plugins/PCG/Source/PCG/Private/PCGActorAndComponentMapping.cpp
--- a/Engine/Plugins/PCG/Source/PCG/Private/PCGActorAndComponentMapping.cpp
+++ b/Engine/Plugins/PCG/Source/PCG/Private/PCGActorAndComponentMapping.cpp
@@ -719,19 +719,11 @@ void FPCGActorAndComponentMapping::UpdateMappingPCGComponentPartitionActor(UPCGC
 
 	if (const APCGWorldActor* WorldActor = PCGSubsystem->GetPCGWorldActor())
 	{
-		// Get the generation grids as a bitflag to compare against.
-		uint32 ValidGrids = static_cast<uint32>(EPCGHiGenGrid::Uninitialized);
-
 		bool bHasUnbounded = false;
 		PCGHiGenGrid::FSizeArray GridSizes;
 		ensure(PCGHelpers::GetGenerationGridSizes(InComponent->GetGraph(), WorldActor, GridSizes, bHasUnbounded));
 
-		for (uint32 GridSize : GridSizes)
-		{
-			ValidGrids |= GridSize;
-		}
-
-		auto UpdateMapping = [this, InComponent, &Bounds, WorldActor, ValidGrids](bool bIsRuntimeGenerated)
+		auto UpdateMapping = [this, InComponent, &Bounds, WorldActor, &GridSizes](bool bIsRuntimeGenerated)
 		{
 			TMap<const UPCGComponent*, TSet<TObjectPtr<APCGPartitionActor>>>& Map = bIsRuntimeGenerated ? ComponentToRuntimeGenPartitionActorsMap : ComponentToPartitionActorsMap;
 			TSet<TObjectPtr<APCGPartitionActor>> RemovedActors;
@@ -748,7 +740,7 @@ void FPCGActorAndComponentMapping::UpdateMappingPCGComponentPartitionActor(UPCGC
 				}
 
 				TSet<TObjectPtr<APCGPartitionActor>> NewMapping;
-				ForAllIntersectingPartitionActors(Bounds, [&NewMapping, InComponent, WorldActor, ValidGrids, bIsRuntimeGenerated](APCGPartitionActor* Actor)
+				ForAllIntersectingPartitionActors(Bounds, [&NewMapping, InComponent, WorldActor, &GridSizes, bIsRuntimeGenerated](APCGPartitionActor* Actor)
 				{
 					if (!Actor)
 					{
@@ -761,7 +753,7 @@ void FPCGActorAndComponentMapping::UpdateMappingPCGComponentPartitionActor(UPCGC
 					// Only add a graph instance to partition actors that are:
 					// * In the same execution domain as the original component.
 					// * On a valid grid for the original component.
-					if (bSameDomain && (ValidGrids & Actor->GetPCGGridSize()))
+					if (bSameDomain && GridSizes.Contains(Actor->GetPCGGridSize()))
 					{
 						Actor->AddGraphInstance(InComponent);
 						NewMapping.Add(Actor);
diff --git a/Engine/Plugins/PCG/Source/PCG/Private/PCGComponent.cpp b/Engine/Plugins/PCG/Source/PCG/Private/PCGComponent.cpp
--- a/Engine/Plugins/PCG/Source/PCG/Private/PCGComponent.cpp
+++ b/Engine/Plugins/PCG/Source/PCG/Private/PCGComponent.cpp
@@ -891,53 +891,49 @@ void UPCGComponent::CleanupLocalImmediate(bool bRemoveComponents, bool bCleanupL
 		ensure(PCGHelpers::GetGenerationGridSizes(GetGraph(), Subsystem->GetPCGWorldActor(), GridSizes, bHasUnbounded));
 	}
 
-	// Cleanup original component if non-partitioned, or if it has nodes that will execute at the Unbounded level.
-	if (!IsPartitioned() || bHasUnbounded)
-	{
-		// Cancels generation of this component if there is an ongoing generation in progress.
-		CancelGeneration();
+	// Cancels generation of this component if there is an ongoing generation in progress.
+	CancelGeneration();
 
-		TSet<TSoftObjectPtr<AActor>> ActorsToDelete;
+	TSet<TSoftObjectPtr<AActor>> ActorsToDelete;
 
-		if (!bRemoveComponents && UPCGManagedResource::DebugForcePurgeAllResourcesOnGenerate())
-		{
-			bRemoveComponents = true;
-		}
+	if (!bRemoveComponents && UPCGManagedResource::DebugForcePurgeAllResourcesOnGenerate())
+	{
+		bRemoveComponents = true;
+	}
 
+	{
+		FScopeLock ResourcesLock(&GeneratedResourcesLock);
+		check(!GeneratedResourcesInaccessible);
+		for (int32 ResourceIndex = GeneratedResources.Num() - 1; ResourceIndex >= 0; --ResourceIndex)
 		{
-			FScopeLock ResourcesLock(&GeneratedResourcesLock);
-			check(!GeneratedResourcesInaccessible);
-			for (int32 ResourceIndex = GeneratedResources.Num() - 1; ResourceIndex >= 0; --ResourceIndex)
-			{
-				// Note: resources can be null here in some loading + bp object cases
-				UPCGManagedResource* Resource = GeneratedResources[ResourceIndex];
+			// Note: resources can be null here in some loading + bp object cases
+			UPCGManagedResource* Resource = GeneratedResources[ResourceIndex];
 
-				PCGGeneratedResourcesLogging::LogCleanupLocalImmediateResource(this, Resource);
+			PCGGeneratedResourcesLogging::LogCleanupLocalImmediateResource(this, Resource);
 
-				if (!Resource || Resource->Release(bRemoveComponents, ActorsToDelete))
+			if (!Resource || Resource->Release(bRemoveComponents, ActorsToDelete))
+			{
+				if (Resource)
 				{
-					if (Resource)
-					{
 #if WITH_EDITOR
-						if (Resource->IsMarkedTransientOnLoad())
-						{
-							LoadedPreviewResources.Add(Resource);
-						}
-						else
-						{
-							Resource->Rename(nullptr, nullptr, REN_ForceNoResetLoaders | REN_DoNotDirty | REN_DontCreateRedirectors | REN_NonTransactional);
-						}
-#endif
+					if (Resource->IsMarkedTransientOnLoad())
+					{
+						LoadedPreviewResources.Add(Resource);
 					}
-
-					GeneratedResources.RemoveAtSwap(ResourceIndex);
+					else
+					{
+						Resource->Rename(nullptr, nullptr, REN_ForceNoResetLoaders | REN_DoNotDirty | REN_DontCreateRedirectors | REN_NonTransactional);
+					}
+#endif
 				}
+
+				GeneratedResources.RemoveAtSwap(ResourceIndex);
 			}
 		}
-
-		UPCGActorHelpers::DeleteActors(GetWorld(), ActorsToDelete.Array());
 	}
 
+	UPCGActorHelpers::DeleteActors(GetWorld(), ActorsToDelete.Array());
+
 	// If bRemoveComponents is true, it means we are in a "real" cleanup, not a pre-cleanup before a generate.
 	// So call PostCleanup in this case.
 	if (bRemoveComponents)
