[PCG] Fix bug where SG invoked from UB grid level cannot be debugged, and fixed presence of unbounded component in cases where unbounded does not execute.

Also fixed case where debug object trees were stale by flushing related stacks after a structural/gen grid change.

#rnx

diff --git a/Engine/Plugins/PCG/Source/PCG/Private/PCGSubsystem.cpp b/Engine/Plugins/PCG/Source/PCG/Private/PCGSubsystem.cpp
--- a/Engine/Plugins/PCG/Source/PCG/Private/PCGSubsystem.cpp
+++ b/Engine/Plugins/PCG/Source/PCG/Private/PCGSubsystem.cpp
@@ -1328,6 +1328,13 @@ void UPCGSubsystem::DeleteSerializedPartitionActors(bool bOnlyDeleteUnused, bool
 
 void UPCGSubsystem::NotifyGraphChanged(UPCGGraph* InGraph, EPCGChangeType ChangeType)
 {
+	if (!!(ChangeType & (EPCGChangeType::Structural | EPCGChangeType::GenerationGrid)))
+	{
+		// If change was deep enough, clear out all executed stacks, and let them refresh upon next generate. Fixes
+		// cases where stale stacks never got flushed.
+		ClearExecutedStacks(InGraph);
+	}
+
 	if (GraphExecutor)
 	{
 		GraphExecutor->NotifyGraphChanged(InGraph, ChangeType);
@@ -1392,22 +1399,32 @@ FPCGGraphCompiler* UPCGSubsystem::GetGraphCompiler()
 
 bool UPCGSubsystem::GetStackContext(const UPCGComponent* InComponent, FPCGStackContext& OutStackContext)
 {
-	if (InComponent && InComponent->GetGraph())
+	UPCGGraph* Graph = InComponent ? InComponent->GetGraph() : nullptr;
+	if (!InComponent || !Graph)
 	{
-		uint32 GenerationGridSize = PCGHiGenGrid::UninitializedGridSize();
-		if (InComponent->GetGraph()->IsHierarchicalGenerationEnabled())
-		{
-			if (InComponent->IsLocalComponent() || InComponent->IsPartitioned())
-			{
-				GenerationGridSize = InComponent->GetGenerationGridSize();
-			}
-		}
+		return false;
+	}
 
-		GetGraphCompiler()->GetCompiledTasks(InComponent->GetGraph(), GenerationGridSize, OutStackContext);
-		return true;
+	// A non-partitioned component generally executes (original component or local component).
+	bool bDoesComponentExecute = !InComponent->IsPartitioned();
+
+	// A partitioned higen original component will execute if the graph has UB grid level.
+	if (!bDoesComponentExecute && Graph->IsHierarchicalGenerationEnabled())
+	{
+		PCGHiGenGrid::FSizeArray GridSizes;
+		Graph->GetGridSizes(GridSizes, /*bOutHasUnbounded=*/bDoesComponentExecute);
 	}
 
-	return false;
+	if (bDoesComponentExecute)
+	{
+		GetGraphCompiler()->GetCompiledTasks(Graph, InComponent->GetGenerationGridSize(), OutStackContext);
+		return true;
+	}
+	else
+	{
+		OutStackContext = FPCGStackContext();
+		return false;
+	}
 }
 
 uint32 UPCGSubsystem::GetGraphCacheEntryCount(IPCGElement* InElement) const
@@ -1466,6 +1483,21 @@ void UPCGSubsystem::ClearExecutedStacks(FPCGStack BeginningWithStack)
 	}
 }
 
+void UPCGSubsystem::ClearExecutedStacks(const UPCGGraph* InContainingGraph)
+{
+	FWriteScopeLock Lock(ExecutedStacksLock);
+
+	for (int StackIndex = ExecutedStacks.Num() - 1; StackIndex >= 0; --StackIndex)
+	{
+		// Clear any dynamic stack that starts with the given base stack, we don't know what stacks will
+		// execute so clean out everything.
+		if (ExecutedStacks[StackIndex].HasObject(InContainingGraph))
+		{
+			ExecutedStacks.RemoveAtSwap(StackIndex);
+		}
+	}
+}
+
 #endif // WITH_EDITOR
 
 void UPCGSubsystem::FlushCache()
diff --git a/Engine/Plugins/PCG/Source/PCG/Public/PCGSubsystem.h b/Engine/Plugins/PCG/Source/PCG/Public/PCGSubsystem.h
--- a/Engine/Plugins/PCG/Source/PCG/Public/PCGSubsystem.h
+++ b/Engine/Plugins/PCG/Source/PCG/Public/PCGSubsystem.h
@@ -291,6 +291,7 @@ public:
 	/** Get a list of stacks that were executed during the last execution. */
 	TArray<FPCGStack> GetExecutedStacks(const UPCGComponent* InComponent, const UPCGGraph* InSubgraph);
 	void ClearExecutedStacks(FPCGStack BeginningWithStack);
+	void ClearExecutedStacks(const UPCGGraph* InContainingGraph);
 
 	FPCGOnComponentGenerationCompleteOrCancelled OnComponentGenerationCompleteOrCancelled;
 
diff --git a/Engine/Plugins/PCG/Source/PCGEditor/Private/SPCGEditorGraphDebugObjectTree.cpp b/Engine/Plugins/PCG/Source/PCGEditor/Private/SPCGEditorGraphDebugObjectTree.cpp
--- a/Engine/Plugins/PCG/Source/PCGEditor/Private/SPCGEditorGraphDebugObjectTree.cpp
+++ b/Engine/Plugins/PCG/Source/PCGEditor/Private/SPCGEditorGraphDebugObjectTree.cpp
@@ -553,28 +553,6 @@ void SPCGEditorGraphDebugObjectTree::AddStacksToTree(const TArray<FPCGStack>& St
 			continue;
 		}
 
-		if (PCGComponent->IsPartitioned())
-		{
-			if (!GraphBeingEdited->IsHierarchicalGenerationEnabled())
-			{
-				// Don't show the original component if graph does not have higen enabled because this component
-				// will never do any interesting execution in this case.
-				continue;
-			}
-			else
-			{
-				// If higen is enabled, check if the UB grid level will execute. If not then, then there is no execution to inspect.
-				bool bHasUnbounded = false;
-				PCGHiGenGrid::FSizeArray GridSizes;
-				GraphBeingEdited->GetGridSizes(GridSizes, bHasUnbounded);
-
-				if (!bHasUnbounded)
-				{
-					continue;
-				}
-			}
-		}
-
 		// Prevent duplicate entries from the editor world while in PIE.
 		if (PCGHelpers::IsRuntimeOrPIE())
 		{
