[PCG] Don't use SetGraph on GraphInstance in Subgraph settings

* SetGraph is doing a lot of work and callbacks, copying user params structs
* Update GetSubgraphInterface to return the subgraph override if valid
* Remove condition in GetUserParameter that was always taking the Graph Instance of the SourceComponent in some cases, which is invalid in a subgraph context.

#rnx

diff --git a/Engine/Plugins/PCG/Source/PCG/Private/Elements/PCGLoopElement.cpp b/Engine/Plugins/PCG/Source/PCG/Private/Elements/PCGLoopElement.cpp
--- a/Engine/Plugins/PCG/Source/PCG/Private/Elements/PCGLoopElement.cpp
+++ b/Engine/Plugins/PCG/Source/PCG/Private/Elements/PCGLoopElement.cpp
@@ -176,11 +176,8 @@ bool FPCGLoopElement::ExecuteInternal(FPCGContext* InContext) const
 	// If we haven't scheduled the subgraph yet (1st call) - this is part of the FPCGSubgraphContext, see FPCGSubgraphElement for more info
 	if (!Context->bScheduledSubgraph)
 	{
-		if (Settings->SubgraphInstance && Settings->SubgraphOverride)
+		if (Settings->SubgraphOverride)
 		{
-			// If OriginalSettings is null, then we ARE the original settings, and writing over the existing graph is incorrect (and potentially a race condition)
-			check(Settings->OriginalSettings);
-			Settings->SubgraphInstance->SetGraph(Settings->SubgraphOverride);
 #if WITH_EDITOR
 			FPCGDynamicTrackingHelper::AddSingleDynamicTrackingKey(Context, FPCGSelectionKey::CreateFromPath(Settings->SubgraphOverride), /*bIsCulled=*/false);
 #endif // WITH_EDITOR
diff --git a/Engine/Plugins/PCG/Source/PCG/Private/Elements/PCGUserParameterGet.cpp b/Engine/Plugins/PCG/Source/PCG/Private/Elements/PCGUserParameterGet.cpp
--- a/Engine/Plugins/PCG/Source/PCG/Private/Elements/PCGUserParameterGet.cpp
+++ b/Engine/Plugins/PCG/Source/PCG/Private/Elements/PCGUserParameterGet.cpp
@@ -48,7 +48,7 @@ namespace PCGUserParameterGetSettings
 		}
 
 		// Then if we don't have any input, we will use the graph instance from the component.
-		if (!GraphInterface || GraphInterface->IsA<UPCGGraph>())
+		if (!GraphInterface)
 		{
 			const UPCGComponent* SourceComponent = InContext.SourceComponent.Get();
 			GraphInterface = SourceComponent ? SourceComponent->GetGraphInstance() : nullptr;
diff --git a/Engine/Plugins/PCG/Source/PCG/Private/PCGSubgraph.cpp b/Engine/Plugins/PCG/Source/PCG/Private/PCGSubgraph.cpp
--- a/Engine/Plugins/PCG/Source/PCG/Private/PCGSubgraph.cpp
+++ b/Engine/Plugins/PCG/Source/PCG/Private/PCGSubgraph.cpp
@@ -378,6 +378,13 @@ FPCGElementPtr UPCGSubgraphSettings::CreateElement() const
 	return MakeShared<FPCGSubgraphElement>();
 }
 
+UPCGGraphInterface* UPCGSubgraphSettings::GetSubgraphInterface() const
+{
+	// The only place when SubgraphOverride is not null, is when we execute a dynamic subgraph.
+	// Everywhere else (for UI, parameter overrides, normal subgraph flow, etc...) we will use the SubgraphInstance.
+	return SubgraphOverride ? SubgraphOverride.Get() : SubgraphInstance.Get(); 
+}
+
 bool UPCGSubgraphSettings::IsDynamicGraph() const
 {
 	UPCGNode* Node = Cast<UPCGNode>(GetOuter());
@@ -591,18 +598,13 @@ bool FPCGSubgraphElement::ExecuteInternal(FPCGContext* InContext) const
 
 	const bool bIsDynamic = Settings->IsDynamicGraph();
 
-	if (bIsDynamic && !Context->bScheduledSubgraph)
+	if (bIsDynamic && !Context->bScheduledSubgraph && Settings->SubgraphOverride)
 	{
-		if (Settings->SubgraphInstance && Settings->OriginalSettings)
-		{
-			// If OriginalSettings is null, then we ARE the original settings, and writing over the existing graph is incorrect (and potentially a race condition)
-			Settings->SubgraphInstance->SetGraph(Settings->SubgraphOverride);
 #if WITH_EDITOR
-			FPCGDynamicTrackingHelper::AddSingleDynamicTrackingKey(Context, FPCGSelectionKey::CreateFromPath(Settings->SubgraphOverride), /*bIsCulled=*/false);
+		FPCGDynamicTrackingHelper::AddSingleDynamicTrackingKey(Context, FPCGSelectionKey::CreateFromPath(Settings->SubgraphOverride), /*bIsCulled=*/false);
 #endif // WITH_EDITOR
 
-			Context->UpdateOverridesWithOverriddenGraph();
-		}
+		Context->UpdateOverridesWithOverriddenGraph();
 	}
 
 	UPCGGraph* Subgraph = Settings->GetSubgraph();
diff --git a/Engine/Plugins/PCG/Source/PCG/Public/PCGSubgraph.h b/Engine/Plugins/PCG/Source/PCG/Public/PCGSubgraph.h
--- a/Engine/Plugins/PCG/Source/PCG/Public/PCGSubgraph.h
+++ b/Engine/Plugins/PCG/Source/PCG/Public/PCGSubgraph.h
@@ -119,7 +119,7 @@ protected:
 
 	//~Begin UPCGBaseSubgraphSettings interface
 public:
-	virtual UPCGGraphInterface* GetSubgraphInterface() const override { return SubgraphInstance.Get(); }
+	virtual UPCGGraphInterface* GetSubgraphInterface() const override;
 	virtual bool IsDynamicGraph() const override;
 protected:
 	virtual void SetSubgraphInternal(UPCGGraphInterface* InGraph) override;
