[PCG] Work directly on BitArray in AttributePartition

#rnx

diff --git a/Engine/Plugins/PCG/Source/PCG/Private/Metadata/PCGMetadataPartitionCommon.cpp b/Engine/Plugins/PCG/Source/PCG/Private/Metadata/PCGMetadataPartitionCommon.cpp
--- a/Engine/Plugins/PCG/Source/PCG/Private/Metadata/PCGMetadataPartitionCommon.cpp
+++ b/Engine/Plugins/PCG/Source/PCG/Private/Metadata/PCGMetadataPartitionCommon.cpp
@@ -19,11 +19,15 @@
 
 namespace PCGMetadataPartitionCommon
 {
+	template <typename T> 
+	constexpr bool IsBitArray = std::is_same_v<T, TBitArray<>>;
+
 	/**
 	* Partition a given attribute, by first partitioning all value keys that point to the same value
 	* and then for each unique value key, list of index in the keys that match for this value.
 	*/
-	TArray<TArray<int32>> AttributePartition(const FPCGMetadataAttributeBase* InAttribute, const IPCGAttributeAccessorKeys& InKeys, FPCGContext* InOptionalContext)
+	template <typename PartitionType>
+	TArray<PartitionType> AttributePartition(const FPCGMetadataAttributeBase* InAttribute, const IPCGAttributeAccessorKeys& InKeys, FPCGContext* InOptionalContext)
 	{
 		TRACE_CPUPROFILER_EVENT_SCOPE(PCGMetadataPartitionCommon::AttributePartition);
 		check(InAttribute);
@@ -83,8 +87,15 @@ namespace PCGMetadataPartitionCommon
 			NumUniqueValueKeys = MetadataValueKeyCount;
 		}
 
-		TArray<TArray<int32>> PartitionedData;
+		TArray<PartitionType> PartitionedData;
 		PartitionedData.SetNum(1 + NumUniqueValueKeys);
+		if constexpr (IsBitArray<PartitionType>)
+		{
+			for (PartitionType& Partition : PartitionedData)
+			{
+				Partition.Init(false, NumberOfEntries);
+			}
+		}
 
 		constexpr int32 ChunkSize = 256;
 		TArray<const PCGMetadataEntryKey*, TInlineAllocator<ChunkSize>> TempEntries;
@@ -109,27 +120,56 @@ namespace PCGMetadataPartitionCommon
 				}
 
 				const int32 PartitionDataIndex = 1 + ValueKey;
-				PartitionedData[PartitionDataIndex].Add(StartIndex + j);
+				if constexpr (IsBitArray<PartitionType>)
+				{
+					PartitionedData[PartitionDataIndex][StartIndex + j] = true;
+				}
+				else
+				{
+					PartitionedData[PartitionDataIndex].Add(StartIndex + j);
+				}
 			}
 		}
 
 		// Since we partition on the value array, it is not guaranteed that the values appears in the same order than the entries.
 		// So sort the final array using the first index as a sort criteria. Empty partitions will be at the beginning too.
-		PartitionedData.Sort([](const TArray<int32>& LHS, const TArray<int32>& RHS) -> bool
-		{ 
-			if (LHS.IsEmpty())
-			{
-				return true;
-			}
-			else if (RHS.IsEmpty())
+		if constexpr (IsBitArray<PartitionType>)
+		{
+			Algo::Sort(PartitionedData, [](const TBitArray<>& LHS, const TBitArray<>& RHS) -> bool
 			{
-				return false;
-			}
-			else
+				const int32 FirstBitSetLHS = LHS.Find(true);
+				if (FirstBitSetLHS == INDEX_NONE)
+				{
+					return true;
+				}
+
+				const int32 FirstBitSetRHS = RHS.Find(true);
+				if (FirstBitSetRHS == INDEX_NONE)
+				{
+					return false;
+				}
+
+				return  FirstBitSetLHS < FirstBitSetRHS;
+			});
+		}
+		else
+		{
+			PartitionedData.Sort([](const PartitionType& LHS, const PartitionType& RHS) -> bool
 			{
-				return LHS[0] < RHS[0];
-			}
-		});
+				if (LHS.IsEmpty())
+				{
+					return true;
+				}
+				else if (RHS.IsEmpty())
+				{
+					return false;
+				}
+				else
+				{
+					return LHS[0] < RHS[0];
+				}
+			});
+		}
 
 		return PartitionedData;
 	}
@@ -138,14 +178,14 @@ namespace PCGMetadataPartitionCommon
 	* Partition a given accessor that iterate on all values, find the identical ones,
 	* and then for each unique value, list of index in the keys that match for this value.
 	*/
-	template <typename T>
-	TArray<TArray<int32>> ValuePartition(const IPCGAttributeAccessor& InAccessor, const IPCGAttributeAccessorKeys& InKeys, FPCGContext* InOptionalContext)
+	template <typename PartitionType, typename T>
+	TArray<PartitionType> ValuePartition(const IPCGAttributeAccessor& InAccessor, const IPCGAttributeAccessorKeys& InKeys, FPCGContext* InOptionalContext)
 	{
 		TRACE_CPUPROFILER_EVENT_SCOPE(PCGMetadataPartitionCommon::ValuePartition);
 		TArray<T> UniqueValues;
-		TArray<TArray<int32>> PartitionedData;
+		TArray<PartitionType> PartitionedData;
 
-		PCGMetadataElementCommon::ApplyOnAccessor<T>(InKeys, InAccessor, [&PartitionedData, &UniqueValues](const T& InValue, int32 InIndex)
+		PCGMetadataElementCommon::ApplyOnAccessor<T>(InKeys, InAccessor, [&PartitionedData, &UniqueValues, NumberOfEntries = InKeys.GetNum()](const T& InValue, int32 InIndex)
 		{
 			// TODO: Might want to upgrade to something better since it can be quadratic and grow quickly.
 			int32 UniqueValueIndex = UniqueValues.IndexOfByPredicate([&InValue](const T& OtherValue)
@@ -157,10 +197,21 @@ namespace PCGMetadataPartitionCommon
 			if (UniqueValueIndex == INDEX_NONE)
 			{
 				UniqueValueIndex = UniqueValues.Add(InValue);
-				PartitionedData.Emplace();
+				PartitionType& Partition = PartitionedData.Emplace_GetRef();
+				if constexpr (IsBitArray<PartitionType>)
+				{
+					Partition.Init(false, NumberOfEntries);
+				}
 			}
 
-			PartitionedData[UniqueValueIndex].Add(InIndex);
+			if constexpr (IsBitArray<PartitionType>)
+			{
+				PartitionedData[UniqueValueIndex][InIndex] = true;
+			}
+			else
+			{
+				PartitionedData[UniqueValueIndex].Add(InIndex);
+			}
 		});
 
 		return PartitionedData;
@@ -169,7 +220,8 @@ namespace PCGMetadataPartitionCommon
 	/**
 	* Dispatch the partition according to the data and selector.
 	*/
-	TArray<TArray<int32>> AttributeGenericPartition(const UPCGData* InData, const FPCGAttributePropertySelector& InSelector, FPCGContext* InOptionalContext, bool bSilenceMissingAttributeErrors)
+	template <typename PartitionType>
+	TArray<PartitionType> AttributeGenericPartition(const UPCGData* InData, const FPCGAttributePropertySelector& InSelector, FPCGContext* InOptionalContext, bool bSilenceMissingAttributeErrors)
 	{
 		TRACE_CPUPROFILER_EVENT_SCOPE(PCGMetadataPartitionCommon::AttributeGenericPartition::SingleSelector);
 		if (!InData)
@@ -204,7 +256,7 @@ namespace PCGMetadataPartitionCommon
 				return {};
 			}
 
-			return AttributePartition(Attribute, *Keys, InOptionalContext);
+			return AttributePartition<PartitionType>(Attribute, *Keys, InOptionalContext);
 		}
 		else
 		{
@@ -219,7 +271,7 @@ namespace PCGMetadataPartitionCommon
 				return {};
 			}
 
-			auto Operation = [&Accessor, &Keys, &InSelector, InOptionalContext](auto Dummy) -> TArray<TArray<int32>>
+			auto Operation = [&Accessor, &Keys, &InSelector, InOptionalContext](auto Dummy) -> TArray<PartitionType>
 			{
 				// Rotators don't have a hash, convert them to Quat
 				using AttributeType = std::conditional_t<std::is_same_v<decltype(Dummy), FRotator>, FQuat, decltype(Dummy)>;
@@ -232,7 +284,7 @@ namespace PCGMetadataPartitionCommon
 				}
 				else
 				{
-					return ValuePartition<AttributeType>(*Accessor, *Keys, InOptionalContext);
+					return ValuePartition<PartitionType, AttributeType>(*Accessor, *Keys, InOptionalContext);
 				}
 			};
 
@@ -260,7 +312,7 @@ namespace PCGMetadataPartitionCommon
 		// Small optimization to partition on a single attribute
 		if (InSelectorArrayView.Num() == 1)
 		{
-			return AttributeGenericPartition(InData, InSelectorArrayView[0], InOptionalContext, bSilenceMissingAttributeErrors);
+			return AttributeGenericPartition<TArray<int32>>(InData, InSelectorArrayView[0], InOptionalContext, bSilenceMissingAttributeErrors);
 		}
 
 		if (!InData || InSelectorArrayView.IsEmpty() || !InData->ConstMetadata())
@@ -282,42 +334,19 @@ namespace PCGMetadataPartitionCommon
 		using IndexPartition = TArray<TArray<int32>>;
 		using BitPartition = TArray<TBitArray<>>;
 
-		TArray<IndexPartition> IndexPartitions;
-		IndexPartitions.SetNum(NumAttributes);
 		TArray<BitPartition> BitPartitions;
 		BitPartitions.SetNum(NumAttributes);
 
 		{
-			TRACE_CPUPROFILER_EVENT_SCOPE(PCGMetadataPartitionCommon::AttributeGenericPartition::MultiSelector::ConversionToBitArray);
+			TRACE_CPUPROFILER_EVENT_SCOPE(PCGMetadataPartitionCommon::AttributeGenericPartition::MultiSelector::PartitionOnBitArray);
 			/* TODO: Can be executed in parallel, threadsafe. There is a follow up task to evaluate between
-			 * option A.) partitioning on all attributes, and then merging and B.) Partitioning on each attribute
-			 * in succession, further partitioning the grouping results of the previous iteration.
-			 */
-			// Calculate each partition and convert it to a bitfield for simple/efficient intersection processing
+			* option A.) partitioning on all attributes, and then merging and B.) Partitioning on each attribute
+			* in succession, further partitioning the grouping results of the previous iteration.
+			*/
+			// Calculate each partition into a bitfield for simple/efficient intersection processing
 			for (int32 I = 0; I < InSelectorArrayView.Num(); ++I)
 			{
-				IndexPartition& CurrentIndexPartition = IndexPartitions[I];
-				BitPartition& CurrentBitPartition = BitPartitions[I];
-
-				// TODO: Ideally, refactor AttributeGenericPartition to return directly into BitArray format instead to avoid conversion
-				// Partition once for each attribute. It is okay if this to be empty, as it will be skipped later during the iterative partition
-				CurrentIndexPartition = AttributeGenericPartition(InData, InSelectorArrayView[I], InOptionalContext, bSilenceMissingAttributeErrors);
-				// The bit partitions will match the index partitions
-				CurrentBitPartition.SetNum(CurrentIndexPartition.Num());
-
-				auto ConvertIndexGroupingToBitGrouping = [NumElements](const IndexPartition& InIndexPartition, BitPartition& OutBitPartition)
-				{
-					for (int32 GroupIndex = 0; GroupIndex < InIndexPartition.Num(); ++GroupIndex)
-					{
-						OutBitPartition[GroupIndex].SetNum(NumElements, false);
-						for (const int32 Index : InIndexPartition[GroupIndex])
-						{
-							OutBitPartition[GroupIndex].Insert(true, Index);
-						}
-					}
-				};
-
-				ConvertIndexGroupingToBitGrouping(CurrentIndexPartition, CurrentBitPartition);
+				BitPartitions[I] = AttributeGenericPartition<TBitArray<>>(InData, InSelectorArrayView[I], InOptionalContext, bSilenceMissingAttributeErrors);
 			}
 		}
 
