[PCG] Fix oversight on Surface Sampler memory safeguard

#rnx

diff --git a/Engine/Plugins/PCG/Source/PCG/Private/Elements/PCGSurfaceSampler.cpp b/Engine/Plugins/PCG/Source/PCG/Private/Elements/PCGSurfaceSampler.cpp
--- a/Engine/Plugins/PCG/Source/PCG/Private/Elements/PCGSurfaceSampler.cpp
+++ b/Engine/Plugins/PCG/Source/PCG/Private/Elements/PCGSurfaceSampler.cpp
@@ -102,9 +102,9 @@ namespace PCGSurfaceSampler
 			return false;
 		}
 
-		if (PCGFeatureSwitches::CVarCheckSamplerMemory.GetValueOnAnyThread() && (PCGFeatureSwitches::CVarSamplerMemoryThreshold.GetValueOnAnyThread() * FPlatformMemory::GetStats().AvailablePhysical) < sizeof(FPCGPoint) * TargetPointCount)
+		if (PCGFeatureSwitches::CVarCheckSamplerMemory.GetValueOnAnyThread() && (PCGFeatureSwitches::CVarSamplerMemoryThreshold.GetValueOnAnyThread() * FPlatformMemory::GetStats().AvailablePhysical) < sizeof(FPCGPoint) * FMath::Min(TargetPointCount, CellCount))
 		{
-			PCGLog::LogErrorOnGraph(FText::Format((LOCTEXT("TooManyPoints", "Skipped - tried to generate too many points ({0}).\nAdjust 'pcg.SamplerMemoryThreshold' if needed.")), TargetPointCount), Context);
+			PCGLog::LogErrorOnGraph(FText::Format((LOCTEXT("TooManyPoints", "Skipped - tried to generate too many points ({0}).\nAdjust 'pcg.SamplerMemoryThreshold' if needed.")), FMath::Min(TargetPointCount, CellCount)), Context);
 			return false;
 		}
 
