[PCG] Update default values for data asset expose to library

* Not exposed by default (like the graphs or BP)
* Default Category: PCGDataAssets

#rnx

diff --git a/Engine/Plugins/PCG/Source/PCG/Public/PCGDataAsset.h b/Engine/Plugins/PCG/Source/PCG/Public/PCGDataAsset.h
--- a/Engine/Plugins/PCG/Source/PCG/Public/PCGDataAsset.h
+++ b/Engine/Plugins/PCG/Source/PCG/Public/PCGDataAsset.h
@@ -50,11 +50,11 @@ public:
 
 	/** Controls whether the asset will be visible from the palette & contextual menus. */
 	UPROPERTY(BlueprintReadWrite, EditAnywhere, Category = AssetInfo, AssetRegistrySearchable)
-	bool bExposeToLibrary = true;
+	bool bExposeToLibrary = false;
 
 	/** Controls in what category the asset will appear in the palette & contextual menus. */
 	UPROPERTY(BlueprintReadWrite, EditAnywhere, Category = AssetInfo, AssetRegistrySearchable)
-	FText Category;
+	FText Category = NSLOCTEXT("PCGDataAsset", "DefaultCategory", "PCGDataAssets");
 
 	/** Additional description (tooltip) shown on the asset/loader node. */
 	UPROPERTY(BlueprintReadWrite, EditAnywhere, Category = AssetInfo, AssetRegistrySearchable)
