[PCG] Structs and array were not triggering a refresh when edited directly on the Get User Parameters node.

#rnx

diff --git a/Engine/Plugins/PCG/Source/PCGEditor/Private/Details/PCGEditableUserParameterDetails.cpp b/Engine/Plugins/PCG/Source/PCGEditor/Private/Details/PCGEditableUserParameterDetails.cpp
--- a/Engine/Plugins/PCG/Source/PCGEditor/Private/Details/PCGEditableUserParameterDetails.cpp
+++ b/Engine/Plugins/PCG/Source/PCGEditor/Private/Details/PCGEditableUserParameterDetails.cpp
@@ -81,45 +81,50 @@ void FPCGEditableUserParameterDetails::CustomizeDetails(IDetailLayoutBuilder& De
 					IDetailPropertyRow* DetailPropertyRow = CategoryBuilder.AddExternalStructureProperty(MakeShared<FInstancePropertyBagStructureDataProvider>(*UserParameters), CachedPropertyDesc.Name);
 					TSharedPtr<IPropertyHandle> RowPropertyHandle = DetailPropertyRow->GetPropertyHandle();
 
-					if (RowPropertyHandle.IsValid())
+					FSimpleDelegate OnPreChange = FSimpleDelegate::CreateLambda([this]()
 					{
-						RowPropertyHandle->SetOnPropertyValuePreChange(FSimpleDelegate::CreateLambda([this]()
+						check(GEditor);
+						if (GEditor->CanTransact())
 						{
-							check(GEditor);
-							if (GEditor->CanTransact())
-							{
-								GEditor->BeginTransaction(LOCTEXT("EditGraphParameter", "Edit Graph Parameter"));
-							}
+							GEditor->BeginTransaction(LOCTEXT("EditGraphParameter", "Edit Graph Parameter"));
+						}
 
-							// TODO: Structs and Arrays don't transact
-							if (CachedGraphInterface.IsValid())
-							{
-								CachedGraphInterface->Modify();
-							}
-						}));
+						if (CachedGraphInterface.IsValid())
+						{
+							CachedGraphInterface->Modify();
+						}
+					});
 
-						RowPropertyHandle->SetOnPropertyValueChanged(FSimpleDelegate::CreateLambda([this]()
+					FSimpleDelegate OnPostChange = FSimpleDelegate::CreateLambda([this]()
+					{
+						check(GEditor);
+						if (CachedGraphInterface.IsValid())
 						{
-							check(GEditor);
-							if (CachedGraphInterface.IsValid())
+							const FInstancedPropertyBag* CurrentUserParameters = CachedGraphInterface->GetMutableUserParametersStruct_Unsafe();
+							if (CurrentUserParameters && CurrentUserParameters->FindPropertyDescByName(CachedPropertyDesc.Name))
 							{
-								const FInstancedPropertyBag* CurrentUserParameters = CachedGraphInterface->GetMutableUserParametersStruct_Unsafe();
-								if (CurrentUserParameters && CurrentUserParameters->FindPropertyDescByName(CachedPropertyDesc.Name))
-								{
-									CachedGraphInterface->OnGraphParametersChanged(EPCGGraphParameterEvent::ValueModifiedLocally, CachedPropertyDesc.Name);
-								}
-								else if (GEditor->IsTransactionActive())
-								{
-									GEditor->CancelTransaction(0);
-									return;
-								}
+								CachedGraphInterface->OnGraphParametersChanged(EPCGGraphParameterEvent::ValueModifiedLocally, CachedPropertyDesc.Name);
 							}
-
-							if (GEditor->IsTransactionActive())
+							else if (GEditor->IsTransactionActive())
 							{
-								GEditor->EndTransaction();
+								GEditor->CancelTransaction(0);
+								return;
 							}
-						}));
+						}
+
+						if (GEditor->IsTransactionActive())
+						{
+							GEditor->EndTransaction();
+						}
+					});
+
+					if (RowPropertyHandle.IsValid())
+					{
+						RowPropertyHandle->SetOnPropertyValuePreChange(OnPreChange);
+						RowPropertyHandle->SetOnPropertyValueChanged(OnPostChange);
+						// We also need to react to child property changes, if the parameter is a struct or an array.
+						RowPropertyHandle->SetOnChildPropertyValuePreChange(OnPreChange);
+						RowPropertyHandle->SetOnChildPropertyValueChanged(OnPostChange);
 					}
 				}
 			}
