[PCG] Fix warning from GetActorProperty/GetPropertyFromObjectPath following parameter rename.

#rnx

diff --git a/Engine/Plugins/PCG/Source/PCG/Public/Elements/PCGGetActorProperty.h b/Engine/Plugins/PCG/Source/PCG/Public/Elements/PCGGetActorProperty.h
--- a/Engine/Plugins/PCG/Source/PCG/Public/Elements/PCGGetActorProperty.h
+++ b/Engine/Plugins/PCG/Source/PCG/Public/Elements/PCGGetActorProperty.h
@@ -67,7 +67,7 @@ public:
 	bool bForceObjectAndStructExtraction = false;
 
 	/** By default, attribute name will be None, but it can be overridden by this name. Use @SourceName to use the property name (only works when not extracting). */
-	UPROPERTY(BlueprintReadWrite, EditAnywhere, Category = Settings, meta = (EditCondition = "!bExtractObjectAndStruct", EditConditionHides))
+	UPROPERTY(BlueprintReadWrite, EditAnywhere, Category = Settings, meta = (EditCondition = "!bForceObjectAndStructExtraction", EditConditionHides))
 	FName OutputAttributeName = NAME_None;
 
 	/** If this is true, we will never put this element in cache, and will always try to re-query the actors and read the latest properties from them. */
diff --git a/Engine/Plugins/PCG/Source/PCG/Public/Elements/PCGGetPropertyFromObjectPath.h b/Engine/Plugins/PCG/Source/PCG/Public/Elements/PCGGetPropertyFromObjectPath.h
--- a/Engine/Plugins/PCG/Source/PCG/Public/Elements/PCGGetPropertyFromObjectPath.h
+++ b/Engine/Plugins/PCG/Source/PCG/Public/Elements/PCGGetPropertyFromObjectPath.h
@@ -57,7 +57,7 @@ public:
 	bool bForceObjectAndStructExtraction = false;
 
 	/** By default, attribute name will be None, but it can be overridden by this name. Use @SourceName to use the property name (only works when not extracting). */
-	UPROPERTY(BlueprintReadWrite, EditAnywhere, Category = Settings, meta = (EditCondition = "!bExtractObjectAndStruct", EditConditionHides))
+	UPROPERTY(BlueprintReadWrite, EditAnywhere, Category = Settings, meta = (EditCondition = "!bForceObjectAndStructExtraction", EditConditionHides))
 	FName OutputAttributeName = NAME_None;
 
 	/** By default, object loading is asynchronous, can force it synchronous if needed. */
