[PCG] Transform in the attribute list view now show the right rotator and not the quaternions

#rnx

diff --git a/Engine/Plugins/PCG/Source/PCGEditor/Private/SPCGEditorGraphAttributeListView.cpp b/Engine/Plugins/PCG/Source/PCGEditor/Private/SPCGEditorGraphAttributeListView.cpp
--- a/Engine/Plugins/PCG/Source/PCGEditor/Private/SPCGEditorGraphAttributeListView.cpp
+++ b/Engine/Plugins/PCG/Source/PCGEditor/Private/SPCGEditorGraphAttributeListView.cpp
@@ -1047,9 +1047,9 @@ void SPCGEditorGraphAttributeListView::GenerateColumnsFromMetadata(const UPCGDat
 				AddMetadataColumn(InPCGData, ColumnName, AttributeType, TEXT(".Position.X"));
 				AddMetadataColumn(InPCGData, ColumnName, AttributeType, TEXT(".Position.Y"));
 				AddMetadataColumn(InPCGData, ColumnName, AttributeType, TEXT(".Position.Z"));
-				AddMetadataColumn(InPCGData, ColumnName, AttributeType, TEXT(".Rotation.X"));
-				AddMetadataColumn(InPCGData, ColumnName, AttributeType, TEXT(".Rotation.Y"));
-				AddMetadataColumn(InPCGData, ColumnName, AttributeType, TEXT(".Rotation.Z"));
+				AddMetadataColumn(InPCGData, ColumnName, AttributeType, TEXT(".Rotation.Roll"));
+				AddMetadataColumn(InPCGData, ColumnName, AttributeType, TEXT(".Rotation.Pitch"));
+				AddMetadataColumn(InPCGData, ColumnName, AttributeType, TEXT(".Rotation.Yaw"));
 				AddMetadataColumn(InPCGData, ColumnName, AttributeType, TEXT(".Scale.X"));
 				AddMetadataColumn(InPCGData, ColumnName, AttributeType, TEXT(".Scale.Y"));
 				AddMetadataColumn(InPCGData, ColumnName, AttributeType, TEXT(".Scale.Z"));
