Fixed point projection wrt rotation and fixed the quaternion weighted sum.

#rnx

diff --git a/Engine/Plugins/PCG/Source/PCG/Private/Data/PCGPointData.cpp b/Engine/Plugins/PCG/Source/PCG/Private/Data/PCGPointData.cpp
--- a/Engine/Plugins/PCG/Source/PCG/Private/Data/PCGPointData.cpp
+++ b/Engine/Plugins/PCG/Source/PCG/Private/Data/PCGPointData.cpp
@@ -560,7 +560,7 @@ bool UPCGPointData::ProjectPoint(const FTransform& InTransform, const FBox& InBo
 
 	// Computed weighted average of spatial properties
 	FVector WeightedPosition = FVector::ZeroVector;
-	FQuat WeightedQuat = FQuat::Identity;
+	FQuat WeightedQuat = FQuat(0.0, 0.0, 0.0, 0.0);
 	FVector WeightedScale = FVector::ZeroVector;
 	FVector::FReal WeightedDensity = 0;
 	FVector WeightedBoundsMin = FVector::ZeroVector;
diff --git a/Engine/Plugins/PCG/Source/PCG/Public/Metadata/PCGMetadataAttributeTpl.h b/Engine/Plugins/PCG/Source/PCG/Public/Metadata/PCGMetadataAttributeTpl.h
--- a/Engine/Plugins/PCG/Source/PCG/Public/Metadata/PCGMetadataAttributeTpl.h
+++ b/Engine/Plugins/PCG/Source/PCG/Public/Metadata/PCGMetadataAttributeTpl.h
@@ -847,7 +847,7 @@ protected:
 	template<typename IT = T, typename TEnableIf<PCG::Private::MetadataTraits<IT>::CanInterpolate>::Type* = nullptr>
 	void Accumulate(PCGMetadataEntryKey ItemKey, const FPCGMetadataAttributeBase* InAttribute, const TArrayView<const TPair<PCGMetadataEntryKey, float>>& InWeightedKeys)
 	{
-		IT Value = PCG::Private::MetadataTraits<IT>::ZeroValue();
+		IT Value = PCG::Private::MetadataTraits<IT>::ZeroValueForWeightedSum();
 		for (const TPair<PCGMetadataEntryKey, float>& WeightedEntry : InWeightedKeys)
 		{
 			Value = PCG::Private::MetadataTraits<IT>::WeightedSum(
diff --git a/Engine/Plugins/PCG/Source/PCG/Public/Metadata/PCGMetadataAttributeTraits.h b/Engine/Plugins/PCG/Source/PCG/Public/Metadata/PCGMetadataAttributeTraits.h
--- a/Engine/Plugins/PCG/Source/PCG/Public/Metadata/PCGMetadataAttributeTraits.h
+++ b/Engine/Plugins/PCG/Source/PCG/Public/Metadata/PCGMetadataAttributeTraits.h
@@ -451,6 +451,11 @@ namespace PCG
 			{
 				return T{};
 			}
+
+			static T ZeroValueForWeightedSum()
+			{
+				return ZeroValue();
+			}
 		};
 
 		template <typename T>
@@ -582,7 +587,7 @@ namespace PCG
 
 		// Vector types
 		template<typename T>
-		struct VectorTraits : DefaultOperationTraits<T>, DefaultWeightedSumTraits<T>, DefaultStringTraits<T>
+		struct VectorTraits : DefaultOperationTraits<T>, DefaultStringTraits<T>
 		{
 			enum { CompressData = false };
 			enum { CanMinMax = true };
@@ -600,6 +605,16 @@ namespace PCG
 				return T::Zero();
 			}
 
+			static T ZeroValueForWeightedSum()
+			{
+				return ZeroValue();
+			}
+
+			static T WeightedSum(const T& A, const T& B, float Weight)
+			{
+				return A + B * Weight;
+			}
+
 			static DistanceType Distance(const T& A, const T& B) 
 			{
 				return (A - B).Size();
@@ -785,6 +800,11 @@ namespace PCG
 				return FQuat::Identity;
 			}
 
+			static FQuat ZeroValueForWeightedSum()
+			{
+				return FQuat(0.0, 0.0, 0.0, 0.0);
+			}
+
 			static FQuat::FReal Distance(const FQuat& A, const FQuat& B)
 			{
 				return A.AngularDistance(B);
@@ -842,6 +862,11 @@ namespace PCG
 				return FRotator::ZeroRotator;
 			}
 
+			static FRotator ZeroValueForWeightedSum()
+			{
+				return ZeroValue();
+			}
+
 			static FRotator::FReal Distance(const FRotator& A, const FRotator& B)
 			{
 				// Use quaternions for distance calculation, but return degrees instead of radians here
@@ -902,6 +927,11 @@ namespace PCG
 			{
 				return FTransform::Identity;
 			}
+
+			static FTransform ZeroValueForWeightedSum()
+			{
+				return FTransform(FQuat(0.0, 0.0, 0.0, 0.0), FVector::ZeroVector, FVector::ZeroVector);
+			}
 		};
 
 		// Strings
