[PCG] Fix force regen removing all resources
* Fix toggling HiGen in graph settings causing a regeneration even when the component is cleaned up

#rnx

diff --git a/Engine/Plugins/PCG/Source/PCG/Private/PCGComponent.cpp b/Engine/Plugins/PCG/Source/PCG/Private/PCGComponent.cpp
--- a/Engine/Plugins/PCG/Source/PCG/Private/PCGComponent.cpp
+++ b/Engine/Plugins/PCG/Source/PCG/Private/PCGComponent.cpp
@@ -2181,10 +2181,7 @@ void UPCGComponent::Refresh(EPCGChangeType ChangeType, bool bCancelExistingRefre
 		// then the bGenerated flag will be false, which will prevent a subsequent update here
 		if (CurrentRefreshTask == InvalidPCGTaskId && CurrentCleanupTask == InvalidPCGTaskId)
 		{
-			// Always force a regeneration if generation grids might be affected - refresh is insufficient.
-			const bool bGridChanged = !!(ChangeType & EPCGChangeType::GenerationGrid);
-			const bool bForceRegen = bGenerationWasInProgress || bGridChanged;
-			CurrentRefreshTask = Subsystem->ScheduleRefresh(this, bForceRegen, /*bForceCleanup=*/bGridChanged);
+			CurrentRefreshTask = Subsystem->ScheduleRefresh(this, bGenerationWasInProgress);
 		}
 	}
 }
@@ -2279,7 +2276,7 @@ bool UPCGComponent::IsObjectTracked(const UObject* InObject, bool& bOutIsCulled)
 	return CheckMap(StaticallyTrackedKeysToSettings) || CheckMap(DynamicallyTrackedKeysToSettings);
 }
 
-void UPCGComponent::OnRefresh(bool bForceRefresh, bool bForceCleanup)
+void UPCGComponent::OnRefresh(bool bForceRefresh)
 {
 	check(!IsManagedByRuntimeGenSystem());
 
@@ -2293,7 +2290,7 @@ void UPCGComponent::OnRefresh(bool bForceRefresh, bool bForceCleanup)
 	const bool bWasGeneratedOrGenerating = bWasGenerated || bForceRefresh;
 
 	// If we are partitioned but we have resources, we need to force a cleanup
-	if (bForceCleanup || (IsPartitioned() && !GeneratedResources.IsEmpty()))
+	if (IsPartitioned() && !GeneratedResources.IsEmpty())
 	{
 		CleanupLocalImmediate(/*bRemoveComponents=*/true);
 	}
diff --git a/Engine/Plugins/PCG/Source/PCG/Private/PCGGraph.cpp b/Engine/Plugins/PCG/Source/PCG/Private/PCGGraph.cpp
--- a/Engine/Plugins/PCG/Source/PCG/Private/PCGGraph.cpp
+++ b/Engine/Plugins/PCG/Source/PCG/Private/PCGGraph.cpp
@@ -1005,10 +1005,10 @@ bool UPCGGraph::RemoveOutboundEdges(UPCGNode* InNode, const FName& OutboundLabel
 }
 
 #if WITH_EDITOR
-void UPCGGraph::ForceNotificationForEditor()
+void UPCGGraph::ForceNotificationForEditor(EPCGChangeType ChangeType)
 {
 	// Queue up the delayed change
-	NotifyGraphChanged(EPCGChangeType::Structural | EPCGChangeType::GenerationGrid);
+	NotifyGraphChanged(ChangeType);
 
 	if (bUserPausedNotificationsInGraphEditor)
 	{
diff --git a/Engine/Plugins/PCG/Source/PCG/Private/PCGSubsystem.cpp b/Engine/Plugins/PCG/Source/PCG/Private/PCGSubsystem.cpp
--- a/Engine/Plugins/PCG/Source/PCG/Private/PCGSubsystem.cpp
+++ b/Engine/Plugins/PCG/Source/PCG/Private/PCGSubsystem.cpp
@@ -811,7 +811,7 @@ void UPCGSubsystem::CancelGeneration(UPCGComponent* Component)
 	TArray<UPCGComponent*> CancelledComponents = GraphExecutor->Cancel(Component);
 	for (UPCGComponent* CancelledComponent : CancelledComponents)
 	{
-		if (ensure(CancelledComponent))
+		if (CancelledComponent)
 		{
 			CancelledComponent->OnProcessGraphAborted(/*bQuiet=*/true);
 		}
@@ -1154,16 +1154,16 @@ bool UPCGSubsystem::GetOutputData(FPCGTaskId TaskId, FPCGDataCollection& OutData
 
 #if WITH_EDITOR
 
-FPCGTaskId UPCGSubsystem::ScheduleRefresh(UPCGComponent* Component, bool bForceRegen, bool bForceCleanup)
+FPCGTaskId UPCGSubsystem::ScheduleRefresh(UPCGComponent* Component, bool bForceRegen)
 {
 	check(Component && !Component->IsManagedByRuntimeGenSystem());
 
 	TWeakObjectPtr<UPCGComponent> ComponentPtr(Component);
 
-	auto RefreshTask = [ComponentPtr, bForceRegen, bForceCleanup]() {
+	auto RefreshTask = [ComponentPtr, bForceRegen]() {
 		if (UPCGComponent* Component = ComponentPtr.Get())
 		{
-			Component->OnRefresh(bForceRegen, bForceCleanup);
+			Component->OnRefresh(bForceRegen);
 		}
 		return true;
 	};
diff --git a/Engine/Plugins/PCG/Source/PCG/Public/PCGComponent.h b/Engine/Plugins/PCG/Source/PCG/Public/PCGComponent.h
--- a/Engine/Plugins/PCG/Source/PCG/Public/PCGComponent.h
+++ b/Engine/Plugins/PCG/Source/PCG/Public/PCGComponent.h
@@ -306,7 +306,7 @@ public:
 	/** Schedules refresh of the component. If bCancelExistingRefresh is true, any existing refresh is cancelled and a new one is scheduled. */
 	void Refresh(EPCGChangeType ChangeType = EPCGChangeType::None, bool bCancelExistingRefresh = false);
 
-	void OnRefresh(bool bForceRefresh, bool bForceCleanup);
+	void OnRefresh(bool bForceRefresh);
 
 	void StartGenerationInProgress();
 	void StopGenerationInProgress();
diff --git a/Engine/Plugins/PCG/Source/PCG/Public/PCGGraph.h b/Engine/Plugins/PCG/Source/PCG/Public/PCGGraph.h
--- a/Engine/Plugins/PCG/Source/PCG/Public/PCGGraph.h
+++ b/Engine/Plugins/PCG/Source/PCG/Public/PCGGraph.h
@@ -300,8 +300,8 @@ public:
 	void ToggleUserPausedNotificationsForEditor();
 	bool NotificationsForEditorArePausedByUser() const { return bUserPausedNotificationsInGraphEditor; }
 
-	UFUNCTION(BlueprintCallable, Category = "Graph|Advanded")
-	void ForceNotificationForEditor();
+	UFUNCTION(BlueprintCallable, Category = "Graph|Advanced")
+	void ForceNotificationForEditor(EPCGChangeType ChangeType = EPCGChangeType::Structural);
 
 	void PreNodeUndo(UPCGNode* InPCGNode);
 	void PostNodeUndo(UPCGNode* InPCGNode);
diff --git a/Engine/Plugins/PCG/Source/PCG/Public/PCGSubsystem.h b/Engine/Plugins/PCG/Source/PCG/Public/PCGSubsystem.h
--- a/Engine/Plugins/PCG/Source/PCG/Public/PCGSubsystem.h
+++ b/Engine/Plugins/PCG/Source/PCG/Public/PCGSubsystem.h
@@ -236,7 +236,7 @@ public:
 #if WITH_EDITOR
 public:
 	/** Schedule refresh on the current or next frame */
-	FPCGTaskId ScheduleRefresh(UPCGComponent* SourceComponent, bool bForceRefresh, bool bForceCleanup);
+	FPCGTaskId ScheduleRefresh(UPCGComponent* SourceComponent, bool bForceRefresh);
 
 	/** Immediately dirties the partition actors in the given bounds */
 	void DirtyGraph(UPCGComponent* Component, const FBox& InBounds, EPCGComponentDirtyFlag DirtyFlag);
diff --git a/Engine/Plugins/PCG/Source/PCGEditor/Private/PCGEditor.cpp b/Engine/Plugins/PCG/Source/PCGEditor/Private/PCGEditor.cpp
--- a/Engine/Plugins/PCG/Source/PCGEditor/Private/PCGEditor.cpp
+++ b/Engine/Plugins/PCG/Source/PCGEditor/Private/PCGEditor.cpp
@@ -882,6 +882,8 @@ void FPCGEditor::OnForceGraphRegeneration_Clicked()
 {
 	if (PCGGraphBeingEdited)
 	{
+		EPCGChangeType ChangeType = EPCGChangeType::Structural;
+
 		FModifierKeysState ModifierKeys = FSlateApplication::Get().GetModifierKeys();
 		if (ModifierKeys.IsControlDown())
 		{
@@ -889,9 +891,11 @@ void FPCGEditor::OnForceGraphRegeneration_Clicked()
 			{
 				Subsystem->FlushCache();
 			}
+
+			ChangeType |= EPCGChangeType::GenerationGrid;
 		}
 
-		PCGGraphBeingEdited->ForceNotificationForEditor();
+		PCGGraphBeingEdited->ForceNotificationForEditor(ChangeType);
 	}
 }
 
