Unshelved from pending changelist '31871104':

[PCG] Fix inspection workflow by re-capturing debug objects through a fuzzy match
* Fix losing debug selection when BP actors recompile
* Fix losing debug selection when RuntimeGen components fully refresh
* Fix case where attribute list view window is not refreshed and left in a stale state
* Fix regression where debug object tree items cannot be unselected
* Fix missing event subscriptions in PIE which mean graph editor visualization was not updated

#rnx

diff --git a/Engine/Plugins/PCG/Source/PCG/Private/Grid/PCGPartitionActor.cpp b/Engine/Plugins/PCG/Source/PCG/Private/Grid/PCGPartitionActor.cpp
--- a/Engine/Plugins/PCG/Source/PCG/Private/Grid/PCGPartitionActor.cpp
+++ b/Engine/Plugins/PCG/Source/PCG/Private/Grid/PCGPartitionActor.cpp
@@ -1,6 +1,6 @@
 // Copyright Epic Games, Inc. All Rights Reserved.
 
-#include "PCGPartitionActor.h"
+#include "Grid/PCGPartitionActor.h"
 
 #include "PCGComponent.h"
 #include "PCGModule.h"
diff --git a/Engine/Plugins/PCG/Source/PCG/Private/Grid/PCGPartitionActor.h b/Engine/Plugins/PCG/Source/PCG/Private/Grid/PCGPartitionActor.h
deleted file mode 100644
diff --git a/Engine/Plugins/PCG/Source/PCG/Private/PCGComponent.cpp b/Engine/Plugins/PCG/Source/PCG/Private/PCGComponent.cpp
--- a/Engine/Plugins/PCG/Source/PCG/Private/PCGComponent.cpp
+++ b/Engine/Plugins/PCG/Source/PCG/Private/PCGComponent.cpp
@@ -481,7 +481,8 @@ void UPCGComponent::PostProcessGraph(const FBox& InNewBounds, bool bInGenerated,
 
 	StopGenerationInProgress();
 
-	GetSubsystem()->OnComponentGenerationCompleteOrCancelled.Broadcast();
+	UPCGSubsystem* Subsystem = GetSubsystem();
+	Subsystem->OnComponentGenerationCompleteOrCancelled.Broadcast(Subsystem);
 #endif
 }
 
@@ -588,7 +589,8 @@ void UPCGComponent::OnProcessGraphAborted(bool bQuiet)
 
 	OnPCGGraphCancelledDelegate.Broadcast(this);
 
-	GetSubsystem()->OnComponentGenerationCompleteOrCancelled.Broadcast();
+	UPCGSubsystem* Subsystem = GetSubsystem();
+	Subsystem->OnComponentGenerationCompleteOrCancelled.Broadcast(Subsystem);
 #endif
 
 	PCGComponent::BroadcastDynamicDelegate(OnPCGGraphCancelledExternal, this);
@@ -1968,7 +1970,7 @@ void UPCGComponent::DisableInspection()
 	
 	if (InspectionCounter == 0)
 	{
-		ClearInspectionData();
+		ClearInspectionData(/*bClearPerNodeExecutionData=*/false);
 	}
 };
 
@@ -2082,26 +2084,29 @@ const FPCGDataCollection* UPCGComponent::GetInspectionData(const FPCGStack& InSt
 	return InspectionCache.Find(InStack);
 }
 
-void UPCGComponent::ClearInspectionData()
+void UPCGComponent::ClearInspectionData(bool bClearPerNodeExecutionData)
 {
 	{
 		FWriteScopeLock Lock(InspectionCacheLock);
 		InspectionCache.Reset();
 	}
 
+	if (bClearPerNodeExecutionData)
 	{
-		FWriteScopeLock Lock(NodeToStacksThatProducedDataLock);
-		NodeToStacksThatProducedData.Reset();
-	}
+		{
+			FWriteScopeLock Lock(NodeToStacksThatProducedDataLock);
+			NodeToStacksThatProducedData.Reset();
+		}
 
-	{
-		FWriteScopeLock Lock(NodeToStacksInWhichNodeExecutedLock);
-		NodeToStacksInWhichNodeExecuted.Reset();
-	}
+		{
+			FWriteScopeLock Lock(NodeToStacksInWhichNodeExecutedLock);
+			NodeToStacksInWhichNodeExecuted.Reset();
+		}
 
-	{
-		FWriteScopeLock Lock(NodeToStackToInactivePinMaskLock);
-		NodeToStackToInactivePinMask.Reset();
+		{
+			FWriteScopeLock Lock(NodeToStackToInactivePinMaskLock);
+			NodeToStackToInactivePinMask.Reset();
+		}
 	}
 }
 
diff --git a/Engine/Plugins/PCG/Source/PCG/Public/Graph/PCGStackContext.h b/Engine/Plugins/PCG/Source/PCG/Public/Graph/PCGStackContext.h
--- a/Engine/Plugins/PCG/Source/PCG/Public/Graph/PCGStackContext.h
+++ b/Engine/Plugins/PCG/Source/PCG/Public/Graph/PCGStackContext.h
@@ -40,6 +40,8 @@ struct PCG_API FPCGStackFrame
 	friend uint32 GetTypeHash(const FPCGStackFrame& In) { return In.Hash; }
 
 	// A valid frame should either point to an object or have a loop index >= 0.
+	bool IsValid() const { return LoopIndex != INDEX_NONE || Object.IsValid(); }
+
 	TWeakObjectPtr<const UObject> Object;
 	int32 LoopIndex = INDEX_NONE;
 
diff --git a/Engine/Plugins/PCG/Source/PCG/Public/Grid/PCGPartitionActor.h b/Engine/Plugins/PCG/Source/PCG/Public/Grid/PCGPartitionActor.h
new file mode 100644
--- /dev/null
+++ b/Engine/Plugins/PCG/Source/PCG/Public/Grid/PCGPartitionActor.h
@@ -0,0 +1,159 @@
+// Copyright Epic Games, Inc. All Rights Reserved.
+
+#pragma once
+
+#include "ActorPartition/PartitionActor.h"
+
+#include "PCGCommon.h"
+
+#include "PCGPartitionActor.generated.h"
+
+namespace EEndPlayReason { enum Type : int; }
+
+class UPCGComponent;
+class UPCGSubsystem;
+class UBoxComponent;
+
+/** 
+* The APCGPartitionActor actor is used to store grid cell data
+* and its size will be a multiple of the grid size.
+*/
+UCLASS(NotBlueprintable, NotPlaceable)
+class PCG_API APCGPartitionActor : public APartitionActor
+{
+	GENERATED_BODY()
+
+public:
+	APCGPartitionActor(const FObjectInitializer& ObjectInitializer = FObjectInitializer::Get());
+
+	//~Begin UObject Interface
+	virtual void PostLoad() override;
+	virtual void BeginDestroy() override;
+	virtual void Serialize(FArchive& Ar) override;
+	//~End UObject Interface
+
+	//~Begin AActor Interface
+	virtual void BeginPlay();
+	virtual void EndPlay(const EEndPlayReason::Type EndPlayReason) override;
+	virtual void GetActorBounds(bool bOnlyCollidingComponents, FVector& Origin, FVector& BoxExtent, bool bIncludeFromChildActors) const override;
+	virtual void PostRegisterAllComponents() override;
+	virtual void PostUnregisterAllComponents() override;
+#if WITH_EDITOR
+	virtual AActor* GetSceneOutlinerParent() const override;
+#endif
+	//~End AActor Interface
+
+#if WITH_EDITOR
+	//~Begin APartitionActor Interface
+	virtual uint32 GetDefaultGridSize(UWorld* InWorld) const override;
+	virtual FGuid GetGridGuid() const override { return PCGGuid; }
+	virtual bool ShouldIncludeGridSizeInLabel() const override { return true; }
+	//~End APartitionActor Interface
+#endif
+
+	FBox GetFixedBounds() const;
+	FIntVector GetGridCoord() const;
+	uint32 GetPCGGridSize() const { return PCGGridSize; }
+
+	bool IsUsing2DGrid() const { return bUse2DGrid; }
+
+	/** Marks this PartitionActor as managed by the runtime generation system. */
+	void SetToRuntimeGenerated() { bIsRuntimeGenerated = true; }
+	bool IsRuntimeGenerated() const { return bIsRuntimeGenerated; }
+
+	/** Forces the actor location to change even if its mobility is static. */
+	bool Teleport(const FVector& NewLocation);
+
+	/** Register with the PCG Subsystem. */
+	void RegisterPCG();
+
+	/** Unregister with the PCG Subsystem. */
+	void UnregisterPCG();
+
+	void AddGraphInstance(UPCGComponent* OriginalComponent);
+	void RemapGraphInstance(const UPCGComponent* OldOriginalComponent, UPCGComponent* NewOriginalComponent);
+	bool RemoveGraphInstance(UPCGComponent* OriginalComponent);
+	void CleanupDeadGraphInstances(bool bRemoveNonNullOnly = false);
+
+	// When a local component is destroyed. It calls this function. We make sure we don't keep mappings that are dead.
+	void RemoveLocalComponent(UPCGComponent* LocalComponent);
+
+	/** To be called after the creation of a new actor to set the grid guid and size. */
+	void PostCreation(const FGuid& InGridGUID, uint32 InGridSize);
+
+	/** [Game thread only] Return if the actor is safe for deletion, meaning no generation is currently running on all original components. */
+	bool IsSafeForDeletion() const;
+
+	/** Return an array of all the PCGComponents on this actor */
+	TSet<TObjectPtr<UPCGComponent>> GetAllLocalPCGComponents() const;
+
+	/** Return a set of all the PCGComponents linked to this actor */
+	TSet<TObjectPtr<UPCGComponent>> GetAllOriginalPCGComponents() const;
+
+	/** Return true if this PA has any graph instances. */
+	bool HasGraphInstances() const { return LocalToOriginal.Num() > 0; }
+
+	/** Changes transient state for the local component matching the given original component. Returns true if PA becomes empty */
+	bool ChangeTransientState(UPCGComponent* OriginalComponent, EPCGEditorDirtyMode EditingMode);
+
+	// TODO: Make this in-editor only; during runtime, we should keep a map of component to bounds/volume only
+	// and preferably precompute the intersection, so this would make it easier/possible to not have the original actor in game version.
+	UFUNCTION(BlueprintCallable, Category = "PCG|PartitionActor")
+	UPCGComponent* GetLocalComponent(const UPCGComponent* OriginalComponent) const;
+
+	UFUNCTION(BlueprintCallable, Category = "PCG|PartitionActor")
+	UPCGComponent* GetOriginalComponent(const UPCGComponent* LocalComponent) const;
+
+	UPROPERTY()
+	FGuid PCGGuid;
+
+private:
+	UPCGSubsystem* GetSubsystem() const;
+
+#if WITH_EDITOR
+	void UpdateBoundsComponentExtents();
+#endif // WITH_EDITOR
+
+	// Note: this map is not a property and not serialized since we will rebuild it from the LocalToOriginal
+	TMap<TObjectPtr<UPCGComponent>, TObjectPtr<UPCGComponent>> OriginalToLocal;
+
+	UPROPERTY()
+	TMap<TObjectPtr<UPCGComponent>, TSoftObjectPtr<UPCGComponent>> LocalToOriginal;
+
+	// PCG components that are cleared when in preview-on-load mode are kept aside and put back when serializing to prevent changes
+	UPROPERTY(Transient)
+	TMap<TObjectPtr<UPCGComponent>, TSoftObjectPtr<UPCGComponent>> LoadedPreviewComponents;
+
+#if WITH_EDITORONLY_DATA
+	UPROPERTY()
+	TMap<TObjectPtr<UPCGComponent>, TWeakObjectPtr<UPCGComponent>> LocalToOriginalMap_DEPRECATED;
+#endif
+
+	UPROPERTY(VisibleAnywhere, Category = WorldPartition)
+	uint32 PCGGridSize;
+
+	UPROPERTY(VisibleAnywhere, Category = WorldPartition)
+	bool bUse2DGrid;
+
+#if WITH_EDITORONLY_DATA
+	/** Box component to draw the Partition actor bounds in the Editor viewport */
+	UPROPERTY(Transient)
+	TObjectPtr<UBoxComponent> BoundsComponent;
+#endif // WITH_EDITORONLY_DATA
+
+	/** Tracks the registration status of this PA with the ActorAndComponentMapping system. Helps us avoid invalid (un)registers. */
+	bool bIsRegistered = false;
+
+	/** Tracks if this actor was created by the Runtime Generation system. */
+	bool bIsRuntimeGenerated = false;
+
+	/** Utility bool to check if PostCreation/PostLoad was called. */
+	bool bWasPostCreatedLoaded = false;
+
+public:
+	/** 
+	 * Gets the name this partition actor should have.
+	 * This does not respect traditional PA name contents like GridGuid, ShouldIncludeGridSizeInName, or ContextHash.
+	 */
+	static FString GetPCGPartitionActorName(uint32 GridSize, const FIntVector& GridCoords, bool bRuntimeGenerated);
+};
diff --git a/Engine/Plugins/PCG/Source/PCG/Public/PCGComponent.h b/Engine/Plugins/PCG/Source/PCG/Public/PCGComponent.h
--- a/Engine/Plugins/PCG/Source/PCG/Public/PCGComponent.h
+++ b/Engine/Plugins/PCG/Source/PCG/Public/PCGComponent.h
@@ -329,7 +329,7 @@ public:
 	void DisableInspection();
 	void StoreInspectionData(const FPCGStack* InStack, const UPCGNode* InNode, const FPCGDataCollection& InInputData, const FPCGDataCollection& InOutputData);
 	const FPCGDataCollection* GetInspectionData(const FPCGStack& InStack) const;
-	void ClearInspectionData();
+	void ClearInspectionData(bool bClearPerNodeExecutionData = true);
 
 	/** Whether a task for the given node and stack was executed during the last execution. */
 	bool WasNodeExecuted(const UPCGNode* InNode, const FPCGStack& Stack) const;
diff --git a/Engine/Plugins/PCG/Source/PCG/Public/PCGSubsystem.h b/Engine/Plugins/PCG/Source/PCG/Public/PCGSubsystem.h
--- a/Engine/Plugins/PCG/Source/PCG/Public/PCGSubsystem.h
+++ b/Engine/Plugins/PCG/Source/PCG/Public/PCGSubsystem.h
@@ -33,7 +33,7 @@ typedef TSharedPtr<IPCGElement, ESPMode::ThreadSafe> FPCGElementPtr;
 class UWorld;
 
 #if WITH_EDITOR
-DECLARE_MULTICAST_DELEGATE(FPCGOnComponentGenerationCompleteOrCancelled);
+DECLARE_MULTICAST_DELEGATE_OneParam(FPCGOnComponentGenerationCompleteOrCancelled, UPCGSubsystem*);
 #endif // WITH_EDITOR
 
 /**
diff --git a/Engine/Plugins/PCG/Source/PCGEditor/Private/PCGEditor.cpp b/Engine/Plugins/PCG/Source/PCGEditor/Private/PCGEditor.cpp
--- a/Engine/Plugins/PCG/Source/PCGEditor/Private/PCGEditor.cpp
+++ b/Engine/Plugins/PCG/Source/PCGEditor/Private/PCGEditor.cpp
@@ -111,11 +111,6 @@ void FPCGEditor::Initialize(const EToolkitMode::Type InMode, const TSharedPtr<cl
 {
 	PCGGraphBeingEdited = InPCGGraph;
 
-	if (UPCGSubsystem* Subsystem = GetSubsystem())
-	{
-		Subsystem->OnComponentGenerationCompleteOrCancelled.AddRaw(this, &FPCGEditor::OnComponentGenerationCompleteOrCancelled);
-	}
-
 	// Initializes the UPCGEditorGraph if needed
 	GetPCGEditorGraph(InPCGGraph);
 
@@ -212,6 +207,18 @@ void FPCGEditor::Initialize(const EToolkitMode::Type InMode, const TSharedPtr<cl
 		GEngine->OnLevelActorDeleted().AddRaw(this, &FPCGEditor::OnLevelActorDeleted);
 	}
 
+	// Hook to PIE start/end to keep callbacks up to date.
+	FEditorDelegates::PostPIEStarted.AddRaw(this, &FPCGEditor::OnPostPIEStarted);
+	FEditorDelegates::EndPIE.AddRaw(this, &FPCGEditor::OnEndPIE);
+
+	if (GEditor)
+	{
+		RegisterDelegatesForWorld(GEditor->GetEditorWorldContext().World());
+
+		// In case the editor is opened while in PIE, we should try setting up callbacks for the PIE world.
+		RegisterDelegatesForWorld(GEditor->PlayWorld.Get());
+	}
+
 	// Clear inspection flag on all nodes.
 	for (UEdGraphNode* EdGraphNode : PCGEditorGraph->Nodes)
 	{
@@ -1917,6 +1924,10 @@ void FPCGEditor::OnToggleInspected()
 
 		DebugObjectTreeWidget->SetNodeBeingInspected(PCGNode);
 	}
+	else
+	{
+		DebugObjectTreeWidget->SetNodeBeingInspected(nullptr);
+	}
 }
 
 bool FPCGEditor::CanToggleInspected() const
@@ -2757,11 +2768,15 @@ void FPCGEditor::OnClose()
 	{
 		LevelEditor->OnMapChanged().RemoveAll(this);
 	}
+
 	if (GEngine)
 	{
 		GEngine->OnLevelActorDeleted().RemoveAll(this);
 	}
 
+	FEditorDelegates::PostPIEStarted.RemoveAll(this);
+	FEditorDelegates::EndPIE.RemoveAll(this);
+
 	// Extra nodes are replicated on editor close, to be saved in the underlying PCGGraph
 	ReplicateExtraNodes();
 
@@ -2788,9 +2803,10 @@ void FPCGEditor::OnClose()
 		}
 	}
 
-	if (UPCGSubsystem* Subsystem = GetSubsystem())
+	if (GEditor)
 	{
-		Subsystem->OnComponentGenerationCompleteOrCancelled.RemoveAll(this);
+		UnregisterDelegatesForWorld(GEditor->GetEditorWorldContext().World());
+		UnregisterDelegatesForWorld(GEditor->PlayWorld.Get());
 	}
 }
 
@@ -2939,11 +2955,10 @@ void FPCGEditor::JumpToDefinition(const UClass* Class) const
 	}
 }
 
-void FPCGEditor::OnComponentGenerationCompleteOrCancelled()
+void FPCGEditor::OnComponentGenerationCompleteOrCancelled(UPCGSubsystem* Subsystem)
 {
 	DebugObjectTreeWidget->RequestRefresh();
 
-	const UPCGSubsystem* Subsystem = GetSubsystem();
 	const bool CacheDebuggingEnabled = Subsystem && Subsystem->IsGraphCacheDebuggingEnabled();
 
 	// Refresh nodes to report any errors/warnings, and to display culling state after execution.
@@ -2972,6 +2987,24 @@ UPCGSubsystem* FPCGEditor::GetSubsystem()
 	return UPCGSubsystem::GetInstance(World);
 }
 
+void FPCGEditor::RegisterDelegatesForWorld(UWorld* World)
+{
+	UnregisterDelegatesForWorld(World);
+
+	if (UPCGSubsystem* Subsystem = UPCGSubsystem::GetInstance(World))
+	{
+		Subsystem->OnComponentGenerationCompleteOrCancelled.AddRaw(this, &FPCGEditor::OnComponentGenerationCompleteOrCancelled);
+	}
+}
+
+void FPCGEditor::UnregisterDelegatesForWorld(UWorld* World)
+{
+	if (UPCGSubsystem* Subsystem = UPCGSubsystem::GetInstance(World))
+	{
+		Subsystem->OnComponentGenerationCompleteOrCancelled.RemoveAll(this);
+	}
+}
+
 void FPCGEditor::OnMapChanged(UWorld* InWorld, EMapChangeType InMapChangedType)
 {
 	if (InMapChangedType != EMapChangeType::SaveMap)
@@ -2982,13 +3015,24 @@ void FPCGEditor::OnMapChanged(UWorld* InWorld, EMapChangeType InMapChangedType)
 		}
 
 		// Subsystem has been torn down and rebuilt.
-		if (UPCGSubsystem* Subsystem = GetSubsystem())
+		if (GEditor)
 		{
-			Subsystem->OnComponentGenerationCompleteOrCancelled.AddRaw(this, &FPCGEditor::OnComponentGenerationCompleteOrCancelled);
+			RegisterDelegatesForWorld(GEditor->GetEditorWorldContext().World());
+			RegisterDelegatesForWorld(GEditor->PlayWorld.Get());
 		}
 	}
 }
 
+void FPCGEditor::OnPostPIEStarted(bool bIsSimulating)
+{
+	RegisterDelegatesForWorld(GEditor ? GEditor->PlayWorld.Get() : nullptr);
+}
+
+void FPCGEditor::OnEndPIE(bool bIsSimulating)
+{
+	UnregisterDelegatesForWorld(GEditor ? GEditor->PlayWorld.Get() : nullptr);
+}
+
 void FPCGEditor::OnLevelActorDeleted(AActor* InActor)
 {
 	// Forward call as this makes an effort to retain the selection if the selected component has not been deleted.
diff --git a/Engine/Plugins/PCG/Source/PCGEditor/Private/PCGEditor.h b/Engine/Plugins/PCG/Source/PCGEditor/Private/PCGEditor.h
--- a/Engine/Plugins/PCG/Source/PCGEditor/Private/PCGEditor.h
+++ b/Engine/Plugins/PCG/Source/PCGEditor/Private/PCGEditor.h
@@ -326,12 +326,17 @@ private:
 	void ReplicateExtraNodes() const;
 
 	/** Called when a component finishes executing. Useful for updating debugging tools/UIs. */
-	void OnComponentGenerationCompleteOrCancelled();
+	void OnComponentGenerationCompleteOrCancelled(UPCGSubsystem* Subsystem);
 
 	/** Trigger any generation required to ensure debug display is up to date. */
 	void UpdateDebugAfterComponentSelection(UPCGComponent* InOldComponent, UPCGComponent* InNewComponent, bool bNewComponentStartedInspecting);
 
+	void RegisterDelegatesForWorld(UWorld* World);
+	void UnregisterDelegatesForWorld(UWorld* World);
+
 	void OnMapChanged(UWorld* InWorld, EMapChangeType InMapChangedType);
+	void OnPostPIEStarted(bool bIsSimulating);
+	void OnEndPIE(bool bIsSimulating);
 	void OnLevelActorDeleted(AActor* InActor);
 
 	TSharedRef<SDockTab> SpawnTab_GraphEditor(const FSpawnTabArgs& Args);
diff --git a/Engine/Plugins/PCG/Source/PCGEditor/Private/SPCGEditorGraphAttributeListView.cpp b/Engine/Plugins/PCG/Source/PCGEditor/Private/SPCGEditorGraphAttributeListView.cpp
--- a/Engine/Plugins/PCG/Source/PCGEditor/Private/SPCGEditorGraphAttributeListView.cpp
+++ b/Engine/Plugins/PCG/Source/PCGEditor/Private/SPCGEditorGraphAttributeListView.cpp
@@ -566,18 +566,9 @@ void SPCGEditorGraphAttributeListView::OnInspectedStackChanged(const FPCGStack&
 	{
 		PCGComponent->OnPCGGraphGeneratedDelegate.AddSP(this, &SPCGEditorGraphAttributeListView::OnGenerateUpdated);
 		PCGComponent->OnPCGGraphCleanedDelegate.AddSP(this, &SPCGEditorGraphAttributeListView::OnGenerateUpdated);
-
-		// Refresh if PCGComponent is being inspected already since we wont get a refresh after generation
-		if (PCGComponent->IsInspecting())
-		{
-			RequestRefresh();
-		}
-	}
-	else
-	{
-		// Refresh if PCGComponent is cleared since we wont get a refresh after generate/cleaned
-		RequestRefresh();
 	}
+
+	RequestRefresh();
 }
 
 UPCGEditorGraphNodeBase* SPCGEditorGraphAttributeListView::GetNodeBeingInspected() const
diff --git a/Engine/Plugins/PCG/Source/PCGEditor/Private/SPCGEditorGraphDebugObjectTree.cpp b/Engine/Plugins/PCG/Source/PCGEditor/Private/SPCGEditorGraphDebugObjectTree.cpp
--- a/Engine/Plugins/PCG/Source/PCGEditor/Private/SPCGEditorGraphDebugObjectTree.cpp
+++ b/Engine/Plugins/PCG/Source/PCGEditor/Private/SPCGEditorGraphDebugObjectTree.cpp
@@ -5,6 +5,7 @@
 #include "PCGComponent.h"
 #include "PCGGraph.h"
 #include "PCGSubsystem.h"
+#include "Grid/PCGPartitionActor.h"
 #include "Helpers/PCGHelpers.h"
 
 #include "PCGEditor.h"
@@ -259,7 +260,7 @@ void SPCGEditorGraphDebugObjectTree::Construct(const FArguments& InArgs, TShared
 		.OnExpansionChanged(this, &SPCGEditorGraphDebugObjectTree::OnExpansionChanged)
 		.OnSetExpansionRecursive(this, &SPCGEditorGraphDebugObjectTree::OnSetExpansionRecursive)
 		.ItemHeight(18)
-		.SelectionMode(ESelectionMode::Single)
+		.SelectionMode(ESelectionMode::SingleToggle)
 		.AllowOverscroll(EAllowOverscroll::No)
 		.ExternalScrollbar(VerticalScrollBar)
 		.ConsumeMouseWheel(EConsumeMouseWheel::Always)
@@ -785,7 +786,27 @@ void SPCGEditorGraphDebugObjectTree::RefreshTree()
 	}
 
 	SortTreeItems();
+	RestoreTreeState();
+}
+
+void SPCGEditorGraphDebugObjectTree::SortTreeItems(bool bIsAscending, bool bIsRecursive)
+{
+	RootItems.Sort([bIsAscending](const FPCGEditorGraphDebugObjectItemPtr& InLHS, const FPCGEditorGraphDebugObjectItemPtr& InRHS)
+	{
+		return (InLHS->GetLabel() < InRHS->GetLabel()) == bIsAscending;
+	});
 
+	if (bIsRecursive)
+	{
+		for (const FPCGEditorGraphDebugObjectItemPtr& Item : RootItems)
+		{
+			Item->SortChildren(bIsAscending, bIsRecursive);
+		}
+	}
+}
+
+void SPCGEditorGraphDebugObjectTree::RestoreTreeState()
+{
 	// Try to restore user item expansion.
 	TSet<FPCGStack> ExpandedStacksBefore = ExpandedStacks;
 	for (const FPCGStack& ExpandedStack : ExpandedStacksBefore)
@@ -801,7 +822,9 @@ void SPCGEditorGraphDebugObjectTree::RefreshTree()
 		}
 	}
 
-	// Try to restore user item selection.
+	bool bFoundMatchingStack = false;
+
+	// Try to restore user item selection by exact matching.
 	for (FPCGEditorGraphDebugObjectItemPtr& Item : AllGraphItems)
 	{
 		const FPCGStack* ItemStack = Item->GetPCGStack();
@@ -809,23 +832,86 @@ void SPCGEditorGraphDebugObjectTree::RefreshTree()
 		if (ItemStack && SelectedStack == *ItemStack)
 		{
 			DebugObjectTreeView->SetItemSelection(Item, true);
+			bFoundMatchingStack = true;
 			break;
 		}
 	}
-}
-
-void SPCGEditorGraphDebugObjectTree::SortTreeItems(bool bIsAscending, bool bIsRecursive)
-{
-	RootItems.Sort([bIsAscending](const FPCGEditorGraphDebugObjectItemPtr& InLHS, const FPCGEditorGraphDebugObjectItemPtr& InRHS)
-	{
-		return (InLHS->GetLabel() < InRHS->GetLabel()) == bIsAscending;
-	});
 
-	if (bIsRecursive)
+	// Try to restore user item selection by fuzzy matching (e.g. share the same owner) if no exactly matching stack was found.
+	if (!bFoundMatchingStack)
 	{
-		for (const FPCGEditorGraphDebugObjectItemPtr& Item : RootItems)
+		for (FPCGEditorGraphDebugObjectItemPtr& Item : AllGraphItems)
 		{
-			Item->SortChildren(bIsAscending, bIsRecursive);
+			const FPCGStack* ItemStack = Item->GetPCGStack();
+
+			bool bFuzzyMatch = false;
+
+			if (ItemStack && SelectedGraph.Get() == ItemStack->GetRootGraph())
+			{
+				const UPCGComponent* RootComponent = ItemStack->GetRootComponent();
+				const AActor* RootOwner = RootComponent ? RootComponent->GetOwner() : nullptr;
+				const APCGPartitionActor* RootPartitionActor = Cast<APCGPartitionActor>(RootOwner);
+
+				if (RootComponent && RootPartitionActor)
+				{
+					// For local components, we can fuzzy match as long as the GridSize, GridCoord, OriginalComponent, and ExecutionDomain are the same.
+					// This is equivalent to saying they are on the same partition actor and come from the same original component.
+					bFuzzyMatch = SelectedGridSize == RootComponent->GetGenerationGridSize()
+						&& SelectedGridCoord == RootPartitionActor->GetGridCoord()
+						&& SelectedOriginalComponent.Get() == RootPartitionActor->GetOriginalComponent(RootComponent)
+						&& (SelectedOriginalComponent.IsValid() && SelectedOriginalComponent->IsManagedByRuntimeGenSystem() == RootComponent->IsManagedByRuntimeGenSystem());
+				}
+				else
+				{
+					// For original components, we can fuzzy match as long as the owning actor is the same.
+					// Note: This fails for multiple original components with the same graph on the same actor, since there is no way to know which one to pick.
+					if (SelectedOwner.Get() == RootOwner && Item->GetParent() && Item->GetParent()->GetChildren().Num() == 1)
+					{
+						int32 ItemRootGraphIndex = INDEX_NONE;
+						int32 SelectedRootGraphIndex = INDEX_NONE;
+
+						ItemStack->GetRootGraph(&ItemRootGraphIndex);
+						SelectedStack.GetRootGraph(&SelectedRootGraphIndex);
+
+						const TArray<FPCGStackFrame>& ItemStackFrames = ItemStack->GetStackFrames();
+						const TArray<FPCGStackFrame>& SelectedStackFrames = SelectedStack.GetStackFrames();
+
+						// If the stacks match from the RootGraph onwards, then our fuzzy match should succeed.
+						if (ItemRootGraphIndex != INDEX_NONE && ItemRootGraphIndex == SelectedRootGraphIndex && ItemStackFrames.Num() == SelectedStackFrames.Num())
+						{
+							bool bAllStackFramesMatch = true;
+
+							for (int I = ItemRootGraphIndex; I < ItemStackFrames.Num(); ++I)
+							{
+								if (!ItemStackFrames[I].IsValid() || ItemStackFrames[I] != SelectedStackFrames[I])
+								{
+									bAllStackFramesMatch = false;
+									break;
+								}
+							}
+
+							if (bAllStackFramesMatch)
+							{
+								bFuzzyMatch = true;
+							}
+						}
+					}
+				}
+			}
+
+			if (bFuzzyMatch)
+			{
+				// Force the selected object to re-expand.
+				FPCGEditorGraphDebugObjectItemPtr Parent = Item->GetParent();
+				while (Parent)
+				{
+					DebugObjectTreeView->SetItemExpansion(Parent, true);
+					Parent = Parent->GetParent();
+				}
+
+				DebugObjectTreeView->SetItemSelection(Item, true);
+				break;
+			}
 		}
 	}
 }
@@ -900,13 +986,34 @@ void SPCGEditorGraphDebugObjectTree::OnSelectionChanged(FPCGEditorGraphDebugObje
 		PreviouslySelectedStack = SelectedStack;
 	}
 
+	// Reset selected item information.
+	SelectedStack = FPCGStack();
+	SelectedGraph = nullptr;
+	SelectedOwner = nullptr;
+	SelectedGridSize = PCGHiGenGrid::UnboundedGridSize();
+	SelectedGridCoord = FIntVector::ZeroValue;
+	SelectedOriginalComponent = nullptr;
+
 	if (const FPCGStack* Stack = InItem ? InItem->GetPCGStack() : nullptr)
 	{
 		SelectedStack = *Stack;
-	}
-	else
-	{
-		SelectedStack = FPCGStack();
+		SelectedGraph = SelectedStack.GetRootGraph();
+
+		if (const UPCGComponent* RootComponent = SelectedStack.GetRootComponent())
+		{
+			SelectedOwner = RootComponent->GetOwner();
+
+			if (const APCGPartitionActor* PartitionActor = Cast<APCGPartitionActor>(SelectedOwner.Get()))
+			{
+				SelectedGridSize = RootComponent->GetGenerationGridSize();
+				SelectedGridCoord = PartitionActor->GetGridCoord();
+				SelectedOriginalComponent = PartitionActor->GetOriginalComponent(RootComponent);
+			}
+			else
+			{
+				SelectedOriginalComponent = RootComponent;
+			}
+		}
 	}
 
 	if (bDisableDebugObjectChangeNotification)
@@ -999,7 +1106,7 @@ void SPCGEditorGraphDebugObjectTree::ExpandAndSelectDebugObject(FPCGEditorGraphD
 	// Set the discovered item as the debug object if it is the only child.
 	if (NumChildren == 1)
 	{
-		OnSelectionChanged(Item, ESelectInfo::Direct);
+		DebugObjectTreeView->SetSelection(Item);
 	}
 }
 
diff --git a/Engine/Plugins/PCG/Source/PCGEditor/Private/SPCGEditorGraphDebugObjectTree.h b/Engine/Plugins/PCG/Source/PCGEditor/Private/SPCGEditorGraphDebugObjectTree.h
--- a/Engine/Plugins/PCG/Source/PCGEditor/Private/SPCGEditorGraphDebugObjectTree.h
+++ b/Engine/Plugins/PCG/Source/PCGEditor/Private/SPCGEditorGraphDebugObjectTree.h
@@ -228,6 +228,7 @@ private:
 
 	void RefreshTree();
 	void SortTreeItems(bool bIsAscending = true, bool bIsRecursive = true);
+	void RestoreTreeState();
 
 	void AddStacksToTree(
 		const TArray<FPCGStack>& Stacks,
@@ -272,6 +273,13 @@ private:
 	/** Used to retain item selection state across tree refreshes. */
 	FPCGStack SelectedStack;
 
+	/** Used to retain item selection state across tree refreshes if the SelectedStack is invalidated (e.g. through BP reconstruction). */
+	TWeakObjectPtr<const UPCGGraph> SelectedGraph = nullptr;
+	TWeakObjectPtr<const AActor> SelectedOwner = nullptr;
+	uint32 SelectedGridSize = PCGHiGenGrid::UnboundedGridSize();
+	FIntVector SelectedGridCoord = FIntVector::ZeroValue;
+	TWeakObjectPtr<const UPCGComponent> SelectedOriginalComponent = nullptr;
+
 	/** The previous stack that the user selected. */
 	FPCGStack PreviouslySelectedStack;
 
