Fixed issues related to undo/redo and enable/disable nodes.
Also fixed the debug toggle not having a proper transaction.

#rnx

diff --git a/Engine/Plugins/PCG/Source/PCG/Private/PCGSettings.cpp b/Engine/Plugins/PCG/Source/PCG/Private/PCGSettings.cpp
--- a/Engine/Plugins/PCG/Source/PCG/Private/PCGSettings.cpp
+++ b/Engine/Plugins/PCG/Source/PCG/Private/PCGSettings.cpp
@@ -545,7 +545,7 @@ void UPCGSettings::PostEditChangeProperty(FPropertyChangedEvent& PropertyChanged
 	if (PropertyName != GET_MEMBER_NAME_CHECKED(UPCGSettings, DeterminismSettings))
 	{
 		// If we have a property name then get the change type for that property, otherwise assume deepest change type.
-		EPCGChangeType ChangeType = EPCGChangeType::Structural | EPCGChangeType::GenerationGrid;
+		EPCGChangeType ChangeType = EPCGChangeType::Cosmetic | EPCGChangeType::Settings | EPCGChangeType::Structural | EPCGChangeType::GenerationGrid;
 		if (PropertyName != NAME_None)
 		{
 			ChangeType = GetChangeTypeForProperty(PropertyChangedEvent);
diff --git a/Engine/Plugins/PCG/Source/PCGEditor/Private/PCGEditor.cpp b/Engine/Plugins/PCG/Source/PCGEditor/Private/PCGEditor.cpp
--- a/Engine/Plugins/PCG/Source/PCGEditor/Private/PCGEditor.cpp
+++ b/Engine/Plugins/PCG/Source/PCGEditor/Private/PCGEditor.cpp
@@ -2005,9 +2005,13 @@ void FPCGEditor::OnToggleEnabled()
 {
 	const ECheckBoxState CheckState = GetEnabledCheckState();
 	const bool bNewCheckState = !(CheckState != ECheckBoxState::Unchecked);
+
+	// To prevent the changes on the editor node from being in the transaction, we delay reconstruction.
+	TArray<FPCGDeferNodeReconstructScope> DeferredEditorNodes;
+
 	if (GraphEditorWidget.IsValid())
 	{
-		const FScopedTransaction Transaction(*FPCGEditorCommon::ContextIdentifier, LOCTEXT("PCGEditorToggleEnableTransactionMessage", "PCG Editor: Toggle Enable Nodes"), nullptr);
+		FScopedTransaction Transaction(*FPCGEditorCommon::ContextIdentifier, LOCTEXT("PCGEditorToggleEnableTransactionMessage", "PCG Editor: Toggle Enable Nodes"), nullptr);
 
 		bool bChanged = false;
 		for (UObject* Object : GraphEditorWidget->GetSelectedNodes())
@@ -2023,6 +2027,7 @@ void FPCGEditor::OnToggleEnabled()
 
 			if (PCGSettingsInterface->bEnabled != bNewCheckState)
 			{
+				DeferredEditorNodes.Emplace(PCGEditorGraphNode);
 				PCGSettingsInterface->Modify();
 				PCGSettingsInterface->SetEnabled(bNewCheckState);
 				bChanged = true;
@@ -2033,6 +2038,10 @@ void FPCGEditor::OnToggleEnabled()
 		{
 			GraphEditorWidget->NotifyGraphChanged();
 		}
+		else
+		{
+			Transaction.Cancel();
+		}
 	}
 }
 
@@ -2104,8 +2113,9 @@ void FPCGEditor::OnToggleDebug()
 
 	if (GraphEditorWidget.IsValid())
 	{
-		const FScopedTransaction Transaction(*FPCGEditorCommon::ContextIdentifier, LOCTEXT("PCGEditorToggleDebugTransactionMessage", "PCG Editor: Toggle Debug Nodes"), nullptr);
+		FScopedTransaction Transaction(*FPCGEditorCommon::ContextIdentifier, LOCTEXT("PCGEditorToggleDebugTransactionMessage", "PCG Editor: Toggle Debug Nodes"), nullptr);
 
+		bool bChanged = false;
 		for (UObject* Object : GraphEditorWidget->GetSelectedNodes())
 		{
 			UPCGEditorGraphNodeBase* PCGEditorGraphNode = Cast<UPCGEditorGraphNodeBase>(Object);
@@ -2119,10 +2129,17 @@ void FPCGEditor::OnToggleDebug()
 
 			if (PCGSettingsInterface->bDebug != bNewCheckState)
 			{
+				PCGSettingsInterface->Modify(/*bAlwaysMarkDirty=*/false);
 				PCGSettingsInterface->bDebug = bNewCheckState;
 				PCGNode->OnNodeChangedDelegate.Broadcast(PCGNode, EPCGChangeType::Settings);
+				bChanged = true;
 			}
 		}
+
+		if (!bChanged)
+		{
+			Transaction.Cancel();
+		}
 	}
 }
 
@@ -2202,7 +2219,7 @@ void FPCGEditor::OnDebugOnlySelected()
 
 			if (PCGSettingsInterface->bDebug != bShouldBeDebug)
 			{
-				PCGSettingsInterface->Modify();
+				PCGSettingsInterface->Modify(/*bAlwaysMarkDirty=*/false);
 				PCGSettingsInterface->bDebug = bShouldBeDebug;
 				PCGNode->OnNodeChangedDelegate.Broadcast(PCGNode, EPCGChangeType::Settings);
 				bChanged = true;
@@ -2235,7 +2252,7 @@ void FPCGEditor::OnDisableDebugOnAllNodes()
 
 			if (PCGSettingsInterface->bDebug)
 			{
-				PCGSettingsInterface->Modify();
+				PCGSettingsInterface->Modify(/*bAlwaysMarkDirty=*/false);
 				PCGSettingsInterface->bDebug = false;
 				PCGNode->OnNodeChangedDelegate.Broadcast(PCGNode, EPCGChangeType::Settings);
 				bChanged = true;
diff --git a/Engine/Plugins/PCG/Source/PCGEditor/Private/PCGEditorGraphNodeBase.h b/Engine/Plugins/PCG/Source/PCGEditor/Private/PCGEditorGraphNodeBase.h
--- a/Engine/Plugins/PCG/Source/PCGEditor/Private/PCGEditorGraphNodeBase.h
+++ b/Engine/Plugins/PCG/Source/PCGEditor/Private/PCGEditorGraphNodeBase.h
@@ -3,6 +3,7 @@
 #pragma once
 
 #include "EdGraph/EdGraphNode.h"
+#include "EdGraph/EdGraphPin.h"
 
 #include "PCGCommon.h"
 
@@ -163,3 +164,51 @@ protected:
 	/** The higen grid currently being inspected if any, otherwise Uninitialized. */
 	EPCGHiGenGrid InspectedGenerationGrid = EPCGHiGenGrid::Uninitialized;
 };
+
+/** Disables reconstruct on nodes (or from a pin) and re-enables in destructor. */
+struct FPCGDeferNodeReconstructScope
+{
+	explicit FPCGDeferNodeReconstructScope(UEdGraphPin* FromPin)
+		: Node(FromPin ? Cast<UPCGEditorGraphNodeBase>(FromPin->GetOwningNode()) : nullptr)
+	{
+		if (Node)
+		{
+			Node->EnableDeferredReconstruct();
+		}
+	}
+
+	explicit FPCGDeferNodeReconstructScope(UPCGEditorGraphNodeBase* InNode)
+		: Node(InNode)
+	{
+		if (Node)
+		{
+			Node->EnableDeferredReconstruct();
+		}
+	}
+
+	FPCGDeferNodeReconstructScope(const FPCGDeferNodeReconstructScope&) = delete;
+	FPCGDeferNodeReconstructScope(FPCGDeferNodeReconstructScope&& Other)
+	{
+		Node = Other.Node;
+		Other.Node = nullptr;
+	}
+
+	FPCGDeferNodeReconstructScope& operator=(const FPCGDeferNodeReconstructScope&) = delete;
+	FPCGDeferNodeReconstructScope& operator=(FPCGDeferNodeReconstructScope&& Other)
+	{
+		Swap(Node, Other.Node);
+		return *this;
+	}
+
+	~FPCGDeferNodeReconstructScope()
+	{
+		if (Node)
+		{
+			Node->DisableDeferredReconstruct();
+			Node = nullptr;
+		}
+	}
+
+private:
+	UPCGEditorGraphNodeBase* Node = nullptr;
+};
\ No newline at end of file
diff --git a/Engine/Plugins/PCG/Source/PCGEditor/Private/PCGEditorGraphSchemaActions.cpp b/Engine/Plugins/PCG/Source/PCGEditor/Private/PCGEditorGraphSchemaActions.cpp
--- a/Engine/Plugins/PCG/Source/PCGEditor/Private/PCGEditorGraphSchemaActions.cpp
+++ b/Engine/Plugins/PCG/Source/PCGEditor/Private/PCGEditorGraphSchemaActions.cpp
@@ -27,30 +27,6 @@
 
 #define LOCTEXT_NAMESPACE "PCGEditorGraphSchemaActions"
 
-/** Disables reconstruct on the owning node of a pin and re-enables in destructor. */
-struct FPCGDeferNodeReconstructScope
-{
-	FPCGDeferNodeReconstructScope(UEdGraphPin* FromPin)
-		: Node(FromPin ? Cast<UPCGEditorGraphNodeBase>(FromPin->GetOwningNode()) : nullptr)
-	{
-		if (Node)
-		{
-			Node->EnableDeferredReconstruct();
-		}
-	}
-
-	~FPCGDeferNodeReconstructScope()
-	{
-		if (Node)
-		{
-			Node->DisableDeferredReconstruct();
-		}
-	}
-
-private:
-	UPCGEditorGraphNodeBase* Node = nullptr;
-};
-
 UEdGraphNode* FPCGEditorGraphSchemaAction_NewNativeElement::PerformAction(UEdGraph* ParentGraph, UEdGraphPin* FromPin, const FVector2D Location, bool bSelectNewNode)
 {
 	UPCGEditorGraph* EditorGraph = Cast<UPCGEditorGraph>(ParentGraph);
