[PCG] Actor Tracking fixes
- Add a new delegate to ULevelInstanceSubsystem to notify which ILevelInstanceInterface have changed after an Edit
- PCG now properly detects that a Level Instance as been changed through this new delegate
- Move registering of PCGActorAndComponentMapping delegates to PCGSubsystem::PostInitialize which happens before actor components are registered
- Add a delay on added/changed Level Instances so we can wait for them to be loaded before processing the tracking
- Make some members editor only since they weren't used in non editor builds
- Other general cleanup
- Fix StorePreviousData not storing data on sub actors if parent actor wasn't tracked causing Component to not regenerate properly once Actors exited the volume (no previous bounds data found)

#rnx

diff --git a/Engine/Plugins/PCG/Source/PCG/Private/PCGActorAndComponentMapping.cpp b/Engine/Plugins/PCG/Source/PCG/Private/PCGActorAndComponentMapping.cpp
--- a/Engine/Plugins/PCG/Source/PCG/Private/PCGActorAndComponentMapping.cpp
+++ b/Engine/Plugins/PCG/Source/PCG/Private/PCGActorAndComponentMapping.cpp
@@ -23,9 +23,10 @@
 #include "Engine/Level.h"
 #include "Engine/StaticMesh.h"
 #include "Kismet/GameplayStatics.h"
-#include "LevelInstance/LevelInstanceActor.h"
+#include "PackedLevelActor/PackedLevelActor.h"
 #include "LevelInstance/LevelInstanceEditorInstanceActor.h"
 #include "LevelInstance/LevelInstanceInterface.h"
+#include "LevelInstance/LevelInstanceEditorPivotInterface.h"
 #include "LevelInstance/LevelInstanceSubsystem.h"
 #include "Materials/MaterialInterface.h"
 
@@ -96,22 +97,39 @@ namespace PCGActorAndComponentMapping
 		TEXT("If delayed unregister for all is introducing bad behavior, disables it, allowing people to continue working while we investigate."));
 
 #if WITH_EDITOR
-	void PropagateToLevelInstanceActors(AActor* InActor, UPCGSubsystem* PCGSubsystem, TFunctionRef<bool(AActor* LevelActor)> InFunc)
+	bool ShouldIgnoreActor(AActor* InActor)
 	{
-		if (!PCGSubsystem || !PCGSubsystem->GetWorld())
+		if(!InActor)
 		{
-			return;
+			return true;
 		}
+				
+		return InActor->bIsEditorPreviewActor || InActor->IsA<ALevelInstanceEditorInstanceActor>() || InActor->IsA<APCGWorldActor>() || InActor->Implements<ULevelInstanceEditorPivotInterface>();
+	}
 
-		if (ILevelInstanceInterface* LevelInstance = Cast<ILevelInstanceInterface>(InActor))
+	void PropagateToLevelInstanceActors(ILevelInstanceInterface* InLevelInstance, TFunctionRef<bool(AActor* LevelActor)> InFunc)
+	{
+		if (InLevelInstance)
 		{
-			if (const ULevelInstanceSubsystem* LevelInstanceSubsystem = PCGSubsystem->GetWorld()->GetSubsystem<ULevelInstanceSubsystem>())
+			if (const ULevelInstanceSubsystem* LevelInstanceSubsystem = InLevelInstance->GetLevelInstanceSubsystem())
 			{
-				LevelInstanceSubsystem->ForEachActorInLevelInstance(LevelInstance, InFunc);
+				LevelInstanceSubsystem->ForEachActorInLevelInstance(InLevelInstance, [InFunc](AActor* LevelActor)
+				{
+					if (!ShouldIgnoreActor(LevelActor))
+					{
+						return InFunc(LevelActor);
+					}
+					return true;
+				});
 			}
 		}
 	}
-
+	
+	void PropagateToLevelInstanceActors(AActor* InActor, TFunctionRef<bool(AActor* LevelActor)> InFunc)
+	{
+		PropagateToLevelInstanceActors(Cast<ILevelInstanceInterface>(InActor), InFunc);
+	}
+		
 	bool ShouldDiscardLandscapeRefresh(const ALandscapeProxy* InLandscape, bool& bIsInEditingMode, bool bIsExitingEditingMode)
 	{
 		bIsInEditingMode = false;
@@ -167,7 +185,7 @@ void FPCGActorAndComponentMapping::Tick()
 	}
 
 #if WITH_EDITOR
-	AddDelayedActors();
+	ProcessDelayedEvents();
 
 	const double CurrentTime = FApp::GetCurrentTime();
 
@@ -1146,6 +1164,11 @@ void FPCGActorAndComponentMapping::RegisterTrackingCallbacks()
 	{
 		World->PersistentLevel->OnLoadedActorAddedToLevelEvent.AddRaw(this, &FPCGActorAndComponentMapping::OnActorLoaded);
 		World->PersistentLevel->OnLoadedActorRemovedFromLevelEvent.AddRaw(this, &FPCGActorAndComponentMapping::OnActorUnloaded);
+
+		if (ULevelInstanceSubsystem* LevelInstanceSubsystem = UWorld::GetSubsystem<ULevelInstanceSubsystem>(World))
+		{
+			LevelInstanceSubsystem->OnLevelInstancesUpdated().AddRaw(this, &FPCGActorAndComponentMapping::OnLevelInstancesUpdated);
+		}
 	}
 }
 
@@ -1163,125 +1186,147 @@ void FPCGActorAndComponentMapping::TeardownTrackingCallbacks()
 	{
 		World->PersistentLevel->OnLoadedActorAddedToLevelEvent.RemoveAll(this);
 		World->PersistentLevel->OnLoadedActorRemovedFromLevelEvent.RemoveAll(this);
+
+		if (ULevelInstanceSubsystem* LevelInstanceSubsystem = UWorld::GetSubsystem<ULevelInstanceSubsystem>(World))
+		{
+			LevelInstanceSubsystem->OnLevelInstancesUpdated().RemoveAll(this);
+		}
 	}
 }
 
-void FPCGActorAndComponentMapping::AddDelayedActors()
+bool FPCGActorAndComponentMapping::ShouldDelayActor(AActor* InActor) const
+{
+	const bool bDisableDelayedActorRegistering = PCGActorAndComponentMapping::CVarDisableDelayedActorRegistering.GetValueOnAnyThread();
+	if (!InActor->HasActorRegisteredAllComponents() && !bDisableDelayedActorRegistering)
+	{
+		return true;
+	}
+
+	bool bShouldDelay = false;
+
+	// Check that whole Level Instance hierarchy is loaded
+	if (ILevelInstanceInterface* LevelInstance = Cast<ILevelInstanceInterface>(InActor))
+	{
+		// If Level Instance isn't loaded yet it should be delayed
+		if (!LevelInstance->IsLoaded() && !InActor->IsA<APackedLevelActor>())
+		{
+			return true;
+		}
+			
+		PCGActorAndComponentMapping::PropagateToLevelInstanceActors(LevelInstance, [this, &bShouldDelay](AActor* InChildActor)
+		{
+			if(ShouldDelayActor(InChildActor))
+			{
+				bShouldDelay = true;
+				return false;
+			}
+			return true;
+		});
+	}
+
+	return bShouldDelay;
+}
+
+void FPCGActorAndComponentMapping::ProcessDelayedEvents()
 {
 	// Safeguard, we can't add delayed actors if the subsystem is not initialized
-	if (!PCGSubsystem || !PCGSubsystem->IsInitialized() || DelayedAddedActors.IsEmpty())
+	if (!PCGSubsystem || !PCGSubsystem->IsInitialized() || (DelayedAddedActors.IsEmpty() && DelayedChangedActors.IsEmpty()))
 	{
 		return;
 	}
 
-	TMap<TObjectKey<AActor>, TTuple<bool, int>> StillDelayedActors;
-	const bool bDisableDelayedActorRegistering = PCGActorAndComponentMapping::CVarDisableDelayedActorRegistering.GetValueOnAnyThread();
-
-	for (const TPair<TObjectKey<AActor>, TTuple<bool, int>>& ActorPtrAndShouldDirty : DelayedAddedActors)
+	TMap<TObjectKey<AActor>, bool> LocalDelayedAddedActors = MoveTemp(DelayedAddedActors);
+	
+	for (const TPair<TObjectKey<AActor>, bool>& ActorPtrAndShouldDirty : LocalDelayedAddedActors)
 	{
 		AActor* Actor = ActorPtrAndShouldDirty.Key.ResolveObjectPtr();
 		if (!Actor)
 		{
 			continue;
 		}
+		
+		OnActorAdded_Internal(Actor, ActorPtrAndShouldDirty.Value);
+	}
+	
+	TMap<TObjectKey<AActor>, TObjectKey<UObject>> LocalDelayedChangedActors = MoveTemp(DelayedChangedActors);
 
-		if (!Actor->HasActorRegisteredAllComponents() && !bDisableDelayedActorRegistering)
-		{
-			StillDelayedActors.Add(ActorPtrAndShouldDirty);
-		}
-		else
+	for (const TPair<TObjectKey<AActor>, TObjectKey<UObject>>& ActorToObject : LocalDelayedChangedActors)
+	{
+		AActor* Actor = ActorToObject.Key.ResolveObjectPtr();
+		UObject* Object = ActorToObject.Value.ResolveObjectPtr();
+		if (!Actor || !Object)
 		{
-			// Implementation note: since the delayed actors list is built from top-level actors (e.g. directly in the level) then depth here is 0.
-			OnActorAdded_Internal(Actor, ActorPtrAndShouldDirty.Value.Get<0>(), ActorPtrAndShouldDirty.Value.Get<1>(), /*bForceAddDelayedActor=*/true);
+			continue;
 		}
+
+		OnActorChanged_Internal(Actor, Object);
 	}
+}
 
-	DelayedAddedActors = MoveTemp(StillDelayedActors);
+void FPCGActorAndComponentMapping::OnLevelInstancesUpdated(const TArray<ILevelInstanceInterface*>& InLevelInstances)
+{
+	for(ILevelInstanceInterface* LevelInstance : InLevelInstances)
+	{
+		OnActorAdded(Cast<AActor>(LevelInstance));
+	}
 }
 
 void FPCGActorAndComponentMapping::OnActorLoaded(AActor& InActor)
 {
 	// We have to make sure to not create a infinite loop
-	if (InActor.IsA<APCGWorldActor>() || !PCGSubsystem || InActor.GetWorld() != PCGSubsystem->GetWorld())
+	if (PCGActorAndComponentMapping::ShouldIgnoreActor(&InActor) || !PCGSubsystem || InActor.GetWorld() != PCGSubsystem->GetWorld())
 	{
 		return;
 	}
 
 	// Loaded actors should not dirty.
-	OnActorAdded_Internal(&InActor, /*bShouldDirty=*/ false, /*LevelInstanceDepth=*/ 0);
+	OnActorAdded_Internal(&InActor, /*bShouldDirty=*/ false);
 }
 
 void FPCGActorAndComponentMapping::OnActorAdded(AActor* InActor)
 {
 	// We have to make sure to not create a infinite loop
-	if (!InActor || InActor->IsA<APCGWorldActor>() || !PCGSubsystem || InActor->GetWorld() != PCGSubsystem->GetWorld())
-	{
-		return;
-	}
-
-	int LevelInstanceDepth = 0;
-
-#if WITH_EDITOR
-	if (InActor->bIsEditorPreviewActor)
+	if (PCGActorAndComponentMapping::ShouldIgnoreActor(InActor) || !PCGSubsystem || InActor->GetWorld() != PCGSubsystem->GetWorld())
 	{
 		return;
 	}
+			
+	// Implementation note: We delay adding because OnActorAdded fires before an actor's properties are set,
+	// so the actor is not ready for processing until the next tick.
+	DelayedAddedActors.Emplace(InActor, true);
+}
 
-	bool bShouldDirty = true;
+void FPCGActorAndComponentMapping::OnActorChanged_Recursive(AActor* InActor, UObject* InOriginatingChangeObject)
+{
+	check(!DelayedAddedActors.Contains(InActor) && !DelayedChangedActors.Contains(InActor));
 
-	// In case of Level Instances, we first need to catch LevelInstance actors. If we have a LevelInstance actor first,
-	// it means we are currently ADDING a new level instance to the level. It is then followed by the creation of a ALevelInstanceEditorInstanceActor.
-	// It also means it needs to be dirtied.
-	// On the other hand, if we have no LevelInstance actor added, it means the LevelInstance is LOADED, and so we should not dirty the level instance.
-	// To future people reading this code: If this doesn't work anymore, verify that the assumption is still valid.
-	if (ALevelInstance* LevelInstanceActor = Cast<ALevelInstance>(InActor))
-	{
-		TempAddedLevelInstances.Add(LevelInstanceActor);
-	}
-	else if (ALevelInstanceEditorInstanceActor* LevelInstanceEditorInstance = Cast<ALevelInstanceEditorInstanceActor>(InActor))
+	// Process actors recursively
+	auto OnActorChanged = [this, InOriginatingChangeObject](AActor* InActor, auto OnActorChangedRecursive) -> void
 	{
-		bShouldDirty = false;
-
-		const ULevelInstanceSubsystem* LevelInstanceSubsystem = PCGSubsystem->GetWorld() ? PCGSubsystem->GetWorld()->GetSubsystem<ULevelInstanceSubsystem>() : nullptr;
-		const ILevelInstanceInterface* ActorLevelInstance = LevelInstanceSubsystem ? LevelInstanceSubsystem->GetParentLevelInstance(InActor) : nullptr;
-
-		// Check if the Level instance actor was added, if so we need to dirty.
-		if (ActorLevelInstance && TempAddedLevelInstances.Contains(ActorLevelInstance))
+		if (InActor)
 		{
-			TempAddedLevelInstances.Remove(ActorLevelInstance);
-			bShouldDirty = true;
-
-			// Also look for child level instances, because the level was not already loaded when the ALevelInstance was added.
-			LevelInstanceSubsystem->ForEachActorInLevelInstance(ActorLevelInstance, [this](AActor* LevelActor)
+			FActorPreviousData* PreviousData = ActorToPreviousDataMap.Find(InActor);
+			OnObjectChanged(InActor, PreviousData, InOriginatingChangeObject);
+			if (PreviousData)
 			{
-				if (ALevelInstance* LevelInstanceActor = Cast<ALevelInstance>(LevelActor))
-				{
-					TempAddedLevelInstances.Add(LevelInstanceActor);
-				}
-
-				return true;
-			});
+				ActorToPreviousDataMap.Remove(InActor);
+			}
 		}
-
-		// We also need to compute the depth.
-		while (ActorLevelInstance)
+		
+		PCGActorAndComponentMapping::PropagateToLevelInstanceActors(InActor, [this, OnActorChangedRecursive](AActor* LevelActor)
 		{
-			++LevelInstanceDepth;
-			ActorLevelInstance = LevelInstanceSubsystem->GetParentLevelInstance(CastChecked<AActor>(ActorLevelInstance));
-		}
-	}
-#endif // WITH_EDITOR
+			OnActorChangedRecursive(LevelActor, OnActorChangedRecursive);
+			return true;
+		});
+	};
 
-	// Implementation note: since this is called only for actors directly in the current level, the depth here is 0.
-	// Another implementation note: We delay adding because OnActorAdded fires before an actor's properties are set,
-	// so the actor is not ready for processing until the next tick.
-	DelayedAddedActors.Emplace(InActor, { bShouldDirty, LevelInstanceDepth });
+	OnActorChanged(InActor, OnActorChanged);
 }
 
-void FPCGActorAndComponentMapping::OnActorAdded_Internal(AActor* InActor, bool bShouldDirty, int32 LevelInstanceDepth, bool bForceAddDelayedActor)
+void FPCGActorAndComponentMapping::OnActorAdded_Internal(AActor* InActor, bool bShouldDirty)
 {
 	TRACE_CPUPROFILER_EVENT_SCOPE(FPCGActorAndComponentMapping::OnActorAdded);
-	check(InActor && !InActor->IsA<APCGWorldActor>() && PCGSubsystem && InActor->GetWorld() == PCGSubsystem->GetWorld());
+	check(!PCGActorAndComponentMapping::ShouldIgnoreActor(InActor) && PCGSubsystem && InActor->GetWorld() == PCGSubsystem->GetWorld());
 
 	if (ALandscapeProxy* LandscapeProxy = Cast<ALandscapeProxy>(InActor))
 	{
@@ -1297,94 +1342,66 @@ void FPCGActorAndComponentMapping::OnActorAdded_Internal(AActor* InActor, bool b
 		return;
 	}
 
-	// A delayed actor should only be added from AddDelayedActors(), because that guarantees the new actor has waited at
-	// least one tick for its properties to be popualated.
-	if (!bForceAddDelayedActor && DelayedAddedActors.Contains(InActor))
+	// Currently delayed, return
+	if (DelayedAddedActors.Contains(InActor))
 	{
 		return;
 	}
 
-	// If the subsystem is not initialized, wait for it to be, and store all the actors to check
-	if (!PCGSubsystem->IsInitialized())
+	// Should still be delayed, add to delayed actors, return
+	if (ShouldDelayActor(InActor))
 	{
-		DelayedAddedActors.Emplace(InActor, { bShouldDirty, LevelInstanceDepth });
+		DelayedAddedActors.Add(InActor, bShouldDirty);
 		return;
 	}
+		
+	// Call Recursive method
+	OnActorChanged_Recursive(InActor, nullptr);
+}
 
-#if WITH_EDITOR
-	// When a level instance is added to the level (or loaded), the Level Instance Editor Instance will be spawned
-	// but when we get the callback, it is not initialized yet (and does not contain the actors) hence adding it to the delayed actors.
-	// Note that while conceptually we'd want a callback when a level instance is loaded, it's a bit tricky vs. what world we want to do our changes in.
-	if (ALevelInstanceEditorInstanceActor* LevelInstanceEditorInstance = Cast<ALevelInstanceEditorInstanceActor>(InActor))
+void FPCGActorAndComponentMapping::OnActorChanged_Internal(AActor* InActor, UObject* InOriginatingChangeObject)
+{
+	// Do not process actor changes if it hasn't been added yet
+	if (DelayedAddedActors.Contains(InActor))
 	{
-		const FLevelInstanceID& LevelInstanceId = LevelInstanceEditorInstance->GetLevelInstanceID();
-		const ULevelInstanceSubsystem* LevelInstanceSubsystem = PCGSubsystem->GetWorld() ? PCGSubsystem->GetWorld()->GetSubsystem<ULevelInstanceSubsystem>() : nullptr;
-
-		if (LevelInstanceId.IsValid() && LevelInstanceSubsystem)
-		{
-			if (ILevelInstanceInterface* LevelInstance = LevelInstanceSubsystem->GetLevelInstance(LevelInstanceId))
-			{
-				LevelInstanceSubsystem->ForEachActorInLevelInstance(LevelInstance, [this, bShouldDirty, LevelInstanceDepth, LevelInstanceEditorInstance](AActor* LevelActor)
-				{
-					if (LevelActor && LevelActor != LevelInstanceEditorInstance && !LevelActor->IsA<APCGWorldActor>())
-					{
-						OnActorAdded_Internal(LevelActor, bShouldDirty, LevelInstanceDepth + 1);
-					}
-					
-					return true;
-				});
-			}
-		}
-		else
-		{
-			DelayedAddedActors.Emplace(InActor, { bShouldDirty, LevelInstanceDepth });
-			return;
-		}
+		return;
+	}
+	
+	// Currently delayed, return
+	if (DelayedChangedActors.Contains(InActor))
+	{
+		return;
 	}
-#endif // WITH_EDITOR
 
-	// Finally notify them all
-	OnObjectChanged(InActor, /*InPreviousData=*/nullptr, /*InOriginatingObject=*/nullptr, LevelInstanceDepth);
+	// Should still be delayed, add to delayed actors, return
+	if (ShouldDelayActor(InActor))
+	{
+		DelayedChangedActors.Add(InActor, InOriginatingChangeObject);
+		return;
+	}
+	
+	// Call Recursive method
+	OnActorChanged_Recursive(InActor, InOriginatingChangeObject);
 }
 
 void FPCGActorAndComponentMapping::OnActorUnloaded(AActor& InActor)
 {
 	// Don't dirty on unload (to mirror the behavior in load)
-	OnActorDeleted_Internal(&InActor, /*bShouldDirty=*/false, /*LevelInstanceDepth=*/0);
+	OnActorDeleted_Internal(&InActor, /*bShouldDirty=*/false);
 }
 
 void FPCGActorAndComponentMapping::OnActorDeleted(AActor* InActor)
 {
-	if (!InActor || !PCGSubsystem || InActor->GetWorld() != PCGSubsystem->GetWorld())
-	{
-		return;
-	}
-
-	int LevelInstanceDepth = 0;
-
-#if WITH_EDITOR
-	if (InActor->bIsEditorPreviewActor)
+	if (PCGActorAndComponentMapping::ShouldIgnoreActor(InActor) || !PCGSubsystem || InActor->GetWorld() != PCGSubsystem->GetWorld())
 	{
 		return;
 	}
 
-	if (ALevelInstanceEditorInstanceActor* LevelInstanceEditorInstance = Cast<ALevelInstanceEditorInstanceActor>(InActor))
-	{
-		const ULevelInstanceSubsystem* LevelInstanceSubsystem = PCGSubsystem->GetWorld() ? PCGSubsystem->GetWorld()->GetSubsystem<ULevelInstanceSubsystem>() : nullptr;
-		const ILevelInstanceInterface* ActorLevelInstance = LevelInstanceSubsystem ? LevelInstanceSubsystem->GetParentLevelInstance(InActor) : nullptr;
-		while (ActorLevelInstance)
-		{
-			++LevelInstanceDepth;
-			ActorLevelInstance = LevelInstanceSubsystem->GetParentLevelInstance(CastChecked<AActor>(ActorLevelInstance));
-		}
-	}
-#endif
-
 	// Implementation note: since this is called only for actors directly in the current level, the depth here is 0.
-	OnActorDeleted_Internal(InActor, /*bShouldDirty=*/true, LevelInstanceDepth);
+	OnActorDeleted_Internal(InActor, /*bShouldDirty=*/true);
 }
 
-void FPCGActorAndComponentMapping::OnActorDeleted_Internal(AActor* InActor, bool bShouldDirty, int32 LevelInstanceDepth)
+void FPCGActorAndComponentMapping::OnActorDeleted_Internal(AActor* InActor, bool bShouldDirty)
 {
 	TRACE_CPUPROFILER_EVENT_SCOPE(UPCGActorAndComponentMapping::OnActorDeleted);
 	check(InActor && PCGSubsystem && InActor->GetWorld() == PCGSubsystem->GetWorld());
@@ -1401,38 +1418,15 @@ void FPCGActorAndComponentMapping::OnActorDeleted_Internal(AActor* InActor, bool
 	}
 
 #if WITH_EDITOR
-	// When editing/deleting/replacing a level instance, this will allow us to first remove the previously tracked actors that were under the level instance
-	if (ALevelInstanceEditorInstanceActor* LevelInstanceEditorInstance = Cast<ALevelInstanceEditorInstanceActor>(InActor))
-	{
-		const FLevelInstanceID& LevelInstanceId = LevelInstanceEditorInstance->GetLevelInstanceID();
-		const ULevelInstanceSubsystem* LevelInstanceSubsystem = PCGSubsystem->GetWorld() ? PCGSubsystem->GetWorld()->GetSubsystem<ULevelInstanceSubsystem>() : nullptr;
-
-		if (LevelInstanceId.IsValid() && LevelInstanceSubsystem)
-		{
-			if (ILevelInstanceInterface* LevelInstance = LevelInstanceSubsystem->GetLevelInstance(LevelInstanceId))
-			{
-				LevelInstanceSubsystem->ForEachActorInLevelInstance(LevelInstance, [this, LevelInstanceDepth, LevelInstanceEditorInstance, bShouldDirty](AActor* LevelActor)
-				{
-					if (LevelActor != LevelInstanceEditorInstance)
-					{
-						OnActorDeleted_Internal(LevelActor, bShouldDirty, LevelInstanceDepth + 1);
-					}
-
-					return true;
-				});
-			}
-		}
-	}
-
-	PCGActorAndComponentMapping::PropagateToLevelInstanceActors(InActor, PCGSubsystem, [this, LevelInstanceDepth, bShouldDirty](AActor* LevelActor)
+	PCGActorAndComponentMapping::PropagateToLevelInstanceActors(InActor, [this, bShouldDirty](AActor* LevelActor)
 	{
-		OnActorDeleted_Internal(LevelActor, bShouldDirty, LevelInstanceDepth + 1);
+		OnActorDeleted_Internal(LevelActor, bShouldDirty);
 		return true;
 	});
 #endif // WITH_EDITOR
 
 	// Notify all components that the actor has changed (was removed), but the Refresh will only happen AFTER the actor was actually removed from the world (because of delayed refresh).
-	OnObjectChanged(InActor, /*InPreviousData=*/nullptr, /*InOriginatingObject=*/nullptr, LevelInstanceDepth, /*bNoRefreshOnOwner=*/true);
+	OnObjectChanged(InActor, /*InPreviousData=*/nullptr, /*InOriginatingObject=*/nullptr, /*bNoRefreshOnOwner=*/true);
 }
 
 void FPCGActorAndComponentMapping::OnObjectModified(UObject* InObject)
@@ -1462,36 +1456,34 @@ void FPCGActorAndComponentMapping::OnObjectModified(UObject* InObject)
 	}
 
 #if WITH_EDITOR
-	if (Actor->bIsEditorPreviewActor)
+	if (PCGActorAndComponentMapping::ShouldIgnoreActor(Actor))
 	{
 		return;
 	}
 #endif
 
-	auto StorePreviousData = [this](AActor* InActor, int32 LevelInstanceDepth, auto RecursiveCall) -> void
+	auto StorePreviousData = [this](AActor* InActor, auto RecursiveCall) -> void
 	{
 		if (!ActorToPreviousDataMap.Contains(InActor))
 		{
-			if (!IsActorTracked(InActor))
+			if (IsActorTracked(InActor))
 			{
-				return;
+				FActorPreviousData& PreviousData = ActorToPreviousDataMap.Add(InActor);
+				PreviousData.Get<0>() = PCGActorAndComponentMapping::GetActorBounds(InActor);
+				PreviousData.Get<1>() = TSet<FName>(InActor->Tags);
+				PreviousData.Get<2>() = FApp::GetCurrentTime();
 			}
-
-			FActorPreviousData& PreviousData = ActorToPreviousDataMap.Add(InActor);
-			PreviousData.Get<0>() = PCGActorAndComponentMapping::GetActorBounds(InActor);
-			PreviousData.Get<1>() = TSet<FName>(InActor->Tags);
-			PreviousData.Get<2>() = FApp::GetCurrentTime();
-
+					
 			// Also propagate the pre-change to all child actors if it is within a level instance.
-			PCGActorAndComponentMapping::PropagateToLevelInstanceActors(InActor, PCGSubsystem, [this, LevelInstanceDepth, RecursiveCall](AActor* LevelActor)
+			PCGActorAndComponentMapping::PropagateToLevelInstanceActors(InActor, [this, RecursiveCall](AActor* LevelActor)
 			{
-				RecursiveCall(LevelActor, LevelInstanceDepth + 1, RecursiveCall);
+				RecursiveCall(LevelActor, RecursiveCall);
 				return true;
 			});
 		}
 	};
 
-	StorePreviousData(Actor, /*LevelInstanceDepth*/ 0, StorePreviousData);
+	StorePreviousData(Actor, StorePreviousData);
 }
 
 void FPCGActorAndComponentMapping::OnObjectSaved(UObject* InObject, FObjectPreSaveContext InObjectSaveContext)
@@ -1551,12 +1543,10 @@ void FPCGActorAndComponentMapping::OnObjectPropertyChanged(UObject* InObject, FP
 		return;
 	}
 
-#if WITH_EDITOR
-	if (Actor && Actor->bIsEditorPreviewActor)
+	if (Actor && PCGActorAndComponentMapping::ShouldIgnoreActor(Actor))
 	{
 		return;
 	}
-#endif
 
 	const bool bNoOperation = (!bValueNotInteractive && !bActorTagChange)
 		|| bIsTextureCompilationResult
@@ -1570,7 +1560,7 @@ void FPCGActorAndComponentMapping::OnObjectPropertyChanged(UObject* InObject, FP
 			auto RemoveAll = [this](AActor* InActor, auto RecursiveCall) -> void
 			{
 				ActorToPreviousDataMap.Remove(InActor);
-				PCGActorAndComponentMapping::PropagateToLevelInstanceActors(InActor, PCGSubsystem, [this, RecursiveCall](AActor* LevelActor)
+				PCGActorAndComponentMapping::PropagateToLevelInstanceActors(InActor, [this, RecursiveCall](AActor* LevelActor)
 				{
 					RecursiveCall(LevelActor, RecursiveCall);
 					return true;
@@ -1585,26 +1575,7 @@ void FPCGActorAndComponentMapping::OnObjectPropertyChanged(UObject* InObject, FP
 
 	if (Actor)
 	{
-		auto OnActorChanged = [this, InObject](AActor* InActor, int32 LevelInstanceDepth, auto RecursiveCall) -> void
-		{
-			if (InActor)
-			{
-				FActorPreviousData* PreviousData = ActorToPreviousDataMap.Find(InActor);
-				OnObjectChanged(InActor, PreviousData, /*InOriginatingChangeObject=*/ InObject, LevelInstanceDepth);
-				if (PreviousData)
-				{
-					ActorToPreviousDataMap.Remove(InActor);
-				}
-			}
-
-			PCGActorAndComponentMapping::PropagateToLevelInstanceActors(InActor, PCGSubsystem, [this, LevelInstanceDepth, RecursiveCall](AActor* LevelActor)
-			{
-				RecursiveCall(LevelActor, LevelInstanceDepth + 1, RecursiveCall);
-				return true;
-			});
-		};
-
-		OnActorChanged(Actor, /*LevelInstanceDepth=*/0, OnActorChanged);
+		OnActorChanged_Internal(Actor, /*InOriginatingChangeObject=*/ InObject);
 	}
 	else
 	{
@@ -1612,7 +1583,7 @@ void FPCGActorAndComponentMapping::OnObjectPropertyChanged(UObject* InObject, FP
 	}
 }
 
-void FPCGActorAndComponentMapping::OnObjectChanged(UObject* InObject, const FActorPreviousData* InPreviousData, const UObject* InOriginatingChangeObject, int32 LevelInstanceDepth, bool bNoRefreshOwner)
+void FPCGActorAndComponentMapping::OnObjectChanged(UObject* InObject, const FActorPreviousData* InPreviousData, const UObject* InOriginatingChangeObject, bool bNoRefreshOwner)
 {
 	TRACE_CPUPROFILER_EVENT_SCOPE(FPCGActorAndComponentMapping::OnObjectChanged);
 
@@ -1633,7 +1604,7 @@ void FPCGActorAndComponentMapping::OnObjectChanged(UObject* InObject, const FAct
 		return;
 	}
 
-	AActor* Actor = Cast<AActor>(InObject);
+	const AActor* Actor = Cast<AActor>(InObject);
 
 	ensure(!PCGSubsystem || !Actor || Actor->GetWorld() == PCGSubsystem->GetWorld());
 
@@ -1883,20 +1854,46 @@ void FPCGActorAndComponentMapping::OnObjectChanged(UObject* InObject, const FAct
 			// PCG components higher-up are properly updated, hence the level instance depth
 			if (LevelInstanceSubsystem && Actor)
 			{
+				// Immediate Level Instance Parent (Actor)
 				const ILevelInstanceInterface* ActorLevelInstance = LevelInstanceSubsystem->GetParentLevelInstance(Actor);
-				int LocalDepth = LevelInstanceDepth;
-				while (ActorLevelInstance && LocalDepth-- > 0)
-				{
-					ActorLevelInstance = LevelInstanceSubsystem->GetParentLevelInstance(CastChecked<AActor>(ActorLevelInstance));
-				}
 
+				// Immediate Level Instance Parent (PCG Component)
 				const ILevelInstanceInterface* ComponentLevelInstance = LevelInstanceSubsystem->GetParentLevelInstance(Component->GetOwner());
-				while (ComponentLevelInstance && ComponentLevelInstance != ActorLevelInstance)
+								
+				// Is the Actor under a Level Instance that is the same or under the Components owning Level Instance
+				bool bActorLevelInstanceInComponentLevelInstanceHierarchy = false;
+
+				// Go up Level Instance hierarchy for Actor and see if we are part of the Component's level instance hierarchy, stop if we find an Editing Level instance as we don't want to impact top level 
+				// Components from an editing level instance
+				if (ComponentLevelInstance)
 				{
-					ComponentLevelInstance = LevelInstanceSubsystem->GetParentLevelInstance(CastChecked<AActor>(ComponentLevelInstance));
-				}
+					LevelInstanceSubsystem->ForEachLevelInstanceAncestors(Actor, [ComponentLevelInstance, &bActorLevelInstanceInComponentLevelInstanceHierarchy](const ILevelInstanceInterface* Ancestor)
+					{
+						if (Ancestor == ComponentLevelInstance)
+						{
+							bActorLevelInstanceInComponentLevelInstanceHierarchy = true;
+							return false; // stop iterating
+						}
 
-				if (ActorLevelInstance != ComponentLevelInstance)
+						return !Ancestor->IsEditing();
+					});
+				}
+				
+				// Actor in same level instance as Component or under the Components Level instance
+				if (bActorLevelInstanceInComponentLevelInstanceHierarchy)
+				{
+					// Allow update if Component is inside Editing level instance or if it is only in preview mode (generates transient data because it is in a non editing Level instance)
+					if (!ComponentLevelInstance->IsEditing() && Component->GetEditingMode() != EPCGEditorDirtyMode::Preview)
+					{
+						continue;
+					}
+				}
+				else if (Actor->IsInEditLevelInstanceHierarchy())
+				{
+					continue;
+				}
+				// If ComponentLevelInstance is non-null and Actor Level instance isn't in its hierarchy then skip it
+				else if (ActorLevelInstance && ComponentLevelInstance != nullptr)
 				{
 					continue;
 				}
@@ -2091,4 +2088,4 @@ void FPCGActorAndComponentMapping::NotifyLandscapeEditModeExited()
 
 	bIsCurrentlyExitingLandscapeEditMode = false;
 }
-#endif // WITH_EDITOR
\ No newline at end of file
+#endif // WITH_EDITOR
diff --git a/Engine/Plugins/PCG/Source/PCG/Private/PCGSubsystem.cpp b/Engine/Plugins/PCG/Source/PCG/Private/PCGSubsystem.cpp
--- a/Engine/Plugins/PCG/Source/PCG/Private/PCGSubsystem.cpp
+++ b/Engine/Plugins/PCG/Source/PCG/Private/PCGSubsystem.cpp
@@ -116,12 +116,6 @@ UPCGSubsystem::UPCGSubsystem()
 	: Super()
 	, ActorAndComponentMapping(this)
 {
-#if WITH_EDITOR
-	if (!HasAnyFlags(RF_ClassDefaultObject))
-	{
-		ActorAndComponentMapping.RegisterTrackingCallbacks();
-	}
-#endif // WITH_EDITOR
 }
 
 UPCGSubsystem* UPCGSubsystem::GetSubsystemForCurrentWorld()
@@ -194,6 +188,10 @@ void UPCGSubsystem::PostInitialize()
 			});
 		}
 	}
+
+#if WITH_EDITOR
+	ActorAndComponentMapping.RegisterTrackingCallbacks();
+#endif // WITH_EDITOR
 }
 
 UPCGSubsystem* UPCGSubsystem::GetInstance(UWorld* World)
diff --git a/Engine/Plugins/PCG/Source/PCG/Public/PCGActorAndComponentMapping.h b/Engine/Plugins/PCG/Source/PCG/Public/PCGActorAndComponentMapping.h
--- a/Engine/Plugins/PCG/Source/PCG/Public/PCGActorAndComponentMapping.h
+++ b/Engine/Plugins/PCG/Source/PCG/Public/PCGActorAndComponentMapping.h
@@ -47,17 +47,6 @@ public:
 	void Tick();
 
 #if WITH_EDITOR
-	/** If the partition grid size change, call this to empty the Partition actors map */
-	void ResetPartitionActorsMap();
-
-	void RegisterTrackingCallbacks();
-	void TeardownTrackingCallbacks();
-
-	void AddDelayedActors();
-
-	void RegisterTracking(UPCGComponent* InComponent);
-	void UpdateTracking(UPCGComponent* InComponent, bool bInShouldDirtyActors, const TArray<FPCGSelectionKey>* OptionalChangedKeys = nullptr);
-
 	/** Notify that we exited the Landscape edit mode. */
 	PCG_API void NotifyLandscapeEditModeExited();
 #endif // WITH_EDITOR
@@ -93,6 +82,21 @@ public:
 	APCGPartitionActor* GetPartitionActor(uint32 GridSize, const FIntVector& CellCoords, bool bRuntimeGenerated = false) const;
 
 private:
+
+#if WITH_EDITOR
+	/** If the partition grid size change, call this to empty the Partition actors map */
+	void ResetPartitionActorsMap();
+
+	void RegisterTrackingCallbacks();
+	void TeardownTrackingCallbacks();
+
+	void RegisterTracking(UPCGComponent* InComponent);
+	void UpdateTracking(UPCGComponent* InComponent, bool bInShouldDirtyActors, const TArray<FPCGSelectionKey>* OptionalChangedKeys = nullptr);
+
+	void ProcessDelayedEvents();
+	bool ShouldDelayActor(AActor* InActor) const;
+#endif
+
 	// This class is only meant to be used as part of the PCG Subsytem and owned by it.
 	// So we put constructors private.
 	// We also need this class to be default constructible, since the PCGSubsytem needs to be default constructible.
@@ -147,11 +151,14 @@ private:
 
 	void OnActorAdded(AActor* InActor);
 	void OnActorLoaded(AActor& InActor);
-	void OnActorAdded_Internal(AActor* InActor, bool bShouldDirty, int32 LevelInstanceDepth, bool bForceAddDelayedActor = false);
+	void OnActorAdded_Internal(AActor* InActor, bool bShouldDirty);
+	void OnActorChanged_Internal(AActor* InActor, UObject* InOriginatingChangeObject);
+	void OnActorChanged_Recursive(AActor* InActor, UObject* InOriginatingChangeObject);
 	void OnActorDeleted(AActor* InActor);
 	void OnActorUnloaded(AActor& InActor);
-	void OnActorDeleted_Internal(AActor* InActor, bool bShouldDirty, int32 LevelInstanceDepth);
+	void OnActorDeleted_Internal(AActor* InActor, bool bShouldDirty);
 	void OnLandscapeChanged(ALandscapeProxy* InLandscape, const FLandscapeProxyComponentDataChangedParams& InChangeParams);
+	void OnLevelInstancesUpdated(const TArray<ILevelInstanceInterface*>& InLevelInstances);
 	void ApplyLandscapeChanges(ALandscapeProxy* InLandscape);
 	void OnObjectModified(UObject* InObject);
 	void OnObjectPropertyChanged(UObject* InObject, FPropertyChangedEvent& InEvent);
@@ -174,7 +181,7 @@ private:
 	* Can also specify an optional object, originating the change, to avoid re-dirtying a component if it was the origin.
 	* Another option when an actor is deleted/unload, don't refresh their components.
 	*/
-	void OnObjectChanged(UObject* InObject, const FActorPreviousData* InPreviousData = nullptr, const UObject* InOriginatingChangeObject = nullptr, int32 LevelInstanceDepth = 0, bool bNoRefreshOnOwner = false);
+	void OnObjectChanged(UObject* InObject, const FActorPreviousData* InPreviousData = nullptr, const UObject* InOriginatingChangeObject = nullptr, bool bNoRefreshOnOwner = false);
 
 	/** Gather all settings from a given component that track the key, and clear the cache for them. Returns true if we should dirty afterwards (aka at least one settings was cleared and/or landscape changed). */
 	bool ClearCacheForKeys(const TArray<FPCGSelectionKey>& InKeys, const UPCGComponent* InComponent, const bool bIntersect, const UObject* InOriginatingChange) const;
@@ -208,10 +215,11 @@ private:
 	TSet<UPCGComponent*> DelayedComponentToUnregister;
 	mutable FCriticalSection DelayedComponentToUnregisterLock;
 
-	// Tracking actors
 	/** Will hold all the components that are not partitioned (and not local) and are tracking something. Will be use to dispatch actor tracking updates. */
 	FPCGComponentOctreeAndMap NonPartitionedOctree;
 
+#if WITH_EDITOR
+	// Tracking actors
 	/** Keep a mapping between tracked keys and the components that track them, and the tracking needs to be culled.*/
 	TMap<FPCGSelectionKey, TSet<UPCGComponent*>> CulledTrackedKeysToComponentsMap;
 
@@ -220,11 +228,9 @@ private:
 
 	mutable FRWLock TrackedComponentsLock;
 
-	// Keep track of actors that aren't yet ready (or if the subsystem is not yet ready), whether we should dirty them and their instance level depth so we can add them in next tick.
-	TMap<TObjectKey<AActor>, TTuple<bool, int>> DelayedAddedActors;
-
-	// Keep track of all Level Instance actors added, so we can detect when a level instance is added vs loaded
-	TSet<ILevelInstanceInterface*> TempAddedLevelInstances;
+	// Keep track of actors that aren't yet ready (or if the subsystem is not yet ready), whether we should dirty them in next tick.
+	TMap<TObjectKey<AActor>, bool> DelayedAddedActors;
+	TMap<TObjectKey<AActor>, TObjectKey<UObject>> DelayedChangedActors;
 
 	/** Transient map of actors and their previous data, it's set in the pre object change to be able to track changes (such as tags or positions) */
 	TMap<TObjectKey<AActor>, FActorPreviousData> ActorToPreviousDataMap;
@@ -232,7 +238,6 @@ private:
 	/** Transient map that keep track of all components that depends on another component currently generating, to trigger the refresh only once, when all are done. */
 	TMap<TObjectKey<UPCGComponent>, TArray<TObjectKey<UPCGComponent>>> ComponentsToDependencyMap;
 
-#if WITH_EDITOR
 	// Part for the delayed landscape change update
 	double LastLandscapeDirtyTime = -1.0;
 	TArray<TObjectKey<ALandscapeProxy>> DelayedModifiedLandscapes;
diff --git a/Engine/Source/Runtime/Engine/Private/LevelInstance/LevelInstanceSubsystem.cpp b/Engine/Source/Runtime/Engine/Private/LevelInstance/LevelInstanceSubsystem.cpp
--- a/Engine/Source/Runtime/Engine/Private/LevelInstance/LevelInstanceSubsystem.cpp
+++ b/Engine/Source/Runtime/Engine/Private/LevelInstance/LevelInstanceSubsystem.cpp
@@ -2647,7 +2647,7 @@ bool ULevelInstanceSubsystem::CommitLevelInstanceInternal(TUniquePtr<FLevelInsta
 	UActorDescContainerSubsystem::GetChecked().NotifyContainerUpdated(WorldAssetPackage);
 
 	
-	TArray<TPair<ULevelInstanceSubsystem*, FLevelInstanceID>> LevelInstancesToUpdate;
+	TMap<ULevelInstanceSubsystem*, TArray<FLevelInstanceID>> LevelInstancesToUpdate;
 	// Gather list to update
 	for (TObjectIterator<UWorld> It(RF_ClassDefaultObject | RF_ArchetypeObject, true); It; ++It)
 	{
@@ -2661,7 +2661,7 @@ bool ULevelInstanceSubsystem::CommitLevelInstanceInternal(TUniquePtr<FLevelInsta
 				{
 					if (CurrentLevelInstance == LevelInstance || bChangesCommitted)
 					{
-						LevelInstancesToUpdate.Add({ LevelInstanceSubsystem, CurrentLevelInstance->GetLevelInstanceID() });
+						LevelInstancesToUpdate.FindOrAdd(LevelInstanceSubsystem).Add(CurrentLevelInstance->GetLevelInstanceID());
 					}
 				}
 			}
@@ -2671,9 +2671,12 @@ bool ULevelInstanceSubsystem::CommitLevelInstanceInternal(TUniquePtr<FLevelInsta
 	// Do update
 	for (const auto& KeyValuePair : LevelInstancesToUpdate)
 	{
-		if (ILevelInstanceInterface* LevelInstanceToUpdate = KeyValuePair.Key->GetLevelInstance(KeyValuePair.Value))
+		for (const FLevelInstanceID& LevelInstanceIDToUpdate : KeyValuePair.Value)
 		{
-			LevelInstanceToUpdate->UpdateLevelInstanceFromWorldAsset();
+			if (ILevelInstanceInterface* LevelInstanceToUpdate = KeyValuePair.Key->GetLevelInstance(LevelInstanceIDToUpdate))
+			{
+				LevelInstanceToUpdate->UpdateLevelInstanceFromWorldAsset();
+			}
 		}
 	}
 
@@ -2708,6 +2711,24 @@ bool ULevelInstanceSubsystem::CommitLevelInstanceInternal(TUniquePtr<FLevelInsta
 	if (bChangesCommitted)
 	{
 		LevelInstanceChangedEvent.Broadcast(WorldAssetPackage);
+
+		// Send event per World (per LevelInstanceSubsystem)
+		for (const auto& KeyValuePair : LevelInstancesToUpdate)
+		{
+			TArray<ILevelInstanceInterface*> UpdatedLevelInstances;
+			for (const FLevelInstanceID& UpdatedLevelInstanceID : KeyValuePair.Value)
+			{
+				if (ILevelInstanceInterface* UpdatedLevelInstance = KeyValuePair.Key->GetLevelInstance(UpdatedLevelInstanceID))
+				{
+					UpdatedLevelInstances.Add(UpdatedLevelInstance);
+				}
+			}
+
+			if (UpdatedLevelInstances.Num() > 0)
+			{
+				KeyValuePair.Key->LevelInstancesUpdatedEvent.Broadcast(UpdatedLevelInstances);
+			}
+		}
 	}
 
 	GEngine->BroadcastLevelActorListChanged();
diff --git a/Engine/Source/Runtime/Engine/Public/LevelInstance/LevelInstanceSubsystem.h b/Engine/Source/Runtime/Engine/Public/LevelInstance/LevelInstanceSubsystem.h
--- a/Engine/Source/Runtime/Engine/Public/LevelInstance/LevelInstanceSubsystem.h
+++ b/Engine/Source/Runtime/Engine/Public/LevelInstance/LevelInstanceSubsystem.h
@@ -161,6 +161,10 @@ public:
 	DECLARE_EVENT_OneParam(ULevelInstanceSubsystem, FLevelInstanceChanged, FName);
 	FLevelInstanceChanged& OnLevelInstanceChanged() { return LevelInstanceChangedEvent; }
 
+	/** Editor-only even triggered when level instances are reloaded after a change */
+	DECLARE_EVENT_OneParam(ULevelInstanceSubsystem, FLevelInstancesUpdated, const TArray<ILevelInstanceInterface*>&);
+	FLevelInstancesUpdated& OnLevelInstancesUpdated() { return LevelInstancesUpdatedEvent; }
+
 	static ENGINE_API void ResetLoadersForWorldAsset(const FString& WorldAsset);
 
 	ENGINE_API bool PassLevelInstanceFilter(UWorld* World, const FWorldPartitionHandle& Actor) const;
@@ -311,6 +315,8 @@ private:
 	bool bIsCommittingLevelInstance;
 
 	FLevelInstanceChanged LevelInstanceChangedEvent;
+	FLevelInstancesUpdated LevelInstancesUpdatedEvent;
+
 #endif
 
 	struct FLevelInstance
