Removed unneeded warnings/error logging when executing disabled nodes

#rnx

diff --git a/Engine/Plugins/PCG/Source/PCG/Private/Graph/PCGGraphExecutor.cpp b/Engine/Plugins/PCG/Source/PCG/Private/Graph/PCGGraphExecutor.cpp
--- a/Engine/Plugins/PCG/Source/PCG/Private/Graph/PCGGraphExecutor.cpp
+++ b/Engine/Plugins/PCG/Source/PCG/Private/Graph/PCGGraphExecutor.cpp
@@ -1265,6 +1265,13 @@ void FPCGGraphExecutor::BuildTaskInput(const FPCGGraphTask& Task, FPCGDataCollec
 	auto LogDiscardedData = [&Task](const UPCGPin* InPin)
 	{
 #if !(UE_BUILD_SHIPPING || UE_BUILD_TEST) || USE_LOGGING_IN_SHIPPING
+		// Turn off eventual errors/warnings when the node is disabled, as this is irrelevant.
+		const bool bNodeIsDisabled = (Task.Context && Task.Context->GetOriginalSettings<UPCGSettings>() && !Task.Context->GetOriginalSettings<UPCGSettings>()->bEnabled);
+		if (bNodeIsDisabled)
+		{
+			return;
+		}
+
 		const FString Message = FString::Printf(
 			TEXT("[%s] %s - BuildTaskInput - too many data items arriving on single data pin '%s', only first data item will be used"),
 			(Task.SourceComponent.Get() && Task.SourceComponent->GetOwner()) ? *Task.SourceComponent->GetOwner()->GetName() : TEXT("MissingComponent"),
diff --git a/Engine/Plugins/PCG/Source/PCG/Private/PCGContext.cpp b/Engine/Plugins/PCG/Source/PCG/Private/PCGContext.cpp
--- a/Engine/Plugins/PCG/Source/PCG/Private/PCGContext.cpp
+++ b/Engine/Plugins/PCG/Source/PCG/Private/PCGContext.cpp
@@ -106,6 +106,12 @@ void FPCGContext::InitializeSettings()
 
 	if (const UPCGSettings* NodeSettings = GetOriginalSettings<UPCGSettings>())
 	{
+		// Don't apply overrides if the original settings are disabled
+		if (!NodeSettings->bEnabled)
+		{
+			return;
+		}
+
 		// Only duplicate the settings if we have overriable params and we have at least one param pin connected.
 		const TArray<FPCGSettingsOverridableParam>& OverridableParams = NodeSettings->OverridableParams();
 		if (!OverridableParams.IsEmpty())
