Fixed profiling issues that were corrupting data in loops & dynamic subgraphs and revisited the way we compute the timings and their display.

diff --git a/Engine/Plugins/PCG/Source/PCG/Private/Graph/PCGGraphExecutor.cpp b/Engine/Plugins/PCG/Source/PCG/Private/Graph/PCGGraphExecutor.cpp
--- a/Engine/Plugins/PCG/Source/PCG/Private/Graph/PCGGraphExecutor.cpp
+++ b/Engine/Plugins/PCG/Source/PCG/Private/Graph/PCGGraphExecutor.cpp
@@ -824,7 +824,7 @@ void FPCGGraphExecutor::Execute()
 						if (Task.StackIndex != INDEX_NONE)
 						{
 							const FPCGStack* Stack = Task.GetStack();
-							SourceComponent->StoreInspectionData(Stack, Task.Node, TaskInput, CachedOutput);
+							SourceComponent->StoreInspectionData(Stack, Task.Node, TaskInput, CachedOutput, /*bUsedCache=*/true);
 						}
 					}
 #endif
@@ -1003,7 +1003,7 @@ void FPCGGraphExecutor::Execute()
 					if (ActiveTask.StackIndex != INDEX_NONE)
 					{
 						const FPCGStack* Stack = ActiveTask.StackContext->GetStack(ActiveTask.StackIndex);
-						SourceComponent->StoreInspectionData(Stack, ActiveTask.Context->Node, ActiveTask.Context->InputData, ActiveTask.Context->OutputData);
+						SourceComponent->StoreInspectionData(Stack, ActiveTask.Context->Node, ActiveTask.Context->InputData, ActiveTask.Context->OutputData, /*bUsedCache=*/false);
 					}
 				}
 			}
diff --git a/Engine/Plugins/PCG/Source/PCG/Private/PCGComponent.cpp b/Engine/Plugins/PCG/Source/PCG/Private/PCGComponent.cpp
--- a/Engine/Plugins/PCG/Source/PCG/Private/PCGComponent.cpp
+++ b/Engine/Plugins/PCG/Source/PCG/Private/PCGComponent.cpp
@@ -638,7 +638,6 @@ FPCGTaskId UPCGComponent::CleanupInternal(bool bRemoveComponents, const TArray<F
 	Modify(!IsInPreviewMode());
 
 #if WITH_EDITOR
-	ExtraCapture.ResetTimers();
 	ExtraCapture.ResetCapturedMessages();
 #endif
 
@@ -1974,15 +1973,29 @@ void UPCGComponent::DisableInspection()
 	}
 };
 
-void UPCGComponent::NotifyNodeExecuted(const UPCGNode* InNode, const FPCGStack* InStack)
+void UPCGComponent::NotifyNodeExecuted(const UPCGNode* InNode, const FPCGStack* InStack, bool bNodeUsedCache)
 {
 	if (!ensure(InStack && InNode))
 	{
 		return;
 	}
 
+	FPCGStack Stack = *InStack;
+
+	// Reset timer information if taken from cache to provide good info in the profiling window
+	if (bNodeUsedCache)
+	{
+		Stack.Timer = PCGUtils::FCallTime();
+	}
+
 	FWriteScopeLock Lock(NodeToStacksInWhichNodeExecutedLock);
-	NodeToStacksInWhichNodeExecuted.FindOrAdd(InNode).Add(*InStack);
+	NodeToStacksInWhichNodeExecuted.FindOrAdd(InNode).Add(MoveTemp(Stack));
+}
+
+TMap<TObjectKey<const UPCGNode>, TSet<FPCGStack>> UPCGComponent::GetExecutedNodeStacks() const
+{
+	FReadScopeLock Lock(NodeToStacksInWhichNodeExecutedLock);
+	return NodeToStacksInWhichNodeExecuted;
 }
 
 uint64 UPCGComponent::GetNodeInactivePinMask(const UPCGNode* InNode, const FPCGStack& Stack) const
@@ -2020,7 +2033,7 @@ bool UPCGComponent::WasNodeExecuted(const UPCGNode* InNode, const FPCGStack& Sta
 	return FoundStacks && FoundStacks->Contains(Stack);
 }
 
-void UPCGComponent::StoreInspectionData(const FPCGStack* InStack, const UPCGNode* InNode, const FPCGDataCollection& InInputData, const FPCGDataCollection& InOutputData)
+void UPCGComponent::StoreInspectionData(const FPCGStack* InStack, const UPCGNode* InNode, const FPCGDataCollection& InInputData, const FPCGDataCollection& InOutputData, bool bUsedCache)
 {
 	if (!InNode || !ensure(InStack))
 	{
@@ -2028,7 +2041,7 @@ void UPCGComponent::StoreInspectionData(const FPCGStack* InStack, const UPCGNode
 	}
 
 	// Notify component that this task executed. Useful for editor visualization.
-	NotifyNodeExecuted(InNode, InStack);
+	NotifyNodeExecuted(InNode, InStack, bUsedCache);
 
 	if (!InOutputData.TaggedData.IsEmpty())
 	{
diff --git a/Engine/Plugins/PCG/Source/PCG/Private/Utils/PCGExtraCapture.cpp b/Engine/Plugins/PCG/Source/PCG/Private/Utils/PCGExtraCapture.cpp
--- a/Engine/Plugins/PCG/Source/PCG/Private/Utils/PCGExtraCapture.cpp
+++ b/Engine/Plugins/PCG/Source/PCG/Private/Utils/PCGExtraCapture.cpp
@@ -5,16 +5,12 @@
 #include "PCGComponent.h"
 #include "PCGContext.h"
 #include "PCGGraph.h"
+#include "PCGSubgraph.h"
 #include "Graph/PCGGraphCompiler.h"
 #include "Graph/PCGGraphExecutor.h"
 
 #if WITH_EDITOR
 
-void PCGUtils::FExtraCapture::ResetTimers()
-{
-	Timers.Empty();
-}
-
 void PCGUtils::FExtraCapture::ResetCapturedMessages()
 {
 	CapturedMessages.Empty();
@@ -22,7 +18,7 @@ void PCGUtils::FExtraCapture::ResetCapturedMessages()
 
 void PCGUtils::FExtraCapture::Update(const PCGUtils::FScopedCall& InScopedCall)
 {
-	if (!InScopedCall.Context->Node)
+	if (!InScopedCall.Context || !InScopedCall.Context->Stack)
 	{
 		return;
 	}
@@ -30,9 +26,7 @@ void PCGUtils::FExtraCapture::Update(const PCGUtils::FScopedCall& InScopedCall)
 	const double CurrentTime = FPlatformTime::Seconds();
 	const double ThisFrameTime = CurrentTime - InScopedCall.StartTime;
 
-	FScopeLock ScopedLock(&Lock);
-
-	FCallTime& Timer = Timers.FindOrAdd(InScopedCall.Context->CompiledTaskId);
+	FCallTime& Timer = const_cast<FPCGStack*>(InScopedCall.Context->Stack)->Timer;
 
 	switch (InScopedCall.Phase)
 	{
@@ -47,7 +41,7 @@ void PCGUtils::FExtraCapture::Update(const PCGUtils::FScopedCall& InScopedCall)
 
 		Timer.PrepareDataFrameCount++;
 		Timer.PrepareDataTime += ThisFrameTime;
-		Timer.PrepareDataWallTime = CurrentTime - Timer.PrepareDataStartTime;
+		Timer.PrepareDataEndTime = CurrentTime;
 		break;
 	case EPCGExecutionPhase::Execute:
 		if (Timer.ExecutionFrameCount == 0)
@@ -57,7 +51,7 @@ void PCGUtils::FExtraCapture::Update(const PCGUtils::FScopedCall& InScopedCall)
 
 		Timer.ExecutionTime += ThisFrameTime;
 		Timer.ExecutionFrameCount++;
-		Timer.ExecutionWallTime = CurrentTime - Timer.ExecutionStartTime;
+		Timer.ExecutionEndTime = CurrentTime;
 
 		Timer.MaxExecutionFrameTime = FMath::Max(Timer.MaxExecutionFrameTime, ThisFrameTime);
 		Timer.MinExecutionFrameTime = FMath::Min(Timer.MinExecutionFrameTime, ThisFrameTime);
@@ -67,6 +61,7 @@ void PCGUtils::FExtraCapture::Update(const PCGUtils::FScopedCall& InScopedCall)
 		break;
 	}
 
+	FScopeLock ScopedLock(&Lock);
 	if (!InScopedCall.CapturedMessages.IsEmpty())
 	{
 		TArray<FCapturedMessage>& InstanceMessages = CapturedMessages.FindOrAdd(InScopedCall.Context->Node);
@@ -76,147 +71,145 @@ void PCGUtils::FExtraCapture::Update(const PCGUtils::FScopedCall& InScopedCall)
 
 namespace PCGUtils
 {
-	void AddTimers(FCallTreeInfo& RootInfo, const FExtraCapture::TTimersMap& Timers, const TMap<FPCGTaskId, const FPCGGraphTask*>& TaskLookup)
+	void BuildTreeInfo(FCallTreeInfo& Info, const TMap<const UPCGNode*, const UPCGGraph*>& SubgraphNodeToGraphMap)
 	{
-		// re-use this map
-		TArray<FPCGTaskId> PathIds;
-		TArray<const UPCGNode*> PathNodes;
+		if (const UPCGGraph* const* Subgraph = SubgraphNodeToGraphMap.Find(Info.Node))
+		{
+			check(*Subgraph);
+			// This is the same as in UPCGSubgraphSettings::GetAdditionalTitleInformation
+			Info.Name = FName::NameToDisplayString((*Subgraph)->GetName(), /*bIsBool=*/false);
+		}
 
-		for (const auto& Pair : Timers)
+		for (FCallTreeInfo& Child : Info.Children)
 		{
-			const FPCGTaskId TaskId = Pair.Key;
-			const FCallTime& CallTime = Pair.Value;
+			BuildTreeInfo(Child, SubgraphNodeToGraphMap);
 
-			PathIds.Reset();
-			PathNodes.Reset();
+			Info.CallTime.PrepareDataTime += Child.CallTime.PrepareDataTime;
+			Info.CallTime.ExecutionTime += Child.CallTime.ExecutionTime;
+			Info.CallTime.PostExecuteTime += Child.CallTime.PostExecuteTime;
 
-			// build a list of parent ids that we need to populate our tree with
-			PathIds.Add(TaskId);
-			while (true)
-			{
-				const FPCGGraphTask*const * TaskItr = TaskLookup.Find(PathIds.Last());
-				if (!TaskItr || !*TaskItr)
-				{
-					// can't find an entry in our list ignore it
-					PathIds.Reset();
-					PathNodes.Reset();
-					break;
-				}
+			Info.CallTime.PrepareDataStartTime = FMath::Min(Info.CallTime.PrepareDataStartTime, Child.CallTime.PrepareDataStartTime);
+			Info.CallTime.PrepareDataEndTime = FMath::Max(Info.CallTime.PrepareDataEndTime, Child.CallTime.PrepareDataEndTime);
+			Info.CallTime.ExecutionStartTime = FMath::Min(Info.CallTime.ExecutionStartTime, Child.CallTime.ExecutionStartTime);
+			Info.CallTime.ExecutionEndTime = FMath::Max(Info.CallTime.ExecutionEndTime, Child.CallTime.ExecutionEndTime);
 
-				const FPCGGraphTask& Task = **TaskItr;
+			Info.CallTime.MinExecutionFrameTime = FMath::Min(Info.CallTime.MinExecutionFrameTime, Child.CallTime.MinExecutionFrameTime);
+			Info.CallTime.MaxExecutionFrameTime = FMath::Max(Info.CallTime.MaxExecutionFrameTime, Child.CallTime.MaxExecutionFrameTime);
+		}
+	}
+}
 
-				if (PathNodes.Num() < PathIds.Num())
-				{
-					PathNodes.Add(Task.Node);
-				}
+PCGUtils::FCallTreeInfo PCGUtils::FExtraCapture::CalculateCallTreeInfo(const UPCGComponent* Component, const FPCGStack& RootStack) const
+{
+	FCallTreeInfo RootInfo;
 
-				if (Task.ParentId == InvalidPCGTaskId)
-				{
-					break;
-				}
+	// Basically, what we want is - visit all entries in the "NodeToStacksInWhichNodeExecuted" and build our information from there.
+	TMap<TObjectKey<const UPCGNode>, TSet<FPCGStack>> NodeToStacksInWhichNodeExecuted = Component->GetExecutedNodeStacks();
 
-				if (Task.ParentId == PathIds.Last())
-				{
-					ensureMsgf(false, TEXT("Parent ID is same as current task ID which will result in infinite loop."));
-					break;
-				}
+	TArray<const UPCGNode*> NodePath;
+	TArray<int32> NodePathLoop;
+	TMap<const UPCGNode*, const UPCGGraph*> SubgraphToGraphMap;
 
-				PathIds.Add(Task.ParentId);
-			}
+	auto GetNodePath = [&NodePath, &NodePathLoop, &RootStack](const FPCGStack& Stack)
+	{
+		NodePath.Reset();
+		NodePathLoop.Reset();
 
-			FCallTreeInfo* CurrentInfo = &RootInfo;
+		if (!Stack.BeginsWith(RootStack))
+		{
+			return false;
+		}
 
-			// now add the entries, the order is reversed b/c top level ones were pushed to the end
-			for (int32 Idx = PathIds.Num()-1; Idx >= 0; --Idx)
+		const TArray<FPCGStackFrame>& StackFrames = Stack.GetStackFrames();
+		for(int StackFrameIndex = RootStack.GetStackFrames().Num(); StackFrameIndex < StackFrames.Num(); ++StackFrameIndex)
+		{
+			const FPCGStackFrame& StackFrame = StackFrames[StackFrameIndex];
+
+			if (const UPCGNode* Node = Cast<UPCGNode>(StackFrame.Object.Get()))
+			{
+				NodePath.Add(Node);
+				NodePathLoop.Add(-1);
+			}
+			else if (StackFrame.LoopIndex != INDEX_NONE)
 			{
-				const FPCGTaskId CurrentId = PathIds[Idx];
+				NodePath.Add(nullptr);
+				NodePathLoop.Add(StackFrame.LoopIndex);
+			}
+		}
 
-				FCallTreeInfo* FoundInfo = nullptr;
-				for (FCallTreeInfo& Child : CurrentInfo->Children)
-				{
-					if (Child.TaskId == CurrentId)
-					{
-						FoundInfo = &Child;
-						break;
-					}
-				}
+		return true;
+	};
 
-				if (!FoundInfo)
-				{
-					// didn't find an existing entry, add one
-					FCallTreeInfo& Child = CurrentInfo->Children.Emplace_GetRef();
-					Child.TaskId = CurrentId;
-					Child.Node = PathNodes[Idx];
+	auto GetCallInfo = [&NodePath, &NodePathLoop, &RootInfo]()
+	{
+		FCallTreeInfo* Current = &RootInfo;
 
-					CurrentInfo = &Child;
-				}
-				else
+		int NodeDepth = 0;
+		while (NodeDepth < NodePath.Num())
+		{
+			const UPCGNode* NodeToFind = NodePath[NodeDepth];
+			const int32 LoopToFind = NodePathLoop[NodeDepth];
+			++NodeDepth;
+
+			bool bFound = false;
+			for (FCallTreeInfo& Child : Current->Children)
+			{
+				if((NodeToFind != nullptr && Child.Node == NodeToFind) || 
+					(LoopToFind != INDEX_NONE && Child.Node == nullptr && Child.LoopIndex == LoopToFind))
 				{
-					CurrentInfo = FoundInfo;
+					bFound = true;
+					Current = &Child;
+					break;
 				}
 			}
 
-			CurrentInfo->CallTime = CallTime;
-		}
-	}
-
-	void BuildTreeInfo(FCallTreeInfo& Info)
-	{
-		for (FCallTreeInfo& Child : Info.Children)
-		{
-			BuildTreeInfo(Child);
-
-			Info.CallTime.PrepareDataTime += Child.CallTime.PrepareDataTime;
-			Info.CallTime.PrepareDataWallTime += Child.CallTime.PrepareDataWallTime;
-			Info.CallTime.ExecutionTime += Child.CallTime.ExecutionTime;
-			Info.CallTime.ExecutionWallTime += Child.CallTime.ExecutionWallTime;
-			Info.CallTime.PostExecuteTime += Child.CallTime.PostExecuteTime;
-			Info.CallTime.MinExecutionFrameTime = FMath::Min(Info.CallTime.MinExecutionFrameTime, Child.CallTime.MinExecutionFrameTime);
-			Info.CallTime.MaxExecutionFrameTime = FMath::Max(Info.CallTime.MaxExecutionFrameTime, Child.CallTime.MaxExecutionFrameTime);
+			if (!bFound)
+			{
+				FCallTreeInfo& Child = Current->Children.Emplace_GetRef();
+				Child.Node = NodeToFind;
+				Child.LoopIndex = LoopToFind;
+				Current = &Child;
+			}
 		}
-	}
-}
 
-PCGUtils::FCallTreeInfo PCGUtils::FExtraCapture::CalculateCallTreeInfo(const UPCGComponent* Component) const
-{
-	UPCGSubsystem* PCGSubsystem = Component ? Component->GetSubsystem() : nullptr;
-	const UPCGGraph* PCGGraph = Component ? Component->GetGraph() : nullptr;
-	if (!PCGSubsystem || !PCGGraph)
-	{
-		return {};
-	}
+		return Current;
+	};
 
-	const FPCGGraphCompiler* Compiler = PCGSubsystem->GetGraphCompiler();
-	if (!Compiler)
+	for (const auto& NodeToStacks : NodeToStacksInWhichNodeExecuted)
 	{
-		return {};
-	}
+		const UPCGNode* Node = NodeToStacks.Key.ResolveObjectPtr();
+		const TSet<FPCGStack>& Stacks = NodeToStacks.Value;
 
-	const uint32 GridSize = PCGGraph->IsHierarchicalGenerationEnabled() ? Component->GetGenerationGridSize() : PCGHiGenGrid::UninitializedGridSize();
-
-	FPCGStackContext DummyStackContext;
-	TArray<FPCGGraphTask> CompiledTasks = Compiler->GetPrecompiledTasks(PCGGraph, GridSize, DummyStackContext);
-	if (CompiledTasks.IsEmpty())
-	{
-		return {};
-	}
+		for (const FPCGStack& Stack : Stacks)
+		{
+			if (!GetNodePath(Stack))
+			{
+				continue;
+			}
 
-	// the last task on a top level graph is post execute task that doesn't need to be in the call tree
-	CompiledTasks.Pop(); 
+			// Set subgraph node name if immediate parent is a subgraph (loop or subgraph - NOT spawn actor)
+			const UPCGNode* ParentNode = NodePath.IsEmpty() ? nullptr : ((NodePathLoop.Last() != INDEX_NONE && NodePath.Num() > 1) ? NodePath.Last(1) : NodePath.Last());
+			if (const UPCGSubgraphNode* SubgraphNode = Cast<UPCGSubgraphNode>(ParentNode))
+			{
+				const UPCGGraph* CurrentGraph = Stack.GetGraphForCurrentFrame();
+				if(CurrentGraph && !SubgraphToGraphMap.Contains(SubgraphNode))
+				{
+					SubgraphToGraphMap.Add(SubgraphNode, CurrentGraph);
+				}
+			}
+			
+			// Need to add "this" node to the stack
+			NodePath.Add(Node);
+			NodePathLoop.Add(INDEX_NONE);
 
-	TMap<FPCGTaskId, const FPCGGraphTask*> TaskLookup;
+			FCallTreeInfo* Info = GetCallInfo();
 
-	// build some lookup maps
-	for (const FPCGGraphTask& Task : CompiledTasks)
-	{
-		TaskLookup.Add(Task.CompiledTaskId, &Task);
+			check(Info);
+			Info->CallTime = Stack.Timer;
+		}
 	}
 
-	FCallTreeInfo RootInfo;
-
-	PCGUtils::AddTimers(RootInfo, Timers, TaskLookup);
-	PCGUtils::BuildTreeInfo(RootInfo);
-
+	PCGUtils::BuildTreeInfo(RootInfo, SubgraphToGraphMap);
 	return RootInfo;
 }
 
diff --git a/Engine/Plugins/PCG/Source/PCG/Public/Graph/PCGStackContext.h b/Engine/Plugins/PCG/Source/PCG/Public/Graph/PCGStackContext.h
--- a/Engine/Plugins/PCG/Source/PCG/Public/Graph/PCGStackContext.h
+++ b/Engine/Plugins/PCG/Source/PCG/Public/Graph/PCGStackContext.h
@@ -2,6 +2,8 @@
 
 #pragma once
 
+#include "Utils/PCGExtraCapture.h"
+
 #include "Containers/Array.h"
 #include "UObject/WeakObjectPtr.h"
 
@@ -111,6 +113,11 @@ public:
 		return Hash;
 	}
 
+#if WITH_EDITOR
+	// Used to store node & hierarchy information
+	PCGUtils::FCallTime Timer;
+#endif
+
 private:
 	TArray<FPCGStackFrame> StackFrames;
 };
diff --git a/Engine/Plugins/PCG/Source/PCG/Public/PCGComponent.h b/Engine/Plugins/PCG/Source/PCG/Public/PCGComponent.h
--- a/Engine/Plugins/PCG/Source/PCG/Public/PCGComponent.h
+++ b/Engine/Plugins/PCG/Source/PCG/Public/PCGComponent.h
@@ -327,7 +327,7 @@ public:
 	bool IsInspecting() const;
 	void EnableInspection();
 	void DisableInspection();
-	void StoreInspectionData(const FPCGStack* InStack, const UPCGNode* InNode, const FPCGDataCollection& InInputData, const FPCGDataCollection& InOutputData);
+	void StoreInspectionData(const FPCGStack* InStack, const UPCGNode* InNode, const FPCGDataCollection& InInputData, const FPCGDataCollection& InOutputData, bool bUsedCache);
 	const FPCGDataCollection* GetInspectionData(const FPCGStack& InStack) const;
 	void ClearInspectionData(bool bClearPerNodeExecutionData = true);
 
@@ -335,7 +335,10 @@ public:
 	bool WasNodeExecuted(const UPCGNode* InNode, const FPCGStack& Stack) const;
 
 	/** Called at execution time each time a node has been executed. */
-	void NotifyNodeExecuted(const UPCGNode* InNode, const FPCGStack* InStack);
+	void NotifyNodeExecuted(const UPCGNode* InNode, const FPCGStack* InStack, bool bNodeUsedCache);
+
+	/* Retrieves the executed nodes information */
+	TMap<TObjectKey<const UPCGNode>, TSet<FPCGStack>> GetExecutedNodeStacks() const;
 
 	/** Retrieve the inactive pin bitmask for the given node and stack in the last execution. */
 	uint64 GetNodeInactivePinMask(const UPCGNode* InNode, const FPCGStack& Stack) const;
diff --git a/Engine/Plugins/PCG/Source/PCG/Public/Utils/PCGExtraCapture.h b/Engine/Plugins/PCG/Source/PCG/Public/Utils/PCGExtraCapture.h
--- a/Engine/Plugins/PCG/Source/PCG/Public/Utils/PCGExtraCapture.h
+++ b/Engine/Plugins/PCG/Source/PCG/Public/Utils/PCGExtraCapture.h
@@ -9,6 +9,7 @@ class UPCGComponent;
 class UPCGNode;
 enum class EPCGExecutionPhase : uint8;
 struct FPCGContext;
+struct FPCGStack;
 
 class IPCGElement;
 
@@ -18,20 +19,25 @@ namespace PCGUtils
 	struct PCG_API FCallTime
 	{
 		// sum of all frames
-		double ExecutionStartTime = 0.0;
+		double ExecutionStartTime = MAX_dbl;
 		double ExecutionTime = 0.0;
-		double ExecutionWallTime = 0.0;
+		double ExecutionEndTime = 0.0;
 		// how many frames
 		int32 ExecutionFrameCount = 0;
 		double MinExecutionFrameTime = MAX_dbl;
 		double MaxExecutionFrameTime = 0.0;
 
-		double PrepareDataStartTime = 0.0;
+		double PrepareDataStartTime = MAX_dbl;
 		double PrepareDataTime = 0.0;
-		double PrepareDataWallTime = 0.0;
+		double PrepareDataEndTime = 0.0;
 		int32 PrepareDataFrameCount = 0;
 
 		double PostExecuteTime = 0.0;
+
+		double PrepareDataWallTime() const { return PrepareDataEndTime - PrepareDataStartTime; }
+		double ExecutionWallTime() const { return ExecutionEndTime - ExecutionStartTime; }
+		double TotalTime() const { return ExecutionTime + PrepareDataTime; }
+		double TotalWallTime() const { return ExecutionEndTime - PrepareDataStartTime; }
 	};
 
 	struct PCG_API FCapturedMessage
@@ -44,9 +50,9 @@ namespace PCGUtils
 
 	struct PCG_API FCallTreeInfo
 	{
-		FPCGTaskId TaskId;
-		FString Name; // name of the task, Node may be nullptr for generated tasks
 		const UPCGNode* Node = nullptr;
+		int32 LoopIndex = INDEX_NONE;
+		FString Name; // overriden name for the task, will take precedence over the node name if not empty
 		FCallTime CallTime;
 
 		TArray<FCallTreeInfo> Children;
@@ -76,19 +82,15 @@ namespace PCGUtils
 	public:
 		void Update(const FScopedCall& InScopedCall);
 
-		void ResetTimers();
 		void ResetCapturedMessages();
 
-		using TTimersMap = TMap<FPCGTaskId, FCallTime>;
 		using TCapturedMessageMap = TMap<TWeakObjectPtr<const UPCGNode>, TArray<FCapturedMessage>>;
 
-		const TTimersMap& GetTimers() const { return Timers; }
 		const TCapturedMessageMap& GetCapturedMessages() const { return CapturedMessages; }
 
-		FCallTreeInfo CalculateCallTreeInfo(const UPCGComponent* Component) const;
+		FCallTreeInfo CalculateCallTreeInfo(const UPCGComponent* Component, const FPCGStack& RootStack) const;
 
 	private:
-		TTimersMap Timers;
 		TCapturedMessageMap CapturedMessages;
 		mutable FCriticalSection Lock;
 	};
diff --git a/Engine/Plugins/PCG/Source/PCGEditor/Private/SPCGEditorGraphProfilingView.cpp b/Engine/Plugins/PCG/Source/PCGEditor/Private/SPCGEditorGraphProfilingView.cpp
--- a/Engine/Plugins/PCG/Source/PCGEditor/Private/SPCGEditorGraphProfilingView.cpp
+++ b/Engine/Plugins/PCG/Source/PCGEditor/Private/SPCGEditorGraphProfilingView.cpp
@@ -10,9 +10,11 @@
 #include "PCGEditorGraphNode.h"
 
 #include "Framework/Views/TableViewMetadata.h"
+#include "HAL/PlatformApplicationMisc.h"
 #include "Widgets/Input/SButton.h"
 #include "Widgets/Input/SCheckBox.h"
 #include "Widgets/Input/SSearchBox.h"
+#include "Widgets/Input/SSpinBox.h"
 #include "Widgets/Layout/SScrollBox.h"
 #include "Widgets/Views/SListView.h"
 
@@ -26,6 +28,7 @@ namespace PCGEditorGraphProfilingView
 	const FName NAME_Node = FName(TEXT("Node"));
 	const FName NAME_PrepareDataTime = FName(TEXT("PrepareDataTime"));
 	const FName NAME_PrepareDataWallTime = FName(TEXT("PrepareData_WallTime"));
+	const FName NAME_NbPrepareFrames = FName(TEXT("NbPrepareFrames"));
 	const FName NAME_MinExecutionFrameTime = FName(TEXT("MinFrameTime"));
 	const FName NAME_MaxExecutionFrameTime = FName(TEXT("MaxFrameTime"));
 	const FName NAME_ExecutionTime = FName(TEXT("ExecutionTime"));
@@ -38,6 +41,7 @@ namespace PCGEditorGraphProfilingView
 	const FText TEXT_NodeLabel = LOCTEXT("NodeLabel", "Node");
 	const FText TEXT_PrepareDataTimeLabel = LOCTEXT("PrepareDataTimeLabel", "Prepare");
 	const FText TEXT_PrepareDataWallTimeLabel = LOCTEXT("PrepareDataWallTimeLabel", "Prepare WallTime");
+	const FText TEXT_NbPrepareFramesLabel = LOCTEXT("NbPrepareFramesLabel", "Prep Frames");
 	const FText TEXT_MinExecutionFrameTimeLabel = LOCTEXT("MinExecutionFrameTimeLabel", "Min FrameTime");
 	const FText TEXT_MaxExecutionFrameTimeLabel = LOCTEXT("MaxExecutionFrameTimeLabel", "Max FrameTime");
 	const FText TEXT_ExecutionTimeLabel = LOCTEXT("ExecutionTimeLabel", "Exec");
@@ -49,6 +53,7 @@ namespace PCGEditorGraphProfilingView
 	/** Tooltips */
 	const FText TEXT_PrepareDataTimeTooltip = LOCTEXT("PrepareDataTimeTooltip", "Cost of the PrepareData execution phase which some nodes use to process the incoming data, in ms.");
 	const FText TEXT_PrepareDataWallTimeTooltip = LOCTEXT("PrepareDataWallTimeTooltip", "Total real time elapsed between prepare data first being called until completion, including any wait/sleep time, in ms.");
+	const FText TEXT_NbPrepareFramesTooltip = LOCTEXT("NbPrepareFramesTooltip", "The number of frames in which one or more prepare data phases were executed.");
 	const FText TEXT_MinExecutionFrameTimeTooltip = LOCTEXT("MinExecutionFrameTimeTooltip", "The minimum time spent of all execution frames, in ms.");
 	const FText TEXT_MaxExecutionFrameTimeTooltip = LOCTEXT("MaxExecutionFrameTimeTooltip", "The maximum time spent of all execution frames, in ms.");
 	const FText TEXT_ExecutionTimeTooltip = LOCTEXT("ExecutionTimeTooltip", "The total time spent for execution, summed over all execution frames, in ms.");
@@ -68,68 +73,98 @@ void SPCGProfilingListViewItemRow::Construct(const FArguments& InArgs, const TSh
 		InOwnerTableView);
 }
 
-TSharedRef<SWidget> SPCGProfilingListViewItemRow::GenerateWidgetForColumn(const FName& ColumnId)
+FText FPCGProfilingListViewItem::GetTextForColumn(FName ColumnId, bool bNoGrouping) const
 {
-	FText ColumnData = LOCTEXT("ColumnError", "Unrecognized Column");
-	if (InternalItem.IsValid())
+	const FNumberFormattingOptions* NumberFormattingOptions = (bNoGrouping ? &FNumberFormattingOptions::DefaultNoGrouping() : nullptr);
+
+	if (ColumnId == PCGEditorGraphProfilingView::NAME_Node)
 	{
-		if (ColumnId == PCGEditorGraphProfilingView::NAME_Node)
+		if (bHasData)
 		{
-			ColumnData = FText::FromString(InternalItem->Name);
+			return FText::FromString(Name);
 		}
-		else if (ColumnId == PCGEditorGraphProfilingView::NAME_NbExecutionFrames)
+		else
 		{
-			ColumnData = ((InternalItem->CallTime.ExecutionFrameCount >= 0) ? FText::AsNumber(InternalItem->CallTime.ExecutionFrameCount) : FText());
+			return FText::FromString(Name + TEXT(" (cached)"));
 		}
-		else if (!InternalItem->bHasData)
-		{
-			ColumnData = PCGEditorGraphProfilingView::NoDataAvailableText;
-		}
-		else if (ColumnId == PCGEditorGraphProfilingView::NAME_MinExecutionFrameTime)
-		{
-			// In ms
-			ColumnData = ((InternalItem->CallTime.MinExecutionFrameTime >= 0) ? FText::AsNumber(InternalItem->CallTime.MinExecutionFrameTime * 1000.0) : FText());
-		}
-		else if (ColumnId == PCGEditorGraphProfilingView::NAME_MaxExecutionFrameTime)
-		{
-			// In ms
-			ColumnData = ((InternalItem->CallTime.MaxExecutionFrameTime >= 0) ? FText::AsNumber(InternalItem->CallTime.MaxExecutionFrameTime * 1000.0) : FText());
-		}
-		else if (ColumnId == PCGEditorGraphProfilingView::NAME_ExecutionTime)
-		{
-			// In ms
-			ColumnData = FText::AsNumber(InternalItem->CallTime.ExecutionTime * 1000.0);
-		}
-		else if (ColumnId == PCGEditorGraphProfilingView::NAME_ExecutionWallTime)
-		{
-			// In ms
-			ColumnData = FText::AsNumber(InternalItem->CallTime.ExecutionWallTime * 1000.0);
-		}
-		else if (ColumnId == PCGEditorGraphProfilingView::NAME_PrepareDataTime)
-		{
-			// In ms
-			ColumnData = FText::AsNumber(InternalItem->CallTime.PrepareDataTime * 1000.0);
-		}
-		else if (ColumnId == PCGEditorGraphProfilingView::NAME_PrepareDataWallTime)
-		{
-			// In ms
-			ColumnData = FText::AsNumber(InternalItem->CallTime.PrepareDataWallTime * 1000.0);
-		}
-		else if (ColumnId == PCGEditorGraphProfilingView::NAME_TotalTime)
-		{
-			// In ms
-			ColumnData = FText::AsNumber((InternalItem->CallTime.PrepareDataTime + InternalItem->CallTime.ExecutionTime) * 1000.0);
-		}
-		else if (ColumnId == PCGEditorGraphProfilingView::NAME_TotalWallTime)
+	}
+	else if (ColumnId == PCGEditorGraphProfilingView::NAME_NbExecutionFrames)
+	{
+		return ((CallTime.ExecutionFrameCount >= 0) ? FText::AsNumber(CallTime.ExecutionFrameCount, NumberFormattingOptions) : FText());
+	}
+	else if (ColumnId == PCGEditorGraphProfilingView::NAME_NbPrepareFrames)
+	{
+		return ((CallTime.PrepareDataFrameCount >= 0) ? FText::AsNumber(CallTime.PrepareDataFrameCount, NumberFormattingOptions) : FText());
+	}
+	else if (!bHasData)
+	{
+		return PCGEditorGraphProfilingView::NoDataAvailableText;
+	}
+	else if (ColumnId == PCGEditorGraphProfilingView::NAME_MinExecutionFrameTime)
+	{
+		// In ms
+		return ((CallTime.MinExecutionFrameTime >= 0) ? FText::AsNumber(CallTime.MinExecutionFrameTime * 1000.0, NumberFormattingOptions) : FText());
+	}
+	else if (ColumnId == PCGEditorGraphProfilingView::NAME_MaxExecutionFrameTime)
+	{
+		// In ms
+		return ((CallTime.MaxExecutionFrameTime >= 0) ? FText::AsNumber(CallTime.MaxExecutionFrameTime * 1000.0, NumberFormattingOptions) : FText());
+	}
+	else if (ColumnId == PCGEditorGraphProfilingView::NAME_ExecutionTime)
+	{
+		// In ms
+		return FText::AsNumber(CallTime.ExecutionTime * 1000.0, NumberFormattingOptions);
+	}
+	else if (ColumnId == PCGEditorGraphProfilingView::NAME_ExecutionWallTime)
+	{
+		// In ms
+		return FText::AsNumber(CallTime.ExecutionWallTime() * 1000.0, NumberFormattingOptions);
+	}
+	else if (ColumnId == PCGEditorGraphProfilingView::NAME_PrepareDataTime)
+	{
+		// In ms
+		return FText::AsNumber(CallTime.PrepareDataTime * 1000.0, NumberFormattingOptions);
+	}
+	else if (ColumnId == PCGEditorGraphProfilingView::NAME_PrepareDataWallTime)
+	{
+		// In ms
+		return FText::AsNumber(CallTime.PrepareDataWallTime() * 1000.0, NumberFormattingOptions);
+	}
+	else if (ColumnId == PCGEditorGraphProfilingView::NAME_TotalTime)
+	{
+		// In ms
+		return FText::AsNumber((CallTime.TotalTime()) * 1000.0, NumberFormattingOptions);
+	}
+	else if (ColumnId == PCGEditorGraphProfilingView::NAME_TotalWallTime)
+	{
+		// In ms
+		return FText::AsNumber((CallTime.TotalWallTime()) * 1000.0, NumberFormattingOptions);
+	}
+	else
+	{
+		return LOCTEXT("ItemColumnError", "Unrecognized Column");
+	}
+}
+
+TSharedRef<SWidget> SPCGProfilingListViewItemRow::GenerateWidgetForColumn(const FName& ColumnId)
+{
+	FText ColumnData = LOCTEXT("ColumnError", "Unrecognized Column");
+	FColor ColumnDataColor = FColor::White;
+
+	if (InternalItem.IsValid())
+	{
+		ColumnData = InternalItem->GetTextForColumn(ColumnId, /*bNoGrouping=*/false);
+
+		if (!InternalItem->bHasData)
 		{
-			// In ms
-			ColumnData = FText::AsNumber((InternalItem->CallTime.PrepareDataWallTime + InternalItem->CallTime.ExecutionWallTime) * 1000.0);
+			ColumnDataColor = FColor(75, 75, 75);
 		}
 	}
 
 	TSharedRef<SWidget> DataWidget = 
 		SNew(STextBlock)
 		.Text(ColumnData)
+		.ColorAndOpacity(ColumnDataColor)
 		.OverflowPolicy(ETextOverflowPolicy::Ellipsis);
 
 	// Add the internal name of the node as tooltip
@@ -182,6 +217,11 @@ void SPCGEditorGraphProfilingView::Construct(const FArguments& InArgs, TSharedPt
 	SortMode = EColumnSortMode::Descending;
 	ListViewHeader = CreateHeaderRowWidget();
 
+	ListViewCommands = MakeShareable(new FUICommandList);
+	ListViewCommands->MapAction(FGenericCommands::Get().Copy,
+		FExecuteAction::CreateSP(this, &SPCGEditorGraphProfilingView::CopySelectionToClipboard),
+		FCanExecuteAction::CreateSP(this, &SPCGEditorGraphProfilingView::CanCopySelectionToClipboard));
+
 	const TSharedRef<SScrollBar> HorizontalScrollBar = SNew(SScrollBar)
 		.Orientation(Orient_Horizontal)
 		.Thickness(FVector2D(12.0f, 12.0f));
@@ -197,6 +237,7 @@ void SPCGEditorGraphProfilingView::Construct(const FArguments& InArgs, TSharedPt
 		.OnMouseButtonDoubleClick(this, &SPCGEditorGraphProfilingView::OnItemDoubleClicked)
 		.AllowOverscroll(EAllowOverscroll::No)
 		.ExternalScrollbar(VerticalScrollBar)
+		.OnKeyDownHandler(this, &SPCGEditorGraphProfilingView::OnListViewKeyDown)
 		.ConsumeMouseWheel(EConsumeMouseWheel::Always);
 	
 	this->ChildSlot
@@ -216,28 +257,22 @@ void SPCGEditorGraphProfilingView::Construct(const FArguments& InArgs, TSharedPt
 			]
 			+SHorizontalBox::Slot()
 			.AutoWidth()
-			.Padding(1.0f, 1.0f)
-			[
-				SNew(SButton)
-				.Text(LOCTEXT("ResetButton", "Reset"))
-				.OnClicked(this, &SPCGEditorGraphProfilingView::ResetTimers)
-			]
-			+SHorizontalBox::Slot()
-			.AutoWidth()
 			.VAlign(VAlign_Center)
 			.Padding(1.0f, 0.0f)
 			[
 				SNew(STextBlock)
-				.Text(LOCTEXT("ExpandSubgraph", "Expand Subgraph"))
+				.Text(LOCTEXT("ExpandSubgraphDepth", "Expand Subgraph Depth"))
 			]
 			+SHorizontalBox::Slot()
 			.AutoWidth()
 			.VAlign(VAlign_Center)
 			.Padding(1.0f, 0.0f)
 			[
-				SNew(SCheckBox)
-				.IsChecked(this, &SPCGEditorGraphProfilingView::IsSubgraphExpanded)
-				.OnCheckStateChanged(this, &SPCGEditorGraphProfilingView::OnSubgraphExpandedChanged)
+				SNew(SSpinBox<int32>)
+				.Value(this, &SPCGEditorGraphProfilingView::GetSubgraphExpandDepth)
+				.OnValueChanged(this, &SPCGEditorGraphProfilingView::OnSubgraphExpandDepthChanged)
+				.MinValue(0)
+				.MaxValue(20)
 			]
 			+SHorizontalBox::Slot()
 			.Padding(FMargin(30.0f, 0.0f, 2.0f, 0.0f))
@@ -307,6 +342,17 @@ void SPCGEditorGraphProfilingView::Construct(const FArguments& InArgs, TSharedPt
 	Refresh();
 }
 
+void SPCGEditorGraphProfilingView::Tick(const FGeometry& AllottedGeometry, const double InCurrentTime, const float InDeltaTime)
+{
+	SCompoundWidget::Tick(AllottedGeometry, InCurrentTime, InDeltaTime);
+
+	if (bNeedsRefresh)
+	{
+		bNeedsRefresh = false;
+		Refresh();
+	}
+}
+
 TSharedRef<SHeaderRow> SPCGEditorGraphProfilingView::CreateHeaderRowWidget()
 {
 	return SNew(SHeaderRow)
@@ -337,6 +383,15 @@ TSharedRef<SHeaderRow> SPCGEditorGraphProfilingView::CreateHeaderRowWidget()
 		.OnSort(this, &SPCGEditorGraphProfilingView::OnSortColumnHeader)
 		.InitialSortMode(EColumnSortMode::Descending)
 		.DefaultTooltip(PCGEditorGraphProfilingView::TEXT_PrepareDataWallTimeTooltip)
+		+ SHeaderRow::Column(PCGEditorGraphProfilingView::NAME_NbPrepareFrames)
+		.ManualWidth(90.0f)
+		.DefaultLabel(PCGEditorGraphProfilingView::TEXT_NbPrepareFramesLabel)
+		.HAlignHeader(HAlign_Center)
+		.HAlignCell(HAlign_Right)
+		.SortMode(this, &SPCGEditorGraphProfilingView::GetColumnSortMode, PCGEditorGraphProfilingView::NAME_NbPrepareFrames)
+		.OnSort(this, &SPCGEditorGraphProfilingView::OnSortColumnHeader)
+		.InitialSortMode(EColumnSortMode::Descending)
+		.DefaultTooltip(PCGEditorGraphProfilingView::TEXT_NbPrepareFramesTooltip)
 		+ SHeaderRow::Column(PCGEditorGraphProfilingView::NAME_NbExecutionFrames)
 		.ManualWidth(90.0f)
 		.DefaultLabel(PCGEditorGraphProfilingView::TEXT_NbExecutionFramesLabel)
@@ -425,50 +480,97 @@ void SPCGEditorGraphProfilingView::OnSortColumnHeader(const EColumnSortPriority:
 		SortMode = NewSortMode;
 	}
 
-	Refresh();
+	RequestRefresh();
 }
 
-ECheckBoxState SPCGEditorGraphProfilingView::IsSubgraphExpanded() const
+void SPCGEditorGraphProfilingView::OnSubgraphExpandDepthChanged(int32 NewValue)
 {
-	return bExpandSubgraph ? ECheckBoxState::Checked : ECheckBoxState::Unchecked;
+	ExpandSubgraphDepth = NewValue;
+	RequestRefresh();
 }
 
-void SPCGEditorGraphProfilingView::OnSubgraphExpandedChanged(ECheckBoxState InNewState)
+FReply SPCGEditorGraphProfilingView::OnListViewKeyDown(const FGeometry& /*InGeometry*/, const FKeyEvent& InKeyEvent) const
 {
-	bExpandSubgraph = (InNewState == ECheckBoxState::Checked);
-	Refresh();
-}
-
-EColumnSortMode::Type SPCGEditorGraphProfilingView::GetColumnSortMode(const FName ColumnId) const
-{
-	if (SortingColumn != ColumnId)
+	if (ListViewCommands->ProcessCommandBindings(InKeyEvent))
 	{
-		return EColumnSortMode::None;
+		return FReply::Handled();
 	}
 
-	return SortMode;
+	return FReply::Unhandled();
 }
 
-FReply SPCGEditorGraphProfilingView::ResetTimers()
+void SPCGEditorGraphProfilingView::CopySelectionToClipboard() const
 {
-	const TSharedPtr<FPCGEditor> PCGEditor = PCGEditorPtr.Pin();
-	if (!PCGEditor.IsValid())
+	constexpr TCHAR Delimiter = TEXT(',');
+	constexpr TCHAR LineEnd = TEXT('\n');
+
+	TStringBuilder<2048> CSVExport;
+	TArray<FName> VisibleColumns;
+
+	// Write column header row
+	bool bHasWrittenAColumn = false;
+	const TIndirectArray<SHeaderRow::FColumn>& Columns = ListViewHeader->GetColumns();
+	for (auto Column : Columns)
 	{
-		return FReply::Handled();
+		if (!Column.bIsVisible)
+		{
+			continue;
+		}
+
+		VisibleColumns.Add(Column.ColumnId);
+		const FText& ColumnTitle = Column.DefaultText.Get();
+
+		if (bHasWrittenAColumn)
+		{
+			CSVExport += Delimiter;
+		}
+
+		CSVExport += ColumnTitle.ToString();
+		bHasWrittenAColumn = true;
 	}
 
-	if (!PCGEditorGraph)
+	// Gather selected rows and sort them to match the display order instead of selection order
+	TArray<PCGProfilingListViewItemPtr> SelectedListViewItems = ListView->GetSelectedItems();
+
+	TArray<int32> SelectedListViewItemsIndices;
+	Algo::Transform(SelectedListViewItems, SelectedListViewItemsIndices, [this](const PCGProfilingListViewItemPtr& InItem) { return ListViewItems.IndexOfByKey(InItem); });
+	Algo::Sort(SelectedListViewItemsIndices, [](const int32& A, const int32& B) { return A < B; });
+
+	SelectedListViewItems.Reset();
+	Algo::Transform(SelectedListViewItemsIndices, SelectedListViewItems, [this](const int32& InItemIndex) { return ListViewItems[InItemIndex]; });
+
+	// Write each row
+	for (const PCGProfilingListViewItemPtr& ListViewItem : SelectedListViewItems)
 	{
-		return FReply::Handled();
+		CSVExport += LineEnd;
+
+		for (int ColumnIndex = 0; ColumnIndex < VisibleColumns.Num(); ++ColumnIndex)
+		{
+			if (ColumnIndex > 0)
+			{
+				CSVExport += Delimiter;
+			}
+
+			CSVExport += ListViewItem->GetTextForColumn(VisibleColumns[ColumnIndex], /*bNoGrouping=*/true).ToString();
+		}
 	}
 
-	if (const UPCGComponent* Component = PCGComponent.Get())
+	FPlatformApplicationMisc::ClipboardCopy(*CSVExport);
+}
+
+bool SPCGEditorGraphProfilingView::CanCopySelectionToClipboard() const
+{
+	return ListView->GetNumItemsSelected() > 0;
+}
+
+EColumnSortMode::Type SPCGEditorGraphProfilingView::GetColumnSortMode(const FName ColumnId) const
+{
+	if (SortingColumn != ColumnId)
 	{
-		Component->ExtraCapture.ResetTimers();
-		Refresh();
+		return EColumnSortMode::None;
 	}
 
-	return FReply::Handled();
+	return SortMode;
 }
 
 namespace PCGEditorGraphProfilingView
@@ -477,7 +579,7 @@ namespace PCGEditorGraphProfilingView
 		TArray<PCGProfilingListViewItemPtr>& OutListViewItems,
 		const TArray<PCGUtils::FCallTreeInfo>& TreeInfo,
 		const TMap<const UPCGNode*, UPCGEditorGraphNode*>& EditorNodeLookup,
-		bool bExpandSubgraph,
+		int ExpandSubgraphDepth,
 		const FString& FolderName,
 		const FString& SearchString,
 		UPCGEditorGraphNode* CurrentEditorNode = nullptr)
@@ -488,38 +590,53 @@ namespace PCGEditorGraphProfilingView
 			UPCGEditorGraphNode* EditorNode = EditorNodeItr ? *EditorNodeItr : CurrentEditorNode;
 
 			FString Fullname = FolderName;
-			if (Info.Node)
+			if (!Info.Name.IsEmpty())
+			{
+				Fullname += Info.Name;
+			}
+			else if (Info.Node)
 			{
 				Fullname += Info.Node->GetNodeTitle(EPCGNodeTitleType::ListView).ToString();
 			}
+			else if (Info.LoopIndex != INDEX_NONE)
+			{
+				Fullname += FString::Printf(TEXT("Loop_%d"), Info.LoopIndex);
+			}
 
 			const bool bFilteredIn = SearchString.IsEmpty() || Fullname.Find(SearchString) != INDEX_NONE || (Info.Node && Info.Node->GetName().Find(SearchString) != INDEX_NONE);
 
-			// don't show inclusive times when bExpandSubgraph is on or the total sum of everything will mismatch reality
-			if (bFilteredIn && (Info.Children.IsEmpty() || bExpandSubgraph == false) && Info.CallTime.MaxExecutionFrameTime > 0)
+			if (bFilteredIn && (Info.Children.IsEmpty() || ExpandSubgraphDepth == 0))
 			{
 				PCGProfilingListViewItemPtr ListViewItem = MakeShared<FPCGProfilingListViewItem>();
 
 				ListViewItem->PCGNode = Info.Node;
 				ListViewItem->EditorNode = EditorNode;
 				ListViewItem->Name = Fullname;
-				ListViewItem->bHasData = true;
+				ListViewItem->bHasData = (Info.CallTime.MaxExecutionFrameTime > 0);
 
 				ListViewItem->CallTime = Info.CallTime;
 
 				OutListViewItems.Add(ListViewItem);
 			}
 
-			if (bExpandSubgraph && !Info.Children.IsEmpty())
+			if(ExpandSubgraphDepth > 0 && !Info.Children.IsEmpty())
 			{
-				AddListItems(OutListViewItems, Info.Children, EditorNodeLookup, /* bExpandSubgraph */ true, Fullname + "/", SearchString, EditorNode);
+				AddListItems(OutListViewItems, Info.Children, EditorNodeLookup, ExpandSubgraphDepth - 1, Fullname + "/", SearchString, EditorNode);
 			}
 		}
 	}
 }
 
+void SPCGEditorGraphProfilingView::RequestRefresh()
+{
+	bNeedsRefresh = true;
+}
+
 FReply SPCGEditorGraphProfilingView::Refresh()
 {
+	TotalTime = 0;
+	TotalWallTime = 0;
+
 	ListViewItems.Empty();
 	ListView->RequestListRefresh();
 
@@ -549,15 +666,18 @@ FReply SPCGEditorGraphProfilingView::Refresh()
 		EditorNodeLookup.Add(EditorNode->GetPCGNode(), EditorNode);
 	}
 
-	PCGUtils::FCallTreeInfo TreeInfo = Component->ExtraCapture.CalculateCallTreeInfo(Component);
-
-	TotalTime = TreeInfo.CallTime.ExecutionTime;
-	TotalWallTime = TreeInfo.CallTime.ExecutionWallTime;
+	PCGUtils::FCallTreeInfo TreeInfo = Component->ExtraCapture.CalculateCallTreeInfo(Component, PCGStack);
 
 	ListViewItems.Reserve(TreeInfo.Children.Num());
 
+	if (TreeInfo.Children.Num() > 0)
+	{
+		TotalTime = TreeInfo.CallTime.TotalTime();
+		TotalWallTime = TreeInfo.CallTime.TotalWallTime();
+	}
+
 	//TODO: could turn this into a tree instead of expanding into a list
-	PCGEditorGraphProfilingView::AddListItems(ListViewItems, TreeInfo.Children, EditorNodeLookup, bExpandSubgraph, FString(), SearchValue);
+	PCGEditorGraphProfilingView::AddListItems(ListViewItems, TreeInfo.Children, EditorNodeLookup, ExpandSubgraphDepth, FString(), SearchValue);
 
 	if (SortingColumn != NAME_None && SortMode != EColumnSortMode::None)
 	{
@@ -574,7 +694,11 @@ FReply SPCGEditorGraphProfilingView::Refresh()
 				}
 				else if (SortingColumn == PCGEditorGraphProfilingView::NAME_PrepareDataWallTime)
 				{
-					isLess = A->CallTime.PrepareDataWallTime < B->CallTime.PrepareDataWallTime;
+					isLess = A->CallTime.PrepareDataWallTime() < B->CallTime.PrepareDataWallTime();
+				}
+				else if (SortingColumn == PCGEditorGraphProfilingView::NAME_NbPrepareFrames)
+				{
+					isLess = A->CallTime.PrepareDataFrameCount < B->CallTime.PrepareDataFrameCount;
 				}
 				else if (SortingColumn == PCGEditorGraphProfilingView::NAME_MinExecutionFrameTime)
 				{
@@ -590,7 +714,7 @@ FReply SPCGEditorGraphProfilingView::Refresh()
 				}
 				else if (SortingColumn == PCGEditorGraphProfilingView::NAME_ExecutionWallTime)
 				{
-					isLess = A->CallTime.ExecutionWallTime < B->CallTime.ExecutionWallTime;
+					isLess = A->CallTime.ExecutionWallTime() < B->CallTime.ExecutionWallTime();
 				}
 				else if (SortingColumn == PCGEditorGraphProfilingView::NAME_NbExecutionFrames)
 				{
@@ -598,11 +722,11 @@ FReply SPCGEditorGraphProfilingView::Refresh()
 				}
 				else if (SortingColumn == PCGEditorGraphProfilingView::NAME_TotalTime)
 				{
-					isLess = (A->CallTime.ExecutionTime + A->CallTime.PrepareDataTime) < (B->CallTime.ExecutionTime + B->CallTime.PrepareDataTime);
+					isLess = A->CallTime.TotalTime() < B->CallTime.TotalTime();
 				}
 				else if (SortingColumn == PCGEditorGraphProfilingView::NAME_TotalWallTime)
 				{
-					isLess = (A->CallTime.ExecutionWallTime + A->CallTime.PrepareDataWallTime) < (B->CallTime.ExecutionWallTime + B->CallTime.PrepareDataWallTime);
+					isLess = A->CallTime.TotalWallTime() < B->CallTime.TotalWallTime();
 				}
 
 				return SortMode == EColumnSortMode::Ascending ? isLess : !isLess;
@@ -616,6 +740,8 @@ FReply SPCGEditorGraphProfilingView::Refresh()
 
 void SPCGEditorGraphProfilingView::OnDebugStackChanged(const FPCGStack& InPCGStack)
 {
+	PCGStack = InPCGStack;
+
 	if (PCGComponent.IsValid())
 	{
 		PCGComponent->OnPCGGraphGeneratedDelegate.RemoveAll(this);
diff --git a/Engine/Plugins/PCG/Source/PCGEditor/Private/SPCGEditorGraphProfilingView.h b/Engine/Plugins/PCG/Source/PCGEditor/Private/SPCGEditorGraphProfilingView.h
--- a/Engine/Plugins/PCG/Source/PCGEditor/Private/SPCGEditorGraphProfilingView.h
+++ b/Engine/Plugins/PCG/Source/PCGEditor/Private/SPCGEditorGraphProfilingView.h
@@ -2,8 +2,10 @@
 
 #pragma once
 
+#include "Graph/PCGStackContext.h"
 #include "Utils/PCGExtraCapture.h"
 
+#include "Framework/Commands/GenericCommands.h"
 #include "Widgets/Views/ITableRow.h"
 #include "Widgets/Views/SListView.h"
 #include "Widgets/Views/STableRow.h"
@@ -14,10 +16,11 @@ class UPCGComponent;
 class UPCGEditorGraphNode;
 class UPCGEditorGraph;
 class UPCGNode;
-struct FPCGStack;
 
 struct FPCGProfilingListViewItem
 {
+	FText GetTextForColumn(FName ColumnId, bool bNoGrouping) const;
+
 	const UPCGNode* PCGNode = nullptr;
 	const UPCGEditorGraphNode* EditorNode = nullptr;
 
@@ -53,20 +56,21 @@ public:
 	~SPCGEditorGraphProfilingView();
 
 	void Construct(const FArguments& InArgs, TSharedPtr<FPCGEditor> InPCGEditor);
+	virtual void Tick(const FGeometry& AllottedGeometry, const double InCurrentTime, const float InDeltaTime) override;
 
 private:
 	TSharedRef<SHeaderRow> CreateHeaderRowWidget();
 
-	ECheckBoxState IsSubgraphExpanded() const;
-	void OnSubgraphExpandedChanged(ECheckBoxState InNewState);
+	int32 GetSubgraphExpandDepth() const { return ExpandSubgraphDepth; }
+	void OnSubgraphExpandDepthChanged(int32 NewValue);
 
 	void OnDebugStackChanged(const FPCGStack& InPCGStack);
 
 	void OnGenerateUpdated(UPCGComponent* InPCGComponent);
 	
 	// Callbacks
+	void RequestRefresh();
 	FReply Refresh();
-	FReply ResetTimers();
 	TSharedRef<ITableRow> OnGenerateRow(PCGProfilingListViewItemPtr Item, const TSharedRef<STableViewBase>& OwnerTable) const;
 	void OnItemDoubleClicked(PCGProfilingListViewItemPtr Item);
 	void OnSortColumnHeader(const EColumnSortPriority::Type SortPriority, const FName& ColumnId, const EColumnSortMode::Type NewSortMode);
@@ -74,6 +78,10 @@ private:
 	FText GetTotalTimeLabel() const;
 	FText GetTotalWallTimeLabel() const;
 
+	FReply OnListViewKeyDown(const FGeometry& InGeometry, const FKeyEvent& InKeyEvent) const;
+	void CopySelectionToClipboard() const;
+	bool CanCopySelectionToClipboard() const;
+
 	/** Called when user changes the text they are searching for */
 	void OnSearchTextChanged(const FText& InText);
 
@@ -89,9 +97,13 @@ private:
 	/** Cached PCGComponent being viewed */
 	TWeakObjectPtr<UPCGComponent> PCGComponent;
 
+	/** Current stack being viewed */
+	FPCGStack PCGStack;
+
 	TSharedPtr<SHeaderRow> ListViewHeader;
 	TSharedPtr<SListView<PCGProfilingListViewItemPtr>> ListView;
 	TArray<PCGProfilingListViewItemPtr> ListViewItems;
+	TSharedPtr<FUICommandList> ListViewCommands;
 
 	// To allow sorting
 	FName SortingColumn = NAME_None;
@@ -100,7 +112,10 @@ private:
 	double TotalTime = 0.0;
 	double TotalWallTime = 0.0;
 
-	bool bExpandSubgraph = true;
+	bool bNeedsRefresh = false;
+
+	/** Currently depth of entry exposition - if 0, only the nodes in the currently debugged graph will be shown. */
+	int32 ExpandSubgraphDepth = 0;
 
 	/** The string to search for */
 	FString SearchValue;
