[PCG] Better guard the blueprint exceptions throw

#rnx

diff --git a/Engine/Plugins/PCG/Source/PCG/Private/Helpers/PCGGraphParametersHelpers.cpp b/Engine/Plugins/PCG/Source/PCG/Private/Helpers/PCGGraphParametersHelpers.cpp
--- a/Engine/Plugins/PCG/Source/PCG/Private/Helpers/PCGGraphParametersHelpers.cpp
+++ b/Engine/Plugins/PCG/Source/PCG/Private/Helpers/PCGGraphParametersHelpers.cpp
@@ -1,10 +1,12 @@
 // Copyright Epic Games, Inc. All Rights Reserved.
 
 #include "Helpers/PCGGraphParametersHelpers.h"
-#include "Blueprint/BlueprintExceptionInfo.h"
+
+#include "PCGModule.h"
 #include "PCGGraph.h"
-#include "PropertyBag.h"
 
+#include "PropertyBag.h"
+#include "Blueprint/BlueprintExceptionInfo.h"
 #include "Templates/ValueOrError.h"
 #include "UObject/Script.h"
 #include "UObject/Stack.h"
@@ -18,8 +20,15 @@ namespace PCGGraphParametersHelpersPrivate
 
 	void ThrowBlueprintException(const FText& ErrorMessage)
 	{
-		const FBlueprintExceptionInfo ExceptionInfo(EBlueprintExceptionType::FatalError, ErrorMessage);
-		FBlueprintCoreDelegates::ThrowScriptException(FFrame::GetThreadLocalTopStackFrame()->Object, *FFrame::GetThreadLocalTopStackFrame(), ExceptionInfo);
+		if (FFrame::GetThreadLocalTopStackFrame() && FFrame::GetThreadLocalTopStackFrame()->Object)
+		{
+			const FBlueprintExceptionInfo ExceptionInfo(EBlueprintExceptionType::FatalError, ErrorMessage);
+			FBlueprintCoreDelegates::ThrowScriptException(FFrame::GetThreadLocalTopStackFrame()->Object, *FFrame::GetThreadLocalTopStackFrame(), ExceptionInfo);
+		}
+		else
+		{
+			UE_LOG(LogPCG, Error, TEXT("%s"), *ErrorMessage.ToString());
+		}
 	}
 
 	void OnException(const EPropertyBagResult Result, const FName PropertyName)
diff --git a/Engine/Plugins/PCG/Source/PCG/Private/Metadata/PCGMetadataAccessor.cpp b/Engine/Plugins/PCG/Source/PCG/Private/Metadata/PCGMetadataAccessor.cpp
--- a/Engine/Plugins/PCG/Source/PCG/Private/Metadata/PCGMetadataAccessor.cpp
+++ b/Engine/Plugins/PCG/Source/PCG/Private/Metadata/PCGMetadataAccessor.cpp
@@ -1,9 +1,12 @@
 // Copyright Epic Games, Inc. All Rights Reserved.
 
 #include "Metadata/PCGMetadataAccessor.h"
-#include "Blueprint/BlueprintExceptionInfo.h"
-#include "Metadata/PCGMetadata.h"
+
+#include "PCGModule.h"
 #include "PCGPoint.h"
+#include "Metadata/PCGMetadata.h"
+
+#include "Blueprint/BlueprintExceptionInfo.h"
 
 #include UE_INLINE_GENERATED_CPP_BY_NAME(PCGMetadataAccessor)
 
@@ -17,7 +20,7 @@ namespace PCGMetadataAccessorHelpers
 
 	void OnException(FText ErrorMessage)
 	{
-		if (FFrame::GetThreadLocalTopStackFrame()->Object)
+		if (FFrame::GetThreadLocalTopStackFrame() && FFrame::GetThreadLocalTopStackFrame()->Object)
 		{
 			const FBlueprintExceptionInfo ExceptionInfo(EBlueprintExceptionType::FatalError, ErrorMessage);
 			FBlueprintCoreDelegates::ThrowScriptException(FFrame::GetThreadLocalTopStackFrame()->Object, *FFrame::GetThreadLocalTopStackFrame(), ExceptionInfo);
