[PCG] Sanitize attribute names when creating from tags. Replaces invalid characters with underscores rather than throwing errors. Throws optional warning which can be silenced in Get Actor Data node settings.

diff --git a/Engine/Plugins/PCG/Source/PCG/Private/Data/PCGPointData.cpp b/Engine/Plugins/PCG/Source/PCG/Private/Data/PCGPointData.cpp
--- a/Engine/Plugins/PCG/Source/PCG/Private/Data/PCGPointData.cpp
+++ b/Engine/Plugins/PCG/Source/PCG/Private/Data/PCGPointData.cpp
@@ -390,18 +390,23 @@ void UPCGPointData::SetPoints(const TArray<FPCGPoint>& InPoints)
 	GetMutablePoints() = InPoints;
 }
 
-void UPCGPointData::InitializeFromActor(AActor* InActor)
+void UPCGPointData::InitializeFromActor(AActor* InActor, bool* bOutOptionalSanitizedTagAttributeName)
 {
 	check(InActor);
 	check(Metadata && Metadata->GetAttributeCount() == 0);
 
-	AddSinglePointFromActor(InActor);
+	AddSinglePointFromActor(InActor, bOutOptionalSanitizedTagAttributeName);
 }
 
-void UPCGPointData::AddSinglePointFromActor(AActor* InActor)
+void UPCGPointData::AddSinglePointFromActor(AActor* InActor, bool* bOutOptionalSanitizedTagAttributeName)
 {
 	check(InActor);
 
+	if (bOutOptionalSanitizedTagAttributeName)
+	{
+		*bOutOptionalSanitizedTagAttributeName = false;
+	}
+
 	FPCGPoint& Point = GetMutablePoints().Emplace_GetRef();
 	Point.Steepness = 1.0f;
 	Point.Transform = InActor->GetActorTransform();
@@ -424,7 +429,7 @@ void UPCGPointData::AddSinglePointFromActor(AActor* InActor)
 	// Parse tags as well
 	for (FName Tag : InActor->Tags)
 	{
-		const FString TagString = Tag.ToString();
+		FString TagString = Tag.ToString();
 		int32 EqualPosition = INDEX_NONE;
 		
 		// Tags that contain a colon will be consider a field:value pair; we'll try to read a number first then default to a string
@@ -438,7 +443,13 @@ void UPCGPointData::AddSinglePointFromActor(AActor* InActor)
 				continue;
 			}
 
-			if(RightSide.IsNumeric())
+			const bool bSanitized = FPCGMetadataAttributeBase::SanitizeName(LeftSide);
+			if (bOutOptionalSanitizedTagAttributeName)
+			{
+				*bOutOptionalSanitizedTagAttributeName |= bSanitized;
+			}
+
+			if (RightSide.IsNumeric())
 			{
 				if (FPCGMetadataAttribute<double>* Attribute = Metadata->FindOrCreateAttribute<double>(FName(LeftSide), 0.0, /*bAllowsInterpolation=*/false, /*bOverrideParent=*/false, /*bOverwriteIfTypeMismatch=*/false))
 				{
@@ -456,7 +467,22 @@ void UPCGPointData::AddSinglePointFromActor(AActor* InActor)
 		}
 		else // Otherwise, consider that the tag is a boolean value
 		{
-			if (FPCGMetadataAttribute<bool>* Attribute = Metadata->FindOrCreateAttribute<bool>(Tag, false, /*bAllowsInterpolation=*/false, /*bOverrideParent=*/false, /*bOverwriteIfTypeMismatch=*/false))
+			FName SanitizedAttributeName = NAME_None;
+			if (FPCGMetadataAttributeBase::SanitizeName(TagString))
+			{
+				SanitizedAttributeName = FName(TagString);
+
+				if (bOutOptionalSanitizedTagAttributeName)
+				{
+					*bOutOptionalSanitizedTagAttributeName = true;
+				}
+			}
+			else
+			{
+				SanitizedAttributeName = Tag;
+			}
+
+			if (FPCGMetadataAttribute<bool>* Attribute = Metadata->FindOrCreateAttribute<bool>(SanitizedAttributeName, false, /*bAllowsInterpolation=*/false, /*bOverrideParent=*/false, /*bOverwriteIfTypeMismatch=*/false))
 			{
 				Attribute->SetValue(Point.MetadataEntry, true);
 			}
diff --git a/Engine/Plugins/PCG/Source/PCG/Private/Elements/PCGDataFromActor.cpp b/Engine/Plugins/PCG/Source/PCG/Private/Elements/PCGDataFromActor.cpp
--- a/Engine/Plugins/PCG/Source/PCG/Private/Elements/PCGDataFromActor.cpp
+++ b/Engine/Plugins/PCG/Source/PCG/Private/Elements/PCGDataFromActor.cpp
@@ -26,6 +26,7 @@ namespace PCGDataFromActorConstants
 {
 	static const FName SinglePointPinLabel = TEXT("Single Point");
 	static const FString PCGComponentDataGridSizeTagPrefix = TEXT("PCG_GridSize_");
+	static const FText TagNamesSanitizedWarning = LOCTEXT("TagAttributeNamesSanitized", "One or more tag names contained invalid characters and were sanitized when creating the corresponding attributes.");
 }
 
 namespace PCGDataFromActorHelpers
@@ -526,16 +527,26 @@ void FPCGDataFromActorElement::MergeActorsIntoPointData(FPCGContext* Context, co
 	{
 		UPCGPointData* Data = NewObject<UPCGPointData>();
 		bool bHasData = false;
+		bool bAnyAttributeNameWasSanitized = false;
 
 		for (AActor* Actor : FoundActors)
 		{
 			if (Actor)
 			{
-				Data->AddSinglePointFromActor(Actor);
+				bool bAttributeNameWasSanitized = false;
+				Data->AddSinglePointFromActor(Actor, &bAttributeNameWasSanitized);
+
+				bAnyAttributeNameWasSanitized |= bAttributeNameWasSanitized;
+
 				bHasData = true;
 			}
 		}
 
+		if (bAnyAttributeNameWasSanitized && !Settings->bSilenceSanitizedAttributeNameWarnings)
+		{
+			PCGE_LOG(Warning, GraphAndLog, PCGDataFromActorConstants::TagNamesSanitizedWarning);
+		}
+
 		if (bHasData)
 		{
 			FPCGTaggedData& TaggedData = Context->OutputData.TaggedData.Emplace_GetRef();
@@ -546,16 +557,26 @@ void FPCGDataFromActorElement::MergeActorsIntoPointData(FPCGContext* Context, co
 	{
 		FPCGDataCollection DataToMerge;
 		const bool bParseActor = false;
+		bool bAnyAttributeNameWasSanitized = false;
 
 		for (AActor* Actor : FoundActors)
 		{
 			if (Actor)
 			{
-				FPCGDataCollection Collection = UPCGComponent::CreateActorPCGDataCollection(Actor, Context->SourceComponent.Get(), EPCGDataType::Any, bParseActor);
+				bool bAttributeNameWasSanitized = false;
+				FPCGDataCollection Collection = UPCGComponent::CreateActorPCGDataCollection(Actor, Context->SourceComponent.Get(), EPCGDataType::Any, bParseActor, &bAttributeNameWasSanitized);
+
+				bAnyAttributeNameWasSanitized |= bAttributeNameWasSanitized;
+
 				DataToMerge.TaggedData += Collection.TaggedData;
 			}
 		}
 
+		if (bAnyAttributeNameWasSanitized && !Settings->bSilenceSanitizedAttributeNameWarnings)
+		{
+			PCGE_LOG(Warning, GraphAndLog, PCGDataFromActorConstants::TagNamesSanitizedWarning);
+		}
+
 		// Perform point data-to-point data merge
 		if (DataToMerge.TaggedData.Num() > 1)
 		{
@@ -693,7 +714,14 @@ void FPCGDataFromActorElement::ProcessActor(FPCGContext* Context, const UPCGData
 	else
 	{
 		const bool bParseActor = (Settings->Mode != EPCGGetDataFromActorMode::GetSinglePoint);
-		FPCGDataCollection Collection = UPCGComponent::CreateActorPCGDataCollection(FoundActor, SourceComponent, Settings->GetDataFilter(), bParseActor);
+		bool bAttributeNameWasSanitized = false;
+		FPCGDataCollection Collection = UPCGComponent::CreateActorPCGDataCollection(FoundActor, SourceComponent, Settings->GetDataFilter(), bParseActor, &bAttributeNameWasSanitized);
+
+		if (bAttributeNameWasSanitized && !Settings->bSilenceSanitizedAttributeNameWarnings)
+		{
+			PCGE_LOG(Warning, GraphAndLog, PCGDataFromActorConstants::TagNamesSanitizedWarning);
+		}
+
 		Outputs += Collection.TaggedData;
 	}
 
@@ -701,7 +729,14 @@ void FPCGDataFromActorElement::ProcessActor(FPCGContext* Context, const UPCGData
 	if (Settings->bAlsoOutputSinglePointData && (Settings->Mode == EPCGGetDataFromActorMode::GetDataFromPCGComponent || Settings->Mode == EPCGGetDataFromActorMode::GetDataFromPCGComponentOrParseComponents))
 	{
 		const bool bParseActor = false;
-		FPCGDataCollection Collection = UPCGComponent::CreateActorPCGDataCollection(FoundActor, SourceComponent, EPCGDataType::Any, bParseActor);
+		bool bAttributeNameWasSanitized = false;
+		FPCGDataCollection Collection = UPCGComponent::CreateActorPCGDataCollection(FoundActor, SourceComponent, EPCGDataType::Any, bParseActor, &bAttributeNameWasSanitized);
+
+		if (bAttributeNameWasSanitized && !Settings->bSilenceSanitizedAttributeNameWarnings)
+		{
+			PCGE_LOG(Warning, GraphAndLog, PCGDataFromActorConstants::TagNamesSanitizedWarning);
+		}
+
 		for (const FPCGTaggedData& SinglePointData : Collection.TaggedData)
 		{
 			FPCGTaggedData& OutSinglePoint = Outputs.Add_GetRef(SinglePointData);
diff --git a/Engine/Plugins/PCG/Source/PCG/Private/Metadata/PCGMetadataAttribute.cpp b/Engine/Plugins/PCG/Source/PCG/Private/Metadata/PCGMetadataAttribute.cpp
--- a/Engine/Plugins/PCG/Source/PCG/Private/Metadata/PCGMetadataAttribute.cpp
+++ b/Engine/Plugins/PCG/Source/PCG/Private/Metadata/PCGMetadataAttribute.cpp
@@ -3,6 +3,29 @@
 #include "Metadata/PCGMetadataAttribute.h"
 #include "Metadata/PCGMetadata.h"
 
+namespace PCGMetadataAttributeBase
+{
+	static constexpr TCHAR AllowedSpecialCharacters[4] = {' ', '_', '-', '/'};
+
+	bool IsValidNameCharacter(TCHAR Character)
+	{
+		if (FChar::IsAlpha(Character) || FChar::IsDigit(Character))
+		{
+			return true;
+		}
+
+		for (const TCHAR AllowedSpecialCharacter : AllowedSpecialCharacters)
+		{
+			if (AllowedSpecialCharacter == Character)
+			{
+				return true;
+			}
+		}
+
+		return false;
+	}
+}
+
 FPCGMetadataAttributeBase::FPCGMetadataAttributeBase(UPCGMetadata* InMetadata, FName InName, const FPCGMetadataAttributeBase* InParent, bool bInAllowsInterpolation)
 	: Metadata(InMetadata)
 	, Parent(InParent)
@@ -206,28 +229,9 @@ void FPCGMetadataAttributeBase::ClearEntries()
 
 bool FPCGMetadataAttributeBase::IsValidName(const FString& Name)
 {
-	// A valid name is alphanumeric with some special characters allowed.
-	static const FString AllowedSpecialCharacters = TEXT(" _-/");
-
 	for (int32 i = 0; i < Name.Len(); ++i)
 	{
-		if (FChar::IsAlpha(Name[i]) || FChar::IsDigit(Name[i]))
-		{
-			continue;
-		}
-
-		bool bAllowedSpecialCharacterFound = false;
-
-		for (int32 j = 0; j < AllowedSpecialCharacters.Len(); ++j)
-		{
-			if (Name[i] == AllowedSpecialCharacters[j])
-			{
-				bAllowedSpecialCharacterFound = true;
-				break;
-			}
-		}
-
-		if (!bAllowedSpecialCharacterFound)
+		if (!PCGMetadataAttributeBase::IsValidNameCharacter(Name[i]))
 		{
 			return false;
 		}
@@ -240,4 +244,20 @@ bool FPCGMetadataAttributeBase::IsValidName(const FName& Name)
 {
 	// Early out on None
 	return (Name == NAME_None) || IsValidName(Name.ToString());
-}
\ No newline at end of file
+}
+
+bool FPCGMetadataAttributeBase::SanitizeName(FString& InOutName)
+{
+	bool bAnyCharactersSanitized = false;
+
+	for (int32 i = 0; i < InOutName.Len(); ++i)
+	{
+		if (!PCGMetadataAttributeBase::IsValidNameCharacter(InOutName[i]))
+		{
+			InOutName[i] = '_';
+			bAnyCharactersSanitized = true;
+		}
+	}
+
+	return bAnyCharactersSanitized;
+}
diff --git a/Engine/Plugins/PCG/Source/PCG/Private/PCGComponent.cpp b/Engine/Plugins/PCG/Source/PCG/Private/PCGComponent.cpp
--- a/Engine/Plugins/PCG/Source/PCG/Private/PCGComponent.cpp
+++ b/Engine/Plugins/PCG/Source/PCG/Private/PCGComponent.cpp
@@ -2477,9 +2477,15 @@ UPCGData* UPCGComponent::CreateActorPCGData(AActor* Actor, const UPCGComponent*
 	}
 }
 
-FPCGDataCollection UPCGComponent::CreateActorPCGDataCollection(AActor* Actor, const UPCGComponent* Component, EPCGDataType InDataFilter, bool bParseActor)
+FPCGDataCollection UPCGComponent::CreateActorPCGDataCollection(AActor* Actor, const UPCGComponent* Component, EPCGDataType InDataFilter, bool bParseActor, bool* bOutOptionalSanitizedTagAttributeName)
 {
 	TRACE_CPUPROFILER_EVENT_SCOPE(UPCGComponent::CreateActorPCGData);
+
+	if (bOutOptionalSanitizedTagAttributeName)
+	{
+		*bOutOptionalSanitizedTagAttributeName = false;
+	}
+
 	FPCGDataCollection Collection;
 
 	if (!Actor)
@@ -2501,7 +2507,7 @@ FPCGDataCollection UPCGComponent::CreateActorPCGDataCollection(AActor* Actor, co
 	if (!bParseActor && !!(InDataFilter & EPCGDataType::Point))
 	{
 		UPCGPointData* Data = NewObject<UPCGPointData>();
-		Data->InitializeFromActor(Actor);
+		Data->InitializeFromActor(Actor, bOutOptionalSanitizedTagAttributeName);
 
 		FPCGTaggedData& TaggedData = Collection.TaggedData.Emplace_GetRef();
 		TaggedData.Data = Data;
@@ -2692,7 +2698,7 @@ FPCGDataCollection UPCGComponent::CreateActorPCGDataCollection(AActor* Actor, co
 	if (Collection.TaggedData.IsEmpty() && !!(InDataFilter & EPCGDataType::Point))
 	{
 		UPCGPointData* Data = NewObject<UPCGPointData>();
-		Data->InitializeFromActor(Actor);
+		Data->InitializeFromActor(Actor, bOutOptionalSanitizedTagAttributeName);
 
 		FPCGTaggedData& TaggedData = Collection.TaggedData.Emplace_GetRef();
 		TaggedData.Data = Data;
diff --git a/Engine/Plugins/PCG/Source/PCG/Public/Data/PCGPointData.h b/Engine/Plugins/PCG/Source/PCG/Public/Data/PCGPointData.h
--- a/Engine/Plugins/PCG/Source/PCG/Public/Data/PCGPointData.h
+++ b/Engine/Plugins/PCG/Source/PCG/Public/Data/PCGPointData.h
@@ -93,10 +93,10 @@ protected:
 
 public:
 	/** Initializes a single point based on the given actor */
-	void InitializeFromActor(AActor* InActor);
+	void InitializeFromActor(AActor* InActor, bool* bOutOptionalSanitizedTagAttributeName = nullptr);
 
 	/** Adds a single point based on the given actor */
-	void AddSinglePointFromActor(AActor* InActor);
+	void AddSinglePointFromActor(AActor* InActor, bool* bOutOptionalSanitizedTagAttributeName = nullptr);
 
 	UFUNCTION(BlueprintCallable, Category = SpatialData)
 	const TArray<FPCGPoint>& GetPoints() const { return Points; }
diff --git a/Engine/Plugins/PCG/Source/PCG/Public/Elements/PCGDataFromActor.h b/Engine/Plugins/PCG/Source/PCG/Public/Elements/PCGDataFromActor.h
--- a/Engine/Plugins/PCG/Source/PCG/Public/Elements/PCGDataFromActor.h
+++ b/Engine/Plugins/PCG/Source/PCG/Public/Elements/PCGDataFromActor.h
@@ -106,6 +106,10 @@ public:
 	UPROPERTY(BlueprintReadWrite, EditAnywhere, Category = "Data Retrieval Settings", meta = (EditCondition = "Mode == EPCGGetDataFromActorMode::GetDataFromProperty", EditConditionHides))
 	FName PropertyName = NAME_None;
 
+	/** Silence warnings that attribute names were sanitized to replace invalid characters. */
+	UPROPERTY(BlueprintReadWrite, EditAnywhere, Category = "Data Retrieval Settings|Advanced")
+	bool bSilenceSanitizedAttributeNameWarnings = false;
+
 #if WITH_EDITORONLY_DATA
 	/** If this is checked, found actors that are outside component bounds will not trigger a refresh. Only works for tags for now in editor. */
 	UPROPERTY(BlueprintReadOnly, EditAnywhere, Category = "Data Retrieval Settings")
diff --git a/Engine/Plugins/PCG/Source/PCG/Public/Metadata/PCGMetadata.h b/Engine/Plugins/PCG/Source/PCG/Public/Metadata/PCGMetadata.h
--- a/Engine/Plugins/PCG/Source/PCG/Public/Metadata/PCGMetadata.h
+++ b/Engine/Plugins/PCG/Source/PCG/Public/Metadata/PCGMetadata.h
@@ -372,7 +372,7 @@ FPCGMetadataAttribute<T>* UPCGMetadata::CreateAttribute(FName AttributeName, con
 {
 	if (!FPCGMetadataAttributeBase::IsValidName(AttributeName))
 	{
-		UE_LOG(LogPCG, Error, TEXT("Attribute name %s is invalid"), *AttributeName.ToString());
+		UE_LOG(LogPCG, Error, TEXT("Attribute name '%s' is invalid"), *AttributeName.ToString());
 		return nullptr;
 	}
 
diff --git a/Engine/Plugins/PCG/Source/PCG/Public/Metadata/PCGMetadataAttribute.h b/Engine/Plugins/PCG/Source/PCG/Public/Metadata/PCGMetadataAttribute.h
--- a/Engine/Plugins/PCG/Source/PCG/Public/Metadata/PCGMetadataAttribute.h
+++ b/Engine/Plugins/PCG/Source/PCG/Public/Metadata/PCGMetadataAttribute.h
@@ -78,9 +78,13 @@ public:
 
 	const FPCGMetadataAttributeBase* GetParent() const { return Parent; }
 
+	/** Returns true if for valid attribute names, which are alphanumeric with some special characters allowed. */
 	static bool IsValidName(const FString& Name);
 	static bool IsValidName(const FName& Name);
 
+	/** Replaces any invalid characters in name with underscores. Returns true if Name was changed. */
+	static bool SanitizeName(FString& InOutName);
+
 private:
 	// Unsafe version, needs to be write lock protected.
 	void SetValueFromValueKey_Unsafe(PCGMetadataEntryKey EntryKey, PCGMetadataValueKey ValueKey, bool bResetValueOnDefaultValueKey, bool bAllowInvalidEntries = false);
diff --git a/Engine/Plugins/PCG/Source/PCG/Public/PCGComponent.h b/Engine/Plugins/PCG/Source/PCG/Public/PCGComponent.h
--- a/Engine/Plugins/PCG/Source/PCG/Public/PCGComponent.h
+++ b/Engine/Plugins/PCG/Source/PCG/Public/PCGComponent.h
@@ -407,7 +407,7 @@ public:
 	FBox GetLastGeneratedBounds() const { return LastGeneratedBounds; }
 
 	/** Builds the PCG data from a given actor and its PCG component, and places it in a data collection with appropriate tags */
-	static FPCGDataCollection CreateActorPCGDataCollection(AActor* Actor, const UPCGComponent* Component, EPCGDataType InDataFilter, bool bParseActor = true);
+	static FPCGDataCollection CreateActorPCGDataCollection(AActor* Actor, const UPCGComponent* Component, EPCGDataType InDataFilter, bool bParseActor = true, bool* bOutOptionalSanitizedTagAttributeName = nullptr);
 
 	/** Builds the canonical PCG data from a given actor and its PCG component if any. */
 	static UPCGData* CreateActorPCGData(AActor* Actor, const UPCGComponent* Component, bool bParseActor = true);
diff --git a/Engine/Plugins/PCG/Source/PCGEditor/Private/PCGLevelToAsset.cpp b/Engine/Plugins/PCG/Source/PCGEditor/Private/PCGLevelToAsset.cpp
--- a/Engine/Plugins/PCG/Source/PCGEditor/Private/PCGLevelToAsset.cpp
+++ b/Engine/Plugins/PCG/Source/PCGEditor/Private/PCGLevelToAsset.cpp
@@ -90,14 +90,14 @@ UPackage* UPCGLevelToAsset::UpdateAsset(const FAssetData& PCGAsset)
 	UPCGDataAsset* Asset = Cast<UPCGDataAsset>(PCGAsset.GetAsset());
 	if (!Asset)
 	{
-		UE_LOG(LogPCG, Error, TEXT("Asset '%s' isn't a PCG data asset or could not be properly loaded."), *PCGAsset.GetObjectPathString());
+		UE_LOG(LogPCGEditor, Error, TEXT("Asset '%s' isn't a PCG data asset or could not be properly loaded."), *PCGAsset.GetObjectPathString());
 		return nullptr;
 	}
 
 	UPackage* Package = Asset->GetPackage();
 	if (!Package)
 	{
-		UE_LOG(LogPCG, Error, TEXT("Unable to retrieve package from Asset '%s'."), *PCGAsset.GetObjectPathString());
+		UE_LOG(LogPCGEditor, Error, TEXT("Unable to retrieve package from Asset '%s'."), *PCGAsset.GetObjectPathString());
 		return nullptr;
 	}
 
@@ -183,7 +183,12 @@ bool UPCGLevelToAsset::BP_ExportWorld_Implementation(UWorld* World, const FStrin
 	FPCGMetadataAttribute<int64>* ParentIndexAttribute = PointMetadata->CreateAttribute<int64>(ParentIndexAttributeName, -1, /*bAllowsInterpolation=*/false, /*bOverrideParent=*/true);
 	FPCGMetadataAttribute<FTransform>* RelativeTransformAttribute = PointMetadata->CreateAttribute<FTransform>(RelativeTransformAttributeName, FTransform::Identity, /*bAllowsInterpolation=*/false, /*bOverrideParent=*/true);
 
+	// Map from raw/unsanitized tag name to corresponding attribute.
 	TMap<FName, FPCGMetadataAttributeBase*> TagToAttributeMap;
+
+	// Relationship Tag:SanitizedName is many:1, so keep track of which sanitized names are created so we don't attempt to create the same one multiple times.
+	TSet<FName> SanitizedAttributeNames;
+
 	TSet<FName> ReservedTags;
 	ReservedTags.Add(MaterialAttributeName);
 	ReservedTags.Add(MeshAttributeName);
@@ -235,7 +240,7 @@ bool UPCGLevelToAsset::BP_ExportWorld_Implementation(UWorld* World, const FStrin
 				continue;
 			}
 
-			const FString TagString = ActorTag.ToString();
+			FString TagString = ActorTag.ToString();
 			int32 DividerPosition = INDEX_NONE;
 			
 			if (TagString.FindChar(':', DividerPosition))
@@ -248,26 +253,48 @@ bool UPCGLevelToAsset::BP_ExportWorld_Implementation(UWorld* World, const FStrin
 					continue;
 				}
 
-				const FName AttributeName = FName(LeftSide);
-				if (ReservedTags.Contains(AttributeName) || TagToAttributeMap.Contains(AttributeName))
+				const FName TagName(LeftSide);
+				FString SanitizedAttributeNameString = LeftSide;
+				const bool bSanitized = FPCGMetadataAttributeBase::SanitizeName(SanitizedAttributeNameString);
+				const FName SanitizedAttributeName(SanitizedAttributeNameString);
+
+				// Once sanitized, multiple tags can map to a single attribute name. The first tag will be used, remaining will be ignored.
+				if (ReservedTags.Contains(SanitizedAttributeName) || TagToAttributeMap.Contains(TagName) || SanitizedAttributeNames.Contains(SanitizedAttributeName))
 				{
 					continue;
 				}
 
+				if (bSanitized)
+				{
+					UE_LOG(LogPCGEditor, Warning, TEXT("Sanitized tag string on actor '%s' to remove invalid characters: '%s' -> '%s'"), *Actor->GetName(), *LeftSide, *SanitizedAttributeNameString);
+				}
+
 				// Otherwise, create the attribute based on the type of the data after the colon.
 				if (RightSide.IsNumeric())
 				{
-					TagToAttributeMap.Add(AttributeName, PointMetadata->CreateAttribute<double>(AttributeName, 0.0f, /*bAllowsInterpolation=*/false, /*bOverrideParent=*/true));
+					TagToAttributeMap.Add(TagName, PointMetadata->CreateAttribute<double>(SanitizedAttributeName, 0.0f, /*bAllowsInterpolation=*/false, /*bOverrideParent=*/true));
 				}
 				else
 				{
-					TagToAttributeMap.Add(AttributeName, PointMetadata->CreateAttribute<FString>(AttributeName, FString(), /*bAllowsInterpolation=*/false, /*bOverrideParent=*/true));
+					TagToAttributeMap.Add(TagName, PointMetadata->CreateAttribute<FString>(SanitizedAttributeName, FString(), /*bAllowsInterpolation=*/false, /*bOverrideParent=*/true));
 				}
+
+				SanitizedAttributeNames.Add(SanitizedAttributeName);
 			}
 			else
 			{
 				// Simple boolean attribute
-				TagToAttributeMap.Add(ActorTag, PointMetadata->CreateAttribute<bool>(ActorTag, false, /*bAllowsInterpolation=*/false, /*bOverrideParent=*/true));
+				FString& SanitizedTagString = TagString;
+				FPCGMetadataAttributeBase::SanitizeName(SanitizedTagString);
+				const FName SanitizedTagName(SanitizedTagString);
+
+				// Once sanitized, multiple tags can map to a single attribute name. The first tag will be used, remaining will be ignored.
+				if (!ReservedTags.Contains(SanitizedTagName) && !SanitizedAttributeNames.Contains(SanitizedTagName))
+				{
+					TagToAttributeMap.Add(ActorTag, PointMetadata->CreateAttribute<bool>(SanitizedTagName, false, /*bAllowsInterpolation=*/false, /*bOverrideParent=*/true));
+
+					SanitizedAttributeNames.Add(SanitizedTagName);
+				}
 			}
 		}
 
@@ -306,7 +333,8 @@ bool UPCGLevelToAsset::BP_ExportWorld_Implementation(UWorld* World, const FStrin
 			RelativeTransformAttribute->SetValue(Point.MetadataEntry, RelativeTransform);
 			HierarchyDepthAttribute->SetValue(Point.MetadataEntry, HierarchyDepth);
 
-			// For all tags, set attribute value to true.
+			// For all tags, if the tag is of format 'Name:Value' then create attribute Name and assign Value, otherwise
+			// create a boolean attribute with the name given by the sanitized tag string.
 			for (FName ActorTag : Actor->Tags)
 			{
 				const FString TagString = ActorTag.ToString();
@@ -322,8 +350,8 @@ bool UPCGLevelToAsset::BP_ExportWorld_Implementation(UWorld* World, const FStrin
 						continue;
 					}
 
-					const FName AttributeName = FName(LeftSide);
-					if (FPCGMetadataAttributeBase** Attribute = TagToAttributeMap.Find(AttributeName))
+					const FName TagName(LeftSide);
+					if (FPCGMetadataAttributeBase** Attribute = TagToAttributeMap.Find(TagName))
 					{
 						check(*Attribute);
 						if (RightSide.IsNumeric() && (*Attribute)->GetTypeId() == PCG::Private::MetadataTypes<double>::Id)
@@ -391,4 +419,4 @@ bool UPCGLevelToAsset::BP_ExportWorld_Implementation(UWorld* World, const FStrin
 	}
 
 	return true;
-}
\ No newline at end of file
+}
