[PCG] Create native node template

diff --git a/Engine/Plugins/PCG/Source/PCGEditor/PCGEditor.Build.cs b/Engine/Plugins/PCG/Source/PCGEditor/PCGEditor.Build.cs
--- a/Engine/Plugins/PCG/Source/PCGEditor/PCGEditor.Build.cs
+++ b/Engine/Plugins/PCG/Source/PCGEditor/PCGEditor.Build.cs
@@ -51,6 +51,7 @@ namespace UnrealBuildTool.Rules
 					"EditorStyle",
 					"EditorSubsystem",
 					"EditorWidgets",
+					"GameProjectGeneration",
 					"GraphEditor",
 					"InputCore",
 					"Kismet",
diff --git a/Engine/Plugins/PCG/Source/PCGEditor/Private/Templates/PCGClassTemplate.cpp b/Engine/Plugins/PCG/Source/PCGEditor/Private/Templates/PCGClassTemplate.cpp
new file mode 100644
--- /dev/null
+++ b/Engine/Plugins/PCG/Source/PCGEditor/Private/Templates/PCGClassTemplate.cpp
@@ -0,0 +1,32 @@
+// Copyright Epic Games, Inc. All Rights Reserved.
+
+#include "PCGClassTemplate.h"
+
+#include "PCGNode.h"
+#include "PCGSettings.h"
+
+#include UE_INLINE_GENERATED_CPP_BY_NAME(PCGClassTemplate)
+
+UPCGNativeSettingsClassTemplate::UPCGNativeSettingsClassTemplate(const FObjectInitializer& ObjectInitializer)
+	: UPluginClassTemplate(ObjectInitializer)
+{
+	PluginName = TEXT("PCG");
+	SetGeneratedBaseClass(UPCGSettings::StaticClass());
+}
+
+FString UPCGNativeSettingsClassTemplate::GetFilename() const
+{
+	return TEXT("PCGTemplateSettings");
+}
+
+UPCGNativeNodeClassTemplate::UPCGNativeNodeClassTemplate(const FObjectInitializer& ObjectInitializer)
+	: UPluginClassTemplate(ObjectInitializer)
+{
+	PluginName = TEXT("PCG");
+	SetGeneratedBaseClass(UPCGNode::StaticClass());
+}
+
+FString UPCGNativeNodeClassTemplate::GetFilename() const
+{
+	return TEXT("PCGTemplateSettings");
+}
\ No newline at end of file
diff --git a/Engine/Plugins/PCG/Source/PCGEditor/Private/Templates/PCGClassTemplate.h b/Engine/Plugins/PCG/Source/PCGEditor/Private/Templates/PCGClassTemplate.h
new file mode 100644
--- /dev/null
+++ b/Engine/Plugins/PCG/Source/PCGEditor/Private/Templates/PCGClassTemplate.h
@@ -0,0 +1,28 @@
+// Copyright Epic Games, Inc. All Rights Reserved.
+
+#pragma once
+
+#include "CoreMinimal.h"
+
+#include "ClassTemplateEditorSubsystem.h"
+
+#include "PCGClassTemplate.generated.h"
+
+
+UCLASS(Hidden)
+class UPCGNativeSettingsClassTemplate : public UPluginClassTemplate
+{
+	GENERATED_UCLASS_BODY()
+
+protected:
+	virtual FString GetFilename() const override;
+};
+
+UCLASS(Hidden)
+class UPCGNativeNodeClassTemplate : public UPluginClassTemplate
+{
+	GENERATED_UCLASS_BODY()
+	
+protected:
+	virtual FString GetFilename() const override;
+};
\ No newline at end of file
