[PCG] Change format string for printing floating point attribute values to %g to scale to different magnitudes.

diff --git a/Engine/Plugins/PCG/Source/PCG/Private/Metadata/PCGMetadataTypesConstantStruct.cpp b/Engine/Plugins/PCG/Source/PCG/Private/Metadata/PCGMetadataTypesConstantStruct.cpp
--- a/Engine/Plugins/PCG/Source/PCG/Private/Metadata/PCGMetadataTypesConstantStruct.cpp
+++ b/Engine/Plugins/PCG/Source/PCG/Private/Metadata/PCGMetadataTypesConstantStruct.cpp
@@ -9,25 +9,25 @@ FString FPCGMetadataTypesConstantStruct::ToString() const
 	case EPCGMetadataTypes::Integer32:
 		return FString::Printf(TEXT("%d"), Int32Value);
 	case EPCGMetadataTypes::Float:
-		return FString::Printf(TEXT("%.2f"), FloatValue);
+		return FString::Printf(TEXT("%g"), FloatValue);
 	case EPCGMetadataTypes::Integer64:
 		return FString::Printf(TEXT("%lld"), IntValue);
 	case EPCGMetadataTypes::Double:
-		return FString::Printf(TEXT("%.2f"), DoubleValue);
+		return FString::Printf(TEXT("%g"), DoubleValue);
 	case EPCGMetadataTypes::String:
 		return FString::Printf(TEXT("\"%s\""), *StringValue);
 	case EPCGMetadataTypes::Name:
 		return FString::Printf(TEXT("N(\"%s\")"), *NameValue.ToString());
 	case EPCGMetadataTypes::Vector2:
-		return FString::Printf(TEXT("V(%.2f, %.2f)"), Vector2Value.X, Vector2Value.Y);
+		return FString::Printf(TEXT("V(%g, %g)"), Vector2Value.X, Vector2Value.Y);
 	case EPCGMetadataTypes::Vector:
-		return FString::Printf(TEXT("V(%.2f, %.2f, %.2f)"), VectorValue.X, VectorValue.Y, VectorValue.Z);
+		return FString::Printf(TEXT("V(%g, %g, %g)"), VectorValue.X, VectorValue.Y, VectorValue.Z);
 	case EPCGMetadataTypes::Vector4:
-		return FString::Printf(TEXT("V(%.2f, %.2f, %.2f, %.2f)"), Vector4Value.X, Vector4Value.Y, Vector4Value.Z, Vector4Value.W);
+		return FString::Printf(TEXT("V(%g, %g, %g, %g)"), Vector4Value.X, Vector4Value.Y, Vector4Value.Z, Vector4Value.W);
 	case EPCGMetadataTypes::Rotator:
-		return FString::Printf(TEXT("R(%.2f, %.2f, %.2f)"), RotatorValue.Roll, RotatorValue.Pitch, RotatorValue.Yaw);
+		return FString::Printf(TEXT("R(%g, %g, %g)"), RotatorValue.Roll, RotatorValue.Pitch, RotatorValue.Yaw);
 	case EPCGMetadataTypes::Quaternion:
-		return FString::Printf(TEXT("Q(%.2f, %.2f, %.2f, %.2f)"), QuatValue.X, QuatValue.Y, QuatValue.Z, QuatValue.W);
+		return FString::Printf(TEXT("Q(%g, %g, %g, %g)"), QuatValue.X, QuatValue.Y, QuatValue.Z, QuatValue.W);
 	case EPCGMetadataTypes::Transform:
 		return FString::Printf(TEXT("Transform"));
 	case EPCGMetadataTypes::Boolean:
