[PCG] PCG Component now allows Force Cleanup

diff --git a/Engine/Plugins/PCG/Source/PCG/Private/PCGComponent.cpp b/Engine/Plugins/PCG/Source/PCG/Private/PCGComponent.cpp
--- a/Engine/Plugins/PCG/Source/PCG/Private/PCGComponent.cpp
+++ b/Engine/Plugins/PCG/Source/PCG/Private/PCGComponent.cpp
@@ -621,9 +621,120 @@ void UPCGComponent::Cleanup_Implementation(bool bRemoveComponents, bool bSave)
 	CleanupLocal(bRemoveComponents);
 }
 
+void UPCGComponent::PurgeUnlinkedResources(const AActor* InActor)
+{
+	const AActor* ThisActor = InActor;
+	if (!ThisActor)
+	{
+		return;
+	}
+
+	TSet<TSoftObjectPtr<AActor>> ActorsToDelete;
+
+	TArray<AActor*>AttachedActors;
+	TArray<UActorComponent*> ActorComponentList;
+
+	ThisActor->GetAttachedActors(AttachedActors);
+
+	for (const AActor* Actor : AttachedActors)
+	{
+		if (Actor && Actor->ActorHasTag(PCGHelpers::DefaultPCGActorTag))
+		{
+			ActorsToDelete.Add(Actor);
+		}
+	};
+
+	// Cleanup any actor components with tag and not managed by any other components
+	ThisActor->ForEachComponent(/*bIncludeFromChildActors=*/true, [&ActorComponentList](UActorComponent* ActorComponent)
+	{
+		if (ActorComponent && ActorComponent->ComponentHasTag(PCGHelpers::DefaultPCGTag))
+		{
+			ActorComponentList.Add(ActorComponent);
+		}
+	});
+	
+	ThisActor->ForEachComponent<UPCGComponent>(/*bIncludeFromChildActors=*/true, [&ActorComponentList, &ActorsToDelete](UPCGComponent* Component)
+	{
+		for (UPCGManagedResource* ManagedResource : Component->GeneratedResources)
+		{
+			if (UPCGManagedComponent* ManagedComponent = Cast<UPCGManagedComponent>(ManagedResource))
+			{
+				ActorComponentList.RemoveSwap(ManagedComponent->GeneratedComponent.Get());
+			}
+			else if (UPCGManagedActors* ManagedActors = Cast<UPCGManagedActors>(ManagedResource))
+			{
+				for (const TSoftObjectPtr<AActor>& GeneratedActor : ManagedActors->GeneratedActors)
+				{
+					ActorsToDelete.Remove(GeneratedActor);
+				}
+			}
+		}
+	});
+
+	for (UActorComponent* ActorComponent : ActorComponentList)
+	{
+		ActorComponent->DestroyComponent();
+	}
+
+	if (UWorld* World = InActor->GetWorld())
+	{
+		UPCGActorHelpers::DeleteActors(World, ActorsToDelete.Array());
+	}
+}
+
+void UPCGComponent::CleanupLocalDeleteAllGeneratedObjects(const TArray<FPCGTaskId>& Dependencies)
+{
+	UPCGSubsystem* Subsystem = GetSubsystem();
+	if (!Subsystem)
+	{
+		return;
+	}
+	
+	TArray<FPCGTaskId> TaskIds;
+	
+	auto SchedulePurge = [this, Subsystem, &TaskIds, &Dependencies](UPCGComponent* Component)
+	{
+		FPCGTaskId TaskId;
+		TWeakObjectPtr<UPCGComponent> ScheduledComponent(Component);
+
+		TaskId = Subsystem->ScheduleGeneric([this, ScheduledComponent]()
+			{
+				if (UPCGComponent* Component = ScheduledComponent.Get())
+				{
+					if (IsValid(Component))
+					{
+						Component->PurgeUnlinkedResources(this->GetOwner());
+					}
+				}
+
+				return true;
+			},
+			this, Dependencies);
+
+		if (TaskId != InvalidPCGTaskId)
+		{
+			TaskIds.Add(TaskId);
+		}
+	};
+	
+	SchedulePurge(this);
+
+	if (IsPartitioned())
+	{
+		Subsystem->ForAllRegisteredLocalComponents(this, SchedulePurge);
+	}
+
+	CleanupLocal(/*bRemoveComponents=*/true, /*bSave=*/true, TaskIds);
+}
+
 void UPCGComponent::CleanupLocal(bool bRemoveComponents, bool bSave)
 {
-	CleanupInternal(bRemoveComponents, {});
+	CleanupLocal(bRemoveComponents, bSave, {});
+}
+
+void UPCGComponent::CleanupLocal(bool bRemoveComponents, bool bSave, const TArray<FPCGTaskId>& Dependencies)
+{
+	CleanupInternal(bRemoveComponents, Dependencies);
 }
 
 FPCGTaskId UPCGComponent::CleanupInternal(bool bRemoveComponents, const TArray<FPCGTaskId>& Dependencies)
diff --git a/Engine/Plugins/PCG/Source/PCG/Public/PCGComponent.h b/Engine/Plugins/PCG/Source/PCG/Public/PCGComponent.h
--- a/Engine/Plugins/PCG/Source/PCG/Public/PCGComponent.h
+++ b/Engine/Plugins/PCG/Source/PCG/Public/PCGComponent.h
@@ -160,6 +160,11 @@ public:
 	UFUNCTION(BlueprintCallable, Category = PCG)
 	void CleanupLocal(bool bRemoveComponents, bool bSave = false);
 
+	void CleanupLocal(bool bRemoveComponents, bool bSave, const TArray<FPCGTaskId>& Dependencies);
+
+	/** Cleanup the generation while purging Actors and Components tagged as generated by PCG but are no longer managed by this or any other actor*/
+	void CleanupLocalDeleteAllGeneratedObjects(const TArray<FPCGTaskId>& Dependencies);
+
 	/**
 	 * Same as CleanupLocal, but without any delayed tasks. All is done immediately. If 'bCleanupLocalComponents' is true and the
 	 * component is partitioned, we will forward the calls to its registered local components.
@@ -415,6 +420,9 @@ public:
 protected:
 	void RefreshSchedulingPolicy();
 
+	/** Purges Actors and Components generated by PCG but are no longer managed by any PCG Component */
+	static void PurgeUnlinkedResources(const AActor* InActor);
+
 protected:
 	UPROPERTY(BlueprintReadOnly, VisibleAnywhere, Category = Settings, Instanced, meta = (NoResetToDefault, DisplayPriority = 100))
 	TObjectPtr<UPCGGraphInstance> GraphInstance;
diff --git a/Engine/Plugins/PCG/Source/PCGEditor/Private/Details/PCGComponentDetails.cpp b/Engine/Plugins/PCG/Source/PCGEditor/Private/Details/PCGComponentDetails.cpp
--- a/Engine/Plugins/PCG/Source/PCGEditor/Private/Details/PCGComponentDetails.cpp
+++ b/Engine/Plugins/PCG/Source/PCGEditor/Private/Details/PCGComponentDetails.cpp
@@ -99,6 +99,7 @@ void FPCGComponentDetails::CustomizeDetails(IDetailLayoutBuilder& DetailBuilder)
 			[
 				SNew(SButton)
 				.OnClicked(this, &FPCGComponentDetails::OnCleanupClicked)
+				.ToolTipText(FText::FromString("Cleans up graph data. \nCtrl + Click purges all components and attached actors tagged as created by PCG."))
 				.Visibility(this, &FPCGComponentDetails::CleanupButtonVisible)
 				[
 					SNew(STextBlock)
@@ -294,7 +295,15 @@ FReply FPCGComponentDetails::OnCleanupClicked()
 	{
 		if (Component.IsValid())
 		{
-			Component.Get()->Cleanup();
+			FModifierKeysState ModifierKeys = FSlateApplication::Get().GetModifierKeys();
+			if (ModifierKeys.IsControlDown())
+			{
+				Component.Get()->CleanupLocalDeleteAllGeneratedObjects({});
+			}
+			else
+			{
+				Component.Get()->Cleanup();
+			}
 		}
 	}
 
