[PCG] If the node will cull if unwired, and if it only has a single normal pin (and no required pins), then display the pin as required, because it effectively is required.

This addresses confusion where a pin may appear as non-required and gives the impression the node won't cull, but the node ends up culling due to cull if unwired being set.

diff --git a/Engine/Plugins/PCG/Source/PCG/Private/Elements/PCGHiGenGridSize.cpp b/Engine/Plugins/PCG/Source/PCG/Private/Elements/PCGHiGenGridSize.cpp
--- a/Engine/Plugins/PCG/Source/PCG/Private/Elements/PCGHiGenGridSize.cpp
+++ b/Engine/Plugins/PCG/Source/PCG/Private/Elements/PCGHiGenGridSize.cpp
@@ -56,7 +56,8 @@ EPCGDataType UPCGHiGenGridSizeSettings::GetCurrentPinTypes(const UPCGPin* InPin)
 TArray<FPCGPinProperties> UPCGHiGenGridSizeSettings::InputPinProperties() const
 {
 	TArray<FPCGPinProperties> PinProperties;
-	PinProperties.Emplace(PCGPinConstants::DefaultInputLabel, EPCGDataType::Any);
+	FPCGPinProperties& InputPinProperty = PinProperties.Emplace_GetRef(PCGPinConstants::DefaultInputLabel, EPCGDataType::Any);
+	InputPinProperty.SetRequiredPin();
 
 	return PinProperties;
 }
diff --git a/Engine/Plugins/PCG/Source/PCGEditor/Private/SPCGEditorGraphNode.cpp b/Engine/Plugins/PCG/Source/PCGEditor/Private/SPCGEditorGraphNode.cpp
--- a/Engine/Plugins/PCG/Source/PCGEditor/Private/SPCGEditorGraphNode.cpp
+++ b/Engine/Plugins/PCG/Source/PCGEditor/Private/SPCGEditorGraphNode.cpp
@@ -17,6 +17,7 @@
 #include "SLevelOfDetailBranchNode.h"
 #include "SPinTypeSelector.h"
 #include "TutorialMetaData.h"
+#include "Algo/AnyOf.h"
 #include "Widgets/Images/SImage.h"
 #include "Widgets/SToolTip.h"
 #include "Widgets/Input/SButton.h"
@@ -48,7 +49,7 @@ public:
 	bool GetExtraIcon(FName& OutExtraIcon, FText& OutTooltip) const;
 
 	/** Whether pin is required to be connected for execution. */
-	bool IsRequiredForExecution() const;
+	bool ShouldDisplayAsRequiredForExecution() const;
 
 private:
 	void ApplyUnusedPinStyle(FSlateColor& InOutColor) const;
@@ -88,7 +89,7 @@ void SPCGEditorGraphNodePin::Construct(const FArguments& InArgs, UEdGraphPin* In
 	{
 		const FSlateBrush* RequiredPinMarkerIcon = FPCGEditorStyle::Get().GetBrush(PCGEditorStyleConstants::Pin_Required);
 		RequiredPinMarkerWidth = RequiredPinMarkerIcon ? RequiredPinMarkerIcon->GetImageSize().X : 8.0f;
-		bDisplayPinMarker = IsRequiredForExecution();
+		bDisplayPinMarker = ShouldDisplayAsRequiredForExecution();
 
 		RequiredPinIconWidget =
 			SNew(SImage)
@@ -398,13 +399,35 @@ bool SPCGEditorGraphNodePin::GetExtraIcon(FName& OutExtraIcon, FText& OutTooltip
 	return (PCGPin && Settings) ? Settings->GetPinExtraIcon(PCGPin, OutExtraIcon, OutTooltip) : false;
 }
 
-bool SPCGEditorGraphNodePin::IsRequiredForExecution() const
+bool SPCGEditorGraphNodePin::ShouldDisplayAsRequiredForExecution() const
 {
 	const UPCGPin* PCGPin = nullptr;
 	const UPCGNode* PCGNode = nullptr;
 	GetPCGNodeAndPin(PCGNode, PCGPin);
 
-	return PCGPin && PCGNode->IsInputPinRequiredByExecution(PCGPin);
+	// Trivial early out tests, and advanced pins should never display as required.
+	if (!PCGPin || !PCGNode || PCGPin->Properties.IsAdvancedPin())
+	{
+		return false;
+	}
+
+	if (PCGNode->IsInputPinRequiredByExecution(PCGPin))
+	{
+		return true;
+	}
+
+	const UPCGSettings* Settings = PCGNode->GetSettings();
+	if (Settings && Settings->CanCullTaskIfUnwired())
+	{
+		// If the node will cull if unwired, and if it only has a single normal pin (and no required pins), then display the pin
+		// as required, because it effectively is. So return false if there are other pins which are not advanced.
+		return !Algo::AnyOf(PCGNode->GetInputPins(), [PCGPin](UPCGPin* InOtherPin)
+		{
+			return InOtherPin && InOtherPin != PCGPin && InOtherPin->Properties.PinStatus != EPCGPinStatus::Advanced;
+		});
+	}
+
+	return false;
 }
 
 void SPCGEditorGraphNode::Construct(const FArguments& InArgs, UPCGEditorGraphNodeBase* InNode)
