[PCG] Fix PCGComponent ignoring changes in PostEditChangeProperty on BP Actors
* Super::PostEditPropertyChanged() was re-running construction scripts, which temporarily marked the component as trash. As a result it would fail the IsValid check and exit without performing any of our PostEdit logic
* Fix PCG recording transactions for generation/cleanup
* Fix RuntimeGen ignoring bActivated flag

#rnx

diff --git a/Engine/Plugins/PCG/Source/PCG/Private/Graph/PCGGraphExecutor.cpp b/Engine/Plugins/PCG/Source/PCG/Private/Graph/PCGGraphExecutor.cpp
--- a/Engine/Plugins/PCG/Source/PCG/Private/Graph/PCGGraphExecutor.cpp
+++ b/Engine/Plugins/PCG/Source/PCG/Private/Graph/PCGGraphExecutor.cpp
@@ -600,6 +600,14 @@ void FPCGGraphExecutor::Execute()
 {
 	TRACE_CPUPROFILER_EVENT_SCOPE(FPCGGraphExecutor::Execute);
 
+#if WITH_EDITOR
+	// Avoid ticking the graph executor during an open transaction.
+	if (GUndo)
+	{
+		return;
+	}
+#endif
+
 	// Process any newly scheduled graphs to execute
 	ScheduleLock.Lock();
 
diff --git a/Engine/Plugins/PCG/Source/PCG/Private/PCGActorAndComponentMapping.cpp b/Engine/Plugins/PCG/Source/PCG/Private/PCGActorAndComponentMapping.cpp
--- a/Engine/Plugins/PCG/Source/PCG/Private/PCGActorAndComponentMapping.cpp
+++ b/Engine/Plugins/PCG/Source/PCG/Private/PCGActorAndComponentMapping.cpp
@@ -153,6 +153,22 @@ FPCGActorAndComponentMapping::FPCGActorAndComponentMapping(UPCGSubsystem* InPCGS
 	NonPartitionedOctree.Reset(FVector::ZeroVector, OctreeExtent);
 }
 
+#if WITH_EDITOR
+void FPCGActorAndComponentMapping::RegisterDelegates()
+{
+	FCoreUObjectDelegates::OnObjectsReplaced.AddRaw(this, &FPCGActorAndComponentMapping::OnObjectsReplaced);
+
+	RegisterTrackingCallbacks();
+}
+
+void FPCGActorAndComponentMapping::UnregisterDelegates()
+{
+	FCoreUObjectDelegates::OnObjectsReplaced.RemoveAll(this);
+
+	TeardownTrackingCallbacks();
+}
+#endif // WITH_EDITOR
+
 void FPCGActorAndComponentMapping::Tick()
 {
 	TSet<UPCGComponent*> ComponentToUnregister;
@@ -2071,6 +2087,44 @@ bool FPCGActorAndComponentMapping::ClearCacheForKeys(const TArray<FPCGSelectionK
 	return bShouldDirty;
 }
 
+void FPCGActorAndComponentMapping::OnObjectsReplaced(const TMap<UObject*, UObject*>& InOldToNewInstances)
+{
+	// TODO: Notify debug object view so that it can replace the selected component if necessary.
+	// TODO: Everything in FComponentInstanceData::ApplyToComponent should be handled here instead.
+
+	for (const TPair<UObject*, UObject*>& OldToNewInstance : InOldToNewInstances)
+	{
+		UPCGComponent* OldComponent = Cast<UPCGComponent>(OldToNewInstance.Key);
+		UPCGComponent* NewComponent = Cast<UPCGComponent>(OldToNewInstance.Value);
+
+		if (OldComponent && NewComponent)
+		{
+			const bool bIsLoaded = !OldComponent->HasAllFlags(RF_WasLoaded);
+
+			if (bIsLoaded)
+			{
+				OldComponent->CleanupLocalImmediate(/*bRemoveComponents=*/true);
+			}
+
+			const bool bDoActorMapping = NewComponent->bGenerated || PCGHelpers::IsRuntimeOrPIE();
+			RemapPCGComponent(OldComponent, NewComponent, bDoActorMapping);
+
+			if (bIsLoaded)
+			{
+				if (NewComponent->bHasDelayedPropertyChangedEvent)
+				{
+					NewComponent->HandlePostEditChangeProperty(NewComponent->DelayedPropertyChangedEvent);
+				}
+				else
+				{
+					NewComponent->DirtyGenerated();
+					NewComponent->Refresh(EPCGChangeType::None, /*bCancelExistingRefresh=*/true);
+				}
+			}
+		}
+	}
+}
+
 void FPCGActorAndComponentMapping::NotifyLandscapeEditModeExited()
 {
 	bIsCurrentlyExitingLandscapeEditMode = true;
diff --git a/Engine/Plugins/PCG/Source/PCG/Private/PCGComponent.cpp b/Engine/Plugins/PCG/Source/PCG/Private/PCGComponent.cpp
--- a/Engine/Plugins/PCG/Source/PCG/Private/PCGComponent.cpp
+++ b/Engine/Plugins/PCG/Source/PCG/Private/PCGComponent.cpp
@@ -55,7 +55,6 @@
 #if WITH_EDITOR
 #include "Editor.h"
 #include "EditorActorFolders.h"
-#include "ScopedTransaction.h"
 #include "Editor/Transactor.h"
 #endif
 
@@ -293,10 +292,6 @@ void UPCGComponent::Generate()
 		return;
 	}
 
-#if WITH_EDITOR
-	FScopedTransaction Transaction(LOCTEXT("PCGGenerate", "Execute generation on PCG component"));
-#endif
-
 	GenerateLocal(/*bForce=*/false);
 }
 
@@ -609,10 +604,6 @@ void UPCGComponent::Cleanup()
 		return;
 	}
 
-#if WITH_EDITOR
-	FScopedTransaction Transaction(LOCTEXT("PCGCleanup", "Clean up PCG component"));
-#endif
-
 	CleanupLocal(/*bRemoveComponents=*/true);
 }
 
@@ -1742,62 +1733,84 @@ void UPCGComponent::PreEditChange(FProperty* PropertyAboutToChange)
 
 void UPCGComponent::PostEditChangeProperty(FPropertyChangedEvent& PropertyChangedEvent)
 {
-	if (!PropertyChangedEvent.Property || !IsValid(this))
-	{
-		Super::PostEditChangeProperty(PropertyChangedEvent);
-		return;
-	}
-
-	const FName PropName = PropertyChangedEvent.Property->GetFName();
-
-	bool bTransientPropertyChangedThatDoesNotRequireARefresh = false;
-
-	// Implementation note:
-	// Since the current editing mode is a transient variable, if we do not do this transition here before going in the Super call,
-	//  we can end up in a situation where BP actors are reconstructed (... this component included ...) which makes this fall into the !IsValid case just after
-	if (PropName == GET_MEMBER_NAME_CHECKED(UPCGComponent, CurrentEditingMode))
-	{
-		// When affecting the editing mode from the user's point of view, we need to change both the current & serialized values
-		SetEditingMode(CurrentEditingMode, CurrentEditingMode);
-		ChangeTransientState(CurrentEditingMode);
-		bTransientPropertyChangedThatDoesNotRequireARefresh = true;
-	}
+	// We can't directly apply the property changed event to a BP actor component, so we save the event and handle it after reconstruction, which happens during Super::PostEditChangeProperty().
+	DelayedPropertyChangedEvent = PropertyChangedEvent;
+	bHasDelayedPropertyChangedEvent = true;
 
 	Super::PostEditChangeProperty(PropertyChangedEvent);
 
-	if (!IsValid(this))
+	// If we're not on a BP actor, the delayed property changed will already be handled, and this call will no-op.
+	HandlePostEditChangeProperty(PropertyChangedEvent);
+}
+
+void UPCGComponent::HandlePostEditChangeProperty(const FPropertyChangedEvent& PropertyChangedEvent)
+{
+	if (!IsValid(this) || !bHasDelayedPropertyChangedEvent)
 	{
 		return;
 	}
 
-	const FName MemberName = PropertyChangedEvent.MemberProperty->GetFName();
-
-	if (MemberName == GET_MEMBER_NAME_CHECKED(UPCGComponent, GenerationRadii))
+	ON_SCOPE_EXIT
 	{
-		// RuntimeGen will automatically pick up any changes to generation radii, we don't need to do any work here.
-		return;
-	}
+		bHasDelayedPropertyChangedEvent = false;
+		DelayedPropertyChangedEvent = FPropertyChangedEvent(nullptr);
+	};
+
+	const FName PropName = PropertyChangedEvent.Property ? PropertyChangedEvent.Property->GetFName() : NAME_None;
+	const FName MemberPropName = PropertyChangedEvent.MemberProperty ? PropertyChangedEvent.MemberProperty->GetFName() : NAME_None;
 
 	// Important note: all property changes already go through the OnObjectPropertyChanged, and will be dirtied here.
 	// So where only a Refresh is needed, it goes through the "capture all" else case.
 	if (PropName == GET_MEMBER_NAME_CHECKED(UPCGComponent, bIsComponentPartitioned))
 	{
-		if (CanPartition())
+		auto SetIsPartitionedHelper = [](UPCGComponent* Component)
 		{
 			// At this point, bIsComponentPartitioned is already set with the new value.
 			// But we need to do some cleanup before
 			// So keep this new value, and take its negation for the cleanup.
-			bool bIsNowPartitioned = bIsComponentPartitioned;
-			bIsComponentPartitioned = !bIsComponentPartitioned;
+			bool bIsNowPartitioned = Component->bIsComponentPartitioned;
+			Component->bIsComponentPartitioned = !Component->bIsComponentPartitioned;
 
 			// SetIsPartitioned cleans up before, so keep track if we were generated or not.
-			bool bWasGenerated = bGenerated;
-			SetIsPartitioned(bIsNowPartitioned);
+			bool bWasGenerated = Component->bGenerated;
+			Component->SetIsPartitioned(bIsNowPartitioned);
 
 			// And finally, re-generate if we were generated and activated
-			if (bWasGenerated && bActivated)
+			if (bWasGenerated && Component->bActivated)
+			{
+				Component->GenerateLocal(/*bForce=*/false);
+			}
+		};
+
+		if (CanPartition())
+		{
+			// If we're running construction scripts then we need to defer as we cannot create partition actors while these are running.
+			if (GetWorld() && GetWorld()->bIsRunningConstructionScript)
+			{
+				if (UPCGSubsystem* Subsystem = GetSubsystem())
+				{
+					TWeakObjectPtr<UPCGComponent> ComponentPtr(this);
+
+					Subsystem->ScheduleGeneric([ComponentPtr, SetIsPartitionedHelper]()
+					{
+						if (UPCGComponent* Component = ComponentPtr.Get())
+						{
+							// If the component is not valid anymore, just early out.
+							if (!IsValid(Component))
+							{
+								return true;
+							}
+
+							SetIsPartitionedHelper(Component);
+						}
+
+						return true;
+					}, this, /*TaskDependencies=*/{});
+				}
+			}
+			else
 			{
-				GenerateLocal(/*bForce=*/false);
+				SetIsPartitionedHelper(this);
 			}
 		}
 	}
@@ -1833,7 +1846,7 @@ void UPCGComponent::PostEditChangeProperty(FPropertyChangedEvent& PropertyChange
 	{
 		// We don't need to refresh the component here because this does not effect generation behavior, only scheduling behavior.
 		RefreshSchedulingPolicy();
-	} 
+	}
 	else if (PropName == GET_MEMBER_NAME_CHECKED(UPCGComponent, GenerationTrigger))
 	{
 		if (!SchedulingPolicy)
@@ -1862,8 +1875,22 @@ void UPCGComponent::PostEditChangeProperty(FPropertyChangedEvent& PropertyChange
 			Refresh();
 		}
 	}
+	else if (PropName == GET_MEMBER_NAME_CHECKED(UPCGComponent, CurrentEditingMode))
+	{
+		// When affecting the editing mode from the user's point of view, we need to change both the current & serialized values
+		SetEditingMode(CurrentEditingMode, CurrentEditingMode);
+		ChangeTransientState(CurrentEditingMode);
+	}
+	else if (PropName == GET_MEMBER_NAME_CHECKED(UPCGComponent, bOverrideGenerationRadii) || MemberPropName == GET_MEMBER_NAME_CHECKED(UPCGComponent, GenerationRadii))
+	{
+		// RuntimeGen will automatically pick up any changes to generation radii, we don't need to do any work here.
+	}
+	else if (PropName == GET_MEMBER_NAME_CHECKED(UPCGComponent, bActivated))
+	{
+		Refresh(IsManagedByRuntimeGenSystem() ? EPCGChangeType::GenerationGrid : EPCGChangeType::None);
+	}
 	// General properties that don't affect behavior
-	else if(!bTransientPropertyChangedThatDoesNotRequireARefresh)
+	else
 	{
 		Refresh();
 	}
@@ -2407,7 +2434,10 @@ bool UPCGComponent::IsObjectTracked(const UObject* InObject, bool& bOutIsCulled)
 
 void UPCGComponent::OnRefresh(bool bForceRefresh)
 {
-	check(!IsManagedByRuntimeGenSystem());
+	if (IsManagedByRuntimeGenSystem())
+	{
+		return;
+	}
 
 	// Mark the refresh task invalid to allow re-triggering refreshes
 	CurrentRefreshTask = InvalidPCGTaskId;
@@ -3481,12 +3511,20 @@ void FPCGComponentInstanceData::ApplyToComponent(UActorComponent* Component, con
 			PCGComponent->CurrentEditingMode = SourceComponent->CurrentEditingMode;
 			PCGComponent->PreviousEditingMode = SourceComponent->PreviousEditingMode;
 			PCGComponent->DynamicallyTrackedKeysToSettings = SourceComponent->DynamicallyTrackedKeysToSettings;
+			PCGComponent->DelayedPropertyChangedEvent = SourceComponent->DelayedPropertyChangedEvent;
+			PCGComponent->bHasDelayedPropertyChangedEvent = SourceComponent->bHasDelayedPropertyChangedEvent;
 #endif // WITH_EDITOR
 
+			// Move over invocation lists for dynamic delegates
+			PCGComponent->OnPCGGraphStartGeneratingExternal = SourceComponent->OnPCGGraphStartGeneratingExternal;
+			PCGComponent->OnPCGGraphCancelledExternal = SourceComponent->OnPCGGraphCancelledExternal;
+			PCGComponent->OnPCGGraphGeneratedExternal = SourceComponent->OnPCGGraphGeneratedExternal;
+			PCGComponent->OnPCGGraphCleanedExternal = SourceComponent->OnPCGGraphCleanedExternal;
+
 			// Non-critical but should be done: transient data, tracked actors cache, landscape tracking
 			// TODO Validate usefulness + move accordingly
 		}
-		
+
 		// Duplicate generated resources + retarget them
 		TArray<TObjectPtr<UPCGManagedResource>> DuplicatedResources;
 		for (const TObjectPtr<UPCGManagedResource>& Resource : GeneratedResources)
@@ -3521,12 +3559,6 @@ void FPCGComponentInstanceData::ApplyToComponent(UActorComponent* Component, con
 			PCGComponent->LoadedPreviewResources = DuplicateLoadedPreviewResources;
 		}
 
-		// Move over invocation lists for dynamic delegates
-		PCGComponent->OnPCGGraphStartGeneratingExternal = SourceComponent->OnPCGGraphStartGeneratingExternal;
-		PCGComponent->OnPCGGraphCancelledExternal = SourceComponent->OnPCGGraphCancelledExternal;
-		PCGComponent->OnPCGGraphGeneratedExternal = SourceComponent->OnPCGGraphGeneratedExternal;
-		PCGComponent->OnPCGGraphCleanedExternal = SourceComponent->OnPCGGraphCleanedExternal;
-
 		// Reconnect callbacks
 		if (PCGComponent->GraphInstance)
 		{
@@ -3534,32 +3566,12 @@ void FPCGComponentInstanceData::ApplyToComponent(UActorComponent* Component, con
 			PCGComponent->GraphInstance->OnGraphChangedDelegate.RemoveAll(PCGComponent);
 			PCGComponent->GraphInstance->OnGraphChangedDelegate.AddUObject(PCGComponent, &UPCGComponent::OnGraphChanged);
 		}
-#endif // WITH_EDITOR
-
-		bool bDoActorMapping = PCGComponent->bGenerated || PCGHelpers::IsRuntimeOrPIE();
-
-		// Also remap
-		UPCGSubsystem* Subsystem = PCGComponent->GetSubsystem();
-		if (Subsystem && SourceComponent)
-		{
-			Subsystem->RemapPCGComponent(SourceComponent, PCGComponent, bDoActorMapping);
-		}
 
-#if WITH_EDITOR
 		// Disconnect callbacks on source.
 		if (SourceComponent && SourceComponent->GraphInstance)
 		{
 			SourceComponent->GraphInstance->TeardownCallbacks();
 		}
-
-		// Finally, start a delayed refresh task (if there is not one already), in editor only
-		// It is important to be delayed, because we cannot spawn Partition Actors within this scope,
-		// because we are in a construction script.
-		// Note that we only do this if we are not currently loading
-		if (!SourceComponent || !SourceComponent->HasAllFlags(RF_WasLoaded))
-		{
-			PCGComponent->Refresh();
-		}
 #endif // WITH_EDITOR
 	}
 }
diff --git a/Engine/Plugins/PCG/Source/PCG/Private/PCGSubsystem.cpp b/Engine/Plugins/PCG/Source/PCG/Private/PCGSubsystem.cpp
--- a/Engine/Plugins/PCG/Source/PCG/Private/PCGSubsystem.cpp
+++ b/Engine/Plugins/PCG/Source/PCG/Private/PCGSubsystem.cpp
@@ -164,7 +164,7 @@ void UPCGSubsystem::Deinitialize()
 	bHasTickedOnce = false;
 
 #if WITH_EDITOR
-	ActorAndComponentMapping.TeardownTrackingCallbacks();
+	ActorAndComponentMapping.UnregisterDelegates();
 #endif // WITH_EDITOR
 
 	Super::Deinitialize();
@@ -174,6 +174,10 @@ void UPCGSubsystem::PostInitialize()
 {
 	Super::PostInitialize();
 
+#if WITH_EDITOR
+	ActorAndComponentMapping.RegisterDelegates();
+#endif
+
 	// Initialize graph executor
 	check(!GraphExecutor);
 	GraphExecutor = new FPCGGraphExecutor();
diff --git a/Engine/Plugins/PCG/Source/PCG/Private/RuntimeGen/PCGRuntimeGenScheduler.cpp b/Engine/Plugins/PCG/Source/PCG/Private/RuntimeGen/PCGRuntimeGenScheduler.cpp
--- a/Engine/Plugins/PCG/Source/PCG/Private/RuntimeGen/PCGRuntimeGenScheduler.cpp
+++ b/Engine/Plugins/PCG/Source/PCG/Private/RuntimeGen/PCGRuntimeGenScheduler.cpp
@@ -150,6 +150,14 @@ bool FPCGRuntimeGenScheduler::ShouldTick()
 		return false;
 	}
 
+#if WITH_EDITOR
+	// Avoid ticking during open editor transactions.
+	if (GUndo)
+	{
+		return false;
+	}
+#endif
+
 	// Disable tick of editor scheduling if in runtime or PIE.
 	if (PCGHelpers::IsRuntimeOrPIE() && !World->IsGameWorld())
 	{
@@ -247,7 +255,7 @@ void FPCGRuntimeGenScheduler::TickQueueComponentsForGeneration(
 	{
 		TRACE_CPUPROFILER_EVENT_SCOPE(FPCGRuntimeGenScheduler::CollectLocalComponents);
 
-		if (!ensure(OriginalComponent) || !OriginalComponent->GetGraph())
+		if (!ensure(OriginalComponent) || !OriginalComponent->GetGraph() || !OriginalComponent->bActivated)
 		{
 			continue;
 		}
@@ -384,7 +392,7 @@ void FPCGRuntimeGenScheduler::TickQueueComponentsForGeneration(
 	{
 		TRACE_CPUPROFILER_EVENT_SCOPE(FPCGRuntimeGenScheduler::CollectNonPartitionedComponents);
 
-		if (!ensure(OriginalComponent) || !OriginalComponent->GetGraph())
+		if (!ensure(OriginalComponent) || !OriginalComponent->GetGraph() || !OriginalComponent->bActivated)
 		{
 			continue;
 		}
@@ -453,10 +461,18 @@ void FPCGRuntimeGenScheduler::TickCleanup(const TSet<IPCGGenSourceBase*>& InGenS
 		const uint32 GridSize = GenerationKey.GetGridSize();
 		const FIntVector& GridCoords = GenerationKey.GetGridCoords();
 		UPCGComponent* OriginalComponent = GenerationKey.GetOriginalComponent();
-		check(OriginalComponent);
-
 		const EPCGHiGenGrid Grid = PCGHiGenGrid::GridSizeToGrid(GridSize);
 
+		// The original component can be marked invalid by BP construction scripts, in which case we should always cleanup to avoid hanging onto dead entries.
+		if (!OriginalComponent || !IsValid(OriginalComponent) || !OriginalComponent->bActivated)
+		{
+			// If the Grid is unbounded, we should clean up the original component.
+			// Otherwise, the local component may be unrecoverable so we provide a nullptr and let the PA clean itself up.
+			UPCGComponent* ComponentToClean = Grid == EPCGHiGenGrid::Unbounded ? OriginalComponent : nullptr;
+			ComponentsToClean.Add({ GenerationKey, ComponentToClean });
+			continue;
+		}
+
 		// If the Grid is unbounded, we have a non-partitioned or unbounded component.
 		if (Grid == EPCGHiGenGrid::Unbounded)
 		{
diff --git a/Engine/Plugins/PCG/Source/PCG/Public/Grid/PCGPartitionActor.h b/Engine/Plugins/PCG/Source/PCG/Public/Grid/PCGPartitionActor.h
--- a/Engine/Plugins/PCG/Source/PCG/Public/Grid/PCGPartitionActor.h
+++ b/Engine/Plugins/PCG/Source/PCG/Public/Grid/PCGPartitionActor.h
@@ -117,7 +117,7 @@ private:
 	// Note: this map is not a property and not serialized since we will rebuild it from the LocalToOriginal
 	TMap<TObjectPtr<UPCGComponent>, TObjectPtr<UPCGComponent>> OriginalToLocal;
 
-	UPROPERTY()
+	UPROPERTY(NonTransactional)
 	TMap<TObjectPtr<UPCGComponent>, TSoftObjectPtr<UPCGComponent>> LocalToOriginal;
 
 	// PCG components that are cleared when in preview-on-load mode are kept aside and put back when serializing to prevent changes
diff --git a/Engine/Plugins/PCG/Source/PCG/Public/PCGActorAndComponentMapping.h b/Engine/Plugins/PCG/Source/PCG/Public/PCGActorAndComponentMapping.h
--- a/Engine/Plugins/PCG/Source/PCG/Public/PCGActorAndComponentMapping.h
+++ b/Engine/Plugins/PCG/Source/PCG/Public/PCGActorAndComponentMapping.h
@@ -47,6 +47,9 @@ public:
 	void Tick();
 
 #if WITH_EDITOR
+	void RegisterDelegates();
+	void UnregisterDelegates();
+
 	/** If the partition grid size change, call this to empty the Partition actors map */
 	void ResetPartitionActorsMap();
 
@@ -178,6 +181,8 @@ private:
 
 	/** Gather all settings from a given component that track the key, and clear the cache for them. Returns true if we should dirty afterwards (aka at least one settings was cleared and/or landscape changed). */
 	bool ClearCacheForKeys(const TArray<FPCGSelectionKey>& InKeys, const UPCGComponent* InComponent, const bool bIntersect, const UObject* InOriginatingChange) const;
+
+	void OnObjectsReplaced(const TMap<UObject*, UObject*>& InOldToNewInstances);
 #endif // WITH_EDITOR
 
 private:
diff --git a/Engine/Plugins/PCG/Source/PCG/Public/PCGComponent.h b/Engine/Plugins/PCG/Source/PCG/Public/PCGComponent.h
--- a/Engine/Plugins/PCG/Source/PCG/Public/PCGComponent.h
+++ b/Engine/Plugins/PCG/Source/PCG/Public/PCGComponent.h
@@ -259,7 +259,7 @@ public:
 	UPROPERTY(BlueprintReadWrite, EditAnywhere, Category = "Editing Settings", meta = (DisplayPriority = 450))
 	bool bOnlyTrackItself = false;
 
-	UPROPERTY(BlueprintReadOnly, VisibleAnywhere, Transient, Category = Debug, meta = (NoResetToDefault))
+	UPROPERTY(BlueprintReadOnly, VisibleAnywhere, Transient, NonTransactional, Category = Debug, meta = (NoResetToDefault))
 	bool bDirtyGenerated = false;
 
 	// Property that will automatically be set on BP templates, to allow for "Generate on add to world" in editor.
@@ -501,6 +501,9 @@ private:
 	virtual void PreEditUndo() override;
 	virtual void PostEditUndo() override;
 
+	/** Helper to handle a property changed event from PostEditChangeProperty(). Necessary to support delayed handling of changed properties for BP actors during reconstruction. */
+	void HandlePostEditChangeProperty(const FPropertyChangedEvent& PropertyChangedEvent);
+
 	/** Sets up actor, tracking, landscape and graph callbacks */
 	void SetupCallbacksOnCreation();
 
@@ -527,19 +530,19 @@ private:
 
 	FBox GetGridBounds(const AActor* InActor) const;
 
-	UPROPERTY(Transient, NonPIEDuplicateTransient)
+	UPROPERTY(Transient, NonPIEDuplicateTransient, NonTransactional)
 	TObjectPtr<UPCGData> CachedPCGData = nullptr;
 
-	UPROPERTY(Transient, NonPIEDuplicateTransient)
+	UPROPERTY(Transient, NonPIEDuplicateTransient, NonTransactional)
 	TObjectPtr<UPCGData> CachedInputData = nullptr;
 
-	UPROPERTY(Transient, NonPIEDuplicateTransient)
+	UPROPERTY(Transient, NonPIEDuplicateTransient, NonTransactional)
 	TObjectPtr<UPCGData> CachedActorData = nullptr;
 
-	UPROPERTY(Transient, NonPIEDuplicateTransient)
+	UPROPERTY(Transient, NonPIEDuplicateTransient, NonTransactional)
 	TObjectPtr<UPCGData> CachedLandscapeData = nullptr;
 
-	UPROPERTY(Transient, NonPIEDuplicateTransient)
+	UPROPERTY(Transient, NonPIEDuplicateTransient, NonTransactional)
 	TObjectPtr<UPCGData> CachedLandscapeHeightData = nullptr;
 
 #if WITH_EDITORONLY_DATA
@@ -549,11 +552,11 @@ private:
 
 	// NOTE: This should not be made visible or editable because it will change the way the BP actors are
 	// duplicated/setup and might trigger an ensure in the resources.
-	UPROPERTY()
+	UPROPERTY(NonTransactional)
 	TArray<TObjectPtr<UPCGManagedResource>> GeneratedResources;
 
 #if WITH_EDITORONLY_DATA
-	UPROPERTY(Transient)
+	UPROPERTY(Transient, NonTransactional)
 	TArray<TObjectPtr<UPCGManagedResource>> LoadedPreviewResources;
 
 	UPROPERTY(Transient, VisibleInstanceOnly, Category = Debug)
@@ -588,6 +591,8 @@ private:
 #if WITH_EDITOR
 	int32 InspectionCounter = 0;
 	FBox LastGeneratedBoundsPriorToUndo = FBox(EForceInit::ForceInit);
+	FPropertyChangedEvent DelayedPropertyChangedEvent = FPropertyChangedEvent(nullptr);
+	bool bHasDelayedPropertyChangedEvent = false;
 #endif
 
 #if WITH_EDITORONLY_DATA
diff --git a/Engine/Plugins/PCG/Source/PCG/Public/PCGSubsystem.h b/Engine/Plugins/PCG/Source/PCG/Public/PCGSubsystem.h
--- a/Engine/Plugins/PCG/Source/PCG/Public/PCGSubsystem.h
+++ b/Engine/Plugins/PCG/Source/PCG/Public/PCGSubsystem.h
@@ -177,9 +177,6 @@ public:
 	/** Register a new PCG Component or update it, will be added to the octree if it doesn't exists yet. Returns true if it was added/updated. Thread safe */
 	bool RegisterOrUpdatePCGComponent(UPCGComponent* InComponent, bool bDoActorMapping = true) { return ActorAndComponentMapping.RegisterOrUpdatePCGComponent(InComponent, bDoActorMapping); }
 
-	/** In case of BP Actors, we need to remap the old component destroyed by the construction script to the new one. Returns true if re-mapping succeeded. */
-	bool RemapPCGComponent(const UPCGComponent* OldComponent, UPCGComponent* NewComponent, bool bDoActorMapping) { return ActorAndComponentMapping.RemapPCGComponent(OldComponent, NewComponent, bDoActorMapping); }
-
 	/** Unregister a PCG Component, will be removed from the octree. Can force it, if we have a delayed unregister. Thread safe */
 	void UnregisterPCGComponent(UPCGComponent* InComponent, bool bForce = false) { ActorAndComponentMapping.UnregisterPCGComponent(InComponent, bForce); }
 
