[PCG] Fixed a bug during copy and paste in the Editable User Parameter Details

#rnx

diff --git a/Engine/Plugins/PCG/Source/PCGEditor/Private/Details/PCGEditableUserParameterDetails.cpp b/Engine/Plugins/PCG/Source/PCGEditor/Private/Details/PCGEditableUserParameterDetails.cpp
--- a/Engine/Plugins/PCG/Source/PCGEditor/Private/Details/PCGEditableUserParameterDetails.cpp
+++ b/Engine/Plugins/PCG/Source/PCGEditor/Private/Details/PCGEditableUserParameterDetails.cpp
@@ -48,15 +48,17 @@ void FPCGEditableUserParameterDetails::CustomizeDetails(IDetailLayoutBuilder& De
 		return;
 	}
 
+	// Cache the parent details view to be used on tick
+	ParentDetailsView = StaticCastWeakPtr<IDetailsView>(DetailBuilder.GetDetailsView()->AsWeak());
+
 	if (UPCGUserParameterGetSettings* Settings = Cast<UPCGUserParameterGetSettings>(SettingsObject.Get()))
 	{
 		if (UObject* NodeObject = Settings->GetOuter())
 		{
 			if (UPCGGraphInterface* GraphInterface = Cast<UPCGGraphInterface>(NodeObject->GetOuter()))
 			{
-				// Cache the interface and view to be used to verify on tick
+				// Cache the interface to be used to verify on tick
 				CachedGraphInterface = MakeWeakObjectPtr(GraphInterface);
-				ParentDetailsView = StaticCastWeakPtr<IDetailsView>(DetailBuilder.GetDetailsView()->AsWeak());
 
 				// It is safe, because we hook pre/post edit changes that will trigger the callbacks
 				if (FInstancedPropertyBag* UserParameters = GraphInterface->GetMutableUserParametersStruct_Unsafe())
diff --git a/Engine/Plugins/PCG/Source/PCGEditor/Public/Details/PCGEditableUserParameterDetails.h b/Engine/Plugins/PCG/Source/PCGEditor/Public/Details/PCGEditableUserParameterDetails.h
--- a/Engine/Plugins/PCG/Source/PCGEditor/Public/Details/PCGEditableUserParameterDetails.h
+++ b/Engine/Plugins/PCG/Source/PCGEditor/Public/Details/PCGEditableUserParameterDetails.h
@@ -20,6 +20,7 @@ public:
 	virtual void CustomizeDetails(IDetailLayoutBuilder& DetailBuilder) override;
 
 	//~Begin FTickableEditorObject interface
+	virtual bool IsAllowedToTick() const override { return ParentDetailsView.IsValid(); }
 	virtual void Tick(float DeltaTime) override;
 	virtual TStatId GetStatId() const override;
 	//~End FTickableEditorObject interface
