[PCG] Added icons for PCG assets and PCG component and in menus as well

#rnx

diff --git a/Engine/Plugins/PCG/Source/PCGEditor/Private/AssetDefinitions/AssetDefinition_PCGGraphInterface.cpp b/Engine/Plugins/PCG/Source/PCGEditor/Private/AssetDefinitions/AssetDefinition_PCGGraphInterface.cpp
--- a/Engine/Plugins/PCG/Source/PCGEditor/Private/AssetDefinitions/AssetDefinition_PCGGraphInterface.cpp
+++ b/Engine/Plugins/PCG/Source/PCGEditor/Private/AssetDefinitions/AssetDefinition_PCGGraphInterface.cpp
@@ -3,6 +3,8 @@
 #include "AssetDefinition_PCGGraphInterface.h"
 
 #include "PCGGraph.h"
+
+#include "PCGEditorStyle.h"
 #include "PCGGraphFactory.h"
 
 #include "ContentBrowserMenuContexts.h"
@@ -70,7 +72,7 @@ namespace MenuExtension_PCGGraphInterface
 					const TAttribute<FText> ToolTip = LOCTEXT("PCGGraph_NewInstanceToolTip", "Creates a parameterized PCG graph using this graph as a base.");
 					const FToolMenuExecuteAction UIAction = FToolMenuExecuteAction::CreateStatic(&ExecuteNewPCGGraphInstance);
 
-					InSection.AddMenuEntry("PCGGraph_NewInstance", Label, ToolTip, FSlateIcon(), UIAction);
+					InSection.AddMenuEntry("PCGGraph_NewInstance", Label, ToolTip, FSlateIcon(FPCGEditorStyle::Get().GetStyleSetName(), "ClassIcon.PCGGraphInstance"), UIAction);
 				}
 			}));
 		}));
diff --git a/Engine/Plugins/PCG/Source/PCGEditor/Private/PCGEditorCommands.cpp b/Engine/Plugins/PCG/Source/PCGEditor/Private/PCGEditorCommands.cpp
--- a/Engine/Plugins/PCG/Source/PCGEditor/Private/PCGEditorCommands.cpp
+++ b/Engine/Plugins/PCG/Source/PCGEditor/Private/PCGEditorCommands.cpp
@@ -18,7 +18,7 @@ FPCGEditorCommands::FPCGEditorCommands()
 void FPCGEditorCommands::RegisterCommands()
 {
 	UI_COMMAND(CollapseNodes, "Collapse into Subgraph", "Collapse selected nodes into a separate PCGGraph asset.", EUserInterfaceActionType::Button, FInputChord(EModifierKey::Alt, EKeys::J));
-	UI_COMMAND(ExportNodes, "Export nodes to AssetData", "Exports selected nodes to separate and reusable PCGSettings assets.", EUserInterfaceActionType::Button, FInputChord());
+	UI_COMMAND(ExportNodes, "Export nodes to PCGSettings", "Exports selected nodes to separate and reusable PCGSettings assets.", EUserInterfaceActionType::Button, FInputChord());
 	UI_COMMAND(ConvertToStandaloneNodes, "Convert to standalone Nodes", "Converts instanced nodes to standalone nodes.", EUserInterfaceActionType::Button, FInputChord());
 	UI_COMMAND(Find, "Find", "Finds PCG nodes and comments in the current graph.", EUserInterfaceActionType::Button, FInputChord(EModifierKey::Control, EKeys::F));
 	UI_COMMAND(ShowSelectedDetails, "Show Node Details", "Opens a details panel for the selected nodes.", EUserInterfaceActionType::Button, FInputChord(EKeys::F4));
diff --git a/Engine/Plugins/PCG/Source/PCGEditor/Private/PCGEditorGraphNodeBase.cpp b/Engine/Plugins/PCG/Source/PCGEditor/Private/PCGEditorGraphNodeBase.cpp
--- a/Engine/Plugins/PCG/Source/PCGEditor/Private/PCGEditorGraphNodeBase.cpp
+++ b/Engine/Plugins/PCG/Source/PCGEditor/Private/PCGEditorGraphNodeBase.cpp
@@ -19,6 +19,7 @@
 #include "PCGEditorGraph.h"
 #include "PCGEditorGraphSchema.h"
 #include "PCGEditorSettings.h"
+#include "PCGEditorStyle.h"
 
 #include "GraphEditorActions.h"
 #include "ScopedTransaction.h"
@@ -198,7 +199,7 @@ void UPCGEditorGraphNodeBase::GetNodeContextMenuActions(UToolMenu* Menu, class U
 		Section.AddMenuEntry(FPCGEditorCommands::Get().ToggleInspect, LOCTEXT("ToggleinspectionLabel", "Inspect"));
 		Section.AddMenuEntry(FGraphEditorCommands::Get().BreakNodeLinks);
 		Section.AddMenuEntry(FPCGEditorCommands::Get().CollapseNodes);
-		Section.AddMenuEntry(FPCGEditorCommands::Get().ExportNodes);
+		Section.AddMenuEntry(FPCGEditorCommands::Get().ExportNodes, TAttribute<FText>(), TAttribute<FText>(), FSlateIcon(FPCGEditorStyle::Get().GetStyleSetName(), "ClassIcon.PCGSettings"));
 		Section.AddMenuEntry(FPCGEditorCommands::Get().ConvertToStandaloneNodes);
 		Section.AddMenuEntry(FPCGEditorCommands::Get().RenameNode, LOCTEXT("RenameNode", "Rename"));
 
diff --git a/Engine/Plugins/PCG/Source/PCGEditor/Private/PCGEditorMenuUtils.cpp b/Engine/Plugins/PCG/Source/PCGEditor/Private/PCGEditorMenuUtils.cpp
--- a/Engine/Plugins/PCG/Source/PCGEditor/Private/PCGEditorMenuUtils.cpp
+++ b/Engine/Plugins/PCG/Source/PCGEditor/Private/PCGEditorMenuUtils.cpp
@@ -3,6 +3,7 @@
 #include "PCGEditorMenuUtils.h"
 
 #include "PCGAssetExporterUtils.h"
+#include "PCGEditorStyle.h"
 #include "PCGLevelToAsset.h"
 
 #include "Engine/World.h"
@@ -63,7 +64,7 @@ namespace PCGEditorMenuUtils
 				"CreatePCGAssetFromMenu",
 				LOCTEXT("CreatePCGAssetFromMenu", "Create PCG Assets from Level(s)"),
 				TAttribute<FText>(),
-				FSlateIcon(),
+				FSlateIcon(FPCGEditorStyle::Get().GetStyleSetName(), "ClassIcon.PCGDataAsset"),
 				UIAction);
 		}
 
@@ -80,7 +81,7 @@ namespace PCGEditorMenuUtils
 				"UpdatePCGAssetFromMenu",
 				LOCTEXT("UpdatePCGAssetFromMenu", "Update PCG Assets"),
 				TAttribute<FText>(),
-				FSlateIcon(),
+				FSlateIcon(FPCGEditorStyle::Get().GetStyleSetName(), "ClassIcon.PCGDataAsset"),
 				UIAction);
 		}
 
diff --git a/Engine/Plugins/PCG/Source/PCGEditor/Private/PCGEditorModule.cpp b/Engine/Plugins/PCG/Source/PCGEditor/Private/PCGEditorModule.cpp
--- a/Engine/Plugins/PCG/Source/PCGEditor/Private/PCGEditorModule.cpp
+++ b/Engine/Plugins/PCG/Source/PCGEditor/Private/PCGEditorModule.cpp
@@ -187,7 +187,9 @@ void FPCGEditorModule::RegisterMenuExtensions()
 			"PCGToolsSubMenu",
 			LOCTEXT("PCGSubMenu", "PCG Framework"),
 			LOCTEXT("PCGSubMenu_Tooltip", "Procedural Content Generation (PCG) Framework related functionality"),
-			FNewMenuDelegate::CreateRaw(this, &FPCGEditorModule::PopulateMenuActions));
+			FNewMenuDelegate::CreateRaw(this, &FPCGEditorModule::PopulateMenuActions),
+			/*bInOpenSubMenuOnClick=*/false,
+			FSlateIcon(FPCGEditorStyle::Get().GetStyleSetName(), "PCG.EditorIcon"));
 	}
 
 	if (UToolMenu* WorldAssetMenu = UToolMenus::Get()->ExtendMenu("ContentBrowser.AssetContextMenu.AssetActionsSubMenu"))
diff --git a/Engine/Plugins/PCG/Source/PCGEditor/Private/PCGEditorStyle.cpp b/Engine/Plugins/PCG/Source/PCGEditor/Private/PCGEditorStyle.cpp
--- a/Engine/Plugins/PCG/Source/PCGEditor/Private/PCGEditorStyle.cpp
+++ b/Engine/Plugins/PCG/Source/PCGEditor/Private/PCGEditorStyle.cpp
@@ -24,7 +24,10 @@ void FPCGEditorStyle::Unregister()
 
 FPCGEditorStyle::FPCGEditorStyle() : FSlateStyleSet("PCGEditorStyle")
 {
+	static const FVector2D Icon16x16(16.0f, 16.0f);
 	static const FVector2D Icon20x20(20.0f, 20.0f);
+	static const FVector2D Icon64x64(64.0f, 64.0f);
+	static const FVector2D Icon128x128(128.0f, 128.0f);
 	
 	SetCoreContentRoot(FPaths::EngineContentDir() / TEXT("Editor/Slate"));
 	SetContentRoot(IPluginManager::Get().FindPlugin("PCG")->GetBaseDir() / TEXT("Content"));
@@ -85,6 +88,26 @@ FPCGEditorStyle::FPCGEditorStyle() : FSlateStyleSet("PCGEditorStyle")
 		PCGEditorStyleConstants::Node_Overlay_GridSizeLabel_BorderRadius,
 		FLinearColor::Black,
 		PCGEditorStyleConstants::Node_Overlay_GridSizeLabel_BorderStroke));
+
+	// Plugin icon/editor/component icons
+	Set("ClassIcon.PCGComponent", new IMAGE_BRUSH_SVG("Icons/PCG_16", Icon16x16));
+	Set("ClassThumbnail.PCGComponent", new IMAGE_BRUSH_SVG("Icons/PCG_64", Icon64x64));
+	Set("PCG.PluginIcon", new IMAGE_BRUSH_SVG("Icons/PCG_128", Icon128x128));
+	Set("PCG.EditorIcon", new IMAGE_BRUSH_SVG("Icons/PCG_16", Icon16x16));
+
+	// Asset/Class icons
+	Set("ClassIcon.PCGVolume", new IMAGE_BRUSH_SVG("Icons/PCG_16", Icon16x16));
+	Set("ClassThumbnail.PCGVolume", new IMAGE_BRUSH_SVG("Icons/PCG_64", Icon64x64));
+	Set("ClassIcon.PCGDataAsset", new IMAGE_BRUSH_SVG("Icons/PCGAsset_16", Icon16x16));
+	Set("ClassThumbnail.PCGDataAsset", new IMAGE_BRUSH_SVG("Icons/PCGAsset_16", Icon64x64));
+	Set("ClassIcon.PCGGraphInterface", new IMAGE_BRUSH_SVG("Icons/PCGGraph_16", Icon16x16));
+	Set("ClassThumbnail.PCGGraphInterface", new IMAGE_BRUSH_SVG("Icons/PCGGraph_16", Icon64x64));
+	Set("ClassIcon.PCGGraph", new IMAGE_BRUSH_SVG("Icons/PCGGraph_16", Icon16x16));
+	Set("ClassThumbnail.PCGGraph", new IMAGE_BRUSH_SVG("Icons/PCGGraph_16", Icon64x64));
+	Set("ClassIcon.PCGGraphInstance", new IMAGE_BRUSH_SVG("Icons/PCGGraphInstance_16", Icon16x16));
+	Set("ClassThumbnail.PCGGraphInstance", new IMAGE_BRUSH_SVG("Icons/PCGGraphInstance_16", Icon64x64));
+	Set("ClassIcon.PCGSettings", new IMAGE_BRUSH_SVG("Icons/PCGSettings_16", Icon16x16));
+	Set("ClassThumbnail.PCGSettings", new IMAGE_BRUSH_SVG("Icons/PCGSettings_16", Icon64x64));
 }
 
 const FPCGEditorStyle& FPCGEditorStyle::Get()
