[PCG] Add safety cleanup in subsystem shutdown to remove any dangling callbacks

#rnx

diff --git a/Engine/Plugins/PCG/Source/PCG/Private/PCGActorAndComponentMapping.cpp b/Engine/Plugins/PCG/Source/PCG/Private/PCGActorAndComponentMapping.cpp
--- a/Engine/Plugins/PCG/Source/PCG/Private/PCGActorAndComponentMapping.cpp
+++ b/Engine/Plugins/PCG/Source/PCG/Private/PCGActorAndComponentMapping.cpp
@@ -1180,6 +1180,23 @@ void FPCGActorAndComponentMapping::TeardownTrackingCallbacks()
 		World->PersistentLevel->OnLoadedActorAddedToLevelEvent.RemoveAll(this);
 		World->PersistentLevel->OnLoadedActorRemovedFromLevelEvent.RemoveAll(this);
 	}
+
+	TSet<UPCGComponent*> StillRegisteredComponents = GetAllRegisteredComponents();
+	// FIXME: When the changes to unregister will go in, we can enable this ensure. At the moment, we didn't unregister when the world is shutdown
+	// because the subsystem is dead before we destroy the component. So it is expected to have still all of those components registered at that point.
+	//if (!ensureMsgf(StillRegisteredComponents.IsEmpty(), TEXT("Found components not unregistered during PCG subsystem shutdown.")))
+	{
+		for (UPCGComponent* Component : StillRegisteredComponents)
+		{
+			if (IsValid(Component))
+			{
+				Component->OnPCGGraphGeneratedDelegate.RemoveAll(this);
+				Component->OnPCGGraphCleanedDelegate.RemoveAll(this);
+				Component->OnPCGGraphStartGeneratingDelegate.RemoveAll(this);
+				Component->OnPCGGraphCancelledDelegate.RemoveAll(this);
+			}
+		}
+	}
 }
 
 void FPCGActorAndComponentMapping::AddDelayedActors()
