[PCG] Updated class icon for PCG-related actors. Added default wire thickness settings in PCG editor.

diff --git a/Engine/Plugins/PCG/Source/PCG/Public/PCGVolume.h b/Engine/Plugins/PCG/Source/PCG/Public/PCGVolume.h
--- a/Engine/Plugins/PCG/Source/PCG/Public/PCGVolume.h
+++ b/Engine/Plugins/PCG/Source/PCG/Public/PCGVolume.h
@@ -19,6 +19,7 @@ public:
 	//~ Begin AActor Interface
 #if WITH_EDITOR
 	virtual bool GetReferencedContentObjects(TArray<UObject*>& Objects) const override;
+	virtual FName GetCustomIconName() const { return NAME_None; }
 #endif // WITH_EDITOR
 	//~ End AActor Interface
 
diff --git a/Engine/Plugins/PCG/Source/PCGEditor/Private/PCGEditorGraphSchema.cpp b/Engine/Plugins/PCG/Source/PCGEditor/Private/PCGEditorGraphSchema.cpp
--- a/Engine/Plugins/PCG/Source/PCGEditor/Private/PCGEditorGraphSchema.cpp
+++ b/Engine/Plugins/PCG/Source/PCGEditor/Private/PCGEditorGraphSchema.cpp
@@ -822,6 +822,8 @@ void FPCGEditorConnectionDrawingPolicy::DetermineWiringStyle(UEdGraphPin* Output
 {
 	FConnectionDrawingPolicy::DetermineWiringStyle(OutputPin, InputPin, Params);
 
+	Params.WireThickness = GetDefault<UPCGEditorSettings>()->DefaultWireThickness;
+
 	// Emphasize wire thickness on hovered pins
 	if (HoveredPins.Contains(InputPin) && HoveredPins.Contains(OutputPin))
 	{
diff --git a/Engine/Plugins/PCG/Source/PCGEditor/Private/PCGEditorStyle.cpp b/Engine/Plugins/PCG/Source/PCGEditor/Private/PCGEditorStyle.cpp
--- a/Engine/Plugins/PCG/Source/PCGEditor/Private/PCGEditorStyle.cpp
+++ b/Engine/Plugins/PCG/Source/PCGEditor/Private/PCGEditorStyle.cpp
@@ -98,16 +98,21 @@ FPCGEditorStyle::FPCGEditorStyle() : FSlateStyleSet("PCGEditorStyle")
 	// Asset/Class icons
 	Set("ClassIcon.PCGVolume", new IMAGE_BRUSH_SVG("Icons/PCG_16", Icon16x16));
 	Set("ClassThumbnail.PCGVolume", new IMAGE_BRUSH_SVG("Icons/PCG_64", Icon64x64));
+	Set("ClassIcon.PCGWorldActor", new IMAGE_BRUSH_SVG("Icons/PCG_16", Icon16x16));
+	Set("ClassThumbnail.PCGWorldActor", new IMAGE_BRUSH_SVG("Icons/PCG_64", Icon64x64));
+	Set("ClassIcon.PCGPartitionActor", new IMAGE_BRUSH_SVG("Icons/PCG_16", Icon16x16));
+	Set("ClassThumbnail.PCGPartitionActor", new IMAGE_BRUSH_SVG("Icons/PCG_64", Icon64x64));
+
 	Set("ClassIcon.PCGDataAsset", new IMAGE_BRUSH_SVG("Icons/PCGAsset_16", Icon16x16));
-	Set("ClassThumbnail.PCGDataAsset", new IMAGE_BRUSH_SVG("Icons/PCGAsset_16", Icon64x64));
+	Set("ClassThumbnail.PCGDataAsset", new IMAGE_BRUSH_SVG("Icons/PCGAsset_64", Icon64x64));
 	Set("ClassIcon.PCGGraphInterface", new IMAGE_BRUSH_SVG("Icons/PCGGraph_16", Icon16x16));
-	Set("ClassThumbnail.PCGGraphInterface", new IMAGE_BRUSH_SVG("Icons/PCGGraph_16", Icon64x64));
+	Set("ClassThumbnail.PCGGraphInterface", new IMAGE_BRUSH_SVG("Icons/PCGGraph_64", Icon64x64));
 	Set("ClassIcon.PCGGraph", new IMAGE_BRUSH_SVG("Icons/PCGGraph_16", Icon16x16));
-	Set("ClassThumbnail.PCGGraph", new IMAGE_BRUSH_SVG("Icons/PCGGraph_16", Icon64x64));
+	Set("ClassThumbnail.PCGGraph", new IMAGE_BRUSH_SVG("Icons/PCGGraph_64", Icon64x64));
 	Set("ClassIcon.PCGGraphInstance", new IMAGE_BRUSH_SVG("Icons/PCGGraphInstance_16", Icon16x16));
-	Set("ClassThumbnail.PCGGraphInstance", new IMAGE_BRUSH_SVG("Icons/PCGGraphInstance_16", Icon64x64));
+	Set("ClassThumbnail.PCGGraphInstance", new IMAGE_BRUSH_SVG("Icons/PCGGraphInstance_64", Icon64x64));
 	Set("ClassIcon.PCGSettings", new IMAGE_BRUSH_SVG("Icons/PCGSettings_16", Icon16x16));
-	Set("ClassThumbnail.PCGSettings", new IMAGE_BRUSH_SVG("Icons/PCGSettings_16", Icon64x64));
+	Set("ClassThumbnail.PCGSettings", new IMAGE_BRUSH_SVG("Icons/PCGSettings_64", Icon64x64));
 }
 
 const FPCGEditorStyle& FPCGEditorStyle::Get()
diff --git a/Engine/Plugins/PCG/Source/PCGEditor/Public/PCGEditorSettings.h b/Engine/Plugins/PCG/Source/PCGEditor/Public/PCGEditorSettings.h
--- a/Engine/Plugins/PCG/Source/PCGEditor/Public/PCGEditorSettings.h
+++ b/Engine/Plugins/PCG/Source/PCGEditor/Public/PCGEditorSettings.h
@@ -149,6 +149,10 @@ public:
 	UPROPERTY(EditAnywhere, config, Category = Node, meta = (HideAlphaChannel))
 	TMap<TSubclassOf<UPCGSettings>, FLinearColor> OverrideNodeColorByClass;
 
+	/** Default wire thickness */
+	UPROPERTY(EditAnywhere, config, Category = Wire, meta = (ClampMin="1", ClampMap="9"))
+	float DefaultWireThickness = 1.5f;
+
 	/** Hover edge size emphasis */
 	UPROPERTY(EditAnywhere, config, Category = Wire, meta = (ClampMin="1", ClampMax="5"))
 	float HoverEdgeEmphasis = 1.5f;
