[PCG] Fixed a bug causing the Volume Sampler to sample incorrectly

diff --git a/Engine/Plugins/PCG/Source/PCG/Private/Elements/PCGVolumeSampler.cpp b/Engine/Plugins/PCG/Source/PCG/Private/Elements/PCGVolumeSampler.cpp
--- a/Engine/Plugins/PCG/Source/PCG/Private/Elements/PCGVolumeSampler.cpp
+++ b/Engine/Plugins/PCG/Source/PCG/Private/Elements/PCGVolumeSampler.cpp
@@ -60,9 +60,9 @@ namespace PCGVolumeSampler
 		int32 NumIterations = -1;
 
 		{
-			const int64 NumX = MaxX - MinX;
-			const int64 NumY = MaxY - MinY;
-			const int64 NumZ = MaxZ - MinZ;
+			const int64 NumX = 1 + MaxX - MinX;
+			const int64 NumY = 1 + MaxY - MinY;
+			const int64 NumZ = 1 + MaxZ - MinZ;
 			const int64 NumIterationsXY64 = NumX * NumY;
 			const int64 NumIterations64 = NumIterationsXY64 * NumZ;
 
@@ -94,9 +94,9 @@ namespace PCGVolumeSampler
 
 		auto AsyncProcessingFunc = [Volume, BoundingShape, PointSteepness = SamplerSettings.PointSteepness, VoxelSize, MinX, MaxX, MinY, MaxY, MinZ](int32 Index, FPCGPoint& OutPoint)
 		{
-			const int X = MinX + (Index % (MaxX - MinX));
-			const int Y = MinY + (Index / (MaxX - MinX) % (MaxY - MinY));
-			const int Z = MinZ + (Index / ((MaxX - MinX) * (MaxY - MinY)));
+			const int X = MinX + (Index % (1 + MaxX - MinX));
+			const int Y = MinY + (Index / (1 + MaxX - MinX) % (1 + MaxY - MinY));
+			const int Z = MinZ + (Index / ((1 + MaxX - MinX) * (1 + MaxY - MinY)));
 
 			const FVector SampleLocation(X * VoxelSize.X, Y * VoxelSize.Y, Z * VoxelSize.Z);
 			const FBox VoxelBox(VoxelSize * -0.5, VoxelSize * 0.5);
