[PCG] Fix unreachable code

#rnx

diff --git a/Engine/Plugins/PCG/Source/PCGEditor/Private/SPCGEditorGraphAttributeListView.cpp b/Engine/Plugins/PCG/Source/PCGEditor/Private/SPCGEditorGraphAttributeListView.cpp
--- a/Engine/Plugins/PCG/Source/PCGEditor/Private/SPCGEditorGraphAttributeListView.cpp
+++ b/Engine/Plugins/PCG/Source/PCGEditor/Private/SPCGEditorGraphAttributeListView.cpp
@@ -273,6 +273,7 @@ bool FPCGPointFilterExpressionContext::TestComplexExpression(const FName& InKey,
 				if (PCGColumnInfo->DataAccessor->Get<ValueType>(Value, Index, *PCGColumnInfo->DataKeys))
 				{
 					FText TextValue;
+					bool bInvalid = false;
 					if constexpr (PCG::Private::IsOfTypes<ValueType, bool>())
 					{
 						TextValue = FText::FromString(LexToString(Value));
@@ -292,12 +293,16 @@ bool FPCGPointFilterExpressionContext::TestComplexExpression(const FName& InKey,
 					else
 					{
 						ensureMsgf(false, TEXT("Unsupported Data Type"));
-						return false;
+						bInvalid = true;
 					}
 
-					const FTextFilterString PointValue(TextValue.ToString());
-					return TextFilterUtils::TestComplexExpression(PointValue, InValue, InComparisonOperation, InTextComparisonMode);
+					if (!bInvalid)
+					{
+						const FTextFilterString PointValue(TextValue.ToString());
+						return TextFilterUtils::TestComplexExpression(PointValue, InValue, InComparisonOperation, InTextComparisonMode);
+					}
 				}
+
 				return false;
 			};
 
