[PCG] PCGWorldPartitionBuilder - return success value if no components found, otherwise this reports a failed build.

diff --git a/Engine/Plugins/PCG/Source/PCGEditor/Private/WorldPartitionBuilder/PCGWorldPartitionBuilder.cpp b/Engine/Plugins/PCG/Source/PCGEditor/Private/WorldPartitionBuilder/PCGWorldPartitionBuilder.cpp
--- a/Engine/Plugins/PCG/Source/PCGEditor/Private/WorldPartitionBuilder/PCGWorldPartitionBuilder.cpp
+++ b/Engine/Plugins/PCG/Source/PCGEditor/Private/WorldPartitionBuilder/PCGWorldPartitionBuilder.cpp
@@ -238,7 +238,7 @@ bool UPCGWorldPartitionBuilder::RunInternal(UWorld* World, const FCellInfo& InCe
 	if (!bGeneratedAnyComponent)
 	{
 		UE_LOG(LogPCGWorldPartitionBuilder, Display, TEXT("Dirty package detection and save skipped due to trivial run"));
-		return false;
+		return !bErrorOccurredWhileGenerating;
 	}
 
 	// TODO: Review the save flow when we have iterative loading.
