[PCG] Remove hard ref to the PCG Graph in PCGUserParametersData + fix refresh issue

* UserParametersData can (and will) be cached because it is a UPCGData that flows.
* Meaning that we hold a hard ref pretty much forever to the OriginalGraph, which is bad
* Instead, we just use the user struct pass in data directly.
* Also when a graph was force deleted, the volume was not refreshed

#rnx

diff --git a/Engine/Plugins/PCG/Source/PCG/Private/Elements/PCGUserParameterGet.cpp b/Engine/Plugins/PCG/Source/PCG/Private/Elements/PCGUserParameterGet.cpp
--- a/Engine/Plugins/PCG/Source/PCG/Private/Elements/PCGUserParameterGet.cpp
+++ b/Engine/Plugins/PCG/Source/PCG/Private/Elements/PCGUserParameterGet.cpp
@@ -32,7 +32,7 @@ namespace PCGUserParameterGetSettings
 		FConstStructView GraphParameters{};
 		const UPCGGraphInterface* GraphInterface = nullptr;
 
-		// First we will read from the input if we find an input for this graph, set by the subgraph element.
+		// First we will read from the input if we find an input for this graph, set by the subgraph element, then use it.
 		TArray<FPCGTaggedData> UserParameterData = InContext.InputData.GetTaggedTypedInputs<UPCGUserParametersData>(PCGBaseSubgraphConstants::UserParameterTagData);
 		if (!UserParameterData.IsEmpty())
 		{
@@ -40,23 +40,15 @@ namespace PCGUserParameterGetSettings
 			// Safe guard to make sure we always have one and only one data of this type.
 			ensure(UserParameterData.Num() == 1);
 #endif // WITH_EDITOR
-			if (const UPCGUserParametersData* OverrideParametersData = CastChecked<UPCGUserParametersData>(UserParameterData[0].Data))
+			if (const UPCGUserParametersData* OverrideParametersData = CastChecked<UPCGUserParametersData>(UserParameterData[0].Data, ECastCheckedType::NullAllowed))
 			{
-				GraphInterface = OverrideParametersData->OriginalGraph;
-				GraphParameters = FConstStructView{ OverrideParametersData->UserParameters };
+				return FConstStructView{ OverrideParametersData->UserParameters };
+			}
+			else
+			{
+				PCGE_LOG_C(Error, LogOnly, &InContext, LOCTEXT("InvalidUserParameterData", "Internal error, PCG User Parameters Data is null"));
+				return FConstStructView{};
 			}
-		}
-
-		// Then if we don't have any input, we will use the graph instance from the component.
-		if (!GraphInterface)
-		{
-			const UPCGComponent* SourceComponent = InContext.SourceComponent.Get();
-			GraphInterface = SourceComponent ? SourceComponent->GetGraphInstance() : nullptr;
-		}
-
-		if (!GraphInterface)
-		{
-			return FConstStructView{};
 		}
 
 		// We gather the outer graph of this node, to make sure it matches our interface.
@@ -65,11 +57,13 @@ namespace PCGUserParameterGetSettings
 		FConstStructView GraphFromNodeParameters = GraphFromNode->GetUserParametersStruct() ? GraphFromNode->GetUserParametersStruct()->GetValue() : FConstStructView{};
 		if (!GraphFromNodeParameters.IsValid())
 		{
+			PCGE_LOG_C(Error, LogOnly, &InContext, LOCTEXT("NoBaseGraphParameters", "Node's graph owner don't have any parameters."));
 			return FConstStructView{};
 		}
 
 		// If we don't have a graph instance, we just use the user parameters from the node graph owner.
-		const UPCGGraphInstance* GraphInstance = Cast<UPCGGraphInstance>(GraphInterface);
+		const UPCGComponent* SourceComponent = InContext.SourceComponent.Get();
+		const UPCGGraphInstance* GraphInstance = SourceComponent ? SourceComponent->GetGraphInstance() : nullptr;
 		if (!GraphInstance)
 		{
 			return GraphFromNodeParameters;
@@ -78,6 +72,7 @@ namespace PCGUserParameterGetSettings
 		// Making sure the graph matches.
 		if (GraphFromNode != GraphInstance->GetGraph())
 		{
+			PCGE_LOG_C(Error, LogOnly, &InContext, LOCTEXT("GraphDontMatch", "Node's graph owner and base graph of the source component doesn't match."));
 			return FConstStructView{};
 		}
 
diff --git a/Engine/Plugins/PCG/Source/PCG/Private/PCGGraph.cpp b/Engine/Plugins/PCG/Source/PCG/Private/PCGGraph.cpp
--- a/Engine/Plugins/PCG/Source/PCG/Private/PCGGraph.cpp
+++ b/Engine/Plugins/PCG/Source/PCG/Private/PCGGraph.cpp
@@ -1730,7 +1730,6 @@ void UPCGGraphInstance::PreEditChange(FProperty* InProperty)
 	if (InProperty->GetOwnerClass() == UPCGGraphInstance::StaticClass() && InProperty->GetFName() == GET_MEMBER_NAME_CHECKED(UPCGGraphInstance, Graph))
 	{
 		PreGraphCache = Graph;
-
 		TeardownCallbacks();
 	}
 }
@@ -1753,16 +1752,19 @@ void UPCGGraphInstance::PostEditChangeProperty(FPropertyChangedEvent& PropertyCh
 		if (Graph && !CanGraphInterfaceBeSet(Graph))
 		{
 			UE_LOG(LogPCG, Error, TEXT("Attempting to assign %s would cause infinite recursion in the graph instance hierarchy, this is not allowed."), *Graph->GetPathName());
-			Graph = PreGraphCache;
+			Graph = PreGraphCache.Get();
 		}
 
 		SetupCallbacks();
 
-		// No need to refresh if it is the same graph
-		if (Graph != PreGraphCache)
+		// No need to refresh if it is the same graph, but we need to refresh if we have no graph anymore, but the pre graph was valid (but isn't anymore like in a Force Delete Asset)
+		if (Graph != PreGraphCache || (!Graph && !PreGraphCache.IsExplicitlyNull()))
 		{
 			OnGraphParametersChanged(Graph, EPCGGraphParameterEvent::GraphChanged, NAME_None);
 		}
+
+		// Reset them there to avoid any side effect if Pre/Post are called multiple times for the same transaction.
+		PreGraphCache = nullptr;
 	}
 	else if (PropertyName == GET_MEMBER_NAME_CHECKED(FPCGOverrideInstancedPropertyBag, PropertiesIDsOverridden))
 	{
@@ -1772,8 +1774,6 @@ void UPCGGraphInstance::PostEditChangeProperty(FPropertyChangedEvent& PropertyCh
 	{
 		OnGraphParametersChanged(this, EPCGGraphParameterEvent::ValueModifiedLocally, PropertyChangedEvent.GetMemberPropertyName());
 	}
-
-	PreGraphCache = nullptr;
 }
 
 void UPCGGraphInstance::PreEditUndo()
diff --git a/Engine/Plugins/PCG/Source/PCG/Private/PCGSubgraph.cpp b/Engine/Plugins/PCG/Source/PCG/Private/PCGSubgraph.cpp
--- a/Engine/Plugins/PCG/Source/PCG/Private/PCGSubgraph.cpp
+++ b/Engine/Plugins/PCG/Source/PCG/Private/PCGSubgraph.cpp
@@ -563,11 +563,20 @@ void FPCGSubgraphElement::PrepareSubgraphUserParameters(const UPCGSubgraphSettin
 	if (const UPCGGraphInterface* SubgraphInterface = Settings->GetSubgraphInterface())
 	{
 		UPCGUserParametersData* UserParamData = NewObject<UPCGUserParametersData>();
-		UserParamData->OriginalGraph = SubgraphInterface;
+
 		if (Context->GraphInstanceParametersOverride.IsValid())
 		{
 			UserParamData->UserParameters = std::move(Context->GraphInstanceParametersOverride);
 		}
+		else if (const FInstancedPropertyBag* InstancedPropertyBag = SubgraphInterface->GetUserParametersStruct())
+		{
+			// FIXME: Copy is done there.
+			UserParamData->UserParameters = InstancedPropertyBag->GetValue();
+		}
+		else
+		{
+			// Do nothing, we still want to have a User Parameter Data to indicate we are in a subgraph context.
+		}
 
 		FPCGTaggedData& TaggedData = OutputData.TaggedData.Emplace_GetRef();
 		TaggedData.Data = UserParamData;
diff --git a/Engine/Plugins/PCG/Source/PCG/Public/Data/PCGUserParametersData.h b/Engine/Plugins/PCG/Source/PCG/Public/Data/PCGUserParametersData.h
--- a/Engine/Plugins/PCG/Source/PCG/Public/Data/PCGUserParametersData.h
+++ b/Engine/Plugins/PCG/Source/PCG/Public/Data/PCGUserParametersData.h
@@ -5,7 +5,6 @@
 #include "PCGData.h"
 
 #include "InstancedStruct.h"
-#include "UObject/ObjectPtr.h"
 
 #include "PCGUserParametersData.generated.h"
 
@@ -29,9 +28,6 @@ public:
 	virtual void AddToCrc(FArchiveCrc32& Ar, bool bFullDataCrc) const override { AddUIDToCrc(Ar); }
 	// ~End UPCGData interface
 
-	UPROPERTY()
-	TObjectPtr<const UPCGGraphInterface> OriginalGraph = nullptr;
-
 	UPROPERTY()
 	FInstancedStruct UserParameters;
 };
\ No newline at end of file
diff --git a/Engine/Plugins/PCG/Source/PCG/Public/PCGGraph.h b/Engine/Plugins/PCG/Source/PCG/Public/PCGGraph.h
--- a/Engine/Plugins/PCG/Source/PCG/Public/PCGGraph.h
+++ b/Engine/Plugins/PCG/Source/PCG/Public/PCGGraph.h
@@ -514,7 +514,7 @@ public:
 private:
 #if WITH_EDITORONLY_DATA
 	// Transient, to keep track of the previous graph when it changed.
-	UPCGGraphInterface* PreGraphCache = nullptr;
+	TWeakObjectPtr<UPCGGraphInterface> PreGraphCache = nullptr;
 #endif // WITH_EDITORONLY_DATA
 };
 
