Added basic validation to prevent polluting the world with events when the state of the current component becomes invalid

#rnx

diff --git a/Engine/Plugins/PCG/Source/PCG/Private/PCGComponent.cpp b/Engine/Plugins/PCG/Source/PCG/Private/PCGComponent.cpp
--- a/Engine/Plugins/PCG/Source/PCG/Private/PCGComponent.cpp
+++ b/Engine/Plugins/PCG/Source/PCG/Private/PCGComponent.cpp
@@ -466,7 +466,7 @@ void UPCGComponent::PostProcessGraph(const FBox& InNewBounds, bool bInGenerated,
 #if WITH_EDITOR
 	const bool bHasGeneratedOutputAfter = GeneratedGraphOutput.TaggedData.Num() > 0;
 
-	if (bHasGeneratedOutputAfter || bHadGeneratedOutputBefore)
+	if (IsValid(this) && (bHasGeneratedOutputAfter || bHadGeneratedOutputBefore))
 	{
 		FProperty* GeneratedOutputProperty = FindFProperty<FProperty>(UPCGComponent::StaticClass(), GET_MEMBER_NAME_CHECKED(UPCGComponent, GeneratedGraphOutput));
 		check(GeneratedOutputProperty);
