[PCG] Fixed an issue where a static mesh spawner stopped mid-execution (because of cancellation) would leave resources in a state that's not correctly reset in subsequent exections.
Simplified code back in managed ISM resource usage so we don't play with the root multiple times per execution
Fixed typos

#rnx

diff --git a/Engine/Plugins/PCG/Source/PCG/Private/Elements/PCGStaticMeshSpawner.cpp b/Engine/Plugins/PCG/Source/PCG/Private/Elements/PCGStaticMeshSpawner.cpp
--- a/Engine/Plugins/PCG/Source/PCG/Private/Elements/PCGStaticMeshSpawner.cpp
+++ b/Engine/Plugins/PCG/Source/PCG/Private/Elements/PCGStaticMeshSpawner.cpp
@@ -418,10 +418,14 @@ void FPCGStaticMeshSpawnerElement::SpawnStaticMeshInstances(FPCGStaticMeshSpawne
 	check(Settings);
 
 	UPCGManagedISMComponent* MISMC = UPCGActorHelpers::GetOrCreateManagedISMC(TargetActor, Context->SourceComponent.Get(), Settings->UID, Params);
-	
+
 	check(MISMC);
 	MISMC->SetCrc(Context->DependenciesCrc);
 
+	// Keep track of all touched resources in the context, because if the execution is cancelled during the SMS execution
+	// we cannot easily guarantee that the state (esp. vs CRCs) is going to be entirely valid
+	Context->TouchedResources.Emplace(MISMC);
+
 	UInstancedStaticMeshComponent* ISMC = MISMC->GetComponent();
 	check(ISMC);
 
@@ -456,6 +460,23 @@ void FPCGStaticMeshSpawnerElement::SpawnStaticMeshInstances(FPCGStaticMeshSpawne
 	}
 }
 
+void FPCGStaticMeshSpawnerElement::AbortInternal(FPCGContext* InContext) const
+{
+	TRACE_CPUPROFILER_EVENT_SCOPE(FPCGStaticMeshSpawnerElement::Execute);
+	FPCGStaticMeshSpawnerContext* Context = static_cast<FPCGStaticMeshSpawnerContext*>(InContext);
+
+	// Any resources we've touched during the execution of this node can potentially be in a "not-quite complete state" especially if we have multiple sources of data writing to the same ISMC.
+	// In this case, we're aiming to mark the resources as "Unused" so they are picked up to be removed during the component's OnProcessGraphAborted, which is why we call Release here.
+	for (TWeakObjectPtr<UPCGManagedISMComponent> ManagedResource : Context->TouchedResources)
+	{
+		if(ManagedResource.IsValid())
+		{
+			TSet<TSoftObjectPtr<AActor>> Dummy;
+			ManagedResource->Release(/*bHardRelease=*/false, Dummy);
+		}
+	}
+}
+
 void UPCGStaticMeshSpawnerSettings::PostLoad()
 {
 	Super::PostLoad();
diff --git a/Engine/Plugins/PCG/Source/PCG/Private/PCGManagedResource.cpp b/Engine/Plugins/PCG/Source/PCG/Private/PCGManagedResource.cpp
--- a/Engine/Plugins/PCG/Source/PCG/Private/PCGManagedResource.cpp
+++ b/Engine/Plugins/PCG/Source/PCG/Private/PCGManagedResource.cpp
@@ -26,14 +26,8 @@ void UPCGManagedResource::PostApplyToComponent()
 // By default, if it is not a hard release, we mark the resource unused.
 bool UPCGManagedResource::Release(bool bHardRelease, TSet<TSoftObjectPtr<AActor>>& /*OutActorsToDelete*/)
 {
-	if (!bHardRelease)
-	{
-		bIsMarkedUnused = true;
-		return false;
-	}
-
-	bIsMarkedUnused = false;
-	return true;
+	bIsMarkedUnused = true;
+	return bHardRelease;
 }
 
 bool UPCGManagedResource::ReleaseIfUnused(TSet<TSoftObjectPtr<AActor>>& OutActorsToDelete)
@@ -566,8 +560,14 @@ void UPCGManagedISMComponent::ResetComponent()
 
 void UPCGManagedISMComponent::MarkAsUsed()
 {
+	const bool bWasMarkedUnused = bIsMarkedUnused;
 	Super::MarkAsUsed();
 
+	if (!bWasMarkedUnused)
+	{
+		return;
+	}
+
 	if (UInstancedStaticMeshComponent* ISMC = GetComponent())
 	{
 		// Keep track of the current root location so if we reuse this later we are able to update this appropriately
diff --git a/Engine/Plugins/PCG/Source/PCG/Public/Elements/PCGStaticMeshSpawner.h b/Engine/Plugins/PCG/Source/PCG/Public/Elements/PCGStaticMeshSpawner.h
--- a/Engine/Plugins/PCG/Source/PCG/Public/Elements/PCGStaticMeshSpawner.h
+++ b/Engine/Plugins/PCG/Source/PCG/Public/Elements/PCGStaticMeshSpawner.h
@@ -138,6 +138,7 @@ protected:
 	virtual FPCGContext* CreateContext() override;
 	virtual bool PrepareDataInternal(FPCGContext* Context) const override;
 	virtual bool ExecuteInternal(FPCGContext* Context) const override;
+	virtual void AbortInternal(FPCGContext* Context) const override;
 	void SpawnStaticMeshInstances(FPCGStaticMeshSpawnerContext* Context, const FPCGMeshInstanceList& InstanceList, AActor* TargetActor, const FPCGPackedCustomData& PackedCustomData) const;
 };
 
diff --git a/Engine/Plugins/PCG/Source/PCG/Public/Elements/PCGStaticMeshSpawnerContext.h b/Engine/Plugins/PCG/Source/PCG/Public/Elements/PCGStaticMeshSpawnerContext.h
--- a/Engine/Plugins/PCG/Source/PCG/Public/Elements/PCGStaticMeshSpawnerContext.h
+++ b/Engine/Plugins/PCG/Source/PCG/Public/Elements/PCGStaticMeshSpawnerContext.h
@@ -3,6 +3,7 @@
 #pragma once
 
 #include "PCGContext.h"
+#include "PCGManagedResource.h"
 #include "Async/PCGAsyncLoadingContext.h"
 #include "MeshSelectors/PCGMeshSelectorBase.h"
 #include "InstanceDataPackers/PCGInstanceDataPackerBase.h"
@@ -73,5 +74,8 @@ struct FPCGStaticMeshSpawnerContext : public FPCGContext, public IPCGAsyncLoadin
 	TArray<TArray<int32>> AttributeOverridePartition;
 	TArray<FSoftISMComponentDescriptor> OverriddenDescriptors;
 
+	// Keeping track of all touched resources to allow for correct cleanup on abort
+	TArray<TWeakObjectPtr<UPCGManagedISMComponent>> TouchedResources;
+
 	void ResetInputIterationData();
 };
\ No newline at end of file
diff --git a/Engine/Plugins/PCG/Source/PCG/Public/PCGManagedResource.h b/Engine/Plugins/PCG/Source/PCG/Public/PCGManagedResource.h
--- a/Engine/Plugins/PCG/Source/PCG/Public/PCGManagedResource.h
+++ b/Engine/Plugins/PCG/Source/PCG/Public/PCGManagedResource.h
@@ -27,7 +27,7 @@ public:
 	/** Called when after a PCG component is applied to (such as after a RerunConstructionScript) */
 	virtual void PostApplyToComponent();
 
-	/** Releases/Mark Unsued the resource depending on the bHardRelease flag. Returns true if resource can be removed from the PCG component */
+	/** Releases/Mark Unused the resource depending on the bHardRelease flag. Returns true if resource can be removed from the PCG component */
 	virtual bool Release(bool bHardRelease, TSet<TSoftObjectPtr<AActor>>& OutActorsToDelete);
 	/** Releases resource if empty or unused. Returns true if the resource can be removed from the PCG component */
 	virtual bool ReleaseIfUnused(TSet<TSoftObjectPtr<AActor>>& OutActorsToDelete);
