[PCG] Remove invalid ensure from runtime gen scheduler.

Was triggered when a player generation source was not returning a result in SIE, where there was a player controller but no pawn. Doing it this way means if you change from SIE to PIE it will seamlessly work, and generally generation sources should be free to not return a result without triggering ensures.

diff --git a/Engine/Plugins/PCG/Source/PCG/Private/RuntimeGen/PCGRuntimeGenScheduler.cpp b/Engine/Plugins/PCG/Source/PCG/Private/RuntimeGen/PCGRuntimeGenScheduler.cpp
--- a/Engine/Plugins/PCG/Source/PCG/Private/RuntimeGen/PCGRuntimeGenScheduler.cpp
+++ b/Engine/Plugins/PCG/Source/PCG/Private/RuntimeGen/PCGRuntimeGenScheduler.cpp
@@ -414,9 +414,8 @@ void FPCGRuntimeGenScheduler::TickQueueComponentsForGeneration(
 			for (const IPCGGenSourceBase* GenSource : InGenSources)
 			{
 				const TOptional<FVector> GenSourcePositionOptional = GenSource->GetPosition();
-				if (!ensure(GenSourcePositionOptional.IsSet()))
+				if (!GenSourcePositionOptional.IsSet())
 				{
-					// If the gen source fails to return a valid result, skip it, to avoid mistakenly considering the origin as a valid gen source.
 					continue;
 				}
 
@@ -487,9 +486,8 @@ void FPCGRuntimeGenScheduler::TickCleanup(const TSet<IPCGGenSourceBase*>& InGenS
 				}
 
 				const TOptional<FVector> GenSourcePositionOptional = GenSource->GetPosition();
-				if (!ensure(GenSourcePositionOptional.IsSet()))
+				if (!GenSourcePositionOptional.IsSet())
 				{
-					// If the gen source fails to return a valid result, skip it, to avoid mistakenly considering the origin as a valid gen source.
 					continue;
 				}
 
