[PCG] Fix crash using runtime gen + SIE when the source component of a higen link is no longer valid, just no-op.

This is the first time I've seen this check() get hit and due to async nature of graph execution it's reasonable that a component providing data to a higen link is destroyed before the link executes.

diff --git a/Engine/Plugins/PCG/Source/PCG/Private/Graph/PCGGraphExecutor.cpp b/Engine/Plugins/PCG/Source/PCG/Private/Graph/PCGGraphExecutor.cpp
--- a/Engine/Plugins/PCG/Source/PCG/Private/Graph/PCGGraphExecutor.cpp
+++ b/Engine/Plugins/PCG/Source/PCG/Private/Graph/PCGGraphExecutor.cpp
@@ -1914,9 +1914,15 @@ namespace PCGGraphExecutor
 		const UPCGNode* InDownstreamNode,
 		FPCGGridLinkageContext* InContext)
 	{
-		check(InContext && InContext->SourceComponent.IsValid());
+		check(InContext);
 		TRACE_CPUPROFILER_EVENT_SCOPE(PCGGraphExecutor::ExecuteGridLinkage);
 
+		if (!InContext->SourceComponent.IsValid())
+		{
+			// Source no longer exists, nothing to be done.
+			return true;
+		}
+
 		// Non-hierarchical generation - no linkage required - data should just pass through.
 		if (!InContext->SourceComponent->GetGraph()->IsHierarchicalGenerationEnabled()
 			|| !ensure(PCGHiGenGrid::IsValidGrid(InFromGrid) || InFromGrid == EPCGHiGenGrid::Unbounded))
