[PCG] Tweak world ray hit UI to reduce confusion and document behaviour around ray params.

In next release we'll deprecate the bOverrideDefaultParams property.

diff --git a/Engine/Plugins/PCG/Source/PCG/Public/Data/PCGWorldData.h b/Engine/Plugins/PCG/Source/PCG/Public/Data/PCGWorldData.h
--- a/Engine/Plugins/PCG/Source/PCG/Public/Data/PCGWorldData.h
+++ b/Engine/Plugins/PCG/Source/PCG/Public/Data/PCGWorldData.h
@@ -122,16 +122,17 @@ struct FPCGWorldRayHitQueryParams : public FPCGWorldCommonQueryParams
 
 	void Initialize();
 
-	UPROPERTY(BlueprintReadWrite, EditAnywhere, Category = Data, meta = (PCG_Overridable))
+	/** Set ray parameters including origin, direction and length explicitly rather than deriving these from the generating actor bounds. */
+	UPROPERTY(BlueprintReadWrite, EditAnywhere, Category = Data, meta = (DisplayName = "Set Ray Parameters", PCG_Overridable))
 	bool bOverrideDefaultParams = false;
 
-	UPROPERTY(BlueprintReadWrite, EditAnywhere, Category = Data, meta = (PCG_Overridable, EditCondition = "bOverrideDefaultParams"))
+	UPROPERTY(BlueprintReadWrite, EditAnywhere, Category = Data, meta = (PCG_Overridable, EditCondition = "bOverrideDefaultParams", EditConditionHides))
 	FVector RayOrigin = FVector::ZeroVector;
 
-	UPROPERTY(BlueprintReadWrite, EditAnywhere, Category = Data, meta = (PCG_Overridable, EditCondition = "bOverrideDefaultParams"))
+	UPROPERTY(BlueprintReadWrite, EditAnywhere, Category = Data, meta = (PCG_Overridable, EditCondition = "bOverrideDefaultParams", EditConditionHides))
 	FVector RayDirection = FVector(0.0, 0.0, -1.0);
 
-	UPROPERTY(BlueprintReadWrite, EditAnywhere, Category = Data, meta = (PCG_Overridable, EditCondition = "bOverrideDefaultParams"))
+	UPROPERTY(BlueprintReadWrite, EditAnywhere, Category = Data, meta = (PCG_Overridable, EditCondition = "bOverrideDefaultParams", EditConditionHides))
 	double RayLength = 1.0e+5; // 100m
 
 	// TODO: see in FCollisionQueryParams if there are some flags we want to expose
