[PCG] FPCGPoint custom struct serialization to optimize size and speed

Alembic import of 562433 points 48mb file

Before this change (regular tagged property serialization of points) was generating a 413mb uasset that took 65.0 seconds to save
With this change (custom FPCGPoint structured archive serialize) generates a 161mb uasset that takes 2.4 seconds to save

Also in this change, serialize mask which skips over some members of the FPCGPoint with default values. This saves another 50mb on the test file bringing size down to 111mb and save to 1.5 seconds

#rnx

diff --git a/Engine/Plugins/PCG/Source/PCG/Private/PCGPoint.cpp b/Engine/Plugins/PCG/Source/PCG/Private/PCGPoint.cpp
--- a/Engine/Plugins/PCG/Source/PCG/Private/PCGPoint.cpp
+++ b/Engine/Plugins/PCG/Source/PCG/Private/PCGPoint.cpp
@@ -5,6 +5,7 @@
 #include "Math/BoxSphereBounds.h"
 #include "Metadata/Accessors/PCGCustomAccessor.h"
 #include "Metadata/Accessors/IPCGAttributeAccessorTpl.h"
+#include "UObject/FortniteReleaseBranchCustomObjectVersion.h"
 
 #include UE_INLINE_GENERATED_CPP_BY_NAME(PCGPoint)
 
@@ -17,6 +18,21 @@ namespace PCGPointCustomPropertyNames
 	const FName ScaleName = TEXT("Scale");
 }
 
+/** Serialized fields of a FPCGPoint, the values here can't change as they are being used to mask out serialization */
+enum class EPCGPointSerializeFields : uint8
+{
+	None = 0 ,
+	Density = 1 << 0,
+	BoundsMin = 1 << 1,
+	BoundsMax = 1 << 2,
+	Color = 1 << 3,
+	Steepness = 1 << 4,
+	Seed = 1 << 5,
+	MetadataEntry = 1 << 6
+};
+
+ENUM_CLASS_FLAGS(EPCGPointSerializeFields);
+
 FPCGPoint::FPCGPoint(const FTransform& InTransform, float InDensity, int32 InSeed)
 	: Transform(InTransform)
 	, Density(InDensity)
@@ -24,6 +40,103 @@ FPCGPoint::FPCGPoint(const FTransform& InTransform, float InDensity, int32 InSee
 {
 }
 
+bool FPCGPoint::Serialize(FStructuredArchive::FSlot Slot)
+{
+	FArchive& UnderlyingArchive = Slot.GetUnderlyingArchive();
+	
+	// Usage of a branch version instead of the PCG version ensures that we can't end up in a situation where two developers modify the PCG version in
+	// two different branches causing issues with saved assets in those branches when integrating.
+	UnderlyingArchive.UsingCustomVersion(FFortniteReleaseBranchCustomObjectVersion::GUID);
+
+	// Previous versions was using default serialization, returning false here ensures older data gets loaded through the default serialization
+	if (UnderlyingArchive.CustomVer(FFortniteReleaseBranchCustomObjectVersion::GUID) < FFortniteReleaseBranchCustomObjectVersion::PCGPointStructuredSerializer)
+	{
+		return false;
+	}
+	
+	const FPCGPoint Default;
+	EPCGPointSerializeFields SerializeMask = EPCGPointSerializeFields::None;
+	if (UnderlyingArchive.IsSaving())
+	{
+		if (Density != Default.Density)
+		{
+			EnumAddFlags(SerializeMask, EPCGPointSerializeFields::Density);
+		}
+
+		if (BoundsMin != Default.BoundsMin)
+		{
+			EnumAddFlags(SerializeMask, EPCGPointSerializeFields::BoundsMin);
+		}
+
+		if (BoundsMax != Default.BoundsMax)
+		{
+			EnumAddFlags(SerializeMask, EPCGPointSerializeFields::BoundsMax);
+		}
+
+		if (Color != Default.Color)
+		{
+			EnumAddFlags(SerializeMask, EPCGPointSerializeFields::Color);
+		}
+
+		if (Steepness != Default.Steepness)
+		{
+			EnumAddFlags(SerializeMask, EPCGPointSerializeFields::Steepness);
+		}
+
+		if (Seed != Default.Seed)
+		{
+			EnumAddFlags(SerializeMask, EPCGPointSerializeFields::Seed);
+		}
+
+		if (MetadataEntry != Default.MetadataEntry)
+		{
+			EnumAddFlags(SerializeMask, EPCGPointSerializeFields::MetadataEntry);
+		}
+	}
+		
+	FStructuredArchive::FRecord Record = Slot.EnterRecord();
+
+	Record << SA_VALUE(TEXT("SerializeMask"), SerializeMask);
+	Record << SA_VALUE(TEXT("Transform"), Transform);
+
+	if (EnumHasAnyFlags(SerializeMask, EPCGPointSerializeFields::Density))
+	{
+		Record << SA_VALUE(TEXT("Density"), Density);
+	}
+		
+	if (EnumHasAnyFlags(SerializeMask, EPCGPointSerializeFields::BoundsMin))
+	{
+		Record << SA_VALUE(TEXT("BoundsMin"), BoundsMin);
+	}
+	
+	if (EnumHasAnyFlags(SerializeMask, EPCGPointSerializeFields::BoundsMax))
+	{
+		Record << SA_VALUE(TEXT("BoundsMax"), BoundsMax);
+	}
+	
+	if (EnumHasAnyFlags(SerializeMask, EPCGPointSerializeFields::Color))
+	{
+		Record << SA_VALUE(TEXT("Color"), Color);
+	}
+	
+	if (EnumHasAnyFlags(SerializeMask, EPCGPointSerializeFields::Steepness))
+	{
+		Record << SA_VALUE(TEXT("Steepness"), Steepness);
+	}
+	
+	if (EnumHasAnyFlags(SerializeMask, EPCGPointSerializeFields::Seed))
+	{
+		Record << SA_VALUE(TEXT("Seed"), Seed);
+	}
+	
+	if (EnumHasAnyFlags(SerializeMask, EPCGPointSerializeFields::MetadataEntry))
+	{
+		Record << SA_VALUE(TEXT("MetadataEntry"), MetadataEntry);
+	}
+	
+	return true;
+}
+
 FBox FPCGPoint::GetLocalBounds() const
 {
 	return FBox(BoundsMin, BoundsMax);
diff --git a/Engine/Plugins/PCG/Source/PCG/Public/PCGPoint.h b/Engine/Plugins/PCG/Source/PCG/Public/PCGPoint.h
--- a/Engine/Plugins/PCG/Source/PCG/Public/PCGPoint.h
+++ b/Engine/Plugins/PCG/Source/PCG/Public/PCGPoint.h
@@ -89,6 +89,17 @@ public:
 
 	static bool HasCustomPropertyGetterSetter(FName Name);
 	static TUniquePtr<IPCGAttributeAccessor> CreateCustomPropertyAccessor(FName Name);
+
+	bool Serialize(FStructuredArchive::FSlot Slot);
+};
+
+template<>
+struct TStructOpsTypeTraits<FPCGPoint> : public TStructOpsTypeTraitsBase2<FPCGPoint>
+{
+	enum
+	{
+		WithStructuredSerializer = true,
+	};
 };
 
 #if UE_ENABLE_INCLUDE_ORDER_DEPRECATED_IN_5_2
diff --git a/Engine/Source/Runtime/Core/Public/UObject/FortniteReleaseBranchCustomObjectVersion.h b/Engine/Source/Runtime/Core/Public/UObject/FortniteReleaseBranchCustomObjectVersion.h
--- a/Engine/Source/Runtime/Core/Public/UObject/FortniteReleaseBranchCustomObjectVersion.h
+++ b/Engine/Source/Runtime/Core/Public/UObject/FortniteReleaseBranchCustomObjectVersion.h
@@ -48,6 +48,9 @@ struct FFortniteReleaseBranchCustomObjectVersion
 		// Store game thread particles data in single precision
 		SinglePrecisonParticleData,
 
+		// UPCGPoint custom serialization
+		PCGPointStructuredSerializer,
+
 		// -----<new versions can be added above this line>-------------------------------------------------
 		VersionPlusOne,
 		LatestVersion = VersionPlusOne - 1
