[PCG] Fix copy paste of nodes not properly renaming out default settings interface
- Also remove PostEditImport on the PCGNode as it can be called multiple times when importing sub objects, instead hook to the SettingsInterface in the new RebuildAfterPaste method

#rnx

diff --git a/Engine/Plugins/PCG/Source/PCG/Private/PCGNode.cpp b/Engine/Plugins/PCG/Source/PCG/Private/PCGNode.cpp
--- a/Engine/Plugins/PCG/Source/PCG/Private/PCGNode.cpp
+++ b/Engine/Plugins/PCG/Source/PCG/Private/PCGNode.cpp
@@ -183,18 +183,34 @@ void UPCGNode::ApplyStructuralDeprecation()
 	}
 }
 
-#endif
-
-#if WITH_EDITOR
-void UPCGNode::PostEditImport()
+void UPCGNode::RebuildAfterPaste()
 {
-	Super::PostEditImport();
+	// When Pasting a node it will get created through NewObject which will create a DefaultNodeSettings object but if
+	// the Source of the Copy has its own non-DefaultNodeSettings this is what is actually going to be pasted on this node through ImportObjectProperties
+	// leaving this DefaultNodeSettings outered to the new node but unused. Here we make sure to find the orphaned UPCGSettingsInterface objects under the node and
+	// if they aren't the one that we are actually using (SettingsInterface member) then we just discard them
+	TArray<UObject*> Objects;
+	GetObjectsWithOuter(this, Objects, false, RF_NoFlags, EInternalObjectFlags::Garbage);
+
+	for (UObject* Object : Objects)
+	{
+		if (UPCGSettingsInterface* OuteredSettingsInterface = Cast<UPCGSettingsInterface>(Object); OuteredSettingsInterface && SettingsInterface != OuteredSettingsInterface)
+		{
+			OuteredSettingsInterface->OnSettingsChangedDelegate.RemoveAll(this);
+			OuteredSettingsInterface->Rename(nullptr, GetTransientPackage(), REN_ForceNoResetLoaders | REN_DoNotDirty | REN_DontCreateRedirectors | REN_NonTransactional);
+			OuteredSettingsInterface->MarkAsGarbage();
+		}
+	}
+
 	if (SettingsInterface)
 	{
 		SettingsInterface->OnSettingsChangedDelegate.AddUObject(this, &UPCGNode::OnSettingsChanged);
 	}
 }
 
+#endif
+
+#if WITH_EDITOR
 void UPCGNode::PreEditUndo()
 {
 	if (SettingsInterface)
diff --git a/Engine/Plugins/PCG/Source/PCG/Public/PCGNode.h b/Engine/Plugins/PCG/Source/PCG/Public/PCGNode.h
--- a/Engine/Plugins/PCG/Source/PCG/Public/PCGNode.h
+++ b/Engine/Plugins/PCG/Source/PCG/Public/PCGNode.h
@@ -37,7 +37,6 @@ public:
 	virtual void BeginDestroy() override;
 
 #if WITH_EDITOR
-	virtual void PostEditImport() override;
 	virtual void PreEditUndo() override;
 	virtual void PostEditUndo() override;
 	/** ~End UObject interface */
@@ -50,6 +49,8 @@ public:
 
 	/** If a node does require structural changes, this will apply them */
 	virtual void ApplyStructuralDeprecation();
+
+	virtual void RebuildAfterPaste();
 #endif
 
 	/** Returns the owning graph */
diff --git a/Engine/Plugins/PCG/Source/PCGEditor/Private/PCGEditorGraphNodeBase.cpp b/Engine/Plugins/PCG/Source/PCGEditor/Private/PCGEditorGraphNodeBase.cpp
--- a/Engine/Plugins/PCG/Source/PCGEditor/Private/PCGEditorGraphNodeBase.cpp
+++ b/Engine/Plugins/PCG/Source/PCGEditor/Private/PCGEditorGraphNodeBase.cpp
@@ -329,6 +329,8 @@ void UPCGEditorGraphNodeBase::RebuildAfterPaste()
 {
 	if (PCGNode)
 	{
+		PCGNode->RebuildAfterPaste();
+
 		RebuildEdgesFromPins();
 
 		PCGNode->OnNodeChangedDelegate.AddUObject(this, &UPCGEditorGraphNodeBase::OnNodeChanged);
