[PCG] Make sure FPCGMetadataTypesConstantStruct::OnPostLoad is called in Create/AddAttribute

#rnx

diff --git a/Engine/Plugins/PCG/Source/PCG/Private/Elements/PCGCreateAttribute.cpp b/Engine/Plugins/PCG/Source/PCG/Private/Elements/PCGCreateAttribute.cpp
--- a/Engine/Plugins/PCG/Source/PCG/Private/Elements/PCGCreateAttribute.cpp
+++ b/Engine/Plugins/PCG/Source/PCG/Private/Elements/PCGCreateAttribute.cpp
@@ -161,6 +161,8 @@ void UPCGAddAttributeSettings::PostLoad()
 		OutputTarget.SetAttributeName(OutputAttributeName_DEPRECATED);
 		OutputAttributeName_DEPRECATED = NAME_None;
 	}
+
+	AttributeTypes.OnPostLoad();
 #endif // WITH_EDITOR
 }
 
@@ -222,6 +224,8 @@ void UPCGCreateAttributeSetSettings::PostLoad()
 		OutputTarget.SetAttributeName(OutputAttributeName_DEPRECATED);
 		OutputAttributeName_DEPRECATED = NAME_None;
 	}
+
+	AttributeTypes.OnPostLoad();
 #endif // WITH_EDITOR
 }
 
