[PCG] Disable caching for Get Object from Property Path to make sure we always have the most up-to-date data

#rnx

diff --git a/Engine/Plugins/PCG/Source/PCG/Public/Elements/PCGGetPropertyFromObjectPath.h b/Engine/Plugins/PCG/Source/PCG/Public/Elements/PCGGetPropertyFromObjectPath.h
--- a/Engine/Plugins/PCG/Source/PCG/Public/Elements/PCGGetPropertyFromObjectPath.h
+++ b/Engine/Plugins/PCG/Source/PCG/Public/Elements/PCGGetPropertyFromObjectPath.h
@@ -83,6 +83,11 @@ class FPCGGetPropertyFromObjectPathElement : public IPCGElement
 public:
 	// Loading needs to be done on the main thread and accessing objects outside of PCG might not be thread safe, so taking the safe approach
 	virtual bool CanExecuteOnlyOnMainThread(FPCGContext* Context) const override { return true; }
+	// FIXME: Disable the caching to be sure that we have the most up-to-date data from external objects
+	// To be re-enabled when we have more control on the cache, to be able to remove entries that were generated
+	// but no longer tracked, because if this object changes but it's not tracked anymore, then it will not properly remove itself from the 
+	// cache and not be caught up in future executions
+	virtual bool IsCacheable(const UPCGSettings* InSettings) const override { return false; }
 
 protected:
 	virtual FPCGContext* CreateContext() override;
