[PCG] Fixed an issue where the landscape changes wouldn't be tracked correctly on initial load due to the changes wrt to initialization and made the landscape cache work with unload/reloads of landscape as well

#rnx

diff --git a/Engine/Plugins/PCG/Source/PCG/Private/Grid/PCGLandscapeCache.cpp b/Engine/Plugins/PCG/Source/PCG/Private/Grid/PCGLandscapeCache.cpp
--- a/Engine/Plugins/PCG/Source/PCG/Private/Grid/PCGLandscapeCache.cpp
+++ b/Engine/Plugins/PCG/Source/PCG/Private/Grid/PCGLandscapeCache.cpp
@@ -880,7 +880,7 @@ void UPCGLandscapeCache::SampleMetadataOnPoint(ALandscapeProxy* Landscape, FPCGP
 		return;
 	}
 
-	const FVector2D ComponentLocalPoint(LocalPoint.X - ComponentMapKey.X * LandscapeInfo->ComponentSizeQuads, LocalPoint.Y - ComponentMapKey.Y * LandscapeInfo->ComponentSizeQuads);	
+	const FVector2D ComponentLocalPoint(LocalPoint.X - ComponentMapKey.X * LandscapeInfo->ComponentSizeQuads, LocalPoint.Y - ComponentMapKey.Y * LandscapeInfo->ComponentSizeQuads);
 	CacheEntry->GetInterpolatedPointMetadataOnly(ComponentLocalPoint, InOutPoint, OutMetadata);
 }
 
@@ -915,6 +915,12 @@ void UPCGLandscapeCache::SetupLandscapeCallbacks()
 		GEngine->OnLevelActorAdded().AddUObject(this, &UPCGLandscapeCache::OnLandscapeAdded);
 		GEngine->OnLevelActorDeleted().AddUObject(this, &UPCGLandscapeCache::OnLandscapeDeleted);
 	}
+
+	// In editor, loading from the persistent level should add it to the cache. Note that we don't need to track unloaded landscapes
+	if (!World->IsPlayInEditor() && World->IsPartitionedWorld() && World->PersistentLevel)
+	{
+		World->PersistentLevel->OnLoadedActorAddedToLevelEvent.AddUObject(this, &UPCGLandscapeCache::OnLandscapeLoaded);
+	}
 }
 
 void UPCGLandscapeCache::TeardownLandscapeCallbacks()
@@ -933,6 +939,14 @@ void UPCGLandscapeCache::TeardownLandscapeCallbacks()
 		GEngine->OnLevelActorAdded().RemoveAll(this);
 		GEngine->OnLevelActorDeleted().RemoveAll(this);
 	}
+
+	if (UWorld* World = GetWorld())
+	{
+		if (!World->IsPlayInEditor() && World->IsPartitionedWorld() && World->PersistentLevel)
+		{
+			World->PersistentLevel->OnLoadedActorAddedToLevelEvent.RemoveAll(this);
+		}
+	}
 }
 
 void UPCGLandscapeCache::OnLandscapeChanged(ALandscapeProxy* InLandscape, const FLandscapeProxyComponentDataChangedParams& InChangeParams)
@@ -1001,12 +1015,21 @@ void UPCGLandscapeCache::OnLandscapeDeleted(AActor* Actor)
 	CacheLock.WriteUnlock();
 }
 
+void UPCGLandscapeCache::OnLandscapeLoaded(AActor& Actor)
+{
+	OnLandscapeAdded(&Actor);
+}
+
 void UPCGLandscapeCache::OnLandscapeAdded(AActor* Actor)
 {
 	if (ALandscapeProxy* LandscapeProxy = Cast<ALandscapeProxy>(Actor))
 	{
 		Landscapes.Add(LandscapeProxy);
-		LandscapeProxy->OnComponentDataChanged.AddUObject(this, &UPCGLandscapeCache::OnLandscapeChanged);
+		if (!LandscapeProxy->OnComponentDataChanged.IsBoundToObject(this))
+		{
+			LandscapeProxy->OnComponentDataChanged.AddUObject(this, &UPCGLandscapeCache::OnLandscapeChanged);
+		}
+
 		// Note: Landscape Proxies have no components at this stage, so they will need to be added on demand
 	}
 }
diff --git a/Engine/Plugins/PCG/Source/PCG/Private/PCGActorAndComponentMapping.cpp b/Engine/Plugins/PCG/Source/PCG/Private/PCGActorAndComponentMapping.cpp
--- a/Engine/Plugins/PCG/Source/PCG/Private/PCGActorAndComponentMapping.cpp
+++ b/Engine/Plugins/PCG/Source/PCG/Private/PCGActorAndComponentMapping.cpp
@@ -171,6 +171,37 @@ FPCGActorAndComponentMapping::FPCGActorAndComponentMapping(UPCGSubsystem* InPCGS
 	NonPartitionedOctree.Reset(FVector::ZeroVector, OctreeExtent);
 }
 
+void FPCGActorAndComponentMapping::Initialize(UWorld* World)
+{
+#if WITH_EDITOR
+	RegisterTrackingCallbacks();
+
+	// Considering that some landscape might already be loaded, we should attach to their change callbacks now
+	if (World)
+	{
+		UPCGActorHelpers::ForEachActorInWorld<ALandscapeProxy>(World, [this](AActor* InActor)
+		{
+			if (ALandscapeProxy* LandscapeProxy = Cast<ALandscapeProxy>(InActor))
+			{
+				if (!LandscapeProxy->OnComponentDataChanged.IsBoundToObject(this))
+				{
+					LandscapeProxy->OnComponentDataChanged.AddRaw(this, &FPCGActorAndComponentMapping::OnLandscapeChanged);
+				}
+			}
+
+			return true;
+		});
+	}
+#endif // WITH_EDITOR
+}
+
+void FPCGActorAndComponentMapping::Deinitialize()
+{
+#if WITH_EDITOR
+	TeardownTrackingCallbacks();
+#endif
+}
+
 void FPCGActorAndComponentMapping::Tick()
 {
 	TSet<UPCGComponent*> ComponentToUnregister;
diff --git a/Engine/Plugins/PCG/Source/PCG/Private/PCGSubsystem.cpp b/Engine/Plugins/PCG/Source/PCG/Private/PCGSubsystem.cpp
--- a/Engine/Plugins/PCG/Source/PCG/Private/PCGSubsystem.cpp
+++ b/Engine/Plugins/PCG/Source/PCG/Private/PCGSubsystem.cpp
@@ -157,9 +157,7 @@ void UPCGSubsystem::Deinitialize()
 	PCGWorldActor = nullptr;
 	bHasTickedOnce = false;
 
-#if WITH_EDITOR
-	ActorAndComponentMapping.TeardownTrackingCallbacks();
-#endif // WITH_EDITOR
+	ActorAndComponentMapping.Deinitialize();
 
 	Super::Deinitialize();
 }
@@ -168,9 +166,7 @@ void UPCGSubsystem::PostInitialize()
 {
 	Super::PostInitialize();
 
-#if WITH_EDITOR
-	ActorAndComponentMapping.RegisterTrackingCallbacks();
-#endif // WITH_EDITOR
+	ActorAndComponentMapping.Initialize(GetWorld());
 
 	// Initialize graph executor
 	check(!GraphExecutor);
diff --git a/Engine/Plugins/PCG/Source/PCG/Public/Grid/PCGLandscapeCache.h b/Engine/Plugins/PCG/Source/PCG/Public/Grid/PCGLandscapeCache.h
--- a/Engine/Plugins/PCG/Source/PCG/Public/Grid/PCGLandscapeCache.h
+++ b/Engine/Plugins/PCG/Source/PCG/Public/Grid/PCGLandscapeCache.h
@@ -171,6 +171,7 @@ private:
 	void OnLandscapeMoved(AActor* InActor);
 	void OnLandscapeAdded(AActor* Actor);
 	void OnLandscapeDeleted(AActor* Actor);
+	void OnLandscapeLoaded(AActor& Actor);
 	void CacheLayerNames(ALandscapeProxy* InLandscape);
 	void CacheLayerNames();
 	void RemoveComponentFromCache(const ALandscapeProxy* LandscapeProxy);
diff --git a/Engine/Plugins/PCG/Source/PCG/Public/PCGActorAndComponentMapping.h b/Engine/Plugins/PCG/Source/PCG/Public/PCGActorAndComponentMapping.h
--- a/Engine/Plugins/PCG/Source/PCG/Public/PCGActorAndComponentMapping.h
+++ b/Engine/Plugins/PCG/Source/PCG/Public/PCGActorAndComponentMapping.h
@@ -26,6 +26,7 @@ class UObject;
 class UPCGComponent;
 class UPCGGraph;
 class UPCGSubsystem;
+class UWorld;
 
 /**
 * This class handle any necessary mapping between actors and pcg components.
@@ -42,6 +43,12 @@ public:
 	friend FPCGRuntimeGenScheduler;
 
 	~FPCGActorAndComponentMapping() = default;
+
+	/** Initializes callbacks, etc, tied to the PCG subsystem */
+	void Initialize(UWorld* World);
+
+	/** Deinitializes callbacks, etc, tied to the PCG subsystem */
+	void Deinitialize();
 	
 	/** Should be called by the subsystem to handle delayed operations. */
 	void Tick();
