[PCG] UI fixes:
- Adjusted some node categories.
- Normalized node names that were missing spaces
- Renamed "Merge" to "Merge Points" to make it less ambiguous
- Hidden the partition actor records as this is really just debug info and could cause issues wrt to UI (very long to show)
- Update category of the import data table blueprint + unexposed it
- Fixed missing update when changing the copy metadata operation

#rnx

#virtualized

diff --git a/Engine/Plugins/PCG/Source/PCG/Private/Elements/Metadata/PCGMetadataElement.cpp b/Engine/Plugins/PCG/Source/PCG/Private/Elements/Metadata/PCGMetadataElement.cpp
--- a/Engine/Plugins/PCG/Source/PCG/Private/Elements/Metadata/PCGMetadataElement.cpp
+++ b/Engine/Plugins/PCG/Source/PCG/Private/Elements/Metadata/PCGMetadataElement.cpp
@@ -48,14 +48,7 @@ FString UPCGMetadataOperationSettings::GetAdditionalTitleInformation() const
 #if WITH_EDITOR
 EPCGChangeType UPCGMetadataOperationSettings::GetChangeTypeForProperty(const FName& InPropertyName) const
 {
-	EPCGChangeType ChangeType = Super::GetChangeTypeForProperty(InPropertyName);
-
-	if (InPropertyName == GET_MEMBER_NAME_CHECKED(UPCGMetadataOperationSettings, bCopyAllAttributes))
-	{
-		ChangeType |= EPCGChangeType::Cosmetic;
-	}
-
-	return ChangeType;
+	return Super::GetChangeTypeForProperty(InPropertyName) | EPCGChangeType::Cosmetic;
 }
 #endif // WITH_EDITOR
 
diff --git a/Engine/Plugins/PCG/Source/PCG/Public/Elements/IO/PCGLoadAssetElement.h b/Engine/Plugins/PCG/Source/PCG/Public/Elements/IO/PCGLoadAssetElement.h
--- a/Engine/Plugins/PCG/Source/PCG/Public/Elements/IO/PCGLoadAssetElement.h
+++ b/Engine/Plugins/PCG/Source/PCG/Public/Elements/IO/PCGLoadAssetElement.h
@@ -32,7 +32,7 @@ public:
 	virtual FText GetDefaultNodeTitle() const override { return NSLOCTEXT("PCGLoadDataAssetSettings", "NodeTitle", "Load PCG Data Asset"); }
 	virtual FText GetNodeTooltipText() const override { return AssetDescription; }
 	virtual FLinearColor GetNodeTitleColor() const override { return AssetColor; }
-	virtual EPCGSettingsType GetType() const override { return EPCGSettingsType::Generic; } // TODO: change type?
+	virtual EPCGSettingsType GetType() const override { return EPCGSettingsType::InputOutput; }
 	virtual void GetStaticTrackedKeys(FPCGSelectionKeyToSettingsMap& OutKeysToSettings, TArray<TObjectPtr<const UPCGGraph>>& OutVisitedGraphs) const override;
 	virtual bool CanDynamicalyTrackKeys() const override { return true; }
 #endif // WITH_EDITOR
diff --git a/Engine/Plugins/PCG/Source/PCG/Public/Elements/PCGDataNum.h b/Engine/Plugins/PCG/Source/PCG/Public/Elements/PCGDataNum.h
--- a/Engine/Plugins/PCG/Source/PCG/Public/Elements/PCGDataNum.h
+++ b/Engine/Plugins/PCG/Source/PCG/Public/Elements/PCGDataNum.h
@@ -18,9 +18,8 @@ public:
 	virtual FName GetDefaultNodeName() const override;
 	virtual FText GetDefaultNodeTitle() const override;
 	virtual FText GetNodeTooltipText() const override;
-	virtual EPCGSettingsType GetType() const override { return EPCGSettingsType::Param; }
+	virtual EPCGSettingsType GetType() const override { return EPCGSettingsType::Generic; }
 #endif
-	
 
 protected:
 	virtual TArray<FPCGPinProperties> OutputPinProperties() const override;
diff --git a/Engine/Plugins/PCG/Source/PCG/Public/Elements/PCGDataTableRowToParamData.h b/Engine/Plugins/PCG/Source/PCG/Public/Elements/PCGDataTableRowToParamData.h
--- a/Engine/Plugins/PCG/Source/PCG/Public/Elements/PCGDataTableRowToParamData.h
+++ b/Engine/Plugins/PCG/Source/PCG/Public/Elements/PCGDataTableRowToParamData.h
@@ -24,7 +24,7 @@ public:
 	//~Begin UPCGSettings interface
 	virtual FName GetDefaultNodeName() const override { return FName(TEXT("DataTableRowToAttributeSet")); }
 	virtual FText GetDefaultNodeTitle() const override { return NSLOCTEXT("PCGDataTableRowToParamDataSettings", "NodeTitle", "Data Table Row To Attribute Set"); }
-	virtual EPCGSettingsType GetType() const override { return EPCGSettingsType::Metadata; }
+	virtual EPCGSettingsType GetType() const override { return EPCGSettingsType::InputOutput; }
 	virtual void GetStaticTrackedKeys(FPCGSelectionKeyToSettingsMap& OutKeysToSettings, TArray<TObjectPtr<const UPCGGraph>>& OutVisitedGraphs) const override;
 	virtual bool CanDynamicalyTrackKeys() const override { return true; }
 	virtual void ApplyDeprecationBeforeUpdatePins(UPCGNode* InOutNode, TArray<TObjectPtr<UPCGPin>>& InputPins, TArray<TObjectPtr<UPCGPin>>& OutputPins);
diff --git a/Engine/Plugins/PCG/Source/PCG/Public/Elements/PCGGather.h b/Engine/Plugins/PCG/Source/PCG/Public/Elements/PCGGather.h
--- a/Engine/Plugins/PCG/Source/PCG/Public/Elements/PCGGather.h
+++ b/Engine/Plugins/PCG/Source/PCG/Public/Elements/PCGGather.h
@@ -19,7 +19,7 @@ public:
 	virtual FName GetDefaultNodeName() const override { return FName(TEXT("Gather")); }
 	virtual FText GetDefaultNodeTitle() const override { return NSLOCTEXT("PCGGatherSettings", "NodeTitle", "Gather"); }
 	virtual FText GetNodeTooltipText() const override { return NSLOCTEXT("PCGGatherSettings", "NodeTooltip", "Gathers multiple data in a single collection. Can also be used to order execution through the dependency-only pin."); }
-	virtual EPCGSettingsType GetType() const override { return EPCGSettingsType::Spatial; }
+	virtual EPCGSettingsType GetType() const override { return EPCGSettingsType::Generic; }
 	virtual bool HasDynamicPins() const override { return true; }
 #endif
 
diff --git a/Engine/Plugins/PCG/Source/PCG/Public/Elements/PCGGetLoopIndex.h b/Engine/Plugins/PCG/Source/PCG/Public/Elements/PCGGetLoopIndex.h
--- a/Engine/Plugins/PCG/Source/PCG/Public/Elements/PCGGetLoopIndex.h
+++ b/Engine/Plugins/PCG/Source/PCG/Public/Elements/PCGGetLoopIndex.h
@@ -19,7 +19,7 @@ public:
 	virtual FName GetDefaultNodeName() const override { return FName(TEXT("GetLoopIndex")); }
 	virtual FText GetDefaultNodeTitle() const override { return NSLOCTEXT("PCGGetLoopIndexElement", "NodeTitle", "Get Loop Index"); }
 	virtual FText GetNodeTooltipText() const override;
-	virtual EPCGSettingsType GetType() const override { return EPCGSettingsType::Param; }
+	virtual EPCGSettingsType GetType() const override { return EPCGSettingsType::Generic; }
 #endif // WITH_EDITOR
 
 protected:
diff --git a/Engine/Plugins/PCG/Source/PCG/Public/Elements/PCGMergeElement.h b/Engine/Plugins/PCG/Source/PCG/Public/Elements/PCGMergeElement.h
--- a/Engine/Plugins/PCG/Source/PCG/Public/Elements/PCGMergeElement.h
+++ b/Engine/Plugins/PCG/Source/PCG/Public/Elements/PCGMergeElement.h
@@ -15,8 +15,8 @@ class PCG_API UPCGMergeSettings : public UPCGSettings
 public:
 	//~Begin UPCGSettings interface
 #if WITH_EDITOR
-	virtual FName GetDefaultNodeName() const override { return FName(TEXT("Merge")); }
-	virtual FText GetDefaultNodeTitle() const override { return NSLOCTEXT("PCGMergeSettings", "NodeTitle", "Merge"); }
+	virtual FName GetDefaultNodeName() const override { return FName(TEXT("MergePoints")); }
+	virtual FText GetDefaultNodeTitle() const override { return NSLOCTEXT("PCGMergeSettings", "NodeTitle", "Merge Points"); }
 	virtual FText GetNodeTooltipText() const override;
 	virtual EPCGSettingsType GetType() const override { return EPCGSettingsType::Spatial; }
 #endif
diff --git a/Engine/Plugins/PCG/Source/PCG/Public/PCGWorldActor.h b/Engine/Plugins/PCG/Source/PCG/Public/PCGWorldActor.h
--- a/Engine/Plugins/PCG/Source/PCG/Public/PCGWorldActor.h
+++ b/Engine/Plugins/PCG/Source/PCG/Public/PCGWorldActor.h
@@ -126,7 +126,7 @@ private:
 	mutable FRWLock TransientGridGuidsLock;
 
 	/** Keeps a record of what grid cells contain a serialized partition actor. Useful for tracking the existence of PAs even when they are not yet loaded. */
-	UPROPERTY(VisibleAnywhere, Category = Debug)
+	UPROPERTY()
 	TSet<FPCGPartitionActorRecord> SerializedPartitionActorRecords;
 };
 
