[PCG] Added automatic deprecation for Apply Tags On Output (only on pins that had edges)

#rnx

diff --git a/Engine/Plugins/PCG/Source/PCG/Private/Elements/PCGCreateAttribute.cpp b/Engine/Plugins/PCG/Source/PCG/Private/Elements/PCGCreateAttribute.cpp
--- a/Engine/Plugins/PCG/Source/PCG/Private/Elements/PCGCreateAttribute.cpp
+++ b/Engine/Plugins/PCG/Source/PCG/Private/Elements/PCGCreateAttribute.cpp
@@ -96,6 +96,8 @@ void UPCGAddAttributeSettings::ApplyStructuralDeprecation(UPCGNode* InOutNode)
 			InOutNode->SetSettingsInterface(NewSettings);
 		}
 	}
+
+	Super::ApplyStructuralDeprecation(InOutNode);
 }
 #endif // WITH_EDITOR
 
diff --git a/Engine/Plugins/PCG/Source/PCG/Private/PCGSettings.cpp b/Engine/Plugins/PCG/Source/PCG/Private/PCGSettings.cpp
--- a/Engine/Plugins/PCG/Source/PCG/Private/PCGSettings.cpp
+++ b/Engine/Plugins/PCG/Source/PCG/Private/PCGSettings.cpp
@@ -11,6 +11,7 @@
 #include "PCGPin.h"
 #include "PCGSubgraph.h"
 #include "PCGSubsystem.h"
+#include "Elements/PCGAddTag.h"
 #include "Helpers/PCGHelpers.h"
 #include "Helpers/PCGSettingsHelpers.h"
 #include "Metadata/Accessors/PCGAttributeAccessorHelpers.h"
@@ -638,6 +639,86 @@ void UPCGSettings::ApplyDeprecationBeforeUpdatePins(UPCGNode* InOutNode, TArray<
 		InOutNode->RenameInputPin(PCGPinConstants::Private::OldDefaultParamsLabel, PCGPinConstants::DefaultParamsLabel);
 	}
 }
+
+void UPCGSettings::ApplyStructuralDeprecation(UPCGNode* InOutNode)
+{
+	if (!InOutNode)
+	{
+		return;
+	}
+
+	// Deprecate the "TagsAppliedOnOutput" feature, replaced by usage of the AddTag node.
+	// Note that the additional nodes on a per-pin basis (we don't have a AddTag passthrough node, and would be hard to understand for users)
+	// and will not create the additional nodes if the pins aren't connected, meaning that in this specific case the tags would be lost.
+	// Implementation note: the version here doesn't match the feature change, but has been bumped just after this was deprecated.
+	if (DataVersion < FPCGCustomVersion::GetPCGComponentDataMustOverlapSourceComponentByDefault && !TagsAppliedOnOutput_DEPRECATED.IsEmpty())
+	{
+		UPCGGraph* PCGGraph = InOutNode->GetGraph();
+		check(PCGGraph);
+
+		const FString TagsToAdd = FString::Join(TagsAppliedOnOutput_DEPRECATED, TEXT(","));
+		bool bAddedAddTagNodes = false;
+
+		const TArray<TObjectPtr<UPCGPin>> OutputPins = InOutNode->GetOutputPins();
+
+		for (TObjectPtr<UPCGPin> OutputPin : OutputPins)
+		{
+			check(OutputPin);
+			const int32 OutputEdgeCount = OutputPin->EdgeCount();
+
+			if (OutputEdgeCount == 0)
+			{
+				continue;
+			}
+
+			// Insert an "Add Tag" node from that pin and move the edges accordingly
+			UPCGAddTagSettings* AddTagSettings = nullptr;
+			UPCGNode* AddTagNode = PCGGraph->AddNodeOfType(AddTagSettings);
+			check(AddTagNode && AddTagSettings);
+
+			AddTagSettings->bEnabled = InOutNode->GetSettings()->bEnabled;
+			AddTagSettings->TagsToAdd = TagsToAdd;
+
+			UPCGPin* AddTagInputPin = AddTagNode->GetInputPins()[0];
+			UPCGPin* AddTagOutputPin = AddTagNode->GetOutputPins()[0];
+
+			int32 SumSourceNodesPositionX = 0;
+			int32 SumSourceNodesPositionY = 0;
+			while (!OutputPin->Edges.IsEmpty())
+			{
+				UPCGPin* DownstreamPin = OutputPin->Edges[0]->OutputPin;
+				check(DownstreamPin);
+
+				int32 NodePositionX, NodePositionY;
+				DownstreamPin->Node->GetNodePosition(NodePositionX, NodePositionY);
+				SumSourceNodesPositionX += NodePositionX;
+				SumSourceNodesPositionY += NodePositionY;
+
+				check(DownstreamPin->Node);
+				OutputPin->BreakEdgeTo(DownstreamPin);
+				AddTagOutputPin->AddEdgeTo(DownstreamPin);
+			}
+
+			// Place the injected node halfway between the original node and the downstream nodes
+			int32 NodePositionX, NodePositionY;
+			InOutNode->GetNodePosition(NodePositionX, NodePositionY);
+			const int32 SourceNodeLocalPositionMeanX = SumSourceNodesPositionX / OutputEdgeCount;
+			const int32 SourceNodeLocalPositionMeanY = SumSourceNodesPositionY / OutputEdgeCount;
+
+			AddTagNode->SetNodePosition((NodePositionX + SourceNodeLocalPositionMeanX) / 2, (NodePositionY + SourceNodeLocalPositionMeanY) / 2);
+			AddTagNode->NodeComment = LOCTEXT("DeprecationAddTagNodeCreated", "Node added to replicate deprecated behavior of the Add Tags To Output functionality.").ToString();
+			AddTagNode->bCommentBubbleVisible = 1;
+
+			OutputPin->AddEdgeTo(AddTagInputPin);
+			bAddedAddTagNodes = true;
+		}
+
+		if (bAddedAddTagNodes)
+		{
+			UE_LOG(LogPCG, Log, TEXT("Tags applied to output detected. One or more 'Add Tags' node was created automatically to repliecate the previous behavior."));
+		}
+	}
+}
 #endif // WITH_EDITOR
 
 bool UPCGSettings::IsPropertyOverriddenByPin(const FProperty* InProperty) const
diff --git a/Engine/Plugins/PCG/Source/PCG/Public/PCGSettings.h b/Engine/Plugins/PCG/Source/PCG/Public/PCGSettings.h
--- a/Engine/Plugins/PCG/Source/PCG/Public/PCGSettings.h
+++ b/Engine/Plugins/PCG/Source/PCG/Public/PCGSettings.h
@@ -278,7 +278,7 @@ public:
 	virtual void ApplyDeprecation(UPCGNode* InOutNode);
 
 	/** If settings require structural changes, this will apply them */
-	virtual void ApplyStructuralDeprecation(UPCGNode* InOutNode) {}
+	virtual void ApplyStructuralDeprecation(UPCGNode* InOutNode);
 
 	virtual FName GetDefaultNodeName() const { return NAME_None; }
 	virtual FText GetDefaultNodeTitle() const { return FText::FromName(GetDefaultNodeName()); }
@@ -372,6 +372,9 @@ public:
 	int Seed = 0xC35A9631; // Default seed is a random prime number, but will be overriden for new settings based on the class type name hash, making each settings class have a different default seed.
 
 #if WITH_EDITORONLY_DATA
+	UPROPERTY()
+	TSet<FString> TagsAppliedOnOutput_DEPRECATED;
+
 	UPROPERTY()
 	EPCGSettingsExecutionMode ExecutionMode_DEPRECATED = EPCGSettingsExecutionMode::Enabled;
 
