[PCG] Nativization of Select Points node from custom blueprints video

**Test of workflow with new custom native node template

diff --git a/Engine/Plugins/PCG/Source/PCG/Private/Elements/PCGRandomChoice.cpp b/Engine/Plugins/PCG/Source/PCG/Private/Elements/PCGRandomChoice.cpp
new file mode 100644
--- /dev/null
+++ b/Engine/Plugins/PCG/Source/PCG/Private/Elements/PCGRandomChoice.cpp
@@ -0,0 +1,151 @@
+// Copyright Epic Games, Inc. All Rights Reserved.
+
+#include "Elements/PCGRandomChoice.h"
+
+#include "PCGContext.h"
+#include "PCGParamData.h"
+#include "PCGPin.h"
+#include "Data/PCGPointData.h"
+
+#define LOCTEXT_NAMESPACE "PCGRandomChoiceElement"
+
+namespace PCGRandomChoiceConstants
+{
+	const FName OutputALabel = TEXT("Chosen Points");
+	const FName OutputBLabel = TEXT("Discarded Points");
+}
+
+#if WITH_EDITOR
+// The label the node is known as internally.
+FName UPCGRandomChoiceSettings::GetDefaultNodeName() const
+{
+	return FName(TEXT("RandomChoice"));
+}
+
+// Default node name shown in the graph editor. Include spaces.
+FText UPCGRandomChoiceSettings::GetDefaultNodeTitle() const
+{
+	return LOCTEXT("NodeTitle", "Random Choice");
+}
+
+// Default tooltip for the node
+FText UPCGRandomChoiceSettings::GetNodeTooltipText() const
+{
+	return LOCTEXT("NodeTooltip", "Chooses points randomly through ratio or a fixed number of points.");
+}
+#endif
+
+// Input/Output pin setup with specific properties, including:
+// Pin data type, allowing singular or multiple inputs per pin, and creating multiple in/out pins.
+TArray<FPCGPinProperties> UPCGRandomChoiceSettings::InputPinProperties() const
+{
+	return Super::DefaultPointInputPinProperties();
+}
+
+TArray<FPCGPinProperties> UPCGRandomChoiceSettings::OutputPinProperties() const
+{
+	TArray<FPCGPinProperties> PinProperties;
+	PinProperties.Emplace(PCGRandomChoiceConstants::OutputALabel,
+		EPCGDataType::Point,
+		/*bAllowMultipleConnections=*/true,
+		/*bAllowMultipleData=*/true);
+
+	PinProperties.Emplace(PCGRandomChoiceConstants::OutputBLabel,
+		EPCGDataType::Point,
+		/*bAllowMultipleConnections=*/true,
+		/*bAllowMultipleData=*/true);
+
+	return PinProperties;
+}
+
+// Creates the Element to be used for ExecuteInternal.
+FPCGElementPtr UPCGRandomChoiceSettings::CreateElement() const
+{
+	return MakeShared<FPCGRandomChoiceElement>();
+}
+
+/*
+* Processing function for this node. 
+* Context holds the InputData, containing the input data collection for this node 
+* and the OutputData, the output data collection to write to as output.
+* Returns true if the processing is done. 
+* Returning false will call back this function at next tick, and will call it until it returns true.
+* Settings contains all the setup options for this node, and if a property was marked PCG_Overridable, 
+* "Context->GetInputSettings" will contain the overridden value for this property if it is overridden
+*/ 
+bool FPCGRandomChoiceElement::ExecuteInternal(FPCGContext* Context) const
+{
+	TRACE_CPUPROFILER_EVENT_SCOPE(FPCGRandomChoiceElement::Execute);
+
+	check(Context);
+
+	const UPCGRandomChoiceSettings* Settings = Context->GetInputSettings<UPCGRandomChoiceSettings>();
+	check(Settings);
+
+	const TArray<FPCGTaggedData> Inputs = Context->InputData.GetInputsByPin(PCGPinConstants::DefaultInputLabel);
+	TArray<FPCGTaggedData>& Outputs = Context->OutputData.TaggedData;
+
+	for (int i = 0; i < Inputs.Num(); ++i)
+	{
+		const UPCGPointData* InputPointData = Cast<UPCGPointData>(Inputs[i].Data);
+		if (!InputPointData)
+		{
+			PCGE_LOG(Verbose, GraphAndLog, FText::Format(LOCTEXT("InvalidPointData", "Input {0} is not point data"), i));
+			continue;
+		}
+
+		const TArray<FPCGPoint>& InPoints = InputPointData->GetPoints();
+		int NumOfElementsToKeep = 0;
+
+		if (Settings->bFixedMode)
+		{
+			NumOfElementsToKeep = FMath::Clamp(Settings->FixedNumber, 1, InPoints.Num());
+		}
+		else
+		{
+			NumOfElementsToKeep = FMath::CeilToInt(InPoints.Num() * FMath::Clamp(Settings->Ratio, 0, 1));
+		}
+
+		FRandomStream RandStream(Context->GetSeed());
+		TArray<int32> ShuffledIndexes;
+		ShuffledIndexes.Reserve(InPoints.Num());
+
+		// Instead of swapping directly, we swap the indexes since it cheaper than copying points around
+		for (int j = 0; j < InPoints.Num(); ++j)
+		{
+			ShuffledIndexes.Add(j);
+		}
+
+		for (int j = 0; j < InPoints.Num() - 1; ++j)
+		{
+			int RandomElement = RandStream.RandRange(j + 1, InPoints.Num() - 1);
+			Swap(ShuffledIndexes[j], ShuffledIndexes[RandomElement]);
+		}
+
+		FPCGTaggedData& ChosenOutput = Outputs.Add_GetRef(Inputs[i]);
+		UPCGPointData* ChosenPointData = NewObject<UPCGPointData>();
+		ChosenPointData->InitializeFromData(InputPointData);
+		TArray<FPCGPoint>& ChosenPoints = ChosenPointData->GetMutablePoints();
+		ChosenPoints.Reserve(NumOfElementsToKeep);
+		ChosenOutput.Data = ChosenPointData;
+		ChosenOutput.Pin = PCGRandomChoiceConstants::OutputALabel;
+		
+		FPCGTaggedData& DiscardedOutput = Outputs.Add_GetRef(Inputs[i]);
+		UPCGPointData* DiscardedPointData = NewObject<UPCGPointData>();
+		DiscardedPointData->InitializeFromData(InputPointData);
+		TArray<FPCGPoint>& DiscardedPoints = DiscardedPointData->GetMutablePoints();
+		DiscardedPoints.Reserve(InPoints.Num() - NumOfElementsToKeep);
+		DiscardedOutput.Data = DiscardedPointData;
+		DiscardedOutput.Pin = PCGRandomChoiceConstants::OutputBLabel;
+
+		
+		for (int j = 0; j < InPoints.Num(); ++j)
+		{
+			(j >= NumOfElementsToKeep ? DiscardedPoints : ChosenPoints).Add(InPoints[ShuffledIndexes[j]]);
+		}
+	}
+
+	return true;
+}
+
+#undef LOCTEXT_NAMESPACE
\ No newline at end of file
diff --git a/Engine/Plugins/PCG/Source/PCG/Public/Elements/PCGRandomChoice.h b/Engine/Plugins/PCG/Source/PCG/Public/Elements/PCGRandomChoice.h
new file mode 100644
--- /dev/null
+++ b/Engine/Plugins/PCG/Source/PCG/Public/Elements/PCGRandomChoice.h
@@ -0,0 +1,49 @@
+// Copyright Epic Games, Inc. All Rights Reserved.
+
+#pragma once
+
+#include "PCGSettings.h"
+
+#include "PCGRandomChoice.generated.h"
+
+/**
+* Chooses points randomly through ratio or a fixed number of points
+*/
+UCLASS(BlueprintType)
+class UPCGRandomChoiceSettings : public UPCGSettings
+{
+	GENERATED_BODY()
+
+public:
+	//~Begin UPCGSettings interface
+#if WITH_EDITOR
+	virtual FName GetDefaultNodeName() const override;
+	virtual FText GetDefaultNodeTitle() const override;
+	virtual FText GetNodeTooltipText() const override;
+	virtual EPCGSettingsType GetType() const override { return EPCGSettingsType::Sampler; }
+#endif
+
+protected:
+	virtual TArray<FPCGPinProperties> InputPinProperties() const override;
+	virtual TArray<FPCGPinProperties> OutputPinProperties() const override;
+	virtual FPCGElementPtr CreateElement() const override;
+	//~End UPCGSettings interface
+
+
+public:
+	UPROPERTY(BlueprintReadWrite, EditAnywhere, Category = Settings, meta = (PCG_Overridable))
+	bool bFixedMode = true;
+
+	UPROPERTY(BlueprintReadWrite, EditAnywhere, Category = Settings, meta = (PCG_Overridable, EditCondition = "bFixedMode", ClampMin = "1"))
+	int FixedNumber = 1;
+
+	UPROPERTY(BlueprintReadWrite, EditAnywhere, Category = Settings, meta = (PCG_Overridable, EditCondition = "!bFixedMode", ClampMin = "0", ClampMax = "1"))
+	float Ratio = 0.5;
+
+};
+
+class FPCGRandomChoiceElement : public IPCGElement
+{
+protected:
+	virtual bool ExecuteInternal(FPCGContext* Context) const override;
+};
