[PCG] Make subgraph element non cacheable

#rnx

diff --git a/Engine/Plugins/PCG/Source/PCG/Private/PCGSubgraph.cpp b/Engine/Plugins/PCG/Source/PCG/Private/PCGSubgraph.cpp
--- a/Engine/Plugins/PCG/Source/PCG/Private/PCGSubgraph.cpp
+++ b/Engine/Plugins/PCG/Source/PCG/Private/PCGSubgraph.cpp
@@ -586,12 +586,6 @@ void FPCGSubgraphElement::PrepareSubgraphUserParameters(const UPCGSubgraphSettin
 	}
 }
 
-bool FPCGSubgraphElement::IsCacheable(const UPCGSettings* InSettings) const
-{
-	const UPCGSubgraphSettings* Settings = Cast<const UPCGSubgraphSettings>(InSettings);
-	return (!Settings || !Settings->IsDynamicGraph());
-}
-
 bool FPCGSubgraphElement::IsPassthrough(const UPCGSettings* InSettings) const
 {
 	const UPCGSubgraphSettings* Settings = Cast<UPCGSubgraphSettings>(InSettings);
diff --git a/Engine/Plugins/PCG/Source/PCG/Public/PCGSubgraph.h b/Engine/Plugins/PCG/Source/PCG/Public/PCGSubgraph.h
--- a/Engine/Plugins/PCG/Source/PCG/Public/PCGSubgraph.h
+++ b/Engine/Plugins/PCG/Source/PCG/Public/PCGSubgraph.h
@@ -179,7 +179,7 @@ class PCG_API FPCGSubgraphElement : public IPCGElement
 {
 public:
 	virtual FPCGContext* Initialize(const FPCGDataCollection& InputData, TWeakObjectPtr<UPCGComponent> SourceComponent, const UPCGNode* Node) override;
-	virtual bool IsCacheable(const UPCGSettings* InSettings) const override;
+	virtual bool IsCacheable(const UPCGSettings* InSettings) const override { return false; }
 
 protected:
 	virtual bool ExecuteInternal(FPCGContext* Context) const override;
