[PCG] Graph Editor: Detect cyclic connections and display connection error
- New UPCGEditorGraphNodeBase::IsCompatible allows for extra validation in UPCGEditorGraphSchema::CanCreateConnection
- Gated behind CVar 'pcg.Editor.CheckConnectionCycles' but 250k edge test was under 3ms and wasn't causing any UX hitches

#rnx

diff --git a/Engine/Plugins/PCG/Source/PCGEditor/Private/PCGEditorGraphNodeBase.cpp b/Engine/Plugins/PCG/Source/PCGEditor/Private/PCGEditorGraphNodeBase.cpp
--- a/Engine/Plugins/PCG/Source/PCGEditor/Private/PCGEditorGraphNodeBase.cpp
+++ b/Engine/Plugins/PCG/Source/PCGEditor/Private/PCGEditorGraphNodeBase.cpp
@@ -31,6 +31,15 @@
 
 #define LOCTEXT_NAMESPACE "PCGEditorGraphNodeBase"
 
+namespace PCGEditorGraphSwitches
+{
+	TAutoConsoleVariable<bool> CVarCheckConnectionCycles{
+		TEXT("pcg.Editor.CheckConnectionCycles"),
+		true,
+		TEXT("Prevents user from creating cycles in graph")
+	};
+}
+
 namespace PCGEditorGraphNodeBase
 {
 	/** Whether this node was culled during graph compilation or during graph execution. */
@@ -697,6 +706,56 @@ void UPCGEditorGraphNodeBase::ExitRenamingMode()
 	OnNodeChangedDelegate.ExecuteIfBound();
 }
 
+bool UPCGEditorGraphNodeBase::IsCompatible(const UPCGPin* InputPin, const UPCGPin* OutputPin, FText& OutReason) const
+{
+	if (PCGEditorGraphSwitches::CVarCheckConnectionCycles.GetValueOnAnyThread() && InputPin && OutputPin && InputPin->Node == PCGNode)
+	{
+		// Upstream Visitor
+		auto Visitor = [ThisPCGNode = PCGNode](const UPCGNode* InNode, auto VisitorLambda) -> bool
+		{
+			if (InNode)
+			{
+				if (InNode == ThisPCGNode)
+				{
+					return false;
+				}
+
+				for (const TObjectPtr<UPCGPin>& InputPin : InNode->GetInputPins())
+				{
+					if (InputPin)
+					{
+						for (const TObjectPtr<UPCGEdge>& Edge : InputPin->Edges)
+						{
+							if (Edge)
+							{
+								if (const UPCGPin* OtherPin = Edge->GetOtherPin(InputPin.Get()))
+								{
+									if (!VisitorLambda(OtherPin->Node, VisitorLambda))
+									{
+										return false;
+									}
+								}
+							}
+						}
+					}
+				}
+			}
+
+			return true;
+		};
+
+		// OutputPin is trying to connect to this nodes InputPin so visit the OutputPin upstream and try to find
+		// a existing connection to this UPCGEditorGraphNodeNamedRerouteDeclaration's PCGNode. If we do deny connection which would create cycle.
+		if (!Visitor(OutputPin->Node, Visitor))
+		{
+			OutReason = LOCTEXT("ConnectionFailedCyclic", "Connection would create cycle");
+			return false;
+		}
+	}
+
+	return true;
+}
+
 void UPCGEditorGraphNodeBase::OnPickColor()
 {
 	FColorPickerArgs PickerArgs;
diff --git a/Engine/Plugins/PCG/Source/PCGEditor/Private/PCGEditorGraphNodeBase.h b/Engine/Plugins/PCG/Source/PCGEditor/Private/PCGEditorGraphNodeBase.h
--- a/Engine/Plugins/PCG/Source/PCGEditor/Private/PCGEditorGraphNodeBase.h
+++ b/Engine/Plugins/PCG/Source/PCGEditor/Private/PCGEditorGraphNodeBase.h
@@ -111,6 +111,9 @@ public:
 	/** Disables renaming, preventing the node from entering editing mode on subsequent updates. */
 	void ExitRenamingMode();
 
+	/** Whether this OutputPin can be connected to this InputPin */
+	virtual bool IsCompatible(const UPCGPin* InputPin, const UPCGPin* OutputPin, FText& OutReason) const;
+
 	DECLARE_DELEGATE(FOnPCGEditorGraphNodeChanged);
 	FOnPCGEditorGraphNodeChanged OnNodeChangedDelegate;
 
diff --git a/Engine/Plugins/PCG/Source/PCGEditor/Private/PCGEditorGraphSchema.cpp b/Engine/Plugins/PCG/Source/PCGEditor/Private/PCGEditorGraphSchema.cpp
--- a/Engine/Plugins/PCG/Source/PCGEditor/Private/PCGEditorGraphSchema.cpp
+++ b/Engine/Plugins/PCG/Source/PCGEditor/Private/PCGEditorGraphSchema.cpp
@@ -105,6 +105,7 @@ const FPinConnectionResponse UPCGEditorGraphSchema::CanCreateConnection(const UE
 
 	const UPCGEditorGraphNodeBase* EditorNodeA = CastChecked<const UPCGEditorGraphNodeBase>(NodeA);
 	const UPCGEditorGraphNodeBase* EditorNodeB = CastChecked<const UPCGEditorGraphNodeBase>(NodeB);
+	const UPCGEditorGraphNodeBase* EditorNodeWithInput = nullptr;
 
 	// Check type compatibility & whether we can connect more pins
 	const UPCGPin* InputPin = nullptr;
@@ -114,11 +115,13 @@ const FPinConnectionResponse UPCGEditorGraphSchema::CanCreateConnection(const UE
 	{
 		OutputPin = EditorNodeA->GetPCGNode()->GetOutputPin(A->PinName);
 		InputPin = EditorNodeB->GetPCGNode()->GetInputPin(B->PinName);
+		EditorNodeWithInput = EditorNodeB;
 	}
 	else
 	{
 		OutputPin = EditorNodeB->GetPCGNode()->GetOutputPin(B->PinName);
 		InputPin = EditorNodeA->GetPCGNode()->GetInputPin(A->PinName);
+		EditorNodeWithInput = EditorNodeA;
 	}
 
 	if (!InputPin || !OutputPin)
@@ -150,6 +153,12 @@ const FPinConnectionResponse UPCGEditorGraphSchema::CanCreateConnection(const UE
 		return FPinConnectionResponse((A->Direction == EGPD_Output) ? CONNECT_RESPONSE_BREAK_OTHERS_B : CONNECT_RESPONSE_BREAK_OTHERS_A, LOCTEXT("ConnectionBreakExisting", "Break existing connection?"));
 	}
 
+	FText Reason;
+	if (!EditorNodeWithInput->IsCompatible(InputPin, OutputPin, Reason))
+	{
+		return FPinConnectionResponse(CONNECT_RESPONSE_DISALLOW, Reason);
+	}
+
 	return FPinConnectionResponse();
 }
 
