[PCG] Copy/name toolbar icons to the PCG plugin
- Add Icon variation for Regen with Clear cache

#rnx

diff --git a/Engine/Plugins/PCG/Source/PCGEditor/Private/Details/PCGComponentDetails.cpp b/Engine/Plugins/PCG/Source/PCGEditor/Private/Details/PCGComponentDetails.cpp
--- a/Engine/Plugins/PCG/Source/PCGEditor/Private/Details/PCGComponentDetails.cpp
+++ b/Engine/Plugins/PCG/Source/PCGEditor/Private/Details/PCGComponentDetails.cpp
@@ -2,10 +2,11 @@
 
 #include "Details/PCGComponentDetails.h"
 #include "PCGComponent.h"
+#include "PCGEditorStyle.h"
 #include "PCGSubsystem.h"
-
 #include "Framework/Application/SlateApplication.h"
 #include "Widgets/Text/STextBlock.h"
+#include "Widgets/Images/SImage.h"
 #include "Widgets/Input/SButton.h"
 #include "DetailLayoutBuilder.h"
 #include "DetailWidgetRow.h"
@@ -68,36 +69,72 @@ void FPCGComponentDetails::CustomizeDetails(IDetailLayoutBuilder& DetailBuilder)
 			SNew(SHorizontalBox)
 			+ SHorizontalBox::Slot()
 			.AutoWidth()
-			.Padding(4.0f)
+			.Padding(2.0f, 0.0f)
+			.VAlign(VAlign_Fill)
 			[
 				SNew(SButton)
 				.OnClicked(this, &FPCGComponentDetails::OnGenerateClicked)
 				.ToolTipText(FText::FromString("Generates graph data. \nCtrl + Click flushes the cache and force generates."))
 				.Visibility(this, &FPCGComponentDetails::GenerateButtonVisible)
 				[
-					SNew(STextBlock)
-					.Font(IDetailLayoutBuilder::GetDetailFont())
-					.Text(LOCTEXT("GenerateButton", "Generate"))
+					SNew(SHorizontalBox)
+					+SHorizontalBox::Slot()
+					.Padding(0.0f, 0.0f, 6.0f, 0.0f)
+					.VAlign(VAlign_Center)
+					.AutoWidth()
+					[
+						SNew(SImage)
+						.DesiredSizeOverride(FVector2D(16,16))
+						.Image_Lambda([]() { return FSlateApplication::Get().GetModifierKeys().IsControlDown() ? FPCGEditorStyle::Get().GetBrush("PCG.Command.ForceRegenClearCache") : FPCGEditorStyle::Get().GetBrush("PCG.Command.ForceRegen"); })
+						.ColorAndOpacity(FSlateColor::UseForeground())
+					]
+					+SHorizontalBox::Slot()
+					.VAlign(VAlign_Center)
+					.AutoWidth()
+					[
+						SNew(STextBlock)
+						.Font(IDetailLayoutBuilder::GetDetailFont())
+						.Text_Lambda([]() { return FSlateApplication::Get().GetModifierKeys().IsControlDown() ? LOCTEXT("ForceRegenerateButton", "Force Generate") : LOCTEXT("GenerateButton", "Generate"); })
+					]
 				]
 			]
 			+ SHorizontalBox::Slot()
 			.AutoWidth()
-			.Padding(4.0f)
+			.Padding(2.0f, 0.0f)
+			.VAlign(VAlign_Fill)
 			[
 				SNew(SButton)
 				.OnClicked(this, &FPCGComponentDetails::OnCancelClicked)
 				.Visibility(this, &FPCGComponentDetails::CancelButtonVisible)
 				[
-					SNew(STextBlock)
-					.Font(IDetailLayoutBuilder::GetDetailFont())
-					.Text(LOCTEXT("CancelButton", "Cancel"))
+					SNew(SHorizontalBox)
+					+ SHorizontalBox::Slot()
+					.Padding(0.0f, 0.0f, 6.0f, 0.0f)
+					.VAlign(VAlign_Center)
+					.AutoWidth()
+					[
+						SNew(SImage)
+						.DesiredSizeOverride(FVector2D(16, 16))
+						.Image(FPCGEditorStyle::Get().GetBrush("PCG.Command.StopRegen"))
+						.ColorAndOpacity(FSlateColor::UseForeground())
+					]
+					+ SHorizontalBox::Slot()
+					.VAlign(VAlign_Center)
+					.AutoWidth()
+					[
+						SNew(STextBlock)
+						.Font(IDetailLayoutBuilder::GetDetailFont())
+						.Text(LOCTEXT("CancelButton", "Cancel"))
+					]
 				]
 			]
 			+ SHorizontalBox::Slot()
 			.AutoWidth()
-			.Padding(4.0f)
+			.VAlign(VAlign_Fill)
+			.Padding(2.0f, 0.0f)
 			[
 				SNew(SButton)
+				.VAlign(VAlign_Center)
 				.OnClicked(this, &FPCGComponentDetails::OnCleanupClicked)
 				.ToolTipText(FText::FromString("Cleans up graph data. \nCtrl + Click purges all components and attached actors tagged as created by PCG."))
 				.Visibility(this, &FPCGComponentDetails::CleanupButtonVisible)
@@ -109,9 +146,11 @@ void FPCGComponentDetails::CustomizeDetails(IDetailLayoutBuilder& DetailBuilder)
 			]
 			+ SHorizontalBox::Slot()
 			.AutoWidth()
-			.Padding(4.0f)
+			.VAlign(VAlign_Fill)
+			.Padding(2.0f, 0.0f)
 			[
 				SNew(SButton)
+				.VAlign(VAlign_Center)
 				.OnClicked(this, &FPCGComponentDetails::OnRefreshClicked)
 				.Visibility(this, &FPCGComponentDetails::RefreshButtonVisible)
 				[
@@ -122,9 +161,11 @@ void FPCGComponentDetails::CustomizeDetails(IDetailLayoutBuilder& DetailBuilder)
 			]
 			+ SHorizontalBox::Slot()
 			.AutoWidth()
-			.Padding(4.0f)
+			.VAlign(VAlign_Fill)
+			.Padding(2.0f, 0.0f)
 			[
 				SNew(SButton)
+				.VAlign(VAlign_Center)
 				.OnClicked(this, &FPCGComponentDetails::OnClearPCGLinkClicked)
 				[
 					SNew(STextBlock)
diff --git a/Engine/Plugins/PCG/Source/PCGEditor/Private/PCGEditor.cpp b/Engine/Plugins/PCG/Source/PCGEditor/Private/PCGEditor.cpp
--- a/Engine/Plugins/PCG/Source/PCGEditor/Private/PCGEditor.cpp
+++ b/Engine/Plugins/PCG/Source/PCGEditor/Private/PCGEditor.cpp
@@ -29,6 +29,7 @@
 #include "PCGEditorGraphSchemaActions.h"
 #include "PCGEditorMenuContext.h"
 #include "PCGEditorSettings.h"
+#include "PCGEditorStyle.h"
 #include "PCGEditorUtils.h"
 #include "SPCGEditorGraphAttributeListView.h"
 #include "SPCGEditorGraphDebugObjectTree.h"
@@ -597,31 +598,38 @@ void FPCGEditor::RegisterToolbar() const
 			PCGEditorCommands.Find,
 			TAttribute<FText>(),
 			TAttribute<FText>(),
-			FSlateIcon(FAppStyle::GetAppStyleSetName(), "BlueprintEditor.FindInBlueprint")));
+			FSlateIcon(FPCGEditorStyle::Get().GetStyleSetName(), "PCG.Command.Find")));
 
 		Section.AddEntry(FToolMenuEntry::InitToolBarButton(
 			FPCGEditorCommands::Get().PauseAutoRegeneration,
 			TAttribute<FText>(),
 			TAttribute<FText>(),
-			FSlateIcon(FAppStyle::GetAppStyleSetName(), "PlayWorld.PausePlaySession")));
+			FSlateIcon(FPCGEditorStyle::Get().GetStyleSetName(), "PCG.Command.PauseRegen")));
 
 		Section.AddEntry(FToolMenuEntry::InitToolBarButton(
 			FPCGEditorCommands::Get().ForceGraphRegeneration,
 			TAttribute<FText>(),
 			TAttribute<FText>(),
-			FSlateIcon(FAppStyle::GetAppStyleSetName(), "SourceControl.Actions.Refresh")));
+			TAttribute<FSlateIcon>::CreateLambda([]()
+			{
+				static const FSlateIcon ForceRegen = FSlateIcon(FPCGEditorStyle::Get().GetStyleSetName(), "PCG.Command.ForceRegen");
+				static const FSlateIcon ForceRegenClearCache = FSlateIcon(FPCGEditorStyle::Get().GetStyleSetName(), "PCG.Command.ForceRegenClearCache");
+					
+				FModifierKeysState ModifierKeys = FSlateApplication::Get().GetModifierKeys();
+				return ModifierKeys.IsControlDown() ? ForceRegenClearCache : ForceRegen;
+			})));
 
 		Section.AddEntry(FToolMenuEntry::InitToolBarButton(
 			FPCGEditorCommands::Get().CancelExecution,
 			TAttribute<FText>(),
 			TAttribute<FText>(),
-			FSlateIcon(FAppStyle::GetAppStyleSetName(), "Icons.Toolbar.Stop")));
-
+			FSlateIcon(FPCGEditorStyle::Get().GetStyleSetName(), "PCG.Command.StopRegen")));
+				
 		Section.AddEntry(FToolMenuEntry::InitToolBarButton(
 			FPCGEditorCommands::Get().OpenDebugObjectTreeTab,
 			TAttribute<FText>(),
 			TAttribute<FText>(),
-			FSlateIcon(FAppStyle::GetAppStyleSetName(), "LevelEditor.CreateBlankBlueprintClass")));
+			FSlateIcon(FPCGEditorStyle::Get().GetStyleSetName(), "PCG.Command.OpenDebugTreeTab")));
 
 		Section.AddSeparator(NAME_None);
 
@@ -629,7 +637,7 @@ void FPCGEditor::RegisterToolbar() const
 			 PCGEditorCommands.RunDeterminismGraphTest,
 			 TAttribute<FText>(),
 			 TAttribute<FText>(),
-			 FSlateIcon(FAppStyle::GetAppStyleSetName(), "BlueprintDebugger.TabIcon")));
+			 FSlateIcon(FPCGEditorStyle::Get().GetStyleSetName(), "PCG.Command.RunDeterminismTest")));
 
 		Section.AddSeparator(NAME_None);
 
@@ -637,7 +645,7 @@ void FPCGEditor::RegisterToolbar() const
 			 PCGEditorCommands.EditGraphSettings,
 			 TAttribute<FText>(),
 			 TAttribute<FText>(),
-			 FSlateIcon(FAppStyle::GetAppStyleSetName(), "FullBlueprintEditor.EditClassDefaults")));
+			 FSlateIcon(FPCGEditorStyle::Get().GetStyleSetName(), "PCG.Command.GraphSettings")));
 	}
 }
 
diff --git a/Engine/Plugins/PCG/Source/PCGEditor/Private/PCGEditorStyle.cpp b/Engine/Plugins/PCG/Source/PCGEditor/Private/PCGEditorStyle.cpp
--- a/Engine/Plugins/PCG/Source/PCGEditor/Private/PCGEditorStyle.cpp
+++ b/Engine/Plugins/PCG/Source/PCGEditor/Private/PCGEditorStyle.cpp
@@ -113,6 +113,17 @@ FPCGEditorStyle::FPCGEditorStyle() : FSlateStyleSet("PCGEditorStyle")
 	Set("ClassThumbnail.PCGGraphInstance", new IMAGE_BRUSH_SVG("Icons/PCGGraphInstance_64", Icon64x64));
 	Set("ClassIcon.PCGSettings", new IMAGE_BRUSH_SVG("Icons/PCGSettings_16", Icon16x16));
 	Set("ClassThumbnail.PCGSettings", new IMAGE_BRUSH_SVG("Icons/PCGSettings_64", Icon64x64));
+
+	// Command icons
+	Set("PCG.Command.Find", new IMAGE_BRUSH_SVG("Style/PCG_Command_Find", Icon20x20));
+	Set("PCG.Command.ForceRegen", new IMAGE_BRUSH_SVG("Style/PCG_Command_ForceRegen", Icon20x20));
+	Set("PCG.Command.ForceRegenClearCache", new IMAGE_BRUSH_SVG("Style/PCG_Command_ForceRegenClearCache", Icon20x20));
+	Set("PCG.Command.PauseRegen", new IMAGE_BRUSH_SVG("Style/PCG_Command_PauseRegen", Icon20x20));
+	Set("PCG.Command.StopRegen", new IMAGE_BRUSH_SVG("Style/PCG_Command_StopRegen", Icon20x20));
+	Set("PCG.Command.GraphSettings", new IMAGE_BRUSH_SVG("Style/PCG_Command_GraphSettings", Icon20x20));
+	Set("PCG.Command.OpenDebugTreeTab", new IMAGE_BRUSH_SVG("Style/PCG_Command_OpenDebugTreeTab", Icon20x20));
+	Set("PCG.Command.RunDeterminismTest", new IMAGE_BRUSH_SVG("Style/PCG_Command_RunDeterminismTest", Icon20x20));
+	
 }
 
 const FPCGEditorStyle& FPCGEditorStyle::Get()
