[PCG] Fix 'Get PCG Component Data' able to get its own output

#rnx

diff --git a/Engine/Plugins/PCG/Source/PCG/Private/Elements/PCGDataFromActor.cpp b/Engine/Plugins/PCG/Source/PCG/Private/Elements/PCGDataFromActor.cpp
--- a/Engine/Plugins/PCG/Source/PCG/Private/Elements/PCGDataFromActor.cpp
+++ b/Engine/Plugins/PCG/Source/PCG/Private/Elements/PCGDataFromActor.cpp
@@ -444,14 +444,15 @@ void FPCGDataFromActorElement::GatherWaitTasks(AActor* FoundActor, FPCGContext*
 	check(Settings);
 
 	UPCGComponent* SourceComponent = Context->SourceComponent.IsValid() ? Context->SourceComponent.Get() : nullptr;
+	const UPCGComponent* SourceOriginalComponent = SourceComponent ? SourceComponent->GetOriginalComponent() : nullptr;
 
-	if (!SourceComponent)
+	if (!SourceOriginalComponent)
 	{
 		return;
 	}
 
 	// We will prevent gathering the current execution - this task cannot wait on itself
-	AActor* ThisOwner = SourceComponent->GetOwner();
+	const AActor* SourceOwner = SourceOriginalComponent->GetOwner();
 
 	TInlineComponentArray<UPCGComponent*, 1> PCGComponents = PCGDataFromActorHelpers::GetPCGComponentsFromActor(
 		FoundActor,
@@ -460,11 +461,14 @@ void FPCGDataFromActorElement::GatherWaitTasks(AActor* FoundActor, FPCGContext*
 		Settings->bGetDataOnAllGrids,
 		Settings->AllowedGrids,
 		Settings->bComponentsMustOverlapSelf,
-		(Settings->bComponentsMustOverlapSelf && SourceComponent) ? SourceComponent->GetGridBounds() : FBox());
+		Settings->bComponentsMustOverlapSelf ? SourceComponent->GetGridBounds() : FBox());
 
 	for (UPCGComponent* Component : PCGComponents)
 	{
-		if (Component->GetOwner() == ThisOwner)
+		const UPCGComponent* OriginalComponent = Component ? Component->GetOriginalComponent() : nullptr;
+
+		// Avoid waiting on our own execution (including local components).
+		if (!OriginalComponent || OriginalComponent == SourceOriginalComponent || (Settings->ActorSelector.bIgnoreSelfAndChildren && OriginalComponent->GetOwner() == SourceOwner))
 		{
 			continue;
 		}
@@ -606,20 +610,19 @@ void FPCGDataFromActorElement::ProcessActor(FPCGContext* Context, const UPCGData
 	check(Settings);
 
 	UPCGComponent* SourceComponent = Context->SourceComponent.IsValid() ? Context->SourceComponent.Get() : nullptr;
+	const UPCGComponent* SourceOriginalComponent = SourceComponent ? SourceComponent->GetOriginalComponent() : nullptr;
 
-	if (!FoundActor || !IsValid(FoundActor) || !SourceComponent)
+	if (!FoundActor || !IsValid(FoundActor) || !SourceOriginalComponent)
 	{
 		return;
 	}
 
-	AActor* ThisOwner = SourceComponent->GetOwner();
+	const AActor* SourceOwner = SourceOriginalComponent->GetOwner();
 	TInlineComponentArray<UPCGComponent*, 1> PCGComponents;
 	bool bHasGeneratedPCGData = false;
 	FProperty* FoundProperty = nullptr;
 
-	const bool bCanGetDataFromComponent = (FoundActor != ThisOwner);
-
-	if (bCanGetDataFromComponent && (Settings->Mode == EPCGGetDataFromActorMode::GetDataFromPCGComponent || Settings->Mode == EPCGGetDataFromActorMode::GetDataFromPCGComponentOrParseComponents))
+	if (Settings->Mode == EPCGGetDataFromActorMode::GetDataFromPCGComponent || Settings->Mode == EPCGGetDataFromActorMode::GetDataFromPCGComponentOrParseComponents)
 	{
 		PCGComponents = PCGDataFromActorHelpers::GetPCGComponentsFromActor(
 			FoundActor,
@@ -630,6 +633,16 @@ void FPCGDataFromActorElement::ProcessActor(FPCGContext* Context, const UPCGData
 			Settings->bComponentsMustOverlapSelf,
 			Settings->bComponentsMustOverlapSelf ? SourceComponent->GetGridBounds() : FBox());
 
+		// Remove any PCG components that don't belong to an external execution context (i.e. share the same original component), or that share a common root actor.
+		PCGComponents.RemoveAllSwap([Settings, SourceOwner, SourceOriginalComponent](UPCGComponent* Component)
+		{
+			const UPCGComponent* OriginalComponent = Component ? Component->GetOriginalComponent() : nullptr;
+
+			return !OriginalComponent
+				|| OriginalComponent == SourceOriginalComponent
+				|| (Settings->ActorSelector.bIgnoreSelfAndChildren && OriginalComponent->GetOwner() == SourceOwner);
+		});
+
 		for (UPCGComponent* Component : PCGComponents)
 		{
 			bHasGeneratedPCGData |= !Component->GetGeneratedGraphOutput().TaggedData.IsEmpty();
@@ -646,15 +659,11 @@ void FPCGDataFromActorElement::ProcessActor(FPCGContext* Context, const UPCGData
 	// Some additional validation
 	if (Settings->Mode == EPCGGetDataFromActorMode::GetDataFromPCGComponent && !bHasGeneratedPCGData)
 	{
-		if (bCanGetDataFromComponent)
+		if (!PCGComponents.IsEmpty())
 		{
 			PCGE_LOG(Log, GraphAndLog, FText::Format(LOCTEXT("ActorHasNoGeneratedData", "Actor '{0}' does not have any previously generated data"), FText::FromName(FoundActor->GetFName())));
 		}
-		else
-		{
-			PCGE_LOG(Error, GraphAndLog, FText::Format(LOCTEXT("ActorCannotGetOwnData", "Actor '{0}' cannot get its own generated data during generation"), FText::FromName(FoundActor->GetFName())));
-		}
-		
+
 		return;
 	}
 	else if (Settings->Mode == EPCGGetDataFromActorMode::GetDataFromProperty && !FoundProperty)
diff --git a/Engine/Plugins/PCG/Source/PCG/Private/Elements/PCGTypedGetter.cpp b/Engine/Plugins/PCG/Source/PCG/Private/Elements/PCGTypedGetter.cpp
--- a/Engine/Plugins/PCG/Source/PCG/Private/Elements/PCGTypedGetter.cpp
+++ b/Engine/Plugins/PCG/Source/PCG/Private/Elements/PCGTypedGetter.cpp
@@ -205,11 +205,8 @@ UPCGGetPCGComponentSettings::UPCGGetPCGComponentSettings()
 	bDisplayModeSettings = false;
 	Mode = EPCGGetDataFromActorMode::GetDataFromPCGComponent;
 
-	// It doesn't make sense (for dependency reasons) to look for components on self/same hierarchy, so we can just hide those options.
 	ActorSelector.bShowActorFilter = false;
 	ActorSelector.ActorFilter = EPCGActorFilter::AllWorldActors;
-	ActorSelector.bShowIgnoreSelfAndChildren = false;
-	ActorSelector.bIgnoreSelfAndChildren = true;
 }
 
 TArray<FPCGPinProperties> UPCGGetPCGComponentSettings::OutputPinProperties() const
@@ -226,7 +223,7 @@ FText UPCGGetPCGComponentSettings::GetNodeTooltipText() const
 {
 	return LOCTEXT("GetPCGComponentTooltip", "Builds a collection of data from other PCG components on the selected actors.\n"
 		"Automatically tags each output with the grid size it was collected from, prefixed by \"PCG_GridSize_\" (e.g.PCG_GridSize_12800).\n"
-		"Note: a component cannot get component data from itself or other components in its hierarchy, as it could create a circular dependency.");
+		"Note: a component cannot get component data from itself or other components in its execution context, as it could create a circular dependency.");
 }
 #endif
 
diff --git a/Engine/Plugins/PCG/Source/PCG/Public/Elements/PCGTypedGetter.h b/Engine/Plugins/PCG/Source/PCG/Public/Elements/PCGTypedGetter.h
--- a/Engine/Plugins/PCG/Source/PCG/Public/Elements/PCGTypedGetter.h
+++ b/Engine/Plugins/PCG/Source/PCG/Public/Elements/PCGTypedGetter.h
@@ -142,7 +142,7 @@ public:
  * Builds a collection of data from other PCG components on the selected actors. Automatically tags each output with the grid size it was collected
  * from, prefixed by "PCG_GridSize_" (e.g. PCG_GridSize_12800).
  *
- * Note: a component cannot get component data from itself or other components in its hierarchy, as it could create a circular dependency.
+ * Note: a component cannot get component data from itself or other components in its execution context, as it could create a circular dependency.
  */
 UCLASS(BlueprintType, ClassGroup = (Procedural))
 class PCG_API UPCGGetPCGComponentSettings : public UPCGDataFromActorSettings
