[PCG] Fix leaking PA records on the PCGWorldActor when changing grid size
* Add PCGWorldActor grid size as a valid grid size to avoid ensures triggered by having a grid size that isn't (power of 2) * 100
* Fix deleting PCGWorldActor leaving orphaned PAs behind

#rnx

diff --git a/Engine/Plugins/PCG/Source/PCG/Private/PCGCommon.cpp b/Engine/Plugins/PCG/Source/PCG/Private/PCGCommon.cpp
--- a/Engine/Plugins/PCG/Source/PCG/Private/PCGCommon.cpp
+++ b/Engine/Plugins/PCG/Source/PCG/Private/PCGCommon.cpp
@@ -2,6 +2,7 @@
 
 #include "PCGCommon.h"
 #include "PCGSubsystem.h"
+#include "PCGWorldActor.h"
 
 namespace PCGFeatureSwitches
 {
@@ -31,9 +32,17 @@ namespace PCGHiGenGrid
 	{
 		// Must be a power of 2 (in m) and within the valid range
 		// TODO: support other units
-		return FMath::IsPowerOfTwo(InGridSize / 100)
+		if (FMath::IsPowerOfTwo(InGridSize / 100)
 			&& InGridSize >= GridToGridSize(EPCGHiGenGrid::GridMin)
-			&& InGridSize <= GridToGridSize(EPCGHiGenGrid::GridMax);
+			&& InGridSize <= GridToGridSize(EPCGHiGenGrid::GridMax))
+		{
+			return true;
+		}
+
+		UPCGSubsystem* PCGSubsystem = UPCGSubsystem::GetSubsystemForCurrentWorld();
+		APCGWorldActor* PCGWorldActor = PCGSubsystem ? PCGSubsystem->GetPCGWorldActor() : nullptr;
+
+		return !PCGWorldActor || PCGWorldActor->PartitionGridSize == InGridSize;
 	}
 
 	bool IsValidGrid(EPCGHiGenGrid InGrid)
diff --git a/Engine/Plugins/PCG/Source/PCG/Private/PCGSubsystem.cpp b/Engine/Plugins/PCG/Source/PCG/Private/PCGSubsystem.cpp
--- a/Engine/Plugins/PCG/Source/PCG/Private/PCGSubsystem.cpp
+++ b/Engine/Plugins/PCG/Source/PCG/Private/PCGSubsystem.cpp
@@ -310,6 +310,9 @@ void UPCGSubsystem::DestroyCurrentPCGWorldActor()
 
 void UPCGSubsystem::DestroyAllPCGWorldActors()
 {
+	// Delete all PAs first to avoid leaving orphans behind.
+	DeleteSerializedPartitionActors(/*bOnlyDeleteUnused=*/false);
+
 	// Get rid of current PCG world actor first
 	DestroyCurrentPCGWorldActor();
 	
@@ -934,8 +937,9 @@ FPCGTaskId UPCGSubsystem::ForAllOverlappingCells(const FBox& InBounds, const PCG
 
 	PCGHiGenGrid::FSizeArray GridSizes = InGridSizes;
 
-	// We have no use for unbounded grids as this is a grid-centric function.
+	// We have no use for unbounded grids as this is a grid-centric function. Also discard invalid grid sizes.
 	GridSizes.Remove(static_cast<uint32>(EPCGHiGenGrid::Unbounded));
+	GridSizes.RemoveAll([](uint32 GridSize) { return !ensure(PCGHiGenGrid::IsValidGridSize(GridSize)); });
 
 	if (GridSizes.IsEmpty())
 	{
@@ -951,7 +955,6 @@ FPCGTaskId UPCGSubsystem::ForAllOverlappingCells(const FBox& InBounds, const PCG
 	TArray<FPCGTaskId> CellTasks;
 	for (uint32 GridSize : GridSizes)
 	{
-		ensure(PCGHiGenGrid::IsValidGridSize(GridSize));
 		const FGuid* GuidPtr = GridSizeToGuid.Find(GridSize);
 		if (!ensure(GuidPtr))
 		{
diff --git a/Engine/Plugins/PCG/Source/PCG/Private/PCGWorldActor.cpp b/Engine/Plugins/PCG/Source/PCG/Private/PCGWorldActor.cpp
--- a/Engine/Plugins/PCG/Source/PCG/Private/PCGWorldActor.cpp
+++ b/Engine/Plugins/PCG/Source/PCG/Private/PCGWorldActor.cpp
@@ -314,6 +314,7 @@ void APCGWorldActor::OnPartitionGridSizeChanged()
 
 	// Then delete all PCGPartitionActors
 	PCGSubsystem->DeleteSerializedPartitionActors(/*bDeleteOnlyUnused=*/false);
+	SerializedPartitionActorRecords.Reset();
 
 	// And finally, regenerate all components that are partitioned (registered to the PCGSubsystem)
 	// to let them recreate the needed PCG Partition Actors.
