Spatial noise use seed deprecation fix

#rnx

diff --git a/Engine/Plugins/PCG/Source/PCG/Private/Elements/PCGSpatialNoise.cpp b/Engine/Plugins/PCG/Source/PCG/Private/Elements/PCGSpatialNoise.cpp
--- a/Engine/Plugins/PCG/Source/PCG/Private/Elements/PCGSpatialNoise.cpp
+++ b/Engine/Plugins/PCG/Source/PCG/Private/Elements/PCGSpatialNoise.cpp
@@ -7,6 +7,7 @@
 
 #include "PCGComponent.h"
 #include "PCGContext.h"
+#include "PCGCustomVersion.h"
 #include "PCGPin.h"
 
 #include "Data/PCGPointData.h"
@@ -27,6 +28,44 @@ UPCGSpatialNoiseSettings::UPCGSpatialNoiseSettings()
 	ValueTarget.SetPointProperty(EPCGPointProperties::Density);
 }
 
+void UPCGSpatialNoiseSettings::PostLoad()
+{
+	Super::PostLoad();
+
+	if (bForceNoUseSeed)
+	{
+		bUseSeed = false;
+	}
+}
+
+void UPCGSpatialNoiseSettings::PostEditImport()
+{
+	Super::PostEditImport();
+
+	if (bForceNoUseSeed)
+	{
+		bUseSeed = false;
+	}
+}
+
+#if WITH_EDITOR
+void UPCGSpatialNoiseSettings::ApplyDeprecation(UPCGNode* InOutNode)
+{
+	// Important note: this version change is not related to the reason we do deprecation here, but
+	// this happened in a close timeframe and we can close the gap here. The original issue is that
+	// we've added seed usage on the SpatialNoise settings at a slightly later version (a few days),
+	// but since this applies at the CDO level too, this was adding seed usage also to the data prior to
+	// this change, which had an effect on generated results.
+	if (DataVersion < FPCGCustomVersion::NoMoreSpatialDataConversionToPointDataByDefaultOnNonPointPins)
+	{
+		bForceNoUseSeed = true;
+		bUseSeed = false;
+	}
+
+	Super::ApplyDeprecation(InOutNode);
+}
+#endif // WITH_EDITOR
+
 TArray<FPCGPinProperties> UPCGSpatialNoiseSettings::InputPinProperties() const
 {
 	return Super::DefaultPointInputPinProperties();
diff --git a/Engine/Plugins/PCG/Source/PCG/Public/Elements/PCGSpatialNoise.h b/Engine/Plugins/PCG/Source/PCG/Public/Elements/PCGSpatialNoise.h
--- a/Engine/Plugins/PCG/Source/PCG/Public/Elements/PCGSpatialNoise.h
+++ b/Engine/Plugins/PCG/Source/PCG/Public/Elements/PCGSpatialNoise.h
@@ -63,14 +63,18 @@ class UPCGSpatialNoiseSettings : public UPCGSettings
 
 public:
 	UPCGSpatialNoiseSettings();
+	//~Begin UObject interface
+	virtual void PostLoad() override;
+	virtual void PostEditImport() override;
+	//~End UObject interface
 
 	//~Begin UPCGSettings interface
 #if WITH_EDITOR
 	virtual FName GetDefaultNodeName() const override { return FName(TEXT("Spatial Noise")); }
 	virtual FText GetDefaultNodeTitle() const override { return NSLOCTEXT("PCGSpatialNoise", "NodeTitle", "Spatial Noise"); }
 	virtual EPCGSettingsType GetType() const override { return EPCGSettingsType::Spatial; }
+	virtual void ApplyDeprecation(UPCGNode* InOutNode) override;
 #endif
-	
 
 protected:
 	virtual TArray<FPCGPinProperties> InputPinProperties() const override;
@@ -145,6 +149,11 @@ public:
 	// will makes the falloff harsher or softer
 	UPROPERTY(BlueprintReadWrite, EditAnywhere, Category = Settings, meta = (EditCondition = "Mode == PCGSpatialNoiseMode::EdgeMask2D", EditConditionHides, ClampMin = "0", PCG_Overridable))
 	float EdgeBlendCurveIntensity = 1.0;
+
+private:
+	// Private field to fix a deprecation issue related to seed usage between UE 5.3 and UE 5.4
+	UPROPERTY()
+	bool bForceNoUseSeed = false;
 };
 
 class FPCGSpatialNoise : public FPCGPointProcessingElementBase
