[PCG] Fix Component "leak" (previous generation component not getting destroyed) on Undo/Redo of BP PCG Component
- FPCGComponentInstanceData not using the proper component to transfer ManagedResources and apply Remapping
- Fix is to register the Destroyed component to the PCGSubsystem so that it can be discovered by name and used in the FPCGComponentInstanceData::ApplyToComponent instead of the captured SourceComponent
- This needs to be revisited post release as there are still some known Undo/Redo issues relating to PCGPartitionActors
- Also fix PCGWorldActor modify call being done after the modification (wasn't properly getting undone when registration happened inside a transaction)
- Since this is a late fix it is gated in a CVar : "pcg.ConstructionScriptFix"

#rnx

diff --git a/Engine/Plugins/PCG/Source/PCG/Private/Graph/PCGGraphExecutor.cpp b/Engine/Plugins/PCG/Source/PCG/Private/Graph/PCGGraphExecutor.cpp
--- a/Engine/Plugins/PCG/Source/PCG/Private/Graph/PCGGraphExecutor.cpp
+++ b/Engine/Plugins/PCG/Source/PCG/Private/Graph/PCGGraphExecutor.cpp
@@ -1160,6 +1160,7 @@ void FPCGGraphExecutor::ClearAllTasks()
 	ReadyTasks.Reset();
 	ActiveTasks.Reset();
 	SleepingTasks.Reset();
+	TaskSuccessors.Reset();
 }
 
 void FPCGGraphExecutor::QueueNextTasks(FPCGTaskId FinishedTask, bool bIgnoreMissingTasks)
diff --git a/Engine/Plugins/PCG/Source/PCG/Private/PCGComponent.cpp b/Engine/Plugins/PCG/Source/PCG/Private/PCGComponent.cpp
--- a/Engine/Plugins/PCG/Source/PCG/Private/PCGComponent.cpp
+++ b/Engine/Plugins/PCG/Source/PCG/Private/PCGComponent.cpp
@@ -68,6 +68,11 @@ namespace PCGComponent
 		false,
 		TEXT("Disable refresh for all PCG Components."));
 
+	static TAutoConsoleVariable<bool> CVarConstructionScriptFix(
+		TEXT("pcg.ConstructionScriptFix"),
+		true,
+		TEXT("This CVar will be removed in future releases, it allows disabling this fix if regressions are found."));
+
 	template <typename DelegateType>
 	static void BroadcastDynamicDelegate(const DelegateType& Delegate, UPCGComponent* PCGComponent)
 	{
@@ -1376,6 +1381,11 @@ void UPCGComponent::OnComponentDestroyed(bool bDestroyingHierarchy)
 		{
 			Subsystem->UnregisterPCGComponent(this);
 		}
+
+		if (IsCreatedByConstructionScript() && PCGComponent::CVarConstructionScriptFix.GetValueOnAnyThread())
+		{
+			Subsystem->SetConstructionScriptSourceComponent(this);
+		}
 	}
 #endif // WITH_EDITOR
 
@@ -3438,18 +3448,11 @@ FPCGComponentInstanceData::FPCGComponentInstanceData(const UPCGComponent* InSour
 	: FActorComponentInstanceData(InSourceComponent)
 	, SourceComponent(InSourceComponent)
 {
-	if (SourceComponent)
-	{
-		SourceComponent->GetManagedResources(GeneratedResources);
-#if WITH_EDITOR
-		LoadedPreviewResources = SourceComponent->LoadedPreviewResources;
-#endif
-	}
 }
 
 bool FPCGComponentInstanceData::ContainsData() const
 {
-	return GeneratedResources.Num() > 0 || Super::ContainsData();
+	return true;
 }
 
 void FPCGComponentInstanceData::ApplyToComponent(UActorComponent* Component, const ECacheApplyPhase CacheApplyPhase)
@@ -3460,14 +3463,29 @@ void FPCGComponentInstanceData::ApplyToComponent(UActorComponent* Component, con
 	{
 		UPCGComponent* PCGComponent = CastChecked<UPCGComponent>(Component);
 
+		// IMPORTANT NOTE:
+		// ConstructionSourceComponent: The previous instance of the PCGComponent (not related to annotation, will be the OldComponent in the OnObjectsReplaced call)
+		// We will transfer/copy properties/resources from this component that flow forward only (aren't undo/redoable)
+		//
+		// SourceComponent: The previous instance (same as ConstructionSourceComponent on a regular edit but an annotated Component when in a Undo/Redo operation)
+		// We will transfer/copy properties from this component that should reflect undo redo (annotation)
+		const UPCGComponent* ConstructionSourceComponent = SourceComponent;
+#if WITH_EDITOR
+		if (UPCGSubsystem* Subsystem = PCGComponent->GetSubsystem())
+		{
+			UPCGComponent* FoundConstructionSourceComponent = nullptr;
+			if(Subsystem->RemoveAndCopyConstructionScriptSourceComponent(Component->GetOwner(), Component->GetFName(), FoundConstructionSourceComponent))
+			{
+				ConstructionSourceComponent = FoundConstructionSourceComponent;
+			}
+		}
+#endif
 		// IMPORTANT NOTE:
 		// Any non-visible (i.e. UPROPERTY() with no specifiers) are NOT copied over when re-running the construction script
 		// This means that some properties need to be reapplied here manually unless we make them visible
 		if (SourceComponent)
 		{
-			// Critical: LastGeneratedBounds & GeneratedGraphOutput
-			PCGComponent->LastGeneratedBounds = SourceComponent->LastGeneratedBounds;
-
+			// Critical: GeneratedGraphOutput
 			PCGComponent->GeneratedGraphOutput = SourceComponent->GeneratedGraphOutput;
 			// Re-outer any data moved here
 			for (FPCGTaggedData& TaggedData : PCGComponent->GeneratedGraphOutput.TaggedData)
@@ -3479,9 +3497,6 @@ void FPCGComponentInstanceData::ApplyToComponent(UActorComponent* Component, con
 			}
 
 #if WITH_EDITOR
-			// bDirtyGenerated is transient.
-			PCGComponent->bDirtyGenerated = SourceComponent->bDirtyGenerated;
-
 			// While this is serialized, it is not properly copied because it is not visible. This is needed otherwise a refresh can retoggle from not generated to generated
 			PCGComponent->bForceGenerateOnBPAddedToWorld = SourceComponent->bForceGenerateOnBPAddedToWorld;
 
@@ -3494,46 +3509,57 @@ void FPCGComponentInstanceData::ApplyToComponent(UActorComponent* Component, con
 			// TODO Validate usefulness + move accordingly
 		}
 		
-		// Duplicate generated resources + retarget them
-		TArray<TObjectPtr<UPCGManagedResource>> DuplicatedResources;
-		for (const TObjectPtr<UPCGManagedResource>& Resource : GeneratedResources)
+		if (ConstructionSourceComponent)
 		{
-			if (Resource)
+			// Critical: LastGeneratedBounds
+			PCGComponent->LastGeneratedBounds = ConstructionSourceComponent->LastGeneratedBounds;
+
+			// Duplicate generated resources + retarget them
+			TArray<TObjectPtr<UPCGManagedResource>> DuplicatedResources;
+			for (const TObjectPtr<UPCGManagedResource>& Resource : ConstructionSourceComponent->GeneratedResources)
 			{
-				UPCGManagedResource* DuplicatedResource = CastChecked<UPCGManagedResource>(StaticDuplicateObject(Resource, PCGComponent, FName()));
-				DuplicatedResource->PostApplyToComponent();
-				DuplicatedResources.Add(DuplicatedResource);
+				if (Resource)
+				{
+					UPCGManagedResource* DuplicatedResource = CastChecked<UPCGManagedResource>(StaticDuplicateObject(Resource, PCGComponent, FName()));
+					DuplicatedResource->PostApplyToComponent();
+					DuplicatedResources.Add(DuplicatedResource);
+				}
 			}
-		}
 
-		if (DuplicatedResources.Num() > 0)
-		{
-			PCGComponent->SetManagedResources(DuplicatedResources);
-		}
+			if (DuplicatedResources.Num() > 0)
+			{
+				PCGComponent->SetManagedResources(DuplicatedResources);
+			}
 
 #if WITH_EDITOR
-		TArray<TObjectPtr<UPCGManagedResource>> DuplicateLoadedPreviewResources;
-		for (const TObjectPtr<UPCGManagedResource>& Resource : LoadedPreviewResources)
-		{
-			if (Resource)
+			// bDirtyGenerated is transient.
+			PCGComponent->bDirtyGenerated = ConstructionSourceComponent->bDirtyGenerated;
+
+			TArray<TObjectPtr<UPCGManagedResource>> DuplicateLoadedPreviewResources;
+			for (const TObjectPtr<UPCGManagedResource>& Resource : ConstructionSourceComponent->LoadedPreviewResources)
+			{
+				if (Resource)
+				{
+					UPCGManagedResource* DuplicatedResource = CastChecked<UPCGManagedResource>(StaticDuplicateObject(Resource, PCGComponent, FName()));
+					DuplicatedResource->PostApplyToComponent();
+					DuplicateLoadedPreviewResources.Add(DuplicatedResource);
+				}
+			}
+
+			if (DuplicateLoadedPreviewResources.Num() > 0)
 			{
-				UPCGManagedResource* DuplicatedResource = CastChecked<UPCGManagedResource>(StaticDuplicateObject(Resource, PCGComponent, FName()));
-				DuplicatedResource->PostApplyToComponent();
-				DuplicateLoadedPreviewResources.Add(DuplicatedResource);
+				PCGComponent->LoadedPreviewResources = DuplicateLoadedPreviewResources;
 			}
-		}
 
-		if (DuplicateLoadedPreviewResources.Num() > 0)
-		{
-			PCGComponent->LoadedPreviewResources = DuplicateLoadedPreviewResources;
+			// Move over invocation lists for dynamic delegates
+			PCGComponent->OnPCGGraphStartGeneratingExternal = ConstructionSourceComponent->OnPCGGraphStartGeneratingExternal;
+			PCGComponent->OnPCGGraphCancelledExternal = ConstructionSourceComponent->OnPCGGraphCancelledExternal;
+			PCGComponent->OnPCGGraphGeneratedExternal = ConstructionSourceComponent->OnPCGGraphGeneratedExternal;
+			PCGComponent->OnPCGGraphCleanedExternal = ConstructionSourceComponent->OnPCGGraphCleanedExternal;
+#endif
 		}
 
-		// Move over invocation lists for dynamic delegates
-		PCGComponent->OnPCGGraphStartGeneratingExternal = SourceComponent->OnPCGGraphStartGeneratingExternal;
-		PCGComponent->OnPCGGraphCancelledExternal = SourceComponent->OnPCGGraphCancelledExternal;
-		PCGComponent->OnPCGGraphGeneratedExternal = SourceComponent->OnPCGGraphGeneratedExternal;
-		PCGComponent->OnPCGGraphCleanedExternal = SourceComponent->OnPCGGraphCleanedExternal;
-
+#if WITH_EDITOR
 		// Reconnect callbacks
 		if (PCGComponent->GraphInstance)
 		{
@@ -3547,23 +3573,23 @@ void FPCGComponentInstanceData::ApplyToComponent(UActorComponent* Component, con
 
 		// Also remap
 		UPCGSubsystem* Subsystem = PCGComponent->GetSubsystem();
-		if (Subsystem && SourceComponent)
+		if (Subsystem && ConstructionSourceComponent)
 		{
-			Subsystem->RemapPCGComponent(SourceComponent, PCGComponent, bDoActorMapping);
+			Subsystem->RemapPCGComponent(ConstructionSourceComponent, PCGComponent, bDoActorMapping);
 		}
 
 #if WITH_EDITOR
 		// Disconnect callbacks on source.
-		if (SourceComponent && SourceComponent->GraphInstance)
+		if (ConstructionSourceComponent && ConstructionSourceComponent->GraphInstance)
 		{
-			SourceComponent->GraphInstance->TeardownCallbacks();
+			ConstructionSourceComponent->GraphInstance->TeardownCallbacks();
 		}
 
 		// Finally, start a delayed refresh task (if there is not one already), in editor only
 		// It is important to be delayed, because we cannot spawn Partition Actors within this scope,
 		// because we are in a construction script.
 		// Note that we only do this if we are not currently loading
-		if (!SourceComponent || !SourceComponent->HasAllFlags(RF_WasLoaded))
+		if (!ConstructionSourceComponent || !ConstructionSourceComponent->HasAllFlags(RF_WasLoaded))
 		{
 			PCGComponent->Refresh();
 		}
@@ -3576,7 +3602,6 @@ void FPCGComponentInstanceData::AddReferencedObjects(FReferenceCollector& Collec
 	Super::AddReferencedObjects(Collector);
 
 	Collector.AddReferencedObject(SourceComponent);
-	Collector.AddReferencedObjects(GeneratedResources);
 }
 
 #undef LOCTEXT_NAMESPACE
diff --git a/Engine/Plugins/PCG/Source/PCG/Private/PCGSubsystem.cpp b/Engine/Plugins/PCG/Source/PCG/Private/PCGSubsystem.cpp
--- a/Engine/Plugins/PCG/Source/PCG/Private/PCGSubsystem.cpp
+++ b/Engine/Plugins/PCG/Source/PCG/Private/PCGSubsystem.cpp
@@ -213,12 +213,46 @@ UPCGSubsystem* UPCGSubsystem::GetActiveEditorInstance()
 
 	return nullptr;
 }
+
+void UPCGSubsystem::SetConstructionScriptSourceComponent(UPCGComponent* InComponent)
+{
+	if (InComponent)
+	{
+		if (AActor* Owner = InComponent->GetOwner())
+		{
+			PerActorConstructionScriptSourceComponents.FindOrAdd(Owner).SourceComponentsMap.Add(InComponent->GetFName(), InComponent);
+		}
+	}
+}
+
+bool UPCGSubsystem::RemoveAndCopyConstructionScriptSourceComponent(AActor* InComponentOwner, FName InComponentName, UPCGComponent*& OutSourceComponent)
+{
+	OutSourceComponent = nullptr;
+	if (FConstructionScriptSourceComponents* Found = PerActorConstructionScriptSourceComponents.Find(InComponentOwner))
+	{
+		TObjectPtr<UPCGComponent> FoundComponent;
+		if (Found->SourceComponentsMap.RemoveAndCopyValue(InComponentName, FoundComponent))
+		{
+			OutSourceComponent = FoundComponent;
+			if (Found->SourceComponentsMap.IsEmpty())
+			{
+				PerActorConstructionScriptSourceComponents.Remove(InComponentOwner);
+			}
+		}
+	}
+
+	return OutSourceComponent != nullptr;
+}
 #endif
 
 void UPCGSubsystem::Tick(float DeltaSeconds)
 {
 	Super::Tick(DeltaSeconds);
 
+#if WITH_EDITOR
+	PerActorConstructionScriptSourceComponents.Empty();
+#endif
+
 	if (!bHasTickedOnce)
 	{
 #if WITH_EDITOR
diff --git a/Engine/Plugins/PCG/Source/PCG/Public/Grid/PCGPartitionActor.h b/Engine/Plugins/PCG/Source/PCG/Public/Grid/PCGPartitionActor.h
--- a/Engine/Plugins/PCG/Source/PCG/Public/Grid/PCGPartitionActor.h
+++ b/Engine/Plugins/PCG/Source/PCG/Public/Grid/PCGPartitionActor.h
@@ -117,7 +117,7 @@ private:
 	// Note: this map is not a property and not serialized since we will rebuild it from the LocalToOriginal
 	TMap<TObjectPtr<UPCGComponent>, TObjectPtr<UPCGComponent>> OriginalToLocal;
 
-	UPROPERTY()
+	UPROPERTY(NonTransactional)
 	TMap<TObjectPtr<UPCGComponent>, TSoftObjectPtr<UPCGComponent>> LocalToOriginal;
 
 	// PCG components that are cleared when in preview-on-load mode are kept aside and put back when serializing to prevent changes
diff --git a/Engine/Plugins/PCG/Source/PCG/Public/PCGComponent.h b/Engine/Plugins/PCG/Source/PCG/Public/PCGComponent.h
--- a/Engine/Plugins/PCG/Source/PCG/Public/PCGComponent.h
+++ b/Engine/Plugins/PCG/Source/PCG/Public/PCGComponent.h
@@ -656,14 +656,6 @@ protected:
 	virtual void ApplyToComponent(UActorComponent* Component, const ECacheApplyPhase CacheApplyPhase) override;
 	virtual void AddReferencedObjects(FReferenceCollector& Collector) override;
 
-	UPROPERTY()
-	TArray<TObjectPtr<UPCGManagedResource>> GeneratedResources;
-
-#if WITH_EDITORONLY_DATA
-	UPROPERTY()
-	TArray<TObjectPtr<UPCGManagedResource>> LoadedPreviewResources;
-#endif
-
 	UPROPERTY()
 	TObjectPtr<const UPCGComponent> SourceComponent;
 };
diff --git a/Engine/Plugins/PCG/Source/PCG/Public/PCGSubsystem.h b/Engine/Plugins/PCG/Source/PCG/Public/PCGSubsystem.h
--- a/Engine/Plugins/PCG/Source/PCG/Public/PCGSubsystem.h
+++ b/Engine/Plugins/PCG/Source/PCG/Public/PCGSubsystem.h
@@ -36,6 +36,17 @@ class UWorld;
 DECLARE_MULTICAST_DELEGATE_OneParam(FPCGOnComponentGenerationCompleteOrCancelled, UPCGSubsystem*);
 #endif // WITH_EDITOR
 
+USTRUCT()
+struct FConstructionScriptSourceComponents
+{
+	GENERATED_USTRUCT_BODY()
+
+#if WITH_EDITORONLY_DATA
+	UPROPERTY(Transient)
+	TMap<FName, TObjectPtr<UPCGComponent>> SourceComponentsMap;
+#endif
+};
+
 /**
 * UPCGSubsystem
 */
@@ -74,6 +85,9 @@ public:
 #if WITH_EDITOR
 	/** Returns PIE world if it is active, otherwise returns editor world. */
 	static UPCGSubsystem* GetActiveEditorInstance();
+
+	void SetConstructionScriptSourceComponent(UPCGComponent* InComponent);
+	bool RemoveAndCopyConstructionScriptSourceComponent(AActor* InComponentOwner, FName InComponentName, UPCGComponent*& OutSourceComponent);
 #endif
 
 	/** Subsystem must not be used without this condition being true. */
@@ -322,6 +336,11 @@ private:
 #if WITH_EDITOR
 	FCriticalSection PCGWorldActorLock;
 #endif
+
+#if WITH_EDITORONLY_DATA
+	UPROPERTY(Transient)
+	TMap<TObjectPtr<AActor>, FConstructionScriptSourceComponents> PerActorConstructionScriptSourceComponents;
+#endif
 };
 
 #if UE_ENABLE_INCLUDE_ORDER_DEPRECATED_IN_5_2
diff --git a/Engine/Plugins/PCG/Source/PCG/Public/PCGWorldActor.h b/Engine/Plugins/PCG/Source/PCG/Public/PCGWorldActor.h
--- a/Engine/Plugins/PCG/Source/PCG/Public/PCGWorldActor.h
+++ b/Engine/Plugins/PCG/Source/PCG/Public/PCGWorldActor.h
@@ -73,10 +73,10 @@ public:
 	void GetTransientGridGuids(PCGHiGenGrid::FSizeToGuidMap& OutSizeToGuidMap) const;
 
 	/** Add a record for tracking loaded and unloaded partition actors. */
-	void AddSerializedPartitionActorRecord(const FPCGPartitionActorRecord& PartitionActorRecord) { SerializedPartitionActorRecords.Add(PartitionActorRecord); Modify(); }
+	void AddSerializedPartitionActorRecord(const FPCGPartitionActorRecord& PartitionActorRecord) { Modify(); SerializedPartitionActorRecords.Add(PartitionActorRecord); }
 
 	/** Remove a record for tracking loaded and unloaded partition actors. */
-	void RemoveSerializedPartitionActorRecord(const FPCGPartitionActorRecord& PartitionActorRecord) { SerializedPartitionActorRecords.Remove(PartitionActorRecord); Modify(); }
+	void RemoveSerializedPartitionActorRecord(const FPCGPartitionActorRecord& PartitionActorRecord) { Modify(); SerializedPartitionActorRecords.Remove(PartitionActorRecord); }
 
 	/** Returns true if there is record of a partition actor living in a certain grid cell, regardless of whether or not it is loaded. */
 	bool DoesSerializedPartitionActorExist(const FGuid& GridGuid, uint32 GridSize, const FIntVector& GridCoords) const { return SerializedPartitionActorRecords.Contains({ GridGuid, GridSize, GridCoords }); }
