Fixed crash when changing the trigger or the editing mode on a PCG component in the blueprint editor.

#rnx

diff --git a/Engine/Plugins/PCG/Source/PCG/Private/PCGComponent.cpp b/Engine/Plugins/PCG/Source/PCG/Private/PCGComponent.cpp
--- a/Engine/Plugins/PCG/Source/PCG/Private/PCGComponent.cpp
+++ b/Engine/Plugins/PCG/Source/PCG/Private/PCGComponent.cpp
@@ -3379,6 +3379,7 @@ void UPCGComponent::ChangeTransientState(EPCGEditorDirtyMode NewEditingMode)
 	{
 		if (NewEditingMode == EPCGEditorDirtyMode::Preview)
 		{
+			bShouldMarkDirty = true;
 			MarkPackageDirty();
 		}
 
@@ -3405,6 +3406,7 @@ void UPCGComponent::ChangeTransientState(EPCGEditorDirtyMode NewEditingMode)
 
 		if (NewEditingMode != EPCGEditorDirtyMode::Preview)
 		{
+			bShouldMarkDirty = true;
 			MarkPackageDirty();
 		}
 	}
@@ -3422,10 +3424,16 @@ void UPCGComponent::ChangeTransientState(EPCGEditorDirtyMode NewEditingMode)
 		}
 	}
 
-	// Changing the transient state can and will play with packages and is not meant to be undoable
-	if (GEditor && GEditor->Trans)
+	// Changing the transient state can and will play with packages and is not meant to be undoable.
+	// However, don't kill the transaction buffer unless it's actually needed, and FIXME: especially not during a construction script
+	if (bShouldMarkDirty && GEditor && GEditor->Trans)
 	{
-		GEditor->Trans->Reset(LOCTEXT("ChangeEditingMode", "Changing Editing Mode"));
+		UWorld* World = GetWorld();
+
+		if (!World || !World->bIsRunningConstructionScript)
+		{
+			GEditor->Trans->Reset(LOCTEXT("ChangeEditingMode", "Changing Editing Mode"));
+		}
 	}
 }
 
