[PCG] Fix issues with collapse

* Was not working for some cases of graph parameters
* Was completely broken for named reroutes
* Refactor code so that it is not a enormous chunk of code, but better organized and with less duplicated code.
* Now detect cycles when collapsing
* Show the user the reason why their collapse has failed.
* Add a unit test that does a combinations of all 4 nodes collapses to cover most edge cases, functional test to be confident it always produces the same result with or without collapse.
* Unit test graph added as an asset

#rnx

diff --git a/Engine/Plugins/PCG/Source/PCG/Private/Helpers/PCGSubgraphHelpers.cpp b/Engine/Plugins/PCG/Source/PCG/Private/Helpers/PCGSubgraphHelpers.cpp
--- a/Engine/Plugins/PCG/Source/PCG/Private/Helpers/PCGSubgraphHelpers.cpp
+++ b/Engine/Plugins/PCG/Source/PCG/Private/Helpers/PCGSubgraphHelpers.cpp
@@ -10,6 +10,7 @@
 #include "PCGPin.h"
 #include "PCGSettings.h"
 #include "PCGSubgraph.h"
+#include "Elements/PCGReroute.h"
 #include "Elements/PCGUserParameterGet.h"
 
 #include "Algo/AnyOf.h"
@@ -19,47 +20,9 @@
 
 namespace PCGSubgraphHelpersExtra
 {
-	// Function that will create a new pin in the input/output node of the subgraph, with a unique name that will
-	// be the related to the pin passed as argument. It will be formatted like this:
-	// "{NodeName} {PinName} {OptionalIndex}"
-	FName CreateNewCustomPin(UPCGGraph* InGraph, UPCGPin* InPinToClone, bool bIsInput, TMap<FString, int>& InOutNameCollisionMapping)
-	{
-		UPCGNode* NewInputOutputNode = bIsInput ? InGraph->GetInputNode() : InGraph->GetOutputNode();
-		UPCGGraphInputOutputSettings* NewInputOutputSettings = CastChecked<UPCGGraphInputOutputSettings>(NewInputOutputNode->GetSettings());
-		FString NewName = InPinToClone->Node->GetNodeTitle(EPCGNodeTitleType::ListView).ToString() + " " + InPinToClone->Properties.Label.ToString();
-		if (InOutNameCollisionMapping.Contains(NewName))
-		{
-			if (!bIsInput)
-			{
-				// Output pin can be re-used. If there is a clash, just return the new name, the pin was already added.
-				return FName(NewName);
-			}
-
-			NewName += " " + FString::FormatAsNumber(++InOutNameCollisionMapping[NewName]);
-		}
-		else
-		{
-			InOutNameCollisionMapping.Emplace(NewName, 1);
-		}
-
-		FPCGPinProperties NewProperties = InPinToClone->Properties;
-
-		// For the pin type, narrow it down to the input edges (if it is an input pin)
-		if (bIsInput)
-		{
-			const EPCGDataType FirstPinTypeUnion = InPinToClone->Node->GetSettings()->GetTypeUnionOfIncidentEdges(InPinToClone->Properties.Label);
-			if (FirstPinTypeUnion != EPCGDataType::None)
-			{
-				NewProperties.AllowedTypes = FirstPinTypeUnion;
-			}
-		}
-
-		NewProperties.Label = FName(NewName);
-		NewProperties = NewInputOutputSettings->AddPin(NewProperties);
-		NewInputOutputNode->UpdateAfterSettingsChangeDuringCreation();
-		return NewProperties.Label;
-	}
-
+	/**
+	* Structure that will hold the necessary information during the collapse and execute it.
+	*/
 	struct FCollapsingInformation
 	{
 #if WITH_EDITOR
@@ -69,31 +32,109 @@ namespace PCGSubgraphHelpersExtra
 		int32 MaxX = std::numeric_limits<int32>::min();
 #endif // WITH_EDITOR
 
-		// Those 3 sets are used to keep track of the subgraph pins
-		// and will be used to extract pins that will be outside the subgraph
-		// and will need special treatment.
-		TSet<UPCGPin*> InputFromSubgraphPins;
-		TSet<UPCGPin*> OutputToSubgraphPins;
-		TSet<UPCGPin*> AllSubgraphPins;
+		/** First pass on the nodes to collapse, gathering information about special nodes like Graph Parameters. */
+		static FCollapsingInformation CreateAndInitialize(const TArray<UPCGNode*>& InNodesToCollapse, const FInstancedPropertyBag& InGraphParameters, UPCGGraph* InOriginalGraph);
+		bool ValidateCollapseIsNotIntroducingCycle();
 
-		// Also keep track of all the edges in the subgraph. Use a set to be able to add the
-		// same edge twice without duplicates.
-		TSet<UPCGEdge*> PCGGraphEdges;
+		void SetNewGraphAndUpdatePositions(UPCGGraph* NewGraph);
 
-		// Keep track of all graph parameters used, they will need to be forwarded to the new graph.
-		TArray<const FPropertyBagPropertyDesc*> GraphParametersUsed;
+		void AddDiscoveredGraphParameters();
+		void CopyNodes(const TArray<UObject*>& ExtraNodesToCollapse);
+		void DeleteNodes(const TArray<UObject*>& ExtraNodesToCollapse);
+		void CreateSubgraphNodeAndConnectRemainingPins();
 
-		TArray<UPCGNode*> ValidNodesToCollapse;
+		/**
+		* Go through all the collapsed nodes and reconnect all of them according to the original connections.
+		* Will also handle any special cases (named reroutes, graph parameters, ...)
+		*/
+		void ReconnectAllNodes();
+
+		// After reconnection, clean up edges that needs to be cleaned up
+		void CleanUpOriginalEdges();
+
+		bool HasEnoughValidNodes() { return ValidNodesToCollapse.Num() > 1; }
+
+	private:
+		// Use initialize to create the struct.
+		FCollapsingInformation() = default;
 
-		// Keep track of all user parameters output pins
-		TMap<FGuid, UPCGPin*> GetUserParametersOutputPins;
-		// Also keep track of duplicated get user parameters, to only keep a single one for each user parameter.
+		// List of all the nodes that are valid to collapse and so will be copied.
+		TArray<UPCGNode*> ValidNodesToCollapse;
+		// By default all nodes that are collapsed are removed, but some might have to stay (like graph parameters getters)
+		TArray<UPCGNode*> NodesToRemove;
+		// Also some node can be mark superfluous, meaning that they can be removed if at the end, there is no edges remaining.
 		TArray<UPCGNode*> SuperfluousNodes;
+
+		// Keep a list of all outside pins and the name of the label to connect to.
+		TArray<TPair<UPCGPin*, FName>> OutsideSubgraphPinToSubgraphPinLabel;
+
+		// Keep a mapping between pins from the inside pin and its connected pin on the input/output node in the subgraph.
+		TMap<UPCGPin*, UPCGPin*> SubgraphInsidePinToInputOutputPinMapping;
+		TMap<FString, int> NameCollisionMapping;
+
+		// Keep a mapping between original nodes and new nodes
+		TMap<const UPCGNode*, UPCGNode*> OriginalToNewNodesMapping;
+
+		UPCGGraph* OriginalGraph = nullptr;
+		UPCGGraph* NewGraph = nullptr;
+
+		/** 
+		* Connect two pins that cross the subgraph boundaries. The Inside Pin is the pin that is inside the subgraph (NewGraph) and outside pin the pin in the main graph.
+		* bIsSubgraphInputBoundary is set if it is the input boundary that is crossed. Otherwise, it is the output boundary.
+		* Returns true if a new custom pin was added, false otherwise.
+		*/
+		bool ConnectPinsBetweenSubgraphBoundaries(UPCGPin* OriginalPin, UPCGPin* NewInsidePin, UPCGPin* OutsidePin, bool bIsSubgraphInputBoundary);
+
+		/**
+		* Function that will create a new pin in the input / output node of the subgraph, with a unique name that will
+		* be the related to the pin passed as argument. It will be formatted like this:
+		* "{OptionalNodeName} {PinName} {OptionalIndex}"
+		* OptionalNodeName and OptionalIndex are added if there are any name clash.
+		*/
+		FName CreateNewCustomPin(UPCGGraph* InGraph, const UPCGPin* InPinToClone, bool bIsInput);
+
+		void ReconnectNode(UPCGNode* NewNode, UPCGPin* OriginalPin);
+
+		////////////////
+		// Graph parameters specific
+		////////////////
+
+		// Keep a mapping between PropertyDesc and their getter node. Useful to discard other getters that can be removed if unused.
+		TMap<const FPropertyBagPropertyDesc*, UPCGNode*> GraphParametersDescToNodeGetters;
+
+		// Also keep a reference of the original graph parameters.
+		const FInstancedPropertyBag* GraphParameters = nullptr;
+
+		/** 
+		* If a graph parameter getter is collapse, this will connect the original getter to the override pin on the subgraph node (it's not a real connect, just adding to OutsideSubgraphPinToSubgraphPinLabelMapping, for placeholder).
+		* returns true if it is the first time we see that graph parameter.
+		*/
+		bool ConnectGraphParameterGetterToSubgraphInput(const UPCGUserParameterGetSettings* InSettings, UPCGNode* InNode);
+
+		/**
+		* Special handling if a graph parameter getter is outside the subgraph. If the graph parameter is part of the subgraph, we will use the getter inside the subgraph for the same parameter.
+		* Otherwise, it is a normal boundary cross.
+		*/
+		void HandleExternalGraphParameterGetter(UPCGPin* OriginalInsidePin, UPCGPin* NewInsidePin, UPCGPin* OutsidePin);
+
+		////////////////
+		// Named reroute specific
+		////////////////
+
+		void HandleReconnectNamedRerouteDeclaration(UPCGNode* OriginalNode, const UPCGNamedRerouteDeclarationSettings* OriginalSettings);
+		void HandleReconnectNamedRerouteUsage(UPCGNode* OriginalNode, const UPCGNamedRerouteUsageSettings* OriginalSettings);
+
+		// List of original pins to break edges. Need to be done after the reconnection, to not alter it.
+		TSet<UPCGPin*> OriginalPinsToBreakEdges;
 	};
 
-	FCollapsingInformation GatheringCollapseInformation(const TArray<UPCGNode*>& InNodesToCollapse, const FInstancedPropertyBag& InGraphParameters)
+	FCollapsingInformation FCollapsingInformation::CreateAndInitialize(const TArray<UPCGNode*>& InNodesToCollapse, const FInstancedPropertyBag& InGraphParameters, UPCGGraph* InOriginalGraph)
 	{
+		check(InOriginalGraph);
+
 		FCollapsingInformation CollapseInfo{};
+		CollapseInfo.GraphParameters = &InGraphParameters;
+		CollapseInfo.OriginalGraph = InOriginalGraph;
 
 		for (UPCGNode* PCGNode : InNodesToCollapse)
 		{
@@ -107,57 +148,18 @@ namespace PCGSubgraphHelpersExtra
 				continue;
 			}
 
-			if (const UPCGUserParameterGetSettings* GetSettings = Cast<UPCGUserParameterGetSettings>(NodeSettings))
-			{
-				if (const FPropertyBagPropertyDesc* PropertyDesc = InGraphParameters.FindPropertyDescByID(GetSettings->PropertyGuid))
-				{
-					CollapseInfo.GraphParametersUsed.AddUnique(PropertyDesc);
-
-					if (CollapseInfo.GetUserParametersOutputPins.Contains(PropertyDesc->ID))
-					{
-						CollapseInfo.SuperfluousNodes.Add(PCGNode);
-					}
-					else
-					{
-						CollapseInfo.GetUserParametersOutputPins.Emplace(PropertyDesc->ID, PCGNode->GetOutputPins()[0]);
-					}
-				}
-				else
-				{
-					// Invalid node
-					continue;
-				}
-			}
+			bool bShouldRemove = true;
 
-			// For each pin of the subgraph, gather all its edges in a set and store the other pin of
-			// the edge as an input from or output to the subgraph. It will be useful for tracking pins
-			// that are outside the subgraph.
-			for (UPCGPin* InputPin : PCGNode->GetInputPins())
-			{
-				check(InputPin);
-
-				for (UPCGEdge* Edge : InputPin->Edges)
-				{
-					check(Edge);
-					CollapseInfo.PCGGraphEdges.Add(Edge);
-					CollapseInfo.OutputToSubgraphPins.Add(Edge->InputPin);
-				}
-
-				CollapseInfo.AllSubgraphPins.Add(InputPin);
-			}
-
-			for (UPCGPin* OutputPin : PCGNode->GetOutputPins())
+			if (const UPCGUserParameterGetSettings* GetSettings = Cast<UPCGUserParameterGetSettings>(NodeSettings))
 			{
-				check(OutputPin);
-
-				for (UPCGEdge* Edge : OutputPin->Edges)
+				bShouldRemove = false;
+				// If it was connected, we want to keep this node outside, so don't remove it.
+				// Otherwise, we already discovered this graph parameter, so we can remove it from the top graph
+				const bool bAlreadyDiscovered = !CollapseInfo.ConnectGraphParameterGetterToSubgraphInput(GetSettings, PCGNode);
+				if (bAlreadyDiscovered)
 				{
-					check(Edge);
-					CollapseInfo.PCGGraphEdges.Add(Edge);
-					CollapseInfo.InputFromSubgraphPins.Add(Edge->OutputPin);
+					CollapseInfo.SuperfluousNodes.Add(PCGNode);
 				}
-
-				CollapseInfo.AllSubgraphPins.Add(OutputPin);
 			}
 
 			CollapseInfo.ValidNodesToCollapse.Add(PCGNode);
@@ -169,6 +171,11 @@ namespace PCGSubgraphHelpersExtra
 			CollapseInfo.MinX = FMath::Min(CollapseInfo.MinX, PCGNode->PositionX);
 			CollapseInfo.MaxX = FMath::Max(CollapseInfo.MaxX, PCGNode->PositionY);
 #endif // WITH_EDITOR
+
+			if (bShouldRemove)
+			{
+				CollapseInfo.NodesToRemove.Add(PCGNode);
+			}
 		}
 
 #if WITH_EDITOR
@@ -179,322 +186,609 @@ namespace PCGSubgraphHelpersExtra
 		}
 #endif // WITH_EDITOR
 
-		// Gather the pins that are outside the subgraph
-		CollapseInfo.InputFromSubgraphPins = CollapseInfo.InputFromSubgraphPins.Difference(CollapseInfo.AllSubgraphPins);
-		CollapseInfo.OutputToSubgraphPins = CollapseInfo.OutputToSubgraphPins.Difference(CollapseInfo.AllSubgraphPins);
-
 		return CollapseInfo;
 	}
-}
 
-UPCGGraph* FPCGSubgraphHelpers::CollapseIntoSubgraph(UPCGGraph* InOriginalGraph, const TArray<UPCGNode*>& InNodesToCollapse, const TArray<UObject*>& InExtraNodesToCollapse, UPCGGraph* OptionalPreAllocatedGraph)
-{
-	// Don't collapse into a subgraph if you don't have at least 2 nodes
-	if (!InOriginalGraph || InNodesToCollapse.Num() < 2)
+	void FCollapsingInformation::SetNewGraphAndUpdatePositions(UPCGGraph* InNewGraph)
 	{
-		return nullptr;
-	}
+		check(InNewGraph);
+		NewGraph = InNewGraph;
 
-	const FInstancedPropertyBag* GraphParameters = InOriginalGraph->GetUserParametersStruct();
-	check(GraphParameters);
-
-	PCGSubgraphHelpersExtra::FCollapsingInformation CollapseInfo = PCGSubgraphHelpersExtra::GatheringCollapseInformation(InNodesToCollapse, *GraphParameters);
+#if WITH_EDITOR
+		// Do some clean-up on input/output nodes
+		constexpr int32 Padding = 300;
+		NewGraph->GetInputNode()->PositionX = MinX - Padding;
+		NewGraph->GetInputNode()->PositionY = AveragePosition.Y;
+		NewGraph->GetOutputNode()->PositionX = MaxX + Padding;
+		NewGraph->GetOutputNode()->PositionY = AveragePosition.Y;
+#endif // WITH_EDITOR
+	}
 
-	// If we have at most 1 valid node to collapse, just exit
-	if (CollapseInfo.ValidNodesToCollapse.Num() <= 1)
+	void FCollapsingInformation::AddDiscoveredGraphParameters()
 	{
-		UE_LOG(LogPCG, Warning, TEXT("There were less than 2 PCG nodes selected, abort"));
-		return nullptr;
+		check(OriginalGraph && NewGraph);
+		TArray<FPropertyBagPropertyDesc> GraphParametersUsed;
+		Algo::Transform(GraphParametersDescToNodeGetters, GraphParametersUsed, [](const auto& It) { check(It.Key); return *It.Key; });
+		NewGraph->AddUserParameters(GraphParametersUsed, OriginalGraph);
 	}
 
-	// 2. Create a new subgraph if necessary.
-	UPCGGraph* NewPCGGraph = OptionalPreAllocatedGraph;
-	if (!NewPCGGraph)
+	void FCollapsingInformation::CopyNodes(const TArray<UObject*>& ExtraNodesToCollapse)
 	{
-		NewPCGGraph = NewObject<UPCGGraph>();
+		check(NewGraph);
+		for (UPCGNode* PCGNode : ValidNodesToCollapse)
+		{
+			check(PCGNode);
+
+			// Reconstruct a new node, same as PCGNode, but without any edges in the new graph 
+			UPCGNode* NewNode = NewGraph->ReconstructNewNode(PCGNode);
+			check(NewNode);
+
+			NewNode->NodeTitle = PCGNode->NodeTitle;
+			OriginalToNewNodesMapping.Emplace(PCGNode, NewNode);
+		}
+
+#if WITH_EDITOR
+		// Also duplicate the extra nodes and assign them to the new graph
+		TArray<TObjectPtr<const UObject>> NewExtraGraphNodes;
+		for (const UObject* ExtraNode : ExtraNodesToCollapse)
+		{
+			NewExtraGraphNodes.Add(DuplicateObject(ExtraNode, NewGraph));
+		}
+
+		NewGraph->SetExtraEditorNodes(NewExtraGraphNodes);
+#endif ///WITH_EDITOR
 	}
 
+	void FCollapsingInformation::DeleteNodes(const TArray<UObject*>& ExtraNodesToCollapse)
+	{
+		check(OriginalGraph);
+
+		// First remove all the nodes that were valid and should be removed.
+		OriginalGraph->RemoveNodes(NodesToRemove);
+
 #if WITH_EDITOR
-	// Do some clean-up on input/output nodes
-	constexpr int32 Padding = 200;
-	NewPCGGraph->GetInputNode()->PositionX = CollapseInfo.MinX - Padding;
-	NewPCGGraph->GetInputNode()->PositionY = CollapseInfo.AveragePosition.Y;
-	NewPCGGraph->GetOutputNode()->PositionX = CollapseInfo.MaxX + Padding;
-	NewPCGGraph->GetOutputNode()->PositionY = CollapseInfo.AveragePosition.Y;
+		for (const UObject* ExtraNode : ExtraNodesToCollapse)
+		{
+			OriginalGraph->RemoveExtraEditorNode(ExtraNode);
+		}
 #endif // WITH_EDITOR
 
-	// 4. Duplicate all the nodes, and keep a mapping between the old pins and new pins
-	TMap<UPCGPin*, UPCGPin*> PinMapping;
-	for (const UPCGNode* PCGNode : CollapseInfo.ValidNodesToCollapse)
+		// Then for all the remaining nodes to remove, make sure they have no edges anymore.
+		TArray<UPCGNode*> SuperfluousNodesToRemove;
+		SuperfluousNodesToRemove.Reserve(SuperfluousNodes.Num());
+		for (UPCGNode* SuperfluousNode : SuperfluousNodes)
+		{
+			check(SuperfluousNode);
+			// Only remove them if they have no edges anymore
+			bool bHasAnyEdges = Algo::AnyOf(SuperfluousNode->GetInputPins(), [](const UPCGPin* InPin) { return InPin && InPin->EdgeCount() > 0; })
+				|| Algo::AnyOf(SuperfluousNode->GetOutputPins(), [](const UPCGPin* InPin) { return InPin && InPin->EdgeCount() > 0; });
+
+			if (!bHasAnyEdges)
+			{
+				SuperfluousNodesToRemove.Add(SuperfluousNode);
+			}
+		}
+
+		// Finally remove them all in one go
+		OriginalGraph->RemoveNodes(SuperfluousNodesToRemove);
+	}
+
+	void FCollapsingInformation::CreateSubgraphNodeAndConnectRemainingPins()
 	{
-		check(PCGNode);
+		check(OriginalGraph && NewGraph);
 
-		// Reconstruct a new node, same as PCGNode, but without any edges in the new graph 
-		UPCGNode* NewNode = NewPCGGraph->ReconstructNewNode(PCGNode);
+		UPCGSettings* DefaultNodeSettings = nullptr;
+		UPCGNode* SubgraphNode = OriginalGraph->AddNodeOfType(UPCGSubgraphSettings::StaticClass(), DefaultNodeSettings);
+		UPCGSubgraphSettings* DefaultSubgraphSettings = CastChecked<UPCGSubgraphSettings>(DefaultNodeSettings);
+		DefaultSubgraphSettings->SetSubgraph(NewGraph);
+
+		check(SubgraphNode);
 
-		// Safeguard: We should have a 1 for 1 matching between pins labels between the original node
-		// and the copied node. If for some reason we don't (perhaps the node was not updated correctly after pins were added/removed)
-		// we will log an error and try to connect as best as we can (probably breaking some edges on the process).
+#if WITH_EDITOR
+		SubgraphNode->PositionX = AveragePosition.X;
+		SubgraphNode->PositionY = AveragePosition.Y;
+#endif // WITH_EDITOR
 
-		auto Mapping = [&PinMapping, PCGNode](const TArray<UPCGPin*>& OriginalPins, const TArray<UPCGPin*>& NewPins)
+		SubgraphNode->UpdateAfterSettingsChangeDuringCreation();
+
+		for (const auto& [OutsidePin, SubgraphPinLabel] : OutsideSubgraphPinToSubgraphPinLabel)
 		{
-			TSet<FName> UnmatchedOriginal;
-			TMap<FName, UPCGPin*> NewMapping;
+			check(OutsidePin);
+			UPCGNode* OutsideNode = OutsidePin->Node;
+			check(OutsideNode);
 
-			for (UPCGPin* NewPin : NewPins)
+			if (OutsidePin->IsOutputPin())
+			{
+				OutsidePin->AddEdgeTo(SubgraphNode->GetInputPin(SubgraphPinLabel));
+			}
+			else
 			{
-				NewMapping.Emplace(NewPin->Properties.Label, NewPin);
+				SubgraphNode->GetOutputPin(SubgraphPinLabel)->AddEdgeTo(OutsidePin);
 			}
+		}
+	}
 
-			for (UPCGPin* OriginalPin : OriginalPins)
+	void FCollapsingInformation::ReconnectAllNodes()
+	{
+		for (UPCGNode* OriginalNode : ValidNodesToCollapse)
+		{
+			check(OriginalToNewNodesMapping.Contains(OriginalNode));
+			UPCGNode* NewNode = OriginalToNewNodesMapping[OriginalNode];
+			check(NewNode);
+
+			// Named reroute specifics.
+			if (const UPCGNamedRerouteDeclarationSettings* OriginalDeclSettings = Cast<UPCGNamedRerouteDeclarationSettings>(OriginalNode->GetSettings()))
 			{
-				FName PinLabel = OriginalPin->Properties.Label;
-				if (UPCGPin** NewPinPtr = NewMapping.Find(PinLabel))
-				{
-					PinMapping.Emplace(OriginalPin, *NewPinPtr);
-				}
-				else if (OriginalPin->IsConnected())
-				{
-					// It is only problematic if the pin was connected
-					UE_LOG(LogPCG, Error, TEXT("[CollapseInSubgraph - %s] %s pin %s does not exist anymore. Edges will be broken."),
-						*PCGNode->GetNodeTitle(EPCGNodeTitleType::ListView).ToString(), (OriginalPin->IsOutputPin() ? TEXT("Output") : TEXT("Input")), *PinLabel.ToString());
-				}
+				HandleReconnectNamedRerouteDeclaration(OriginalNode, OriginalDeclSettings);
+				continue;
+			}
+			else if (const UPCGNamedRerouteUsageSettings* OriginalUsageSettings = Cast<UPCGNamedRerouteUsageSettings>(OriginalNode->GetSettings()))
+			{
+				HandleReconnectNamedRerouteUsage(OriginalNode, OriginalUsageSettings);
+				continue;
 			}
-		};
 
-		Mapping(PCGNode->GetInputPins(), NewNode->GetInputPins());
-		Mapping(PCGNode->GetOutputPins(), NewNode->GetOutputPins());
+			for (UPCGPin* OriginalInputPin : OriginalNode->GetInputPins())
+			{
+				ReconnectNode(NewNode, OriginalInputPin);
+			}
+
+			for (UPCGPin* OriginalOutputPin : OriginalNode->GetOutputPins())
+			{
+				ReconnectNode(NewNode, OriginalOutputPin);
+			}
+		}
 	}
 
-#if WITH_EDITOR
-	// Also duplicate the extra nodes and assign them to the new graph
-	TArray<TObjectPtr<const UObject>> NewExtraGraphNodes;
-	for (const UObject* ExtraNode : InExtraNodesToCollapse)
+	void FCollapsingInformation::CleanUpOriginalEdges()
 	{
-		NewExtraGraphNodes.Add(DuplicateObject(ExtraNode, NewPCGGraph));
+		for (UPCGPin* OriginalPin : OriginalPinsToBreakEdges)
+		{
+			if (ensure(OriginalPin))
+			{
+				OriginalPin->BreakAllEdges();
+			}
+		}
 	}
 
-	NewPCGGraph->SetExtraEditorNodes(NewExtraGraphNodes);
-#endif ///WITH_EDITOR
-
-	// 5. Iterate over all the edges and create edges "placeholders"
-	// Most of them will already be complete, but for those that needs to be connected to the new
-	// subgraph nodes, the pins don't exist yet. Therefore we identify them with their pin labels.
-	//
-	// The logic behind the new pins is this:
-	// -> If the pin is connected to the simple pin of the input/output node, let it like this
-	// -> If the pin is part of the original input/output advanced pins, we will trigger the advanced pins flag on the input/output node in the subgraph
-	// -> If the pin is connected to a get user parameters that is outside of the subgraph, we'll create a new node and connect it to it.
-	// -> Otherwise, we add a new custom pin, with the name of the node, the name of the pin and a number if there is name collision
-	struct EdgePlaceholder
+	FName FCollapsingInformation::CreateNewCustomPin(UPCGGraph* InGraph, const UPCGPin* InPinToClone, bool bIsInput)
 	{
-		UPCGPin* InputPin = nullptr;
-		UPCGPin* OutputPin = nullptr;
-		FName InputPinLabel;
-		FName OutputPinLabel;
-	};
+		check(InGraph && InPinToClone);
 
-	TArray<EdgePlaceholder> EdgePlaceholders;
-	TMap<FString, int> NameCollisionMapping;
+		UPCGNode* NewInputOutputNode = bIsInput ? InGraph->GetInputNode() : InGraph->GetOutputNode();
+		UPCGGraphInputOutputSettings* NewInputOutputSettings = CastChecked<UPCGGraphInputOutputSettings>(NewInputOutputNode->GetSettings());
+		const UPCGNode* PinToCloneNode = InPinToClone->Node;
+		const UPCGSettings* PinToCloneSettings = PinToCloneNode ? PinToCloneNode->GetSettings() : nullptr;
+		check(PinToCloneNode && PinToCloneSettings);
 
-	// For all internal parameters, create extra edges to connect those nodes to their override pin counterparts
-	for (const TPair<FGuid, UPCGPin*>& It : CollapseInfo.GetUserParametersOutputPins)
-	{
-		check(It.Value);
+		// First try with just the pin name
+		FString NewName = InPinToClone->Properties.Label.ToString();
+		if (NameCollisionMapping.Contains(NewName))
+		{
+			// If there is a clash, use the node title on top.
+			NewName = PinToCloneNode->GetNodeTitle(EPCGNodeTitleType::ListView).ToString() + " " + NewName;
+			if (NameCollisionMapping.Contains(NewName))
+			{
+				NewName += " " + FString::FormatAsNumber(++NameCollisionMapping[NewName]);
+			}
+			else
+			{
+				NameCollisionMapping.Emplace(NewName, 1);
+			}
+		}
+		else
+		{
+			NameCollisionMapping.Emplace(NewName, 1);
+		}
 
-		EdgePlaceholder OutsideSubgraphEdge;
-		OutsideSubgraphEdge.InputPin = It.Value;
-		// Gymnastics to make sure we have a 1 to 1 match between labels, because subgraph override uses the display text name.
-		OutsideSubgraphEdge.OutputPinLabel = FName(FName::NameToDisplayString(OutsideSubgraphEdge.InputPin->Properties.Label.ToString(), /*bIsBool=*/ false));
+		FPCGPinProperties NewProperties = InPinToClone->Properties;
 
-		EdgePlaceholders.Add(std::move(OutsideSubgraphEdge));
+		// For the pin type, narrow it down to the input edges (if it is an input pin)
+		if (!InPinToClone->IsOutputPin())
+		{
+			const EPCGDataType FirstPinTypeUnion = PinToCloneSettings->GetTypeUnionOfIncidentEdges(InPinToClone->Properties.Label);
+			if (FirstPinTypeUnion != EPCGDataType::None)
+			{
+				NewProperties.AllowedTypes = FirstPinTypeUnion;
+			}
+		}
+		else
+		{
+			NewProperties.AllowedTypes = PinToCloneSettings->GetCurrentPinTypes(InPinToClone);
+		}
+
+		NewProperties.Label = FName(NewName);
+		NewProperties = NewInputOutputSettings->AddPin(NewProperties);
+		NewInputOutputNode->UpdateAfterSettingsChangeDuringCreation();
+		return NewProperties.Label;
 	}
 
-	for (UPCGEdge* Edge : CollapseInfo.PCGGraphEdges)
+	void FCollapsingInformation::ReconnectNode(UPCGNode* NewNode, UPCGPin* OriginalPin)
 	{
-		UPCGPin* const* InPin = PinMapping.Find(Edge->InputPin);
-		UPCGPin* const* OutPin = PinMapping.Find(Edge->OutputPin);
+		check(NewGraph && NewNode && OriginalPin);
+		const bool bIsInput = !OriginalPin->IsOutputPin();
 
-		check(InPin || OutPin);
+		UPCGPin* NewPin = bIsInput ? NewNode->GetInputPin(OriginalPin->Properties.Label) : NewNode->GetOutputPin(OriginalPin->Properties.Label);
+		check(NewPin);
 
-		if (InPin == nullptr)
+		for (UPCGEdge* Edge : OriginalPin->Edges)
 		{
-			// The edge comes from outside the graph.
-			// If it is from the input node, we have a special behavior
-			// Same if it is from a get user parameter
-			EdgePlaceholder OutsideSubgraphEdge;
-			EdgePlaceholder InsideSubgraphEdge;
-			bool bProcessed = false;
+			check(Edge);
+			UPCGPin* OtherOriginalPin = bIsInput ? Edge->InputPin : Edge->OutputPin;
+			check(OtherOriginalPin);
 
-			OutsideSubgraphEdge.InputPin = Edge->InputPin;
-			InsideSubgraphEdge.OutputPin = *OutPin;
+			const UPCGNode* OtherOriginalNode = OtherOriginalPin->Node;
+			check(OtherOriginalNode);
 
-			if (Edge->InputPin->Node && Edge->InputPin->Node->GetSettings() && Edge->InputPin->Node->GetSettings()->IsA<UPCGUserParameterGetSettings>())
+			// If the other original is in the mapping, it is part of the subgraph. So Connect the 2, only if the new node is the input pin, as we will pass there twice (once per edge extremity)
+			if (UPCGNode** OtherNewNodePtr = OriginalToNewNodesMapping.Find(OtherOriginalNode))
 			{
-				const UPCGUserParameterGetSettings* OldSettings = CastChecked<const UPCGUserParameterGetSettings>(Edge->InputPin->Node->GetSettings());
-				if (CollapseInfo.GetUserParametersOutputPins.Contains(OldSettings->PropertyGuid))
+				if (bIsInput)
 				{
-					if (CollapseInfo.GetUserParametersOutputPins[OldSettings->PropertyGuid] != Edge->InputPin->Node->GetOutputPins()[0])
-					{
-						CollapseInfo.SuperfluousNodes.Add(Edge->InputPin->Node);
-					}
+					check(OtherOriginalPin->IsOutputPin());
+					UPCGPin* OtherNewPin = (*OtherNewNodePtr)->GetOutputPin(OtherOriginalPin->Properties.Label);
+					check(OtherNewPin);
+					OtherNewPin->AddEdgeTo(NewPin);
+				}
+			}
+			else
+			{
+				if (const UPCGUserParameterGetSettings* GraphParameterSettings = Cast<const UPCGUserParameterGetSettings>(OtherOriginalNode->GetSettings()))
+				{
+					HandleExternalGraphParameterGetter(OriginalPin, NewPin, OtherOriginalPin);
+				}
+				else if (Cast<const UPCGUserParameterGetSettings>(NewNode->GetSettings()) != nullptr)
+				{
+					// If the inside node is a graph parameter that goes outside, we have nothing to do, the edge already exists on the original node and the original node won't be deleted.
 				}
 				else
 				{
-					CollapseInfo.GetUserParametersOutputPins.Emplace(OldSettings->PropertyGuid, Edge->InputPin->Node->GetOutputPins()[0]);
+					ConnectPinsBetweenSubgraphBoundaries(OriginalPin, NewPin, OtherOriginalPin, bIsInput);
 				}
+			}
+		}
+	}
 
-				check(CollapseInfo.GetUserParametersOutputPins.Contains(OldSettings->PropertyGuid));
+	// Returns true if a custom pin was created.
+	bool FCollapsingInformation::ConnectPinsBetweenSubgraphBoundaries(UPCGPin* OriginalInsidePin, UPCGPin* NewInsidePin, UPCGPin* OutsidePin, bool bIsSubgraphInputBoundary)
+	{
+		check(OriginalInsidePin && NewInsidePin && OutsidePin && NewGraph);
 
-				// The label will be the same as the input pin, since it will be an override pin.
-				OutsideSubgraphEdge.InputPin = CollapseInfo.GetUserParametersOutputPins[OldSettings->PropertyGuid];
-				// Gymnastics to make sure we have a 1 to 1 match between labels, because subgraph override uses the display text name.
-				OutsideSubgraphEdge.OutputPinLabel = FName(FName::NameToDisplayString(OutsideSubgraphEdge.InputPin->Properties.Label.ToString(), /*bIsBool=*/ false));
+		bool bCustomPinCreated = false;
 
-				// Create a new node and connect it to it.
-				UPCGSettings* NewSettings = nullptr;
-				UPCGNode* NewNode = NewPCGGraph->AddNodeCopy(Edge->InputPin->Node->GetSettings(), NewSettings);
-				check(NewNode && NewSettings && NewNode->GetOutputPins().Num() == 1);
+		// IO pins are a 1 for 1 with the pins that are inside the subgraph.
+		// So we keep a mapping between the subgraph IO pins and the inside pins, and only create a new IO pin if we have a new inside pin that cross boundaries.
+		UPCGPin* EdgeInputPin = bIsSubgraphInputBoundary ? OutsidePin : OriginalInsidePin;
+		UPCGPin** SubgraphIOPinPtr = SubgraphInsidePinToInputOutputPinMapping.Find(EdgeInputPin);
+		UPCGPin* SubgraphIOPin = nullptr;
+		if (!SubgraphIOPinPtr)
+		{
+			FName SubgraphPinLabel = CreateNewCustomPin(NewGraph, EdgeInputPin, bIsSubgraphInputBoundary);
+			SubgraphIOPin = bIsSubgraphInputBoundary ? NewGraph->GetInputNode()->GetOutputPin(SubgraphPinLabel) : NewGraph->GetOutputNode()->GetInputPin(SubgraphPinLabel);
+			check(SubgraphIOPin);
+			SubgraphInsidePinToInputOutputPinMapping.Emplace(EdgeInputPin, SubgraphIOPin);
 
-				// Also need to add this parameter to the list.
-				if (const FPropertyBagPropertyDesc* PropertyDesc = GraphParameters->FindPropertyDescByID(OldSettings->PropertyGuid))
-				{
-					CollapseInfo.GraphParametersUsed.AddUnique(PropertyDesc);
-				}
+			bCustomPinCreated = true;
+		}
+		else
+		{
+			SubgraphIOPin = *SubgraphIOPinPtr;
+		}
 
-				InsideSubgraphEdge.InputPin = NewNode->GetOutputPins()[0];
+		check(SubgraphIOPin);
 
-				bProcessed = true;
-#if WITH_EDITOR
-				// Place the newly created node slightly to the left of the output pin node.
-				if (ensure(Edge->OutputPin->Node))
-				{
-					constexpr int32 OffsetX = -200;
-					NewNode->PositionX = Edge->OutputPin->Node->PositionX + OffsetX;
-					NewNode->PositionY = Edge->OutputPin->Node->PositionY;
-				}
-#endif // WITH_EDITOR
+		// And we keep a list of all the outside pins and the label they should connect to, when the subgraph node will be created later.
+		if (bCustomPinCreated || !bIsSubgraphInputBoundary)
+		{
+			OutsideSubgraphPinToSubgraphPinLabel.Emplace(OutsidePin, SubgraphIOPin->Properties.Label);
+		}
+		
+		if (bCustomPinCreated || bIsSubgraphInputBoundary)
+		{
+			if (bIsSubgraphInputBoundary)
+			{
+				SubgraphIOPin->AddEdgeTo(NewInsidePin);
 			}
-
-			if (!bProcessed)
+			else
 			{
-				FName NewPinName = PCGSubgraphHelpersExtra::CreateNewCustomPin(NewPCGGraph, Edge->OutputPin, true, NameCollisionMapping);
-				OutsideSubgraphEdge.OutputPinLabel = NewPinName;
-				InsideSubgraphEdge.InputPin = NewPCGGraph->GetInputNode()->GetOutputPin(NewPinName);
+				NewInsidePin->AddEdgeTo(SubgraphIOPin);
 			}
-
-			EdgePlaceholders.Add(std::move(OutsideSubgraphEdge));
-			EdgePlaceholders.Add(std::move(InsideSubgraphEdge));
 		}
-		else if (OutPin == nullptr)
+
+		return bCustomPinCreated;
+	}
+
+	bool FCollapsingInformation::ConnectGraphParameterGetterToSubgraphInput(const UPCGUserParameterGetSettings* InSettings, UPCGNode* InNode)
+	{
+		check(InSettings && InNode);
+		const FPropertyBagPropertyDesc* PropertyDesc = GraphParameters->FindPropertyDescByID(InSettings->PropertyGuid);
+
+		if (ensure(PropertyDesc) && !GraphParametersDescToNodeGetters.Contains(PropertyDesc))
 		{
-			// The edge comes from outside the graph.
-			// If it is from the output node, we have a special behavior
-			EdgePlaceholder OutsideSubgraphEdge;
-			EdgePlaceholder InsideSubgraphEdge;
-			bool bProcessed = false;
+			GraphParametersDescToNodeGetters.Emplace(PropertyDesc, InNode);
+			check(InNode->GetOutputPins().Num() == 1);
+			UPCGPin* GraphParameterPin = InNode->GetOutputPins()[0];
+			check(GraphParameterPin);
+			OutsideSubgraphPinToSubgraphPinLabel.Emplace(GraphParameterPin, FName(FName::NameToDisplayString(GraphParameterPin->Properties.Label.ToString(), /*bIsBool=*/ false)));
 
-			OutsideSubgraphEdge.OutputPin = Edge->OutputPin;
-			InsideSubgraphEdge.InputPin = *InPin;
+			return true;
+		}
 
-			if (!bProcessed)
-			{
-				FName NewPinName = PCGSubgraphHelpersExtra::CreateNewCustomPin(NewPCGGraph, Edge->InputPin, false, NameCollisionMapping);
-				OutsideSubgraphEdge.InputPinLabel = NewPinName;
-				InsideSubgraphEdge.OutputPin = NewPCGGraph->GetOutputNode()->GetInputPin(NewPinName);
-			}
+		// Already connected or invalid.
+		return false;
+	}
 
-			EdgePlaceholders.Add(std::move(OutsideSubgraphEdge));
-			EdgePlaceholders.Add(std::move(InsideSubgraphEdge));
+	void FCollapsingInformation::HandleExternalGraphParameterGetter(UPCGPin* OriginalInsidePin, UPCGPin* NewInsidePin, UPCGPin* OutsidePin)
+	{
+		check(OriginalInsidePin && NewInsidePin && OutsidePin);
+		UPCGNode* ExternalGraphParameterNode = OutsidePin->Node;
+		check(ExternalGraphParameterNode);
+		const UPCGUserParameterGetSettings* Settings = CastChecked<const UPCGUserParameterGetSettings>(ExternalGraphParameterNode->GetSettings());
+
+		const FPropertyBagPropertyDesc* PropertyDesc = GraphParameters->FindPropertyDescByID(Settings->PropertyGuid);
+		check(PropertyDesc);
+
+		// If the graph parameter exists, we should use that new getter inside the subgraph, and mark the external as superfluous.
+		if (UPCGNode** ExistingGraphParameterNode = GraphParametersDescToNodeGetters.Find(PropertyDesc))
+		{
+			check(*ExistingGraphParameterNode);
+			UPCGNode* ExistingGraphParameterNodeInsideSubgraph = OriginalToNewNodesMapping.FindChecked(*ExistingGraphParameterNode);
+			check(ExistingGraphParameterNodeInsideSubgraph && ExistingGraphParameterNodeInsideSubgraph->GetOutputPins().Num() == 1);
+
+			ExistingGraphParameterNodeInsideSubgraph->GetOutputPins()[0]->AddEdgeTo(NewInsidePin);
+			SuperfluousNodes.Add(ExternalGraphParameterNode);
 		}
+		// If the external graph parameter is not part of the subgraph, we treat it as a normal node.
 		else
 		{
-			// Both nodes are inside
-			EdgePlaceholders.Add(EdgePlaceholder{ *InPin, *OutPin });
+			ConnectPinsBetweenSubgraphBoundaries(OriginalInsidePin, NewInsidePin, OutsidePin, /*bIsSubgraphInputBoundary=*/true);
 		}
 	}
 
-	// Transfer the values from the original graph parameters to the new graph
-	TArray<FPropertyBagPropertyDesc> GraphParametersUsed;
-	Algo::Transform(CollapseInfo.GraphParametersUsed, GraphParametersUsed, [](const FPropertyBagPropertyDesc* In) { check(In); return *In; });
-	NewPCGGraph->AddUserParameters(GraphParametersUsed, InOriginalGraph);
-
-	// 6. Create subgraph and delete old nodes, extra nodes and superfluous nodes
-	UPCGNode* SubgraphNode = nullptr;
+	void FCollapsingInformation::HandleReconnectNamedRerouteDeclaration(UPCGNode* OriginalDeclNode, const UPCGNamedRerouteDeclarationSettings* OriginalDeclSettings)
 	{
-		TArray<UPCGNode*> NodesToRemove;
-		NodesToRemove.Reserve(CollapseInfo.ValidNodesToCollapse.Num());
+		check(OriginalDeclNode && OriginalDeclSettings);
+		UPCGNode* NewDeclNode = OriginalToNewNodesMapping.FindChecked(OriginalDeclNode);
+		check(NewDeclNode);
 
-		for (UPCGNode* PCGNode : CollapseInfo.ValidNodesToCollapse)
+		// Connect all the input pins normally
+		for (UPCGPin* OriginalInputPin : OriginalDeclNode->GetInputPins())
 		{
-			const UPCGSettings* NodeSettings = PCGNode->GetSettings();
+			ReconnectNode(NewDeclNode, OriginalInputPin);
+		}
 
-			// Don't delete input, output and get user parameters nodes from the subgraph.
-			if (!NodeSettings || NodeSettings->IsA<UPCGGraphInputOutputSettings>() || NodeSettings->IsA<UPCGUserParameterGetSettings>())
+		// Connect all the output pins normally, except the Invisible Pin
+		UPCGPin* InvisiblePin = OriginalDeclNode->GetOutputPin(PCGNamedRerouteConstants::InvisiblePinLabel);
+		check(InvisiblePin);
+		for (UPCGPin* OriginalOutputPin : OriginalDeclNode->GetOutputPins())
+		{
+			if (!ensure(OriginalOutputPin) || OriginalOutputPin == InvisiblePin)
 			{
 				continue;
 			}
 
-			NodesToRemove.Add(PCGNode);
+			ReconnectNode(NewDeclNode, OriginalOutputPin);
 		}
 
-#if WITH_EDITOR
-		for (const UObject* ExtraNode : InExtraNodesToCollapse)
+		// Then for the invisible pin, only check for usages that could be outside. If they are outside, we don't remove the original declaration node, 
+		// we make sure to create a new output pin to the subgraph and connect the original to this pin. We also break all the edges at the output, as they would be rewired normally.
+		// For the usage inside, they will be patched and rewired in the Usage case.
+		for (const UPCGEdge* Edge : InvisiblePin->Edges)
 		{
-			InOriginalGraph->RemoveExtraEditorNode(ExtraNode);
+			const UPCGPin* OtherPin = Edge->OutputPin;
+			check(OtherPin && OtherPin != InvisiblePin);
+			const UPCGNode* OtherNode = OtherPin->Node;
+			check(OtherNode);
+			const UPCGNamedRerouteUsageSettings* OtherUsageSettings = Cast<UPCGNamedRerouteUsageSettings>(OtherNode->GetSettings());
+			if (!OriginalToNewNodesMapping.Contains(OtherNode) && OtherUsageSettings && NodesToRemove.Contains(OriginalDeclNode))
+			{
+				NodesToRemove.RemoveSwap(OriginalDeclNode);
+				UPCGPin* OriginalInputPin = OriginalDeclNode->GetInputPin(PCGPinConstants::DefaultInputLabel);
+				check(OriginalInputPin);
+				OriginalPinsToBreakEdges.Add(OriginalInputPin);
+
+				UPCGPin* OriginalOutputPin = OriginalDeclNode->GetOutputPin(PCGPinConstants::DefaultOutputLabel);
+				check(OriginalOutputPin);
+				OriginalPinsToBreakEdges.Add(OriginalOutputPin);
+
+				UPCGPin* NewOutputPin = NewDeclNode->GetOutputPin(PCGPinConstants::DefaultOutputLabel);
+				check(NewOutputPin);
+				ConnectPinsBetweenSubgraphBoundaries(OriginalOutputPin, NewOutputPin, OriginalInputPin, /*bIsInput=*/false);
+				break;
+			}
 		}
-#endif // WITH_EDITOR
+	}
+
+	void FCollapsingInformation::HandleReconnectNamedRerouteUsage(UPCGNode* OriginalUsageNode, const UPCGNamedRerouteUsageSettings* OriginalUsageSettings)
+	{
+		check(OriginalUsageNode && OriginalUsageSettings)
+		UPCGNode* NewUsageNode = OriginalToNewNodesMapping.FindChecked(OriginalUsageNode);
+		check(NewUsageNode);
+		UPCGNamedRerouteUsageSettings* NewUsageSettings = CastChecked<UPCGNamedRerouteUsageSettings>(NewUsageNode->GetSettings());
 
-		for (UPCGNode* SuperfluousNode : CollapseInfo.SuperfluousNodes)
+		const UPCGNamedRerouteDeclarationSettings* OriginalDeclSettings = OriginalUsageSettings->Declaration;
+		if (!ensure(OriginalDeclSettings))
 		{
-			check(SuperfluousNode);
-			// Only remove them if they have no edges anymore
-			bool bHasAnyEdges = Algo::AnyOf(SuperfluousNode->GetInputPins(), [](const UPCGPin* InPin) { return InPin && InPin->EdgeCount() > 0; })
-				|| Algo::AnyOf(SuperfluousNode->GetOutputPins(), [](const UPCGPin* InPin) { return InPin && InPin->EdgeCount() > 0; });
+			return;
+		}
 
-			if (!bHasAnyEdges)
+		UPCGNode* OriginalDeclNode = CastChecked<UPCGNode>(OriginalDeclSettings->GetOuter());
+		UPCGNode** NewDeclNode = OriginalToNewNodesMapping.Find(OriginalDeclNode);
+		if (!NewDeclNode)
+		{
+			// If the declaration is outside of the subgraph, we need to create a new declaration inside the subgraph.
+			UPCGNode* NewlyCreatedDeclNode = NewGraph->ReconstructNewNode(OriginalDeclNode);
+			check(NewlyCreatedDeclNode);
+			NewlyCreatedDeclNode->NodeTitle = OriginalDeclNode->NodeTitle;
+			OriginalToNewNodesMapping.Emplace(OriginalDeclNode, NewlyCreatedDeclNode);
+
+			UPCGPin* OriginalOutputPin = OriginalDeclNode->GetOutputPin(PCGPinConstants::DefaultOutputLabel);
+			UPCGPin* NewInputPin = NewlyCreatedDeclNode->GetInputPin(PCGPinConstants::DefaultInputLabel);
+			check(OriginalOutputPin && NewInputPin);
+			ConnectPinsBetweenSubgraphBoundaries(OriginalOutputPin, NewInputPin, OriginalOutputPin, /*bIsInput=*/true);
+
+			NewDeclNode = &NewlyCreatedDeclNode;
+		}
+
+		// Patch and connect the edges.
+		check(NewDeclNode && *NewDeclNode);
+		NewUsageSettings->Declaration = CastChecked<UPCGNamedRerouteDeclarationSettings>((*NewDeclNode)->GetSettings());
+		UPCGPin* InvisiblePin = (*NewDeclNode)->GetOutputPin(PCGNamedRerouteConstants::InvisiblePinLabel);
+		UPCGPin* InputUsagePin = NewUsageNode->GetInputPin(PCGPinConstants::DefaultInputLabel);
+		check(InvisiblePin && InputUsagePin);
+		InvisiblePin->AddEdgeTo(InputUsagePin);
+
+		// And finally connect all the output pins normally
+		for (UPCGPin* OriginalOutputPin : OriginalUsageNode->GetOutputPins())
+		{
+			ReconnectNode(NewUsageNode, OriginalOutputPin);
+		}
+	}
+
+	/** Pre validation before starting the collapse, to verify that we won't introduce a cycle by collapsing. */
+	bool FCollapsingInformation::ValidateCollapseIsNotIntroducingCycle()
+	{
+		const TSet<UPCGNode*> NodesToCollapse{ ValidNodesToCollapse };
+
+		for (const UPCGNode* Node : NodesToCollapse)
+		{
+			if (!Node)
 			{
-				NodesToRemove.Add(SuperfluousNode);
+				continue;
+			}
+
+			// We can only switch boundaries once, otherwise it will introduce cycle
+			// Also keep track of all seen nodes, so our algorithm stops even if there are cycles already existing.
+			TSet<const UPCGNode*> SeenNodes;
+
+			auto Visitor = [&NodesToCollapse, &SeenNodes](const UPCGNode* CurrentNode, bool bInSubgraphBoundary, auto VisitorRecurse) -> bool
+			{
+				if (!CurrentNode || SeenNodes.Contains(CurrentNode))
+				{
+					//Nothing to do
+					return true;
+				}
+
+				SeenNodes.Add(CurrentNode);
+
+				for (const UPCGPin* OutputPin : CurrentNode->GetOutputPins())
+				{
+					if (!OutputPin)
+					{
+						continue;
+					}
+
+					for (const UPCGEdge* Edge : OutputPin->Edges)
+					{
+						if (!Edge)
+						{
+							continue;
+						}
+
+						if (const UPCGNode* OtherNode = Edge->OutputPin ? Edge->OutputPin->Node : nullptr)
+						{
+							const bool bNodeInSubgraphBoundary = NodesToCollapse.Contains(OtherNode);
+							if (!bInSubgraphBoundary && bNodeInSubgraphBoundary)
+							{
+								// We crossed twice, we can't collapse
+								return false;
+							}
+
+							if (!VisitorRecurse(OtherNode, bNodeInSubgraphBoundary, VisitorRecurse))
+							{
+								return false;
+							}
+						}
+					}
+				}
+
+				return true;
+			};
+
+			if (!Visitor(Node, /*bIsInSubgraphBoundary=*/true, Visitor))
+			{
+				return false;
 			}
 		}
 
-		// Finally remove them all in one go
-		InOriginalGraph->RemoveNodes(NodesToRemove);
+		return true;
+	}
+}
 
-		UPCGSettings* DefaultNodeSettings = nullptr;
-		SubgraphNode = InOriginalGraph->AddNodeOfType(UPCGSubgraphSettings::StaticClass(), DefaultNodeSettings);
-		UPCGSubgraphSettings* DefaultSubgraphSettings = CastChecked<UPCGSubgraphSettings>(DefaultNodeSettings);
-		DefaultSubgraphSettings->SetSubgraph(NewPCGGraph);
+UPCGGraph* FPCGSubgraphHelpers::CollapseIntoSubgraph(UPCGGraph* InOriginalGraph, const TArray<UPCGNode*>& InNodesToCollapse, const TArray<UObject*>& InExtraNodesToCollapse, UPCGGraph* OptionalPreAllocatedGraph)
+{
+	FText OutFailReason;
+	UPCGGraph* OutSubgraph = CollapseIntoSubgraphWithReason(InOriginalGraph, InNodesToCollapse, InExtraNodesToCollapse, OutFailReason, OptionalPreAllocatedGraph);
+	if (!OutSubgraph)
+	{
+		UE_LOG(LogPCG, Warning, TEXT("%s"), *OutFailReason.ToString());
+	}
 
-#if WITH_EDITOR
-		SubgraphNode->PositionX = CollapseInfo.AveragePosition.X;
-		SubgraphNode->PositionY = CollapseInfo.AveragePosition.Y;
-#endif // WITH_EDITOR
+	return OutSubgraph;
+}
 
-		SubgraphNode->UpdateAfterSettingsChangeDuringCreation();
+UPCGGraph* FPCGSubgraphHelpers::CollapseIntoSubgraphWithReason(UPCGGraph* InOriginalGraph, const TArray<UPCGNode*>& InNodesToCollapse, const TArray<UObject*>& InExtraNodesToCollapse, FText& OutFailReason, UPCGGraph* OptionalPreAllocatedGraph)
+{
+	// Don't collapse into a subgraph if you don't have at least 2 nodes
+	if (!InOriginalGraph)
+	{
+		OutFailReason = LOCTEXT("InvalidGraph", "Original Graph is invalid.");
+		return nullptr;
 	}
 
-	// 8. Connect all the edges
-	TSet<UPCGNode*> TouchedNodes;
-	for (EdgePlaceholder& Edge : EdgePlaceholders)
+	if (InNodesToCollapse.Num() < 2)
 	{
-		if (Edge.InputPin == nullptr)
-		{
-			SubgraphNode->GetOutputPin(Edge.InputPinLabel)->AddEdgeTo(Edge.OutputPin, &TouchedNodes);
-		}
-		else if (Edge.OutputPin == nullptr)
-		{
-			Edge.InputPin->AddEdgeTo(SubgraphNode->GetInputPin(Edge.OutputPinLabel), &TouchedNodes);
-		}
-		else
-		{
-			Edge.InputPin->AddEdgeTo(Edge.OutputPin, &TouchedNodes);
-		}
+		OutFailReason = LOCTEXT("CollapseLessThanTwoNodes", "Try to collapse less than 2 nodes.");
+		return nullptr;
 	}
 
+	const FInstancedPropertyBag* GraphParameters = InOriginalGraph->GetUserParametersStruct();
+	check(GraphParameters);
+
+	PCGSubgraphHelpersExtra::FCollapsingInformation CollapseInfo = PCGSubgraphHelpersExtra::FCollapsingInformation::CreateAndInitialize(InNodesToCollapse, *GraphParameters, InOriginalGraph);
+
+	// If we have at most 1 valid node to collapse, just exit
+	if (!CollapseInfo.HasEnoughValidNodes())
+	{
+		OutFailReason = LOCTEXT("ValidLessThanTwoNodes", "There were less than 2 PCG nodes selected, abort");
+		return nullptr;
+	}
+
+	if (!CollapseInfo.ValidateCollapseIsNotIntroducingCycle())
+	{
+		OutFailReason = LOCTEXT("CycleDetected", "Collapsing would introduce a cycle (A -> B -> C situation when we try to collapse A and C, B would be connected to both ends). Abort.");
+		return nullptr;
+	}
+
+	// Create a new subgraph if necessary.
+	UPCGGraph* NewPCGGraph = OptionalPreAllocatedGraph;
+	if (!NewPCGGraph)
+	{
+		NewPCGGraph = NewObject<UPCGGraph>();
+	}
+
+	CollapseInfo.SetNewGraphAndUpdatePositions(NewPCGGraph);
+
 #if WITH_EDITOR
+	// Disable all notifications for both graphs as we do our changes.
+	InOriginalGraph->DisableNotificationsForEditor();
+	NewPCGGraph->DisableNotificationsForEditor();
+#endif // WITH_EDITOR
+
+	CollapseInfo.AddDiscoveredGraphParameters();
+	CollapseInfo.CopyNodes(InExtraNodesToCollapse);
+	CollapseInfo.ReconnectAllNodes();
+	CollapseInfo.CleanUpOriginalEdges();
+	CollapseInfo.DeleteNodes(InExtraNodesToCollapse);
+	CollapseInfo.CreateSubgraphNodeAndConnectRemainingPins();
+
+#if WITH_EDITOR
+	// Trigger a structural change for original graph and flush all notifications for subgraph and re-enable notifications for both.
 	InOriginalGraph->NotifyGraphChanged(EPCGChangeType::Structural);
 	NewPCGGraph->NotifyGraphChanged(EPCGChangeType::Structural);
+
+	InOriginalGraph->EnableNotificationsForEditor();
+	NewPCGGraph->EnableNotificationsForEditor();
 #endif // WITH_EDITOR
 
 	return NewPCGGraph;
diff --git a/Engine/Plugins/PCG/Source/PCG/Private/PCGGraph.cpp b/Engine/Plugins/PCG/Source/PCG/Private/PCGGraph.cpp
--- a/Engine/Plugins/PCG/Source/PCG/Private/PCGGraph.cpp
+++ b/Engine/Plugins/PCG/Source/PCG/Private/PCGGraph.cpp
@@ -610,7 +610,7 @@ UPCGNode* UPCGGraph::AddNodeInstance(UPCGSettings* InSettings)
 	return Node;
 }
 
-UPCGNode* UPCGGraph::AddNodeCopy(UPCGSettings* InSettings, UPCGSettings*& DefaultNodeSettings)
+UPCGNode* UPCGGraph::AddNodeCopy(const UPCGSettings* InSettings, UPCGSettings*& DefaultNodeSettings)
 {
 	if (!InSettings)
 	{
diff --git a/Engine/Plugins/PCG/Source/PCG/Private/Tests/Subgraphs/PCGCollapseSubgraphTest.cpp b/Engine/Plugins/PCG/Source/PCG/Private/Tests/Subgraphs/PCGCollapseSubgraphTest.cpp
--- a/Engine/Plugins/PCG/Source/PCG/Private/Tests/Subgraphs/PCGCollapseSubgraphTest.cpp
+++ b/Engine/Plugins/PCG/Source/PCG/Private/Tests/Subgraphs/PCGCollapseSubgraphTest.cpp
@@ -1,18 +1,26 @@
 // Copyright Epic Games, Inc. All Rights Reserved.
 
+#include "PCGComponent.h"
 #include "PCGContext.h"
 #include "PCGEdge.h"
 #include "PCGGraph.h"
-#include "PCGSubgraph.h"
 #include "PCGInputOutputSettings.h"
+#include "PCGPoint.h"
+#include "PCGSubgraph.h"
+#include "Data/PCGPointData.h"
 #include "Elements/PCGAttributeNoise.h"
 #include "Elements/PCGSurfaceSampler.h"
 #include "Elements/PCGUserParameterGet.h"
+#include "Graph/PCGGraphExecutor.h"
 #include "Helpers/PCGSubgraphHelpers.h"
 #include "Tests/PCGTestsCommon.h"
 
+#include "UObject/Package.h"
+
+
 IMPLEMENT_CUSTOM_SIMPLE_AUTOMATION_TEST(FPCGCollapseSubgraphSimple, FPCGTestBaseClass, "Plugins.PCG.Subgraph.Collapse.Simple", PCGTestsCommon::TestFlags);
 IMPLEMENT_CUSTOM_SIMPLE_AUTOMATION_TEST(FPCGCollapseSubgraphWithParams, FPCGTestBaseClass, "Plugins.PCG.Subgraph.Collapse.WithParams", PCGTestsCommon::TestFlags);
+IMPLEMENT_CUSTOM_SIMPLE_AUTOMATION_TEST(FPCGCollapseSubgraphCombinationOfCollapses, FPCGTestBaseClass, "Plugins.PCG.Subgraph.Collapse.Combination", PCGTestsCommon::TestFlags);
 
 namespace PCGCollapseSubgraphTests
 {
@@ -173,7 +181,7 @@ bool FPCGCollapseSubgraphSimple::RunTest(const FString& Parameters)
 	UTEST_NOT_NULL("Node is a subgraph node", SubgraphSettings);
 	UTEST_EQUAL("Subgraph graph is the right one", SubgraphSettings->SubgraphInstance->GetGraph(), Subgraph);
 
-	const FName InputPinSubgraphLabel = TEXT("Surface Sampler Surface");
+	const FName InputPinSubgraphLabel = TEXT("Landscape");
 
 	// Landscape pin from input is connected to created pin of subgraph node
 	VALIDATE_EDGE(UPCGGraphInputOutputSettings, UPCGSubgraphSettings, MainGraph, InputLandscapePin.Label, InputPinSubgraphLabel)
@@ -248,4 +256,144 @@ bool FPCGCollapseSubgraphWithParams::RunTest(const FString& Parameters)
 	return true;
 }
 
+bool FPCGCollapseSubgraphCombinationOfCollapses::RunTest(const FString& Parameters)
+{
+	const FSoftObjectPath UnitTestGraphAssetPath(TEXT("/Script/PCG.PCGGraph'/PCG/UnitTest/UnitTestSubgraphCollapse.UnitTestSubgraphCollapse'"));
+	UPCGGraph* UnitTestGraph = TSoftObjectPtr<UPCGGraph>(UnitTestGraphAssetPath).LoadSynchronous();
+	UTEST_NOT_NULL("UnitTest graph was successfully loaded", UnitTestGraph);
+	check(UnitTestGraph);
+
+	// We expect our graph to have 12 nodes, as it will impact the number of combinations below.
+	UTEST_EQUAL("Expected number of nodes", UnitTestGraph->GetNodes().Num(), 12);
+
+	PCGTestsCommon::FTestData TestData{};
+
+	// Temporary solution to run a full graph and extract the resulting points.
+	auto ExecuteGraph = [&TestData](UPCGGraph* GraphToExecute, double FirstConstant, double SecondConstant) -> TArray<FPCGPoint>
+	{
+		FPCGGraphExecutor GraphExecutor;
+		TArray<FPCGPoint> OutputPoints;
+		bool bDone = false;
+
+		TestData.TestPCGComponent->SetGraphLocal(GraphToExecute);
+		TestData.TestPCGComponent->GetGraphInstance()->SetGraphParameter<double>(TEXT("FloatConstant"), FirstConstant);
+		TestData.TestPCGComponent->GetGraphInstance()->SetGraphParameter<double>(TEXT("AnotherConstant"), SecondConstant);
+
+		const FPCGTaskId TaskId = GraphExecutor.Schedule(TestData.TestPCGComponent);
+		const FPCGTaskId FinalTaskId = GraphExecutor.ScheduleGeneric([&OutputPoints, &GraphExecutor, TaskId, &bDone]() -> bool
+		{
+			FPCGDataCollection OutputData{};
+			GraphExecutor.GetOutputData(TaskId, OutputData);
+
+			if (OutputData.TaggedData.Num() == 1)
+			{
+				if (const UPCGPointData* PointData = Cast<UPCGPointData>(OutputData.TaggedData[0].Data))
+				{
+					OutputPoints = PointData->GetPoints();
+				}
+			}
+
+			bDone = true;
+			return true;
+		}, TestData.TestPCGComponent, { TaskId });
+
+		while(!bDone)
+		{
+			GraphExecutor.Execute();
+		}
+
+		return OutputPoints;
+	};
+
+	auto ComparePoints = [this](const TArray<FPCGPoint>& InPointsA, const TArray<FPCGPoint>& InPointsB, const FString& Name) -> bool
+	{
+		UTEST_EQUAL(*FString::Printf(TEXT("%s: Same number of points"), *Name), InPointsB.Num(), InPointsA.Num());
+		for (int i = 0; i < InPointsA.Num(); ++i)
+		{
+			UTEST_EQUAL(*FString::Printf(TEXT("%s: Same position"), *Name), InPointsB[i].Transform.GetLocation(), InPointsA[i].Transform.GetLocation());
+		}
+
+		return true;
+	};
+
+	constexpr double FirstConstantValue = 50.0;
+	constexpr double FirstConstantValue2 = 100.0;
+	constexpr double SecondConstantValue = 500.0;
+	constexpr double SecondConstantValue2 = 5000.0;
+
+	// Run 3 times with different graph parameters on each
+	const TArray<FPCGPoint> FirstRunPoints = ExecuteGraph(UnitTestGraph, FirstConstantValue, SecondConstantValue);
+	const TArray<FPCGPoint> SecondRunPoints = ExecuteGraph(UnitTestGraph, FirstConstantValue2, SecondConstantValue);
+	const TArray<FPCGPoint> ThirdRunPoints = ExecuteGraph(UnitTestGraph, FirstConstantValue, SecondConstantValue2);
+
+	// Then do all the 4 nodes collapse combinations. We expect our initial node to have 12 nodes. 12 choose 4 is 495 (which is the sum of ExpectedNumberOfCycles and ExpectedNumberOfSuccess)
+	bool bSuccess = true;
+	const int32 NodesNum = UnitTestGraph->GetNodes().Num();
+	constexpr int32 ExpectedNumberOfCycles = 315;
+	constexpr int32 ExpectedNumberOfSuccess = 180;
+	int32 NumberOfCycles = 0;
+	int32 NumberOfSuccess = 0;
+	for (int i = 0; i < NodesNum && bSuccess; ++i)
+	{
+		for (int j = i + 1; j < NodesNum && bSuccess; j++)
+		{
+			for (int k = j + 1; k < NodesNum && bSuccess; ++k)
+			{
+				for (int l = k + 1; l < NodesNum && bSuccess; ++l)
+				{
+					UPCGGraph* OriginalGraph = DuplicateObject(UnitTestGraph, GetTransientPackage(), TEXT("UnitTest_GraphDuplicate"));
+					check(OriginalGraph);
+					OriginalGraph->SetFlags(RF_Transient);
+
+					FString Name = FString::Printf(TEXT("[%d, %d, %d, %d]"), i, j, k, l);
+					FString AllNodes = FString::Printf(TEXT("[%s, %s, %s, %s]"), *OriginalGraph->GetNodes()[i]->GetName(), *OriginalGraph->GetNodes()[j]->GetName(), *OriginalGraph->GetNodes()[k]->GetName(), *OriginalGraph->GetNodes()[l]->GetName());
+
+					const TArray<UPCGNode*> NodesToCollapse = { OriginalGraph->GetNodes()[i], OriginalGraph->GetNodes()[j], OriginalGraph->GetNodes()[k], OriginalGraph->GetNodes()[l] };
+					FText OutFailReason;
+					if (UPCGGraph* Subgraph = FPCGSubgraphHelpers::CollapseIntoSubgraphWithReason(OriginalGraph, NodesToCollapse, {}, OutFailReason))
+					{
+						Subgraph->SetFlags(RF_Transient);
+
+						const TArray<FPCGPoint> FirstRunCollapsedPoints = ExecuteGraph(OriginalGraph, FirstConstantValue, SecondConstantValue);
+						const TArray<FPCGPoint> SecondRunCollapsedPoints = ExecuteGraph(OriginalGraph, FirstConstantValue2, SecondConstantValue);
+						const TArray<FPCGPoint> ThirdRunCollapsedPoints = ExecuteGraph(OriginalGraph, FirstConstantValue, SecondConstantValue2);
+
+						bool bLocalSuccess = ComparePoints(FirstRunPoints, FirstRunCollapsedPoints, Name + TEXT(" First run"));
+						bLocalSuccess &= ComparePoints(SecondRunPoints, SecondRunCollapsedPoints, Name + TEXT(" Second run"));
+						bLocalSuccess &= ComparePoints(ThirdRunPoints, ThirdRunCollapsedPoints, Name + TEXT(" Third run"));
+
+						if (!bLocalSuccess)
+						{
+							UE_LOG(LogPCG, Log, TEXT("Combination %s failed. Nodes: %s"), *Name, *AllNodes);
+						}
+						else
+						{
+							NumberOfSuccess++;
+						}
+
+						bSuccess &= bLocalSuccess;
+
+						Subgraph->MarkAsGarbage();
+					}
+					else
+					{
+						UE_LOG(LogPCG, Log, TEXT("Combination %s has probably a cycle (Msg: %s)"), *Name, *OutFailReason.ToString());
+						NumberOfCycles++;
+					}
+
+					OriginalGraph->MarkAsGarbage();
+				}
+			}
+		}
+	}
+
+	if (bSuccess)
+	{
+		UTEST_EQUAL("Number of success expected", NumberOfSuccess, ExpectedNumberOfSuccess);
+		UTEST_EQUAL("Number of cycles expected", NumberOfCycles, ExpectedNumberOfCycles);
+	}
+
+	return bSuccess;
+}
+
 #undef VALIDATE_EDGE
\ No newline at end of file
diff --git a/Engine/Plugins/PCG/Source/PCG/Public/Elements/PCGReroute.h b/Engine/Plugins/PCG/Source/PCG/Public/Elements/PCGReroute.h
--- a/Engine/Plugins/PCG/Source/PCG/Public/Elements/PCGReroute.h
+++ b/Engine/Plugins/PCG/Source/PCG/Public/Elements/PCGReroute.h
@@ -55,6 +55,13 @@ class PCG_API UPCGNamedRerouteDeclarationSettings : public UPCGNamedRerouteBaseS
 {
 	GENERATED_BODY()
 
+public:
+#if WITH_EDITOR
+	//~Begin UPCGSettings interface
+	virtual FName GetDefaultNodeName() const override { return FName("NamedRerouteDeclaration"); }
+	//~End UPCGSettings interface
+#endif
+
 protected:
 	virtual TArray<FPCGPinProperties> OutputPinProperties() const override;
 
@@ -69,6 +76,13 @@ class PCG_API UPCGNamedRerouteUsageSettings : public UPCGNamedRerouteBaseSetting
 {
 	GENERATED_BODY()
 
+public:
+#if WITH_EDITOR
+	//~Begin UPCGSettings interface
+	virtual FName GetDefaultNodeName() const override { return FName("NamedRerouteUsage"); }
+	//~End UPCGSettings interface
+#endif
+
 protected:
 	virtual TArray<FPCGPinProperties> InputPinProperties() const override;
 
diff --git a/Engine/Plugins/PCG/Source/PCG/Public/Helpers/PCGSubgraphHelpers.h b/Engine/Plugins/PCG/Source/PCG/Public/Helpers/PCGSubgraphHelpers.h
--- a/Engine/Plugins/PCG/Source/PCG/Public/Helpers/PCGSubgraphHelpers.h
+++ b/Engine/Plugins/PCG/Source/PCG/Public/Helpers/PCGSubgraphHelpers.h
@@ -11,5 +11,6 @@ class UObject;
 class FPCGSubgraphHelpers
 {
 public:
+	static PCG_API UPCGGraph* CollapseIntoSubgraphWithReason(UPCGGraph* InOriginalGraph, const TArray<UPCGNode*>& InNodesToCollapse, const TArray<UObject*>& InExtraNodesToCollapse, FText& OutFailReason, UPCGGraph* OptionalPreAllocatedGraph = nullptr);
 	static PCG_API UPCGGraph* CollapseIntoSubgraph(UPCGGraph* InOriginalGraph, const TArray<UPCGNode*>& InNodesToCollapse, const TArray<UObject*>& InExtraNodesToCollapse, UPCGGraph* OptionalPreAllocatedGraph = nullptr);
 };
\ No newline at end of file
diff --git a/Engine/Plugins/PCG/Source/PCG/Public/PCGGraph.h b/Engine/Plugins/PCG/Source/PCG/Public/PCGGraph.h
--- a/Engine/Plugins/PCG/Source/PCG/Public/PCGGraph.h
+++ b/Engine/Plugins/PCG/Source/PCG/Public/PCGGraph.h
@@ -237,7 +237,7 @@ public:
 
 	/** Creates a node and copies the input settings. Returns the created node. */
 	UFUNCTION(BlueprintCallable, Category = Graph, meta = (DeterminesOutputType = "InSettings", DynamicOutputParam = "OutCopiedSettings"))
-	UPCGNode* AddNodeCopy(UPCGSettings* InSettings, UPCGSettings*& DefaultNodeSettings);
+	UPCGNode* AddNodeCopy(const UPCGSettings* InSettings, UPCGSettings*& DefaultNodeSettings);
 
 	/** Removes a node from the graph. */
 	UFUNCTION(BlueprintCallable, Category = Graph)
diff --git a/Engine/Plugins/PCG/Source/PCGEditor/Private/PCGEditor.cpp b/Engine/Plugins/PCG/Source/PCGEditor/Private/PCGEditor.cpp
--- a/Engine/Plugins/PCG/Source/PCGEditor/Private/PCGEditor.cpp
+++ b/Engine/Plugins/PCG/Source/PCGEditor/Private/PCGEditor.cpp
@@ -63,6 +63,7 @@
 #include "Framework/Notifications/NotificationManager.h"
 #include "HAL/PlatformApplicationMisc.h"
 #include "Misc/ITransaction.h"
+#include "Misc/MessageDialog.h"
 #include "Misc/TransactionObjectEvent.h"
 #include "Preferences/UnrealEdOptions.h"
 #include "Widgets/Docking/SDockTab.h"
@@ -1643,31 +1644,32 @@ void FPCGEditor::OnCollapseNodesInSubgraph()
 
 	// Create a new subgraph, by creating a new PCGGraph asset.
 	TObjectPtr<UPCGGraph> NewPCGGraph = nullptr;
-	{
-		FScopedTransaction Transaction(LOCTEXT("PCGCollapseInSubgraphMessage", "[PCG] Collapse into Subgraph"));
 
-		IAssetTools& AssetTools = FModuleManager::Get().LoadModuleChecked<FAssetToolsModule>("AssetTools").Get();
+	IAssetTools& AssetTools = FModuleManager::Get().LoadModuleChecked<FAssetToolsModule>("AssetTools").Get();
 
-		TObjectPtr<UPCGGraphFactory> Factory = NewObject<UPCGGraphFactory>();
+	TObjectPtr<UPCGGraphFactory> Factory = NewObject<UPCGGraphFactory>();
 
-		FString NewPackageName;
-		FString NewAssetName;
-		PCGEditorUtils::GetParentPackagePathAndUniqueName(PCGGraph, LOCTEXT("NewPCGSubgraphAsset", "NewPCGSubgraph").ToString(), NewPackageName, NewAssetName);
+	FString NewPackageName;
+	FString NewAssetName;
+	PCGEditorUtils::GetParentPackagePathAndUniqueName(PCGGraph, LOCTEXT("NewPCGSubgraphAsset", "NewPCGSubgraph").ToString(), NewPackageName, NewAssetName);
 
-		NewPCGGraph = Cast<UPCGGraph>(AssetTools.CreateAssetWithDialog(NewAssetName, NewPackageName, PCGGraph->GetClass(), Factory, "PCGEditor_CollapseInSubgraph"));
+	NewPCGGraph = Cast<UPCGGraph>(AssetTools.CreateAssetWithDialog(NewAssetName, NewPackageName, PCGGraph->GetClass(), Factory, "PCGEditor_CollapseInSubgraph"));
 
-		if (NewPCGGraph == nullptr)
-		{
-			UE_LOG(LogPCGEditor, Warning, TEXT("Subgraph asset creation was aborted or failed, abort."));
-			Transaction.Cancel();
-			return;
-		}
+	if (NewPCGGraph == nullptr)
+	{
+		UE_LOG(LogPCGEditor, Warning, TEXT("Subgraph asset creation was aborted or failed, abort."));
+		return;
+	}
 
-		NewPCGGraph = FPCGSubgraphHelpers::CollapseIntoSubgraph(PCGGraph, NodesToCollapse, ExtraNodesToCollapse, NewPCGGraph);
+	{
+		FScopedTransaction Transaction(LOCTEXT("PCGCollapseInSubgraphMessage", "[PCG] Collapse into Subgraph"));
+		FText OutFailReason;
+		NewPCGGraph = FPCGSubgraphHelpers::CollapseIntoSubgraphWithReason(PCGGraph, NodesToCollapse, ExtraNodesToCollapse, OutFailReason, NewPCGGraph);
 
 		if (NewPCGGraph == nullptr)
 		{
-			UE_LOG(LogPCGEditor, Warning, TEXT("Subgraph collapse failed, abort."));
+			FMessageDialog::Open(EAppMsgType::Ok, OutFailReason, LOCTEXT("PCGCollapseInSubgraphFailed", "PCG Subgraph Collapse Failed"));
+			Transaction.Cancel();
 			return;
 		}
 
@@ -1675,11 +1677,14 @@ void FPCGEditor::OnCollapseNodesInSubgraph()
 		PCGEditorGraph->ReconstructGraph();
 	}
 
-	// Save the new asset
-	UEditorAssetLibrary::SaveLoadedAsset(NewPCGGraph);
+	if (NewPCGGraph)
+	{
+		// Save the new asset
+		UEditorAssetLibrary::SaveLoadedAsset(NewPCGGraph);
 
-	// Notify the widget
-	GraphEditorWidget->NotifyGraphChanged();
+		// Notify the widget
+		GraphEditorWidget->NotifyGraphChanged();
+	}
 }
 
 bool FPCGEditor::CanExportNodes() const
