[PCG] PCG_API Internalization and Deprecation Elimination Pass

#rnx

diff --git a/Engine/Plugins/PCG/Source/PCG/Private/Elements/Metadata/PCGMetadataElement.cpp b/Engine/Plugins/PCG/Source/PCG/Private/Elements/Metadata/PCGMetadataElement.cpp
--- a/Engine/Plugins/PCG/Source/PCG/Private/Elements/Metadata/PCGMetadataElement.cpp
+++ b/Engine/Plugins/PCG/Source/PCG/Private/Elements/Metadata/PCGMetadataElement.cpp
@@ -96,40 +96,6 @@ TArray<FPCGPinProperties> UPCGMetadataOperationSettings::InputPinProperties() co
 void UPCGMetadataOperationSettings::PostLoad()
 {
 	Super::PostLoad();
-
-#if WITH_EDITOR
-PRAGMA_DISABLE_DEPRECATION_WARNINGS
-	if (Target_DEPRECATED != EPCGMetadataOperationTarget::PropertyToAttribute
-		|| (SourceAttribute_DEPRECATED != NAME_None)
-		|| (DestinationAttribute_DEPRECATED != NAME_None)
-		|| (PointProperty_DEPRECATED != EPCGPointProperties::Density))
-	{
-		switch (Target_DEPRECATED)
-		{
-		case EPCGMetadataOperationTarget::PropertyToAttribute:
-			InputSource.SetPointProperty(PointProperty_DEPRECATED);
-			OutputTarget.SetAttributeName(DestinationAttribute_DEPRECATED);
-			break;
-		case EPCGMetadataOperationTarget::AttributeToProperty:
-			InputSource.SetAttributeName(SourceAttribute_DEPRECATED);
-			OutputTarget.SetPointProperty(PointProperty_DEPRECATED);
-			break;
-		case EPCGMetadataOperationTarget::AttributeToAttribute:
-			InputSource.SetAttributeName(SourceAttribute_DEPRECATED);
-			OutputTarget.SetAttributeName(DestinationAttribute_DEPRECATED);
-			break;
-		default:
-			break;
-		}
-
-		// Default values.
-		SourceAttribute_DEPRECATED = NAME_None;
-		DestinationAttribute_DEPRECATED = NAME_None;
-		Target_DEPRECATED = EPCGMetadataOperationTarget::PropertyToAttribute;
-		PointProperty_DEPRECATED = EPCGPointProperties::Density;
-	}
-PRAGMA_ENABLE_DEPRECATION_WARNINGS
-#endif // WITH_EDITOR
 }
 
 bool FPCGMetadataOperationElement::ExecuteInternal(FPCGContext* Context) const
diff --git a/Engine/Plugins/PCG/Source/PCG/Private/Elements/PCGCreateTargetActor.cpp b/Engine/Plugins/PCG/Source/PCG/Private/Elements/PCGCreateTargetActor.cpp
--- a/Engine/Plugins/PCG/Source/PCG/Private/Elements/PCGCreateTargetActor.cpp
+++ b/Engine/Plugins/PCG/Source/PCG/Private/Elements/PCGCreateTargetActor.cpp
@@ -222,6 +222,29 @@ void UPCGCreateTargetActor::PostLoad()
 #endif // WITH_EDITOR
 }
 
+void UPCGCreateTargetActor::SetTemplateActorClass(const TSubclassOf<AActor>& InTemplateActorClass)
+{
+#if WITH_EDITOR
+	TeardownBlueprintEvent();
+#endif // WITH_EDITOR
+
+	TemplateActorClass = InTemplateActorClass;
+
+#if WITH_EDITOR
+	SetupBlueprintEvent();
+	RefreshTemplateActor();
+#endif // WITH_EDITOR
+}
+
+void UPCGCreateTargetActor::SetAllowTemplateActorEditing(bool bInAllowTemplateActorEditing)
+{
+	bAllowTemplateActorEditing = bInAllowTemplateActorEditing;
+
+#if WITH_EDITOR
+	RefreshTemplateActor();
+#endif // WITH_EDITOR
+}
+
 bool FPCGCreateTargetActorElement::ExecuteInternal(FPCGContext* Context) const
 {
 	TRACE_CPUPROFILER_EVENT_SCOPE(FPCGCreateTargetActorElement::Execute);
diff --git a/Engine/Plugins/PCG/Source/PCG/Private/Elements/PCGSpawnActor.cpp b/Engine/Plugins/PCG/Source/PCG/Private/Elements/PCGSpawnActor.cpp
--- a/Engine/Plugins/PCG/Source/PCG/Private/Elements/PCGSpawnActor.cpp
+++ b/Engine/Plugins/PCG/Source/PCG/Private/Elements/PCGSpawnActor.cpp
@@ -141,6 +141,29 @@ UPCGNode* UPCGSpawnActorSettings::CreateNode() const
 	return NewObject<UPCGSpawnActorNode>();
 }
 
+void UPCGSpawnActorSettings::SetTemplateActorClass(const TSubclassOf<AActor>& InTemplateActorClass)
+{
+#if WITH_EDITOR
+	TeardownBlueprintEvent();
+#endif // WITH_EDITOR
+
+	TemplateActorClass = InTemplateActorClass;
+
+#if WITH_EDITOR
+	SetupBlueprintEvent();
+	RefreshTemplateActor();
+#endif // WITH_EDITOR
+}
+
+void UPCGSpawnActorSettings::SetAllowTemplateActorEditing(bool bInAllowTemplateActorEditing)
+{
+	bAllowTemplateActorEditing = bInAllowTemplateActorEditing;
+
+#if WITH_EDITOR
+	RefreshTemplateActor();
+#endif // WITH_EDITOR
+}
+
 FPCGElementPtr UPCGSpawnActorSettings::CreateElement() const
 {
 	return MakeShared<FPCGSpawnActorElement>();
diff --git a/Engine/Plugins/PCG/Source/PCG/Private/Elements/PCGStaticMeshSpawner.cpp b/Engine/Plugins/PCG/Source/PCG/Private/Elements/PCGStaticMeshSpawner.cpp
--- a/Engine/Plugins/PCG/Source/PCG/Private/Elements/PCGStaticMeshSpawner.cpp
+++ b/Engine/Plugins/PCG/Source/PCG/Private/Elements/PCGStaticMeshSpawner.cpp
@@ -481,26 +481,6 @@ void UPCGStaticMeshSpawnerSettings::PostLoad()
 {
 	Super::PostLoad();
 
-#if WITH_EDITOR
-	if (Meshes_DEPRECATED.Num() != 0)
-	{
-		SetMeshSelectorType(UPCGMeshSelectorWeighted::StaticClass());
-
-		UPCGMeshSelectorWeighted* MeshSelector = CastChecked<UPCGMeshSelectorWeighted>(MeshSelectorParameters);
-
-		for (const FPCGStaticMeshSpawnerEntry& Entry : Meshes_DEPRECATED)
-		{
-			FPCGMeshSelectorWeightedEntry& NewEntry = MeshSelector->MeshEntries.Emplace_GetRef(Entry.Mesh, Entry.Weight);
-			if (Entry.bOverrideCollisionProfile)
-			{
-				NewEntry.Descriptor.BodyInstance.SetCollisionProfileName(Entry.CollisionProfile.Name);
-			}
-		}
-
-		Meshes_DEPRECATED.Reset();
-	}
-#endif
-
 	const EObjectFlags Flags = GetMaskedFlags(RF_PropagateToSubObjects) | RF_Transactional;
 	
 	if (!MeshSelectorParameters)
diff --git a/Engine/Plugins/PCG/Source/PCG/Private/Elements/PCGTextureSampler.cpp b/Engine/Plugins/PCG/Source/PCG/Private/Elements/PCGTextureSampler.cpp
--- a/Engine/Plugins/PCG/Source/PCG/Private/Elements/PCGTextureSampler.cpp
+++ b/Engine/Plugins/PCG/Source/PCG/Private/Elements/PCGTextureSampler.cpp
@@ -247,6 +247,15 @@ FPCGContext* FPCGTextureSamplerElement::CreateContext()
 	return new FPCGTextureSamplerContext();
 }
 
+void UPCGTextureSamplerSettings::SetTexture(TSoftObjectPtr<UTexture> InTexture)
+{
+	Texture = InTexture;
+
+#if WITH_EDITOR
+	UpdateDisplayTextureArrayIndex();
+#endif // WITH_EDITOR
+}
+
 void FPCGTextureSamplerElement::GetDependenciesCrc(const FPCGDataCollection& InInput, const UPCGSettings* InSettings, UPCGComponent* InComponent, FPCGCrc& OutCrc) const
 {
 	FPCGCrc Crc;
diff --git a/Engine/Plugins/PCG/Source/PCG/Private/Helpers/PCGDynamicTrackingHelpers.cpp b/Engine/Plugins/PCG/Source/PCG/Private/Helpers/PCGDynamicTrackingHelpers.cpp
--- a/Engine/Plugins/PCG/Source/PCG/Private/Helpers/PCGDynamicTrackingHelpers.cpp
+++ b/Engine/Plugins/PCG/Source/PCG/Private/Helpers/PCGDynamicTrackingHelpers.cpp
@@ -12,7 +12,7 @@ void FPCGDynamicTrackingHelper::EnableAndInitialize(const FPCGContext* InContext
 	if (InContext && InContext->SourceComponent.IsValid())
 	{
 		const UPCGSettings* Settings = InContext->GetOriginalSettings<UPCGSettings>();
-		if (Settings && Settings->CanDynamicalyTrackKeys())
+		if (Settings && Settings->CanDynamicallyTrackKeys())
 		{
 			CachedComponent = InContext->SourceComponent;
 			bDynamicallyTracked = true;
diff --git a/Engine/Plugins/PCG/Source/PCG/Private/PCGElement.cpp b/Engine/Plugins/PCG/Source/PCG/Private/PCGElement.cpp
--- a/Engine/Plugins/PCG/Source/PCG/Private/PCGElement.cpp
+++ b/Engine/Plugins/PCG/Source/PCG/Private/PCGElement.cpp
@@ -176,7 +176,7 @@ void IPCGElement::PostExecute(FPCGContext* Context) const
 
 #if WITH_EDITOR
 	// Register the element to the component indicating the element has run and can have dynamic tracked keys.
-	if (Settings && Settings->CanDynamicalyTrackKeys() && Context->SourceComponent.IsValid())
+	if (Settings && Settings->CanDynamicallyTrackKeys() && Context->SourceComponent.IsValid())
 	{
 		Context->SourceComponent->RegisterDynamicTracking(Settings, {});
 	}
diff --git a/Engine/Plugins/PCG/Source/PCG/Public/Data/PCGCollisionShapeData.h b/Engine/Plugins/PCG/Source/PCG/Public/Data/PCGCollisionShapeData.h
--- a/Engine/Plugins/PCG/Source/PCG/Public/Data/PCGCollisionShapeData.h
+++ b/Engine/Plugins/PCG/Source/PCG/Public/Data/PCGCollisionShapeData.h
@@ -12,14 +12,14 @@
 
 class UShapeComponent;
 
-UCLASS(BlueprintType, ClassGroup = (Procedural))
-class PCG_API UPCGCollisionShapeData : public UPCGSpatialDataWithPointCache
+UCLASS(MinimalAPI, BlueprintType, ClassGroup = (Procedural))
+class UPCGCollisionShapeData : public UPCGSpatialDataWithPointCache
 {
 	GENERATED_BODY()
 
 public:
-	void Initialize(UShapeComponent* InShape);
-	static bool IsSupported(UShapeComponent* InShape);
+	PCG_API void Initialize(UShapeComponent* InShape);
+	PCG_API static bool IsSupported(UShapeComponent* InShape);
 
 	// ~Begin UPCGData interface
 	virtual EPCGDataType GetDataType() const override { return EPCGDataType::Primitive; }
diff --git a/Engine/Plugins/PCG/Source/PCG/Public/Data/PCGLandscapeData.h b/Engine/Plugins/PCG/Source/PCG/Public/Data/PCGLandscapeData.h
--- a/Engine/Plugins/PCG/Source/PCG/Public/Data/PCGLandscapeData.h
+++ b/Engine/Plugins/PCG/Source/PCG/Public/Data/PCGLandscapeData.h
@@ -14,7 +14,7 @@ class ULandscapeInfo;
 class UPCGLandscapeCache;
 
 USTRUCT(BlueprintType)
-struct PCG_API FPCGLandscapeDataProps
+struct FPCGLandscapeDataProps
 {
 	GENERATED_BODY()
 
diff --git a/Engine/Plugins/PCG/Source/PCG/Public/Data/PCGRenderTargetData.h b/Engine/Plugins/PCG/Source/PCG/Public/Data/PCGRenderTargetData.h
--- a/Engine/Plugins/PCG/Source/PCG/Public/Data/PCGRenderTargetData.h
+++ b/Engine/Plugins/PCG/Source/PCG/Public/Data/PCGRenderTargetData.h
@@ -10,8 +10,8 @@ class UTextureRenderTarget2D;
 
 //TODO: It's possible that caching the result in this class is not as efficient as it could be
 // if we expect to sample in different ways (e.g. channel) in the same render target
-UCLASS(BlueprintType, ClassGroup = (Procedural))
-class PCG_API UPCGRenderTargetData : public UPCGBaseTextureData
+UCLASS(MinimalAPI, BlueprintType, ClassGroup = (Procedural))
+class UPCGRenderTargetData : public UPCGBaseTextureData
 {
 	GENERATED_BODY()
 
@@ -29,12 +29,8 @@ protected:
 
 public:
 	UFUNCTION(BlueprintCallable, Category = RenderTarget)
-	void Initialize(UTextureRenderTarget2D* InRenderTarget, const FTransform& InTransform);
+	PCG_API void Initialize(UTextureRenderTarget2D* InRenderTarget, const FTransform& InTransform);
 
 	UPROPERTY(BlueprintReadOnly, VisibleAnywhere, Category = Properties)
 	TObjectPtr<UTextureRenderTarget2D> RenderTarget = nullptr;
 };
-
-#if UE_ENABLE_INCLUDE_ORDER_DEPRECATED_IN_5_2
-#include "Engine/TextureRenderTarget2D.h"
-#endif
diff --git a/Engine/Plugins/PCG/Source/PCG/Public/Data/PCGSplineData.h b/Engine/Plugins/PCG/Source/PCG/Public/Data/PCGSplineData.h
--- a/Engine/Plugins/PCG/Source/PCG/Public/Data/PCGSplineData.h
+++ b/Engine/Plugins/PCG/Source/PCG/Public/Data/PCGSplineData.h
@@ -69,12 +69,13 @@ protected:
 };
 
 /* The projection of a spline onto a surface. */
-UCLASS(BlueprintType, ClassGroup=(Procedural))
-class PCG_API UPCGSplineProjectionData : public UPCGProjectionData
+UCLASS(MinimalAPI, BlueprintType, ClassGroup=(Procedural))
+class UPCGSplineProjectionData : public UPCGProjectionData
 {
 	GENERATED_BODY()
+
 public:
-	void Initialize(const UPCGSplineData* InSourceSpline, const UPCGSpatialData* InTargetSurface, const FPCGProjectionParams& InParams);
+	PCG_API void Initialize(const UPCGSplineData* InSourceSpline, const UPCGSpatialData* InTargetSurface, const FPCGProjectionParams& InParams);
 
 	// ~Begin UPCGData interface
 	virtual EPCGDataType GetDataType() const override { return EPCGDataType::Spline; }
diff --git a/Engine/Plugins/PCG/Source/PCG/Public/Data/PCGSplineInteriorSurfaceData.h b/Engine/Plugins/PCG/Source/PCG/Public/Data/PCGSplineInteriorSurfaceData.h
--- a/Engine/Plugins/PCG/Source/PCG/Public/Data/PCGSplineInteriorSurfaceData.h
+++ b/Engine/Plugins/PCG/Source/PCG/Public/Data/PCGSplineInteriorSurfaceData.h
@@ -14,13 +14,13 @@ struct FPCGProjectionParams;
 /**
 * Represents a surface implicitly using the top-down 2D projection of a closed spline.
 */
-UCLASS(BlueprintType, ClassGroup = (Procedural))
-class PCG_API UPCGSplineInteriorSurfaceData : public UPCGSurfaceData
+UCLASS(MinimalAPI, BlueprintType, ClassGroup = (Procedural))
+class UPCGSplineInteriorSurfaceData : public UPCGSurfaceData
 {
 	GENERATED_BODY()
 
 public:
-	void Initialize(const UPCGSplineData* InSplineData);
+	PCG_API void Initialize(const UPCGSplineData* InSplineData);
 
 	// ~Begin UObject interface
 	virtual void PostLoad() override;
diff --git a/Engine/Plugins/PCG/Source/PCG/Public/Data/PCGTextureData.h b/Engine/Plugins/PCG/Source/PCG/Public/Data/PCGTextureData.h
--- a/Engine/Plugins/PCG/Source/PCG/Public/Data/PCGTextureData.h
+++ b/Engine/Plugins/PCG/Source/PCG/Public/Data/PCGTextureData.h
@@ -126,13 +126,13 @@ protected:
 	void CopyBaseTextureData(UPCGBaseTextureData* NewTextureData) const;
 };
 
-UCLASS(BlueprintType, ClassGroup=(Procedural))
-class PCG_API UPCGTextureData : public UPCGBaseTextureData
+UCLASS(MinimalAPI, BlueprintType, ClassGroup=(Procedural))
+class UPCGTextureData : public UPCGBaseTextureData
 {
 	GENERATED_BODY()
 
 public:
-	void Initialize(UTexture* InTexture, uint32 InTextureIndex, const FTransform& InTransform, const TFunction<void()>& PostInitializeCallback, bool bCreateCPUDuplicateEditorOnly = false);
+	PCG_API void Initialize(UTexture* InTexture, uint32 InTextureIndex, const FTransform& InTransform, const TFunction<void()>& PostInitializeCallback, bool bCreateCPUDuplicateEditorOnly = false);
 
 	// ~Begin UPCGData interface
 	virtual EPCGDataType GetDataType() const override { return EPCGDataType::Texture; }
diff --git a/Engine/Plugins/PCG/Source/PCG/Public/Data/PCGWorldData.h b/Engine/Plugins/PCG/Source/PCG/Public/Data/PCGWorldData.h
--- a/Engine/Plugins/PCG/Source/PCG/Public/Data/PCGWorldData.h
+++ b/Engine/Plugins/PCG/Source/PCG/Public/Data/PCGWorldData.h
@@ -81,13 +81,13 @@ struct FPCGWorldVolumetricQueryParams : public FPCGWorldCommonQueryParams
 };
 
 /** Queries volume for presence of world collision or not. Can be used to voxelize environment. */
-UCLASS(BlueprintType, ClassGroup = (Procedural))
-class PCG_API UPCGWorldVolumetricData : public UPCGVolumeData
+UCLASS(MinimalAPI, BlueprintType, ClassGroup = (Procedural))
+class UPCGWorldVolumetricData : public UPCGVolumeData
 {
 	GENERATED_BODY()
 
 public:
-	void Initialize(UWorld* InWorld, const FBox& InBounds = FBox(EForceInit::ForceInit));
+	PCG_API void Initialize(UWorld* InWorld, const FBox& InBounds = FBox(EForceInit::ForceInit));
 
 	//~Begin UPCGSpatialData interface
 	virtual bool IsBounded() const override { return !!Bounds.IsValid; }
@@ -146,13 +146,13 @@ struct FPCGWorldRayHitQueryParams : public FPCGWorldCommonQueryParams
 };
 
 /** Executes collision queries against world collision. */
-UCLASS(BlueprintType, ClassGroup = (Procedural))
-class PCG_API UPCGWorldRayHitData : public UPCGSurfaceData
+UCLASS(MinimalAPI, BlueprintType, ClassGroup = (Procedural))
+class UPCGWorldRayHitData : public UPCGSurfaceData
 {
 	GENERATED_BODY()
 
 public:
-	void Initialize(UWorld* InWorld, const FBox& InBounds = FBox(EForceInit::ForceInit));
+	PCG_API void Initialize(UWorld* InWorld, const FBox& InBounds = FBox(EForceInit::ForceInit));
 
 	// ~Begin UPCGData interface
 	virtual void AddToCrc(FArchiveCrc32& Ar, bool bFullDataCrc) const override;
diff --git a/Engine/Plugins/PCG/Source/PCG/Public/Elements/IO/PCGDataTableElement.h b/Engine/Plugins/PCG/Source/PCG/Public/Elements/IO/PCGDataTableElement.h
--- a/Engine/Plugins/PCG/Source/PCG/Public/Elements/IO/PCGDataTableElement.h
+++ b/Engine/Plugins/PCG/Source/PCG/Public/Elements/IO/PCGDataTableElement.h
@@ -10,8 +10,8 @@ class UDataTable;
 
 #include "PCGDataTableElement.generated.h"
 
-UCLASS(BlueprintType, ClassGroup = (Procedural))
-class PCG_API UPCGLoadDataTableSettings : public UPCGExternalDataSettings
+UCLASS(MinimalAPI, BlueprintType, ClassGroup = (Procedural))
+class UPCGLoadDataTableSettings : public UPCGExternalDataSettings
 {
 	GENERATED_BODY()
 
@@ -23,7 +23,7 @@ public:
 	virtual FText GetNodeTooltipText() const override;
 	virtual bool HasDynamicPins() const override { return true; }
 	virtual void GetStaticTrackedKeys(FPCGSelectionKeyToSettingsMap& OutKeysToSettings, TArray<TObjectPtr<const UPCGGraph>>& OutVisitedGraphs) const override;
-	virtual bool CanDynamicalyTrackKeys() const override { return true; }
+	virtual bool CanDynamicallyTrackKeys() const override { return true; }
 #endif
 
 	virtual EPCGDataType GetCurrentPinTypes(const UPCGPin* InPin) const override;
diff --git a/Engine/Plugins/PCG/Source/PCG/Public/Elements/IO/PCGLoadAssetElement.h b/Engine/Plugins/PCG/Source/PCG/Public/Elements/IO/PCGLoadAssetElement.h
--- a/Engine/Plugins/PCG/Source/PCG/Public/Elements/IO/PCGLoadAssetElement.h
+++ b/Engine/Plugins/PCG/Source/PCG/Public/Elements/IO/PCGLoadAssetElement.h
@@ -34,7 +34,7 @@ public:
 	virtual FLinearColor GetNodeTitleColor() const override { return AssetColor; }
 	virtual EPCGSettingsType GetType() const override { return EPCGSettingsType::InputOutput; }
 	virtual void GetStaticTrackedKeys(FPCGSelectionKeyToSettingsMap& OutKeysToSettings, TArray<TObjectPtr<const UPCGGraph>>& OutVisitedGraphs) const override;
-	virtual bool CanDynamicalyTrackKeys() const override { return true; }
+	virtual bool CanDynamicallyTrackKeys() const override { return true; }
 #endif // WITH_EDITOR
 
 	virtual bool HasDynamicPins() const override { return true; }
diff --git a/Engine/Plugins/PCG/Source/PCG/Public/Elements/Metadata/PCGMetadataBitwiseOpElement.h b/Engine/Plugins/PCG/Source/PCG/Public/Elements/Metadata/PCGMetadataBitwiseOpElement.h
--- a/Engine/Plugins/PCG/Source/PCG/Public/Elements/Metadata/PCGMetadataBitwiseOpElement.h
+++ b/Engine/Plugins/PCG/Source/PCG/Public/Elements/Metadata/PCGMetadataBitwiseOpElement.h
@@ -15,8 +15,8 @@ enum class EPCGMetadataBitwiseOperation : uint16
 	Xor
 };
 
-UCLASS(BlueprintType, ClassGroup = (Procedural))
-class PCG_API UPCGMetadataBitwiseSettings : public UPCGMetadataSettingsBase
+UCLASS(MinimalAPI, BlueprintType, ClassGroup = (Procedural))
+class UPCGMetadataBitwiseSettings : public UPCGMetadataSettingsBase
 {
 	GENERATED_BODY()
 
diff --git a/Engine/Plugins/PCG/Source/PCG/Public/Elements/Metadata/PCGMetadataBooleanOpElement.h b/Engine/Plugins/PCG/Source/PCG/Public/Elements/Metadata/PCGMetadataBooleanOpElement.h
--- a/Engine/Plugins/PCG/Source/PCG/Public/Elements/Metadata/PCGMetadataBooleanOpElement.h
+++ b/Engine/Plugins/PCG/Source/PCG/Public/Elements/Metadata/PCGMetadataBooleanOpElement.h
@@ -15,8 +15,8 @@ enum class EPCGMetadataBooleanOperation : uint16
 	Xor
 };
 
-UCLASS(BlueprintType, ClassGroup = (Procedural))
-class PCG_API UPCGMetadataBooleanSettings : public UPCGMetadataSettingsBase
+UCLASS(MinimalAPI, BlueprintType, ClassGroup = (Procedural))
+class UPCGMetadataBooleanSettings : public UPCGMetadataSettingsBase
 {
 	GENERATED_BODY()
 
diff --git a/Engine/Plugins/PCG/Source/PCG/Public/Elements/Metadata/PCGMetadataBreakTransform.h b/Engine/Plugins/PCG/Source/PCG/Public/Elements/Metadata/PCGMetadataBreakTransform.h
--- a/Engine/Plugins/PCG/Source/PCG/Public/Elements/Metadata/PCGMetadataBreakTransform.h
+++ b/Engine/Plugins/PCG/Source/PCG/Public/Elements/Metadata/PCGMetadataBreakTransform.h
@@ -7,8 +7,8 @@
 
 #include "PCGMetadataBreakTransform.generated.h"
 
-UCLASS(BlueprintType, ClassGroup = (Procedural))
-class PCG_API UPCGMetadataBreakTransformSettings : public UPCGMetadataSettingsBase
+UCLASS(MinimalAPI, BlueprintType, ClassGroup = (Procedural))
+class UPCGMetadataBreakTransformSettings : public UPCGMetadataSettingsBase
 {
 	GENERATED_BODY()
 
diff --git a/Engine/Plugins/PCG/Source/PCG/Public/Elements/Metadata/PCGMetadataBreakVector.h b/Engine/Plugins/PCG/Source/PCG/Public/Elements/Metadata/PCGMetadataBreakVector.h
--- a/Engine/Plugins/PCG/Source/PCG/Public/Elements/Metadata/PCGMetadataBreakVector.h
+++ b/Engine/Plugins/PCG/Source/PCG/Public/Elements/Metadata/PCGMetadataBreakVector.h
@@ -15,8 +15,8 @@ namespace PCGMetadataBreakVectorConstants
 	const FName WLabel = TEXT("W");
 }
 
-UCLASS(BlueprintType, ClassGroup = (Procedural))
-class PCG_API UPCGMetadataBreakVectorSettings : public UPCGMetadataSettingsBase
+UCLASS(MinimalAPI, BlueprintType, ClassGroup = (Procedural))
+class UPCGMetadataBreakVectorSettings : public UPCGMetadataSettingsBase
 {
 	GENERATED_BODY()
 
diff --git a/Engine/Plugins/PCG/Source/PCG/Public/Elements/Metadata/PCGMetadataCompareOpElement.h b/Engine/Plugins/PCG/Source/PCG/Public/Elements/Metadata/PCGMetadataCompareOpElement.h
--- a/Engine/Plugins/PCG/Source/PCG/Public/Elements/Metadata/PCGMetadataCompareOpElement.h
+++ b/Engine/Plugins/PCG/Source/PCG/Public/Elements/Metadata/PCGMetadataCompareOpElement.h
@@ -17,8 +17,8 @@ enum class EPCGMetadataCompareOperation : uint16
 	LessOrEqual
 };
 
-UCLASS(BlueprintType, ClassGroup = (Procedural))
-class PCG_API UPCGMetadataCompareSettings : public UPCGMetadataSettingsBase
+UCLASS(MinimalAPI, BlueprintType, ClassGroup = (Procedural))
+class UPCGMetadataCompareSettings : public UPCGMetadataSettingsBase
 {
 	GENERATED_BODY()
 
diff --git a/Engine/Plugins/PCG/Source/PCG/Public/Elements/Metadata/PCGMetadataElement.h b/Engine/Plugins/PCG/Source/PCG/Public/Elements/Metadata/PCGMetadataElement.h
--- a/Engine/Plugins/PCG/Source/PCG/Public/Elements/Metadata/PCGMetadataElement.h
+++ b/Engine/Plugins/PCG/Source/PCG/Public/Elements/Metadata/PCGMetadataElement.h
@@ -7,16 +7,8 @@
 
 #include "PCGMetadataElement.generated.h"
 
-UENUM()
-enum class UE_DEPRECATED(5.2, "Not used anymore") EPCGMetadataOperationTarget : uint8
-{
-	PropertyToAttribute,
-	AttributeToProperty,
-	AttributeToAttribute,
-};
-
-UCLASS(BlueprintType, ClassGroup = (Procedural))
-class PCG_API UPCGMetadataOperationSettings : public UPCGSettings
+UCLASS(MinimalAPI, BlueprintType, ClassGroup = (Procedural))
+class UPCGMetadataOperationSettings : public UPCGSettings
 {
 	GENERATED_BODY()
 
@@ -56,22 +48,6 @@ public:
 
 	UPROPERTY(BlueprintReadWrite, EditAnywhere, Category = Settings, meta = (PCG_Overridable))
 	bool bCopyAllAttributes = false;
-
-#if WITH_EDITORONLY_DATA
-	UPROPERTY()
-	FName SourceAttribute_DEPRECATED = NAME_None;
-
-	UPROPERTY()
-	EPCGPointProperties PointProperty_DEPRECATED = EPCGPointProperties::Density;
-
-	UPROPERTY()
-	FName DestinationAttribute_DEPRECATED = NAME_None;
-
-PRAGMA_DISABLE_DEPRECATION_WARNINGS
-	UPROPERTY()
-	EPCGMetadataOperationTarget Target_DEPRECATED = EPCGMetadataOperationTarget::PropertyToAttribute;
-PRAGMA_ENABLE_DEPRECATION_WARNINGS
-#endif // WITH_EDITORONLY_DATA
 };
 
 class FPCGMetadataOperationElement : public IPCGElement
diff --git a/Engine/Plugins/PCG/Source/PCG/Public/Elements/Metadata/PCGMetadataMakeRotator.h b/Engine/Plugins/PCG/Source/PCG/Public/Elements/Metadata/PCGMetadataMakeRotator.h
--- a/Engine/Plugins/PCG/Source/PCG/Public/Elements/Metadata/PCGMetadataMakeRotator.h
+++ b/Engine/Plugins/PCG/Source/PCG/Public/Elements/Metadata/PCGMetadataMakeRotator.h
@@ -36,8 +36,8 @@ enum class EPCGMetadataMakeRotatorOp : uint8
 };
 
 /* Create a Rotator from 1, 2 or 3 axis. */
-UCLASS(BlueprintType, ClassGroup = (Procedural))
-class PCG_API UPCGMetadataMakeRotatorSettings : public UPCGMetadataSettingsBase
+UCLASS(MinimalAPI, BlueprintType, ClassGroup = (Procedural))
+class UPCGMetadataMakeRotatorSettings : public UPCGMetadataSettingsBase
 {
 	GENERATED_BODY()
 
diff --git a/Engine/Plugins/PCG/Source/PCG/Public/Elements/Metadata/PCGMetadataMakeTransform.h b/Engine/Plugins/PCG/Source/PCG/Public/Elements/Metadata/PCGMetadataMakeTransform.h
--- a/Engine/Plugins/PCG/Source/PCG/Public/Elements/Metadata/PCGMetadataMakeTransform.h
+++ b/Engine/Plugins/PCG/Source/PCG/Public/Elements/Metadata/PCGMetadataMakeTransform.h
@@ -14,8 +14,8 @@ namespace PCGMetadataTransformConstants
 	const FName Scale = TEXT("Scale");
 }
 
-UCLASS(BlueprintType, ClassGroup = (Procedural))
-class PCG_API UPCGMetadataMakeTransformSettings : public UPCGMetadataSettingsBase
+UCLASS(MinimalAPI, BlueprintType, ClassGroup = (Procedural))
+class UPCGMetadataMakeTransformSettings : public UPCGMetadataSettingsBase
 {
 	GENERATED_BODY()
 
diff --git a/Engine/Plugins/PCG/Source/PCG/Public/Elements/Metadata/PCGMetadataMakeVector.h b/Engine/Plugins/PCG/Source/PCG/Public/Elements/Metadata/PCGMetadataMakeVector.h
--- a/Engine/Plugins/PCG/Source/PCG/Public/Elements/Metadata/PCGMetadataMakeVector.h
+++ b/Engine/Plugins/PCG/Source/PCG/Public/Elements/Metadata/PCGMetadataMakeVector.h
@@ -36,8 +36,8 @@ enum class EPCGMetadataMakeVector4 : uint8
 	Vector3AndValue
 };
 
-UCLASS(BlueprintType, ClassGroup = (Procedural))
-class PCG_API UPCGMetadataMakeVectorSettings : public UPCGMetadataSettingsBase
+UCLASS(MinimalAPI, BlueprintType, ClassGroup = (Procedural))
+class UPCGMetadataMakeVectorSettings : public UPCGMetadataSettingsBase
 {
 	GENERATED_BODY()
 
diff --git a/Engine/Plugins/PCG/Source/PCG/Public/Elements/Metadata/PCGMetadataMathsOpElement.h b/Engine/Plugins/PCG/Source/PCG/Public/Elements/Metadata/PCGMetadataMathsOpElement.h
--- a/Engine/Plugins/PCG/Source/PCG/Public/Elements/Metadata/PCGMetadataMathsOpElement.h
+++ b/Engine/Plugins/PCG/Source/PCG/Public/Elements/Metadata/PCGMetadataMathsOpElement.h
@@ -42,8 +42,8 @@ enum class EPCGMetadataMathsOperation : uint16
 };
 ENUM_CLASS_FLAGS(EPCGMetadataMathsOperation);
 
-UCLASS(BlueprintType, ClassGroup = (Procedural))
-class PCG_API UPCGMetadataMathsSettings : public UPCGMetadataSettingsBase
+UCLASS(MinimalAPI, BlueprintType, ClassGroup = (Procedural))
+class UPCGMetadataMathsSettings : public UPCGMetadataSettingsBase
 {
 	GENERATED_BODY()
 
diff --git a/Engine/Plugins/PCG/Source/PCG/Public/Elements/Metadata/PCGMetadataPartition.h b/Engine/Plugins/PCG/Source/PCG/Public/Elements/Metadata/PCGMetadataPartition.h
--- a/Engine/Plugins/PCG/Source/PCG/Public/Elements/Metadata/PCGMetadataPartition.h
+++ b/Engine/Plugins/PCG/Source/PCG/Public/Elements/Metadata/PCGMetadataPartition.h
@@ -7,8 +7,8 @@
 
 #include "PCGMetadataPartition.generated.h"
 
-UCLASS(BlueprintType, ClassGroup = (Procedural))
-class PCG_API UPCGMetadataPartitionSettings : public UPCGSettings
+UCLASS(MinimalAPI, BlueprintType, ClassGroup = (Procedural))
+class UPCGMetadataPartitionSettings : public UPCGSettings
 {
 	GENERATED_BODY()
 
diff --git a/Engine/Plugins/PCG/Source/PCG/Public/Elements/Metadata/PCGMetadataRenameElement.h b/Engine/Plugins/PCG/Source/PCG/Public/Elements/Metadata/PCGMetadataRenameElement.h
--- a/Engine/Plugins/PCG/Source/PCG/Public/Elements/Metadata/PCGMetadataRenameElement.h
+++ b/Engine/Plugins/PCG/Source/PCG/Public/Elements/Metadata/PCGMetadataRenameElement.h
@@ -6,8 +6,8 @@
 
 #include "PCGMetadataRenameElement.generated.h"
 
-UCLASS(BlueprintType, ClassGroup = (Procedural))
-class PCG_API UPCGMetadataRenameSettings : public UPCGSettings
+UCLASS(MinimalAPI, BlueprintType, ClassGroup = (Procedural))
+class UPCGMetadataRenameSettings : public UPCGSettings
 {
 	GENERATED_BODY()
 
diff --git a/Engine/Plugins/PCG/Source/PCG/Public/Elements/Metadata/PCGMetadataRotatorOpElement.h b/Engine/Plugins/PCG/Source/PCG/Public/Elements/Metadata/PCGMetadataRotatorOpElement.h
--- a/Engine/Plugins/PCG/Source/PCG/Public/Elements/Metadata/PCGMetadataRotatorOpElement.h
+++ b/Engine/Plugins/PCG/Source/PCG/Public/Elements/Metadata/PCGMetadataRotatorOpElement.h
@@ -26,8 +26,8 @@ namespace PCGMetadataRotatorHelpers
 	FRotator RLerp(const FRotator& A, const FRotator& B, double Alpha, bool bShortestPath);
 }
 
-UCLASS(BlueprintType, ClassGroup = (Procedural))
-class PCG_API UPCGMetadataRotatorSettings : public UPCGMetadataSettingsBase
+UCLASS(MinimalAPI, BlueprintType, ClassGroup = (Procedural))
+class UPCGMetadataRotatorSettings : public UPCGMetadataSettingsBase
 {
 	GENERATED_BODY()
 
diff --git a/Engine/Plugins/PCG/Source/PCG/Public/Elements/Metadata/PCGMetadataStringOpElement.h b/Engine/Plugins/PCG/Source/PCG/Public/Elements/Metadata/PCGMetadataStringOpElement.h
--- a/Engine/Plugins/PCG/Source/PCG/Public/Elements/Metadata/PCGMetadataStringOpElement.h
+++ b/Engine/Plugins/PCG/Source/PCG/Public/Elements/Metadata/PCGMetadataStringOpElement.h
@@ -12,8 +12,8 @@ enum class EPCGMetadataStringOperation : uint16
 	Append
 };
 
-UCLASS(BlueprintType, ClassGroup = (Procedural))
-class PCG_API UPCGMetadataStringOpSettings : public UPCGMetadataSettingsBase
+UCLASS(MinimalAPI, BlueprintType, ClassGroup = (Procedural))
+class UPCGMetadataStringOpSettings : public UPCGMetadataSettingsBase
 {
 	GENERATED_BODY()
 
diff --git a/Engine/Plugins/PCG/Source/PCG/Public/Elements/Metadata/PCGMetadataTransformOpElement.h b/Engine/Plugins/PCG/Source/PCG/Public/Elements/Metadata/PCGMetadataTransformOpElement.h
--- a/Engine/Plugins/PCG/Source/PCG/Public/Elements/Metadata/PCGMetadataTransformOpElement.h
+++ b/Engine/Plugins/PCG/Source/PCG/Public/Elements/Metadata/PCGMetadataTransformOpElement.h
@@ -28,8 +28,8 @@ enum class EPCGTransformLerpMode : uint16
 	DualQuatInterp
 };
 
-UCLASS(BlueprintType, ClassGroup = (Procedural))
-class PCG_API UPCGMetadataTransformSettings : public UPCGMetadataSettingsBase
+UCLASS(MinimalAPI, BlueprintType, ClassGroup = (Procedural))
+class UPCGMetadataTransformSettings : public UPCGMetadataSettingsBase
 {
 	GENERATED_BODY()
 
diff --git a/Engine/Plugins/PCG/Source/PCG/Public/Elements/Metadata/PCGMetadataTrigOpElement.h b/Engine/Plugins/PCG/Source/PCG/Public/Elements/Metadata/PCGMetadataTrigOpElement.h
--- a/Engine/Plugins/PCG/Source/PCG/Public/Elements/Metadata/PCGMetadataTrigOpElement.h
+++ b/Engine/Plugins/PCG/Source/PCG/Public/Elements/Metadata/PCGMetadataTrigOpElement.h
@@ -20,8 +20,8 @@ enum class EPCGMetadataTrigOperation : uint16
 	RadToDeg
 };
 
-UCLASS(BlueprintType, ClassGroup = (Procedural))
-class PCG_API UPCGMetadataTrigSettings : public UPCGMetadataSettingsBase
+UCLASS(MinimalAPI, BlueprintType, ClassGroup = (Procedural))
+class UPCGMetadataTrigSettings : public UPCGMetadataSettingsBase
 {
 	GENERATED_BODY()
 
diff --git a/Engine/Plugins/PCG/Source/PCG/Public/Elements/Metadata/PCGMetadataVectorOpElement.h b/Engine/Plugins/PCG/Source/PCG/Public/Elements/Metadata/PCGMetadataVectorOpElement.h
--- a/Engine/Plugins/PCG/Source/PCG/Public/Elements/Metadata/PCGMetadataVectorOpElement.h
+++ b/Engine/Plugins/PCG/Source/PCG/Public/Elements/Metadata/PCGMetadataVectorOpElement.h
@@ -24,8 +24,8 @@ enum class EPCGMetadataVectorOperation : uint16
 	InverseTransformLocation
 };
 
-UCLASS(BlueprintType, ClassGroup = (Procedural))
-class PCG_API UPCGMetadataVectorSettings : public UPCGMetadataSettingsBase
+UCLASS(MinimalAPI, BlueprintType, ClassGroup = (Procedural))
+class UPCGMetadataVectorSettings : public UPCGMetadataSettingsBase
 {
 	GENERATED_BODY()
 
diff --git a/Engine/Plugins/PCG/Source/PCG/Public/Elements/PCGActorSelector.h b/Engine/Plugins/PCG/Source/PCG/Public/Elements/PCGActorSelector.h
--- a/Engine/Plugins/PCG/Source/PCG/Public/Elements/PCGActorSelector.h
+++ b/Engine/Plugins/PCG/Source/PCG/Public/Elements/PCGActorSelector.h
@@ -43,7 +43,7 @@ enum class EPCGActorFilter : uint8
 * Actors have more options for selection with Self/Parent/Root/Original and also EPCGActorSelection::ByTag
 */
 USTRUCT()
-struct FPCGSelectionKey
+struct PCG_API FPCGSelectionKey
 {
 	GENERATED_BODY()
 
@@ -87,11 +87,11 @@ struct FPCGSelectionKey
 	TObjectPtr<const UClass> OptionalExtraDependency = nullptr;
 };
 
-FArchive& operator<<(FArchive& Ar, FPCGSelectionKey& Key);
+PCG_API FArchive& operator<<(FArchive& Ar, FPCGSelectionKey& Key);
 
 /** Helper struct for organizing queries against the world to gather actors. */
 USTRUCT(BlueprintType)
-struct FPCGActorSelectorSettings
+struct PCG_API FPCGActorSelectorSettings
 {
 	GENERATED_BODY()
 
@@ -161,6 +161,6 @@ struct FPCGActorSelectorSettings
 
 namespace PCGActorSelector
 {
-	TArray<AActor*> FindActors(const FPCGActorSelectorSettings& Settings, const UPCGComponent* InComponent, const TFunction<bool(const AActor*)>& BoundsCheck, const TFunction<bool(const AActor*)>& SelfIgnoreCheck);
-	AActor* FindActor(const FPCGActorSelectorSettings& InSettings, UPCGComponent* InComponent, const TFunction<bool(const AActor*)>& BoundsCheck, const TFunction<bool(const AActor*)>& SelfIgnoreCheck);
+	PCG_API TArray<AActor*> FindActors(const FPCGActorSelectorSettings& Settings, const UPCGComponent* InComponent, const TFunction<bool(const AActor*)>& BoundsCheck, const TFunction<bool(const AActor*)>& SelfIgnoreCheck);
+	PCG_API AActor* FindActor(const FPCGActorSelectorSettings& InSettings, UPCGComponent* InComponent, const TFunction<bool(const AActor*)>& BoundsCheck, const TFunction<bool(const AActor*)>& SelfIgnoreCheck);
 }
diff --git a/Engine/Plugins/PCG/Source/PCG/Public/Elements/PCGAttributeCast.h b/Engine/Plugins/PCG/Source/PCG/Public/Elements/PCGAttributeCast.h
--- a/Engine/Plugins/PCG/Source/PCG/Public/Elements/PCGAttributeCast.h
+++ b/Engine/Plugins/PCG/Source/PCG/Public/Elements/PCGAttributeCast.h
@@ -12,8 +12,8 @@
 /**
 * Cast an attribute to another type. Support broadcastable cast (like double -> FVector) and constructible cast (like double -> float)
 */
-UCLASS(BlueprintType, ClassGroup = (Procedural))
-class PCG_API UPCGAttributeCastSettings : public UPCGSettings
+UCLASS(MinimalAPI, BlueprintType, ClassGroup = (Procedural))
+class UPCGAttributeCastSettings : public UPCGSettings
 {
 	GENERATED_BODY()
 
@@ -47,7 +47,7 @@ protected:
 	//~End UPCGSettings interface
 };
 
-class PCG_API FPCGAttributeCastElement : public IPCGElement
+class FPCGAttributeCastElement : public IPCGElement
 {
 protected:
 	virtual bool ExecuteInternal(FPCGContext* Context) const override;
diff --git a/Engine/Plugins/PCG/Source/PCG/Public/Elements/PCGAttributeFilter.h b/Engine/Plugins/PCG/Source/PCG/Public/Elements/PCGAttributeFilter.h
--- a/Engine/Plugins/PCG/Source/PCG/Public/Elements/PCGAttributeFilter.h
+++ b/Engine/Plugins/PCG/Source/PCG/Public/Elements/PCGAttributeFilter.h
@@ -23,12 +23,12 @@ enum class EPCGAttributeFilterOperator : uint8
 };
 
 USTRUCT(BlueprintType)
-struct PCG_API FPCGAttributeFilterThresholdSettings
+struct FPCGAttributeFilterThresholdSettings
 {
 	GENERATED_BODY()
 
 #if WITH_EDITOR
-	void OnPostLoad();
+	PCG_API void OnPostLoad();
 #endif
 
 	/** If the threshold in included or excluded from the range. */
@@ -61,8 +61,8 @@ struct PCG_API FPCGAttributeFilterThresholdSettings
 * - Threshold on attribute by metadata attribute(A.aaa < B.bbb)
 * - Threshold on attribute by property(A.aaa == B.color)
 */
-UCLASS(BlueprintType, ClassGroup = (Procedural))
-class PCG_API UPCGAttributeFilteringSettings : public UPCGSettings
+UCLASS(MinimalAPI, BlueprintType, ClassGroup = (Procedural))
+class UPCGAttributeFilteringSettings : public UPCGSettings
 {
 	GENERATED_BODY()
 
@@ -135,8 +135,8 @@ public:
 * - Threshold on attribute by metadata attribute(A.aaa in [B.bbmin, B.bbmax])
 * - Threshold on attribute by property(A.aaa in [B.position, B.scale])
 */
-UCLASS(BlueprintType, ClassGroup = (Procedural))
-class PCG_API UPCGAttributeFilteringRangeSettings : public UPCGSettings
+UCLASS(MinimalAPI, BlueprintType, ClassGroup = (Procedural))
+class UPCGAttributeFilteringRangeSettings : public UPCGSettings
 {
 	GENERATED_BODY()
 
diff --git a/Engine/Plugins/PCG/Source/PCG/Public/Elements/PCGAttributeGetFromIndexElement.h b/Engine/Plugins/PCG/Source/PCG/Public/Elements/PCGAttributeGetFromIndexElement.h
--- a/Engine/Plugins/PCG/Source/PCG/Public/Elements/PCGAttributeGetFromIndexElement.h
+++ b/Engine/Plugins/PCG/Source/PCG/Public/Elements/PCGAttributeGetFromIndexElement.h
@@ -6,8 +6,8 @@
 
 #include "PCGAttributeGetFromIndexElement.generated.h"
 
-UCLASS(BlueprintType, ClassGroup=(Procedural))
-class PCG_API UPCGAttributeGetFromIndexSettings : public UPCGSettings
+UCLASS(MinimalAPI, BlueprintType, ClassGroup=(Procedural))
+class UPCGAttributeGetFromIndexSettings : public UPCGSettings
 {
 	GENERATED_BODY()
 
diff --git a/Engine/Plugins/PCG/Source/PCG/Public/Elements/PCGAttributeGetFromPointIndexElement.h b/Engine/Plugins/PCG/Source/PCG/Public/Elements/PCGAttributeGetFromPointIndexElement.h
--- a/Engine/Plugins/PCG/Source/PCG/Public/Elements/PCGAttributeGetFromPointIndexElement.h
+++ b/Engine/Plugins/PCG/Source/PCG/Public/Elements/PCGAttributeGetFromPointIndexElement.h
@@ -21,8 +21,8 @@ namespace PCGAttributeGetFromPointIndexConstants
 * 
 * The Index can be overridden by a second Params input.
 */
-UCLASS(BlueprintType, ClassGroup = (Procedural))
-class PCG_API UPCGAttributeGetFromPointIndexSettings : public UPCGSettings
+UCLASS(MinimalAPI, BlueprintType, ClassGroup = (Procedural))
+class UPCGAttributeGetFromPointIndexSettings : public UPCGSettings
 {
 	GENERATED_BODY()
 
diff --git a/Engine/Plugins/PCG/Source/PCG/Public/Elements/PCGAttributeNoise.h b/Engine/Plugins/PCG/Source/PCG/Public/Elements/PCGAttributeNoise.h
--- a/Engine/Plugins/PCG/Source/PCG/Public/Elements/PCGAttributeNoise.h
+++ b/Engine/Plugins/PCG/Source/PCG/Public/Elements/PCGAttributeNoise.h
@@ -27,8 +27,8 @@ enum class EPCGAttributeNoiseMode : uint8
 * Apply some noise to an attribute/property. You can select the mode you want and a noise range.
 * Support all numerical types and vectors/rotators.
 */
-UCLASS(BlueprintType, ClassGroup = (Procedural))
-class PCG_API UPCGAttributeNoiseSettings : public UPCGSettings
+UCLASS(MinimalAPI, BlueprintType, ClassGroup = (Procedural))
+class UPCGAttributeNoiseSettings : public UPCGSettings
 {
 	GENERATED_BODY()
 
diff --git a/Engine/Plugins/PCG/Source/PCG/Public/Elements/PCGAttributeReduceElement.h b/Engine/Plugins/PCG/Source/PCG/Public/Elements/PCGAttributeReduceElement.h
--- a/Engine/Plugins/PCG/Source/PCG/Public/Elements/PCGAttributeReduceElement.h
+++ b/Engine/Plugins/PCG/Source/PCG/Public/Elements/PCGAttributeReduceElement.h
@@ -24,8 +24,8 @@ enum class EPCGAttributeReduceOperation
 * that would be accurate only if the quaternions are close to each other. The accurate version of the average is using eigenvectors/eigenvalues
 * which is way more complicated and computationally expensive. Quaternion will also be normatilzed at the end. Beware if you are using this average.
 */
-UCLASS(BlueprintType, ClassGroup = (Procedural))
-class PCG_API UPCGAttributeReduceSettings : public UPCGSettings
+UCLASS(MinimalAPI, BlueprintType, ClassGroup = (Procedural))
+class UPCGAttributeReduceSettings : public UPCGSettings
 {
 	GENERATED_BODY()
 
diff --git a/Engine/Plugins/PCG/Source/PCG/Public/Elements/PCGAttributeSelectElement.h b/Engine/Plugins/PCG/Source/PCG/Public/Elements/PCGAttributeSelectElement.h
--- a/Engine/Plugins/PCG/Source/PCG/Public/Elements/PCGAttributeSelectElement.h
+++ b/Engine/Plugins/PCG/Source/PCG/Public/Elements/PCGAttributeSelectElement.h
@@ -45,8 +45,8 @@ enum class EPCGAttributeSelectAxis
 *
 * If the OutputAttributeName is None, we will use InputSource.GetName().
 */
-UCLASS(BlueprintType, ClassGroup = (Procedural))
-class PCG_API UPCGAttributeSelectSettings : public UPCGSettings
+UCLASS(MinimalAPI, BlueprintType, ClassGroup = (Procedural))
+class UPCGAttributeSelectSettings : public UPCGSettings
 {
 	GENERATED_BODY()
 
diff --git a/Engine/Plugins/PCG/Source/PCG/Public/Elements/PCGAttributeTransferElement.h b/Engine/Plugins/PCG/Source/PCG/Public/Elements/PCGAttributeTransferElement.h
--- a/Engine/Plugins/PCG/Source/PCG/Public/Elements/PCGAttributeTransferElement.h
+++ b/Engine/Plugins/PCG/Source/PCG/Public/Elements/PCGAttributeTransferElement.h
@@ -25,8 +25,8 @@ namespace PCGAttributeTransferConstants
 * 
 * To do the same but with a Source param data, use CreateAttribute.
 */
-UCLASS(BlueprintType, ClassGroup = (Procedural))
-class PCG_API UPCGAttributeTransferSettings : public UPCGSettings
+UCLASS(MinimalAPI, BlueprintType, ClassGroup = (Procedural))
+class UPCGAttributeTransferSettings : public UPCGSettings
 {
 	GENERATED_BODY()
 
diff --git a/Engine/Plugins/PCG/Source/PCG/Public/Elements/PCGBoundsModifier.h b/Engine/Plugins/PCG/Source/PCG/Public/Elements/PCGBoundsModifier.h
--- a/Engine/Plugins/PCG/Source/PCG/Public/Elements/PCGBoundsModifier.h
+++ b/Engine/Plugins/PCG/Source/PCG/Public/Elements/PCGBoundsModifier.h
@@ -20,8 +20,8 @@ enum class EPCGBoundsModifierMode : uint8
 /**
 * This class controls/sets up a node that modifies the min/max bounds of the input points.
 */
-UCLASS(BlueprintType, ClassGroup = (Procedural))
-class PCG_API UPCGBoundsModifierSettings : public UPCGSettings
+UCLASS(MinimalAPI, BlueprintType, ClassGroup = (Procedural))
+class UPCGBoundsModifierSettings : public UPCGSettings
 {
 	GENERATED_BODY()
 
diff --git a/Engine/Plugins/PCG/Source/PCG/Public/Elements/PCGCollapseElement.h b/Engine/Plugins/PCG/Source/PCG/Public/Elements/PCGCollapseElement.h
--- a/Engine/Plugins/PCG/Source/PCG/Public/Elements/PCGCollapseElement.h
+++ b/Engine/Plugins/PCG/Source/PCG/Public/Elements/PCGCollapseElement.h
@@ -7,8 +7,8 @@
 #include "PCGCollapseElement.generated.h"
 
 /** Convert input to point data, performing sampling with default settings if necessary */
-UCLASS(BlueprintType, ClassGroup = (Procedural))
-class PCG_API UPCGCollapseSettings : public UPCGSettings
+UCLASS(MinimalAPI, BlueprintType, ClassGroup = (Procedural))
+class UPCGCollapseSettings : public UPCGSettings
 {
 	GENERATED_BODY()
 
@@ -31,8 +31,8 @@ protected:
 };
 
 /** Converts attribute sets to point data */
-UCLASS(BlueprintType, ClassGroup = (Procedural))
-class PCG_API UPCGConvertToPointDataSettings : public UPCGCollapseSettings
+UCLASS(MinimalAPI, BlueprintType, ClassGroup = (Procedural))
+class UPCGConvertToPointDataSettings : public UPCGCollapseSettings
 {
 	GENERATED_BODY()
 
diff --git a/Engine/Plugins/PCG/Source/PCG/Public/Elements/PCGConvertToAttributeSet.h b/Engine/Plugins/PCG/Source/PCG/Public/Elements/PCGConvertToAttributeSet.h
--- a/Engine/Plugins/PCG/Source/PCG/Public/Elements/PCGConvertToAttributeSet.h
+++ b/Engine/Plugins/PCG/Source/PCG/Public/Elements/PCGConvertToAttributeSet.h
@@ -7,8 +7,8 @@
 #include "PCGConvertToAttributeSet.generated.h"
 
 /** Converts point data to an attribute set with one entry per point and the same attributes. */
-UCLASS(BlueprintType, ClassGroup = (Procedural))
-class PCG_API UPCGConvertToAttributeSetSettings : public UPCGSettings
+UCLASS(MinimalAPI, BlueprintType, ClassGroup = (Procedural))
+class UPCGConvertToAttributeSetSettings : public UPCGSettings
 {
 	GENERATED_BODY()
 
diff --git a/Engine/Plugins/PCG/Source/PCG/Public/Elements/PCGCopyPoints.h b/Engine/Plugins/PCG/Source/PCG/Public/Elements/PCGCopyPoints.h
--- a/Engine/Plugins/PCG/Source/PCG/Public/Elements/PCGCopyPoints.h
+++ b/Engine/Plugins/PCG/Source/PCG/Public/Elements/PCGCopyPoints.h
@@ -38,8 +38,8 @@ enum class EPCGCopyPointsMetadataInheritanceMode : uint8
 	None        UMETA(Tooltip = "Points will have no metadata.")
 };
 
-UCLASS(BlueprintType, ClassGroup = (Procedural))
-class PCG_API UPCGCopyPointsSettings : public UPCGSettings
+UCLASS(MinimalAPI, BlueprintType, ClassGroup = (Procedural))
+class UPCGCopyPointsSettings : public UPCGSettings
 {
 	GENERATED_BODY()
 
diff --git a/Engine/Plugins/PCG/Source/PCG/Public/Elements/PCGCreateAttribute.h b/Engine/Plugins/PCG/Source/PCG/Public/Elements/PCGCreateAttribute.h
--- a/Engine/Plugins/PCG/Source/PCG/Public/Elements/PCGCreateAttribute.h
+++ b/Engine/Plugins/PCG/Source/PCG/Public/Elements/PCGCreateAttribute.h
@@ -17,8 +17,8 @@ class UPCGMetadata;
 * New attribute can be a constant, hardcoded in the node, or can come from another Attribute Set.
 * Can also add all the attributes coming from the other Attribute Set.
 */
-UCLASS(BlueprintType, ClassGroup = (Procedural))
-class PCG_API UPCGAddAttributeSettings : public UPCGSettings
+UCLASS(MinimalAPI, BlueprintType, ClassGroup = (Procedural))
+class UPCGAddAttributeSettings : public UPCGSettings
 {
 	GENERATED_BODY()
 
@@ -77,8 +77,8 @@ protected:
 };
 
 /* Creates a new Attribute Set. */
-UCLASS(BlueprintType, ClassGroup = (Procedural))
-class PCG_API UPCGCreateAttributeSetSettings : public UPCGSettings
+UCLASS(MinimalAPI, BlueprintType, ClassGroup = (Procedural))
+class UPCGCreateAttributeSetSettings : public UPCGSettings
 {
 	GENERATED_BODY()
 
@@ -129,8 +129,4 @@ class FPCGAddAttributeElement : public IPCGElement
 {
 protected:
 	virtual bool ExecuteInternal(FPCGContext* Context) const override;
-};
-
-#if UE_ENABLE_INCLUDE_ORDER_DEPRECATED_IN_5_2
-#include "PCGPoint.h"
-#endif
+};
\ No newline at end of file
diff --git a/Engine/Plugins/PCG/Source/PCG/Public/Elements/PCGCreateSpline.h b/Engine/Plugins/PCG/Source/PCG/Public/Elements/PCGCreateSpline.h
--- a/Engine/Plugins/PCG/Source/PCG/Public/Elements/PCGCreateSpline.h
+++ b/Engine/Plugins/PCG/Source/PCG/Public/Elements/PCGCreateSpline.h
@@ -19,8 +19,8 @@ enum class EPCGCreateSplineMode : uint8
 };
 
 /** PCG node that creates a spline presentation from the input points data, with optional tangents */
-UCLASS(BlueprintType, ClassGroup = (Procedural))
-class PCG_API UPCGCreateSplineSettings : public UPCGSettings
+UCLASS(MinimalAPI, BlueprintType, ClassGroup = (Procedural))
+class UPCGCreateSplineSettings : public UPCGSettings
 {
 	GENERATED_BODY()
 
diff --git a/Engine/Plugins/PCG/Source/PCG/Public/Elements/PCGCreateSurfaceFromSpline.h b/Engine/Plugins/PCG/Source/PCG/Public/Elements/PCGCreateSurfaceFromSpline.h
--- a/Engine/Plugins/PCG/Source/PCG/Public/Elements/PCGCreateSurfaceFromSpline.h
+++ b/Engine/Plugins/PCG/Source/PCG/Public/Elements/PCGCreateSurfaceFromSpline.h
@@ -12,8 +12,8 @@
  * Note that by default a low resolution polyline is used to represent the spline. If you observe that your interior samples are not sufficiently accurate,
  * you can resample the spline with a SplineSampler at your desired resolution, and convert the points back to a spline with the CreateSpline node.
  */
-UCLASS(BlueprintType, ClassGroup = (Procedural))
-class PCG_API UPCGCreateSurfaceFromSplineSettings : public UPCGSettings
+UCLASS(MinimalAPI, BlueprintType, ClassGroup = (Procedural))
+class UPCGCreateSurfaceFromSplineSettings : public UPCGSettings
 {
 	GENERATED_BODY()
 
diff --git a/Engine/Plugins/PCG/Source/PCG/Public/Elements/PCGCreateTargetActor.h b/Engine/Plugins/PCG/Source/PCG/Public/Elements/PCGCreateTargetActor.h
--- a/Engine/Plugins/PCG/Source/PCG/Public/Elements/PCGCreateTargetActor.h
+++ b/Engine/Plugins/PCG/Source/PCG/Public/Elements/PCGCreateTargetActor.h
@@ -9,8 +9,8 @@
 
 class AActor;
 
-UCLASS(BlueprintType, ClassGroup = (Procedural))
-class PCG_API UPCGCreateTargetActor : public UPCGSettings
+UCLASS(MinimalAPI, BlueprintType, ClassGroup = (Procedural))
+class UPCGCreateTargetActor : public UPCGSettings
 {
 	GENERATED_BODY()
 
@@ -52,16 +52,9 @@ private:
 #endif
 
 public:
-	UPROPERTY(BlueprintReadOnly, EditAnywhere, Category = Settings, meta = (OnlyPlaceable, DisallowCreateNew))
-	TSubclassOf<AActor> TemplateActorClass = nullptr;
-
 	UPROPERTY(EditAnywhere, BlueprintReadWrite, Instanced, Category = Settings, meta = (ShowInnerProperties, EditCondition = "bAllowTemplateActorEditing", EditConditionHides))
 	TObjectPtr<AActor> TemplateActor;
 
-	// TODO: make this InlineEditConditionToggle, not done because property changed event does not propagate correctly so we can't track accurately the need to create the target actor
-	UPROPERTY(EditAnywhere, BlueprintReadWrite, Category = Settings)
-	bool bAllowTemplateActorEditing = false;
-
 	UPROPERTY(BlueprintReadWrite, EditAnywhere, Category = Settings)
 	EPCGAttachOptions AttachOptions = EPCGAttachOptions::Attached; // Note that this is no longer the default value for new nodes, it is now EPCGAttachOptions::InFolder
 
@@ -84,6 +77,21 @@ public:
 	/** Specify a list of functions to be called on the target actor after creation. Functions need to be parameter-less and with "CallInEditor" flag enabled. */
 	UPROPERTY(BlueprintReadWrite, EditAnywhere, Category = Settings)
 	TArray<FName> PostProcessFunctionNames;
+
+	PCG_API void SetTemplateActorClass(const TSubclassOf<AActor>& InTemplateActorClass);
+	PCG_API void SetAllowTemplateActorEditing(bool bInAllowTemplateActorEditing);
+	const TSubclassOf<AActor>& GetTemplateActorClass() const { return TemplateActorClass; }
+	bool GetAllowTemplateActorEditing() const { return bAllowTemplateActorEditing; }
+
+protected:
+	UPROPERTY(BlueprintReadOnly, EditAnywhere, Category = Settings, meta = (OnlyPlaceable, DisallowCreateNew))
+	TSubclassOf<AActor> TemplateActorClass = nullptr;
+
+	// TODO: make this InlineEditConditionToggle, not done because property changed event does not propagate correctly so we can't track accurately the need to create the target actor
+	UPROPERTY(EditAnywhere, BlueprintReadWrite, Category = Settings)
+	bool bAllowTemplateActorEditing = false;
+
+	friend class FPCGCreateTargetActorElement;
 };
 
 class FPCGCreateTargetActorElement : public IPCGElement
diff --git a/Engine/Plugins/PCG/Source/PCG/Public/Elements/PCGCullPointsOutsideActorBounds.h b/Engine/Plugins/PCG/Source/PCG/Public/Elements/PCGCullPointsOutsideActorBounds.h
--- a/Engine/Plugins/PCG/Source/PCG/Public/Elements/PCGCullPointsOutsideActorBounds.h
+++ b/Engine/Plugins/PCG/Source/PCG/Public/Elements/PCGCullPointsOutsideActorBounds.h
@@ -9,8 +9,8 @@
 /**
  * Removes points that lie outside the current actor bounds.
  */
-UCLASS(BlueprintType, ClassGroup = (Procedural))
-class PCG_API UPCGCullPointsOutsideActorBoundsSettings : public UPCGSettings
+UCLASS(MinimalAPI, BlueprintType, ClassGroup = (Procedural))
+class UPCGCullPointsOutsideActorBoundsSettings : public UPCGSettings
 {
 	GENERATED_BODY()
 
diff --git a/Engine/Plugins/PCG/Source/PCG/Public/Elements/PCGDataNum.h b/Engine/Plugins/PCG/Source/PCG/Public/Elements/PCGDataNum.h
--- a/Engine/Plugins/PCG/Source/PCG/Public/Elements/PCGDataNum.h
+++ b/Engine/Plugins/PCG/Source/PCG/Public/Elements/PCGDataNum.h
@@ -7,8 +7,8 @@
 #include "PCGDataNum.generated.h"
 
 /** Counts number of data in a data collection */
-UCLASS(BlueprintType, ClassGroup = (Procedural))
-class PCG_API UPCGDataNumSettings : public UPCGSettings
+UCLASS(MinimalAPI, BlueprintType, ClassGroup = (Procedural))
+class UPCGDataNumSettings : public UPCGSettings
 {
 	GENERATED_BODY()
 
diff --git a/Engine/Plugins/PCG/Source/PCG/Public/Elements/PCGDataTableRowToParamData.h b/Engine/Plugins/PCG/Source/PCG/Public/Elements/PCGDataTableRowToParamData.h
--- a/Engine/Plugins/PCG/Source/PCG/Public/Elements/PCGDataTableRowToParamData.h
+++ b/Engine/Plugins/PCG/Source/PCG/Public/Elements/PCGDataTableRowToParamData.h
@@ -10,7 +10,7 @@
 
 class UDataTable;
 
-UCLASS(BlueprintType, ClassGroup = (Procedural))
+UCLASS(MinimalAPI, BlueprintType, ClassGroup = (Procedural))
 class UPCGDataTableRowToParamDataSettings : public UPCGSettings
 {
 	GENERATED_BODY()
@@ -26,7 +26,7 @@ public:
 	virtual FText GetDefaultNodeTitle() const override { return NSLOCTEXT("PCGDataTableRowToParamDataSettings", "NodeTitle", "Data Table Row To Attribute Set"); }
 	virtual EPCGSettingsType GetType() const override { return EPCGSettingsType::InputOutput; }
 	virtual void GetStaticTrackedKeys(FPCGSelectionKeyToSettingsMap& OutKeysToSettings, TArray<TObjectPtr<const UPCGGraph>>& OutVisitedGraphs) const override;
-	virtual bool CanDynamicalyTrackKeys() const override { return true; }
+	virtual bool CanDynamicallyTrackKeys() const override { return true; }
 	virtual void ApplyDeprecationBeforeUpdatePins(UPCGNode* InOutNode, TArray<TObjectPtr<UPCGPin>>& InputPins, TArray<TObjectPtr<UPCGPin>>& OutputPins);
 #endif
 
diff --git a/Engine/Plugins/PCG/Source/PCG/Public/Elements/PCGDeleteAttributesElement.h b/Engine/Plugins/PCG/Source/PCG/Public/Elements/PCGDeleteAttributesElement.h
--- a/Engine/Plugins/PCG/Source/PCG/Public/Elements/PCGDeleteAttributesElement.h
+++ b/Engine/Plugins/PCG/Source/PCG/Public/Elements/PCGDeleteAttributesElement.h
@@ -19,8 +19,8 @@ enum class EPCGAttributeFilterOperation
 * 
 * The output will be the original data with the updated metadata.
 */
-UCLASS(BlueprintType, ClassGroup = (Procedural))
-class PCG_API UPCGDeleteAttributesSettings : public UPCGSettings
+UCLASS(MinimalAPI, BlueprintType, ClassGroup = (Procedural))
+class UPCGDeleteAttributesSettings : public UPCGSettings
 {
 	GENERATED_BODY()
 
diff --git a/Engine/Plugins/PCG/Source/PCG/Public/Elements/PCGDeleteTags.h b/Engine/Plugins/PCG/Source/PCG/Public/Elements/PCGDeleteTags.h
--- a/Engine/Plugins/PCG/Source/PCG/Public/Elements/PCGDeleteTags.h
+++ b/Engine/Plugins/PCG/Source/PCG/Public/Elements/PCGDeleteTags.h
@@ -14,8 +14,8 @@ enum class EPCGTagFilterOperation
 };
 
 /** Filters the tags on the input data. */
-UCLASS(BlueprintType, ClassGroup = (Procedural))
-class PCG_API UPCGDeleteTagsSettings : public UPCGSettings
+UCLASS(MinimalAPI, BlueprintType, ClassGroup = (Procedural))
+class UPCGDeleteTagsSettings : public UPCGSettings
 {
 	GENERATED_BODY()
 
diff --git a/Engine/Plugins/PCG/Source/PCG/Public/Elements/PCGDensityFilter.h b/Engine/Plugins/PCG/Source/PCG/Public/Elements/PCGDensityFilter.h
--- a/Engine/Plugins/PCG/Source/PCG/Public/Elements/PCGDensityFilter.h
+++ b/Engine/Plugins/PCG/Source/PCG/Public/Elements/PCGDensityFilter.h
@@ -6,8 +6,8 @@
 
 #include "PCGDensityFilter.generated.h"
 
-UCLASS(BlueprintType, ClassGroup = (Procedural))
-class PCG_API UPCGDensityFilterSettings : public UPCGSettings
+UCLASS(MinimalAPI, BlueprintType, ClassGroup = (Procedural))
+class UPCGDensityFilterSettings : public UPCGSettings
 {
 	GENERATED_BODY()
 
diff --git a/Engine/Plugins/PCG/Source/PCG/Public/Elements/PCGDensityRemap.h b/Engine/Plugins/PCG/Source/PCG/Public/Elements/PCGDensityRemap.h
--- a/Engine/Plugins/PCG/Source/PCG/Public/Elements/PCGDensityRemap.h
+++ b/Engine/Plugins/PCG/Source/PCG/Public/Elements/PCGDensityRemap.h
@@ -6,8 +6,8 @@
 
 #include "PCGDensityRemap.generated.h"
 
-UCLASS(BlueprintType, ClassGroup = (Procedural), Hidden, meta = (Deprecated = "5.1", DeprecationMessage = "Use DensityNoise instead."))
-class PCG_API UPCGLinearDensityRemapSettings : public UPCGSettings
+UCLASS(MinimalAPI, BlueprintType, ClassGroup = (Procedural), Hidden, meta = (Deprecated = "5.1", DeprecationMessage = "Use DensityNoise instead."))
+class UPCGLinearDensityRemapSettings : public UPCGSettings
 {
 	GENERATED_BODY()
 
diff --git a/Engine/Plugins/PCG/Source/PCG/Public/Elements/PCGDensityRemapElement.h b/Engine/Plugins/PCG/Source/PCG/Public/Elements/PCGDensityRemapElement.h
--- a/Engine/Plugins/PCG/Source/PCG/Public/Elements/PCGDensityRemapElement.h
+++ b/Engine/Plugins/PCG/Source/PCG/Public/Elements/PCGDensityRemapElement.h
@@ -8,8 +8,8 @@
 
 #include "PCGDensityRemapElement.generated.h"
 
-UCLASS(BlueprintType, ClassGroup = (Procedural))
-class PCG_API UPCGDensityRemapSettings : public UPCGSettings
+UCLASS(MinimalAPI, BlueprintType, ClassGroup = (Procedural))
+class UPCGDensityRemapSettings : public UPCGSettings
 {
 	GENERATED_BODY()
 
diff --git a/Engine/Plugins/PCG/Source/PCG/Public/Elements/PCGDifferenceElement.h b/Engine/Plugins/PCG/Source/PCG/Public/Elements/PCGDifferenceElement.h
--- a/Engine/Plugins/PCG/Source/PCG/Public/Elements/PCGDifferenceElement.h
+++ b/Engine/Plugins/PCG/Source/PCG/Public/Elements/PCGDifferenceElement.h
@@ -13,8 +13,8 @@ namespace PCGDifferenceConstants
 	const FName DifferencesLabel = TEXT("Differences");
 }
 
-UCLASS(BlueprintType, ClassGroup = (Procedural))
-class PCG_API UPCGDifferenceSettings : public UPCGSettings
+UCLASS(MinimalAPI, BlueprintType, ClassGroup = (Procedural))
+class UPCGDifferenceSettings : public UPCGSettings
 {
 	GENERATED_BODY()
 
diff --git a/Engine/Plugins/PCG/Source/PCG/Public/Elements/PCGFilterByAttribute.h b/Engine/Plugins/PCG/Source/PCG/Public/Elements/PCGFilterByAttribute.h
--- a/Engine/Plugins/PCG/Source/PCG/Public/Elements/PCGFilterByAttribute.h
+++ b/Engine/Plugins/PCG/Source/PCG/Public/Elements/PCGFilterByAttribute.h
@@ -7,8 +7,8 @@
 #include "PCGFilterByAttribute.generated.h"
 
 /** Separates data on whether they have a specific metadata attribute (not by value) */
-UCLASS(BlueprintType, ClassGroup = (Procedural))
-class PCG_API UPCGFilterByAttributeSettings : public UPCGFilterDataBaseSettings
+UCLASS(MinimalAPI, BlueprintType, ClassGroup = (Procedural))
+class UPCGFilterByAttributeSettings : public UPCGFilterDataBaseSettings
 {
 	GENERATED_BODY()
 
diff --git a/Engine/Plugins/PCG/Source/PCG/Public/Elements/PCGFilterByIndex.h b/Engine/Plugins/PCG/Source/PCG/Public/Elements/PCGFilterByIndex.h
--- a/Engine/Plugins/PCG/Source/PCG/Public/Elements/PCGFilterByIndex.h
+++ b/Engine/Plugins/PCG/Source/PCG/Public/Elements/PCGFilterByIndex.h
@@ -7,8 +7,8 @@
 #include "PCGFilterByIndex.generated.h"
 
 /** Filters a data collection based on a user defined index range expression. */
-UCLASS(BlueprintType, ClassGroup = (Procedural))
-class PCG_API UPCGFilterByIndexSettings : public UPCGFilterDataBaseSettings
+UCLASS(MinimalAPI, BlueprintType, ClassGroup = (Procedural))
+class UPCGFilterByIndexSettings : public UPCGFilterDataBaseSettings
 {
 	GENERATED_BODY()
 
diff --git a/Engine/Plugins/PCG/Source/PCG/Public/Elements/PCGFilterByTag.h b/Engine/Plugins/PCG/Source/PCG/Public/Elements/PCGFilterByTag.h
--- a/Engine/Plugins/PCG/Source/PCG/Public/Elements/PCGFilterByTag.h
+++ b/Engine/Plugins/PCG/Source/PCG/Public/Elements/PCGFilterByTag.h
@@ -14,8 +14,8 @@ enum class EPCGFilterByTagOperation
 };
 
 /** Filters a data collection based on some tag criterion */
-UCLASS(BlueprintType, ClassGroup = (Procedural))
-class PCG_API UPCGFilterByTagSettings : public UPCGFilterDataBaseSettings
+UCLASS(MinimalAPI, BlueprintType, ClassGroup = (Procedural))
+class UPCGFilterByTagSettings : public UPCGFilterDataBaseSettings
 {
 	GENERATED_BODY()
 
diff --git a/Engine/Plugins/PCG/Source/PCG/Public/Elements/PCGFilterByType.h b/Engine/Plugins/PCG/Source/PCG/Public/Elements/PCGFilterByType.h
--- a/Engine/Plugins/PCG/Source/PCG/Public/Elements/PCGFilterByType.h
+++ b/Engine/Plugins/PCG/Source/PCG/Public/Elements/PCGFilterByType.h
@@ -7,8 +7,8 @@
 #include "PCGFilterByType.generated.h"
 
 /** Filters an input collection based on data type. */
-UCLASS(BlueprintType, ClassGroup = (Procedural))
-class PCG_API UPCGFilterByTypeSettings : public UPCGFilterDataBaseSettings
+UCLASS(MinimalAPI, BlueprintType, ClassGroup = (Procedural))
+class UPCGFilterByTypeSettings : public UPCGFilterDataBaseSettings
 {
 	GENERATED_BODY()
 
diff --git a/Engine/Plugins/PCG/Source/PCG/Public/Elements/PCGGetActorProperty.h b/Engine/Plugins/PCG/Source/PCG/Public/Elements/PCGGetActorProperty.h
--- a/Engine/Plugins/PCG/Source/PCG/Public/Elements/PCGGetActorProperty.h
+++ b/Engine/Plugins/PCG/Source/PCG/Public/Elements/PCGGetActorProperty.h
@@ -13,8 +13,8 @@ class UActorComponent;
 /**
 * Extract a property value from an actor/component into a ParamData.
 */
-UCLASS(BlueprintType, ClassGroup = (Procedural))
-class PCG_API UPCGGetActorPropertySettings : public UPCGSettings
+UCLASS(MinimalAPI, BlueprintType, ClassGroup = (Procedural))
+class UPCGGetActorPropertySettings : public UPCGSettings
 {
 	GENERATED_BODY()
 
@@ -32,7 +32,7 @@ public:
 	virtual FText GetDefaultNodeTitle() const override { return NSLOCTEXT("PCGPropertyToParamDataSettings", "NodeTitle", "Get Actor Property"); }
 	virtual EPCGSettingsType GetType() const override { return EPCGSettingsType::Param; }
 	virtual void GetStaticTrackedKeys(FPCGSelectionKeyToSettingsMap& OutKeysToSettings, TArray<TObjectPtr<const UPCGGraph>>& OutVisitedGraphs) const override;
-	virtual bool CanDynamicalyTrackKeys() const override { return true; }
+	virtual bool CanDynamicallyTrackKeys() const override { return true; }
 #endif
 
 	virtual FString GetAdditionalTitleInformation() const override;
diff --git a/Engine/Plugins/PCG/Source/PCG/Public/Elements/PCGGetBounds.h b/Engine/Plugins/PCG/Source/PCG/Public/Elements/PCGGetBounds.h
--- a/Engine/Plugins/PCG/Source/PCG/Public/Elements/PCGGetBounds.h
+++ b/Engine/Plugins/PCG/Source/PCG/Public/Elements/PCGGetBounds.h
@@ -9,8 +9,8 @@
 /**
  * Computes the bounds of the inputs as attributes.
  */
-UCLASS(BlueprintType, ClassGroup = (Procedural))
-class PCG_API UPCGGetBoundsSettings : public UPCGSettings
+UCLASS(MinimalAPI, BlueprintType, ClassGroup = (Procedural))
+class UPCGGetBoundsSettings : public UPCGSettings
 {
 	GENERATED_BODY()
 
diff --git a/Engine/Plugins/PCG/Source/PCG/Public/Elements/PCGGetLoopIndex.h b/Engine/Plugins/PCG/Source/PCG/Public/Elements/PCGGetLoopIndex.h
--- a/Engine/Plugins/PCG/Source/PCG/Public/Elements/PCGGetLoopIndex.h
+++ b/Engine/Plugins/PCG/Source/PCG/Public/Elements/PCGGetLoopIndex.h
@@ -8,8 +8,8 @@
 #include "PCGGetLoopIndex.generated.h"
 
 /** Returns the current loop iteration index of the "nearest" subgraph in the execution stack */
-UCLASS(BlueprintType, ClassGroup=(Procedural))
-class PCG_API UPCGGetLoopIndexSettings : public UPCGSettings
+UCLASS(MinimalAPI, BlueprintType, ClassGroup=(Procedural))
+class UPCGGetLoopIndexSettings : public UPCGSettings
 {
 	GENERATED_BODY()
 
diff --git a/Engine/Plugins/PCG/Source/PCG/Public/Elements/PCGGetPropertyFromObjectPath.h b/Engine/Plugins/PCG/Source/PCG/Public/Elements/PCGGetPropertyFromObjectPath.h
--- a/Engine/Plugins/PCG/Source/PCG/Public/Elements/PCGGetPropertyFromObjectPath.h
+++ b/Engine/Plugins/PCG/Source/PCG/Public/Elements/PCGGetPropertyFromObjectPath.h
@@ -11,8 +11,8 @@
 /**
 * Extract property from a list of soft object paths.
 */
-UCLASS(BlueprintType, ClassGroup = (Procedural))
-class PCG_API UPCGGetPropertyFromObjectPathSettings : public UPCGSettings
+UCLASS(MinimalAPI, BlueprintType, ClassGroup = (Procedural))
+class UPCGGetPropertyFromObjectPathSettings : public UPCGSettings
 {
 	GENERATED_BODY()
 
@@ -25,7 +25,7 @@ public:
 	virtual bool CanEditChange(const FProperty* InProperty) const override;
 	virtual bool IsPinUsedByNodeExecution(const UPCGPin* InPin) const override;
 	virtual void GetStaticTrackedKeys(FPCGSelectionKeyToSettingsMap& OutKeysToSettings, TArray<TObjectPtr<const UPCGGraph>>& OutVisitedGraphs) const override;
-	virtual bool CanDynamicalyTrackKeys() const override { return true; }
+	virtual bool CanDynamicallyTrackKeys() const override { return true; }
 #endif
 	virtual bool CanCullTaskIfUnwired() const override { return false; }
 	virtual FString GetAdditionalTitleInformation() const override;
diff --git a/Engine/Plugins/PCG/Source/PCG/Public/Elements/PCGHiGenGridSize.h b/Engine/Plugins/PCG/Source/PCG/Public/Elements/PCGHiGenGridSize.h
--- a/Engine/Plugins/PCG/Source/PCG/Public/Elements/PCGHiGenGridSize.h
+++ b/Engine/Plugins/PCG/Source/PCG/Public/Elements/PCGHiGenGridSize.h
@@ -9,8 +9,8 @@
 /**
  * Set the execution grid size for downstream nodes. Enables executing a single graph across a hierarchy of grids.
  */
-UCLASS(BlueprintType, ClassGroup = (Procedural))
-class PCG_API UPCGHiGenGridSizeSettings : public UPCGSettings
+UCLASS(MinimalAPI, BlueprintType, ClassGroup = (Procedural))
+class UPCGHiGenGridSizeSettings : public UPCGSettings
 {
 	GENERATED_BODY()
 
@@ -29,6 +29,8 @@ public:
 	virtual FString GetAdditionalTitleInformation() const override;
 	virtual EPCGDataType GetCurrentPinTypes(const UPCGPin* InPin) const override;
 
+	UPROPERTY(BlueprintReadWrite, EditAnywhere, Category = Settings, meta = (DisplayName = "HiGen Grid Size"))
+	EPCGHiGenGrid HiGenGridSize = EPCGHiGenGrid::Grid256;
 protected:
 #if WITH_EDITOR
 	virtual EPCGChangeType GetChangeTypeForProperty(const FName& InPropertyName) const override;
@@ -37,9 +39,6 @@ protected:
 	virtual TArray<FPCGPinProperties> OutputPinProperties() const override;
 	virtual FPCGElementPtr CreateElement() const override;
 	//~End UPCGSettings interface
-
-	UPROPERTY(BlueprintReadWrite, EditAnywhere, Category = Settings, meta = (DisplayName = "HiGen Grid Size"))
-	EPCGHiGenGrid HiGenGridSize = EPCGHiGenGrid::Grid256;
 };
 
 class FPCGHiGenGridSizeElement : public IPCGElement
diff --git a/Engine/Plugins/PCG/Source/PCG/Public/Elements/PCGIndirectionElement.h b/Engine/Plugins/PCG/Source/PCG/Public/Elements/PCGIndirectionElement.h
--- a/Engine/Plugins/PCG/Source/PCG/Public/Elements/PCGIndirectionElement.h
+++ b/Engine/Plugins/PCG/Source/PCG/Public/Elements/PCGIndirectionElement.h
@@ -17,8 +17,8 @@ enum class EPCGProxyInterfaceMode : uint8
 	BySettings UMETA(Tooltip = "User defined settings will define the pin interface")
 };
 
-UCLASS(BlueprintType, ClassGroup=(Procedural))
-class PCG_API UPCGIndirectionSettings : public UPCGSettings
+UCLASS(MinimalAPI, BlueprintType, ClassGroup=(Procedural))
+class UPCGIndirectionSettings : public UPCGSettings
 {
 	GENERATED_BODY()
 
@@ -30,7 +30,7 @@ public:
 	virtual FText GetNodeTooltipText() const override;
 	virtual EPCGSettingsType GetType() const override { return EPCGSettingsType::Generic; }
 	virtual void GetStaticTrackedKeys(FPCGSelectionKeyToSettingsMap& OutKeysToSettings, TArray<TObjectPtr<const UPCGGraph>>& OutVisitedGraphs) const override;
-	virtual bool CanDynamicalyTrackKeys() const override { return true; }
+	virtual bool CanDynamicallyTrackKeys() const override { return true; }
 #endif // WITH_EDITOR
 	virtual FString GetAdditionalTitleInformation() const override;
 	virtual bool HasFlippedTitleLines() const override { return true; }
diff --git a/Engine/Plugins/PCG/Source/PCG/Public/Elements/PCGInnerIntersectionElement.h b/Engine/Plugins/PCG/Source/PCG/Public/Elements/PCGInnerIntersectionElement.h
--- a/Engine/Plugins/PCG/Source/PCG/Public/Elements/PCGInnerIntersectionElement.h
+++ b/Engine/Plugins/PCG/Source/PCG/Public/Elements/PCGInnerIntersectionElement.h
@@ -7,8 +7,8 @@
 
 #include "PCGInnerIntersectionElement.generated.h"
 
-UCLASS(BlueprintType, ClassGroup = (Procedural))
-class PCG_API UPCGInnerIntersectionSettings : public UPCGSettings
+UCLASS(MinimalAPI, BlueprintType, ClassGroup = (Procedural))
+class UPCGInnerIntersectionSettings : public UPCGSettings
 {
 	GENERATED_BODY()
 
diff --git a/Engine/Plugins/PCG/Source/PCG/Public/Elements/PCGLoopElement.h b/Engine/Plugins/PCG/Source/PCG/Public/Elements/PCGLoopElement.h
--- a/Engine/Plugins/PCG/Source/PCG/Public/Elements/PCGLoopElement.h
+++ b/Engine/Plugins/PCG/Source/PCG/Public/Elements/PCGLoopElement.h
@@ -6,8 +6,8 @@
 
 #include "PCGLoopElement.generated.h"
 
-UCLASS(BlueprintType, ClassGroup=(Procedural))
-class PCG_API UPCGLoopSettings : public UPCGSubgraphSettings
+UCLASS(MinimalAPI, BlueprintType, ClassGroup=(Procedural))
+class UPCGLoopSettings : public UPCGSubgraphSettings
 {
 	GENERATED_BODY()
 
@@ -61,7 +61,7 @@ public:
 	FString FeedbackPins;
 };
 
-class PCG_API FPCGLoopElement : public FPCGSubgraphElement
+class FPCGLoopElement : public FPCGSubgraphElement
 {
 protected:
 	virtual bool ExecuteInternal(FPCGContext* Context) const override;
@@ -69,7 +69,7 @@ protected:
 };
 
 /** This element does the same static input forwarding mechanism as the FPCGInputForwardingElement, but it also will go fetch the previous iteration data, if any. */
-class PCG_API FPCGLoopInputForwardingElement : public FPCGInputForwardingElement
+class FPCGLoopInputForwardingElement : public FPCGInputForwardingElement
 {
 public:
 	FPCGLoopInputForwardingElement(const FPCGDataCollection& StaticInputToForward, FPCGTaskId InPreviousIterationTaskId, const TArray<FName>& InFeedbackPinNames);
diff --git a/Engine/Plugins/PCG/Source/PCG/Public/Elements/PCGMakeConcreteElement.h b/Engine/Plugins/PCG/Source/PCG/Public/Elements/PCGMakeConcreteElement.h
--- a/Engine/Plugins/PCG/Source/PCG/Public/Elements/PCGMakeConcreteElement.h
+++ b/Engine/Plugins/PCG/Source/PCG/Public/Elements/PCGMakeConcreteElement.h
@@ -7,8 +7,8 @@
 #include "PCGMakeConcreteElement.generated.h"
 
 /** Makes data Concrete, collapsing to Point if necessary. Discards non-Spatial data. */
-UCLASS(BlueprintType, ClassGroup = (Procedural))
-class PCG_API UPCGMakeConcreteSettings : public UPCGSettings
+UCLASS(MinimalAPI, BlueprintType, ClassGroup = (Procedural))
+class UPCGMakeConcreteSettings : public UPCGSettings
 {
 	GENERATED_BODY()
 
diff --git a/Engine/Plugins/PCG/Source/PCG/Public/Elements/PCGMatchAndSetAttributes.h b/Engine/Plugins/PCG/Source/PCG/Public/Elements/PCGMatchAndSetAttributes.h
--- a/Engine/Plugins/PCG/Source/PCG/Public/Elements/PCGMatchAndSetAttributes.h
+++ b/Engine/Plugins/PCG/Source/PCG/Public/Elements/PCGMatchAndSetAttributes.h
@@ -39,8 +39,8 @@ enum class EPCGMatchMaxDistanceMode
 * E.g. for a given point, if the point has the same specified attribute as the matching attribute in the attribute set,
 * then we will copy all the other non-selection attributes to the point.
 */
-UCLASS(BlueprintType, ClassGroup = (Procedural))
-class PCG_API UPCGMatchAndSetAttributesSettings : public UPCGSettings
+UCLASS(MinimalAPI, BlueprintType, ClassGroup = (Procedural))
+class UPCGMatchAndSetAttributesSettings : public UPCGSettings
 {
 	GENERATED_BODY()
 
diff --git a/Engine/Plugins/PCG/Source/PCG/Public/Elements/PCGMergeAttributes.h b/Engine/Plugins/PCG/Source/PCG/Public/Elements/PCGMergeAttributes.h
--- a/Engine/Plugins/PCG/Source/PCG/Public/Elements/PCGMergeAttributes.h
+++ b/Engine/Plugins/PCG/Source/PCG/Public/Elements/PCGMergeAttributes.h
@@ -7,8 +7,8 @@
 #include "PCGMergeAttributes.generated.h"
 
 /** Merges multiple attribute sets together into a single attribute set */
-UCLASS(BlueprintType, Classgroup = (Procedural))
-class PCG_API UPCGMergeAttributesSettings : public UPCGSettings
+UCLASS(MinimalAPI, BlueprintType, Classgroup = (Procedural))
+class UPCGMergeAttributesSettings : public UPCGSettings
 {
 	GENERATED_BODY()
 
diff --git a/Engine/Plugins/PCG/Source/PCG/Public/Elements/PCGMergeElement.h b/Engine/Plugins/PCG/Source/PCG/Public/Elements/PCGMergeElement.h
--- a/Engine/Plugins/PCG/Source/PCG/Public/Elements/PCGMergeElement.h
+++ b/Engine/Plugins/PCG/Source/PCG/Public/Elements/PCGMergeElement.h
@@ -7,8 +7,8 @@
 #include "PCGMergeElement.generated.h"
 
 /** Merges multiple data sources (currently only points supported) into a single output. */
-UCLASS(BlueprintType, Classgroup = (Procedural))
-class PCG_API UPCGMergeSettings : public UPCGSettings
+UCLASS(MinimalAPI, BlueprintType, Classgroup = (Procedural))
+class UPCGMergeSettings : public UPCGSettings
 {
 	GENERATED_BODY()
 
diff --git a/Engine/Plugins/PCG/Source/PCG/Public/Elements/PCGOuterIntersectionElement.h b/Engine/Plugins/PCG/Source/PCG/Public/Elements/PCGOuterIntersectionElement.h
--- a/Engine/Plugins/PCG/Source/PCG/Public/Elements/PCGOuterIntersectionElement.h
+++ b/Engine/Plugins/PCG/Source/PCG/Public/Elements/PCGOuterIntersectionElement.h
@@ -16,8 +16,8 @@ namespace PCGIntersectionConstants
 		"Empty data passed along on one of the secondary pins will result in an empty intersection output, unless the 'Ignore Empty Secondary Input' flag is enabled.");
 }
 
-UCLASS(BlueprintType, ClassGroup = (Procedural))
-class PCG_API UPCGOuterIntersectionSettings : public UPCGSettingsWithDynamicInputs
+UCLASS(MinimalAPI, BlueprintType, ClassGroup = (Procedural))
+class UPCGOuterIntersectionSettings : public UPCGSettingsWithDynamicInputs
 {
 	GENERATED_BODY()
 
diff --git a/Engine/Plugins/PCG/Source/PCG/Public/Elements/PCGPointExtentsModifier.h b/Engine/Plugins/PCG/Source/PCG/Public/Elements/PCGPointExtentsModifier.h
--- a/Engine/Plugins/PCG/Source/PCG/Public/Elements/PCGPointExtentsModifier.h
+++ b/Engine/Plugins/PCG/Source/PCG/Public/Elements/PCGPointExtentsModifier.h
@@ -18,8 +18,8 @@ enum class EPCGPointExtentsModifierMode : uint8
 	Multiply
 };
 
-UCLASS(BlueprintType, ClassGroup = (Procedural))
-class PCG_API UPCGPointExtentsModifierSettings : public UPCGSettings
+UCLASS(MinimalAPI, BlueprintType, ClassGroup = (Procedural))
+class UPCGPointExtentsModifierSettings : public UPCGSettings
 {
 	GENERATED_BODY()
 
diff --git a/Engine/Plugins/PCG/Source/PCG/Public/Elements/PCGPointFromMeshElement.h b/Engine/Plugins/PCG/Source/PCG/Public/Elements/PCGPointFromMeshElement.h
--- a/Engine/Plugins/PCG/Source/PCG/Public/Elements/PCGPointFromMeshElement.h
+++ b/Engine/Plugins/PCG/Source/PCG/Public/Elements/PCGPointFromMeshElement.h
@@ -13,8 +13,8 @@
 class UStaticMesh;
 
 // PointFromMesh creates a single point at the origin with an attribute containing a SoftObjectPath to the selected UStaticMesh
-UCLASS(BlueprintType, ClassGroup = (Procedural))
-class PCG_API UPCGPointFromMeshSettings : public UPCGSettings
+UCLASS(MinimalAPI, BlueprintType, ClassGroup = (Procedural))
+class UPCGPointFromMeshSettings : public UPCGSettings
 {
 	GENERATED_BODY()
 
diff --git a/Engine/Plugins/PCG/Source/PCG/Public/Elements/PCGPointMatchAndSet.h b/Engine/Plugins/PCG/Source/PCG/Public/Elements/PCGPointMatchAndSet.h
--- a/Engine/Plugins/PCG/Source/PCG/Public/Elements/PCGPointMatchAndSet.h
+++ b/Engine/Plugins/PCG/Source/PCG/Public/Elements/PCGPointMatchAndSet.h
@@ -14,8 +14,8 @@ class UPCGMatchAndSetBase;
 * on the point data it consumes as input.
 * E.g. for a given point, if it matches with something in the Match & Set object, it will set a value on the point.
 */
-UCLASS(BlueprintType, ClassGroup = (Procedural), meta = (PrioritizeCategories = "Settings"))
-class PCG_API UPCGPointMatchAndSetSettings : public UPCGSettings
+UCLASS(MinimalAPI, BlueprintType, ClassGroup = (Procedural), meta = (PrioritizeCategories = "Settings"))
+class UPCGPointMatchAndSetSettings : public UPCGSettings
 {
 	GENERATED_BODY()
 
@@ -48,7 +48,7 @@ public:
 
 	/** Recreates the match & set instance stored in this settings object if needed. */
 	UFUNCTION(BlueprintCallable, Category = Settings)
-	void SetMatchAndSetType(TSubclassOf<UPCGMatchAndSetBase> InMatchAndSetType);
+	PCG_API void SetMatchAndSetType(TSubclassOf<UPCGMatchAndSetBase> InMatchAndSetType);
 
 protected:
 	void RefreshMatchAndSet();
@@ -70,15 +70,15 @@ public:
 	UPROPERTY(BlueprintReadWrite, EditAnywhere, Category = Settings, meta = (EditCondition = "bSetTargetIsAttribute", HideEditConditionToggle, EditConditionHides))
 	EPCGMetadataTypes SetTargetType = EPCGMetadataTypes::Double;
 
-	/** For string types, the subtype is used to cleanup the UI. */
-	UPROPERTY()
-	EPCGMetadataTypesConstantStructStringMode SetTargetStringMode_DEPRECATED;
-
 private:
 #if WITH_EDITORONLY_DATA
 	// Property used to sidestep edit condition issues - reflects SetTarget.Selection == Attribute
 	UPROPERTY()
 	bool bSetTargetIsAttribute = true;
+
+	/** For string types, the subtype is used to cleanup the UI. */
+	UPROPERTY()
+	EPCGMetadataTypesConstantStructStringMode SetTargetStringMode_DEPRECATED;
 #endif
 };
 
diff --git a/Engine/Plugins/PCG/Source/PCG/Public/Elements/PCGProjectionElement.h b/Engine/Plugins/PCG/Source/PCG/Public/Elements/PCGProjectionElement.h
--- a/Engine/Plugins/PCG/Source/PCG/Public/Elements/PCGProjectionElement.h
+++ b/Engine/Plugins/PCG/Source/PCG/Public/Elements/PCGProjectionElement.h
@@ -12,8 +12,8 @@ namespace PCGProjectionConstants
 	const FName ProjectionTargetLabel = TEXT("Projection Target");
 }
 
-UCLASS(BlueprintType, ClassGroup = (Procedural))
-class PCG_API UPCGProjectionSettings : public UPCGSettings
+UCLASS(MinimalAPI, BlueprintType, ClassGroup = (Procedural))
+class UPCGProjectionSettings : public UPCGSettings
 {
 	GENERATED_BODY()
 
diff --git a/Engine/Plugins/PCG/Source/PCG/Public/Elements/PCGReplaceTags.h b/Engine/Plugins/PCG/Source/PCG/Public/Elements/PCGReplaceTags.h
--- a/Engine/Plugins/PCG/Source/PCG/Public/Elements/PCGReplaceTags.h
+++ b/Engine/Plugins/PCG/Source/PCG/Public/Elements/PCGReplaceTags.h
@@ -7,8 +7,8 @@
 #include "PCGReplaceTags.generated.h"
 
 /** Replaces the tags on the input data. */
-UCLASS(BlueprintType, ClassGroup = (Procedural))
-class PCG_API UPCGReplaceTagsSettings : public UPCGSettings
+UCLASS(MinimalAPI, BlueprintType, ClassGroup = (Procedural))
+class UPCGReplaceTagsSettings : public UPCGSettings
 {
 	GENERATED_BODY()
 
diff --git a/Engine/Plugins/PCG/Source/PCG/Public/Elements/PCGReroute.h b/Engine/Plugins/PCG/Source/PCG/Public/Elements/PCGReroute.h
--- a/Engine/Plugins/PCG/Source/PCG/Public/Elements/PCGReroute.h
+++ b/Engine/Plugins/PCG/Source/PCG/Public/Elements/PCGReroute.h
@@ -14,8 +14,8 @@ namespace PCGNamedRerouteConstants
 	const FName InvisiblePinLabel = TEXT("InvisiblePin");
 }
 
-UCLASS(BlueprintType, ClassGroup = (Procedural))
-class PCG_API UPCGRerouteSettings : public UPCGSettings
+UCLASS(MinimalAPI, BlueprintType, ClassGroup = (Procedural))
+class UPCGRerouteSettings : public UPCGSettings
 {
 	GENERATED_BODY()
 
@@ -44,14 +44,14 @@ protected:
 };
 
 /** Base class for both reroute declaration and usage to share implementation, but also because they use the same visual node representation in the editor. */
-UCLASS(ClassGroup = (Procedural))
-class PCG_API UPCGNamedRerouteBaseSettings : public UPCGRerouteSettings
+UCLASS(MinimalAPI, ClassGroup = (Procedural))
+class UPCGNamedRerouteBaseSettings : public UPCGRerouteSettings
 {
 	GENERATED_BODY()
 };
 
-UCLASS(BlueprintType, ClassGroup = (Procedural))
-class PCG_API UPCGNamedRerouteDeclarationSettings : public UPCGNamedRerouteBaseSettings
+UCLASS(MinimalAPI, BlueprintType, ClassGroup = (Procedural))
+class UPCGNamedRerouteDeclarationSettings : public UPCGNamedRerouteBaseSettings
 {
 	GENERATED_BODY()
 
@@ -71,8 +71,8 @@ public:
 #endif
 };
 
-UCLASS(BlueprintType, ClassGroup = (Procedural))
-class PCG_API UPCGNamedRerouteUsageSettings : public UPCGNamedRerouteBaseSettings
+UCLASS(MinimalAPI, BlueprintType, ClassGroup = (Procedural))
+class UPCGNamedRerouteUsageSettings : public UPCGNamedRerouteBaseSettings
 {
 	GENERATED_BODY()
 
@@ -96,7 +96,7 @@ public:
 	TObjectPtr<const UPCGNamedRerouteDeclarationSettings> Declaration;
 };
 
-class PCG_API FPCGRerouteElement : public IPCGElement
+class FPCGRerouteElement : public IPCGElement
 {
 protected:
 	virtual bool ExecuteInternal(FPCGContext* Context) const override;
diff --git a/Engine/Plugins/PCG/Source/PCG/Public/Elements/PCGSanityCheckPointData.h b/Engine/Plugins/PCG/Source/PCG/Public/Elements/PCGSanityCheckPointData.h
--- a/Engine/Plugins/PCG/Source/PCG/Public/Elements/PCGSanityCheckPointData.h
+++ b/Engine/Plugins/PCG/Source/PCG/Public/Elements/PCGSanityCheckPointData.h
@@ -6,8 +6,8 @@
 
 #include "PCGSanityCheckPointData.generated.h"
 
-UCLASS(BlueprintType, ClassGroup = (Procedural))
-class PCG_API UPCGSanityCheckPointDataSettings : public UPCGSettings
+UCLASS(MinimalAPI, BlueprintType, ClassGroup = (Procedural))
+class UPCGSanityCheckPointDataSettings : public UPCGSettings
 {
 	GENERATED_BODY()
 
diff --git a/Engine/Plugins/PCG/Source/PCG/Public/Elements/PCGSelectPoints.h b/Engine/Plugins/PCG/Source/PCG/Public/Elements/PCGSelectPoints.h
--- a/Engine/Plugins/PCG/Source/PCG/Public/Elements/PCGSelectPoints.h
+++ b/Engine/Plugins/PCG/Source/PCG/Public/Elements/PCGSelectPoints.h
@@ -6,8 +6,8 @@
 
 #include "PCGSelectPoints.generated.h"
 
-UCLASS(BlueprintType, ClassGroup = (Procedural))
-class PCG_API UPCGSelectPointsSettings : public UPCGSettings
+UCLASS(MinimalAPI, BlueprintType, ClassGroup = (Procedural))
+class UPCGSelectPointsSettings : public UPCGSettings
 {
 	GENERATED_BODY()
 
diff --git a/Engine/Plugins/PCG/Source/PCG/Public/Elements/PCGSelfPruning.h b/Engine/Plugins/PCG/Source/PCG/Public/Elements/PCGSelfPruning.h
--- a/Engine/Plugins/PCG/Source/PCG/Public/Elements/PCGSelfPruning.h
+++ b/Engine/Plugins/PCG/Source/PCG/Public/Elements/PCGSelfPruning.h
@@ -46,7 +46,7 @@ namespace PCGSelfPruningElement
 {
 	// all points are in contiguous memory so this will create a bitset to just look up the index to set/test a bit
 	// this should reduce the 'set' cost to O[1]
-	struct PCG_API FPointBitSet
+	struct FPointBitSet
 	{
 		TArray<uint32> Bits;
 		const FPCGPoint* FirstPoint = nullptr;
@@ -77,8 +77,8 @@ namespace PCGSelfPruningElement
 	PCG_API bool ExecuteSlice(FIterationState& InState, const FPCGSelfPruningParameters& InParameters, FPCGContext* InOptionalContext = nullptr);
 }
 
-UCLASS(BlueprintType, ClassGroup=(Procedural))
-class PCG_API UPCGSelfPruningSettings : public UPCGSettings
+UCLASS(MinimalAPI, BlueprintType, ClassGroup=(Procedural))
+class UPCGSelfPruningSettings : public UPCGSettings
 {
 	GENERATED_BODY()
 
diff --git a/Engine/Plugins/PCG/Source/PCG/Public/Elements/PCGSpawnActor.h b/Engine/Plugins/PCG/Source/PCG/Public/Elements/PCGSpawnActor.h
--- a/Engine/Plugins/PCG/Source/PCG/Public/Elements/PCGSpawnActor.h
+++ b/Engine/Plugins/PCG/Source/PCG/Public/Elements/PCGSpawnActor.h
@@ -35,17 +35,14 @@ enum class EPCGSpawnActorGenerationTrigger : uint8
 * as far as data passing and dispatch go.
 * Note that at this point in time, results from the underlying graphs being generated is not propagated back as results of this node.
 */
-UCLASS(BlueprintType, ClassGroup = (Procedural))
-class PCG_API UPCGSpawnActorSettings : public UPCGBaseSubgraphSettings
+UCLASS(MinimalAPI, BlueprintType, ClassGroup = (Procedural))
+class UPCGSpawnActorSettings : public UPCGBaseSubgraphSettings
 {
 	GENERATED_BODY()
 
 public:
 	UPCGSpawnActorSettings(const FObjectInitializer& ObjectInitializer);
 
-	UPROPERTY(BlueprintReadOnly, EditAnywhere, Category = Settings, meta = (OnlyPlaceable, DisallowCreateNew))
-	TSubclassOf<AActor> TemplateActorClass = nullptr;
-
 	/** 
 	 * Can specify a list of functions from the template class to be called on each actor spawned, in order. Need to have "CallInEditor" flag enabled
 	 * and have either no parameters or exactly the parameters PCGPoint and PCGMetadata
@@ -69,9 +66,6 @@ public:
 	UPROPERTY(BlueprintReadWrite, EditAnywhere, Category = Settings, meta = (EditCondition = "Option!=EPCGSpawnActorOption::CollapseActors", EditConditionHides))
 	TArray<FName> TagsToAddOnActors;
 
-	UPROPERTY(BlueprintReadWrite, EditAnywhere, Category = Settings, meta = (EditCondition = "Option != EPCGSpawnActorOption::CollapseActors"))
-	bool bAllowTemplateActorEditing = false;
-
 	UPROPERTY(EditAnywhere, BlueprintReadWrite, Instanced, Category = Settings, meta = (ShowInnerProperties, EditCondition = "bAllowTemplateActorEditing && Option != EPCGSpawnActorOption::CollapseActors", EditConditionHides))
 	TObjectPtr<AActor> TemplateActor;
 
@@ -91,6 +85,12 @@ public:
 	FName SpawnAttribute = NAME_None;
 
 protected:
+	UPROPERTY(BlueprintReadOnly, EditAnywhere, Category = Settings, meta = (OnlyPlaceable, DisallowCreateNew))
+	TSubclassOf<AActor> TemplateActorClass = nullptr;
+
+	UPROPERTY(BlueprintReadWrite, EditAnywhere, Category = Settings, meta = (EditCondition = "Option != EPCGSpawnActorOption::CollapseActors"))
+	bool bAllowTemplateActorEditing = false;
+
 #if WITH_EDITORONLY_DATA
 	UPROPERTY()
 	EPCGSpawnActorGenerationTrigger bGenerationTrigger_DEPRECATED = EPCGSpawnActorGenerationTrigger::Default;
@@ -111,11 +111,15 @@ public:
 	virtual void PostEditUndo() override;
 #endif // WITH_EDITOR
 
-	//~Begin UCPGSettings interface
-	
+	//~Begin UPCGSettings interface
 	virtual UPCGNode* CreateNode() const override;
 
-#if WITH_EDITOR	
+	PCG_API void SetTemplateActorClass(const TSubclassOf<AActor>& InTemplateActorClass);
+	PCG_API void SetAllowTemplateActorEditing(bool bInAllowTemplateActorEditing);
+	const TSubclassOf<AActor>& GetTemplateActorClass() const { return TemplateActorClass; }
+	bool GetAllowTemplateActorEditing() const { return bAllowTemplateActorEditing; }
+
+#if WITH_EDITOR
 	virtual FName GetDefaultNodeName() const override { return FName(TEXT("SpawnActor")); }
 	virtual FText GetDefaultNodeTitle() const override { return NSLOCTEXT("PCGSpawnActorSettings", "NodeTitle", "Spawn Actor"); }
 	virtual EPCGSettingsType GetType() const override;
@@ -148,10 +152,12 @@ protected:
 private:
 	void RefreshTemplateActor();
 #endif
+
+	friend class FPCGSpawnActorElement;
 };
 
 UCLASS(ClassGroup = (Procedural))
-class PCG_API UPCGSpawnActorNode : public UPCGBaseSubgraphNode
+class UPCGSpawnActorNode : public UPCGBaseSubgraphNode
 {
 	GENERATED_BODY()
 public:
@@ -177,8 +183,4 @@ private:
 
 	void CollapseIntoTargetActor(FPCGSubgraphContext* Context, AActor* TargetActor, TSubclassOf<AActor> TemplateActorClass, const UPCGPointData* PointData) const;
 	void SpawnActors(FPCGSubgraphContext* Context, AActor* TargetActor, TSubclassOf<AActor> TemplateActorClass, AActor* TemplateActor, FPCGTaggedData& Output, const UPCGPointData* PointData, UPCGPointData* OutPointData) const;
-};
-
-#if UE_ENABLE_INCLUDE_ORDER_DEPRECATED_IN_5_2
-#include "CoreMinimal.h"
-#endif
+};
\ No newline at end of file
diff --git a/Engine/Plugins/PCG/Source/PCG/Public/Elements/PCGSplineSampler.h b/Engine/Plugins/PCG/Source/PCG/Public/Elements/PCGSplineSampler.h
--- a/Engine/Plugins/PCG/Source/PCG/Public/Elements/PCGSplineSampler.h
+++ b/Engine/Plugins/PCG/Source/PCG/Public/Elements/PCGSplineSampler.h
@@ -53,7 +53,7 @@ enum class EPCGSplineSamplingInteriorOrientation : uint8
 };
 
 USTRUCT(BlueprintType)
-struct PCG_API FPCGSplineSamplerParams
+struct FPCGSplineSamplerParams
 {
 	GENERATED_BODY()
 
@@ -210,8 +210,8 @@ namespace PCGSplineSamplerHelpers
 	FVector::FReal ProjectOntoSplineInteriorSurface(const TArray<FVector>& SplinePoints, const FVector& PointToProject);
 }
 
-UCLASS(BlueprintType, ClassGroup = (Procedural))
-class PCG_API UPCGSplineSamplerSettings : public UPCGSettings
+UCLASS(MinimalAPI, BlueprintType, ClassGroup = (Procedural))
+class UPCGSplineSamplerSettings : public UPCGSettings
 {
 	GENERATED_BODY()
 
diff --git a/Engine/Plugins/PCG/Source/PCG/Public/Elements/PCGStaticMeshSpawner.h b/Engine/Plugins/PCG/Source/PCG/Public/Elements/PCGStaticMeshSpawner.h
--- a/Engine/Plugins/PCG/Source/PCG/Public/Elements/PCGStaticMeshSpawner.h
+++ b/Engine/Plugins/PCG/Source/PCG/Public/Elements/PCGStaticMeshSpawner.h
@@ -20,26 +20,8 @@ struct FPCGStaticMeshSpawnerContext;
 
 class UStaticMesh;
 
-USTRUCT(BlueprintType, meta=(Deprecated = "5.0", DeprecationMessage="Use MeshSelectorWeighted instead."))
-struct PCG_API FPCGStaticMeshSpawnerEntry
-{
-	GENERATED_BODY()
-
-	UPROPERTY(BlueprintReadWrite, EditAnywhere, Category = Settings, meta = (ClampMin = "0"))
-	int Weight = 1;
-
-	UPROPERTY(BlueprintReadWrite, EditAnywhere, Category = Settings)
-	TSoftObjectPtr<UStaticMesh> Mesh;
-
-	UPROPERTY(BlueprintReadWrite, EditAnywhere, Category = Settings)
-	bool bOverrideCollisionProfile = false;
-
-	UPROPERTY(BlueprintReadWrite, EditAnywhere, Category = Settings)
-	FCollisionProfileName CollisionProfile;
-};
-
-UCLASS(BlueprintType, ClassGroup = (Procedural), meta = (PrioritizeCategories = "Settings"))
-class PCG_API UPCGStaticMeshSpawnerSettings : public UPCGSettings
+UCLASS(MinimalAPI, BlueprintType, ClassGroup = (Procedural), meta = (PrioritizeCategories = "Settings"))
+class UPCGStaticMeshSpawnerSettings : public UPCGSettings
 {
 	GENERATED_BODY()
 
@@ -71,10 +53,10 @@ public:
 #endif
 
 	UFUNCTION(BlueprintCallable, Category = Settings)
-	void SetMeshSelectorType(TSubclassOf<UPCGMeshSelectorBase> InMeshSelectorType);
+	PCG_API void SetMeshSelectorType(TSubclassOf<UPCGMeshSelectorBase> InMeshSelectorType);
 
 	UFUNCTION(BlueprintCallable, Category = Settings)
-	void SetInstancePackerType(TSubclassOf<UPCGInstanceDataPackerBase> InInstancePackerType);
+	PCG_API void SetInstancePackerType(TSubclassOf<UPCGInstanceDataPackerBase> InInstancePackerType);
 
 public:
 	/** Defines the method of mesh selection per input data */
@@ -112,9 +94,6 @@ public:
 	UPROPERTY(BlueprintReadWrite, EditAnywhere, Category = Settings)
 	TArray<FName> PostProcessFunctionNames;
 
-	UPROPERTY(meta=(DeprecatedProperty, DeprecationMessage="Use MeshSelectorType and MeshSelectorParameters instead."))
-	TArray<FPCGStaticMeshSpawnerEntry> Meshes_DEPRECATED;
-
 	/** Meshes/Materials will be synchronously loaded before spawning instead of asynchronously. */
 	UPROPERTY(BlueprintReadWrite, EditAnywhere, Category = "Settings|Debug")
 	bool bSynchronousLoad = false;
@@ -141,10 +120,3 @@ protected:
 	virtual void AbortInternal(FPCGContext* Context) const override;
 	void SpawnStaticMeshInstances(FPCGStaticMeshSpawnerContext* Context, const FPCGMeshInstanceList& InstanceList, AActor* TargetActor, const FPCGPackedCustomData& PackedCustomData) const;
 };
-
-
-#if UE_ENABLE_INCLUDE_ORDER_DEPRECATED_IN_5_2
-#include "CoreMinimal.h"
-#include "InstanceDataPackers/PCGInstanceDataPackerBase.h"
-#include "MeshSelectors/PCGMeshSelectorBase.h"
-#endif
diff --git a/Engine/Plugins/PCG/Source/PCG/Public/Elements/PCGSurfaceSampler.h b/Engine/Plugins/PCG/Source/PCG/Public/Elements/PCGSurfaceSampler.h
--- a/Engine/Plugins/PCG/Source/PCG/Public/Elements/PCGSurfaceSampler.h
+++ b/Engine/Plugins/PCG/Source/PCG/Public/Elements/PCGSurfaceSampler.h
@@ -74,8 +74,8 @@ namespace PCGSurfaceSampler
 	bool SampleSurface(FPCGContext* Context, const FSurfaceSamplerParams& Settings, const UPCGSpatialData* InSurface, const UPCGSpatialData* InBoundingShape, UPCGPointData* SampledData, const bool bTimeSlicingIsEnabled = false);
 }
 
-UCLASS(BlueprintType, ClassGroup = (Procedural))
-class PCG_API UPCGSurfaceSamplerSettings : public UPCGSettings
+UCLASS(MinimalAPI, BlueprintType, ClassGroup = (Procedural))
+class UPCGSurfaceSamplerSettings : public UPCGSettings
 {
 	GENERATED_BODY()
 
diff --git a/Engine/Plugins/PCG/Source/PCG/Public/Elements/PCGTextureSampler.h b/Engine/Plugins/PCG/Source/PCG/Public/Elements/PCGTextureSampler.h
--- a/Engine/Plugins/PCG/Source/PCG/Public/Elements/PCGTextureSampler.h
+++ b/Engine/Plugins/PCG/Source/PCG/Public/Elements/PCGTextureSampler.h
@@ -11,8 +11,8 @@
 
 class UTexture;
 
-UCLASS(BlueprintType, ClassGroup = (Procedural))
-class PCG_API UPCGTextureSamplerSettings : public UPCGSettings
+UCLASS(MinimalAPI, BlueprintType, ClassGroup = (Procedural))
+class UPCGTextureSamplerSettings : public UPCGSettings
 {
 	GENERATED_BODY()
 
@@ -31,7 +31,7 @@ public:
 	virtual FText GetNodeTooltipText() const override;
 	virtual EPCGSettingsType GetType() const override { return EPCGSettingsType::Spatial; }
 	virtual void GetStaticTrackedKeys(FPCGSelectionKeyToSettingsMap& OutKeysToSettings, TArray<TObjectPtr<const UPCGGraph>>& OutVisitedGraphs) const override;
-	virtual bool CanDynamicalyTrackKeys() const override { return true; }
+	virtual bool CanDynamicallyTrackKeys() const override { return true; }
 #endif
 
 protected:
@@ -45,6 +45,9 @@ protected:
 #endif
 
 public:
+	PCG_API void SetTexture(TSoftObjectPtr<UTexture> InTexture);
+	TSoftObjectPtr<UTexture> GetTexture() const { return Texture; }
+
 	/** Surface transform */
 	UPROPERTY(BlueprintReadWrite, EditAnywhere, Category = Settings, meta = (PCG_Overridable))
 	FTransform Transform = FTransform::Identity;
@@ -52,10 +55,6 @@ public:
 	UPROPERTY(BlueprintReadWrite, EditAnywhere, Category = Settings, meta = (PCG_Overridable))
 	bool bUseAbsoluteTransform = false;
 
-	/** Texture specific parameters */
-	UPROPERTY(BlueprintReadWrite, EditAnywhere, Category = Settings, meta = (PCG_Overridable))
-	TSoftObjectPtr<UTexture> Texture = nullptr;
-
 	/** Index of texture array slice. Only used when built with editor and if the type of Texture is UTexture2DArray. */
 	UPROPERTY(BlueprintReadWrite, EditAnywhere, Category = Settings, meta = (EditCondition = bDisplayTextureArrayIndex, EditConditionHides, HideEditConditionToggle, ClampMin = '0', PCG_Overridable))
 	int TextureArrayIndex = 0;
@@ -111,14 +110,20 @@ public:
 	bool bSynchronousLoad = false;
 
 protected:
+	/** Texture specific parameters */
+	UPROPERTY(BlueprintReadWrite, EditAnywhere, Category = Settings, meta = (PCG_Overridable))
+	TSoftObjectPtr<UTexture> Texture = nullptr;
+
 #if WITH_EDITORONLY_DATA
 	// Used to hide the 'TextureArrayIndex' property.
 	UPROPERTY(Transient)
 	bool bDisplayTextureArrayIndex = false;
 #endif
+
+	friend class FPCGTextureSamplerElement;
 };
 
-struct PCG_API FPCGTextureSamplerContext : public FPCGContext, public IPCGAsyncLoadingContext
+struct FPCGTextureSamplerContext : public FPCGContext, public IPCGAsyncLoadingContext
 {
 	bool bTextureReadbackDone = false;
 };
diff --git a/Engine/Plugins/PCG/Source/PCG/Public/Elements/PCGTransformPoints.h b/Engine/Plugins/PCG/Source/PCG/Public/Elements/PCGTransformPoints.h
--- a/Engine/Plugins/PCG/Source/PCG/Public/Elements/PCGTransformPoints.h
+++ b/Engine/Plugins/PCG/Source/PCG/Public/Elements/PCGTransformPoints.h
@@ -6,8 +6,8 @@
 
 #include "PCGTransformPoints.generated.h"
 
-UCLASS(BlueprintType, ClassGroup = (Procedural))
-class PCG_API UPCGTransformPointsSettings : public UPCGSettings
+UCLASS(MinimalAPI, BlueprintType, ClassGroup = (Procedural))
+class UPCGTransformPointsSettings : public UPCGSettings
 {
 	GENERATED_BODY()
 
diff --git a/Engine/Plugins/PCG/Source/PCG/Public/Elements/PCGTypedGetter.h b/Engine/Plugins/PCG/Source/PCG/Public/Elements/PCGTypedGetter.h
--- a/Engine/Plugins/PCG/Source/PCG/Public/Elements/PCGTypedGetter.h
+++ b/Engine/Plugins/PCG/Source/PCG/Public/Elements/PCGTypedGetter.h
@@ -8,8 +8,8 @@
 #include "PCGTypedGetter.generated.h"
 
 /** Builds a collection of landscape data from the selected actors. */
-UCLASS(BlueprintType, ClassGroup = (Procedural))
-class PCG_API UPCGGetLandscapeSettings : public UPCGDataFromActorSettings
+UCLASS(MinimalAPI, BlueprintType, ClassGroup = (Procedural))
+class UPCGGetLandscapeSettings : public UPCGDataFromActorSettings
 {
 	GENERATED_BODY()
 
@@ -87,8 +87,8 @@ public:
 };
 
 /** Builds a collection of volume data from the selected actors. */
-UCLASS(BlueprintType, ClassGroup = (Procedural))
-class PCG_API UPCGGetVolumeSettings : public UPCGDataFromActorSettings
+UCLASS(MinimalAPI, BlueprintType, ClassGroup = (Procedural))
+class UPCGGetVolumeSettings : public UPCGDataFromActorSettings
 {
 	GENERATED_BODY()
 
@@ -113,8 +113,8 @@ public:
 };
 
 /** Builds a collection of primitive data from primitive components on the selected actors. */
-UCLASS(BlueprintType, ClassGroup = (Procedural))
-class PCG_API UPCGGetPrimitiveSettings : public UPCGDataFromActorSettings
+UCLASS(MinimalAPI, BlueprintType, ClassGroup = (Procedural))
+class UPCGGetPrimitiveSettings : public UPCGDataFromActorSettings
 {
 	GENERATED_BODY()
 
@@ -144,8 +144,8 @@ public:
  *
  * Note: a component cannot get component data from itself or other components in its execution context, as it could create a circular dependency.
  */
-UCLASS(BlueprintType, ClassGroup = (Procedural))
-class PCG_API UPCGGetPCGComponentSettings : public UPCGDataFromActorSettings
+UCLASS(MinimalAPI, BlueprintType, ClassGroup = (Procedural))
+class UPCGGetPCGComponentSettings : public UPCGDataFromActorSettings
 {
 	GENERATED_BODY()
 
diff --git a/Engine/Plugins/PCG/Source/PCG/Public/Elements/PCGUnionElement.h b/Engine/Plugins/PCG/Source/PCG/Public/Elements/PCGUnionElement.h
--- a/Engine/Plugins/PCG/Source/PCG/Public/Elements/PCGUnionElement.h
+++ b/Engine/Plugins/PCG/Source/PCG/Public/Elements/PCGUnionElement.h
@@ -7,8 +7,8 @@
 
 #include "PCGUnionElement.generated.h"
 
-UCLASS(BlueprintType, ClassGroup = (Procedural))
-class PCG_API UPCGUnionSettings : public UPCGSettings
+UCLASS(MinimalAPI, BlueprintType, ClassGroup = (Procedural))
+class UPCGUnionSettings : public UPCGSettings
 {
 	GENERATED_BODY()
 
diff --git a/Engine/Plugins/PCG/Source/PCG/Public/Elements/PCGUserParameterGet.h b/Engine/Plugins/PCG/Source/PCG/Public/Elements/PCGUserParameterGet.h
--- a/Engine/Plugins/PCG/Source/PCG/Public/Elements/PCGUserParameterGet.h
+++ b/Engine/Plugins/PCG/Source/PCG/Public/Elements/PCGUserParameterGet.h
@@ -13,8 +13,8 @@
 * Getter for user parameters defined in PCGGraph, by the user.
 * Will pick up the value from the graph instance.
 */
-UCLASS(BlueprintType, ClassGroup = (Procedural))
-class PCG_API UPCGUserParameterGetSettings : public UPCGSettings
+UCLASS(MinimalAPI, BlueprintType, ClassGroup = (Procedural))
+class UPCGUserParameterGetSettings : public UPCGSettings
 {
 	GENERATED_BODY()
 
@@ -52,7 +52,7 @@ protected:
 	//~End UPCGSettings interface
 };
 
-class PCG_API FPCGUserParameterGetElement : public IPCGElement
+class FPCGUserParameterGetElement : public IPCGElement
 {
 public:
 	// Graph parameters access is not thread safe
diff --git a/Engine/Plugins/PCG/Source/PCG/Public/Elements/PCGVolumeSampler.h b/Engine/Plugins/PCG/Source/PCG/Public/Elements/PCGVolumeSampler.h
--- a/Engine/Plugins/PCG/Source/PCG/Public/Elements/PCGVolumeSampler.h
+++ b/Engine/Plugins/PCG/Source/PCG/Public/Elements/PCGVolumeSampler.h
@@ -51,8 +51,8 @@ namespace PCGVolumeSampler
 	bool SampleVolume(FPCGContext* Context, const FVolumeSamplerParams& SamplerSettings, const UPCGSpatialData* Volume, const UPCGSpatialData* BoundingShape, UPCGPointData* OutputData, const bool bTimeSlicingIsEnabled = false);
 }
 
-UCLASS(BlueprintType, ClassGroup = (Procedural))
-class PCG_API UPCGVolumeSamplerSettings : public UPCGSettings
+UCLASS(MinimalAPI, BlueprintType, ClassGroup = (Procedural))
+class UPCGVolumeSamplerSettings : public UPCGSettings
 {
 	GENERATED_BODY()
 
diff --git a/Engine/Plugins/PCG/Source/PCG/Public/Elements/PCGWorldQuery.h b/Engine/Plugins/PCG/Source/PCG/Public/Elements/PCGWorldQuery.h
--- a/Engine/Plugins/PCG/Source/PCG/Public/Elements/PCGWorldQuery.h
+++ b/Engine/Plugins/PCG/Source/PCG/Public/Elements/PCGWorldQuery.h
@@ -7,8 +7,8 @@
 
 #include "PCGWorldQuery.generated.h"
 
-UCLASS(BlueprintType, ClassGroup = (Procedural))
-class PCG_API UPCGWorldQuerySettings : public UPCGSettings
+UCLASS(MinimalAPI, BlueprintType, ClassGroup = (Procedural))
+class UPCGWorldQuerySettings : public UPCGSettings
 {
 	GENERATED_BODY()
 
@@ -41,8 +41,8 @@ protected:
 	virtual bool ExecuteInternal(FPCGContext* Context) const override;
 };
 
-UCLASS(BlueprintType, ClassGroup = (Procedural))
-class PCG_API UPCGWorldRayHitSettings : public UPCGSettings
+UCLASS(MinimalAPI, BlueprintType, ClassGroup = (Procedural))
+class UPCGWorldRayHitSettings : public UPCGSettings
 {
 	GENERATED_BODY()
 
diff --git a/Engine/Plugins/PCG/Source/PCG/Public/Grid/PCGComponentOctree.h b/Engine/Plugins/PCG/Source/PCG/Public/Grid/PCGComponentOctree.h
--- a/Engine/Plugins/PCG/Source/PCG/Public/Grid/PCGComponentOctree.h
+++ b/Engine/Plugins/PCG/Source/PCG/Public/Grid/PCGComponentOctree.h
@@ -13,14 +13,14 @@ template <typename ElementType, typename OctreeSemantics> class TOctree2;
 
 class UPCGComponent;
 
-struct PCG_API FPCGComponentOctreeID : public TSharedFromThis<FPCGComponentOctreeID, ESPMode::ThreadSafe>
+struct FPCGComponentOctreeID : public TSharedFromThis<FPCGComponentOctreeID, ESPMode::ThreadSafe>
 {
 	FOctreeElementId2 Id;
 };
 
 using FPCGComponentOctreeIDSharedRef = TSharedRef<struct FPCGComponentOctreeID, ESPMode::ThreadSafe>;
 
-struct PCG_API FPCGComponentRef
+struct FPCGComponentRef
 {
 	FPCGComponentRef(UPCGComponent* InComponent, const FPCGComponentOctreeIDSharedRef& InIdShared);
 
@@ -31,7 +31,7 @@ struct PCG_API FPCGComponentRef
 	FBoxSphereBounds Bounds;
 };
 
-struct PCG_API FPCGComponentRefSemantics
+struct FPCGComponentRefSemantics
 {
 	enum { MaxElementsPerLeaf = 16 };
 	enum { MinInclusiveElementsPerNode = 7 };
@@ -91,8 +91,4 @@ private:
 	FPCGComponentOctree Octree;
 	FPCGComponentToIdMap ComponentToIdMap;
 	mutable FRWLock Lock;
-};
-
-#if UE_ENABLE_INCLUDE_ORDER_DEPRECATED_IN_5_2
-#include "Math/GenericOctree.h"
-#endif
+};
\ No newline at end of file
diff --git a/Engine/Plugins/PCG/Source/PCG/Public/Helpers/PCGHelpers.h b/Engine/Plugins/PCG/Source/PCG/Public/Helpers/PCGHelpers.h
--- a/Engine/Plugins/PCG/Source/PCG/Public/Helpers/PCGHelpers.h
+++ b/Engine/Plugins/PCG/Source/PCG/Public/Helpers/PCGHelpers.h
@@ -86,8 +86,8 @@ namespace PCGHelpers
 };
 
 /** Holds function prototypes used to match against actor function signatures. */
-UCLASS()
-class PCG_API UPCGFunctionPrototypes : public UBlueprintFunctionLibrary
+UCLASS(MinimalAPI)
+class UPCGFunctionPrototypes : public UBlueprintFunctionLibrary
 {
 	GENERATED_BODY()
 
diff --git a/Engine/Plugins/PCG/Source/PCG/Public/Helpers/PCGPropertyHelpers.h b/Engine/Plugins/PCG/Source/PCG/Public/Helpers/PCGPropertyHelpers.h
--- a/Engine/Plugins/PCG/Source/PCG/Public/Helpers/PCGPropertyHelpers.h
+++ b/Engine/Plugins/PCG/Source/PCG/Public/Helpers/PCGPropertyHelpers.h
@@ -70,11 +70,11 @@ namespace PCGPropertyHelpers
 
 	/**
 	* Extract a given property in an Attribute Set.
-	* @param Parameters - Parameters for extaction, cf above.
-	* @param InOptionalContext - Optional context if the extraction is done in a PCG Node, so errors are using the context to log.
+	* @param Parameters - Parameters for extraction, cf above.
+	* @param OptionalContext - Optional context if the extraction is done in a PCG Node, so errors are using the context to log.
 	* @param OptionalObjectTraversed - Optional set to store all objects that we traversed, to be able to react to those objects changes.
 	*/
-	UPCGParamData* ExtractPropertyAsAttributeSet(const FExtractorParameters& Parameters, FPCGContext* OptionalContext = nullptr, TSet<FSoftObjectPath>* OptionalObjectTraversed = nullptr);
+	PCG_API UPCGParamData* ExtractPropertyAsAttributeSet(const FExtractorParameters& Parameters, FPCGContext* OptionalContext = nullptr, TSet<FSoftObjectPath>* OptionalObjectTraversed = nullptr);
 }
 
 //////
diff --git a/Engine/Plugins/PCG/Source/PCG/Public/Helpers/PCGSettingsHelpers.h b/Engine/Plugins/PCG/Source/PCG/Public/Helpers/PCGSettingsHelpers.h
--- a/Engine/Plugins/PCG/Source/PCG/Public/Helpers/PCGSettingsHelpers.h
+++ b/Engine/Plugins/PCG/Source/PCG/Public/Helpers/PCGSettingsHelpers.h
@@ -248,7 +248,7 @@ PRAGMA_ENABLE_DEPRECATION_WARNINGS
 		bool bExtractArrays = false;
 	};
 
-	TArray<FPCGSettingsOverridableParam> GetAllOverridableParams(const UStruct* InClass, const FPCGGetAllOverridableParamsConfig& InConfig);
+	PCG_API TArray<FPCGSettingsOverridableParam> GetAllOverridableParams(const UStruct* InClass, const FPCGGetAllOverridableParamsConfig& InConfig);
 }
 
 // Deprecated macro, not necessary anymore. Cf. GetValue
diff --git a/Engine/Plugins/PCG/Source/PCG/Public/InstanceDataPackers/PCGInstanceDataPackerByAttribute.h b/Engine/Plugins/PCG/Source/PCG/Public/InstanceDataPackers/PCGInstanceDataPackerByAttribute.h
--- a/Engine/Plugins/PCG/Source/PCG/Public/InstanceDataPackers/PCGInstanceDataPackerByAttribute.h
+++ b/Engine/Plugins/PCG/Source/PCG/Public/InstanceDataPackers/PCGInstanceDataPackerByAttribute.h
@@ -6,8 +6,8 @@
 
 #include "PCGInstanceDataPackerByAttribute.generated.h"
 
-UCLASS(BlueprintType, ClassGroup = (Procedural))
-class PCG_API UPCGInstanceDataPackerByAttribute : public UPCGInstanceDataPackerBase 
+UCLASS(MinimalAPI, BlueprintType, ClassGroup = (Procedural))
+class UPCGInstanceDataPackerByAttribute : public UPCGInstanceDataPackerBase 
 {
 	GENERATED_BODY()
 
diff --git a/Engine/Plugins/PCG/Source/PCG/Public/InstanceDataPackers/PCGInstanceDataPackerByRegex.h b/Engine/Plugins/PCG/Source/PCG/Public/InstanceDataPackers/PCGInstanceDataPackerByRegex.h
--- a/Engine/Plugins/PCG/Source/PCG/Public/InstanceDataPackers/PCGInstanceDataPackerByRegex.h
+++ b/Engine/Plugins/PCG/Source/PCG/Public/InstanceDataPackers/PCGInstanceDataPackerByRegex.h
@@ -6,8 +6,8 @@
 
 #include "PCGInstanceDataPackerByRegex.generated.h"
 
-UCLASS(BlueprintType, ClassGroup = (Procedural))
-class PCG_API UPCGInstanceDataPackerByRegex : public UPCGInstanceDataPackerBase 
+UCLASS(MinimalAPI, BlueprintType, ClassGroup = (Procedural))
+class UPCGInstanceDataPackerByRegex : public UPCGInstanceDataPackerBase 
 {
 	GENERATED_BODY()
 
diff --git a/Engine/Plugins/PCG/Source/PCG/Public/MatchAndSet/PCGMatchAndSetByAttribute.h b/Engine/Plugins/PCG/Source/PCG/Public/MatchAndSet/PCGMatchAndSetByAttribute.h
--- a/Engine/Plugins/PCG/Source/PCG/Public/MatchAndSet/PCGMatchAndSetByAttribute.h
+++ b/Engine/Plugins/PCG/Source/PCG/Public/MatchAndSet/PCGMatchAndSetByAttribute.h
@@ -30,8 +30,8 @@ struct PCG_API FPCGMatchAndSetByAttributeEntry
 * This Match & Set object looks up an attribute on a given point,
 * then looks up its entries to find a match; if there is one, then it sets it value. 
 */
-UCLASS(BlueprintType, ClassGroup = (Procedural))
-class PCG_API UPCGMatchAndSetByAttribute : public UPCGMatchAndSetBase
+UCLASS(MinimalAPI, BlueprintType, ClassGroup = (Procedural))
+class UPCGMatchAndSetByAttribute : public UPCGMatchAndSetBase
 {
 	GENERATED_BODY()
 
diff --git a/Engine/Plugins/PCG/Source/PCG/Public/MeshSelectors/PCGMeshSelectorByAttribute.h b/Engine/Plugins/PCG/Source/PCG/Public/MeshSelectors/PCGMeshSelectorByAttribute.h
--- a/Engine/Plugins/PCG/Source/PCG/Public/MeshSelectors/PCGMeshSelectorByAttribute.h
+++ b/Engine/Plugins/PCG/Source/PCG/Public/MeshSelectors/PCGMeshSelectorByAttribute.h
@@ -6,8 +6,8 @@
 
 #include "PCGMeshSelectorByAttribute.generated.h"
 
-UCLASS(BlueprintType, ClassGroup = (Procedural))
-class PCG_API UPCGMeshSelectorByAttribute : public UPCGMeshSelectorBase 
+UCLASS(MinimalAPI, BlueprintType, ClassGroup = (Procedural))
+class UPCGMeshSelectorByAttribute : public UPCGMeshSelectorBase 
 {
 	GENERATED_BODY()
 
diff --git a/Engine/Plugins/PCG/Source/PCG/Public/MeshSelectors/PCGMeshSelectorWeighted.h b/Engine/Plugins/PCG/Source/PCG/Public/MeshSelectors/PCGMeshSelectorWeighted.h
--- a/Engine/Plugins/PCG/Source/PCG/Public/MeshSelectors/PCGMeshSelectorWeighted.h
+++ b/Engine/Plugins/PCG/Source/PCG/Public/MeshSelectors/PCGMeshSelectorWeighted.h
@@ -68,8 +68,8 @@ struct PCG_API FPCGMeshSelectorWeightedEntry
 #endif
 };
 
-UCLASS(BlueprintType, ClassGroup = (Procedural))
-class PCG_API UPCGMeshSelectorWeighted : public UPCGMeshSelectorBase 
+UCLASS(MinimalAPI, BlueprintType, ClassGroup = (Procedural))
+class UPCGMeshSelectorWeighted : public UPCGMeshSelectorBase 
 {
 	GENERATED_BODY()
 
@@ -89,9 +89,8 @@ public:
 	virtual void PostEditChangeProperty(FPropertyChangedEvent& PropertyChangedEvent) override;
 	// ~End UObject interface
 
-protected:
 	/** Refresh MeshEntries display names */
-	void RefreshDisplayNames();
+	PCG_API void RefreshDisplayNames();
 #endif // WITH_EDITOR
 
 public:
diff --git a/Engine/Plugins/PCG/Source/PCG/Public/MeshSelectors/PCGMeshSelectorWeightedByCategory.h b/Engine/Plugins/PCG/Source/PCG/Public/MeshSelectors/PCGMeshSelectorWeightedByCategory.h
--- a/Engine/Plugins/PCG/Source/PCG/Public/MeshSelectors/PCGMeshSelectorWeightedByCategory.h
+++ b/Engine/Plugins/PCG/Source/PCG/Public/MeshSelectors/PCGMeshSelectorWeightedByCategory.h
@@ -32,8 +32,8 @@ struct PCG_API FPCGWeightedByCategoryEntryList
 	TArray<FPCGMeshSelectorWeightedEntry> WeightedMeshEntries;
 };
 
-UCLASS(BlueprintType, ClassGroup = (Procedural))
-class PCG_API UPCGMeshSelectorWeightedByCategory : public UPCGMeshSelectorBase 
+UCLASS(MinimalAPI, BlueprintType, ClassGroup = (Procedural))
+class UPCGMeshSelectorWeightedByCategory : public UPCGMeshSelectorBase 
 {
 	GENERATED_BODY()
 
@@ -53,9 +53,8 @@ public:
 	virtual void PostEditChangeProperty(FPropertyChangedEvent& PropertyChangedEvent) override;
 	// ~End UObject interface
 
-protected:
 	// Refresh MeshEntries display names
-	void RefreshDisplayNames();
+	PCG_API void RefreshDisplayNames();
 #endif // WITH_EDITOR
 
 public:
diff --git a/Engine/Plugins/PCG/Source/PCG/Public/Metadata/PCGAttributePropertySelector.h b/Engine/Plugins/PCG/Source/PCG/Public/Metadata/PCGAttributePropertySelector.h
--- a/Engine/Plugins/PCG/Source/PCG/Public/Metadata/PCGAttributePropertySelector.h
+++ b/Engine/Plugins/PCG/Source/PCG/Public/Metadata/PCGAttributePropertySelector.h
@@ -206,7 +206,7 @@ struct TStructOpsTypeTraits<FPCGAttributePropertyOutputSelector> : public TStruc
 * Helper class to allow the BP to call the custom setters and getters on FPCGAttributePropertySelector.
 */
 UCLASS()
-class PCG_API UPCGAttributePropertySelectorBlueprintHelpers : public UBlueprintFunctionLibrary
+class UPCGAttributePropertySelectorBlueprintHelpers : public UBlueprintFunctionLibrary
 {
 	GENERATED_BODY()
 
diff --git a/Engine/Plugins/PCG/Source/PCG/Public/Metadata/PCGMetadataAccessor.h b/Engine/Plugins/PCG/Source/PCG/Public/Metadata/PCGMetadataAccessor.h
--- a/Engine/Plugins/PCG/Source/PCG/Public/Metadata/PCGMetadataAccessor.h
+++ b/Engine/Plugins/PCG/Source/PCG/Public/Metadata/PCGMetadataAccessor.h
@@ -12,7 +12,7 @@ struct FPCGPoint;
 class UPCGMetadata;
 
 UCLASS()
-class PCG_API UPCGMetadataAccessorHelpers : public UBlueprintFunctionLibrary
+class UPCGMetadataAccessorHelpers : public UBlueprintFunctionLibrary
 {
 	GENERATED_BODY()
 
@@ -106,19 +106,20 @@ public:
 	static bool SetAttributeFromPropertyByMetadataKey(UPARAM(ref) int64& Key, UPCGMetadata* Metadata, FName AttributeName, const UObject* Object, FName PropertyName);
 
 	UFUNCTION(BlueprintCallable, Category = "PCG|Metadata")
-	static bool HasAttributeSetByMetadataKey(int64 Key, const UPCGMetadata* Metadata, FName AttributeName);
+	static PCG_API bool HasAttributeSetByMetadataKey(int64 Key, const UPCGMetadata* Metadata, FName AttributeName);
 
 	/** Point functions */
 	UFUNCTION(BlueprintCallable, Category = "PCG", meta = (ScriptMethod))
-	static void CopyPoint(const FPCGPoint& InPoint, FPCGPoint& OutPoint, bool bCopyMetadata = true, const UPCGMetadata* InMetadata = nullptr, UPCGMetadata* OutMetadata = nullptr);
+	static PCG_API void CopyPoint(const FPCGPoint& InPoint, FPCGPoint& OutPoint, bool bCopyMetadata = true, const UPCGMetadata* InMetadata = nullptr, UPCGMetadata* OutMetadata = nullptr);
 
 	static void InitializeMetadata(FPCGPoint& Point, UPCGMetadata* Metadata);
+
 	/** Assigns a metadata entry but does not copy values if from a non-parented metadata */
-	static void InitializeMetadataWithParent(FPCGPoint& Point, UPCGMetadata* Metadata, const FPCGPoint& ParentPoint, const UPCGMetadata* ParentMetadata);
+	static PCG_API void InitializeMetadataWithParent(FPCGPoint& Point, UPCGMetadata* Metadata, const FPCGPoint& ParentPoint, const UPCGMetadata* ParentMetadata);
 
 	/** Assigns a metadata entry key and will copy attribute values if from an unrelated metadata. Note: a null ParentMetadata assumes this is the same as Metadata */
 	UFUNCTION(BlueprintCallable, Category = "PCG|Metadata", meta = (ScriptMethod))
-	static void InitializeMetadata(UPARAM(ref) FPCGPoint& Point, UPCGMetadata* Metadata, const FPCGPoint& ParentPoint, const UPCGMetadata* ParentMetadata = nullptr);
+	static PCG_API void InitializeMetadata(UPARAM(ref) FPCGPoint& Point, UPCGMetadata* Metadata, const FPCGPoint& ParentPoint, const UPCGMetadata* ParentMetadata = nullptr);
 
 	UFUNCTION(BlueprintCallable, Category = "PCG|Metadata", meta = (ScriptMethod))
 	static int32 GetInteger32Attribute(const FPCGPoint& Point, const UPCGMetadata* Metadata, FName AttributeName);
@@ -211,7 +212,7 @@ public:
 	static void SetSoftClassPathAttribute(UPARAM(ref) FPCGPoint& Point, UPCGMetadata* Metadata, FName AttributeName, const FSoftClassPath& Value);
 
 	UFUNCTION(BlueprintCallable, Category = "PCG|Metadata", meta = (ScriptMethod))
-	static bool HasAttributeSet(const FPCGPoint& Point, const UPCGMetadata* Metadata, FName AttributeName);
+	static PCG_API bool HasAttributeSet(const FPCGPoint& Point, const UPCGMetadata* Metadata, FName AttributeName);
 
 protected:
 	template<typename T>
@@ -219,9 +220,4 @@ protected:
 
 	template<typename T>
 	static void SetAttribute(PCGMetadataEntryKey& Key, UPCGMetadata* Metadata, FName AttributeName, const T& Value);
-};
-
-#if UE_ENABLE_INCLUDE_ORDER_DEPRECATED_IN_5_2
-#include "CoreMinimal.h"
-#include "PCGPoint.h"
-#endif
+};
\ No newline at end of file
diff --git a/Engine/Plugins/PCG/Source/PCG/Public/PCGActorAndComponentMapping.h b/Engine/Plugins/PCG/Source/PCG/Public/PCGActorAndComponentMapping.h
--- a/Engine/Plugins/PCG/Source/PCG/Public/PCGActorAndComponentMapping.h
+++ b/Engine/Plugins/PCG/Source/PCG/Public/PCGActorAndComponentMapping.h
@@ -55,7 +55,7 @@ public:
 
 #if WITH_EDITOR
 	/** Notify that we exited the Landscape edit mode. */
-	PCG_API void NotifyLandscapeEditModeExited();
+	void NotifyLandscapeEditModeExited();
 #endif // WITH_EDITOR
 
 	/** Register a new PCG Component or update it. Returns true if it was added/updated. Thread safe */
diff --git a/Engine/Plugins/PCG/Source/PCG/Public/PCGCommon.h b/Engine/Plugins/PCG/Source/PCG/Public/PCGCommon.h
--- a/Engine/Plugins/PCG/Source/PCG/Public/PCGCommon.h
+++ b/Engine/Plugins/PCG/Source/PCG/Public/PCGCommon.h
@@ -259,7 +259,7 @@ enum class EPCGEditorDirtyMode : uint8
 };
 
 USTRUCT(BlueprintType)
-struct FPCGRuntimeGenerationRadii
+struct PCG_API FPCGRuntimeGenerationRadii
 {
 	GENERATED_BODY()
 
@@ -343,5 +343,5 @@ namespace PCGQualityHelpers
 	const FName PinLabelCinematic = TEXT("Cinematic");
 
 	/** Get the pin label associated with the current 'pcg.Quality' value. If the quality level is invalid, it will return the default pin label. */
-	FName GetQualityPinLabel();
+	PCG_API FName GetQualityPinLabel();
 }
diff --git a/Engine/Plugins/PCG/Source/PCG/Public/PCGCustomVersion.h b/Engine/Plugins/PCG/Source/PCG/Public/PCGCustomVersion.h
--- a/Engine/Plugins/PCG/Source/PCG/Public/PCGCustomVersion.h
+++ b/Engine/Plugins/PCG/Source/PCG/Public/PCGCustomVersion.h
@@ -7,7 +7,7 @@
 struct FGuid;
 
 // Custom serialization version for assets/classes in the PCG plugin
-struct PCG_API FPCGCustomVersion
+struct FPCGCustomVersion
 {
 	enum Type
 	{
diff --git a/Engine/Plugins/PCG/Source/PCG/Public/PCGInputOutputSettings.h b/Engine/Plugins/PCG/Source/PCG/Public/PCGInputOutputSettings.h
--- a/Engine/Plugins/PCG/Source/PCG/Public/PCGInputOutputSettings.h
+++ b/Engine/Plugins/PCG/Source/PCG/Public/PCGInputOutputSettings.h
@@ -16,15 +16,15 @@ namespace PCGInputOutputConstants
 	const FName DefaultNewCustomPinName = TEXT("NewPin");
 }
 
-class PCG_API FPCGInputOutputElement : public IPCGElement
+class FPCGInputOutputElement : public IPCGElement
 {
 protected:
 	virtual bool ExecuteInternal(FPCGContext* Context) const override;
 	virtual bool IsCacheable(const UPCGSettings* InSettings) const override { return false; }
 };
 
-UCLASS(NotBlueprintable, Hidden, ClassGroup = (Procedural))
-class PCG_API UPCGGraphInputOutputSettings : public UPCGSettings
+UCLASS(MinimalAPI, NotBlueprintable, Hidden, ClassGroup = (Procedural))
+class UPCGGraphInputOutputSettings : public UPCGSettings
 {
 	GENERATED_BODY()
 
@@ -55,35 +55,30 @@ protected:
 #endif // WITH_EDITOR
 	// ~End UPCGSettings interface
 
+	UPROPERTY(EditAnywhere, BlueprintReadWrite, Category = Settings)
+	TArray<FPCGPinProperties> Pins;
+
 public:
 	bool IsInput() const { return bIsInput; }
 	void SetInput(bool bInIsInput);
-	
+
 	// Add a new custom pin
 	// Note that you should use the return value of this function, since it can be different from
 	// the one passed as argument. It will change if its label collides with existing pins.
-	[[nodiscard]] const FPCGPinProperties& AddPin(const FPCGPinProperties& NewCustomPinProperties);
+	PCG_API [[nodiscard]] const FPCGPinProperties& AddPin(const FPCGPinProperties& NewCustomPinProperties);
 
 protected:
 	TArray<FPCGPinProperties> DefaultPinProperties(bool bInvisiblePin) const;
 	void FixPinProperties();
 
-protected:
 	UPROPERTY()
 	TSet<FName> PinLabels_DEPRECATED;
 
 	UPROPERTY()
 	TArray<FPCGPinProperties> CustomPins_DEPRECATED;
 
-	UPROPERTY(EditAnywhere, BlueprintReadWrite, Category = Settings)
-	TArray<FPCGPinProperties> Pins;
-
 	UPROPERTY()
 	bool bHasAddedDefaultPin = false;
 
 	bool bIsInput = false;
 };
-
-#if UE_ENABLE_INCLUDE_ORDER_DEPRECATED_IN_5_2
-#include "CoreMinimal.h"
-#endif
diff --git a/Engine/Plugins/PCG/Source/PCG/Public/PCGSettings.h b/Engine/Plugins/PCG/Source/PCG/Public/PCGSettings.h
--- a/Engine/Plugins/PCG/Source/PCG/Public/PCGSettings.h
+++ b/Engine/Plugins/PCG/Source/PCG/Public/PCGSettings.h
@@ -300,7 +300,7 @@ public:
 	virtual void GetStaticTrackedKeys(FPCGSelectionKeyToSettingsMap& OutKeysToSettings, TArray<TObjectPtr<const UPCGGraph>>& OutVisitedGraphs) const {}
 
 	/** Derived classes must implement this to communicate that they might have dynamic dependencies. */
-	virtual bool CanDynamicalyTrackKeys() const { return false; }
+	virtual bool CanDynamicallyTrackKeys() const { return false; }
 
 	/** Override this class to provide an UObject to jump to in case of double click on node
 	 *  ie. returning a blueprint instance will open the given blueprint in its editor.
@@ -537,8 +537,8 @@ public:
 };
 
 /** Trivial / Pass-through settings used for input/output nodes */
-UCLASS(BlueprintType, ClassGroup = (Procedural))
-class PCG_API UPCGTrivialSettings : public UPCGSettings
+UCLASS(MinimalAPI, BlueprintType, ClassGroup = (Procedural))
+class UPCGTrivialSettings : public UPCGSettings
 {
 	GENERATED_BODY()
 	
@@ -550,7 +550,7 @@ protected:
 	virtual FPCGElementPtr CreateElement() const override;
 };
 
-class PCG_API FPCGTrivialElement : public IPCGElement
+class FPCGTrivialElement : public IPCGElement
 {
 protected:
 	virtual bool ExecuteInternal(FPCGContext* Context) const override;
diff --git a/Engine/Plugins/PCG/Source/PCG/Public/PCGSubgraph.h b/Engine/Plugins/PCG/Source/PCG/Public/PCGSubgraph.h
--- a/Engine/Plugins/PCG/Source/PCG/Public/PCGSubgraph.h
+++ b/Engine/Plugins/PCG/Source/PCG/Public/PCGSubgraph.h
@@ -55,7 +55,7 @@ public:
 protected:
 #if WITH_EDITOR
 	virtual void GetStaticTrackedKeys(FPCGSelectionKeyToSettingsMap& OutKeysToSettings, TArray<TObjectPtr<const UPCGGraph>>& OutVisitedGraphs) const override;
-	virtual bool CanDynamicalyTrackKeys() const override { return true; }
+	virtual bool CanDynamicallyTrackKeys() const override { return true; }
 	virtual EPCGChangeType GetChangeTypeForProperty(const FName& InPropertyName) const override;
 #endif
 
@@ -148,8 +148,8 @@ public:
 	virtual TObjectPtr<UPCGGraphInterface> GetSubgraphInterface() const { return nullptr; }
 };
 
-UCLASS(ClassGroup = (Procedural))
-class PCG_API UPCGSubgraphNode : public UPCGBaseSubgraphNode
+UCLASS(MinimalAPI, ClassGroup = (Procedural))
+class UPCGSubgraphNode : public UPCGBaseSubgraphNode
 {
 	GENERATED_BODY()
 
@@ -159,7 +159,7 @@ public:
 	/** ~End UPCGBaseSubgraphNode interface */
 };
 
-struct PCG_API FPCGSubgraphContext : public FPCGContext
+struct FPCGSubgraphContext : public FPCGContext
 {
 	TArray<FPCGTaskId> SubgraphTaskIds;
 	bool bScheduledSubgraph = false;
@@ -175,7 +175,7 @@ protected:
 	virtual void* GetUnsafeExternalContainerForOverridableParam(const FPCGSettingsOverridableParam& InParam) override;
 };
 
-class PCG_API FPCGSubgraphElement : public IPCGElement
+class FPCGSubgraphElement : public IPCGElement
 {
 public:
 	virtual FPCGContext* Initialize(const FPCGDataCollection& InputData, TWeakObjectPtr<UPCGComponent> SourceComponent, const UPCGNode* Node) override;
@@ -188,7 +188,7 @@ protected:
 	void PrepareSubgraphUserParameters(const UPCGSubgraphSettings* Settings, FPCGSubgraphContext* Context, FPCGDataCollection& OutputData) const;
 };
 
-class PCG_API FPCGInputForwardingElement : public IPCGElement
+class FPCGInputForwardingElement : public IPCGElement
 {
 public:
 	FPCGInputForwardingElement(const FPCGDataCollection& InputToForward);
@@ -200,7 +200,3 @@ protected:
 
 	TArray<UPCGData*> RootedData;
 };
-
-#if UE_ENABLE_INCLUDE_ORDER_DEPRECATED_IN_5_2
-#include "CoreMinimal.h"
-#endif
diff --git a/Engine/Plugins/PCG/Source/PCG/Public/RuntimeGen/SchedulingPolicies/PCGSchedulingPolicyDistanceAndDirection.h b/Engine/Plugins/PCG/Source/PCG/Public/RuntimeGen/SchedulingPolicies/PCGSchedulingPolicyDistanceAndDirection.h
--- a/Engine/Plugins/PCG/Source/PCG/Public/RuntimeGen/SchedulingPolicies/PCGSchedulingPolicyDistanceAndDirection.h
+++ b/Engine/Plugins/PCG/Source/PCG/Public/RuntimeGen/SchedulingPolicies/PCGSchedulingPolicyDistanceAndDirection.h
@@ -14,8 +14,8 @@ class IPCGGenSourceBase;
  *
  * Distance and Direction are calculated with respect to the Generation Source.
  */
-UCLASS(BlueprintType, ClassGroup = (Procedural))
-class PCG_API UPCGSchedulingPolicyDistanceAndDirection : public UPCGSchedulingPolicyBase
+UCLASS(MinimalAPI, BlueprintType, ClassGroup = (Procedural))
+class UPCGSchedulingPolicyDistanceAndDirection : public UPCGSchedulingPolicyBase
 {
 	GENERATED_BODY()
 
diff --git a/Engine/Plugins/PCG/Source/PCG/Public/Tests/Determinism/PCGDeterminismNativeTests.h b/Engine/Plugins/PCG/Source/PCG/Public/Tests/Determinism/PCGDeterminismNativeTests.h
--- a/Engine/Plugins/PCG/Source/PCG/Public/Tests/Determinism/PCGDeterminismNativeTests.h
+++ b/Engine/Plugins/PCG/Source/PCG/Public/Tests/Determinism/PCGDeterminismNativeTests.h
@@ -16,9 +16,9 @@ namespace PCGDeterminismTests
 		static void Destroy();
 
 		/** Registers a test function given the StaticClass of an element's settings */
-		static PCG_API void RegisterTestFunction(UClass* SettingsStaticClass, TFunction<bool()> TestFunction);
+		static void RegisterTestFunction(UClass* SettingsStaticClass, TFunction<bool()> TestFunction);
 		/** Deregisters a test function given the StaticClass of an element's settings */
-		static PCG_API void DeregisterTestFunction(const UClass* SettingsStaticClass);
+		static void DeregisterTestFunction(const UClass* SettingsStaticClass);
 		/** Gets a native test function, given an element's UPCGSettings */
 		static PCG_API TFunction<bool()> GetNativeTestFunction(const UPCGSettings* PCGSettings);
 
@@ -33,8 +33,4 @@ namespace PCGDeterminismTests
 		TMap<UClass*, TFunction<bool()>> NativeTestMapping;
 	};
 }
-#endif
-
-#if UE_ENABLE_INCLUDE_ORDER_DEPRECATED_IN_5_2
-#include "CoreMinimal.h"
-#endif
+#endif
\ No newline at end of file
diff --git a/Engine/Plugins/PCG/Source/PCG/Public/Tests/Determinism/PCGDifferenceDeterminismTest.h b/Engine/Plugins/PCG/Source/PCG/Public/Tests/Determinism/PCGDifferenceDeterminismTest.h
--- a/Engine/Plugins/PCG/Source/PCG/Public/Tests/Determinism/PCGDifferenceDeterminismTest.h
+++ b/Engine/Plugins/PCG/Source/PCG/Public/Tests/Determinism/PCGDifferenceDeterminismTest.h
@@ -8,12 +8,8 @@ namespace PCGDeterminismTests
 {
 	namespace DifferenceElement
 	{
-		PCG_API bool RunTestSuite();
+		bool RunTestSuite();
 	}
 }
 
-#endif
-
-#if UE_ENABLE_INCLUDE_ORDER_DEPRECATED_IN_5_2
-#include "CoreMinimal.h"
-#endif
+#endif
\ No newline at end of file
diff --git a/Engine/Plugins/PCG/Source/PCG/Public/Tests/PCGTestsCommon.h b/Engine/Plugins/PCG/Source/PCG/Public/Tests/PCGTestsCommon.h
--- a/Engine/Plugins/PCG/Source/PCG/Public/Tests/PCGTestsCommon.h
+++ b/Engine/Plugins/PCG/Source/PCG/Public/Tests/PCGTestsCommon.h
@@ -47,9 +47,9 @@ namespace PCGTestsCommon
 		FRandomStream RandomStream;
 	};
 
-	PCG_API AActor* CreateTemporaryActor();
+	AActor* CreateTemporaryActor();
 	PCG_API UPCGPointData* CreateEmptyPointData();
-	PCG_API UPCGParamData* CreateEmptyParamData();
+	UPCGParamData* CreateEmptyParamData();
 
 	/** Creates a PointData with a single point at the origin */
 	PCG_API UPCGPointData* CreatePointData();
@@ -59,15 +59,15 @@ namespace PCGTestsCommon
 	/** Creates a PointData with PointCount many points, and randomizes the Transform, Color, and Density */
 	PCG_API UPCGPointData* CreateRandomPointData(int32 PointCount, int32 Seed, bool RandomDensity = false);
 
-	PCG_API UPCGPolyLineData* CreatePolyLineData();
-	PCG_API UPCGSurfaceData* CreateSurfaceData();
-	PCG_API UPCGVolumeData* CreateVolumeData(const FBox& InBounds = FBox::BuildAABB(FVector::ZeroVector, FVector::OneVector * 100));
-	PCG_API UPCGPrimitiveData* CreatePrimitiveData();
+	UPCGPolyLineData* CreatePolyLineData();
+	UPCGSurfaceData* CreateSurfaceData();
+	UPCGVolumeData* CreateVolumeData(const FBox& InBounds = FBox::BuildAABB(FVector::ZeroVector, FVector::OneVector * 100));
+	UPCGPrimitiveData* CreatePrimitiveData();
 
-	PCG_API TArray<FPCGDataCollection> GenerateAllowedData(const FPCGPinProperties& PinProperties);
+	TArray<FPCGDataCollection> GenerateAllowedData(const FPCGPinProperties& PinProperties);
 
 	/** Validates that two Spatial Points are identical */
-	PCG_API bool PointsAreIdentical(const FPCGPoint& FirstPoint, const FPCGPoint& SecondPoint);
+	bool PointsAreIdentical(const FPCGPoint& FirstPoint, const FPCGPoint& SecondPoint);
 
 	/** Generates settings based upon a UPCGSettings subclass */
 	template<typename SettingsType>
@@ -91,7 +91,7 @@ namespace PCGTestsCommon
 	}
 }
 
-class PCG_API FPCGTestBaseClass : public FAutomationTestBase
+class FPCGTestBaseClass : public FAutomationTestBase
 {
 public:
 	using FAutomationTestBase::FAutomationTestBase;
@@ -102,7 +102,3 @@ protected:
 	/** Generates all valid input combinations */
 	bool SmokeTestAnyValidInput(UPCGSettings* InSettings, TFunction<bool(const FPCGDataCollection&, const FPCGDataCollection&)> ValidationFn = TFunction<bool(const FPCGDataCollection&, const FPCGDataCollection&)>());
 };
-
-#if UE_ENABLE_INCLUDE_ORDER_DEPRECATED_IN_5_2
-#include "Tests/AutomationCommon.h"
-#endif
diff --git a/Engine/Plugins/PCGExternalDataInterop/Source/PCGExternalDataInterop/Public/Elements/PCGLoadAlembicElement.h b/Engine/Plugins/PCGExternalDataInterop/Source/PCGExternalDataInterop/Public/Elements/PCGLoadAlembicElement.h
--- a/Engine/Plugins/PCGExternalDataInterop/Source/PCGExternalDataInterop/Public/Elements/PCGLoadAlembicElement.h
+++ b/Engine/Plugins/PCGExternalDataInterop/Source/PCGExternalDataInterop/Public/Elements/PCGLoadAlembicElement.h
@@ -15,8 +15,8 @@ enum class EPCGLoadAlembicStandardSetup : uint8
 	CitySample UMETA(Tooltip="Uses the same setup as in the City Sample demo: right handed Y-up and the orient and scale mapping to the rotation and scale, respectively")
 };
 
-UCLASS(BlueprintType, ClassGroup = (Procedural))
-class PCGEXTERNALDATAINTEROP_API UPCGLoadAlembicSettings : public UPCGExternalDataSettings
+UCLASS(MinimalAPI, BlueprintType, ClassGroup = (Procedural))
+class UPCGLoadAlembicSettings : public UPCGExternalDataSettings
 {
 	GENERATED_BODY()
 
@@ -56,9 +56,9 @@ public:
 	bool bConversionFlipHandedness = false;
 
 	UFUNCTION(BlueprintCallable, Category = "PCG|AlembicImport")
-	void SetupFromStandard(EPCGLoadAlembicStandardSetup InSetup);
+	PCGEXTERNALDATAINTEROP_API void SetupFromStandard(EPCGLoadAlembicStandardSetup InSetup);
 
-	static void SetupFromStandard(EPCGLoadAlembicStandardSetup InSetup, FVector& InConversionScale, FVector& InConversionRotation, bool& bInConversionFlipHandedness, TMap<FString, FPCGAttributePropertyInputSelector>& InAttributeMapping);
+	static PCGEXTERNALDATAINTEROP_API void SetupFromStandard(EPCGLoadAlembicStandardSetup InSetup, FVector& InConversionScale, FVector& InConversionRotation, bool& bInConversionFlipHandedness, TMap<FString, FPCGAttributePropertyInputSelector>& InAttributeMapping);
 
 #if WITH_EDITORONLY_DATA
 	UPROPERTY(EditAnywhere, Category = "Alembic", meta = (DisplayName="Setup from standard"))
