[PCG] Fix Mac compile error in PCGInputOutputSettings

#rnx

diff --git a/Engine/Plugins/PCG/Source/PCG/Public/PCGInputOutputSettings.h b/Engine/Plugins/PCG/Source/PCG/Public/PCGInputOutputSettings.h
--- a/Engine/Plugins/PCG/Source/PCG/Public/PCGInputOutputSettings.h
+++ b/Engine/Plugins/PCG/Source/PCG/Public/PCGInputOutputSettings.h
@@ -65,7 +65,7 @@ public:
 	// Add a new custom pin
 	// Note that you should use the return value of this function, since it can be different from
 	// the one passed as argument. It will change if its label collides with existing pins.
-	PCG_API [[nodiscard]] const FPCGPinProperties& AddPin(const FPCGPinProperties& NewCustomPinProperties);
+	[[nodiscard]] PCG_API const FPCGPinProperties& AddPin(const FPCGPinProperties& NewCustomPinProperties);
 
 protected:
 	TArray<FPCGPinProperties> DefaultPinProperties(bool bInvisiblePin) const;
