[PCG] Fix for a subsequent build problem

#rnx

diff --git a/Engine/Plugins/PCG/Source/PCG/Private/PCGSubsystem.cpp b/Engine/Plugins/PCG/Source/PCG/Private/PCGSubsystem.cpp
--- a/Engine/Plugins/PCG/Source/PCG/Private/PCGSubsystem.cpp
+++ b/Engine/Plugins/PCG/Source/PCG/Private/PCGSubsystem.cpp
@@ -1463,6 +1463,11 @@ uint32 UPCGSubsystem::GetGraphCacheEntryCount(IPCGElement* InElement) const
 	return GraphExecutor ? GraphExecutor->GetGraphCacheEntryCount(InElement) : 0;
 }
 
+void UPCGSubsystem::NotifyLandscapeEditModeExited()
+{
+	ActorAndComponentMapping.NotifyLandscapeEditModeExited();
+}
+
 void UPCGSubsystem::ClearExecutionMetadata(const FPCGStack& BaseStack)
 {
 	ClearExecutedStacks(BaseStack);
diff --git a/Engine/Plugins/PCG/Source/PCG/Public/PCGSubsystem.h b/Engine/Plugins/PCG/Source/PCG/Public/PCGSubsystem.h
--- a/Engine/Plugins/PCG/Source/PCG/Public/PCGSubsystem.h
+++ b/Engine/Plugins/PCG/Source/PCG/Public/PCGSubsystem.h
@@ -296,7 +296,7 @@ public:
 	FPCGNodeVisualLogs& GetNodeVisualLogsMutable() { return NodeVisualLogs; }
 
 	/** Notify that we exited the Landscape edit mode. */
-	void NotifyLandscapeEditModeExited() { ActorAndComponentMapping.NotifyLandscapeEditModeExited(); }
+	void NotifyLandscapeEditModeExited();
 
 	/** Clear any data collected during execution, normally called prior to generating the component. */
 	void ClearExecutionMetadata(const FPCGStack& BaseStack);
