[PCG] Tentative fix for ASAN

#rnx

diff --git a/Engine/Plugins/PCG/Source/PCG/Private/Helpers/PCGSubgraphHelpers.cpp b/Engine/Plugins/PCG/Source/PCG/Private/Helpers/PCGSubgraphHelpers.cpp
--- a/Engine/Plugins/PCG/Source/PCG/Private/Helpers/PCGSubgraphHelpers.cpp
+++ b/Engine/Plugins/PCG/Source/PCG/Private/Helpers/PCGSubgraphHelpers.cpp
@@ -620,14 +620,12 @@ namespace PCGSubgraphHelpersExtra
 			UPCGNode* NewlyCreatedDeclNode = NewGraph->ReconstructNewNode(OriginalDeclNode);
 			check(NewlyCreatedDeclNode);
 			NewlyCreatedDeclNode->NodeTitle = OriginalDeclNode->NodeTitle;
-			OriginalToNewNodesMapping.Emplace(OriginalDeclNode, NewlyCreatedDeclNode);
+			NewDeclNode = &OriginalToNewNodesMapping.Emplace(OriginalDeclNode, NewlyCreatedDeclNode);
 
 			UPCGPin* OriginalOutputPin = OriginalDeclNode->GetOutputPin(PCGPinConstants::DefaultOutputLabel);
 			UPCGPin* NewInputPin = NewlyCreatedDeclNode->GetInputPin(PCGPinConstants::DefaultInputLabel);
 			check(OriginalOutputPin && NewInputPin);
 			ConnectPinsBetweenSubgraphBoundaries(OriginalOutputPin, NewInputPin, OriginalOutputPin, /*bIsInput=*/true);
-
-			NewDeclNode = &NewlyCreatedDeclNode;
 		}
 
 		// Patch and connect the edges.
