[PCG] Disable pcg.ConstructionScriptFix until real fix gets robomerged in

#rnx

diff --git a/Engine/Plugins/PCG/Source/PCG/Private/PCGComponent.cpp b/Engine/Plugins/PCG/Source/PCG/Private/PCGComponent.cpp
--- a/Engine/Plugins/PCG/Source/PCG/Private/PCGComponent.cpp
+++ b/Engine/Plugins/PCG/Source/PCG/Private/PCGComponent.cpp
@@ -70,7 +70,7 @@ namespace PCGComponent
 
 	static TAutoConsoleVariable<bool> CVarConstructionScriptFix(
 		TEXT("pcg.ConstructionScriptFix"),
-		true,
+		false,
 		TEXT("This CVar will be removed in future releases, it allows disabling this fix if regressions are found."));
 
 	template <typename DelegateType>
