[PCG] Fix ensure caused by last fix CL32585341
- Unloading a Level Instance containing PCG Components we would hit the new map in FindAndPrintStaleReferencesToObjects (which gets cleared on next tick)
- Replace by TObjectKeys to avoid hard refs

#rnx

diff --git a/Engine/Plugins/PCG/Source/PCG/Private/PCGSubsystem.cpp b/Engine/Plugins/PCG/Source/PCG/Private/PCGSubsystem.cpp
--- a/Engine/Plugins/PCG/Source/PCG/Private/PCGSubsystem.cpp
+++ b/Engine/Plugins/PCG/Source/PCG/Private/PCGSubsystem.cpp
@@ -220,7 +220,7 @@ void UPCGSubsystem::SetConstructionScriptSourceComponent(UPCGComponent* InCompon
 	{
 		if (AActor* Owner = InComponent->GetOwner())
 		{
-			PerActorConstructionScriptSourceComponents.FindOrAdd(Owner).SourceComponentsMap.Add(InComponent->GetFName(), InComponent);
+			PerActorConstructionScriptSourceComponents.FindOrAdd(Owner).Add(InComponent->GetFName(), InComponent);
 		}
 	}
 }
@@ -230,11 +230,11 @@ bool UPCGSubsystem::RemoveAndCopyConstructionScriptSourceComponent(AActor* InCom
 	OutSourceComponent = nullptr;
 	if (FConstructionScriptSourceComponents* Found = PerActorConstructionScriptSourceComponents.Find(InComponentOwner))
 	{
-		TObjectPtr<UPCGComponent> FoundComponent;
-		if (Found->SourceComponentsMap.RemoveAndCopyValue(InComponentName, FoundComponent))
+		TObjectKey<UPCGComponent> FoundComponent;
+		if (Found->RemoveAndCopyValue(InComponentName, FoundComponent))
 		{
-			OutSourceComponent = FoundComponent;
-			if (Found->SourceComponentsMap.IsEmpty())
+			OutSourceComponent = FoundComponent.ResolveObjectPtrEvenIfGarbage();
+			if (Found->IsEmpty())
 			{
 				PerActorConstructionScriptSourceComponents.Remove(InComponentOwner);
 			}
diff --git a/Engine/Plugins/PCG/Source/PCG/Public/PCGSubsystem.h b/Engine/Plugins/PCG/Source/PCG/Public/PCGSubsystem.h
--- a/Engine/Plugins/PCG/Source/PCG/Public/PCGSubsystem.h
+++ b/Engine/Plugins/PCG/Source/PCG/Public/PCGSubsystem.h
@@ -7,6 +7,7 @@
 #include "PCGCommon.h"
 #include "PCGActorAndComponentMapping.h"
 #include "Grid/PCGComponentOctree.h"
+#include "UObject/ObjectKey.h"
 #include "Utils/PCGNodeVisualLogs.h"
 
 #include "PCGSubsystem.generated.h"
@@ -36,17 +37,6 @@ class UWorld;
 DECLARE_MULTICAST_DELEGATE_OneParam(FPCGOnComponentGenerationCompleteOrCancelled, UPCGSubsystem*);
 #endif // WITH_EDITOR
 
-USTRUCT()
-struct FConstructionScriptSourceComponents
-{
-	GENERATED_USTRUCT_BODY()
-
-#if WITH_EDITORONLY_DATA
-	UPROPERTY(Transient)
-	TMap<FName, TObjectPtr<UPCGComponent>> SourceComponentsMap;
-#endif
-};
-
 /**
 * UPCGSubsystem
 */
@@ -335,11 +325,9 @@ private:
 
 #if WITH_EDITOR
 	FCriticalSection PCGWorldActorLock;
-#endif
-
-#if WITH_EDITORONLY_DATA
-	UPROPERTY(Transient)
-	TMap<TObjectPtr<AActor>, FConstructionScriptSourceComponents> PerActorConstructionScriptSourceComponents;
+		
+	using FConstructionScriptSourceComponents = TMap<FName, TObjectKey<UPCGComponent>>;
+	TMap<TObjectKey<AActor>, FConstructionScriptSourceComponents> PerActorConstructionScriptSourceComponents;
 #endif
 };
 
