[PCG] Fixing ASAN in PCG Subgraph Collapse, storing a pointer on something that goes out of scope

#rnx

diff --git a/Engine/Plugins/PCG/Source/PCG/Private/Helpers/PCGSubgraphHelpers.cpp b/Engine/Plugins/PCG/Source/PCG/Private/Helpers/PCGSubgraphHelpers.cpp
--- a/Engine/Plugins/PCG/Source/PCG/Private/Helpers/PCGSubgraphHelpers.cpp
+++ b/Engine/Plugins/PCG/Source/PCG/Private/Helpers/PCGSubgraphHelpers.cpp
@@ -613,25 +613,29 @@ namespace PCGSubgraphHelpersExtra
 		}
 
 		UPCGNode* OriginalDeclNode = CastChecked<UPCGNode>(OriginalDeclSettings->GetOuter());
-		UPCGNode** NewDeclNode = OriginalToNewNodesMapping.Find(OriginalDeclNode);
-		if (!NewDeclNode)
+		UPCGNode* NewDeclNode = nullptr;
+		if (UPCGNode** NewDeclNodePtr = OriginalToNewNodesMapping.Find(OriginalDeclNode))
+		{
+			NewDeclNode = *NewDeclNodePtr;
+		}
+		else
 		{
 			// If the declaration is outside of the subgraph, we need to create a new declaration inside the subgraph.
-			UPCGNode* NewlyCreatedDeclNode = NewGraph->ReconstructNewNode(OriginalDeclNode);
-			check(NewlyCreatedDeclNode);
-			NewlyCreatedDeclNode->NodeTitle = OriginalDeclNode->NodeTitle;
-			NewDeclNode = &OriginalToNewNodesMapping.Emplace(OriginalDeclNode, NewlyCreatedDeclNode);
+			NewDeclNode = NewGraph->ReconstructNewNode(OriginalDeclNode);
+			check(NewDeclNode);
+			NewDeclNode->NodeTitle = OriginalDeclNode->NodeTitle;
+			OriginalToNewNodesMapping.Emplace(OriginalDeclNode, NewDeclNode);
 
 			UPCGPin* OriginalOutputPin = OriginalDeclNode->GetOutputPin(PCGPinConstants::DefaultOutputLabel);
-			UPCGPin* NewInputPin = NewlyCreatedDeclNode->GetInputPin(PCGPinConstants::DefaultInputLabel);
+			UPCGPin* NewInputPin = NewDeclNode->GetInputPin(PCGPinConstants::DefaultInputLabel);
 			check(OriginalOutputPin && NewInputPin);
 			ConnectPinsBetweenSubgraphBoundaries(OriginalOutputPin, NewInputPin, OriginalOutputPin, /*bIsInput=*/true);
 		}
 
 		// Patch and connect the edges.
-		check(NewDeclNode && *NewDeclNode);
-		NewUsageSettings->Declaration = CastChecked<UPCGNamedRerouteDeclarationSettings>((*NewDeclNode)->GetSettings());
-		UPCGPin* InvisiblePin = (*NewDeclNode)->GetOutputPin(PCGNamedRerouteConstants::InvisiblePinLabel);
+		check(NewDeclNode);
+		NewUsageSettings->Declaration = CastChecked<UPCGNamedRerouteDeclarationSettings>(NewDeclNode->GetSettings());
+		UPCGPin* InvisiblePin = NewDeclNode->GetOutputPin(PCGNamedRerouteConstants::InvisiblePinLabel);
 		UPCGPin* InputUsagePin = NewUsageNode->GetInputPin(PCGPinConstants::DefaultInputLabel);
 		check(InvisiblePin && InputUsagePin);
 		InvisiblePin->AddEdgeTo(InputUsagePin);
