[PCG] Distance and Spatial Noise elements incorrectly inheriting from PointProcessingElement

* Also fixed minor WITH_EDITOR issue with PCGArtificialWait in PCGSample

#rnx

diff --git a/Engine/Plugins/PCG/Source/PCG/Public/Elements/PCGDistance.h b/Engine/Plugins/PCG/Source/PCG/Public/Elements/PCGDistance.h
--- a/Engine/Plugins/PCG/Source/PCG/Public/Elements/PCGDistance.h
+++ b/Engine/Plugins/PCG/Source/PCG/Public/Elements/PCGDistance.h
@@ -2,7 +2,6 @@
 #pragma once
 
 #include "PCGSettings.h"
-#include "Elements/PCGPointProcessingElementBase.h"
 
 #include "PCGDistance.generated.h"
 
@@ -87,7 +86,7 @@ public:
 #endif // WITH_EDITORONLY_DATA
 };
 
-class FPCGDistanceElement : public FPCGPointProcessingElementBase
+class FPCGDistanceElement : public IPCGElement
 {
 protected:
 	virtual bool ExecuteInternal(FPCGContext* Context) const override;
diff --git a/Engine/Plugins/PCG/Source/PCG/Public/Elements/PCGSpatialNoise.h b/Engine/Plugins/PCG/Source/PCG/Public/Elements/PCGSpatialNoise.h
--- a/Engine/Plugins/PCG/Source/PCG/Public/Elements/PCGSpatialNoise.h
+++ b/Engine/Plugins/PCG/Source/PCG/Public/Elements/PCGSpatialNoise.h
@@ -3,7 +3,6 @@
 #pragma once
 
 #include "PCGSettings.h"
-#include "Elements/PCGPointProcessingElementBase.h"
 #include "Metadata/PCGAttributePropertySelector.h"
 
 #include "PCGSpatialNoise.generated.h"
@@ -156,7 +155,7 @@ private:
 	bool bForceNoUseSeed = false;
 };
 
-class FPCGSpatialNoise : public FPCGPointProcessingElementBase
+class FPCGSpatialNoise : public IPCGElement
 {
 protected:
 	virtual bool ExecuteInternal(FPCGContext* Context) const override;
