[PCG] Conversion of Density Remap to Point Operation Element

- Also Deprecation of Old Density Remap (set for deprecation in 5.1)

#rnx

diff --git a/Engine/Plugins/PCG/Source/PCG/Private/Elements/PCGDensityRemap.cpp b/Engine/Plugins/PCG/Source/PCG/Private/Elements/PCGDensityRemap.cpp
deleted file mode 100644
diff --git a/Engine/Plugins/PCG/Source/PCG/Private/Elements/PCGDensityRemapElement.cpp b/Engine/Plugins/PCG/Source/PCG/Private/Elements/PCGDensityRemapElement.cpp
--- a/Engine/Plugins/PCG/Source/PCG/Private/Elements/PCGDensityRemapElement.cpp
+++ b/Engine/Plugins/PCG/Source/PCG/Private/Elements/PCGDensityRemapElement.cpp
@@ -3,7 +3,6 @@
 #include "Elements/PCGDensityRemapElement.h"
 
 #include "PCGContext.h"
-#include "PCGCustomVersion.h"
 
 #include "PCGPoint.h"
 
@@ -14,16 +13,15 @@ FPCGElementPtr UPCGDensityRemapSettings::CreateElement() const
 	return MakeShared<FPCGDensityRemapElement>();
 }
 
-bool FPCGDensityRemapElement::ExecuteInternal(FPCGContext* Context) const
+bool FPCGDensityRemapElement::ExecuteInternal(FPCGContext* InContext) const
 {
 	TRACE_CPUPROFILER_EVENT_SCOPE(FPCGDensityRemapElement::Execute);
 
+	ContextType* Context = static_cast<ContextType*>(InContext);
+
 	const UPCGDensityRemapSettings* Settings = Context->GetInputSettings<UPCGDensityRemapSettings>();
 	check(Settings);
 
-	TArray<FPCGTaggedData> Inputs = Context->InputData.GetInputs();
-	TArray<FPCGTaggedData>& Outputs = Context->OutputData.TaggedData;
-
 	const float InRangeMin = Settings->InRangeMin;
 	const float InRangeMax = Settings->InRangeMax;
 	const float OutRangeMin = Settings->OutRangeMin;
@@ -52,7 +50,7 @@ bool FPCGDensityRemapElement::ExecuteInternal(FPCGContext* Context) const
 		Intercept = OutRangeMin;
 	}
 
-	ProcessPoints(Context, Inputs, Outputs, [bExcludeValuesOutsideInputRange, InRangeTrueMin, InRangeTrueMax, Slope, InRangeMin, Intercept](const FPCGPoint& InPoint, FPCGPoint& OutPoint)
+	return ExecutePointOperation(Context, [bExcludeValuesOutsideInputRange, InRangeTrueMin, InRangeTrueMax, Slope, InRangeMin, Intercept](const FPCGPoint& InPoint, FPCGPoint& OutPoint)
 	{
 		OutPoint = InPoint;
 		const float SourceDensity = InPoint.Density;
@@ -65,6 +63,4 @@ bool FPCGDensityRemapElement::ExecuteInternal(FPCGContext* Context) const
 
 		return true;
 	});
-
-	return true;
 }
diff --git a/Engine/Plugins/PCG/Source/PCG/Public/Elements/PCGDensityRemap.h b/Engine/Plugins/PCG/Source/PCG/Public/Elements/PCGDensityRemap.h
deleted file mode 100644
diff --git a/Engine/Plugins/PCG/Source/PCG/Public/Elements/PCGDensityRemapElement.h b/Engine/Plugins/PCG/Source/PCG/Public/Elements/PCGDensityRemapElement.h
--- a/Engine/Plugins/PCG/Source/PCG/Public/Elements/PCGDensityRemapElement.h
+++ b/Engine/Plugins/PCG/Source/PCG/Public/Elements/PCGDensityRemapElement.h
@@ -2,10 +2,9 @@
 
 #pragma once
 
+#include "PCGPointOperationElementBase.h"
 #include "PCGSettings.h"
 
-#include "Elements/PCGPointProcessingElementBase.h"
-
 #include "PCGDensityRemapElement.generated.h"
 
 UCLASS(MinimalAPI, BlueprintType, ClassGroup = (Procedural))
@@ -20,7 +19,6 @@ public:
 	virtual FText GetDefaultNodeTitle() const override { return NSLOCTEXT("PCGDensityRemapSettings", "NodeTitle", "Density Remap"); }
 	virtual EPCGSettingsType GetType() const override { return EPCGSettingsType::Density; }
 #endif
-	
 
 protected:
 	virtual TArray<FPCGPinProperties> InputPinProperties() const override { return Super::DefaultPointInputPinProperties(); }
@@ -48,13 +46,8 @@ public:
 	bool bExcludeValuesOutsideInputRange = false;
 };
 
-class FPCGDensityRemapElement : public FPCGPointProcessingElementBase
+class FPCGDensityRemapElement : public FPCGPointOperationElementBase
 {
 protected:
-	virtual bool ExecuteInternal(FPCGContext* Context) const override;
+	virtual bool ExecuteInternal(FPCGContext* InContext) const override;
 };
-
-#if UE_ENABLE_INCLUDE_ORDER_DEPRECATED_IN_5_2
-#include "CoreMinimal.h"
-#include "PCGNode.h"
-#endif
