[PCG] Conversion of Point Extents to Point Operation Element

#rnx

diff --git a/Engine/Plugins/PCG/Source/PCG/Private/Elements/PCGPointExtentsModifier.cpp b/Engine/Plugins/PCG/Source/PCG/Private/Elements/PCGPointExtentsModifier.cpp
--- a/Engine/Plugins/PCG/Source/PCG/Private/Elements/PCGPointExtentsModifier.cpp
+++ b/Engine/Plugins/PCG/Source/PCG/Private/Elements/PCGPointExtentsModifier.cpp
@@ -12,61 +12,59 @@ FPCGElementPtr UPCGPointExtentsModifierSettings::CreateElement() const
 	return MakeShared<FPCGPointExtentsModifier>();
 }
 
-bool FPCGPointExtentsModifier::ExecuteInternal(FPCGContext* Context) const
+bool FPCGPointExtentsModifier::ExecuteInternal(FPCGContext* InContext) const
 {
 	TRACE_CPUPROFILER_EVENT_SCOPE(FPCGPointExtentsModifier::Execute);
 
+	ContextType* Context = static_cast<ContextType*>(InContext);
 	const UPCGPointExtentsModifierSettings* Settings = Context->GetInputSettings<UPCGPointExtentsModifierSettings>();
 	check(Settings);
 
-	TArray<FPCGTaggedData> Inputs = Context->InputData.GetInputsByPin(PCGPinConstants::DefaultInputLabel);
-	TArray<FPCGTaggedData>& Outputs = Context->OutputData.TaggedData;
-
 	const EPCGPointExtentsModifierMode Mode = Settings->Mode;
 	const FVector& Extents = Settings->Extents;
 
 	switch (Mode)
 	{
-	case EPCGPointExtentsModifierMode::Minimum:
-		ProcessPoints(Context, Inputs, Outputs, [Extents](const FPCGPoint& InPoint, FPCGPoint& OutPoint){
-			OutPoint = InPoint;
-			OutPoint.SetExtents(FVector::Min(Extents, InPoint.GetExtents()));
-			return true;
-		});
-		break;
+		case EPCGPointExtentsModifierMode::Minimum:
+			return ExecutePointOperation(Context, [Extents](const FPCGPoint& InPoint, FPCGPoint& OutPoint)
+			{
+				OutPoint = InPoint;
+				OutPoint.SetExtents(FVector::Min(Extents, InPoint.GetExtents()));
+				return true;
+			});
 
-	case EPCGPointExtentsModifierMode::Maximum:
-		ProcessPoints(Context, Inputs, Outputs, [Extents](const FPCGPoint& InPoint, FPCGPoint& OutPoint){
-			OutPoint = InPoint;
-			OutPoint.SetExtents(FVector::Max(Extents, InPoint.GetExtents()));
-			return true;
-		});
-		break;
+		case EPCGPointExtentsModifierMode::Maximum:
+			return ExecutePointOperation(Context, [Extents](const FPCGPoint& InPoint, FPCGPoint& OutPoint)
+			{
+				OutPoint = InPoint;
+				OutPoint.SetExtents(FVector::Max(Extents, InPoint.GetExtents()));
+				return true;
+			});
 
 	case EPCGPointExtentsModifierMode::Add:
-		ProcessPoints(Context, Inputs, Outputs, [Extents](const FPCGPoint& InPoint, FPCGPoint& OutPoint){
-			OutPoint = InPoint;
-			OutPoint.SetExtents(Extents + InPoint.GetExtents());
-			return true;
-		});
-		break;
+			return ExecutePointOperation(Context, [Extents](const FPCGPoint& InPoint, FPCGPoint& OutPoint)
+			{
+				OutPoint = InPoint;
+				OutPoint.SetExtents(Extents + InPoint.GetExtents());
+				return true;
+			});
 
-	case EPCGPointExtentsModifierMode::Multiply:
-		ProcessPoints(Context, Inputs, Outputs, [Extents](const FPCGPoint& InPoint, FPCGPoint& OutPoint){
-			OutPoint = InPoint;
-			OutPoint.SetExtents(Extents * InPoint.GetExtents());
-			return true;
-		});
-		break;
+		case EPCGPointExtentsModifierMode::Multiply:
+			return ExecutePointOperation(Context, [Extents](const FPCGPoint& InPoint, FPCGPoint& OutPoint)
+			{
+				OutPoint = InPoint;
+				OutPoint.SetExtents(Extents * InPoint.GetExtents());
+				return true;
+			});
 
-	case EPCGPointExtentsModifierMode::Set:
-		ProcessPoints(Context, Inputs, Outputs, [Extents](const FPCGPoint& InPoint, FPCGPoint& OutPoint){
-			OutPoint = InPoint;
-			OutPoint.SetExtents(Extents);
-			return true;
-		});
-		break;
+		case EPCGPointExtentsModifierMode::Set:
+			return ExecutePointOperation(Context, [Extents](const FPCGPoint& InPoint, FPCGPoint& OutPoint)
+			{
+				OutPoint = InPoint;
+				OutPoint.SetExtents(Extents);
+				return true;
+			});
 	}
-	
+
 	return true;
 }
diff --git a/Engine/Plugins/PCG/Source/PCG/Public/Elements/PCGPointExtentsModifier.h b/Engine/Plugins/PCG/Source/PCG/Public/Elements/PCGPointExtentsModifier.h
--- a/Engine/Plugins/PCG/Source/PCG/Public/Elements/PCGPointExtentsModifier.h
+++ b/Engine/Plugins/PCG/Source/PCG/Public/Elements/PCGPointExtentsModifier.h
@@ -2,10 +2,9 @@
 
 #pragma once
 
+#include "PCGPointOperationElementBase.h"
 #include "PCGSettings.h"
 
-#include "Elements/PCGPointProcessingElementBase.h"
-
 #include "PCGPointExtentsModifier.generated.h"
 
 UENUM()
@@ -47,13 +46,8 @@ public:
 	EPCGPointExtentsModifierMode Mode = EPCGPointExtentsModifierMode::Set;
 };
 
-class FPCGPointExtentsModifier : public FPCGPointProcessingElementBase
+class FPCGPointExtentsModifier : public FPCGPointOperationElementBase
 {
 protected:
-	virtual bool ExecuteInternal(FPCGContext* Context) const override;
-};
-
-#if UE_ENABLE_INCLUDE_ORDER_DEPRECATED_IN_5_2
-#include "CoreMinimal.h"
-#include "PCGNode.h"
-#endif
+	virtual bool ExecuteInternal(FPCGContext* InContext) const override;
+};
\ No newline at end of file
