[PCG] Conversion of NormalToDensity to Point Operation Element

#rnx

diff --git a/Engine/Plugins/PCG/Source/PCG/Private/Elements/PCGNormalToDensity.cpp b/Engine/Plugins/PCG/Source/PCG/Private/Elements/PCGNormalToDensity.cpp
--- a/Engine/Plugins/PCG/Source/PCG/Private/Elements/PCGNormalToDensity.cpp
+++ b/Engine/Plugins/PCG/Source/PCG/Private/Elements/PCGNormalToDensity.cpp
@@ -5,17 +5,18 @@
 #include "Data/PCGPointData.h"
 #include "Kismet/KismetMathLibrary.h"
 #include "PCGContext.h"
-#include "PCGPin.h"
 
 FPCGElementPtr UPCGNormalToDensitySettings::CreateElement() const
 {
 	return MakeShared<FPCGNormalToDensityElement>();
 }
 
-bool FPCGNormalToDensityElement::ExecuteInternal(FPCGContext* Context) const
+bool FPCGNormalToDensityElement::ExecuteInternal(FPCGContext* InContext) const
 {
 	TRACE_CPUPROFILER_EVENT_SCOPE(FPCGNormalToDensityElement::Execute);
 
+	ContextType* Context = static_cast<ContextType*>(InContext);
+
 	const UPCGNormalToDensitySettings* Settings = Context->GetInputSettings<UPCGNormalToDensitySettings>();
 	check(Settings);
 
@@ -24,9 +25,6 @@ bool FPCGNormalToDensityElement::ExecuteInternal(FPCGContext* Context) const
 	const double Strength = Settings->Strength;
 	const PCGNormalToDensityMode DensityMode = Settings->DensityMode;
 
-	TArray<FPCGTaggedData> Inputs = Context->InputData.GetInputsByPin(PCGPinConstants::DefaultInputLabel);
-	TArray<FPCGTaggedData>& Outputs = Context->OutputData.TaggedData;
-
 	const double InvStrength = 1.0/FMath::Max(0.0001, Strength);
 
 	const auto CalcValue = [Normal, Offset, InvStrength](const FPCGPoint& InPoint)
@@ -37,91 +35,70 @@ bool FPCGNormalToDensityElement::ExecuteInternal(FPCGContext* Context) const
 
 	switch (DensityMode)
 	{
-	case PCGNormalToDensityMode::Set:
-		ProcessPoints(Context, Inputs, Outputs, [CalcValue](const FPCGPoint& InPoint, FPCGPoint& OutPoint)->bool{
-
-			OutPoint = InPoint;
-
-			OutPoint.Density = CalcValue(InPoint);
-
-			return true;
-
-		});
-		break;
-
-	case PCGNormalToDensityMode::Minimum:
-		ProcessPoints(Context, Inputs, Outputs, [CalcValue](const FPCGPoint& InPoint, FPCGPoint& OutPoint)->bool{
-
-			OutPoint = InPoint;
-
-			OutPoint.Density = FMath::Min(InPoint.Density, CalcValue(InPoint));
+		case PCGNormalToDensityMode::Set:
+			return ExecutePointOperation(Context, [CalcValue](const FPCGPoint& InPoint, FPCGPoint& OutPoint)
+			{
+				OutPoint = InPoint;
 
-			return true;
+				OutPoint.Density = CalcValue(InPoint);
 
-		});
-		break;
+				return true;
+			});
 
-	case PCGNormalToDensityMode::Maximum:
-		ProcessPoints(Context, Inputs, Outputs, [CalcValue](const FPCGPoint& InPoint, FPCGPoint& OutPoint)->bool{
+		case PCGNormalToDensityMode::Minimum:
+			return ExecutePointOperation(Context, [CalcValue](const FPCGPoint& InPoint, FPCGPoint& OutPoint)->bool{
+				OutPoint = InPoint;
 
-			OutPoint = InPoint;
+				OutPoint.Density = FMath::Min(InPoint.Density, CalcValue(InPoint));
 
-			OutPoint.Density = FMath::Max(InPoint.Density, CalcValue(InPoint));
+				return true;
+			});
 
-			return true;
+		case PCGNormalToDensityMode::Maximum:
+			return ExecutePointOperation(Context, [CalcValue](const FPCGPoint& InPoint, FPCGPoint& OutPoint)->bool{
+				OutPoint = InPoint;
 
-		});
-		break;
+				OutPoint.Density = FMath::Max(InPoint.Density, CalcValue(InPoint));
 
-	case PCGNormalToDensityMode::Add:
-		ProcessPoints(Context, Inputs, Outputs, [CalcValue](const FPCGPoint& InPoint, FPCGPoint& OutPoint)->bool{
+				return true;
+			});
 
-			OutPoint = InPoint;
+		case PCGNormalToDensityMode::Add:
+			return ExecutePointOperation(Context, [CalcValue](const FPCGPoint& InPoint, FPCGPoint& OutPoint)->bool{
+				OutPoint = InPoint;
 
-			OutPoint.Density += CalcValue(InPoint);
+				OutPoint.Density += CalcValue(InPoint);
 
-			return true;
+				return true;
+			});
 
-		});
-		break;
+		case PCGNormalToDensityMode::Subtract:
+			return ExecutePointOperation(Context, [CalcValue](const FPCGPoint& InPoint, FPCGPoint& OutPoint)->bool{
+				OutPoint = InPoint;
 
-	case PCGNormalToDensityMode::Subtract:
-		ProcessPoints(Context, Inputs, Outputs, [CalcValue](const FPCGPoint& InPoint, FPCGPoint& OutPoint)->bool{
+				OutPoint.Density -= CalcValue(InPoint);
 
-			OutPoint = InPoint;
+				return true;
+			});
 
-			OutPoint.Density -= CalcValue(InPoint);
+		case PCGNormalToDensityMode::Multiply:
+			return ExecutePointOperation(Context, [CalcValue](const FPCGPoint& InPoint, FPCGPoint& OutPoint)->bool{
+				OutPoint = InPoint;
 
-			return true;
+				OutPoint.Density *= CalcValue(InPoint);
 
-		});
-		break;
+				return true;
+			});
 
-	case PCGNormalToDensityMode::Multiply:
-		ProcessPoints(Context, Inputs, Outputs, [CalcValue](const FPCGPoint& InPoint, FPCGPoint& OutPoint)->bool{
+		case PCGNormalToDensityMode::Divide:
+			return ExecutePointOperation(Context, [CalcValue](const FPCGPoint& InPoint, FPCGPoint& OutPoint)->bool{
+				OutPoint = InPoint;
 
-			OutPoint = InPoint;
+				OutPoint.Density = UKismetMathLibrary::SafeDivide(OutPoint.Density, CalcValue(InPoint));
 
-			OutPoint.Density *= CalcValue(InPoint);
-
-			return true;
-
-		});
-		break;
-
-	case PCGNormalToDensityMode::Divide:
-		ProcessPoints(Context, Inputs, Outputs, [CalcValue](const FPCGPoint& InPoint, FPCGPoint& OutPoint)->bool{
-
-			OutPoint = InPoint;
-
-			OutPoint.Density = UKismetMathLibrary::SafeDivide(OutPoint.Density, CalcValue(InPoint));
-
-			return true;
-
-		});
-		break;
+				return true;
+			});
 	}
 
 	return true;
-
 }
diff --git a/Engine/Plugins/PCG/Source/PCG/Public/Elements/PCGNormalToDensity.h b/Engine/Plugins/PCG/Source/PCG/Public/Elements/PCGNormalToDensity.h
--- a/Engine/Plugins/PCG/Source/PCG/Public/Elements/PCGNormalToDensity.h
+++ b/Engine/Plugins/PCG/Source/PCG/Public/Elements/PCGNormalToDensity.h
@@ -1,8 +1,8 @@
 // Copyright Epic Games, Inc. All Rights Reserved.
 #pragma once
 
+#include "PCGPointOperationElementBase.h"
 #include "PCGSettings.h"
-#include "Elements/PCGPointProcessingElementBase.h"
 
 #include "PCGNormalToDensity.generated.h"
 
@@ -60,9 +60,9 @@ public:
 	PCGNormalToDensityMode DensityMode = PCGNormalToDensityMode::Set;
 };
 
-class FPCGNormalToDensityElement : public FPCGPointProcessingElementBase
+class FPCGNormalToDensityElement : public FPCGPointOperationElementBase
 {
 protected:
-	virtual bool ExecuteInternal(FPCGContext* Context) const override;
+	virtual bool ExecuteInternal(FPCGContext* InContext) const override;
 };
 
