[PCG] Added sorting to the Determinism List View

#rnx

diff --git a/Engine/Plugins/PCG/Source/PCG/Public/Tests/Determinism/PCGDeterminismTestsCommon.h b/Engine/Plugins/PCG/Source/PCG/Public/Tests/Determinism/PCGDeterminismTestsCommon.h
--- a/Engine/Plugins/PCG/Source/PCG/Public/Tests/Determinism/PCGDeterminismTestsCommon.h
+++ b/Engine/Plugins/PCG/Source/PCG/Public/Tests/Determinism/PCGDeterminismTestsCommon.h
@@ -36,6 +36,10 @@ struct FDeterminismTestResult
 {
 	GENERATED_BODY()
 
+	/** The test result index. */
+	UPROPERTY(VisibleAnywhere, BlueprintReadOnly, Category = Determinism)
+	int64 Index = INDEX_NONE;
+
 	/** The node's title */
 	UPROPERTY(VisibleAnywhere, BlueprintReadOnly, Category = Determinism)
 	FName TestResultTitle = TEXT("Untitled");
diff --git a/Engine/Plugins/PCG/Source/PCGEditor/Private/PCGEditor.cpp b/Engine/Plugins/PCG/Source/PCGEditor/Private/PCGEditor.cpp
--- a/Engine/Plugins/PCG/Source/PCGEditor/Private/PCGEditor.cpp
+++ b/Engine/Plugins/PCG/Source/PCGEditor/Private/PCGEditor.cpp
@@ -973,6 +973,7 @@ void FPCGEditor::OnDeterminismNodeTest()
 	DeterminismWidget->ClearItems();
 	DeterminismWidget->BuildBaseColumns();
 
+	int64 TestIndex = 0;
 	for (UObject* Object : GraphEditorWidget->GetSelectedNodes())
 	{
 		check(Object);
@@ -996,6 +997,7 @@ void FPCGEditor::OnDeterminismNodeTest()
 				check(PCGNode && PCGNode->GetSettings());
 
 				TSharedPtr<FDeterminismTestResult> NodeResult = MakeShared<FDeterminismTestResult>();
+				NodeResult->Index = TestIndex++;
 				NodeResult->TestResultTitle = FName(*PCGNode->GetNodeTitle(EPCGNodeTitleType::ListView).ToString());
 				NodeResult->TestResultName = PCGNode->GetName();
 				NodeResult->Seed = PCGNode->GetSettings()->GetSeed();
@@ -1085,6 +1087,7 @@ void FPCGEditor::OnDeterminismGraphTest()
 	DeterminismWidget->AddColumn(ColumnInfo);
 
 	TSharedPtr<FDeterminismTestResult> TestResult = MakeShared<FDeterminismTestResult>();
+	TestResult->Index = 0;
 	TestResult->TestResultTitle = TEXT("Full Graph Test");
 	TestResult->TestResultName = PCGGraphBeingEdited->GetName();
 	TestResult->Seed = PCGComponentBeingInspected->Seed;
diff --git a/Engine/Plugins/PCG/Source/PCGEditor/Private/SPCGEditorGraphDeterminism.cpp b/Engine/Plugins/PCG/Source/PCGEditor/Private/SPCGEditorGraphDeterminism.cpp
--- a/Engine/Plugins/PCG/Source/PCGEditor/Private/SPCGEditorGraphDeterminism.cpp
+++ b/Engine/Plugins/PCG/Source/PCGEditor/Private/SPCGEditorGraphDeterminism.cpp
@@ -39,7 +39,6 @@ namespace
 void SPCGEditorGraphDeterminismRow::Construct(const FArguments& InArgs, const TSharedRef<STableViewBase>& InOwnerTableView, const FPCGNodeTestResultPtr& Item, int32 ItemIndex)
 {
 	CurrentItem = Item;
-	CurrentIndex = ItemIndex;
 	SMultiColumnTableRow<FPCGNodeTestResultPtr>::Construct(SMultiColumnTableRow::FArguments(), InOwnerTableView);
 }
 
@@ -59,7 +58,7 @@ TSharedRef<SWidget> SPCGEditorGraphDeterminismRow::GenerateWidgetForColumn(const
 	// Permanent columns
 	if (ColumnId == NAME_Index)
 	{
-		return ReturnColorCodedResultBlock(FText::FromString(FString::FromInt(CurrentIndex)));
+		return ReturnColorCodedResultBlock(FText::FromString(FString::FromInt(CurrentItem->Index)));
 	}
 	else if (ColumnId == NAME_NodeTitle)
 	{
@@ -124,6 +123,8 @@ void SPCGEditorGraphDeterminismListView::Construct(const FArguments& InArgs, TWe
 	PCGEditorPtr = InPCGEditor;
 
 	GeneratedHeaderRow = SNew(SHeaderRow);
+	SortMode = EColumnSortMode::None;
+	SortingColumn = NAME_None;
 
 	SAssignNew(ListView, SListView<FPCGNodeTestResultPtr>)
 		.ListItemsSource(&ListViewItems)
@@ -168,6 +169,8 @@ void SPCGEditorGraphDeterminismListView::AddColumn(const FTestColumnInfo& Column
 	}
 	Arguments.HAlignHeader(ColumnInfo.HAlign);
 	Arguments.HAlignCell(ColumnInfo.HAlign);
+	Arguments.SortMode(this, &SPCGEditorGraphDeterminismListView::GetColumnSortMode, ColumnInfo.ColumnID);
+	Arguments.OnSort(this, &SPCGEditorGraphDeterminismListView::OnSortColumnHeader);
 	GeneratedHeaderRow->AddColumn(Arguments);
 }
 
@@ -203,4 +206,72 @@ TSharedRef<ITableRow> SPCGEditorGraphDeterminismListView::OnGenerateRow(const FP
 	return SNew(SPCGEditorGraphDeterminismRow, OwnerTable, Item, ++ItemIndexCounter);
 }
 
+void SPCGEditorGraphDeterminismListView::OnSortColumnHeader(const EColumnSortPriority::Type SortPriority, const FName& ColumnId, const EColumnSortMode::Type NewSortMode)
+{
+	if (SortingColumn == ColumnId)
+	{
+		// Circling
+		SortMode = static_cast<EColumnSortMode::Type>((SortMode + 1) % 3);
+	}
+	else
+	{
+		SortingColumn = ColumnId;
+		SortMode = NewSortMode;
+	}
+
+	if (SortingColumn != NAME_None && SortMode != EColumnSortMode::None)
+	{
+		Algo::Sort(ListViewItems, [this, ColumnId](const FPCGNodeTestResultPtr& A, const FPCGNodeTestResultPtr& B)
+		{
+			bool bIsLess = false;
+			if (SortingColumn == NAME_Index)
+			{
+				bIsLess = A->Index < B->Index;
+			}
+			else if (SortingColumn == NAME_NodeTitle)
+			{
+				bIsLess = A->TestResultTitle.ToString() < B->TestResultTitle.ToString();
+			}
+			else if (SortingColumn == NAME_NodeName)
+			{
+				bIsLess = A->TestResultName < B->TestResultName;
+			}
+			else if (SortingColumn == NAME_Seed)
+			{
+				bIsLess = A->Seed < B->Seed;
+			}
+			else if (SortingColumn == NAME_DataTypesTested)
+			{
+				bIsLess = A->DataTypesTested < B->DataTypesTested;
+			}
+			else if (SortingColumn == NAME_AdditionalDetails)
+			{
+				bIsLess = (A->AdditionalDetails.IsEmpty() && !B->AdditionalDetails.IsEmpty()) ||
+					(!A->AdditionalDetails.IsEmpty() && !B->AdditionalDetails.IsEmpty() && A->AdditionalDetails[0] < B->AdditionalDetails[0]);
+			}
+			else
+			{
+				const EDeterminismLevel* DeterminismLevelA = A->TestResults.Find(ColumnId);
+				const EDeterminismLevel* DeterminismLevelB = B->TestResults.Find(ColumnId);
+				bIsLess = (DeterminismLevelA && !DeterminismLevelB) ||
+					(DeterminismLevelA && DeterminismLevelB && *DeterminismLevelA < *DeterminismLevelB);
+			}
+
+			return SortMode == EColumnSortMode::Ascending ? bIsLess : !bIsLess;
+		});
+	}
+
+	RefreshItems();
+}
+
+EColumnSortMode::Type SPCGEditorGraphDeterminismListView::GetColumnSortMode(const FName ColumnId) const
+{
+	if (SortingColumn != ColumnId)
+	{
+		return EColumnSortMode::None;
+	}
+
+	return SortMode;
+}
+
 #undef LOCTEXT_NAMESPACE
diff --git a/Engine/Plugins/PCG/Source/PCGEditor/Private/SPCGEditorGraphDeterminism.h b/Engine/Plugins/PCG/Source/PCGEditor/Private/SPCGEditorGraphDeterminism.h
--- a/Engine/Plugins/PCG/Source/PCGEditor/Private/SPCGEditorGraphDeterminism.h
+++ b/Engine/Plugins/PCG/Source/PCGEditor/Private/SPCGEditorGraphDeterminism.h
@@ -40,7 +40,6 @@ class SPCGEditorGraphDeterminismRow final : public SMultiColumnTableRow<FPCGNode
 
 protected:
 	FPCGNodeTestResultPtr CurrentItem;
-	int32 CurrentIndex = -1;
 };
 
 class SPCGEditorGraphDeterminismListView final : public SCompoundWidget
@@ -72,6 +71,9 @@ private:
 	void ClearColumns();
 	/** Generate the row widget */
 	TSharedRef<ITableRow> OnGenerateRow(const FPCGNodeTestResultPtr Item, const TSharedRef<STableViewBase>& OwnerTable) const;
+	void OnSortColumnHeader(const EColumnSortPriority::Type SortPriority, const FName& ColumnId, const EColumnSortMode::Type NewSortMode);
+	EColumnSortMode::Type GetColumnSortMode(const FName ColumnId) const;
+
 	TSharedPtr<SHeaderRow> GeneratedHeaderRow;
 
 	TWeakPtr<FPCGEditor> PCGEditorPtr;
@@ -79,6 +81,10 @@ private:
 	TSharedPtr<SListView<FPCGNodeTestResultPtr>> ListView;
 	TArray<FPCGNodeTestResultPtr> ListViewItems;
 
+	// To allow sorting
+	FName SortingColumn = NAME_None;
+	EColumnSortMode::Type SortMode = EColumnSortMode::Type::None;
+
 	bool bIsConstructed = false;
 	mutable int32 ItemIndexCounter = -1;
 };
