[PCG] Fix invalid bounds after a Volume Sampler on a Point Data

#rnx

diff --git a/Engine/Plugins/PCG/Source/PCG/Private/Elements/PCGVolumeSampler.cpp b/Engine/Plugins/PCG/Source/PCG/Private/Elements/PCGVolumeSampler.cpp
--- a/Engine/Plugins/PCG/Source/PCG/Private/Elements/PCGVolumeSampler.cpp
+++ b/Engine/Plugins/PCG/Source/PCG/Private/Elements/PCGVolumeSampler.cpp
@@ -119,6 +119,8 @@ namespace PCGVolumeSampler
 
 				OutPoint.Seed = PCGHelpers::ComputeSeed(X, Y, Z);
 				OutPoint.Steepness = PointSteepness;
+				// Reset the bounds to the expected bounds, if it was modified by the Sample Point
+				OutPoint.SetLocalBounds(VoxelBox);
 				return true;
 			}
 			else
