[PCG] Improvements on Random Choice node.

* Allow to set the seed, no deprecation since it is a new node
* Add optimization to avoid copying points if we keep none or all the points
* Add optimization to not shuffle more than we need
* Add optimization to not output in the discarded points, if not needed.
* Improve shuffling algorithm to have a uniform sampling (was not uniform before)
* Add option to make the choice stable, to keep the original order
* Add unit tests

#rnx

diff --git a/Engine/Plugins/PCG/Source/PCG/Private/Elements/PCGRandomChoice.cpp b/Engine/Plugins/PCG/Source/PCG/Private/Elements/PCGRandomChoice.cpp
--- a/Engine/Plugins/PCG/Source/PCG/Private/Elements/PCGRandomChoice.cpp
+++ b/Engine/Plugins/PCG/Source/PCG/Private/Elements/PCGRandomChoice.cpp
@@ -9,34 +9,37 @@
 
 #define LOCTEXT_NAMESPACE "PCGRandomChoiceElement"
 
-namespace PCGRandomChoiceConstants
+UPCGRandomChoiceSettings::UPCGRandomChoiceSettings()
 {
-	const FName OutputALabel = TEXT("Chosen Points");
-	const FName OutputBLabel = TEXT("Discarded Points");
+	bUseSeed = true;
 }
 
 #if WITH_EDITOR
-// The label the node is known as internally.
 FName UPCGRandomChoiceSettings::GetDefaultNodeName() const
 {
 	return FName(TEXT("RandomChoice"));
 }
 
-// Default node name shown in the graph editor. Include spaces.
 FText UPCGRandomChoiceSettings::GetDefaultNodeTitle() const
 {
 	return LOCTEXT("NodeTitle", "Random Choice");
 }
 
-// Default tooltip for the node
 FText UPCGRandomChoiceSettings::GetNodeTooltipText() const
 {
-	return LOCTEXT("NodeTooltip", "Chooses points randomly through ratio or a fixed number of points.");
+	return LOCTEXT("NodeTooltip", "Chooses points randomly through ratio or a fixed number of points.\n"
+		"Chosen/Discarded Points will be in the same order than they appear in the input data.");
 }
 #endif
 
-// Input/Output pin setup with specific properties, including:
-// Pin data type, allowing singular or multiple inputs per pin, and creating multiple in/out pins.
+EPCGDataType UPCGRandomChoiceSettings::GetCurrentPinTypes(const UPCGPin* InPin) const
+{
+	// We have dynamic pins to true because we can change the number of output pins.
+	// But the output type is fully defined (point) so just return the pin allowed type.
+	check(InPin);
+	return InPin->Properties.AllowedTypes;
+}
+
 TArray<FPCGPinProperties> UPCGRandomChoiceSettings::InputPinProperties() const
 {
 	return Super::DefaultPointInputPinProperties();
@@ -45,34 +48,27 @@ TArray<FPCGPinProperties> UPCGRandomChoiceSettings::InputPinProperties() const
 TArray<FPCGPinProperties> UPCGRandomChoiceSettings::OutputPinProperties() const
 {
 	TArray<FPCGPinProperties> PinProperties;
-	PinProperties.Emplace(PCGRandomChoiceConstants::OutputALabel,
+	PinProperties.Emplace(PCGRandomChoiceConstants::ChosenPointsLabel,
 		EPCGDataType::Point,
 		/*bAllowMultipleConnections=*/true,
 		/*bAllowMultipleData=*/true);
 
-	PinProperties.Emplace(PCGRandomChoiceConstants::OutputBLabel,
-		EPCGDataType::Point,
-		/*bAllowMultipleConnections=*/true,
-		/*bAllowMultipleData=*/true);
+	if (bOutputDiscardedPoints)
+	{
+		PinProperties.Emplace(PCGRandomChoiceConstants::DiscardedPointsLabel,
+			EPCGDataType::Point,
+			/*bAllowMultipleConnections=*/true,
+			/*bAllowMultipleData=*/true);
+	}
 
 	return PinProperties;
 }
 
-// Creates the Element to be used for ExecuteInternal.
 FPCGElementPtr UPCGRandomChoiceSettings::CreateElement() const
 {
 	return MakeShared<FPCGRandomChoiceElement>();
 }
 
-/*
-* Processing function for this node. 
-* Context holds the InputData, containing the input data collection for this node 
-* and the OutputData, the output data collection to write to as output.
-* Returns true if the processing is done. 
-* Returning false will call back this function at next tick, and will call it until it returns true.
-* Settings contains all the setup options for this node, and if a property was marked PCG_Overridable, 
-* "Context->GetInputSettings" will contain the overridden value for this property if it is overridden
-*/ 
 bool FPCGRandomChoiceElement::ExecuteInternal(FPCGContext* Context) const
 {
 	TRACE_CPUPROFILER_EVENT_SCOPE(FPCGRandomChoiceElement::Execute);
@@ -87,7 +83,9 @@ bool FPCGRandomChoiceElement::ExecuteInternal(FPCGContext* Context) const
 
 	for (int i = 0; i < Inputs.Num(); ++i)
 	{
-		const UPCGPointData* InputPointData = Cast<UPCGPointData>(Inputs[i].Data);
+		const FPCGTaggedData& CurrentInput = Inputs[i];
+
+		const UPCGPointData* InputPointData = Cast<UPCGPointData>(CurrentInput.Data);
 		if (!InputPointData)
 		{
 			PCGE_LOG(Verbose, GraphAndLog, FText::Format(LOCTEXT("InvalidPointData", "Input {0} is not point data"), i));
@@ -99,13 +97,60 @@ bool FPCGRandomChoiceElement::ExecuteInternal(FPCGContext* Context) const
 
 		if (Settings->bFixedMode)
 		{
-			NumOfElementsToKeep = FMath::Clamp(Settings->FixedNumber, 1, InPoints.Num());
+			NumOfElementsToKeep = FMath::Clamp(Settings->FixedNumber, 0, InPoints.Num());
 		}
 		else
 		{
 			NumOfElementsToKeep = FMath::CeilToInt(InPoints.Num() * FMath::Clamp(Settings->Ratio, 0, 1));
 		}
 
+		auto CreateData = [InputPointData, &Outputs, &CurrentInput](const int NumPoints, const FName PinLabel) -> TArray<FPCGPoint>*
+		{
+			FPCGTaggedData& Output = Outputs.Add_GetRef(CurrentInput);
+			UPCGPointData* OutPointData = NewObject<UPCGPointData>();
+			OutPointData->InitializeFromData(InputPointData);
+
+			TArray<FPCGPoint>& Points = OutPointData->GetMutablePoints();
+			Points.Reserve(NumPoints);
+
+			Output.Data = OutPointData;
+			Output.Pin = PinLabel;
+
+			return &Points;
+		};
+
+		if (NumOfElementsToKeep == 0)
+		{
+			// We keep no points, forward the input to the Discarded Points, and create an empty point data on chosen for parity.
+			if (Settings->bOutputDiscardedPoints)
+			{
+				FPCGTaggedData& DiscardedOutput = Outputs.Add_GetRef(CurrentInput);
+				DiscardedOutput.Pin = PCGRandomChoiceConstants::DiscardedPointsLabel;
+			}
+
+			CreateData(0, PCGRandomChoiceConstants::ChosenPointsLabel);
+			continue;
+		}
+		else if (NumOfElementsToKeep == InPoints.Num())
+		{
+			// We keep all the points, forward the input to the Chosen Points, and create an empty point data on discarded for parity.
+			FPCGTaggedData& ChosenOutput = Outputs.Add_GetRef(CurrentInput);
+			ChosenOutput.Pin = PCGRandomChoiceConstants::ChosenPointsLabel;
+
+			if (Settings->bOutputDiscardedPoints)
+			{
+				CreateData(0, PCGRandomChoiceConstants::DiscardedPointsLabel);
+			}
+
+			continue;
+		}
+
+		// TODO: While shuffling is the most intuitive way of selecting randomly points, it is still inefficient in memory, especially if we have a lot of points and want to select just a few of them.
+		// Perhaps we could chose another algorithm in that case.
+		// Like for example:
+		// * Pick a number in [0, n[
+		// * Pick a number in[0, n - 1[, then for each previously selected number, if it's larger, add +1
+		// It's O(n) cpu + O(n) memory vs O(s^2) cpu + O(s) memory. (n = total number of points, s = number of points to keep)
 		FRandomStream RandStream(Context->GetSeed());
 		TArray<int32> ShuffledIndexes;
 		ShuffledIndexes.Reserve(InPoints.Num());
@@ -116,32 +161,34 @@ bool FPCGRandomChoiceElement::ExecuteInternal(FPCGContext* Context) const
 			ShuffledIndexes.Add(j);
 		}
 
-		for (int j = 0; j < InPoints.Num() - 1; ++j)
+		// We only have to shuffle until we reached the number of elements to keep.
+		for (int j = 0; j < NumOfElementsToKeep; ++j)
 		{
-			int RandomElement = RandStream.RandRange(j + 1, InPoints.Num() - 1);
-			Swap(ShuffledIndexes[j], ShuffledIndexes[RandomElement]);
+			const int RandomElement = RandStream.RandRange(j, InPoints.Num() - 1);
+			if (RandomElement != j)
+			{
+				Swap(ShuffledIndexes[j], ShuffledIndexes[RandomElement]);
+			}
 		}
 
-		FPCGTaggedData& ChosenOutput = Outputs.Add_GetRef(Inputs[i]);
-		UPCGPointData* ChosenPointData = NewObject<UPCGPointData>();
-		ChosenPointData->InitializeFromData(InputPointData);
-		TArray<FPCGPoint>& ChosenPoints = ChosenPointData->GetMutablePoints();
-		ChosenPoints.Reserve(NumOfElementsToKeep);
-		ChosenOutput.Data = ChosenPointData;
-		ChosenOutput.Pin = PCGRandomChoiceConstants::OutputALabel;
-		
-		FPCGTaggedData& DiscardedOutput = Outputs.Add_GetRef(Inputs[i]);
-		UPCGPointData* DiscardedPointData = NewObject<UPCGPointData>();
-		DiscardedPointData->InitializeFromData(InputPointData);
-		TArray<FPCGPoint>& DiscardedPoints = DiscardedPointData->GetMutablePoints();
-		DiscardedPoints.Reserve(InPoints.Num() - NumOfElementsToKeep);
-		DiscardedOutput.Data = DiscardedPointData;
-		DiscardedOutput.Pin = PCGRandomChoiceConstants::OutputBLabel;
-
-		
-		for (int j = 0; j < InPoints.Num(); ++j)
+		auto AddToOutput = [&CreateData, &InPoints, &ShuffledIndexes](const int StartIndex, const int Count, const FName Label)
+		{
+			// Order needs to be stable to sort this part of the array
+			Algo::Sort(MakeArrayView(ShuffledIndexes.GetData() + StartIndex, Count));
+
+			TArray<FPCGPoint>* Points = CreateData(Count, Label);
+			check(Points);
+			for (int j = 0; j < Count; ++j)
+			{
+				Points->Add(InPoints[ShuffledIndexes[StartIndex + j]]);
+			}
+		};
+
+		AddToOutput(0, NumOfElementsToKeep, PCGRandomChoiceConstants::ChosenPointsLabel);
+
+		if (Settings->bOutputDiscardedPoints)
 		{
-			(j >= NumOfElementsToKeep ? DiscardedPoints : ChosenPoints).Add(InPoints[ShuffledIndexes[j]]);
+			AddToOutput(NumOfElementsToKeep, InPoints.Num() - NumOfElementsToKeep, PCGRandomChoiceConstants::DiscardedPointsLabel);
 		}
 	}
 
diff --git a/Engine/Plugins/PCG/Source/PCG/Private/Tests/Elements/PCGRandomChoiceTest.cpp b/Engine/Plugins/PCG/Source/PCG/Private/Tests/Elements/PCGRandomChoiceTest.cpp
new file mode 100644
--- /dev/null
+++ b/Engine/Plugins/PCG/Source/PCG/Private/Tests/Elements/PCGRandomChoiceTest.cpp
@@ -0,0 +1,253 @@
+// Copyright Epic Games, Inc. All Rights Reserved.
+
+#if WITH_EDITOR
+
+#include "Tests/PCGTestsCommon.h"
+
+#include "PCGComponent.h"
+#include "PCGContext.h"
+#include "Data/PCGPointData.h"
+
+#include "Elements/PCGRandomChoice.h"
+
+IMPLEMENT_CUSTOM_SIMPLE_AUTOMATION_TEST(FPCGRandomChoiceTest_Fixed, FPCGTestBaseClass, "Plugins.PCG.RandomChoice,Fixed", PCGTestsCommon::TestFlags)
+IMPLEMENT_CUSTOM_SIMPLE_AUTOMATION_TEST(FPCGRandomChoiceTest_Ratio, FPCGTestBaseClass, "Plugins.PCG.RandomChoice.Ratio", PCGTestsCommon::TestFlags)
+IMPLEMENT_CUSTOM_SIMPLE_AUTOMATION_TEST(FPCGRandomChoiceTest_SelectNone, FPCGTestBaseClass, "Plugins.PCG.RandomChoice.SelectNone", PCGTestsCommon::TestFlags)
+IMPLEMENT_CUSTOM_SIMPLE_AUTOMATION_TEST(FPCGRandomChoiceTest_SelectAll, FPCGTestBaseClass, "Plugins.PCG.RandomChoice.SelectAll", PCGTestsCommon::TestFlags)
+IMPLEMENT_CUSTOM_SIMPLE_AUTOMATION_TEST(FPCGRandomChoiceTest_NoDiscard, FPCGTestBaseClass, "Plugins.PCG.RandomChoice.NoDiscard", PCGTestsCommon::TestFlags)
+
+namespace PCGRandomChoiceTest
+{
+	UPCGPointData* CreateInputPointData(FPCGContext* Context, const int NumPoints)
+	{
+		check(Context);
+
+		UPCGPointData* NewPointData = NewObject<UPCGPointData>();
+		NewPointData->SetFlags(RF_Transient);
+		TArray<FPCGPoint>& Points = NewPointData->GetMutablePoints();
+		Points.SetNum(NumPoints);
+		for (int i = 0; i < NumPoints; ++i)
+		{
+			// Store the index in the density
+			Points[i].Density = i;
+		}
+
+		FPCGTaggedData& InputData = Context->InputData.TaggedData.Emplace_GetRef();
+		InputData.Data = NewPointData;
+		InputData.Pin = PCGPinConstants::DefaultInputLabel;
+
+		return NewPointData;
+	}
+
+	bool VerifyAllPointsThere(const int NumPoints, const UPCGPointData* ChosenPointData, const UPCGPointData* DiscardedPointData)
+	{
+		check(ChosenPointData && DiscardedPointData);
+		TSet<int> IndexesSeen;
+
+		auto Check = [&IndexesSeen](const UPCGPointData* PointData)
+		{
+			const TArray<FPCGPoint>& Points = PointData->GetPoints();
+			for (int i = 0; i < Points.Num(); ++i)
+			{
+				const FPCGPoint& Point = Points[i];
+				const int Index = static_cast<int>(Point.Density);
+				if (IndexesSeen.Contains(Index))
+				{
+					return false;
+				}
+
+				IndexesSeen.Add(Index);
+
+				// It needs to be stable so density should be ascending.
+				if (i > 0 && Points[i].Density < Points[i - 1].Density)
+				{
+					return false;
+				}
+			}
+
+			return true;
+		};
+
+		if (!Check(ChosenPointData))
+		{
+			return false;
+		}
+
+		if (!Check(DiscardedPointData))
+		{
+			return false;
+		}
+
+		return IndexesSeen.Num() == NumPoints;
+	}
+}
+
+bool FPCGRandomChoiceTest_Fixed::RunTest(const FString& Parameters)
+{
+	PCGTestsCommon::FTestData TestData;
+	UPCGRandomChoiceSettings* Settings = PCGTestsCommon::GenerateSettings<UPCGRandomChoiceSettings>(TestData);
+	check(Settings);
+
+	constexpr int NumOfPoints = 20;
+	constexpr int ExpectedNumElementsChosen = 7;
+	constexpr int ExpectedNumElementsDiscarded = 13;
+
+	Settings->bFixedMode = true;
+	Settings->FixedNumber = ExpectedNumElementsChosen;
+
+	TUniquePtr<FPCGContext> Context = TestData.InitializeTestContext();
+	const UPCGPointData* InputPointData = PCGRandomChoiceTest::CreateInputPointData(Context.Get(), NumOfPoints);
+
+	FPCGElementPtr TestElement = TestData.Settings->GetElement();
+	while (!TestElement->Execute(Context.Get())) {}
+
+	TArray<FPCGTaggedData> ChosenOutputTagged = Context->OutputData.GetInputsByPin(PCGRandomChoiceConstants::ChosenPointsLabel);
+	TArray<FPCGTaggedData> DiscardedOutputTagged = Context->OutputData.GetInputsByPin(PCGRandomChoiceConstants::DiscardedPointsLabel);
+
+	const UPCGPointData* ChosenOutputData = ChosenOutputTagged.Num() == 1 ? Cast<const UPCGPointData>(ChosenOutputTagged[0].Data) : nullptr;
+	const UPCGPointData* DiscardedOutputData = DiscardedOutputTagged.Num() == 1 ? Cast<const UPCGPointData>(DiscardedOutputTagged[0].Data) : nullptr;
+
+	UTEST_NOT_NULL("There is a point data in chosen points", ChosenOutputData);
+	UTEST_NOT_NULL("There is a point data in discarded points", DiscardedOutputData);
+
+	UTEST_EQUAL("There is the right number of points in chosen", ChosenOutputData->GetPoints().Num(), ExpectedNumElementsChosen);
+	UTEST_EQUAL("There is the right number of points in discarded", DiscardedOutputData->GetPoints().Num(), ExpectedNumElementsDiscarded);
+
+	UTEST_TRUE("All points are there and in the right order", PCGRandomChoiceTest::VerifyAllPointsThere(NumOfPoints, ChosenOutputData, DiscardedOutputData));
+
+	return true;
+}
+
+bool FPCGRandomChoiceTest_Ratio::RunTest(const FString& Parameters)
+{
+	PCGTestsCommon::FTestData TestData;
+	UPCGRandomChoiceSettings* Settings = PCGTestsCommon::GenerateSettings<UPCGRandomChoiceSettings>(TestData);
+	check(Settings);
+
+	constexpr int NumOfPoints = 20;
+	constexpr int ExpectedNumElementsChosen = 5;
+	constexpr int ExpectedNumElementsDiscarded = 15;
+
+	Settings->bFixedMode = false;
+	Settings->Ratio = 0.25f;
+
+	TUniquePtr<FPCGContext> Context = TestData.InitializeTestContext();
+	const UPCGPointData* InputPointData = PCGRandomChoiceTest::CreateInputPointData(Context.Get(), NumOfPoints);
+
+	FPCGElementPtr TestElement = TestData.Settings->GetElement();
+	while (!TestElement->Execute(Context.Get())) {}
+
+	TArray<FPCGTaggedData> ChosenOutputTagged = Context->OutputData.GetInputsByPin(PCGRandomChoiceConstants::ChosenPointsLabel);
+	TArray<FPCGTaggedData> DiscardedOutputTagged = Context->OutputData.GetInputsByPin(PCGRandomChoiceConstants::DiscardedPointsLabel);
+
+	const UPCGPointData* ChosenOutputData = ChosenOutputTagged.Num() == 1 ? Cast<const UPCGPointData>(ChosenOutputTagged[0].Data) : nullptr;
+	const UPCGPointData* DiscardedOutputData = DiscardedOutputTagged.Num() == 1 ? Cast<const UPCGPointData>(DiscardedOutputTagged[0].Data) : nullptr;
+
+	UTEST_NOT_NULL("There is a point data in chosen points", ChosenOutputData);
+	UTEST_NOT_NULL("There is a point data in discarded points", DiscardedOutputData);
+
+	UTEST_EQUAL("There is the right number of points in chosen", ChosenOutputData->GetPoints().Num(), ExpectedNumElementsChosen);
+	UTEST_EQUAL("There is the right number of points in discarded", DiscardedOutputData->GetPoints().Num(), ExpectedNumElementsDiscarded);
+
+	UTEST_TRUE("All points are there and in the right order", PCGRandomChoiceTest::VerifyAllPointsThere(NumOfPoints, ChosenOutputData, DiscardedOutputData));
+
+	return true;
+}
+
+bool FPCGRandomChoiceTest_SelectNone::RunTest(const FString& Parameters)
+{
+	PCGTestsCommon::FTestData TestData;
+	UPCGRandomChoiceSettings* Settings = PCGTestsCommon::GenerateSettings<UPCGRandomChoiceSettings>(TestData);
+	check(Settings);
+
+	constexpr int NumOfPoints = 20;
+
+	Settings->bFixedMode = true;
+	Settings->FixedNumber = 0;
+
+	TUniquePtr<FPCGContext> Context = TestData.InitializeTestContext();
+	const UPCGPointData* InputPointData = PCGRandomChoiceTest::CreateInputPointData(Context.Get(), NumOfPoints);
+
+	FPCGElementPtr TestElement = TestData.Settings->GetElement();
+	while (!TestElement->Execute(Context.Get())) {}
+
+	TArray<FPCGTaggedData> ChosenOutputTagged = Context->OutputData.GetInputsByPin(PCGRandomChoiceConstants::ChosenPointsLabel);
+	TArray<FPCGTaggedData> DiscardedOutputTagged = Context->OutputData.GetInputsByPin(PCGRandomChoiceConstants::DiscardedPointsLabel);
+
+	const UPCGPointData* ChosenOutputData = ChosenOutputTagged.Num() == 1 ? Cast<const UPCGPointData>(ChosenOutputTagged[0].Data) : nullptr;
+	const UPCGPointData* DiscardedOutputData = DiscardedOutputTagged.Num() == 1 ? Cast<const UPCGPointData>(DiscardedOutputTagged[0].Data) : nullptr;
+
+	UTEST_NOT_NULL("There is a point data in chosen points", ChosenOutputData);
+	UTEST_NOT_NULL("There is a point data in discarded points", DiscardedOutputData);
+
+	UTEST_EQUAL("Chosen points is empty", ChosenOutputData->GetPoints().Num(), 0);
+	UTEST_EQUAL("Discarded points is the input data", DiscardedOutputData, InputPointData);
+
+	return true;
+}
+
+bool FPCGRandomChoiceTest_SelectAll::RunTest(const FString& Parameters)
+{
+	PCGTestsCommon::FTestData TestData;
+	UPCGRandomChoiceSettings* Settings = PCGTestsCommon::GenerateSettings<UPCGRandomChoiceSettings>(TestData);
+	check(Settings);
+
+	constexpr int NumOfPoints = 20;
+
+	Settings->bFixedMode = false;
+	Settings->Ratio = 1.0f;
+
+	TUniquePtr<FPCGContext> Context = TestData.InitializeTestContext();
+	const UPCGPointData* InputPointData = PCGRandomChoiceTest::CreateInputPointData(Context.Get(), NumOfPoints);
+
+	FPCGElementPtr TestElement = TestData.Settings->GetElement();
+	while (!TestElement->Execute(Context.Get())) {}
+
+	TArray<FPCGTaggedData> ChosenOutputTagged = Context->OutputData.GetInputsByPin(PCGRandomChoiceConstants::ChosenPointsLabel);
+	TArray<FPCGTaggedData> DiscardedOutputTagged = Context->OutputData.GetInputsByPin(PCGRandomChoiceConstants::DiscardedPointsLabel);
+
+	const UPCGPointData* ChosenOutputData = ChosenOutputTagged.Num() == 1 ? Cast<const UPCGPointData>(ChosenOutputTagged[0].Data) : nullptr;
+	const UPCGPointData* DiscardedOutputData = DiscardedOutputTagged.Num() == 1 ? Cast<const UPCGPointData>(DiscardedOutputTagged[0].Data) : nullptr;
+
+	UTEST_NOT_NULL("There is a point data in chosen points", ChosenOutputData);
+	UTEST_NOT_NULL("There is a point data in discarded points", DiscardedOutputData);
+
+	UTEST_EQUAL("Chosen points is the input data", ChosenOutputData, InputPointData);
+	UTEST_EQUAL("Discarded points is empty", DiscardedOutputData->GetPoints().Num(), 0);
+
+	return true;
+}
+
+bool FPCGRandomChoiceTest_NoDiscard::RunTest(const FString& Parameters)
+{
+	PCGTestsCommon::FTestData TestData;
+	UPCGRandomChoiceSettings* Settings = PCGTestsCommon::GenerateSettings<UPCGRandomChoiceSettings>(TestData);
+	check(Settings);
+
+	constexpr int NumOfPoints = 20;
+	constexpr int ExpectedNumElementsChosen = 2;
+
+	Settings->bFixedMode = false;
+	Settings->Ratio = 0.1f;
+	Settings->bOutputDiscardedPoints = false;
+
+	TUniquePtr<FPCGContext> Context = TestData.InitializeTestContext();
+	const UPCGPointData* InputPointData = PCGRandomChoiceTest::CreateInputPointData(Context.Get(), NumOfPoints);
+
+	FPCGElementPtr TestElement = TestData.Settings->GetElement();
+	while (!TestElement->Execute(Context.Get())) {}
+
+	TArray<FPCGTaggedData> ChosenOutputTagged = Context->OutputData.GetInputsByPin(PCGRandomChoiceConstants::ChosenPointsLabel);
+	TArray<FPCGTaggedData> DiscardedOutputTagged = Context->OutputData.GetInputsByPin(PCGRandomChoiceConstants::DiscardedPointsLabel);
+
+	const UPCGPointData* ChosenOutputData = ChosenOutputTagged.Num() == 1 ? Cast<const UPCGPointData>(ChosenOutputTagged[0].Data) : nullptr;
+	const UPCGPointData* DiscardedOutputData = DiscardedOutputTagged.Num() == 1 ? Cast<const UPCGPointData>(DiscardedOutputTagged[0].Data) : nullptr;
+
+	UTEST_NOT_NULL("There is a point data in chosen points", ChosenOutputData);
+	UTEST_NULL("There is no point data in discarded points", DiscardedOutputData);
+
+	UTEST_EQUAL("There is the right number of points in chosen", ChosenOutputData->GetPoints().Num(), ExpectedNumElementsChosen);
+
+	return true;
+}
+#endif // WITH_EDITOR
diff --git a/Engine/Plugins/PCG/Source/PCG/Public/Elements/PCGRandomChoice.h b/Engine/Plugins/PCG/Source/PCG/Public/Elements/PCGRandomChoice.h
--- a/Engine/Plugins/PCG/Source/PCG/Public/Elements/PCGRandomChoice.h
+++ b/Engine/Plugins/PCG/Source/PCG/Public/Elements/PCGRandomChoice.h
@@ -6,8 +6,15 @@
 
 #include "PCGRandomChoice.generated.h"
 
+namespace PCGRandomChoiceConstants
+{
+	const FName ChosenPointsLabel = TEXT("Chosen");
+	const FName DiscardedPointsLabel = TEXT("Discarded");
+}
+
 /**
-* Chooses points randomly through ratio or a fixed number of points
+* Chooses points randomly through ratio or a fixed number of points.
+* Chosen/Discarded Points will be in the same order than they appear in the input data.
 */
 UCLASS(BlueprintType)
 class UPCGRandomChoiceSettings : public UPCGSettings
@@ -15,14 +22,19 @@ class UPCGRandomChoiceSettings : public UPCGSettings
 	GENERATED_BODY()
 
 public:
+	UPCGRandomChoiceSettings();
+
 	//~Begin UPCGSettings interface
 #if WITH_EDITOR
 	virtual FName GetDefaultNodeName() const override;
 	virtual FText GetDefaultNodeTitle() const override;
 	virtual FText GetNodeTooltipText() const override;
-	virtual EPCGSettingsType GetType() const override { return EPCGSettingsType::Sampler; }
+	virtual EPCGSettingsType GetType() const override { return EPCGSettingsType::Filter; }
 #endif
 
+	virtual bool HasDynamicPins() const override { return true; }
+	virtual EPCGDataType GetCurrentPinTypes(const UPCGPin* InPin) const override;
+
 protected:
 	virtual TArray<FPCGPinProperties> InputPinProperties() const override;
 	virtual TArray<FPCGPinProperties> OutputPinProperties() const override;
@@ -31,15 +43,21 @@ protected:
 
 
 public:
+	// Either choose a fixed number of points, or a ratio of points.
 	UPROPERTY(BlueprintReadWrite, EditAnywhere, Category = Settings, meta = (PCG_Overridable))
 	bool bFixedMode = true;
 
-	UPROPERTY(BlueprintReadWrite, EditAnywhere, Category = Settings, meta = (PCG_Overridable, EditCondition = "bFixedMode", ClampMin = "1"))
+	// Defines the number of points to keep.
+	UPROPERTY(BlueprintReadWrite, EditAnywhere, Category = Settings, meta = (PCG_Overridable, EditCondition = "bFixedMode", ClampMin = "0"))
 	int FixedNumber = 1;
 
+	// Defines the ratio of points to keep.
 	UPROPERTY(BlueprintReadWrite, EditAnywhere, Category = Settings, meta = (PCG_Overridable, EditCondition = "!bFixedMode", ClampMin = "0", ClampMax = "1"))
 	float Ratio = 0.5;
 
+	// By default, we output discarded points. If you don't need them, disable this option.
+	UPROPERTY(BlueprintReadWrite, EditAnywhere, Category = Settings)
+	bool bOutputDiscardedPoints = true;
 };
 
 class FPCGRandomChoiceElement : public IPCGElement
