[PCG] Implement Import/Export Text for PCG Attribute Selector

* Fix duplicate issue with serialization
* Also clean up Update to avoid having unwanted characters that could mess with the serialization

#rnx

diff --git a/Engine/Plugins/PCG/Source/PCG/Private/Metadata/PCGAttributePropertySelector.cpp b/Engine/Plugins/PCG/Source/PCG/Private/Metadata/PCGAttributePropertySelector.cpp
--- a/Engine/Plugins/PCG/Source/PCG/Private/Metadata/PCGAttributePropertySelector.cpp
+++ b/Engine/Plugins/PCG/Source/PCG/Private/Metadata/PCGAttributePropertySelector.cpp
@@ -18,6 +18,22 @@ namespace PCGAttributePropertySelectorConstants
 	static const TCHAR ExtraSeparatorChar = ExtraSeparator[0];
 }
 
+bool FPCGAttributePropertySelector::ExportTextItem(FString& ValueStr, FPCGAttributePropertySelector const& DefaultValue, UObject* Parent, int32 PortFlags, UObject* ExportRootScope) const
+{
+	// Quoted string
+	ValueStr = FString(TEXT("\"")) + ToString() + TEXT("\"");
+	return true;
+}
+
+bool FPCGAttributePropertySelector::ImportTextItem(const TCHAR*& Buffer, int32 PortFlags, UObject* Parent, FOutputDevice* ErrorText)
+{
+	FString Token;
+	int NumCharsRead;
+	FParse::QuotedString(Buffer, Token, &NumCharsRead);
+	Update(Token);
+	return true;
+}
+
 FName FPCGAttributePropertySelector::GetName() const
 {
 	switch (Selection)
@@ -117,7 +133,7 @@ bool FPCGAttributePropertySelector::SetExtraProperty(EPCGExtraProperties InExtra
 	}
 }
 
-FText FPCGAttributePropertySelector::GetDisplayText() const
+FString FPCGAttributePropertySelector::ToString() const
 {
 	FString Res;
 	const FName Name = GetName();
@@ -140,7 +156,7 @@ FText FPCGAttributePropertySelector::GetDisplayText() const
 		Res = FString::Join(AllNames, PCGAttributePropertySelectorConstants::ExtraSeparator);
 	}
 
-	return FText::FromString(Res);
+	return Res;
 }
 
 
@@ -192,6 +208,13 @@ void FPCGAttributePropertySelector::ImportFromOtherSelector(const FPCGAttributeP
 bool FPCGAttributePropertySelector::IsValid() const
 {
 	const FName ThisAttributeName = GetAttributeName();
+	static const FName EmptyName = TEXT("");
+
+	if (!ExtraNames.IsEmpty() && ThisAttributeName == EmptyName)
+	{
+		return false;
+	}
+
 	return (Selection != EPCGAttributePropertySelection::Attribute) || 
 		ThisAttributeName == PCGMetadataAttributeConstants::LastAttributeName ||
 		ThisAttributeName == PCGMetadataAttributeConstants::LastCreatedAttributeName ||
@@ -209,7 +232,23 @@ bool FPCGAttributePropertySelector::Update(const FString& NewValue)
 	}
 	else
 	{
-		NewValue.ParseIntoArray(NewValues, PCGAttributePropertySelectorConstants::ExtraSeparator, /*InCullEmpty=*/ false);
+		TArray<FString> NewValuesTemp;
+		NewValue.ParseIntoArray(NewValuesTemp, PCGAttributePropertySelectorConstants::ExtraSeparator, /*InCullEmpty=*/ false);
+
+		// We can't sanitize the first name because of the `@` and `$` that are valid symbols here, but are not for the sanitizer.
+		// Also if we have parsed nothing, just add a single empty string.
+		NewValues.Emplace(NewValuesTemp.IsEmpty() ? FString() : std::move(NewValuesTemp[0]));
+
+		// Then we make sure extra names are not containing non-alphanumerical chars by sanitizing it
+		for (int i = 1; i < NewValuesTemp.Num(); ++i)
+		{
+			FPCGMetadataAttributeBase::SanitizeName(NewValuesTemp[i]);
+
+			if (!NewValuesTemp[i].IsEmpty())
+			{
+				NewValues.Emplace(std::move(NewValuesTemp[i]));
+			}
+		}
 	}
 
 	const FString& NewName = NewValues[0];
diff --git a/Engine/Plugins/PCG/Source/PCG/Public/Metadata/PCGAttributePropertySelector.h b/Engine/Plugins/PCG/Source/PCG/Public/Metadata/PCGAttributePropertySelector.h
--- a/Engine/Plugins/PCG/Source/PCG/Public/Metadata/PCGAttributePropertySelector.h
+++ b/Engine/Plugins/PCG/Source/PCG/Public/Metadata/PCGAttributePropertySelector.h
@@ -65,8 +65,10 @@ public:
 	// Return the name of the selector.
 	FName GetName() const;
 
+	FString ToString() const;
+
 	// Returns the text to display in the widget.
-	FText GetDisplayText() const;
+	FText GetDisplayText() const { return FText::FromString(ToString()); }
 
 	// Return true if the underlying name is valid.
 	bool IsValid() const;
@@ -126,6 +128,9 @@ public:
 
 	friend uint32 GetTypeHash(const FPCGAttributePropertySelector& Selector);
 
+	bool ExportTextItem(FString& ValueStr, FPCGAttributePropertySelector const& DefaultValue, UObject* Parent, int32 PortFlags, UObject* ExportRootScope) const;
+	bool ImportTextItem(const TCHAR*& Buffer, int32 PortFlags, UObject* Parent, FOutputDevice* ErrorText);
+
 protected:
 	UPROPERTY()
 	EPCGAttributePropertySelection Selection = EPCGAttributePropertySelection::Attribute;
@@ -183,22 +188,37 @@ struct PCG_API FPCGAttributePropertyOutputNoSourceSelector : public FPCGAttribut
 	GENERATED_BODY()
 };
 
-// Extra trait to specify to the deserializer that those two classes can be deserialized using the old class.
+
+// All selectors have a custom Export/Import text to use the Update function to serialize/deserialize.
+// It fixes errors of serialization when duplicating/copying.
 template<>
-struct TStructOpsTypeTraits<FPCGAttributePropertyInputSelector> : public TStructOpsTypeTraitsBase2<FPCGAttributePropertyInputSelector>
+struct TStructOpsTypeTraits<FPCGAttributePropertySelector> : public TStructOpsTypeTraitsBase2<FPCGAttributePropertySelector>
+{
+	enum
+	{
+		WithExportTextItem = true,
+		WithImportTextItem = true,
+	};
+};
+
+template<> struct TStructOpsTypeTraits<FPCGAttributePropertyOutputNoSourceSelector> : public TStructOpsTypeTraits<FPCGAttributePropertySelector>{};
+
+// Also extra trait to specify to the deserializer that those two classes can be deserialized using the old class.
+template<> 
+struct TStructOpsTypeTraits<FPCGAttributePropertyInputSelector> : public TStructOpsTypeTraits<FPCGAttributePropertySelector>
 {
 	enum
 	{
-		WithStructuredSerializeFromMismatchedTag = true,
+		WithStructuredSerializeFromMismatchedTag = true
 	};
 };
 
 template<>
-struct TStructOpsTypeTraits<FPCGAttributePropertyOutputSelector> : public TStructOpsTypeTraitsBase2<FPCGAttributePropertyOutputSelector>
+struct TStructOpsTypeTraits<FPCGAttributePropertyOutputSelector> : public TStructOpsTypeTraits<FPCGAttributePropertySelector>
 {
 	enum
 	{
-		WithStructuredSerializeFromMismatchedTag = true,
+		WithStructuredSerializeFromMismatchedTag = true
 	};
 };
 
