[PCG] Metadata operation were not outputing data in the same order in case of no-op

#rnx

diff --git a/Engine/Plugins/PCG/Source/PCG/Private/Elements/Metadata/PCGMetadataOpElementBase.cpp b/Engine/Plugins/PCG/Source/PCG/Private/Elements/Metadata/PCGMetadataOpElementBase.cpp
--- a/Engine/Plugins/PCG/Source/PCG/Private/Elements/Metadata/PCGMetadataOpElementBase.cpp
+++ b/Engine/Plugins/PCG/Source/PCG/Private/Elements/Metadata/PCGMetadataOpElementBase.cpp
@@ -430,6 +430,10 @@ bool FPCGMetadataElementBase::PrepareDataInternal(FPCGContext* Context) const
 		SourceAttribute.SetNum(OperandNum);
 		InputTaggedData.SetNum(OperandNum);
 
+		// Since we add the output data (in CreateAttribute below) if the operation is valid in the PrepareData, if we ever have a no-op, we have to passthrough the inputs now and not in the Execute. So that the order is respected
+		// in the end. (ie. { Input1(Valid), Input2(No-Op), Input3(Valid) } will have in output { Output1, Input2, Output3 } and not { Output1, Output3, Input2 } if we do the passthrough in the Execute.)
+		auto NoOperation = [this, Context, IterationIndex, &Outputs]() { PassthroughInput(Context, Outputs, IterationIndex); return EPCGTimeSliceInitResult::NoOperation; };
+
 		const uint32 PrimaryPinIndex = Settings->GetInputPinToForward();
 
 		// Iterate over the inputs and validate
@@ -466,7 +470,7 @@ bool FPCGMetadataElementBase::PrepareDataInternal(FPCGContext* Context) const
 			{
 				// If we have no data, there is no operation
 				PCGE_LOG(Verbose, LogOnly, FText::Format(LOCTEXT("MissingInputDataForPin", "No data provided on pin '{0}'."), FText::FromName(CurrentPinLabel)));
-				return EPCGTimeSliceInitResult::NoOperation;
+				return NoOperation();
 			}
 			else if (CurrentPinInputData.Num() != 1 && CurrentPinInputData.Num() != OperandInputNumMax)
 			{
@@ -488,7 +492,7 @@ bool FPCGMetadataElementBase::PrepareDataInternal(FPCGContext* Context) const
 				{
 					// If we have no points, there is no operation
 					PCGE_LOG(Verbose, LogOnly, FText::Format(LOCTEXT("NoPointsForPin", "No points in point data provided on pin {0}"), FText::FromName(CurrentPinLabel)));
-					return EPCGTimeSliceInitResult::NoOperation;
+					return NoOperation();
 				}
 			}
 
@@ -522,7 +526,7 @@ bool FPCGMetadataElementBase::PrepareDataInternal(FPCGContext* Context) const
 		if (OperationData.NumberOfElementsToProcess == 0)
 		{
 			PCGE_LOG(Verbose, LogOnly, FText::Format(LOCTEXT("NoElementsInForwardedInput", "No elements in data from forwarded pin '{0}'."), FText::FromName(PrimaryPinData.Pin)));
-			return EPCGTimeSliceInitResult::NoOperation;
+			return NoOperation();
 		}
 
 		// Create the accessors and validate them for each of the other operands
@@ -550,7 +554,7 @@ bool FPCGMetadataElementBase::PrepareDataInternal(FPCGContext* Context) const
 				if (ElementNum == 0)
 				{
 					PCGE_LOG(Verbose, LogOnly, FText::Format(LOCTEXT("NoElementsInInput", "No elements in data from secondary pin '{0}'."), FText::FromName(PrimaryPinData.Pin)));
-					return EPCGTimeSliceInitResult::NoOperation;
+					return NoOperation();
 				}
 
 				// Verify that the number of elements makes sense
@@ -722,7 +726,6 @@ bool FPCGMetadataElementBase::ExecuteInternal(FPCGContext* Context) const
 		// No operation, so skip the iteration.
 		if (Context->GetIterationStateResult(IterationIndex) == EPCGTimeSliceInitResult::NoOperation)
 		{
-			PassthroughInput(Context, Outputs, IterationIndex);
 			return true;
 		}
 
@@ -730,7 +733,6 @@ bool FPCGMetadataElementBase::ExecuteInternal(FPCGContext* Context) const
 		if (!DoOperation(IterState))
 		{
 			PCGE_LOG(Error, GraphAndLog, LOCTEXT("ErrorOccurred", "Error while performing the metadata operation, check logs for more information"));
-			PassthroughInput(Context, Outputs, IterationIndex);
 		}
 
 		return true;
