[PCG] Add PCGNode to the transaction when we delete nodes

* Otherwise, Settings callbacks are not reconnected on Undo/Redo

#rnx

diff --git a/Engine/Plugins/PCG/Source/PCG/Private/PCGGraph.cpp b/Engine/Plugins/PCG/Source/PCG/Private/PCGGraph.cpp
--- a/Engine/Plugins/PCG/Source/PCG/Private/PCGGraph.cpp
+++ b/Engine/Plugins/PCG/Source/PCG/Private/PCGGraph.cpp
@@ -843,6 +843,9 @@ void UPCGGraph::RemoveNodes_Internal(TArrayView<UPCGNode*> InNodes)
 		// We're about to remove InNode, so don't bother triggering updates
 		TouchedNodes.Remove(Node);
 
+		// Add the node to the transaction, to make sure we reconnect everything correctly on Undo/Redo
+		Node->Modify();
+
 		Nodes.Remove(Node);
 	}
 
