[PCG] Enable Dynamic pin behavior for Blueprint nodes

#rnx

diff --git a/Engine/Plugins/PCG/Source/PCG/Private/Elements/PCGExecuteBlueprint.cpp b/Engine/Plugins/PCG/Source/PCG/Private/Elements/PCGExecuteBlueprint.cpp
--- a/Engine/Plugins/PCG/Source/PCG/Private/Elements/PCGExecuteBlueprint.cpp
+++ b/Engine/Plugins/PCG/Source/PCG/Private/Elements/PCGExecuteBlueprint.cpp
@@ -179,6 +179,28 @@ bool UPCGBlueprintElement::IsCacheableOverride_Implementation() const
 	return bIsCacheable;
 }
 
+int32 UPCGBlueprintElement::DynamicPinTypesOverride_Implementation(const UPCGSettings* InSettings, const UPCGPin* InPin) const
+{
+	// Default implementation will modify the Out pin type depending on the input data coming from the In pin. 
+	// If no data arrives, it's not dynamic, or it is another pin, it returns the default allowed type
+	check(InPin);
+
+	if (InSettings->HasDynamicPins() && InPin->IsOutputPin())
+	{
+		const UPCGNode* Node = Cast<const UPCGNode>(InSettings->GetOuter());
+		if (Node && Node->GetInputPin(PCGPinConstants::DefaultInputLabel) != nullptr)
+		{
+			const EPCGDataType InputTypeUnion = InSettings->GetTypeUnionOfIncidentEdges(PCGPinConstants::DefaultInputLabel);
+			if (InputTypeUnion != EPCGDataType::None)
+			{
+				return static_cast<int32>(InputTypeUnion);
+			}
+		}
+	}
+
+	return static_cast<int32>(InPin->Properties.AllowedTypes);
+}
+
 TSet<FName> UPCGBlueprintElement::CustomInputLabels() const
 {
 	TSet<FName> Labels;
@@ -475,6 +497,16 @@ void UPCGBlueprintSettings::RefreshBlueprintElement()
 	InitializeCachedOverridableParams(/*bReset=*/true);
 }
 
+bool UPCGBlueprintSettings::HasDynamicPins() const
+{
+	return BlueprintElementInstance ? BlueprintElementInstance->bHasDynamicPins : UPCGSettings::HasDynamicPins();
+}
+
+EPCGDataType UPCGBlueprintSettings::GetCurrentPinTypes(const UPCGPin* InPin) const
+{
+	return BlueprintElementInstance ? static_cast<EPCGDataType>(BlueprintElementInstance->DynamicPinTypesOverride(this, InPin)) : UPCGSettings::GetCurrentPinTypes(InPin);
+}
+
 #if WITH_EDITOR
 FLinearColor UPCGBlueprintSettings::GetNodeTitleColor() const
 {
diff --git a/Engine/Plugins/PCG/Source/PCG/Public/Elements/PCGExecuteBlueprint.h b/Engine/Plugins/PCG/Source/PCG/Public/Elements/PCGExecuteBlueprint.h
--- a/Engine/Plugins/PCG/Source/PCG/Public/Elements/PCGExecuteBlueprint.h
+++ b/Engine/Plugins/PCG/Source/PCG/Public/Elements/PCGExecuteBlueprint.h
@@ -139,6 +139,10 @@ public:
 	UFUNCTION(BlueprintNativeEvent, Category = "PCG|Node Customization")
 	EPCGSettingsType NodeTypeOverride() const;
 
+	/** If Dynamic Pins is enabled in the BP settings, override this function to provide the type for the given pin. You can use "GetTypeUnionOfIncidentEdges" from the settings to get the union of input types on a given pin. Use the bitwise OR to combine multiple types together.*/
+	UFUNCTION(BlueprintNativeEvent, Category = "PCG|Node Customization")
+	int32 DynamicPinTypesOverride(const UPCGSettings* InSettings, const UPCGPin* InPin) const;
+
 	/** Override for the IsCacheable node property when it depends on the settings in your node. If true, the node will be cached, if not it will always be executed. */
 	UFUNCTION(BlueprintNativeEvent, Category = "PCG|Execution")
 	bool IsCacheableOverride() const;
@@ -247,6 +251,10 @@ public:
 	UPROPERTY(BlueprintReadWrite, EditDefaultsOnly, Category = "Settings|Input & Output")
 	bool bHasDefaultOutPin = true;
 
+	/** If enabled, by default, the Out pin type will have the union of In pin types. Default only works if the pins are In and Out. For custom behavior, implement DynamicPinTypesOverride. */
+	UPROPERTY(BlueprintReadWrite, EditDefaultsOnly, Category = "Settings|Input & Output")
+	bool bHasDynamicPins = false;
+
 #if WITH_EDITORONLY_DATA
 	UPROPERTY(BlueprintReadWrite, EditDefaultsOnly, Category = AssetInfo, AssetRegistrySearchable)
 	bool bExposeToLibrary = false;
@@ -307,6 +315,9 @@ public:
 	virtual TArray<FPCGPreConfiguredSettingsInfo> GetPreconfiguredInfo() const override;
 	virtual bool OnlyExposePreconfiguredSettings() const override;
 #endif
+	virtual bool HasDynamicPins() const override;
+	virtual EPCGDataType GetCurrentPinTypes(const UPCGPin* InPin) const override;
+
 	virtual void ApplyPreconfiguredSettings(const FPCGPreConfiguredSettingsInfo& InPreconfiguredsInfo) override;
 	virtual FString GetAdditionalTitleInformation() const override;
 	virtual bool HasFlippedTitleLines() const override { return true; }
diff --git a/Engine/Plugins/PCG/Source/PCG/Public/PCGPin.h b/Engine/Plugins/PCG/Source/PCG/Public/PCGPin.h
--- a/Engine/Plugins/PCG/Source/PCG/Public/PCGPin.h
+++ b/Engine/Plugins/PCG/Source/PCG/Public/PCGPin.h
@@ -123,6 +123,12 @@ class PCG_API UPCGBlueprintPinHelpers : public UBlueprintFunctionLibrary
 	GENERATED_BODY()
 
 public:
+	UFUNCTION(BlueprintCallable, Category = "PCG|Pins")
+	static int32 GetCorrespondingDataType(EPCGExclusiveDataType InExclusiveDataType)
+	{
+		return static_cast<int32>(GetCorrespondingDataType_Internal(InExclusiveDataType));
+	}
+
 	UFUNCTION(BlueprintPure, Category="PCG|Pins", meta = (NativeBreakFunc))
 	static void BreakPinProperty(const FPCGPinProperties& PinProperty, FName& Label, bool& bAllowMultipleData, bool& bAllowMultipleConnections, bool& bIsAdvancedPin, EPCGExclusiveDataType& AllowedType)
 	{
@@ -154,6 +160,34 @@ public:
 
 	UFUNCTION(BlueprintPure, Category="PCG|Pins", meta = (NativeMakeFunc))
 	static FPCGPinProperties MakePinProperty(FName Label, bool bAllowMultipleData, bool bAllowMultipleConnections, bool bIsAdvancedPin, EPCGExclusiveDataType AllowedType = EPCGExclusiveDataType::Any)
+	{
+		EPCGDataType DataType = GetCorrespondingDataType_Internal(AllowedType);
+
+		FPCGPinProperties PinProperties = FPCGPinProperties(Label, DataType, bAllowMultipleConnections, bAllowMultipleData);
+		if (bIsAdvancedPin)
+		{
+			PinProperties.SetAdvancedPin();
+		}
+
+		return PinProperties;
+	}
+
+	UFUNCTION(BlueprintCallable, Category = "PCG|Pins")
+	static bool IsOfType(int32 AllowedTypes, EPCGExclusiveDataType TypeToCheck)
+	{
+		int32 DataType = GetCorrespondingDataType(TypeToCheck);
+		return !!(AllowedTypes & DataType);
+	}
+
+	UFUNCTION(BlueprintCallable, Category = "PCG|Pins")
+	static bool IsExactlySameType(int32 AllowedTypes, EPCGExclusiveDataType TypeToCheck)
+	{
+		int32 DataType = GetCorrespondingDataType(TypeToCheck);
+		return AllowedTypes == DataType;
+	}
+
+private:
+	static EPCGDataType GetCorrespondingDataType_Internal(EPCGExclusiveDataType InExclusiveDataType)
 	{
 		const UEnum* DataTypeEnum = StaticEnum<EPCGDataType>();
 		const UEnum* ExclusiveDataTypeEnum = StaticEnum<EPCGExclusiveDataType>();
@@ -162,7 +196,7 @@ public:
 
 		if (DataTypeEnum && ExclusiveDataTypeEnum)
 		{
-			FName ExclusiveDataTypeName = ExclusiveDataTypeEnum->GetNameByValue(static_cast<__underlying_type(EPCGExclusiveDataType)>(AllowedType));
+			FName ExclusiveDataTypeName = ExclusiveDataTypeEnum->GetNameByValue(static_cast<__underlying_type(EPCGExclusiveDataType)>(InExclusiveDataType));
 			if (ensure(ExclusiveDataTypeName != NAME_None))
 			{
 				const int64 MatchingType = DataTypeEnum->GetValueByName(ExclusiveDataTypeName);
@@ -173,13 +207,7 @@ public:
 			}
 		}
 
-		FPCGPinProperties PinProperties = FPCGPinProperties(Label, DataType, bAllowMultipleConnections, bAllowMultipleData);
-		if (bIsAdvancedPin)
-		{
-			PinProperties.SetAdvancedPin();
-		}
-
-		return PinProperties;
+		return DataType;
 	}
 };
 
@@ -230,8 +258,12 @@ public:
 	/** Break all edges the connect pins of invalid type. Optionally returns list of affected nodes (including the pin owner). */
 	bool BreakAllIncompatibleEdges(TSet<UPCGNode*>* InTouchedNodes = nullptr);
 
+	UFUNCTION(BlueprintCallable, Category = Settings)
 	bool IsConnected() const;
+
+	UFUNCTION(BlueprintCallable, Category = Settings)
 	bool IsOutputPin() const;
+
 	int32 EdgeCount() const;
 
 	/** Returns the current pin types, which can either be the static types from the pin properties, or a dynamic type based on connected edges. */
diff --git a/Engine/Plugins/PCG/Source/PCG/Public/PCGSettings.h b/Engine/Plugins/PCG/Source/PCG/Public/PCGSettings.h
--- a/Engine/Plugins/PCG/Source/PCG/Public/PCGSettings.h
+++ b/Engine/Plugins/PCG/Source/PCG/Public/PCGSettings.h
@@ -401,6 +401,11 @@ protected:
 	virtual TArray<FPCGPinProperties> InputPinProperties() const;
 	virtual TArray<FPCGPinProperties> OutputPinProperties() const;
 
+	// BP version since EPCGDataType is uint32 and BP only supports uint8 enums.
+	/** Bitwise union of the allowed types of each incident edge on pin. Returns None type if no common bits, or no edges. Use the BP function helpers to extract the types from the result. */
+	UFUNCTION(BlueprintCallable, Category = "Settings|DynamicPins", DisplayName = "Get Type Union Of Incident Edges")
+	int32 BP_GetTypeUnionOfIncidentEdges(const FName& PinLabel) const { return static_cast<int32>(GetTypeUnionOfIncidentEdges(PinLabel)); }
+
 	virtual FPCGElementPtr CreateElement() const PURE_VIRTUAL(UPCGSettings::CreateElement, return nullptr;);
 
 	/** An additional custom version number that external system users can use to track versions. This version will be serialized into the asset and will be provided by UserDataVersion after load. */
