[PCG] Adding extractor menu for selector for documentation and training purposes.

#rnx

diff --git a/Engine/Plugins/PCG/Source/PCGEditor/Private/Details/PCGAttributePropertySelectorDetails.cpp b/Engine/Plugins/PCG/Source/PCGEditor/Private/Details/PCGAttributePropertySelectorDetails.cpp
--- a/Engine/Plugins/PCG/Source/PCGEditor/Private/Details/PCGAttributePropertySelectorDetails.cpp
+++ b/Engine/Plugins/PCG/Source/PCGEditor/Private/Details/PCGAttributePropertySelectorDetails.cpp
@@ -3,6 +3,7 @@
 #include "Details/PCGAttributePropertySelectorDetails.h"
 
 #include "PCGCommon.h"
+#include "Metadata/Accessors/PCGAttributeExtractor.h"
 #include "Metadata/PCGAttributePropertySelector.h"
 #include "Metadata/PCGMetadataAttribute.h"
 
@@ -44,6 +45,81 @@ namespace PCGAttributePropertySelectorDetails
 		}
 	}
 
+	void BuildSubMenuSectionForExtractor(FMenuBuilder& InMenuBuilder, FPCGAttributePropertySelectorDetails* InDetailsObject, const FName ExtractorLabel, const FText& InTooltip = FText{})
+	{
+		InMenuBuilder.AddMenuEntry(
+			FText::FromName(ExtractorLabel),
+			InTooltip,
+			FSlateIcon(),
+			FExecuteAction::CreateSP(InDetailsObject, &FPCGAttributePropertySelectorDetails::AddExtractor, ExtractorLabel),
+			NAME_None,
+			EUserInterfaceActionType::Button
+		);
+	}
+
+	void BuildSubMenuSectionVector(FMenuBuilder& InMenuBuilder, FPCGAttributePropertySelectorDetails* InDetailsObject, bool bIsInput)
+	{
+		InMenuBuilder.AddSubMenu(
+			LOCTEXT("Vector", "Vector"),
+			FText{},
+			FNewMenuDelegate::CreateLambda([InDetailsObject, bIsInput](FMenuBuilder& InSubMenuBuilder)
+		{
+			static const FText TooltipComponents = LOCTEXT("TooltipComponents", "Components can be stacked like XYZ or ZYX.\nZ/B only for Vec3/Vec4, W/A only for Vec4.\nCan't use XYZW and RGBA at the same time.");
+
+			BuildSubMenuSectionForExtractor(InSubMenuBuilder, InDetailsObject, PCGAttributeExtractorConstants::VectorX, TooltipComponents);
+			BuildSubMenuSectionForExtractor(InSubMenuBuilder, InDetailsObject, PCGAttributeExtractorConstants::VectorY, TooltipComponents);
+			BuildSubMenuSectionForExtractor(InSubMenuBuilder, InDetailsObject, PCGAttributeExtractorConstants::VectorZ, TooltipComponents);
+			BuildSubMenuSectionForExtractor(InSubMenuBuilder, InDetailsObject, PCGAttributeExtractorConstants::VectorW, TooltipComponents);
+			BuildSubMenuSectionForExtractor(InSubMenuBuilder, InDetailsObject, PCGAttributeExtractorConstants::VectorR, TooltipComponents);
+			BuildSubMenuSectionForExtractor(InSubMenuBuilder, InDetailsObject, PCGAttributeExtractorConstants::VectorG, TooltipComponents);
+			BuildSubMenuSectionForExtractor(InSubMenuBuilder, InDetailsObject, PCGAttributeExtractorConstants::VectorB, TooltipComponents);
+			BuildSubMenuSectionForExtractor(InSubMenuBuilder, InDetailsObject, PCGAttributeExtractorConstants::VectorA, TooltipComponents);
+
+			if (bIsInput)
+			{
+				BuildSubMenuSectionForExtractor(InSubMenuBuilder, InDetailsObject, PCGAttributeExtractorConstants::VectorLength, LOCTEXT("VectorLengthTooltip", "Get the length of the vector (double value). Read only."));
+				BuildSubMenuSectionForExtractor(InSubMenuBuilder, InDetailsObject, PCGAttributeExtractorConstants::VectorSquaredLength, LOCTEXT("VectorSquaredLengthTooltip", "Get the squared length of the vector (double value). Read only."));
+				BuildSubMenuSectionForExtractor(InSubMenuBuilder, InDetailsObject, PCGAttributeExtractorConstants::VectorNormalized, LOCTEXT("VectorLengthTooltip", "Get the vector normalized (double value). Read only."));
+			}
+		}));
+	}
+
+	void BuildSubMenuSectionRotator(FMenuBuilder& InMenuBuilder, FPCGAttributePropertySelectorDetails* InDetailsObject, bool bIsInput)
+	{
+		InMenuBuilder.AddSubMenu(
+			LOCTEXT("Rotator", "Rotator"),
+			FText{},
+			FNewMenuDelegate::CreateLambda([InDetailsObject, bIsInput](FMenuBuilder& InSubMenuBuilder)
+		{
+			static const FText TooltipComponents = LOCTEXT("TooltipComponentsRotator", "Euler angles");
+
+			BuildSubMenuSectionForExtractor(InSubMenuBuilder, InDetailsObject, PCGAttributeExtractorConstants::RotatorRoll, TooltipComponents);
+			BuildSubMenuSectionForExtractor(InSubMenuBuilder, InDetailsObject, PCGAttributeExtractorConstants::RotatorPitch, TooltipComponents);
+			BuildSubMenuSectionForExtractor(InSubMenuBuilder, InDetailsObject, PCGAttributeExtractorConstants::RotatorYaw, TooltipComponents);
+
+			if (bIsInput)
+			{
+				BuildSubMenuSectionForExtractor(InSubMenuBuilder, InDetailsObject, PCGAttributeExtractorConstants::RotatorForward, LOCTEXT("RotatorForwardTooltip", "X Axis of the rotation. Read only."));
+				BuildSubMenuSectionForExtractor(InSubMenuBuilder, InDetailsObject, PCGAttributeExtractorConstants::RotatorRight, LOCTEXT("RotatorRightTooltip", "Y Axis of the rotation. Read only."));
+				BuildSubMenuSectionForExtractor(InSubMenuBuilder, InDetailsObject, PCGAttributeExtractorConstants::RotatorUp, LOCTEXT("RotatorUpTooltip", "Z Axis of the rotation. Read only."));
+			}
+		}));
+	}
+
+	void BuildSubMenuSectionTransform(FMenuBuilder& InMenuBuilder, FPCGAttributePropertySelectorDetails* InDetailsObject, bool bIsInput)
+	{
+		InMenuBuilder.AddSubMenu(
+			LOCTEXT("Transform", "Transform"),
+			FText{},
+			FNewMenuDelegate::CreateLambda([InDetailsObject](FMenuBuilder& InSubMenuBuilder)
+		{
+
+			BuildSubMenuSectionForExtractor(InSubMenuBuilder, InDetailsObject, PCGAttributeExtractorConstants::TransformLocation);
+			BuildSubMenuSectionForExtractor(InSubMenuBuilder, InDetailsObject, PCGAttributeExtractorConstants::TransformRotation);
+			BuildSubMenuSectionForExtractor(InSubMenuBuilder, InDetailsObject, PCGAttributeExtractorConstants::TransformScale);
+		}));
+	}
+
 	// Func signature: bool(FPCGAttributePropertySelector*, T)
 	template <typename T, typename Func>
 	void SetValue(FPCGAttributePropertySelector* InSelector, TSharedPtr<IPropertyHandle>& InPropertyHandle, T InValue, Func InCallback)
@@ -69,7 +145,8 @@ void FPCGAttributePropertySelectorDetails::CustomizeHeader(TSharedRef<IPropertyH
 	PropertyHandle = InPropertyHandle;
 	check(PropertyHandle.IsValid());
 
-	FText TooltipText = LOCTEXT("TooltipText", "Enter the name of the attribute. You can prefix it by '$' to get a property. Red text means that your attribute or property is invalid.");
+	static FText TooltipText = LOCTEXT("TooltipText", "Enter the name of the attribute. You can prefix it by '$' to get a property. Red text means that your attribute or property is invalid.\n"
+		"You can also postfix with `.` to extract values. For example, $Position.X will extract the X component of the position. They can also be chained. You can see the list under the '+'.");
 
 	auto Validation = [this]() -> FSlateColor
 	{
@@ -159,6 +236,17 @@ TSharedRef<SWidget> FPCGAttributePropertySelectorDetails::GenerateExtraMenu()
 	FMenuBuilder MenuBuilder(true, nullptr);
 
 	FStructProperty* StructProperty = CastFieldChecked<FStructProperty>(PropertyHandle->GetProperty());
+	const bool bIsInput = StructProperty->Struct == FPCGAttributePropertyInputSelector::StaticStruct();
+
+	MenuBuilder.AddSubMenu(
+		LOCTEXT("AddAnExtractor", "Add an extractor"),
+		LOCTEXT("AddAnExtractorTooltip", "List of all the extractor to add to the current attribute."),
+		FNewMenuDelegate::CreateLambda([this, bIsInput](FMenuBuilder& InSubMenuBuilder)
+	{
+		PCGAttributePropertySelectorDetails::BuildSubMenuSectionVector(InSubMenuBuilder, this, bIsInput);
+		PCGAttributePropertySelectorDetails::BuildSubMenuSectionRotator(InSubMenuBuilder, this, bIsInput);
+		PCGAttributePropertySelectorDetails::BuildSubMenuSectionTransform(InSubMenuBuilder, this, bIsInput);
+	}));
 
 	MenuBuilder.BeginSection("Attributes", LOCTEXT("AttributesHeader", "Attributes"));
 	{
@@ -257,4 +345,18 @@ void FPCGAttributePropertySelectorDetails::SetExtraProperty(EPCGExtraProperties
 	PCGAttributePropertySelectorDetails::SetValue(GetStruct(), PropertyHandle, EnumValue, [](FPCGAttributePropertySelector* InStruct, EPCGExtraProperties InValue) -> bool { return InStruct && InStruct->SetExtraProperty(InValue); });
 }
 
+void FPCGAttributePropertySelectorDetails::AddExtractor(FName InExtractor)
+{
+	PCGAttributePropertySelectorDetails::SetValue(GetStruct(), PropertyHandle, InExtractor, [](FPCGAttributePropertySelector* InStruct, FName InValue) -> bool 
+	{
+		if (!InStruct)
+		{
+			return false;
+		}
+
+		InStruct->GetExtraNamesMutable().Add(InValue.ToString());
+		return true;
+	});
+}
+
 #undef LOCTEXT_NAMESPACE
diff --git a/Engine/Plugins/PCG/Source/PCGEditor/Public/Details/PCGAttributePropertySelectorDetails.h b/Engine/Plugins/PCG/Source/PCGEditor/Public/Details/PCGAttributePropertySelectorDetails.h
--- a/Engine/Plugins/PCG/Source/PCGEditor/Public/Details/PCGAttributePropertySelectorDetails.h
+++ b/Engine/Plugins/PCG/Source/PCGEditor/Public/Details/PCGAttributePropertySelectorDetails.h
@@ -25,6 +25,8 @@ public:
 	virtual void CustomizeChildren(TSharedRef<IPropertyHandle> InPropertyHandle, IDetailChildrenBuilder& ChildBuilder, IPropertyTypeCustomizationUtils& CustomizationUtils) override {};
 	/** ~End IPropertyTypeCustomization interface */
 
+	void AddExtractor(FName InExtractor);
+
 protected:
 	FPCGAttributePropertySelector* GetStruct();
 	const FPCGAttributePropertySelector* GetStruct() const;
