[PCG] Improve PointOperationElementBase to include an optional pin to select inputs rather than defaulting to all spatial

#rnx

diff --git a/Engine/Plugins/PCG/Source/PCG/Private/Elements/PCGPointOperationElementBase.cpp b/Engine/Plugins/PCG/Source/PCG/Private/Elements/PCGPointOperationElementBase.cpp
--- a/Engine/Plugins/PCG/Source/PCG/Private/Elements/PCGPointOperationElementBase.cpp
+++ b/Engine/Plugins/PCG/Source/PCG/Private/Elements/PCGPointOperationElementBase.cpp
@@ -17,20 +17,20 @@ bool FPCGPointOperationElementBase::PrepareDataInternal(FPCGContext* Context) co
 	return PreparePointOperationData(PointProcessContext);
 }
 
-bool FPCGPointOperationElementBase::PreparePointOperationData(ContextType* Context) const
+bool FPCGPointOperationElementBase::PreparePointOperationData(ContextType* InContext, FName InputPinLabel) const
 {
 	TRACE_CPUPROFILER_EVENT_SCOPE(FPCGPointOperationElementBase::PreparePointProcessing);
-	check(Context);
+	check(InContext);
 
-	const TArray<FPCGTaggedData>& Inputs = Context->InputData.GetInputs();
-	TArray<FPCGTaggedData>& Outputs = Context->OutputData.TaggedData;
+	const TArray<FPCGTaggedData>& Inputs = InContext->InputData.GetInputsByPin(InputPinLabel);
+	TArray<FPCGTaggedData>& Outputs = InContext->OutputData.TaggedData;
 
 	// There is no execution state, so this just flags that its okay to continue
-	Context->InitializePerExecutionState();
+	InContext->InitializePerExecutionState();
 
 	// Prepare the 'per iteration' time slice context state and allocate output point data
-	Context->InitializePerIterationStates(Inputs.Num(),
-		[this, &Context, &Inputs, &Outputs](IterStateType& OutState, const ExecStateType& ExecState, const uint32 IterationIndex)
+	InContext->InitializePerIterationStates(Inputs.Num(),
+		[this, &InContext, &Inputs, &Outputs](IterStateType& OutState, const ExecStateType& ExecState, const uint32 IterationIndex)
 		{
 			TRACE_CPUPROFILER_EVENT_SCOPE(FPCGPointOperationElementBase::InitializePerIterationStates);
 
@@ -40,15 +40,15 @@ bool FPCGPointOperationElementBase::PreparePointOperationData(ContextType* Conte
 
 			if (!SpatialData)
 			{
-				PCGE_LOG(Error, GraphAndLog, LOCTEXT("InputMissingSpatialData", "Unable to get Spatial data from input"));
+				PCGLog::LogErrorOnGraph(LOCTEXT("InputMissingSpatialData", "Unable to get Spatial data from input"), InContext);
 				return EPCGTimeSliceInitResult::NoOperation;
 			}
 
-			OutState.InputPointData = SpatialData->ToPointData(Context);
+			OutState.InputPointData = SpatialData->ToPointData(InContext);
 
 			if (!OutState.InputPointData)
 			{
-				PCGE_LOG(Error, GraphAndLog, LOCTEXT("InputMissingPointData", "Unable to get Point data from input"));
+				PCGLog::LogErrorOnGraph(LOCTEXT("InputMissingPointData", "Unable to get Point data from input"), InContext);
 				return EPCGTimeSliceInitResult::NoOperation;
 			}
 
diff --git a/Engine/Plugins/PCG/Source/PCG/Public/Elements/PCGPointOperationElementBase.h b/Engine/Plugins/PCG/Source/PCG/Public/Elements/PCGPointOperationElementBase.h
--- a/Engine/Plugins/PCG/Source/PCG/Public/Elements/PCGPointOperationElementBase.h
+++ b/Engine/Plugins/PCG/Source/PCG/Public/Elements/PCGPointOperationElementBase.h
@@ -40,7 +40,6 @@ protected:
 	/** Executes the PointFunction function/lambda for every point copied from PreparePointOperationData */
 	bool ExecutePointOperation(ContextType* Context, TFunctionRef<PointExecSignature> PointFunction, int32 PointsPerChunk = PCGPointOperation::Constants::PointsPerChunk) const;
 
-private:
 	/** Mandatory call. Using the context, prepares the state data for time slice execution */
-	bool PreparePointOperationData(ContextType* Context) const;
+	bool PreparePointOperationData(ContextType* Context, FName InputPinLabel = PCGPinConstants::DefaultInputLabel) const;
 };
\ No newline at end of file
