[PCG] Fix issue with Dynamic Pins function in custom BP that was called on PostLoad

* Never call it in post load or if we don't have dynamic pins
* Returning the current pin types (saved on the pin so it won't break the edges on post load).

#rnx

diff --git a/Engine/Plugins/PCG/Source/PCG/Private/Elements/PCGExecuteBlueprint.cpp b/Engine/Plugins/PCG/Source/PCG/Private/Elements/PCGExecuteBlueprint.cpp
--- a/Engine/Plugins/PCG/Source/PCG/Private/Elements/PCGExecuteBlueprint.cpp
+++ b/Engine/Plugins/PCG/Source/PCG/Private/Elements/PCGExecuteBlueprint.cpp
@@ -12,6 +12,7 @@
 #include "Engine/World.h"
 #include "PCGPin.h"
 #include "UObject/Package.h"
+#include "UObject/UObjectThreadContext.h"
 
 #include UE_INLINE_GENERATED_CPP_BY_NAME(PCGExecuteBlueprint)
 
@@ -504,7 +505,17 @@ bool UPCGBlueprintSettings::HasDynamicPins() const
 
 EPCGDataType UPCGBlueprintSettings::GetCurrentPinTypes(const UPCGPin* InPin) const
 {
-	return BlueprintElementInstance ? static_cast<EPCGDataType>(BlueprintElementInstance->DynamicPinTypesOverride(this, InPin)) : UPCGSettings::GetCurrentPinTypes(InPin);
+	check(InPin);
+
+	// We can't call a BP function while we are PostLoading, in that case (or if we don't have an instance or it's not dynamic) just return the pin allowed types.
+	if (BlueprintElementInstance && BlueprintElementInstance->bHasDynamicPins && !FUObjectThreadContext::Get().IsRoutingPostLoad)
+	{
+		return static_cast<EPCGDataType>(BlueprintElementInstance->DynamicPinTypesOverride(this, InPin));
+	}
+	else
+	{
+		return InPin->Properties.AllowedTypes;
+	}
 }
 
 #if WITH_EDITOR
