[PCG] Allow attributes with spaces to be filtered by [REDACTED_PERSON]

* Deprecation set up to usher old nodes into an update before 5.6 with a warning to users

diff --git a/Engine/Plugins/PCG/Source/PCG/Private/Data/PCGWorldData.cpp b/Engine/Plugins/PCG/Source/PCG/Private/Data/PCGWorldData.cpp
--- a/Engine/Plugins/PCG/Source/PCG/Private/Data/PCGWorldData.cpp
+++ b/Engine/Plugins/PCG/Source/PCG/Private/Data/PCGWorldData.cpp
@@ -29,7 +29,7 @@ void FPCGWorldCommonQueryParams::Initialize()
 	}
 	else
 	{
-		TArray<FString> ParsedList = PCGHelpers::GetStringArrayFromCommaSeparatedString(ActorTagsList);
+		TArray<FString> ParsedList = PCGHelpers::GetStringArrayFromCommaSeparatedList(ActorTagsList);
 		ParsedActorTagsList.Reset();
 		for (const FString& Tag : ParsedList)
 		{
diff --git a/Engine/Plugins/PCG/Source/PCG/Private/Elements/Metadata/PCGMetadataPartition.cpp b/Engine/Plugins/PCG/Source/PCG/Private/Elements/Metadata/PCGMetadataPartition.cpp
--- a/Engine/Plugins/PCG/Source/PCG/Private/Elements/Metadata/PCGMetadataPartition.cpp
+++ b/Engine/Plugins/PCG/Source/PCG/Private/Elements/Metadata/PCGMetadataPartition.cpp
@@ -3,7 +3,7 @@
 #include "Elements/Metadata/PCGMetadataPartition.h"
 
 #include "PCGContext.h"
-#include "PCGModule.h"
+#include "PCGCustomVersion.h"
 #include "Helpers/PCGHelpers.h"
 #include "Metadata/PCGMetadataAttribute.h"
 #include "Metadata/PCGMetadataPartitionCommon.h"
@@ -30,6 +30,7 @@ void UPCGMetadataPartitionSettings::PostLoad()
 	Super::PostLoad();
 
 #if WITH_EDITOR
+	PRAGMA_DISABLE_DEPRECATION_WARNINGS
 	if (PartitionAttribute_DEPRECATED != NAME_None)
 	{
 		PartitionAttributeSelectors.Empty();
@@ -43,9 +44,24 @@ void UPCGMetadataPartitionSettings::PostLoad()
 		PartitionAttributeSelectors.Emplace(PartitionAttributeSource_DEPRECATED);
 		PartitionAttributeSource_DEPRECATED = FPCGAttributePropertyInputSelector();
 	}
+	PRAGMA_ENABLE_DEPRECATION_WARNINGS
 #endif // WITH_EDITOR
 }
 
+#if WITH_EDITOR
+void UPCGMetadataPartitionSettings::ApplyDeprecation(UPCGNode* InOutNode)
+{
+	PRAGMA_DISABLE_DEPRECATION_WARNINGS
+	if (DataVersion < FPCGCustomVersion::AttributesAndTagsCanContainSpaces)
+	{
+		bTokenizeOnWhiteSpace = true;
+	}
+	PRAGMA_ENABLE_DEPRECATION_WARNINGS
+
+	Super::ApplyDeprecation(InOutNode);
+}
+#endif // WITH_EDITOR
+
 FString UPCGMetadataPartitionSettings::GetAdditionalTitleInformation() const
 {
 #if WITH_EDITOR
@@ -81,7 +97,13 @@ bool FPCGMetadataPartitionElement::ExecuteInternal(FPCGContext* Context) const
 
 	// TODO: This is a temporary solution for overrides until arrays are supported
 	TArray<FPCGAttributePropertyInputSelector> OverriddenSelectors;
-	const TArray<FString> AttributeNames = PCGHelpers::GetStringArrayFromCommaSeparatedString(Settings->PartitionAttributeNames);
+
+	PRAGMA_DISABLE_DEPRECATION_WARNINGS
+	const TArray<FString> AttributeNames = Settings->bTokenizeOnWhiteSpace
+		? PCGHelpers::GetStringArrayFromCommaSeparatedString(Settings->PartitionAttributeNames, Context)
+		: PCGHelpers::GetStringArrayFromCommaSeparatedList(Settings->PartitionAttributeNames);
+	PRAGMA_ENABLE_DEPRECATION_WARNINGS
+
 	// If the names are overridden by the user, generate the selectors with them
 	OverriddenSelectors.SetNum(AttributeNames.Num());
 	for (const FString& AttributeName : AttributeNames)
diff --git a/Engine/Plugins/PCG/Source/PCG/Private/Elements/PCGAddTag.cpp b/Engine/Plugins/PCG/Source/PCG/Private/Elements/PCGAddTag.cpp
--- a/Engine/Plugins/PCG/Source/PCG/Private/Elements/PCGAddTag.cpp
+++ b/Engine/Plugins/PCG/Source/PCG/Private/Elements/PCGAddTag.cpp
@@ -3,11 +3,26 @@
 #include "Elements/PCGAddTag.h"
 
 #include "PCGContext.h"
+#include "PCGCustomVersion.h"
 #include "PCGPin.h"
 #include "Helpers/PCGHelpers.h"
 
 #define LOCTEXT_NAMESPACE "PCGAddTagElement"
 
+#if WITH_EDITOR
+void UPCGAddTagSettings::ApplyDeprecation(UPCGNode* InOutNode)
+{
+	PRAGMA_DISABLE_DEPRECATION_WARNINGS
+	if (DataVersion < FPCGCustomVersion::AttributesAndTagsCanContainSpaces)
+	{
+		bTokenizeOnWhiteSpace = true;
+	}
+	PRAGMA_ENABLE_DEPRECATION_WARNINGS
+
+	Super::ApplyDeprecation(InOutNode);
+}
+#endif // WITH_EDITOR
+
 TArray<FPCGPinProperties> UPCGAddTagSettings::OutputPinProperties() const
 {
 	TArray<FPCGPinProperties> PinProperties;
@@ -32,8 +47,12 @@ bool FPCGAddTagElement::ExecuteInternal(FPCGContext* Context) const
 	
 	Context->OutputData.TaggedData = Context->InputData.GetInputsByPin(PCGPinConstants::DefaultInputLabel);
 
-	const TArray<FString> TagsArray = PCGHelpers::GetStringArrayFromCommaSeparatedString(Settings->TagsToAdd);
-	
+    PRAGMA_DISABLE_DEPRECATION_WARNINGS
+	const TArray<FString> TagsArray = Settings->bTokenizeOnWhiteSpace
+		? PCGHelpers::GetStringArrayFromCommaSeparatedString(Settings->TagsToAdd, Context)
+		: PCGHelpers::GetStringArrayFromCommaSeparatedList(Settings->TagsToAdd);
+    PRAGMA_ENABLE_DEPRECATION_WARNINGS
+
 	for (const FString& Tag : TagsArray)
 	{
 		for (FPCGTaggedData& OutputTaggedData : Context->OutputData.TaggedData)
diff --git a/Engine/Plugins/PCG/Source/PCG/Private/Elements/PCGDeleteAttributesElement.cpp b/Engine/Plugins/PCG/Source/PCG/Private/Elements/PCGDeleteAttributesElement.cpp
--- a/Engine/Plugins/PCG/Source/PCG/Private/Elements/PCGDeleteAttributesElement.cpp
+++ b/Engine/Plugins/PCG/Source/PCG/Private/Elements/PCGDeleteAttributesElement.cpp
@@ -3,6 +3,7 @@
 #include "Elements/PCGDeleteAttributesElement.h"
 
 #include "PCGContext.h"
+#include "PCGCustomVersion.h"
 #include "PCGParamData.h"
 #include "PCGPin.h"
 #include "Data/PCGSpatialData.h"
@@ -51,6 +52,18 @@ void UPCGDeleteAttributesSettings::PostLoad()
 }
 
 #if WITH_EDITOR
+void UPCGDeleteAttributesSettings::ApplyDeprecation(UPCGNode* InOutNode)
+{
+	PRAGMA_DISABLE_DEPRECATION_WARNINGS
+	if (DataVersion < FPCGCustomVersion::AttributesAndTagsCanContainSpaces)
+	{
+		bTokenizeOnWhiteSpace = true;
+	}
+	PRAGMA_ENABLE_DEPRECATION_WARNINGS
+
+	Super::ApplyDeprecation(InOutNode);
+}
+
 FName UPCGDeleteAttributesSettings::GetDefaultNodeName() const
 {
 	return PCGAttributeFilterConstants::NodeName;
@@ -84,7 +97,12 @@ FString UPCGDeleteAttributesSettings::GetAdditionalTitleInformation() const
 		ensureMsgf(false, TEXT("Unrecognized operation"));
 	}
 
-	TArray<FString> AttributesToKeep = PCGHelpers::GetStringArrayFromCommaSeparatedString(SelectedAttributes);
+	PRAGMA_DISABLE_DEPRECATION_WARNINGS
+	const TArray<FString> AttributesToKeep = bTokenizeOnWhiteSpace
+		? PCGHelpers::GetStringArrayFromCommaSeparatedString(SelectedAttributes)
+		: PCGHelpers::GetStringArrayFromCommaSeparatedList(SelectedAttributes);
+	PRAGMA_ENABLE_DEPRECATION_WARNINGS
+
 	if (AttributesToKeep.Num() == 1)
 	{
 		return FString::Printf(TEXT("%s (%s)"), *OperationString, *AttributesToKeep[0]);
@@ -123,8 +141,14 @@ bool FPCGDeleteAttributesElement::ExecuteInternal(FPCGContext* Context) const
 	const bool bAddAttributesFromParent = (Settings->Operation == EPCGAttributeFilterOperation::DeleteSelectedAttributes);
 	const EPCGMetadataFilterMode FilterMode = bAddAttributesFromParent ? EPCGMetadataFilterMode::ExcludeAttributes : EPCGMetadataFilterMode::IncludeAttributes;
 
-	TSet<FName> AttributesToFilter; 
-	const TArray<FString> FilterAttributes = PCGHelpers::GetStringArrayFromCommaSeparatedString(Settings->SelectedAttributes);
+	TSet<FName> AttributesToFilter;
+
+	PRAGMA_DISABLE_DEPRECATION_WARNINGS
+	const TArray<FString> FilterAttributes = Settings->bTokenizeOnWhiteSpace
+	   ? PCGHelpers::GetStringArrayFromCommaSeparatedString(Settings->SelectedAttributes, Context)
+	   : PCGHelpers::GetStringArrayFromCommaSeparatedList(Settings->SelectedAttributes);
+	PRAGMA_ENABLE_DEPRECATION_WARNINGS
+
 	for (const FString& FilterAttribute : FilterAttributes)
 	{
 		AttributesToFilter.Add(FName(*FilterAttribute));
diff --git a/Engine/Plugins/PCG/Source/PCG/Private/Elements/PCGDeleteTags.cpp b/Engine/Plugins/PCG/Source/PCG/Private/Elements/PCGDeleteTags.cpp
--- a/Engine/Plugins/PCG/Source/PCG/Private/Elements/PCGDeleteTags.cpp
+++ b/Engine/Plugins/PCG/Source/PCG/Private/Elements/PCGDeleteTags.cpp
@@ -3,12 +3,25 @@
 #include "Elements/PCGDeleteTags.h"
 
 #include "PCGContext.h"
+#include "PCGCustomVersion.h"
 #include "PCGData.h"
 #include "Helpers/PCGHelpers.h"
 
 #define LOCTEXT_NAMESPACE "PCGDeleteTagsElement"
 
 #if WITH_EDITOR
+void UPCGDeleteTagsSettings::ApplyDeprecation(UPCGNode* InOutNode)
+{
+	PRAGMA_DISABLE_DEPRECATION_WARNINGS
+	if (DataVersion < FPCGCustomVersion::AttributesAndTagsCanContainSpaces)
+	{
+		bTokenizeOnWhiteSpace = true;
+	}
+	PRAGMA_ENABLE_DEPRECATION_WARNINGS
+
+	Super::ApplyDeprecation(InOutNode);
+}
+
 FName UPCGDeleteTagsSettings::GetDefaultNodeName() const
 {
 	return TEXT("DeleteTags");
@@ -25,7 +38,12 @@ FString UPCGDeleteTagsSettings::GetAdditionalTitleInformation() const
 	if (const UEnum* SelectionEnum = StaticEnum<EPCGTagFilterOperation>())
 	{
 		FText OperationText = SelectionEnum->GetDisplayNameTextByValue(static_cast<int64>(Operation));
-		TArray<FString> TagsToProcess = PCGHelpers::GetStringArrayFromCommaSeparatedString(SelectedTags);
+
+		PRAGMA_DISABLE_DEPRECATION_WARNINGS
+		const TArray<FString> TagsToProcess = bTokenizeOnWhiteSpace
+			? PCGHelpers::GetStringArrayFromCommaSeparatedString(SelectedTags)
+			: PCGHelpers::GetStringArrayFromCommaSeparatedList(SelectedTags);
+		PRAGMA_ENABLE_DEPRECATION_WARNINGS
 
 		if (TagsToProcess.Num() == 1)
 		{
@@ -63,7 +81,12 @@ bool FPCGDeleteTagsElement::ExecuteInternal(FPCGContext* Context) const
 	const UPCGDeleteTagsSettings* Settings = Context->GetInputSettings<UPCGDeleteTagsSettings>();
 	check(Settings);
 
-	const TArray<FString> FilterTags = PCGHelpers::GetStringArrayFromCommaSeparatedString(Settings->SelectedTags);
+	PRAGMA_DISABLE_DEPRECATION_WARNINGS
+	const TArray<FString> FilterTags = Settings->bTokenizeOnWhiteSpace
+		? PCGHelpers::GetStringArrayFromCommaSeparatedString(Settings->SelectedTags, Context)
+		: PCGHelpers::GetStringArrayFromCommaSeparatedList(Settings->SelectedTags);
+	PRAGMA_ENABLE_DEPRECATION_WARNINGS
+
 	TSet<FString> TagsToFilter(FilterTags);
 
 	const bool bKeepInFilter = (Settings->Operation == EPCGTagFilterOperation::KeepOnlySelectedTags);
diff --git a/Engine/Plugins/PCG/Source/PCG/Private/Elements/PCGFilterByTag.cpp b/Engine/Plugins/PCG/Source/PCG/Private/Elements/PCGFilterByTag.cpp
--- a/Engine/Plugins/PCG/Source/PCG/Private/Elements/PCGFilterByTag.cpp
+++ b/Engine/Plugins/PCG/Source/PCG/Private/Elements/PCGFilterByTag.cpp
@@ -2,7 +2,7 @@
 
 #include "Elements/PCGFilterByTag.h"
 #include "PCGContext.h"
-#include "PCGPin.h"
+#include "PCGCustomVersion.h"
 #include "Helpers/PCGHelpers.h"
 
 #include UE_INLINE_GENERATED_CPP_BY_NAME(PCGFilterByTag)
@@ -10,6 +10,18 @@
 #define LOCTEXT_NAMESPACE "PCGFilterByTag"
 
 #if WITH_EDITOR
+void UPCGFilterByTagSettings::ApplyDeprecation(UPCGNode* InOutNode)
+{
+	PRAGMA_DISABLE_DEPRECATION_WARNINGS
+	if (DataVersion < FPCGCustomVersion::AttributesAndTagsCanContainSpaces)
+	{
+		bTokenizeOnWhiteSpace = true;
+	}
+	PRAGMA_ENABLE_DEPRECATION_WARNINGS
+
+	Super::ApplyDeprecation(InOutNode);
+}
+
 FName UPCGFilterByTagSettings::GetDefaultNodeName() const
 {
 	return FName(TEXT("FilterDataByTag"));
@@ -28,7 +40,11 @@ FText UPCGFilterByTagSettings::GetNodeTooltipText() const
 
 FString UPCGFilterByTagSettings::GetAdditionalTitleInformation() const
 {
-	const TArray<FString> Tags = PCGHelpers::GetStringArrayFromCommaSeparatedString(SelectedTags);
+	PRAGMA_DISABLE_DEPRECATION_WARNINGS
+	const TArray<FString> Tags = bTokenizeOnWhiteSpace
+		? PCGHelpers::GetStringArrayFromCommaSeparatedString(SelectedTags)
+		: PCGHelpers::GetStringArrayFromCommaSeparatedList(SelectedTags);
+	PRAGMA_ENABLE_DEPRECATION_WARNINGS
 
 	const FString Prefix = (Operation == EPCGFilterByTagOperation::KeepTagged ? TEXT("Tag (Keep):") : TEXT("Tag (Remove):"));
 
@@ -59,7 +75,12 @@ bool FPCGFilterByTagElement::ExecuteInternal(FPCGContext* Context) const
 	const UPCGFilterByTagSettings* Settings = Context->GetInputSettings<UPCGFilterByTagSettings>();
 
 	const bool bKeepIfTag = (Settings->Operation == EPCGFilterByTagOperation::KeepTagged);
-	const TArray<FString> Tags = PCGHelpers::GetStringArrayFromCommaSeparatedString(Settings->SelectedTags);
+
+	PRAGMA_DISABLE_DEPRECATION_WARNINGS
+	const TArray<FString> Tags = Settings->bTokenizeOnWhiteSpace
+		? PCGHelpers::GetStringArrayFromCommaSeparatedString(Settings->SelectedTags, Context)
+		: PCGHelpers::GetStringArrayFromCommaSeparatedList(Settings->SelectedTags);
+	PRAGMA_ENABLE_DEPRECATION_WARNINGS
 
 	TArray<FPCGTaggedData> Inputs = Context->InputData.GetInputsByPin(PCGPinConstants::DefaultInputLabel);
 	TArray<FPCGTaggedData>& Outputs = Context->OutputData.TaggedData;
diff --git a/Engine/Plugins/PCG/Source/PCG/Private/Elements/PCGLoopElement.cpp b/Engine/Plugins/PCG/Source/PCG/Private/Elements/PCGLoopElement.cpp
--- a/Engine/Plugins/PCG/Source/PCG/Private/Elements/PCGLoopElement.cpp
+++ b/Engine/Plugins/PCG/Source/PCG/Private/Elements/PCGLoopElement.cpp
@@ -35,6 +35,13 @@ void UPCGLoopSettings::ApplyDeprecation(UPCGNode* InOutNode)
 		bUseGraphDefaultPinUsage = LoopPins.IsEmpty();
 	}
 
+	PRAGMA_DISABLE_DEPRECATION_WARNINGS
+	if (DataVersion < FPCGCustomVersion::AttributesAndTagsCanContainSpaces)
+	{
+		bTokenizeOnWhiteSpace = true;
+	}
+	PRAGMA_ENABLE_DEPRECATION_WARNINGS
+
 	Super::ApplyDeprecation(InOutNode);
 }
 
@@ -112,13 +119,21 @@ void UPCGLoopSettings::GetLoopPinNames(FPCGContext* Context, TArray<FName>& Loop
 	}
 	else
 	{
-		TArray<FString> PinsFromSettings = PCGHelpers::GetStringArrayFromCommaSeparatedString(LoopPins);
+		PRAGMA_DISABLE_DEPRECATION_WARNINGS
+		TArray<FString> PinsFromSettings = bTokenizeOnWhiteSpace
+			? PCGHelpers::GetStringArrayFromCommaSeparatedString(LoopPins, Context)
+			: PCGHelpers::GetStringArrayFromCommaSeparatedList(LoopPins);
+
 		for (const FString& PinLabel : PinsFromSettings)
 		{
 			LoopPinNames.Emplace(PinLabel);
 		}
 
-		PinsFromSettings = PCGHelpers::GetStringArrayFromCommaSeparatedString(FeedbackPins);
+		PinsFromSettings = bTokenizeOnWhiteSpace
+			? PCGHelpers::GetStringArrayFromCommaSeparatedString(FeedbackPins, Context)
+			: PCGHelpers::GetStringArrayFromCommaSeparatedList(FeedbackPins);
+		PRAGMA_ENABLE_DEPRECATION_WARNINGS
+
 		for (const FString& PinLabel : PinsFromSettings)
 		{
 			if (LoopPinNames.Contains(PinLabel))
diff --git a/Engine/Plugins/PCG/Source/PCG/Private/Elements/PCGReplaceTags.cpp b/Engine/Plugins/PCG/Source/PCG/Private/Elements/PCGReplaceTags.cpp
--- a/Engine/Plugins/PCG/Source/PCG/Private/Elements/PCGReplaceTags.cpp
+++ b/Engine/Plugins/PCG/Source/PCG/Private/Elements/PCGReplaceTags.cpp
@@ -3,6 +3,7 @@
 #include "Elements/PCGReplaceTags.h"
 
 #include "PCGContext.h"
+#include "PCGCustomVersion.h"
 #include "PCGData.h"
 #include "PCGModule.h"
 #include "PCGPin.h"
@@ -11,6 +12,18 @@
 #define LOCTEXT_NAMESPACE "PCGReplaceTagsElement"
 
 #if WITH_EDITOR
+void UPCGReplaceTagsSettings::ApplyDeprecation(UPCGNode* InOutNode)
+{
+	PRAGMA_DISABLE_DEPRECATION_WARNINGS
+	if (DataVersion < FPCGCustomVersion::AttributesAndTagsCanContainSpaces)
+	{
+		bTokenizeOnWhiteSpace = true;
+	}
+	PRAGMA_ENABLE_DEPRECATION_WARNINGS
+
+	Super::ApplyDeprecation(InOutNode);
+}
+
 FName UPCGReplaceTagsSettings::GetDefaultNodeName() const
 {
 	return TEXT("ReplaceTags");
@@ -24,8 +37,14 @@ FText UPCGReplaceTagsSettings::GetDefaultNodeTitle() const
 
 FString UPCGReplaceTagsSettings::GetAdditionalTitleInformation() const
 {
-	TArray<FString> TagsToReplace = PCGHelpers::GetStringArrayFromCommaSeparatedString(SelectedTags);
-	TArray<FString> TagsToProcess = PCGHelpers::GetStringArrayFromCommaSeparatedString(ReplacedTags);
+	PRAGMA_DISABLE_DEPRECATION_WARNINGS
+	const TArray<FString> TagsToReplace = bTokenizeOnWhiteSpace
+		? PCGHelpers::GetStringArrayFromCommaSeparatedString(SelectedTags)
+		: PCGHelpers::GetStringArrayFromCommaSeparatedList(SelectedTags);
+	const TArray<FString> TagsToProcess = bTokenizeOnWhiteSpace
+		? PCGHelpers::GetStringArrayFromCommaSeparatedString(ReplacedTags)
+		: PCGHelpers::GetStringArrayFromCommaSeparatedList(ReplacedTags);
+	PRAGMA_ENABLE_DEPRECATION_WARNINGS
 
 	if (TagsToReplace.Num() == 1 && TagsToProcess.Num() == 1)
 	{
@@ -69,8 +88,15 @@ bool FPCGReplaceTagsElement::ExecuteInternal(FPCGContext* Context) const
 	const TArray<FPCGTaggedData> Inputs = Context->InputData.GetInputsByPin(PCGPinConstants::DefaultInputLabel);
 	TArray<FPCGTaggedData>& Outputs = Context->OutputData.TaggedData;
 
-	const TArray<FString> SelectedTags = PCGHelpers::GetStringArrayFromCommaSeparatedString(Settings->SelectedTags);
-	const TArray<FString> ReplacedTags = PCGHelpers::GetStringArrayFromCommaSeparatedString(Settings->ReplacedTags);
+	PRAGMA_DISABLE_DEPRECATION_WARNINGS
+	const TArray<FString> SelectedTags = Settings->bTokenizeOnWhiteSpace
+		? PCGHelpers::GetStringArrayFromCommaSeparatedString(Settings->SelectedTags)
+		: PCGHelpers::GetStringArrayFromCommaSeparatedList(Settings->SelectedTags);
+	const TArray<FString> ReplacedTags = Settings->bTokenizeOnWhiteSpace
+		? PCGHelpers::GetStringArrayFromCommaSeparatedString(Settings->ReplacedTags)
+		: PCGHelpers::GetStringArrayFromCommaSeparatedList(Settings->ReplacedTags);
+	PRAGMA_ENABLE_DEPRECATION_WARNINGS
+
 	if (SelectedTags.IsEmpty() && ReplacedTags.IsEmpty())
 	{
 		// Forward the input and early out
diff --git a/Engine/Plugins/PCG/Source/PCG/Private/Helpers/PCGHelpers.cpp b/Engine/Plugins/PCG/Source/PCG/Private/Helpers/PCGHelpers.cpp
--- a/Engine/Plugins/PCG/Source/PCG/Private/Helpers/PCGHelpers.cpp
+++ b/Engine/Plugins/PCG/Source/PCG/Private/Helpers/PCGHelpers.cpp
@@ -288,10 +288,37 @@ namespace PCGHelpers
 		return (InWorld && InWorld->GetSubsystem<UPCGSubsystem>()) ? InWorld->GetSubsystem<UPCGSubsystem>()->GetPCGWorldActor() : nullptr;
 	}
 
-	TArray<FString> GetStringArrayFromCommaSeparatedString(const FString& InCommaSeparatedString)
+	// TODO: Temporary validation during transition of allowing spaces in tags/attributes. Deprecate in 5.6.
+	TArray<FString> GetStringArrayFromCommaSeparatedString(const FString& InCommaSeparatedString, const FPCGContext* InOptionalContext)
 	{
+#if WITH_EDITOR
+		if (InCommaSeparatedString.Contains(" "))
+		{
+			PCGLog::LogWarningOnGraph(
+				FText::Format(LOCTEXT(
+					"AttributeOrTagContainsSpace", "The comma separated list '{0}' contains an internal space character, which should no longer be parsed as a separator. \n"
+					"Disable 'bParseOnWhiteSpace' on the node to deprecate and update the behavior."),
+					FText::FromString(InCommaSeparatedString)),
+				InOptionalContext);
+		}
+#endif // WITH_EDITOR
+
 		TArray<FString> Result;
 		InCommaSeparatedString.ParseIntoArrayWS(Result, TEXT(","));
+
+		return Result;
+	}
+
+	TArray<FString> GetStringArrayFromCommaSeparatedList(const FString& InCommaSeparatedString)
+	{
+		TArray<FString> Result;
+		InCommaSeparatedString.ParseIntoArray(Result, TEXT(","));
+		// Trim leading and trailing spaces
+		for (FString& String : Result)
+		{
+			String.TrimStartAndEndInline();
+		}
+
 		return Result;
 	}
 
diff --git a/Engine/Plugins/PCG/Source/PCG/Public/Elements/Metadata/PCGMetadataPartition.h b/Engine/Plugins/PCG/Source/PCG/Public/Elements/Metadata/PCGMetadataPartition.h
--- a/Engine/Plugins/PCG/Source/PCG/Public/Elements/Metadata/PCGMetadataPartition.h
+++ b/Engine/Plugins/PCG/Source/PCG/Public/Elements/Metadata/PCGMetadataPartition.h
@@ -19,6 +19,8 @@ public:
 
 	//~Begin UPCGSettings interface
 #if WITH_EDITOR
+	virtual void ApplyDeprecation(UPCGNode* InOutNode) override;
+
 	virtual FName GetDefaultNodeName() const override { return FName(TEXT("AttributePartition")); }
 	virtual FText GetDefaultNodeTitle() const override { return NSLOCTEXT("PCGMetadataPartitionSettings", "NodeTitle", "Attribute Partition"); }
 	virtual EPCGSettingsType GetType() const override { return EPCGSettingsType::Metadata; }
@@ -42,11 +44,18 @@ public:
 	UPROPERTY(meta = (PCG_Overridable))
 	FString PartitionAttributeNames;
 
+	/** Enables deprecated behavior using spaces as separators. Disable to update the node to current behavior. */
+	UE_DEPRECATED(5.6, "bTokenizeOnWhiteSpace has been deprecated.")
+	UPROPERTY(EditAnywhere, Category = Settings, meta = (EditCondition = "bTokenizeOnWhiteSpace", EditConditionHides, DeprecationMessage = "bTokenizeOnWhiteSpace has been deprecated."))
+	bool bTokenizeOnWhiteSpace = false;
+
 #if WITH_EDITORONLY_DATA
-	UPROPERTY()
+	UE_DEPRECATED(5.6, "PartitionAttribute has been deprecated.")
+	UPROPERTY(meta = (DeprecatedProperty, DeprecationMessage = "PartitionAttribute has been deprecated."))
 	FName PartitionAttribute_DEPRECATED = NAME_None;
 
-	UPROPERTY()
+	UE_DEPRECATED(5.6, "PartitionAttributeSource has been deprecated.")
+	UPROPERTY(meta = (DeprecatedProperty, DeprecationMessage = "PartitionAttributeSource has been deprecated."))
 	FPCGAttributePropertyInputSelector PartitionAttributeSource_DEPRECATED;
 #endif // WITH_EDITORONLY_DATA
 };
diff --git a/Engine/Plugins/PCG/Source/PCG/Public/Elements/PCGAddTag.h b/Engine/Plugins/PCG/Source/PCG/Public/Elements/PCGAddTag.h
--- a/Engine/Plugins/PCG/Source/PCG/Public/Elements/PCGAddTag.h
+++ b/Engine/Plugins/PCG/Source/PCG/Public/Elements/PCGAddTag.h
@@ -17,13 +17,14 @@ class UPCGAddTagSettings : public UPCGSettings
 public:
 	//~Begin UPCGSettings interface
 #if WITH_EDITOR
+	virtual void ApplyDeprecation(UPCGNode* InOutNode) override;
+
 	virtual FName GetDefaultNodeName() const override { return FName(TEXT("AddTags")); }
 	virtual FText GetDefaultNodeTitle() const override { return NSLOCTEXT("PCGAddTagElement", "NodeTitle", "Add Tags"); }
 	virtual FText GetNodeTooltipText() const override { return NSLOCTEXT("PCGAddTagElement", "NodeTooltip", "Applies the specified tags on the output data."); }
 	virtual EPCGSettingsType GetType() const override { return EPCGSettingsType::Generic; }
 	virtual bool HasDynamicPins() const override { return true; }
 #endif
-	
 
 protected:
 	virtual TArray<FPCGPinProperties> OutputPinProperties() const override;
@@ -34,6 +35,11 @@ public:
 	/** Comma-separated list of tags to apply to the node */
 	UPROPERTY(EditAnywhere, BlueprintReadWrite, Category = Settings, meta = (PCG_Overridable))
 	FString TagsToAdd;
+
+	/** Enables deprecated behavior using spaces as separators. Disable to update the node to current behavior. */
+	UE_DEPRECATED(5.6, "bTokenizeOnWhiteSpace has been deprecated.")
+	UPROPERTY(EditAnywhere, Category = Settings, meta = (EditCondition = "bTokenizeOnWhiteSpace", EditConditionHides, DeprecationMessage = "bTokenizeOnWhiteSpace has been deprecated."))
+	bool bTokenizeOnWhiteSpace = false;
 };
 
 class FPCGAddTagElement : public IPCGElement
diff --git a/Engine/Plugins/PCG/Source/PCG/Public/Elements/PCGDeleteAttributesElement.h b/Engine/Plugins/PCG/Source/PCG/Public/Elements/PCGDeleteAttributesElement.h
--- a/Engine/Plugins/PCG/Source/PCG/Public/Elements/PCGDeleteAttributesElement.h
+++ b/Engine/Plugins/PCG/Source/PCG/Public/Elements/PCGDeleteAttributesElement.h
@@ -33,6 +33,8 @@ public:
 
 	//~Begin UPCGSettings interface
 #if WITH_EDITOR
+	virtual void ApplyDeprecation(UPCGNode* InOutNode) override;
+
 	virtual FName GetDefaultNodeName() const override;
 	virtual FText GetDefaultNodeTitle() const override;
 	virtual TArray<FText> GetNodeTitleAliases() const override;
@@ -59,6 +61,11 @@ public:
 	UPROPERTY(BlueprintReadWrite, EditAnywhere, Category = Settings, meta = (PCG_Overridable))
 	FString SelectedAttributes;
 
+	/** Enables deprecated behavior using spaces as separators. Disable to update the node to current behavior. */
+	UE_DEPRECATED(5.6, "bTokenizeOnWhiteSpace has been deprecated.")
+	UPROPERTY(EditAnywhere, Category = Settings, meta = (EditCondition = "bTokenizeOnWhiteSpace", EditConditionHides, DeprecationMessage = "bTokenizeOnWhiteSpace has been deprecated."))
+	bool bTokenizeOnWhiteSpace = false;
+
 #if WITH_EDITORONLY_DATA
 	UPROPERTY()
 	TArray<FName> AttributesToKeep_DEPRECATED;
diff --git a/Engine/Plugins/PCG/Source/PCG/Public/Elements/PCGDeleteTags.h b/Engine/Plugins/PCG/Source/PCG/Public/Elements/PCGDeleteTags.h
--- a/Engine/Plugins/PCG/Source/PCG/Public/Elements/PCGDeleteTags.h
+++ b/Engine/Plugins/PCG/Source/PCG/Public/Elements/PCGDeleteTags.h
@@ -22,6 +22,8 @@ class UPCGDeleteTagsSettings : public UPCGSettings
 public:
 	//~Begin UPCGSettings interface
 #if WITH_EDITOR
+	virtual void ApplyDeprecation(UPCGNode* InOutNode) override;
+
 	virtual FName GetDefaultNodeName() const override;
 	virtual FText GetDefaultNodeTitle() const override;
 	virtual EPCGSettingsType GetType() const override { return EPCGSettingsType::Generic; }
@@ -45,6 +47,11 @@ public:
 	/** Comma-separated list of tags to add or remove from the input data. */
 	UPROPERTY(BlueprintReadWrite, EditAnywhere, Category = Settings, meta = (PCG_Overridable))
 	FString SelectedTags;
+
+	/** Enables deprecated behavior using spaces as separators. Disable to update the node to current behavior. */
+	UE_DEPRECATED(5.6, "bTokenizeOnWhiteSpace has been deprecated.")
+	UPROPERTY(EditAnywhere, Category = Settings, meta = (EditCondition = "bTokenizeOnWhiteSpace", EditConditionHides, DeprecationMessage = "bTokenizeOnWhiteSpace has been deprecated."))
+	bool bTokenizeOnWhiteSpace = false;
 };
 
 class FPCGDeleteTagsElement : public IPCGElement
diff --git a/Engine/Plugins/PCG/Source/PCG/Public/Elements/PCGFilterByTag.h b/Engine/Plugins/PCG/Source/PCG/Public/Elements/PCGFilterByTag.h
--- a/Engine/Plugins/PCG/Source/PCG/Public/Elements/PCGFilterByTag.h
+++ b/Engine/Plugins/PCG/Source/PCG/Public/Elements/PCGFilterByTag.h
@@ -22,6 +22,8 @@ class UPCGFilterByTagSettings : public UPCGFilterDataBaseSettings
 public:
 	//~Begin UPCGSettings interface
 #if WITH_EDITOR
+	virtual void ApplyDeprecation(UPCGNode* InOutNode) override;
+
 	virtual FName GetDefaultNodeName() const override;
 	virtual FText GetDefaultNodeTitle() const override;
 	virtual FText GetNodeTooltipText() const override;
@@ -43,6 +45,11 @@ public:
 	/** Comma-separated list of tags */
 	UPROPERTY(BlueprintReadWrite, EditAnywhere, Category = Settings, meta = (PCG_Overridable))
 	FString SelectedTags;
+
+	/** Enables deprecated behavior using spaces as separators. Disable to update the node to current behavior. */
+	UE_DEPRECATED(5.6, "bTokenizeOnWhiteSpace has been deprecated.")
+	UPROPERTY(EditAnywhere, Category = Settings, meta = (EditCondition = "bTokenizeOnWhiteSpace", EditConditionHides, DeprecationMessage = "bTokenizeOnWhiteSpace has been deprecated."))
+	bool bTokenizeOnWhiteSpace = false;
 };
 
 class FPCGFilterByTagElement : public IPCGElement
diff --git a/Engine/Plugins/PCG/Source/PCG/Public/Elements/PCGLoopElement.h b/Engine/Plugins/PCG/Source/PCG/Public/Elements/PCGLoopElement.h
--- a/Engine/Plugins/PCG/Source/PCG/Public/Elements/PCGLoopElement.h
+++ b/Engine/Plugins/PCG/Source/PCG/Public/Elements/PCGLoopElement.h
@@ -59,6 +59,11 @@ public:
 	*/
 	UPROPERTY(BlueprintReadWrite, EditAnywhere, Category = Data, meta = (EditCondition = "!bUseGraphDefaultPinUsage"))
 	FString FeedbackPins;
+
+	/** Enables deprecated behavior using spaces as separators. Disable to update the node to current behavior. */
+	UE_DEPRECATED(5.6, "bTokenizeOnWhiteSpace has been deprecated.")
+	UPROPERTY(EditAnywhere, Category = Settings, meta = (EditCondition = "bTokenizeOnWhiteSpace", EditConditionHides, DeprecationMessage = "bTokenizeOnWhiteSpace has been deprecated."))
+	bool bTokenizeOnWhiteSpace = false;
 };
 
 class FPCGLoopElement : public FPCGSubgraphElement
diff --git a/Engine/Plugins/PCG/Source/PCG/Public/Elements/PCGReplaceTags.h b/Engine/Plugins/PCG/Source/PCG/Public/Elements/PCGReplaceTags.h
--- a/Engine/Plugins/PCG/Source/PCG/Public/Elements/PCGReplaceTags.h
+++ b/Engine/Plugins/PCG/Source/PCG/Public/Elements/PCGReplaceTags.h
@@ -15,6 +15,8 @@ class UPCGReplaceTagsSettings : public UPCGSettings
 public:
 	//~Begin UPCGSettings interface
 #if WITH_EDITOR
+	virtual void ApplyDeprecation(UPCGNode* InOutNode) override;
+
 	virtual FName GetDefaultNodeName() const override;
 	virtual FText GetDefaultNodeTitle() const override;
 	virtual EPCGSettingsType GetType() const override { return EPCGSettingsType::Generic; }
@@ -41,6 +43,11 @@ public:
 	/** Comma-separated list of new tags to replace the selected tags. */
 	UPROPERTY(BlueprintReadWrite, EditAnywhere, Category = Settings, meta = (PCG_Overridable))
 	FString ReplacedTags;
+
+	/** Enables deprecated behavior using spaces as separators. Disable to update the node to current behavior. */
+	UE_DEPRECATED(5.6, "bTokenizeOnWhiteSpace has been deprecated.")
+	UPROPERTY(EditAnywhere, Category = Settings, meta = (EditCondition = "bTokenizeOnWhiteSpace", EditConditionHides, DeprecationMessage = "bTokenizeOnWhiteSpace has been deprecated."))
+	bool bTokenizeOnWhiteSpace = false;
 };
 
 class FPCGReplaceTagsElement : public IPCGElement
diff --git a/Engine/Plugins/PCG/Source/PCG/Public/Helpers/PCGHelpers.h b/Engine/Plugins/PCG/Source/PCG/Public/Helpers/PCGHelpers.h
--- a/Engine/Plugins/PCG/Source/PCG/Public/Helpers/PCGHelpers.h
+++ b/Engine/Plugins/PCG/Source/PCG/Public/Helpers/PCGHelpers.h
@@ -54,7 +54,9 @@ namespace PCGHelpers
 
 	PCG_API APCGWorldActor* GetPCGWorldActor(UWorld* InWorld);
 
-	PCG_API TArray<FString> GetStringArrayFromCommaSeparatedString(const FString& InCommaSeparatedString);
+	UE_DEPRECATED(5.6, "This function was incorrectly parsing by whitespace. Please use 'GetStringArrayFromCommaSeparatedList' instead.")
+	PCG_API TArray<FString> GetStringArrayFromCommaSeparatedString(const FString& InCommaSeparatedString, const FPCGContext* InOptionalContext = nullptr);
+	PCG_API TArray<FString> GetStringArrayFromCommaSeparatedList(const FString& InCommaSeparatedString);
 
 #if WITH_EDITOR
 	PCG_API void GatherDependencies(UObject* Object, TSet<TObjectPtr<UObject>>& OutDependencies, int32 MaxDepth = -1, const TArray<UClass*>& InExcludedClasses = {});
diff --git a/Engine/Plugins/PCG/Source/PCG/Public/PCGCustomVersion.h b/Engine/Plugins/PCG/Source/PCG/Public/PCGCustomVersion.h
--- a/Engine/Plugins/PCG/Source/PCG/Public/PCGCustomVersion.h
+++ b/Engine/Plugins/PCG/Source/PCG/Public/PCGCustomVersion.h
@@ -77,6 +77,9 @@ struct FPCGCustomVersion
 		// New gate for new data, so any node that has a non Point pin don't do any ToPointData by default.
 		NoMoreSpatialDataConversionToPointDataByDefaultOnNonPointPins = 22,
 
+		// Attributes and tags can now contain spaces and will no longer be parsed by spaces.
+		AttributesAndTagsCanContainSpaces = 23,
+
 		// -----<new versions can be added above this line>-------------------------------------------------
 		VersionPlusOne,
 		LatestVersion = VersionPlusOne - 1
