[PCG] Finish deprecating PointProcessingElement

#rnx

diff --git a/Engine/Plugins/PCG/Source/PCG/Public/Elements/PCGPointProcessingElementBase.h b/Engine/Plugins/PCG/Source/PCG/Public/Elements/PCGPointProcessingElementBase.h
--- a/Engine/Plugins/PCG/Source/PCG/Public/Elements/PCGPointProcessingElementBase.h
+++ b/Engine/Plugins/PCG/Source/PCG/Public/Elements/PCGPointProcessingElementBase.h
@@ -4,19 +4,11 @@
 
 #include "PCGElement.h"
 
-
 struct FPCGPoint;
 struct FPCGTaggedData;
 
-// TODO: Deprecate this class once all of the internal inheritors have been cleaned up
-// class UE_DEPRECATED(5.4, "FPCGPointProcessingElementBase is deprecated. Please use 'FPCGPointOperationElementBase' instead.") PCG_API FPCGPointProcessingElementBase : public IPCGElement
-class PCG_API FPCGPointProcessingElementBase : public IPCGElement
+class UE_DEPRECATED(5.6, "FPCGPointProcessingElementBase is deprecated. Please use 'FPCGPointOperationElementBase' instead.") PCG_API FPCGPointProcessingElementBase : public IPCGElement
 {
 protected:
 	void ProcessPoints(FPCGContext* Context, const TArray<FPCGTaggedData>& Inputs, TArray<FPCGTaggedData>& Outputs, const TFunction<bool(const FPCGPoint&, FPCGPoint&)>& PointFunc) const;
 };
-
-#if UE_ENABLE_INCLUDE_ORDER_DEPRECATED_IN_5_2
-#include "CoreMinimal.h"
-#include "Helpers/PCGAsync.h"
-#endif
