[PCG] Loc fix

#rnx

diff --git a/Engine/Plugins/PCG/Source/PCGEditor/Private/Details/PCGAttributePropertySelectorDetails.cpp b/Engine/Plugins/PCG/Source/PCGEditor/Private/Details/PCGAttributePropertySelectorDetails.cpp
--- a/Engine/Plugins/PCG/Source/PCGEditor/Private/Details/PCGAttributePropertySelectorDetails.cpp
+++ b/Engine/Plugins/PCG/Source/PCGEditor/Private/Details/PCGAttributePropertySelectorDetails.cpp
@@ -79,7 +79,7 @@ namespace PCGAttributePropertySelectorDetails
 			{
 				BuildSubMenuSectionForExtractor(InSubMenuBuilder, InDetailsObject, PCGAttributeExtractorConstants::VectorLength, LOCTEXT("VectorLengthTooltip", "Get the length of the vector (double value). Read only."));
 				BuildSubMenuSectionForExtractor(InSubMenuBuilder, InDetailsObject, PCGAttributeExtractorConstants::VectorSquaredLength, LOCTEXT("VectorSquaredLengthTooltip", "Get the squared length of the vector (double value). Read only."));
-				BuildSubMenuSectionForExtractor(InSubMenuBuilder, InDetailsObject, PCGAttributeExtractorConstants::VectorNormalized, LOCTEXT("VectorLengthTooltip", "Get the vector normalized (double value). Read only."));
+				BuildSubMenuSectionForExtractor(InSubMenuBuilder, InDetailsObject, PCGAttributeExtractorConstants::VectorNormalized, LOCTEXT("VectorNormalizedTooltip", "Get the vector normalized (double value). Read only."));
 			}
 		}));
 	}
