[PCG] Builder uses FEditorFileUtils::GetDirtyWorldPackages instead of FEditorFileUtils::GetDirtyPackages to avoid saving non Actor files (the builder is not expected to modify non Actor assets)

#rnx

diff --git a/Engine/Plugins/PCG/Source/PCGEditor/Private/WorldPartitionBuilder/PCGWorldPartitionBuilder.cpp b/Engine/Plugins/PCG/Source/PCGEditor/Private/WorldPartitionBuilder/PCGWorldPartitionBuilder.cpp
--- a/Engine/Plugins/PCG/Source/PCGEditor/Private/WorldPartitionBuilder/PCGWorldPartitionBuilder.cpp
+++ b/Engine/Plugins/PCG/Source/PCGEditor/Private/WorldPartitionBuilder/PCGWorldPartitionBuilder.cpp
@@ -163,7 +163,7 @@ bool UPCGWorldPartitionBuilder::RunInternal(UWorld* World, const FCellInfo& InCe
 
 	{
 		TArray<UPackage*> DirtyPackages;
-		FEditorFileUtils::GetDirtyPackages(DirtyPackages);
+		FEditorFileUtils::GetDirtyWorldPackages(DirtyPackages);
 		DirtyPackages.RemoveSwap(nullptr);
 		for (UPackage* DirtyPackage : DirtyPackages)
 		{
@@ -212,7 +212,7 @@ bool UPCGWorldPartitionBuilder::RunInternal(UWorld* World, const FCellInfo& InCe
 	if (bGeneratedAnyComponent)
 	{
 		TArray<UPackage*> DirtyPackages;
-		FEditorFileUtils::GetDirtyPackages(DirtyPackages);
+		FEditorFileUtils::GetDirtyWorldPackages(DirtyPackages);
 		
 		DirtyPackages.RemoveSwap(nullptr);
 		for (UPackage* DirtyPackage : DirtyPackages)
