Add redirector for updated parameter names on bp function in PCG

diff --git a/Engine/Plugins/PCG/Config/BasePCG.ini b/Engine/Plugins/PCG/Config/BasePCG.ini
--- a/Engine/Plugins/PCG/Config/BasePCG.ini
+++ b/Engine/Plugins/PCG/Config/BasePCG.ini
@@ -59,6 +59,8 @@
 +PropertyRedirects=(OldName="PCGGraphParametersHelpers.SetRotatorParameter.GraphInstance",NewName="PCGGraphParametersHelpers.SetRotatorParameter.GraphInterface")
 +PropertyRedirects=(OldName="PCGGraphParametersHelpers.SetTransformParameter.GraphInstance",NewName="PCGGraphParametersHelpers.SetTransformParameter.GraphInterface")
 +PropertyRedirects=(OldName="PCGPropertyToParamDataSettings.bExtractObjectAndStruct",NewName="PCGPropertyToParamDataSettings.bForceObjectAndStructExtraction")
++PropertyRedirects=(OldName="PCGAttributePropertySelectorBlueprintHelpers.CopyAndFixSource.Selector",NewName="PCGAttributePropertySelectorBlueprintHelpers.CopyAndFixSource.OutputSelector")
++PropertyRedirects=(OldName="PCGAttributePropertySelectorBlueprintHelpers.CopyAndFixSource.InSelector",NewName="PCGAttributePropertySelectorBlueprintHelpers.CopyAndFixSource.InputSelector")
 
 +ClassRedirects=(OldName="PCGPointSamplerSettings",NewName="/Script/PCG.PCGSelectPointsSettings")
 +ClassRedirects=(OldName="PCGDensityNoiseSettings",NewName="/Script/PCG.PCGAttributeNoiseSettings")
