[PCG] World Ray Hit Query and Surface Sampler now work with extreme angles

* Separated out passed in param data with calculated data in S.Sampler

#rnx

diff --git a/Engine/Plugins/PCG/Source/PCG/Private/Data/PCGWorldData.cpp b/Engine/Plugins/PCG/Source/PCG/Private/Data/PCGWorldData.cpp
--- a/Engine/Plugins/PCG/Source/PCG/Private/Data/PCGWorldData.cpp
+++ b/Engine/Plugins/PCG/Source/PCG/Private/Data/PCGWorldData.cpp
@@ -226,10 +226,12 @@ UPCGSpatialData* UPCGWorldVolumetricData::CopyInternal() const
 }
 
 /** World Ray Hit data implementation */
-void UPCGWorldRayHitData::Initialize(UWorld* InWorld, const FBox& InBounds)
+void UPCGWorldRayHitData::Initialize(UWorld* InWorld, const FTransform& InTransform, const FBox& InBounds, const FBox& InLocalBounds)
 {
 	World = InWorld;
+	Transform = InTransform;
 	Bounds = InBounds;
+	LocalBounds = InLocalBounds;
 }
 
 void UPCGWorldRayHitData::AddToCrc(FArchiveCrc32& Ar, bool bFullDataCrc) const
@@ -254,13 +256,11 @@ bool UPCGWorldRayHitData::SamplePoint(const FTransform& InTransform, const FBox&
 	Params.bTraceComplex = QueryParams.bTraceComplex;
 	Params.bReturnPhysicalMaterial = (PhysicalMaterialAttribute != nullptr);
 
-	// Project the InTransform location on the ray origin plane
-	const FVector PointLocation = InTransform.GetLocation();
-	FVector RayStart = PointLocation - ((PointLocation - QueryParams.RayOrigin) | QueryParams.RayDirection) * QueryParams.RayDirection;
-	FVector RayEnd = RayStart + QueryParams.RayDirection * QueryParams.RayLength;
+	FVector RayOrigin = InTransform.GetLocation() - ((InTransform.GetLocation() - QueryParams.RayOrigin) | QueryParams.RayDirection) * QueryParams.RayDirection;
+	FVector RayEnd = RayOrigin + QueryParams.RayDirection * QueryParams.RayLength;
 
 	TArray<FHitResult> Hits;
-	World->LineTraceMultiByObjectType(Hits, RayStart, RayEnd, ObjectQueryParams, Params);
+	World->LineTraceMultiByObjectType(Hits, RayOrigin, RayEnd, ObjectQueryParams, Params);
 
 	for (const FHitResult& Hit : Hits)
 	{
@@ -324,7 +324,7 @@ bool UPCGWorldRayHitData::SamplePoint(const FTransform& InTransform, const FBox&
 		}
 
 		bool bHitOnLandscape = false;
-		if (QueryParams.bIgnoreLandscapeHits || QueryParams.bApplyMetadataFromLandscape)
+		if (QueryParams.bIgnoreLandscapeHits || QueryParams.bApplyMetadataFromLandscape || QueryParams.bIgnoreBackfaceHits)
 		{
 			bHitOnLandscape = HitComponent->GetOwner() && HitComponent->GetOwner()->IsA<ALandscapeProxy>();
 		}
@@ -334,6 +334,16 @@ bool UPCGWorldRayHitData::SamplePoint(const FTransform& InTransform, const FBox&
 			continue;
 		}
 
+		// Optionally skip backface hits
+		if (QueryParams.bIgnoreBackfaceHits)
+		{
+			// If its a landscape, we cull if the normal is negative in Z direction (landscape normal is always the +Z axis). If not, then we cull if the impact normal and the ray are headed in the same direction
+			if ((bHitOnLandscape && Hit.ImpactNormal.Z < 0) || (QueryParams.RayDirection).Dot(Hit.ImpactNormal) > 0)
+			{
+				continue;
+			}
+		}
+
 		// Finally, fill in OutPoint - we're done
 		// Implementation note: this uses the same orthonormalization process as the landscape cache
 		ensure(Hit.ImpactNormal.IsNormalized());
@@ -401,11 +411,9 @@ const UPCGPointData* UPCGWorldRayHitData::CreatePointData(FPCGContext* Context,
 		return Data;
 	}
 
-	PCGSurfaceSampler::FSurfaceSamplerParams SamplerParams;
-	if (SamplerParams.Initialize(nullptr, Context, EffectiveBounds))
-	{
-		Data = PCGSurfaceSampler::SampleSurface(Context, this, nullptr, SamplerParams);
-	}
+	// The default params will be fine in this case
+	const PCGSurfaceSampler::FSurfaceSamplerParams Params;
+	Data = PCGSurfaceSampler::SampleSurface(Context, /*InSurface=*/this, /*InBoundingShape=*/nullptr, EffectiveBounds, Params);
 
 	return Data;
 }
diff --git a/Engine/Plugins/PCG/Source/PCG/Private/Elements/PCGSurfaceSampler.cpp b/Engine/Plugins/PCG/Source/PCG/Private/Elements/PCGSurfaceSampler.cpp
--- a/Engine/Plugins/PCG/Source/PCG/Private/Elements/PCGSurfaceSampler.cpp
+++ b/Engine/Plugins/PCG/Source/PCG/Private/Elements/PCGSurfaceSampler.cpp
@@ -10,7 +10,7 @@
 #include "PCGGraph.h"
 #include "PCGPin.h"
 #include "Data/PCGPointData.h"
-#include "Data/PCGSpatialData.h"
+#include "Data/PCGSurfaceData.h"
 #include "Helpers/PCGAsync.h"
 #include "Helpers/PCGHelpers.h"
 #include "Helpers/PCGSettingsHelpers.h"
@@ -24,31 +24,38 @@
 
 namespace PCGSurfaceSampler
 {
-	bool FSurfaceSamplerParams::Initialize(const UPCGSurfaceSamplerSettings* Settings, const FPCGContext* Context, const FBox& InputBounds)
+	void FSurfaceSamplerParams::InitializeFromSettings(const UPCGSurfaceSamplerSettings* Settings)
 	{
-		if (!Context)
-		{
-			return false;
-		}
-
-		if (Settings)
-		{
-			// Compute used values
-			PointsPerSquaredMeter = Settings->PointsPerSquaredMeter;
-			PointExtents = Settings->PointExtents;
-			Looseness = Settings->Looseness;
-			bApplyDensityToPoints = Settings->bApplyDensityToPoints;
-			PointSteepness = Settings->PointSteepness;
+		check(Settings);
+		// Compute used values
+		PointsPerSquaredMeter = Settings->PointsPerSquaredMeter;
+		PointExtents = Settings->PointExtents;
+		Looseness = Settings->Looseness;
+		bApplyDensityToPoints = Settings->bApplyDensityToPoints;
+		PointSteepness = Settings->PointSteepness;
 #if WITH_EDITOR
-			bKeepZeroDensityPoints = Settings->bKeepZeroDensityPoints;
+		bKeepZeroDensityPoints = Settings->bKeepZeroDensityPoints;
 #endif
+	}
+
+	bool FSurfaceSamplerData::Initialize(const UPCGSurfaceSamplerSettings* Settings, const FPCGContext* Context, const FBox& InEffectiveGridBounds, const FTransform& InSurfaceTransform)
+	{
+		Params.InitializeFromSettings(Settings);
+		return Initialize(Context, InEffectiveGridBounds, InSurfaceTransform);
+	}
+
+	bool FSurfaceSamplerData::Initialize(const FPCGContext* Context, const FBox& InEffectiveGridBounds, const FTransform& InSurfaceTransform)
+	{
+		if (!Context || !InEffectiveGridBounds.IsValid)
+		{
+			return false;
 		}
 
 		Seed = Context->GetSeed();
 
 		// Conceptually, we will break down the surface bounds in a N x M grid
-		InterstitialDistance = PointExtents * 2;
-		InnerCellSize = InterstitialDistance * Looseness;
+		InterstitialDistance = Params.PointExtents * 2;
+		InnerCellSize = InterstitialDistance * Params.Looseness;
 		CellSize = InterstitialDistance + InnerCellSize;
 		if (CellSize.X <= 0 || CellSize.Y <= 0)
 		{
@@ -57,11 +64,33 @@ namespace PCGSurfaceSampler
 			return false;
 		}
 
+		// Local transformation is only needed if we're rotating.
+		bNeedsLocalTransformation = !InSurfaceTransform.Rotator().IsNearlyZero();
+		if (bNeedsLocalTransformation)
+		{
+			// Build the semi-local transform matrix for transforming points to the pre-projection plane
+			const FTransform TranslationTransform(InSurfaceTransform.GetTranslation());
+			PreProjectionTransform = TranslationTransform.Inverse().ToMatrixNoScale();
+			// Find the rotation between the world normal and the surface and apply to the matrix
+			PreProjectionTransform *= FQuat::FindBetweenNormals(FVector::UpVector, InSurfaceTransform.GetRotation().GetUpVector().GetSafeNormal()).ToMatrix();
+			PreProjectionTransform *= TranslationTransform.ToMatrixNoScale();
+		}
+
+		// Drop points slightly by an epsilon otherwise point can be culled. If the sampler has a volume connected as the Bounding Shape,
+		// the volume will call through to PCGHelpers::IsInsideBounds() which is a one sided test and points at the top of the volume
+		// will fail it. TODO perhaps the one-sided check can be isolated to component-bounds
+		constexpr FVector::FReal DefaultHeightModifier = 1.0 - UE_DOUBLE_SMALL_NUMBER;
+		// Try to use a multiplier instead of a simply offset to combat loss of precision in floats. However if MaxZ is very small,
+		// then multiplier will not work, so just use an offset.
+		PreProjectionDisplacement = (FMath::Abs(InEffectiveGridBounds.Max.Z) > UE_DOUBLE_SMALL_NUMBER) ? InEffectiveGridBounds.Max.Z * DefaultHeightModifier : -UE_DOUBLE_SMALL_NUMBER;
+		// Make sure we're still in bounds though!
+		PreProjectionDisplacement = FMath::Max(PreProjectionDisplacement, InEffectiveGridBounds.Min.Z);
+
 		// By using scaled indices in the world, we can easily make this process deterministic
-		CellMinX = FMath::CeilToInt((InputBounds.Min.X) / CellSize.X);
-		CellMaxX = FMath::FloorToInt((InputBounds.Max.X) / CellSize.X);
-		CellMinY = FMath::CeilToInt((InputBounds.Min.Y) / CellSize.Y);
-		CellMaxY = FMath::FloorToInt((InputBounds.Max.Y) / CellSize.Y);
+		CellMinX = FMath::CeilToInt((InEffectiveGridBounds.Min.X) / CellSize.X);
+		CellMaxX = FMath::FloorToInt((InEffectiveGridBounds.Max.X) / CellSize.X);
+		CellMinY = FMath::CeilToInt((InEffectiveGridBounds.Min.Y) / CellSize.Y);
+		CellMaxY = FMath::FloorToInt((InEffectiveGridBounds.Max.Y) / CellSize.Y);
 
 		{
 			const int64 CellCountX = 1 + CellMaxX - CellMinX;
@@ -88,8 +117,8 @@ namespace PCGSurfaceSampler
 
 		check(CellCount > 0);
 
-		const FVector::FReal InvSquaredMeterUnits = 1.0 / (100.0 * 100.0);
-		const FVector::FReal TargetPointCount = (InputBounds.Max.X - InputBounds.Min.X) * (InputBounds.Max.Y - InputBounds.Min.Y) * PointsPerSquaredMeter * InvSquaredMeterUnits;
+		constexpr FVector::FReal InvSquaredMeterUnits = 1.0 / (100.0 * 100.0);
+		const FVector::FReal TargetPointCount = (InEffectiveGridBounds.Max.X - InEffectiveGridBounds.Min.X) * (InEffectiveGridBounds.Max.Y - InEffectiveGridBounds.Min.Y) * Params.PointsPerSquaredMeter * InvSquaredMeterUnits;
 		Ratio = static_cast<float>(FMath::Clamp(TargetPointCount / (FVector::FReal)CellCount, 0.0, 1.0));
 
 		if (Ratio < UE_SMALL_NUMBER)
@@ -108,13 +137,10 @@ namespace PCGSurfaceSampler
 			return false;
 		}
 
-		InputBoundsMinZ = InputBounds.Min.Z;
-		InputBoundsMaxZ = InputBounds.Max.Z;
-
 		return true;
 	}
 
-	FIntVector2 FSurfaceSamplerParams::ComputeCellIndices(int32 Index) const
+	FIntVector2 FSurfaceSamplerData::ComputeCellIndices(int32 Index) const
 	{
 		check(Index >= 0 && Index < CellCount);
 		const int32 CellCountX = 1 + CellMaxX - CellMinX;
@@ -122,51 +148,45 @@ namespace PCGSurfaceSampler
 		return FIntVector2(CellMinX + (Index % CellCountX), CellMinY + (Index / CellCountX));
 	}
 
-	UPCGPointData* SampleSurface(FPCGContext* Context, const UPCGSpatialData* InSurface, const UPCGSpatialData* InBoundingShape, const FSurfaceSamplerParams& ExecutionSettings)
+	UPCGPointData* SampleSurface(FPCGContext* Context, const UPCGSurfaceData* InSurface, const UPCGSpatialData* InBoundingShape, const FBox& EffectiveBounds, const FSurfaceSamplerParams& ExecutionParams)
 	{
 		UPCGPointData* SampledData = NewObject<UPCGPointData>();
 		SampledData->InitializeFromData(InSurface);
 
+		FSurfaceSamplerData SamplerData;
+		SamplerData.Params = ExecutionParams;
+		SamplerData.Initialize(Context, EffectiveBounds);
 		// We don't support time slicing here
-		SampleSurface(Context, ExecutionSettings, InSurface, InBoundingShape, SampledData, /*bTimeSlicingIsEnabled=*/false);
+		SampleSurface(Context, SamplerData, InSurface, InBoundingShape, SampledData, /*bTimeSlicingIsEnabled=*/false);
 
 		return SampledData;
 	}
 
-	bool SampleSurface(FPCGContext* Context, const FSurfaceSamplerParams& Settings, const UPCGSpatialData* InSurface, const UPCGSpatialData* InBoundingShape, UPCGPointData* SampledData, const bool bTimeSlicingIsEnabled)
+	bool SampleSurface(FPCGContext* Context, const FSurfaceSamplerData& SamplerData, const UPCGSurfaceData* InSurface, const UPCGSpatialData* InBoundingShape, UPCGPointData* SampledData, const bool bTimeSlicingIsEnabled)
 	{
 		check(InSurface && SampledData);
 
 		TArray<FPCGPoint>& SampledPoints = SampledData->GetMutablePoints();
 
-		FPCGProjectionParams ProjectionParams{};
-
-		// Drop points slightly by an epsilon otherwise point can be culled. If the sampler has a volume connected as the Bounding Shape,
-		// the volume will call through to PCGHelpers::IsInsideBounds() which is a one sided test and points at the top of the volume
-		// will fail it. TODO perhaps the one-sided check can be isolated to component-bounds
-		constexpr FVector::FReal ZMultiplier = 1.0 - UE_DOUBLE_SMALL_NUMBER;
-		// Try to use a multiplier instead of a simply offset to combat loss of precision in floats. However if MaxZ is very small,
-		// then multiplier will not work, so just use an offset.
-		FVector::FReal SampleZ = (FMath::Abs(Settings.InputBoundsMaxZ) > UE_DOUBLE_SMALL_NUMBER) ? Settings.InputBoundsMaxZ * ZMultiplier : -UE_DOUBLE_SMALL_NUMBER;
-		// Make sure we're still in bounds though!
-		SampleZ = FMath::Max(SampleZ, Settings.InputBoundsMinZ);
+		const FPCGProjectionParams ProjectionParams{};
 
 		// Cache pointer ahead of time to avoid dereferencing object pointer which does access tracking and supports lazy loading, and can come with substantial
 		// overhead (add trace marker to FObjectPtr::Get to see).
 		UPCGMetadata* OutMetadata = SampledData->Metadata.Get();
 
-		auto AsyncProcessFunc = [&Settings, InBoundingShape, InSurface, &ProjectionParams, SampleZ, OutMetadata](int32 Index, FPCGPoint& OutPoint)
+		auto AsyncProcessFunc = [&SamplerData, InBoundingShape, InSurface, &ProjectionParams, OutMetadata](int32 Index, FPCGPoint& OutPoint)
 		{
-			const FIntVector2 Indices = Settings.ComputeCellIndices(Index);
+			const FSurfaceSamplerParams& SamplerParams = SamplerData.Params;
+			const FIntVector2 Indices = SamplerData.ComputeCellIndices(Index);
 
-			const FVector::FReal CurrentX = Indices.X * Settings.CellSize.X;
-			const FVector::FReal CurrentY = Indices.Y * Settings.CellSize.Y;
-			const FVector InnerCellSize = Settings.InnerCellSize;
+			const FVector::FReal CurrentX = Indices.X * SamplerData.CellSize.X;
+			const FVector::FReal CurrentY = Indices.Y * SamplerData.CellSize.Y;
+			const FVector InnerCellSize = SamplerData.InnerCellSize;
 
-			FRandomStream RandomSource(PCGHelpers::ComputeSeed(Settings.Seed, Indices.X, Indices.Y));
+			FRandomStream RandomSource(PCGHelpers::ComputeSeed(SamplerData.Seed, Indices.X, Indices.Y));
 			const float Chance = RandomSource.FRand();
 
-			const float Ratio = Settings.Ratio;
+			const float Ratio = SamplerData.Ratio;
 
 			if (Chance >= Ratio)
 			{
@@ -176,8 +196,16 @@ namespace PCGSurfaceSampler
 			const float RandX = RandomSource.FRand();
 			const float RandY = RandomSource.FRand();
 
-			const FVector TentativeLocation = FVector(CurrentX + RandX * InnerCellSize.X, CurrentY + RandY * InnerCellSize.Y, SampleZ);
-			const FBox LocalBound(-Settings.PointExtents, Settings.PointExtents);
+			FVector TentativeLocation = FVector(CurrentX + RandX * InnerCellSize.X, CurrentY + RandY * InnerCellSize.Y, SamplerData.PreProjectionDisplacement);
+
+			// If pre-projected points need a local transformation (ex. World Ray Hit Query) and not default to -Z
+			if (SamplerData.bNeedsLocalTransformation)
+			{
+				// Transform the pre-projected sample point around the surface's origin from local to world space
+				TentativeLocation = SamplerData.PreProjectionTransform.TransformPosition(TentativeLocation);
+			}
+
+			const FBox LocalBound(-SamplerParams.PointExtents, SamplerParams.PointExtents);
 
 			// The output at this point is not initialized
 			OutPoint = FPCGPoint();
@@ -189,15 +217,15 @@ namespace PCGSurfaceSampler
 			}
 
 			// Set physical properties that are needed for the bounding shape checks, etc.
-			OutPoint.SetExtents(Settings.PointExtents);
-			OutPoint.Steepness = Settings.PointSteepness;
+			OutPoint.SetExtents(SamplerParams.PointExtents);
+			OutPoint.Steepness = SamplerParams.PointSteepness;
 
 			// Now run gauntlet of shape network (if there is one) to accept or reject the point.
 			if (InBoundingShape)
 			{
 				FPCGPoint BoundingShapeSample;
 #if WITH_EDITOR
-				if (!InBoundingShape->SamplePoint(OutPoint.Transform, OutPoint.GetLocalBounds(), BoundingShapeSample, nullptr) && !Settings.bKeepZeroDensityPoints)
+				if (!InBoundingShape->SamplePoint(OutPoint.Transform, OutPoint.GetLocalBounds(), BoundingShapeSample, nullptr) && !SamplerParams.bKeepZeroDensityPoints)
 #else
 				if (!InBoundingShape->SamplePoint(OutPoint.Transform, OutPoint.GetLocalBounds(), BoundingShapeSample, nullptr))
 #endif
@@ -210,14 +238,14 @@ namespace PCGSurfaceSampler
 			}
 
 			// Apply final parameters on the point
-			OutPoint.Density *= (Settings.bApplyDensityToPoints ? ((Ratio - Chance) / Ratio) : 1.0f);
+			OutPoint.Density *= (SamplerParams.bApplyDensityToPoints ? ((Ratio - Chance) / Ratio) : 1.0f);
 			OutPoint.Seed = RandomSource.GetCurrentSeed();
 
 			return true;
 		};
 
 		FPCGAsyncState* AsyncState = Context ? &Context->AsyncState : nullptr;
-		return FPCGAsync::AsyncProcessing<FPCGPoint>(AsyncState, Settings.CellCount, SampledPoints, AsyncProcessFunc, /*bEnableTimeSlicing=*/Context && bTimeSlicingIsEnabled);
+		return FPCGAsync::AsyncProcessing<FPCGPoint>(AsyncState, SamplerData.CellCount, SampledPoints, AsyncProcessFunc, /*bEnableTimeSlicing=*/Context && bTimeSlicingIsEnabled);
 	}
 
 #if WITH_EDITOR
@@ -375,13 +403,13 @@ namespace PCGSurfaceSamplerHelpers
 
 		TArray<FPCGTaggedData>& Outputs = Context->OutputData.TaggedData;
 		// Find the generating shapes to sample
-		TArray<const UPCGSpatialData*>& GeneratingShapes = OutState.GeneratingShapes;
+		TArray<const UPCGSurfaceData*>& GeneratingShapes = OutState.GeneratingShapes;
 		GeneratingShapes.Reserve(SurfaceInputs.Num());
 
 		// Construct a list of shapes to generate samples from. Get these directly from the first input pin.
 		for (const FPCGTaggedData& TaggedData : SurfaceInputs)
 		{
-			if (const UPCGSpatialData* SurfaceData = Cast<UPCGSpatialData>(TaggedData.Data))
+			if (const UPCGSurfaceData* SurfaceData = Cast<UPCGSurfaceData>(TaggedData.Data))
 			{
 				GeneratingShapes.Add(SurfaceData);
 				Outputs.Add(TaggedData);
@@ -414,7 +442,7 @@ bool FPCGSurfaceSamplerElement::PrepareDataInternal(FPCGContext* InContext) cons
 		return true;
 	}
 
-	TArray<const UPCGSpatialData*>& GeneratingShapes = Context->GetPerExecutionState().GeneratingShapes;
+	TArray<const UPCGSurfaceData*>& GeneratingShapes = Context->GetPerExecutionState().GeneratingShapes;
 	TArray<FPCGTaggedData>& Outputs = Context->OutputData.TaggedData;
 
 	// Initialize the per-iteration data, using the generating shapes as the source of iteration
@@ -422,38 +450,55 @@ bool FPCGSurfaceSamplerElement::PrepareDataInternal(FPCGContext* InContext) cons
 		[&GeneratingShapes, &Outputs, &Settings, &Context](IterStateType& OutState, const ExecStateType& ExecState, const uint32 IterationIndex)
 		{
 			// If we have generating shape inputs, use them
-			const UPCGSpatialData* GeneratingShape = GeneratingShapes[IterationIndex];
+			const UPCGSurfaceData* GeneratingShape = GeneratingShapes[IterationIndex];
 			check(GeneratingShape);
 
 			OutState.OutputPoints = NewObject<UPCGPointData>();
 			OutState.OutputPoints->InitializeFromData(GeneratingShape);
 
-			// Calculate the intersection of bounds of the provided inputs
-			FBox InputBounds = FBox(EForceInit::ForceInit);
+			// This bounds will be used to generate the pre-projected grid
+			FBox EffectiveGridBounds = FBox(EForceInit::ForceInit);
 			const FBox& BoundingShapeBounds = ExecState.BoundingShapeBounds;
 
-			if (GeneratingShape->IsBounded())
+			// The shape's local bounds is most ideal for generating the grid
+			EffectiveGridBounds = GeneratingShape->GetLocalBounds();
+
+			// If local bounds exists, apply the transform without rotation to get into grid sampling space
+			if (EffectiveGridBounds.IsValid)
 			{
-				InputBounds = GeneratingShape->GetBounds();
+				// Remove rotation as we don't want to shear an AABB.
+				FTransform NoRotTransform(GeneratingShape->GetTransform());
+				NoRotTransform.SetRotation(FQuat::Identity);
 
-				if (BoundingShapeBounds.IsValid)
-				{
-					InputBounds = PCGHelpers::OverlapBounds(InputBounds, BoundingShapeBounds);
-				}
+				// Transport the box into semi-local 2D space, where we can get the deterministic grid samples
+				EffectiveGridBounds = EffectiveGridBounds.TransformBy(NoRotTransform);
 			}
-			else
+			else // If no local bounds, try to use the generating shape's bounds
 			{
-				InputBounds = BoundingShapeBounds;
+				if (GeneratingShape->IsBounded())
+				{
+					EffectiveGridBounds = GeneratingShape->GetBounds();
+
+					// If we're using the generating shape's bounds, we can further optimize by overlapping the bounding shape if it was provided
+					if (BoundingShapeBounds.IsValid)
+					{
+						EffectiveGridBounds = PCGHelpers::OverlapBounds(EffectiveGridBounds, BoundingShapeBounds);
+					}
+				}
+				else // If no local or world bounds, then finally try to use the bounding shape
+				{
+					EffectiveGridBounds = BoundingShapeBounds;
+				}
 			}
 
-			if (!InputBounds.IsValid || !OutState.Settings.Initialize(Settings, Context, InputBounds))
+			if (!OutState.SamplerData.Initialize(Settings, Context, EffectiveGridBounds, GeneratingShape->GetTransform()))
 			{
 				if (!GeneratingShape->IsBounded())
 				{
 					// Some inputs are unable to provide bounds, like the WorldRayHit, in which case the user must provide bounds.
 					PCGE_LOG_C(Warning, GraphAndLog, Context, LOCTEXT("CouldNotObtainInputBounds", "Input data is not bounded, so bounds must be provided for sampling. Consider providing a Bounding Shape input."));
 				}
-				else if(!InputBounds.IsValid)
+				else if (!EffectiveGridBounds.IsValid)
 				{
 					PCGE_LOG_C(Verbose, LogOnly, Context, LOCTEXT("InvalidSamplingBounds", "Final sampling bounds is invalid/zero-sized."));
 				}
@@ -519,11 +564,11 @@ bool FPCGSurfaceSamplerElement::ExecuteInternal(FPCGContext* InContext) const
 		check(InitResult == EPCGTimeSliceInitResult::Success);
 
 		// Run the execution until the time slice is finished
-		const bool bAsyncDone = PCGSurfaceSampler::SampleSurface(Context, IterState.Settings, ExecState.GeneratingShapes[IterationIndex], ExecState.BoundingShape, IterState.OutputPoints, Context->TimeSliceIsEnabled());
+		const bool bAsyncDone = PCGSurfaceSampler::SampleSurface(Context, IterState.SamplerData, ExecState.GeneratingShapes[IterationIndex], ExecState.BoundingShape, IterState.OutputPoints, Context->TimeSliceIsEnabled());
 
 		if (bAsyncDone)
 		{
-			PCGE_LOG_C(Verbose, LogOnly, Context, FText::Format(LOCTEXT("GenerationInfo", "Generated {0} points in {1} cells"), IterState.OutputPoints->GetPoints().Num(), IterState.Settings.CellCount));
+			PCGE_LOG_C(Verbose, LogOnly, Context, FText::Format(LOCTEXT("GenerationInfo", "Generated {0} points in {1} cells"), IterState.OutputPoints->GetPoints().Num(), IterState.SamplerData.CellCount));
 		}
 
 		return bAsyncDone;
diff --git a/Engine/Plugins/PCG/Source/PCG/Private/Elements/PCGWorldQuery.cpp b/Engine/Plugins/PCG/Source/PCG/Private/Elements/PCGWorldQuery.cpp
--- a/Engine/Plugins/PCG/Source/PCG/Private/Elements/PCGWorldQuery.cpp
+++ b/Engine/Plugins/PCG/Source/PCG/Private/Elements/PCGWorldQuery.cpp
@@ -93,18 +93,18 @@ bool FPCGWorldRayHitQueryElement::ExecuteInternal(FPCGContext* Context) const
 	
 	// TODO: Support params pin + Apply param data
 
+	// TODO: Might want to revisit this for 3D partitioning. The reasoning here is that ray origin should be the same
+	// if we are partitioned or non-partitioned. But with 3D partitioning we might want something different.
+	// Since it is not yet widely used, we'll stick with same behavior for all cases.
+	UPCGComponent* SourceComponent = Context->SourceComponent.Get();
+	UPCGComponent* OriginalComponent = SourceComponent ? SourceComponent->GetOriginalComponent() : nullptr;
+	AActor* Owner = OriginalComponent ? OriginalComponent->GetOwner() : nullptr;
+	FTransform Transform = Owner ? Owner->GetTransform() : FTransform::Identity;
+	FBox LocalBounds = Owner ? PCGHelpers::GetActorLocalBounds(Owner) : FBox(EForceInit::ForceInit);
+
 	if (!QueryParams.bOverrideDefaultParams)
 	{
 		// Compute default parameters based on original owner component - raycast down local Z axis
-		// TODO: Might want to revisit this for 3D partitioning. The reasoning here is that ray origin should be the same
-		// if we are partitioned or non-partitioned. But with 3D partitioning we might want something different.
-		// Since it is not yet widely used, we'll stick with same behavior for all cases.
-		UPCGComponent* SourceComponent = Context->SourceComponent.Get();
-		UPCGComponent* OriginalComponent = SourceComponent ? SourceComponent->GetOriginalComponent() : nullptr;
-		AActor* Owner = OriginalComponent ? OriginalComponent->GetOwner() : nullptr;
-		const FTransform& Transform = Owner ? Owner->GetTransform() : FTransform::Identity;
-
-		const FBox LocalBounds = PCGHelpers::GetActorLocalBounds(Owner);
 		const FVector RayOrigin = Transform.TransformPosition(FVector(0, 0, LocalBounds.Max.Z));
 		const FVector RayEnd = Transform.TransformPosition(FVector(0, 0, LocalBounds.Min.Z));
 
@@ -133,7 +133,7 @@ bool FPCGWorldRayHitQueryElement::ExecuteInternal(FPCGContext* Context) const
 	UWorld* World = Context->SourceComponent->GetWorld();
 
 	UPCGWorldRayHitData* Data = NewObject<UPCGWorldRayHitData>();
-	Data->Initialize(World);
+	Data->Initialize(World, Transform, /*InBounds=*/FBox(EForceInit::ForceInit), LocalBounds);
 	Data->QueryParams = QueryParams;
 	Data->QueryParams.Initialize();
 	Data->OriginatingComponent = Context->SourceComponent;
diff --git a/Engine/Plugins/PCG/Source/PCG/Public/Data/PCGSurfaceData.h b/Engine/Plugins/PCG/Source/PCG/Public/Data/PCGSurfaceData.h
--- a/Engine/Plugins/PCG/Source/PCG/Public/Data/PCGSurfaceData.h
+++ b/Engine/Plugins/PCG/Source/PCG/Public/Data/PCGSurfaceData.h
@@ -21,11 +21,17 @@ public:
 	virtual bool HasNonTrivialTransform() const override { return true; }
 	//~End UPCGSpatialData interface
 
+	const FTransform& GetTransform() const { return Transform; }
+	const FBox& GetLocalBounds() const { return LocalBounds; }
+
 protected:
 	void CopyBaseSurfaceData(UPCGSurfaceData* NewSurfaceData) const;
 
 	UPROPERTY(BlueprintReadOnly, VisibleAnywhere, Category = SpatialData)
 	FTransform Transform;
+
+	UPROPERTY(BlueprintReadOnly, VisibleAnywhere, Category = SpatialData)
+	FBox LocalBounds = FBox(EForceInit::ForceInit);
 };
 
 #if UE_ENABLE_INCLUDE_ORDER_DEPRECATED_IN_5_2
diff --git a/Engine/Plugins/PCG/Source/PCG/Public/Data/PCGWorldData.h b/Engine/Plugins/PCG/Source/PCG/Public/Data/PCGWorldData.h
--- a/Engine/Plugins/PCG/Source/PCG/Public/Data/PCGWorldData.h
+++ b/Engine/Plugins/PCG/Source/PCG/Public/Data/PCGWorldData.h
@@ -138,6 +138,9 @@ struct FPCGWorldRayHitQueryParams : public FPCGWorldCommonQueryParams
 	// TODO: see in FCollisionQueryParams if there are some flags we want to expose
 	// examples: bReturnFaceIndex, bReturnPhysicalMaterial, some ignore patterns
 
+	UPROPERTY(BlueprintReadWrite, EditAnywhere, Category = Data, meta = (PCG_Overridable))
+	bool bIgnoreBackfaceHits = false;
+
 	UPROPERTY(BlueprintReadWrite, EditAnywhere, Category = Data, meta = (PCG_Overridable))
 	bool bApplyMetadataFromLandscape = false;
 
@@ -152,7 +155,7 @@ class UPCGWorldRayHitData : public UPCGSurfaceData
 	GENERATED_BODY()
 
 public:
-	PCG_API void Initialize(UWorld* InWorld, const FBox& InBounds = FBox(EForceInit::ForceInit));
+	PCG_API void Initialize(UWorld* InWorld, const FTransform& InTransform, const FBox& InBounds = FBox(EForceInit::ForceInit), const FBox& InLocalBounds = FBox(EForceInit::ForceInit));
 
 	// ~Begin UPCGData interface
 	virtual void AddToCrc(FArchiveCrc32& Ar, bool bFullDataCrc) const override;
@@ -164,6 +167,7 @@ public:
 	virtual bool IsBounded() const override { return !!Bounds.IsValid; }
 	virtual bool SamplePoint(const FTransform& Transform, const FBox& Bounds, FPCGPoint& OutPoint, UPCGMetadata* OutMetadata) const override;
 	virtual bool HasNonTrivialTransform() const override { return true; }
+	virtual FVector GetNormal() const override { return Transform.GetRotation().GetUpVector(); }
 protected:
 	virtual UPCGSpatialData* CopyInternal() const override;
 	//~End UPCGSpatialData interface
@@ -173,7 +177,7 @@ public:
 	virtual bool SupportsBoundedPointData() const { return true; }
 	virtual const UPCGPointData* CreatePointData(FPCGContext* Context) const override { return CreatePointData(Context, FBox(EForceInit::ForceInit)); }
 	virtual const UPCGPointData* CreatePointData(FPCGContext* Context, const FBox& InBounds) const override;
-	// ~End UPCGConcreteDataWithPointCache interface
+	// ~End UPCGSpatialDataWithPointCache interface
 
 	UPROPERTY()
 	TWeakObjectPtr<UWorld> World = nullptr;
diff --git a/Engine/Plugins/PCG/Source/PCG/Public/Elements/PCGSurfaceSampler.h b/Engine/Plugins/PCG/Source/PCG/Public/Elements/PCGSurfaceSampler.h
--- a/Engine/Plugins/PCG/Source/PCG/Public/Elements/PCGSurfaceSampler.h
+++ b/Engine/Plugins/PCG/Source/PCG/Public/Elements/PCGSurfaceSampler.h
@@ -11,7 +11,7 @@
 
 class UPCGNode;
 class UPCGPointData;
-class UPCGSpatialData;
+class UPCGSurfaceData;
 class UPCGSurfaceSamplerSettings;
 
 namespace PCGSurfaceSamplerConstants
@@ -22,9 +22,10 @@ namespace PCGSurfaceSamplerConstants
 
 namespace PCGSurfaceSampler
 {
-	// TODO: Factor out "computed values" that shouldn't be exposed to the external user into a more appropriate state
 	struct FSurfaceSamplerParams
 	{
+		void InitializeFromSettings(const UPCGSurfaceSamplerSettings* Settings);
+
 		float PointsPerSquaredMeter = 1.0f;
 		FVector PointExtents = FVector::One() * 0.5f;
 		float Looseness = 0.0f;
@@ -33,45 +34,58 @@ namespace PCGSurfaceSampler
 #if WITH_EDITORONLY_DATA
 		bool bKeepZeroDensityPoints = false;
 #endif
+	};
+
+	struct FSurfaceSamplerData
+	{
+		bool Initialize(const UPCGSurfaceSamplerSettings* Settings, const FPCGContext* Context, const FBox& InEffectiveGridBounds, const FTransform& InSurfaceTransform = FTransform::Identity);
+		bool Initialize(const FPCGContext* Context, const FBox& InEffectiveGridBounds, const FTransform& InSurfaceTransform = FTransform::Identity);
 
-		bool Initialize(const UPCGSurfaceSamplerSettings* Settings, const FPCGContext* Context, const FBox& InputBounds);
 		FIntVector2 ComputeCellIndices(int32 Index) const;
 
+		FSurfaceSamplerParams Params;
+
+		/** Whether the pre-projected points need a local->world transformation. */
+		bool bNeedsLocalTransformation = false;
+
 		/** Computed values **/
-		FVector InterstitialDistance;
-		FVector InnerCellSize;
-		FVector CellSize;
-
-		int32 CellMinX;
-		int32 CellMaxX;
-		int32 CellMinY;
-		int32 CellMaxY;
-		int32 CellCount;
-		float Ratio;
-		int Seed;
-
-		FVector::FReal InputBoundsMinZ;
-		FVector::FReal InputBoundsMaxZ;
+		FVector InterstitialDistance = FVector::ZeroVector;
+		FVector InnerCellSize = FVector::ZeroVector;
+		FVector CellSize = FVector::ZeroVector;
+
+		int32 CellMinX = 0;
+		int32 CellMaxX = 0;
+		int32 CellMinY = 0;
+		int32 CellMaxY = 0;
+		int32 CellCount = 0;
+		float Ratio = 0.f;
+		int Seed = 0;
+
+		/** The displacement of the pre-projected, pre-sampled points. */
+		FVector::FReal PreProjectionDisplacement = 0.0;
+
+		/** Will be applied to pre-sampled points to transform them onto the pre-projection plane. */
+		FMatrix PreProjectionTransform = FMatrix::Identity;
 	};
 
 	struct FSurfaceSamplerExecutionState
 	{
 		const UPCGSpatialData* BoundingShape = nullptr;
 		FBox BoundingShapeBounds = FBox(EForceInit::ForceInit);
-		TArray<const UPCGSpatialData*> GeneratingShapes;
+		TArray<const UPCGSurfaceData*> GeneratingShapes;
 	};
 
 	struct FSurfaceSamplerIterationState
 	{
-		FSurfaceSamplerParams Settings;
+		FSurfaceSamplerData SamplerData;
 		UPCGPointData* OutputPoints = nullptr;
 	};
 
 	/** Sample a surface and returns the resulting point data. Can't be timesliced. */
-	UPCGPointData* SampleSurface(FPCGContext* Context, const UPCGSpatialData* InSurface, const UPCGSpatialData* InBoundingShape, const FSurfaceSamplerParams& ExecutionSettings);
+	UPCGPointData* SampleSurface(FPCGContext* Context, const UPCGSurfaceData* InSurface, const UPCGSpatialData* InBoundingShape, const FBox& EffectiveBounds, const FSurfaceSamplerParams& ExecutionParams);
 
 	/** Sample a surface and write the results in the given point data. Can be timesliced, and will return false if the processing is not done, true otherwise. */
-	bool SampleSurface(FPCGContext* Context, const FSurfaceSamplerParams& Settings, const UPCGSpatialData* InSurface, const UPCGSpatialData* InBoundingShape, UPCGPointData* SampledData, const bool bTimeSlicingIsEnabled = false);
+	bool SampleSurface(FPCGContext* Context, const FSurfaceSamplerData& SamplerData, const UPCGSurfaceData* InSurface, const UPCGSpatialData* InBoundingShape, UPCGPointData* SampledData, const bool bTimeSlicingIsEnabled = false);
 }
 
 UCLASS(MinimalAPI, BlueprintType, ClassGroup = (Procedural))
