[PCG] Copy Attribute on Points was not working when copying an attribute into another

* We were not copying it, just creating a new attribute with a default value
* Because we were setting the default value to the attribute first entry, we keep this behavior with the copy.
* Make it also work with Params
* Add unit tests for all those cases

#rnx

diff --git a/Engine/Plugins/PCG/Source/PCG/Private/Elements/Metadata/PCGMetadataElement.cpp b/Engine/Plugins/PCG/Source/PCG/Private/Elements/Metadata/PCGMetadataElement.cpp
--- a/Engine/Plugins/PCG/Source/PCG/Private/Elements/Metadata/PCGMetadataElement.cpp
+++ b/Engine/Plugins/PCG/Source/PCG/Private/Elements/Metadata/PCGMetadataElement.cpp
@@ -86,13 +86,21 @@ void UPCGMetadataOperationSettings::ApplyDeprecation(UPCGNode* InOutNode)
 TArray<FPCGPinProperties> UPCGMetadataOperationSettings::InputPinProperties() const
 {
 	TArray<FPCGPinProperties> Properties;
-	FPCGPinProperties& InputPinProperty = Properties.Emplace_GetRef(PCGPinConstants::DefaultInputLabel, EPCGDataType::Point);
+	FPCGPinProperties& InputPinProperty = Properties.Emplace_GetRef(PCGPinConstants::DefaultInputLabel, EPCGDataType::PointOrParam);
 	InputPinProperty.SetRequiredPin();
 
 	Properties.Emplace(PCGMetadataOperationSettings::AttributeLabel, EPCGDataType::Param, /*bInAllowMultipleConnections=*/ false, /*bAllowMultipleData=*/ false, PCGMetadataOperationSettings::AttributeTooltip);
 	return Properties;
 }
 
+TArray<FPCGPinProperties> UPCGMetadataOperationSettings::OutputPinProperties() const
+{
+	TArray<FPCGPinProperties> Properties;
+	Properties.Emplace(PCGPinConstants::DefaultOutputLabel, EPCGDataType::PointOrParam);
+
+	return Properties;
+}
+
 void UPCGMetadataOperationSettings::PostLoad()
 {
 	Super::PostLoad();
@@ -121,25 +129,16 @@ bool FPCGMetadataOperationElement::ExecuteInternal(FPCGContext* Context) const
 	for (const FPCGTaggedData& Input : Inputs)
 	{
 		FPCGTaggedData& Output = Outputs.Add_GetRef(Input);
-
-		const UPCGSpatialData* SpatialInput = Cast<const UPCGSpatialData>(Input.Data);
-
-		if (!SpatialInput)
-		{
-			PCGE_LOG(Error, GraphAndLog, LOCTEXT("InvalidInputData", "Invalid input data type, must be of type Spatial"));
-			continue;
-		}
-
-		const UPCGPointData* OriginalData = SpatialInput->ToPointData(Context);
+		const UPCGData* OriginalData = Input.Data;
 
 		if (!OriginalData)
 		{
-			PCGE_LOG(Error, GraphAndLog, LOCTEXT("UnableToGetPointData", "Unable to get point data from input"));
+			PCGE_LOG(Error, GraphAndLog, LOCTEXT("InvalidInputData", "Invalid input data"));
 			continue;
 		}
 
 		const UPCGData* SourceData = SourceAttributeSet ? static_cast<const UPCGData*>(SourceAttributeSet) : static_cast<const UPCGData*>(OriginalData);
-		const UPCGMetadata* SourceMetadata = SourceAttributeSet ? SourceAttributeSet->Metadata : OriginalData->Metadata;
+		const UPCGMetadata* SourceMetadata = SourceAttributeSet ? SourceAttributeSet->Metadata : OriginalData->ConstMetadata();
 
 		if (!SourceMetadata)
 		{
@@ -150,21 +149,28 @@ bool FPCGMetadataOperationElement::ExecuteInternal(FPCGContext* Context) const
 		// Early out when trying to copy all attributes from self, since nothing will happen
 		if (!SourceAttributeSet && Settings->bCopyAllAttributes)
 		{
-			PCGE_LOG(Verbose, LogOnly, LOCTEXT("TrivialCopy", "Copying all attributes on itself is a trivial operation."));
+			PCGE_LOG(Verbose, LogOnly, LOCTEXT("TrivialAllCopy", "Copying all attributes on itself is a trivial operation."));
+			continue;
+		}
+
+		// Early out when trying to copy same value to self
+		if (!SourceAttributeSet && Settings->InputSource == Settings->OutputTarget)
+		{
+			PCGE_LOG(Verbose, LogOnly, LOCTEXT("TrivialCopy", "Copying attribute to itself is a trivial operation."));
 			continue;
 		}
 
-		UPCGPointData* SampledData = CastChecked<UPCGPointData>(OriginalData->DuplicateData());
+		UPCGData* OutputData = OriginalData->DuplicateData();
 		bool bSuccess = false;
 		if (Settings->bCopyAllAttributes)
 		{
-			bSuccess = PCGMetadataHelpers::CopyAllAttributes(SourceData, SampledData, Context);
+			bSuccess = PCGMetadataHelpers::CopyAllAttributes(SourceData, OutputData, Context);
 		}
 		else
 		{
 			PCGMetadataHelpers::FPCGCopyAttributeParams Params{};
 			Params.SourceData = SourceData;
-			Params.TargetData = SampledData;
+			Params.TargetData = OutputData;
 			Params.InputSource = Settings->InputSource;
 			Params.OutputTarget = Settings->OutputTarget;
 			Params.OptionalContext = Context;
@@ -175,7 +181,7 @@ bool FPCGMetadataOperationElement::ExecuteInternal(FPCGContext* Context) const
 
 		if (bSuccess)
 		{
-			Output.Data = SampledData;
+			Output.Data = OutputData;
 		}
 	}
 
diff --git a/Engine/Plugins/PCG/Source/PCG/Private/Helpers/PCGMetadataHelpers.cpp b/Engine/Plugins/PCG/Source/PCG/Private/Helpers/PCGMetadataHelpers.cpp
--- a/Engine/Plugins/PCG/Source/PCG/Private/Helpers/PCGMetadataHelpers.cpp
+++ b/Engine/Plugins/PCG/Source/PCG/Private/Helpers/PCGMetadataHelpers.cpp
@@ -162,15 +162,12 @@ namespace PCGMetadataHelpers
 				// Presence of attribute was already checked before, this should not return null
 				check(SourceAttribute);
 
-				// Copy the attribute using the first entry of the source attribute as the default value (there is just a single entry or none). If there is no first entry, will be the default value anyway.
-				auto CreateAttribute = [TargetMetadata, SourceAttribute, LocalDestinationAttribute](auto Dummy) -> FPCGMetadataAttributeBase*
+				if (FPCGMetadataAttributeBase* NewAttr = TargetMetadata->CopyAttribute(SourceAttribute, LocalDestinationAttribute, /*bKeepParent=*/false, /*bCopyEntries=*/true,/*bCopyValues=*/true))
 				{
-					using AttributeType = decltype(Dummy);
-					AttributeType DefaultValue = static_cast<const FPCGMetadataAttribute<AttributeType>*>(SourceAttribute)->GetValue(PCGMetadataEntryKey(0));
-					return PCGMetadataElementCommon::ClearOrCreateAttribute<AttributeType>(TargetMetadata, LocalDestinationAttribute, DefaultValue);
-				};
-
-				if (!PCGMetadataAttribute::CallbackWithRightType(SourceAttribute->GetTypeId(), std::move(CreateAttribute)))
+					// To keep the previous behavior, we force the copied attribute to have its default value set to the first entry.
+					NewAttr->SetDefaultValueToFirstEntry();
+				}
+				else
 				{
 					PCGLog::LogWarningOnGraph(FText::Format(LOCTEXT("FailedCreateNewAttribute", "Failed to create new attribute '{0}'"), FText::FromName(LocalDestinationAttribute)));
 					continue;
diff --git a/Engine/Plugins/PCG/Source/PCG/Private/Tests/Elements/PCGCopyAttributeTest.cpp b/Engine/Plugins/PCG/Source/PCG/Private/Tests/Elements/PCGCopyAttributeTest.cpp
new file mode 100644
--- /dev/null
+++ b/Engine/Plugins/PCG/Source/PCG/Private/Tests/Elements/PCGCopyAttributeTest.cpp
@@ -0,0 +1,436 @@
+// Copyright Epic Games, Inc. All Rights Reserved.
+
+#if WITH_EDITOR
+
+#include "Tests/PCGTestsCommon.h"
+
+#include "PCGComponent.h"
+#include "PCGContext.h"
+#include "PCGParamData.h"
+#include "Data/PCGPointData.h"
+#include "Metadata/PCGMetadata.h"
+#include "Metadata/PCGMetadataAttributeTpl.h"
+
+#include "Elements/Metadata/PCGMetadataElement.h"
+
+IMPLEMENT_CUSTOM_SIMPLE_AUTOMATION_TEST(FPCGCopyAttributeTests_Points_PropertyToProperty, FPCGTestBaseClass, "Plugins.PCG.CopyAttribute.Points.PropertyToProperty", PCGTestsCommon::TestFlags)
+IMPLEMENT_CUSTOM_SIMPLE_AUTOMATION_TEST(FPCGCopyAttributeTests_Points_PropertyToAttribute, FPCGTestBaseClass, "Plugins.PCG.CopyAttribute.Points.PropertyToAttribute", PCGTestsCommon::TestFlags)
+IMPLEMENT_CUSTOM_SIMPLE_AUTOMATION_TEST(FPCGCopyAttributeTests_Points_AttributeToProperty, FPCGTestBaseClass, "Plugins.PCG.CopyAttribute.Points.AttributeToProperty", PCGTestsCommon::TestFlags)
+IMPLEMENT_CUSTOM_SIMPLE_AUTOMATION_TEST(FPCGCopyAttributeTests_Points_AttributeToAttribute, FPCGTestBaseClass, "Plugins.PCG.CopyAttribute.Points.AttributeToAttribute", PCGTestsCommon::TestFlags)
+IMPLEMENT_CUSTOM_SIMPLE_AUTOMATION_TEST(FPCGCopyAttributeTests_Points_CopyingToItself, FPCGTestBaseClass, "Plugins.PCG.CopyAttribute.Points.CopyingToItself", PCGTestsCommon::TestFlags)
+IMPLEMENT_CUSTOM_SIMPLE_AUTOMATION_TEST(FPCGCopyAttributeTests_Points_CopyingAllToItself, FPCGTestBaseClass, "Plugins.PCG.CopyAttribute.Points.CopyingAllToItself", PCGTestsCommon::TestFlags)
+
+IMPLEMENT_CUSTOM_SIMPLE_AUTOMATION_TEST(FPCGCopyAttributeTests_Params_SingleValue, FPCGTestBaseClass, "Plugins.PCG.CopyAttribute.Params.SingleValue", PCGTestsCommon::TestFlags)
+IMPLEMENT_CUSTOM_SIMPLE_AUTOMATION_TEST(FPCGCopyAttributeTests_Params_MultiValue, FPCGTestBaseClass, "Plugins.PCG.CopyAttribute.Params.MultiValue", PCGTestsCommon::TestFlags)
+IMPLEMENT_CUSTOM_SIMPLE_AUTOMATION_TEST(FPCGCopyAttributeTests_Params_CopyingToItself, FPCGTestBaseClass, "Plugins.PCG.CopyAttribute.Params.CopyingToItself", PCGTestsCommon::TestFlags)
+IMPLEMENT_CUSTOM_SIMPLE_AUTOMATION_TEST(FPCGCopyAttributeTests_Params_CopyingAllToItself, FPCGTestBaseClass, "Plugins.PCG.CopyAttribute.Params.CopyingAllToItself", PCGTestsCommon::TestFlags)
+
+
+namespace PCGCopyAttributeTests
+{
+	static const FName AttributeName = TEXT("Double");
+
+	UPCGPointData* CreateInputPointData(FPCGContext* Context, const int NumPoints)
+	{
+		check(Context);
+
+		UPCGPointData* NewPointData = NewObject<UPCGPointData>();
+		NewPointData->SetFlags(RF_Transient);
+
+		check(NewPointData && NewPointData->Metadata);
+
+		// Create an attribute
+		FPCGMetadataAttribute<double>* NewAttribute = NewPointData->Metadata->CreateAttribute<double>(AttributeName, /*DefaultValue=*/0.0, /*bAllowInterpolation=*/true, /*bOverrideParent=*/false);
+
+		TArray<FPCGPoint>& Points = NewPointData->GetMutablePoints();
+		Points.SetNum(NumPoints);
+		for (int i = 0; i < NumPoints; ++i)
+		{
+			// Store the index in the density
+			Points[i].Density = i;
+
+			 // And also in NewAttribute (offset by 5 to differientiate from density)
+			NewPointData->Metadata->InitializeOnSet(Points[i].MetadataEntry);
+			NewAttribute->SetValue(Points[i].MetadataEntry, i + 5);
+		}
+
+		FPCGTaggedData& InputData = Context->InputData.TaggedData.Emplace_GetRef();
+		InputData.Data = NewPointData;
+		InputData.Pin = PCGPinConstants::DefaultInputLabel;
+
+		return NewPointData;
+	}
+
+	UPCGParamData* CreateInputParamData(FPCGContext* Context, const int NumEntries)
+	{
+		check(Context);
+
+		UPCGParamData* NewParamData = NewObject<UPCGParamData>();
+		NewParamData->SetFlags(RF_Transient);
+
+		check(NewParamData && NewParamData->Metadata);
+
+		// Create an attribute
+		FPCGMetadataAttribute<double>* NewAttribute = NewParamData->Metadata->CreateAttribute<double>(AttributeName, /*DefaultValue=*/0.0, /*bAllowInterpolation=*/true, /*bOverrideParent=*/false);
+
+		for (int i = 0; i < NumEntries; ++i)
+		{
+			// Store the index offset by one in the attribute
+			NewAttribute->SetValue(NewParamData->Metadata->AddEntry(), i + 1);
+		}
+
+		FPCGTaggedData& InputData = Context->InputData.TaggedData.Emplace_GetRef();
+		InputData.Data = NewParamData;
+		InputData.Pin = PCGPinConstants::DefaultInputLabel;
+
+		return NewParamData;
+	}
+}
+
+bool FPCGCopyAttributeTests_Points_PropertyToProperty::RunTest(const FString& Parameters)
+{
+	PCGTestsCommon::FTestData TestData;
+	UPCGMetadataOperationSettings* Settings = PCGTestsCommon::GenerateSettings<UPCGMetadataOperationSettings>(TestData);
+	check(Settings);
+
+	constexpr int NumOfPoints = 20;
+
+	// Write Density in Position.X
+	Settings->InputSource.SetPointProperty(EPCGPointProperties::Density);
+	Settings->OutputTarget.Update(TEXT("$Position.X"));
+
+	TUniquePtr<FPCGContext> Context = TestData.InitializeTestContext();
+	const UPCGPointData* InputPointData = PCGCopyAttributeTests::CreateInputPointData(Context.Get(), NumOfPoints);
+
+	FPCGElementPtr TestElement = TestData.Settings->GetElement();
+	while (!TestElement->Execute(Context.Get())) {}
+
+	TArray<FPCGTaggedData> OutputTagged = Context->OutputData.GetInputsByPin(PCGPinConstants::DefaultOutputLabel);
+	const UPCGPointData* OutputData = OutputTagged.Num() == 1 ? Cast<const UPCGPointData>(OutputTagged[0].Data) : nullptr;
+
+	UTEST_NOT_NULL("There is a point data in output points", OutputData);
+	UTEST_EQUAL("There is the right number of points in output", OutputData->GetPoints().Num(), NumOfPoints);
+
+	check(OutputData);
+
+	// Density was copied in the position correctly
+	for (int i = 0; i < NumOfPoints; ++i)
+	{
+		const FPCGPoint& OutputPoint = OutputData->GetPoints()[i];
+		UTEST_EQUAL(*FString::Printf(TEXT("Position.X has the same value as density for point %d"), i), OutputPoint.Transform.GetLocation().X, static_cast<double>(OutputPoint.Density));
+	}
+
+	return true;
+}
+
+bool FPCGCopyAttributeTests_Points_PropertyToAttribute::RunTest(const FString& Parameters)
+{
+	PCGTestsCommon::FTestData TestData;
+	UPCGMetadataOperationSettings* Settings = PCGTestsCommon::GenerateSettings<UPCGMetadataOperationSettings>(TestData);
+	check(Settings);
+
+	constexpr int NumOfPoints = 20;
+	const FName OutputAttributeName = TEXT("OutputAttr");
+
+	// Write Density in Attribute
+	Settings->InputSource.SetPointProperty(EPCGPointProperties::Density);
+	Settings->OutputTarget.SetAttributeName(OutputAttributeName);
+
+	TUniquePtr<FPCGContext> Context = TestData.InitializeTestContext();
+	const UPCGPointData* InputPointData = PCGCopyAttributeTests::CreateInputPointData(Context.Get(), NumOfPoints);
+
+	FPCGElementPtr TestElement = TestData.Settings->GetElement();
+	while (!TestElement->Execute(Context.Get())) {}
+
+	TArray<FPCGTaggedData> OutputTagged = Context->OutputData.GetInputsByPin(PCGPinConstants::DefaultOutputLabel);
+	const UPCGPointData* OutputData = OutputTagged.Num() == 1 ? Cast<const UPCGPointData>(OutputTagged[0].Data) : nullptr;
+
+	UTEST_NOT_NULL("There is a point data in output points", OutputData);
+	UTEST_EQUAL("There is the right number of points in output", OutputData->GetPoints().Num(), NumOfPoints);
+
+	check(OutputData && OutputData->Metadata);
+	const FPCGMetadataAttribute<double>* OutputAttribute = OutputData->Metadata->GetConstTypedAttribute<double>(OutputAttributeName);
+	UTEST_NOT_NULL("Output Attribute exists in the output data", OutputAttribute);
+
+	check(OutputAttribute);
+
+	// Attribute was copied in the position correctly
+	for (int i = 0; i < NumOfPoints; ++i)
+	{
+		const FPCGPoint& OutputPoint = OutputData->GetPoints()[i];
+		UTEST_EQUAL(*FString::Printf(TEXT("Output attribute has the same value than density for point %d"), i), OutputAttribute->GetValueFromItemKey(OutputPoint.MetadataEntry), static_cast<double>(OutputPoint.Density));
+	}
+
+	return true;
+}
+
+bool FPCGCopyAttributeTests_Points_AttributeToProperty::RunTest(const FString& Parameters)
+{
+	PCGTestsCommon::FTestData TestData;
+	UPCGMetadataOperationSettings* Settings = PCGTestsCommon::GenerateSettings<UPCGMetadataOperationSettings>(TestData);
+	check(Settings);
+
+	constexpr int NumOfPoints = 20;
+
+	// Write Attribute in Position.Y
+	Settings->InputSource.SetAttributeName(PCGCopyAttributeTests::AttributeName);
+	Settings->OutputTarget.Update(TEXT("$Position.Y"));
+
+	TUniquePtr<FPCGContext> Context = TestData.InitializeTestContext();
+	const UPCGPointData* InputPointData = PCGCopyAttributeTests::CreateInputPointData(Context.Get(), NumOfPoints);
+
+	FPCGElementPtr TestElement = TestData.Settings->GetElement();
+	while (!TestElement->Execute(Context.Get())) {}
+
+	TArray<FPCGTaggedData> OutputTagged = Context->OutputData.GetInputsByPin(PCGPinConstants::DefaultOutputLabel);
+	const UPCGPointData* OutputData = OutputTagged.Num() == 1 ? Cast<const UPCGPointData>(OutputTagged[0].Data) : nullptr;
+
+	UTEST_NOT_NULL("There is a point data in output points", OutputData);
+	UTEST_EQUAL("There is the right number of points in output", OutputData->GetPoints().Num(), NumOfPoints);
+
+	check(OutputData && OutputData->Metadata);
+	const FPCGMetadataAttribute<double>* InputAttribute = OutputData->Metadata->GetConstTypedAttribute<double>(PCGCopyAttributeTests::AttributeName);
+	UTEST_NOT_NULL("Input Attribute exists in the output data", InputAttribute);
+
+	// Attribute value was copied in Position.Y correctly
+	for (int i = 0; i < NumOfPoints; ++i)
+	{
+		const FPCGPoint& OutputPoint = OutputData->GetPoints()[i];
+		UTEST_EQUAL(*FString::Printf(TEXT("Position.Y has the same value than the input attribute for point %d"), i), OutputPoint.Transform.GetLocation().Y, InputAttribute->GetValueFromItemKey(OutputPoint.MetadataEntry));
+	}
+
+	return true;
+}
+
+bool FPCGCopyAttributeTests_Points_AttributeToAttribute::RunTest(const FString& Parameters)
+{
+	PCGTestsCommon::FTestData TestData;
+	UPCGMetadataOperationSettings* Settings = PCGTestsCommon::GenerateSettings<UPCGMetadataOperationSettings>(TestData);
+	check(Settings);
+
+	constexpr int NumOfPoints = 20;
+	const FName OutputAttributeName = TEXT("OutputAttr");
+
+	// Write Input attribute to Output attribute
+	Settings->InputSource.SetAttributeName(PCGCopyAttributeTests::AttributeName);
+	Settings->OutputTarget.SetAttributeName(OutputAttributeName);
+
+	TUniquePtr<FPCGContext> Context = TestData.InitializeTestContext();
+	const UPCGPointData* InputPointData = PCGCopyAttributeTests::CreateInputPointData(Context.Get(), NumOfPoints);
+
+	FPCGElementPtr TestElement = TestData.Settings->GetElement();
+	while (!TestElement->Execute(Context.Get())) {}
+
+	TArray<FPCGTaggedData> OutputTagged = Context->OutputData.GetInputsByPin(PCGPinConstants::DefaultOutputLabel);
+	const UPCGPointData* OutputData = OutputTagged.Num() == 1 ? Cast<const UPCGPointData>(OutputTagged[0].Data) : nullptr;
+
+	UTEST_NOT_NULL("There is a point data in output points", OutputData);
+	UTEST_EQUAL("There is the right number of points in output", OutputData->GetPoints().Num(), NumOfPoints);
+
+	check(OutputData && OutputData->Metadata);
+
+	const FPCGMetadataAttribute<double>* InputAttribute = OutputData->Metadata->GetConstTypedAttribute<double>(PCGCopyAttributeTests::AttributeName);
+	UTEST_NOT_NULL("Input Attribute exists in the output data", InputAttribute);
+	const FPCGMetadataAttribute<double>* OutputAttribute = OutputData->Metadata->GetConstTypedAttribute<double>(OutputAttributeName);
+	UTEST_NOT_NULL("Output Attribute exists in the output data", OutputAttribute);
+
+	check(InputAttribute && OutputAttribute);
+
+	UTEST_EQUAL("Output Attribute default value is the same as first entry of Input Attribute", OutputAttribute->GetValueFromItemKey(PCGInvalidEntryKey), InputAttribute->GetValueFromItemKey(PCGMetadataEntryKey(0)));
+
+	// Input attribute was copied in the output attribute correctly
+	for (int i = 0; i < NumOfPoints; ++i)
+	{
+		const FPCGPoint& OutputPoint = OutputData->GetPoints()[i];
+		UTEST_EQUAL(*FString::Printf(TEXT("Output Attribute has the same value as Input Attribute for point %d"), i), OutputAttribute->GetValueFromItemKey(OutputPoint.MetadataEntry), InputAttribute->GetValueFromItemKey(OutputPoint.MetadataEntry));
+	}
+
+	return true;
+}
+
+bool FPCGCopyAttributeTests_Points_CopyingToItself::RunTest(const FString& Parameters)
+{
+	PCGTestsCommon::FTestData TestData;
+	UPCGMetadataOperationSettings* Settings = PCGTestsCommon::GenerateSettings<UPCGMetadataOperationSettings>(TestData);
+	check(Settings);
+
+	constexpr int NumOfPoints = 20;
+
+	// Write Attribute in itself
+	Settings->InputSource.SetAttributeName(PCGCopyAttributeTests::AttributeName);
+	Settings->OutputTarget.SetAttributeName(PCGCopyAttributeTests::AttributeName);
+
+	TUniquePtr<FPCGContext> Context = TestData.InitializeTestContext();
+	const UPCGPointData* InputPointData = PCGCopyAttributeTests::CreateInputPointData(Context.Get(), NumOfPoints);
+
+	FPCGElementPtr TestElement = TestData.Settings->GetElement();
+	while (!TestElement->Execute(Context.Get())) {}
+
+	TArray<FPCGTaggedData> OutputTagged = Context->OutputData.GetInputsByPin(PCGPinConstants::DefaultOutputLabel);
+	const UPCGPointData* OutputData = OutputTagged.Num() == 1 ? Cast<const UPCGPointData>(OutputTagged[0].Data) : nullptr;
+
+	UTEST_NOT_NULL("There is a point data in output points", OutputData);
+	UTEST_EQUAL("It's the same data as input", OutputData, InputPointData);
+
+	return true;
+}
+
+bool FPCGCopyAttributeTests_Points_CopyingAllToItself::RunTest(const FString& Parameters)
+{
+	PCGTestsCommon::FTestData TestData;
+	UPCGMetadataOperationSettings* Settings = PCGTestsCommon::GenerateSettings<UPCGMetadataOperationSettings>(TestData);
+	check(Settings);
+
+	constexpr int NumOfPoints = 20;
+
+	Settings->bCopyAllAttributes = true;
+
+	TUniquePtr<FPCGContext> Context = TestData.InitializeTestContext();
+	const UPCGPointData* InputPointData = PCGCopyAttributeTests::CreateInputPointData(Context.Get(), NumOfPoints);
+
+	FPCGElementPtr TestElement = TestData.Settings->GetElement();
+	while (!TestElement->Execute(Context.Get())) {}
+
+	TArray<FPCGTaggedData> OutputTagged = Context->OutputData.GetInputsByPin(PCGPinConstants::DefaultOutputLabel);
+	const UPCGPointData* OutputData = OutputTagged.Num() == 1 ? Cast<const UPCGPointData>(OutputTagged[0].Data) : nullptr;
+
+	UTEST_NOT_NULL("There is a point data in output points", OutputData);
+	UTEST_EQUAL("It's the same data as input", OutputData, InputPointData);
+
+	return true;
+}
+
+bool FPCGCopyAttributeTests_Params_SingleValue::RunTest(const FString& Parameters)
+{
+	PCGTestsCommon::FTestData TestData;
+	UPCGMetadataOperationSettings* Settings = PCGTestsCommon::GenerateSettings<UPCGMetadataOperationSettings>(TestData);
+	check(Settings);
+
+	const FName OutputAttributeName = TEXT("OutputAttr");
+
+	// Write Input attribute to Output attribute
+	Settings->InputSource.SetAttributeName(PCGCopyAttributeTests::AttributeName);
+	Settings->OutputTarget.SetAttributeName(OutputAttributeName);
+
+	TUniquePtr<FPCGContext> Context = TestData.InitializeTestContext();
+	const UPCGParamData* InputParamData = PCGCopyAttributeTests::CreateInputParamData(Context.Get(), 1);
+
+	FPCGElementPtr TestElement = TestData.Settings->GetElement();
+	while (!TestElement->Execute(Context.Get())) {}
+
+	TArray<FPCGTaggedData> OutputTagged = Context->OutputData.GetInputsByPin(PCGPinConstants::DefaultOutputLabel);
+	const UPCGParamData* OutputData = OutputTagged.Num() == 1 ? Cast<const UPCGParamData>(OutputTagged[0].Data) : nullptr;
+
+	UTEST_NOT_NULL("There is a param data in output points", OutputData);
+	UTEST_EQUAL("There is the right number of entries in output metadata", OutputData->Metadata->GetItemCountForChild(), 1);
+
+	check(OutputData && OutputData->Metadata);
+
+	const FPCGMetadataAttribute<double>* InputAttribute = OutputData->Metadata->GetConstTypedAttribute<double>(PCGCopyAttributeTests::AttributeName);
+	UTEST_NOT_NULL("Input Attribute exists in the output data", InputAttribute);
+	const FPCGMetadataAttribute<double>* OutputAttribute = OutputData->Metadata->GetConstTypedAttribute<double>(OutputAttributeName);
+	UTEST_NOT_NULL("Output Attribute exists in the output data", OutputAttribute);
+
+	check(InputAttribute && OutputAttribute);
+
+	UTEST_EQUAL("Output Attribute default value is the same as first entry of Input Attribute", OutputAttribute->GetValueFromItemKey(PCGInvalidEntryKey), InputAttribute->GetValueFromItemKey(PCGMetadataEntryKey(0)));
+	UTEST_EQUAL("Output Attribute has the same value as Input Attribute for entry 0", OutputAttribute->GetValueFromItemKey(PCGMetadataEntryKey(0)), InputAttribute->GetValueFromItemKey(PCGMetadataEntryKey(0)));
+
+	return true;
+}
+
+bool FPCGCopyAttributeTests_Params_MultiValue::RunTest(const FString& Parameters)
+{
+	PCGTestsCommon::FTestData TestData;
+	UPCGMetadataOperationSettings* Settings = PCGTestsCommon::GenerateSettings<UPCGMetadataOperationSettings>(TestData);
+	check(Settings);
+
+	constexpr int NumEntries = 20;
+	const FName OutputAttributeName = TEXT("OutputAttr");
+
+	// Write Input attribute to Output attribute
+	Settings->InputSource.SetAttributeName(PCGCopyAttributeTests::AttributeName);
+	Settings->OutputTarget.SetAttributeName(OutputAttributeName);
+
+	TUniquePtr<FPCGContext> Context = TestData.InitializeTestContext();
+	const UPCGParamData* InputParamData = PCGCopyAttributeTests::CreateInputParamData(Context.Get(), NumEntries);
+
+	FPCGElementPtr TestElement = TestData.Settings->GetElement();
+	while (!TestElement->Execute(Context.Get())) {}
+
+	TArray<FPCGTaggedData> OutputTagged = Context->OutputData.GetInputsByPin(PCGPinConstants::DefaultOutputLabel);
+	const UPCGParamData* OutputData = OutputTagged.Num() == 1 ? Cast<const UPCGParamData>(OutputTagged[0].Data) : nullptr;
+
+	UTEST_NOT_NULL("There is a param data in output points", OutputData);
+	UTEST_EQUAL("There is the right number of entries in output metadata", OutputData->Metadata->GetItemCountForChild(), NumEntries);
+
+	check(OutputData && OutputData->Metadata);
+
+	const FPCGMetadataAttribute<double>* InputAttribute = OutputData->Metadata->GetConstTypedAttribute<double>(PCGCopyAttributeTests::AttributeName);
+	UTEST_NOT_NULL("Input Attribute exists in the output data", InputAttribute);
+	const FPCGMetadataAttribute<double>* OutputAttribute = OutputData->Metadata->GetConstTypedAttribute<double>(OutputAttributeName);
+	UTEST_NOT_NULL("Output Attribute exists in the output data", OutputAttribute);
+
+	check(InputAttribute && OutputAttribute);
+
+	UTEST_EQUAL("Output Attribute default value is the same as first entry of Input Attribute", OutputAttribute->GetValueFromItemKey(PCGInvalidEntryKey), InputAttribute->GetValueFromItemKey(PCGMetadataEntryKey(0)));
+
+	// Input attribute was copied in the output attribute correctly
+	for (int i = 0; i < NumEntries; ++i)
+	{
+		UTEST_EQUAL(*FString::Printf(TEXT("Output Attribute has the same value as Input Attribute for entry %d"), i), OutputAttribute->GetValueFromItemKey(PCGMetadataEntryKey(i)), InputAttribute->GetValueFromItemKey(PCGMetadataEntryKey(i)));
+	}
+
+	return true;
+}
+
+bool FPCGCopyAttributeTests_Params_CopyingToItself::RunTest(const FString& Parameters)
+{
+	PCGTestsCommon::FTestData TestData;
+	UPCGMetadataOperationSettings* Settings = PCGTestsCommon::GenerateSettings<UPCGMetadataOperationSettings>(TestData);
+	check(Settings);
+
+	// Write Input attribute to Output attribute
+	Settings->InputSource.SetAttributeName(PCGCopyAttributeTests::AttributeName);
+	Settings->OutputTarget.SetAttributeName(PCGCopyAttributeTests::AttributeName);
+
+	TUniquePtr<FPCGContext> Context = TestData.InitializeTestContext();
+	const UPCGParamData* InputParamData = PCGCopyAttributeTests::CreateInputParamData(Context.Get(), 1);
+
+	FPCGElementPtr TestElement = TestData.Settings->GetElement();
+	while (!TestElement->Execute(Context.Get())) {}
+
+	TArray<FPCGTaggedData> OutputTagged = Context->OutputData.GetInputsByPin(PCGPinConstants::DefaultOutputLabel);
+	const UPCGParamData* OutputData = OutputTagged.Num() == 1 ? Cast<const UPCGParamData>(OutputTagged[0].Data) : nullptr;
+
+	UTEST_NOT_NULL("There is a param data in output points", OutputData);
+	UTEST_EQUAL("It's the same data as input", OutputData, InputParamData);
+
+	return true;
+}
+
+bool FPCGCopyAttributeTests_Params_CopyingAllToItself::RunTest(const FString& Parameters)
+{
+	PCGTestsCommon::FTestData TestData;
+	UPCGMetadataOperationSettings* Settings = PCGTestsCommon::GenerateSettings<UPCGMetadataOperationSettings>(TestData);
+	check(Settings);
+
+	Settings->bCopyAllAttributes = true;
+
+	TUniquePtr<FPCGContext> Context = TestData.InitializeTestContext();
+	const UPCGParamData* InputParamData = PCGCopyAttributeTests::CreateInputParamData(Context.Get(), 1);
+
+	FPCGElementPtr TestElement = TestData.Settings->GetElement();
+	while (!TestElement->Execute(Context.Get())) {}
+
+	TArray<FPCGTaggedData> OutputTagged = Context->OutputData.GetInputsByPin(PCGPinConstants::DefaultOutputLabel);
+	const UPCGParamData* OutputData = OutputTagged.Num() == 1 ? Cast<const UPCGParamData>(OutputTagged[0].Data) : nullptr;
+
+	UTEST_NOT_NULL("There is a param data in output points", OutputData);
+	UTEST_EQUAL("It's the same data as input", OutputData, InputParamData);
+
+	return true;
+}
+
+
+#endif // WITH_EDITOR
diff --git a/Engine/Plugins/PCG/Source/PCG/Public/Elements/Metadata/PCGMetadataElement.h b/Engine/Plugins/PCG/Source/PCG/Public/Elements/Metadata/PCGMetadataElement.h
--- a/Engine/Plugins/PCG/Source/PCG/Public/Elements/Metadata/PCGMetadataElement.h
+++ b/Engine/Plugins/PCG/Source/PCG/Public/Elements/Metadata/PCGMetadataElement.h
@@ -27,6 +27,7 @@ public:
 	virtual EPCGSettingsType GetType() const override { return EPCGSettingsType::Metadata; }
 	virtual bool IsPinUsedByNodeExecution(const UPCGPin* InPin) const override;
 	virtual void ApplyDeprecation(UPCGNode* InOutNode) override;
+	virtual bool HasDynamicPins() const override { return true; }
 #endif
 	virtual FString GetAdditionalTitleInformation() const override;
 
@@ -35,7 +36,7 @@ protected:
 	virtual EPCGChangeType GetChangeTypeForProperty(const FName& InPropertyName) const override;
 #endif
 	virtual TArray<FPCGPinProperties> InputPinProperties() const override;
-	virtual TArray<FPCGPinProperties> OutputPinProperties() const override { return Super::DefaultPointOutputPinProperties(); }
+	virtual TArray<FPCGPinProperties> OutputPinProperties() const override;
 	virtual FPCGElementPtr CreateElement() const override;
 	//~End UPCGSettings interface
 
