[PCG] Fix issue where Static Mesh Spawner was not setting the output point bounds correctly

* When ISMCs were re-used, we were skipping the bounds set in the Execute.

#rnx

diff --git a/Engine/Plugins/PCG/Source/PCG/Private/Elements/PCGStaticMeshSpawner.cpp b/Engine/Plugins/PCG/Source/PCG/Private/Elements/PCGStaticMeshSpawner.cpp
--- a/Engine/Plugins/PCG/Source/PCG/Private/Elements/PCGStaticMeshSpawner.cpp
+++ b/Engine/Plugins/PCG/Source/PCG/Private/Elements/PCGStaticMeshSpawner.cpp
@@ -198,25 +198,20 @@ bool FPCGStaticMeshSpawnerElement::PrepareDataInternal(FPCGContext* InContext) c
 				Context->CurrentOutputPointData = OutputPointData;
 			}
 
-			// At this point, if we're in a reuse case we don't need to create the instance list here as it won't be processed or spawned
-			if (!bSkippedDueToReuse)
-			{
-				FPCGStaticMeshSpawnerContext::FPackedInstanceListData& InstanceListData = Context->MeshInstancesData.Emplace_GetRef();
-				InstanceListData.TargetActor = TargetActor;
-				InstanceListData.SpatialData = PointData;
-			}
+			FPCGStaticMeshSpawnerContext::FPackedInstanceListData& InstanceListData = Context->MeshInstancesData.Emplace_GetRef();
+			InstanceListData.TargetActor = TargetActor;
+			InstanceListData.SpatialData = PointData;
 
 			Context->CurrentPointData = PointData;
 			Context->bCurrentInputSetup = true;
 		}
 
+		// TODO: If we know we re-use the ISMCs, we should not run the Selection, as it can be pretty costly.
+		// At the moment, the selection is filling the output point data, so it is necessary to run it. But we should just hit the cache in that case.
 		if (!Context->bSelectionDone)
 		{
-			TArray<FPCGMeshInstanceList> DummyMeshInstances;
-			TArray<FPCGMeshInstanceList>& MeshInstances = (bSkippedDueToReuse ? DummyMeshInstances : Context->MeshInstancesData.Last().MeshInstances);
-
 			check(Context->CurrentPointData);
-			Context->bSelectionDone = Settings->MeshSelectorParameters->SelectInstances(*Context, Settings, Context->CurrentPointData, MeshInstances, Context->CurrentOutputPointData);
+			Context->bSelectionDone = Settings->MeshSelectorParameters->SelectInstances(*Context, Settings, Context->CurrentPointData, Context->MeshInstancesData.Last().MeshInstances, Context->CurrentOutputPointData);
 		}
 
 		if (!Context->bSelectionDone)
@@ -295,7 +290,7 @@ bool FPCGStaticMeshSpawnerElement::ExecuteInternal(FPCGContext* InContext) const
 	while(!Context->MeshInstancesData.IsEmpty())
 	{
 		const FPCGStaticMeshSpawnerContext::FPackedInstanceListData& InstanceList = Context->MeshInstancesData.Last();
-		check(InstanceList.MeshInstances.Num() == InstanceList.PackedCustomData.Num());
+		check(Context->bSkippedDueToReuse || InstanceList.MeshInstances.Num() == InstanceList.PackedCustomData.Num());
 
 		const bool bTargetActorValid = (InstanceList.TargetActor && IsValid(InstanceList.TargetActor));
 
@@ -304,7 +299,9 @@ bool FPCGStaticMeshSpawnerElement::ExecuteInternal(FPCGContext* InContext) const
 			while (Context->CurrentDataIndex < InstanceList.MeshInstances.Num())
 			{
 				const FPCGMeshInstanceList& MeshInstance = InstanceList.MeshInstances[Context->CurrentDataIndex];
-				SpawnStaticMeshInstances(Context, MeshInstance, InstanceList.TargetActor, InstanceList.PackedCustomData[Context->CurrentDataIndex]);
+				// We always have mesh instances, but if we are in re-use, we don't compute the packed custom data.
+				const FPCGPackedCustomData* PackedCustomData = InstanceList.PackedCustomData.IsValidIndex(Context->CurrentDataIndex) ? &InstanceList.PackedCustomData[Context->CurrentDataIndex] : nullptr;
+				SpawnStaticMeshInstances(Context, MeshInstance, InstanceList.TargetActor, PackedCustomData);
 
 				// Now that the mesh is loaded/spawned, set the bounds to out points if requested.
 				if (MeshInstance.Descriptor.StaticMesh && Settings->bApplyMeshBoundsToPoints)
@@ -370,7 +367,7 @@ bool FPCGStaticMeshSpawnerElement::CanExecuteOnlyOnMainThread(FPCGContext* Conte
 	return Context->CurrentPhase == EPCGExecutionPhase::Execute || Context->CurrentPhase == EPCGExecutionPhase::PrepareData;
 }
 
-void FPCGStaticMeshSpawnerElement::SpawnStaticMeshInstances(FPCGStaticMeshSpawnerContext* Context, const FPCGMeshInstanceList& InstanceList, AActor* TargetActor, const FPCGPackedCustomData& PackedCustomData) const
+void FPCGStaticMeshSpawnerElement::SpawnStaticMeshInstances(FPCGStaticMeshSpawnerContext* Context, const FPCGMeshInstanceList& InstanceList, AActor* TargetActor, const FPCGPackedCustomData* InPackedCustomData) const
 {
 	// Populate the (H)ISM from the previously prepared entries
 	TRACE_CPUPROFILER_EVENT_SCOPE(FPCGStaticMeshSpawnerElement::Execute::PopulateISMs);
@@ -394,6 +391,12 @@ void FPCGStaticMeshSpawnerElement::SpawnStaticMeshInstances(FPCGStaticMeshSpawne
 		return;
 	}
 
+	// Don't spawn meshes if we reuse the ISMCs, but we still want to be sure that the mesh is loaded at least (for operations downstream).
+	if (Context->bSkippedDueToReuse)
+	{
+		return;
+	}
+
 	for (TSoftObjectPtr<UMaterialInterface> OverrideMaterial : InstanceList.Descriptor.OverrideMaterials)
 	{
 		// Will be synchronously loaded if not loaded. But by default it should already have been loaded asynchronously in PrepareData, so this is free.
@@ -404,6 +407,14 @@ void FPCGStaticMeshSpawnerElement::SpawnStaticMeshInstances(FPCGStaticMeshSpawne
 		}
 	}
 
+	// If we spawn the meshes, we should have computed a packed custom data.
+	if (!ensure(InPackedCustomData))
+	{
+		return;
+	}
+
+	const FPCGPackedCustomData& PackedCustomData = *InPackedCustomData;
+
 	FPCGISMCBuilderParameters Params;
 	Params.Descriptor = FISMComponentDescriptor(InstanceList.Descriptor);
 	Params.NumCustomDataFloats = PackedCustomData.NumCustomDataFloats;
diff --git a/Engine/Plugins/PCG/Source/PCG/Public/Elements/PCGStaticMeshSpawner.h b/Engine/Plugins/PCG/Source/PCG/Public/Elements/PCGStaticMeshSpawner.h
--- a/Engine/Plugins/PCG/Source/PCG/Public/Elements/PCGStaticMeshSpawner.h
+++ b/Engine/Plugins/PCG/Source/PCG/Public/Elements/PCGStaticMeshSpawner.h
@@ -118,5 +118,5 @@ protected:
 	virtual bool PrepareDataInternal(FPCGContext* Context) const override;
 	virtual bool ExecuteInternal(FPCGContext* Context) const override;
 	virtual void AbortInternal(FPCGContext* Context) const override;
-	void SpawnStaticMeshInstances(FPCGStaticMeshSpawnerContext* Context, const FPCGMeshInstanceList& InstanceList, AActor* TargetActor, const FPCGPackedCustomData& PackedCustomData) const;
+	void SpawnStaticMeshInstances(FPCGStaticMeshSpawnerContext* Context, const FPCGMeshInstanceList& InstanceList, AActor* TargetActor, const FPCGPackedCustomData* InPackedCustomData) const;
 };
