Modification to the Visualize Attribute PCG node to support larger point number display using the distance from the camera as a limit.
#changelist validated

diff --git a/Engine/Plugins/PCG/Source/PCG/Private/Elements/PCGVisualizeAttribute.cpp b/Engine/Plugins/PCG/Source/PCG/Private/Elements/PCGVisualizeAttribute.cpp
--- a/Engine/Plugins/PCG/Source/PCG/Private/Elements/PCGVisualizeAttribute.cpp
+++ b/Engine/Plugins/PCG/Source/PCG/Private/Elements/PCGVisualizeAttribute.cpp
@@ -8,6 +8,9 @@
 #include "Data/PCGPointData.h"
 #include "Metadata/Accessors/PCGAttributeAccessorHelpers.h"
 
+#include "LevelEditorViewport.h"
+#include "Math/UnrealMathUtility.h"
+
 #define LOCTEXT_NAMESPACE "PCGVisualizeAttributeElement"
 
 FPCGElementPtr UPCGVisualizeAttributeSettings::CreateElement() const
@@ -112,27 +115,51 @@ bool FPCGVisualizeAttribute::ExecuteInternal(FPCGContext* Context) const
 		TArray<FDebugRenderSceneProxy::FText3d> DebugStrings;
 		DebugStrings.Reserve(FMath::Min(Points.Num(), Settings->PointLimit));
 
-		// TODO: Optimize which points will be drawn by distance, or other measures, rather than by Index
+		// TODO: Optimize which points will be drawn by other measures, rather than only by Index and Distance
 		// Debug messages are already culled by camera frustum, but a point limit is desired to prevent overloading
-		for (int I = 0; I < Points.Num() && I < Settings->PointLimit; ++I)
-		{
-			FStringBuilderBase CompoundedString;
 
-			CompoundedString += Settings->CustomPrefixString;
+		// Store the number of points we add debug info too
+		int count = 0;
 
-			if (Settings->bPrefixWithIndex)
-			{
-				CompoundedString += FString::Printf(TEXT("[%d]"), I);
+		// Get the current camera location to use with the distance check
+		FViewportCameraTransform& ViewTransform = GCurrentLevelEditingViewportClient->GetViewTransform();
+		FVector CurrentCameraLocation = ViewTransform.GetLocation();
+
+		for (int I = 0; I < Points.Num(); ++I)
+		{
+			// Break when the number of points exceed the point limit
+			if (count > Settings->PointLimit) {
+				break;
 			}
 
-			if (Settings->bPrefixWithAttributeName)
+			// Get the current point location and the distance between it and the camera
+			FVector CurrentPointLocation = Points[I].Transform.GetLocation();
+			float Distance = FVector::DistSquared(CurrentPointLocation, CurrentCameraLocation);
+
+			// If the distance is smaller than the max allowed distance, do the display of attribute
+			if (Distance < FMath::Square(Settings->MaxDistance))
 			{
-				CompoundedString += FString::Printf(TEXT("[%s]"), *InputSource.GetDisplayText().ToString());
-			}
+				FStringBuilderBase CompoundedString;
+
+				CompoundedString += Settings->CustomPrefixString;
+
+				if (Settings->bPrefixWithIndex)
+				{
+					CompoundedString += FString::Printf(TEXT("[%d]"), I);
+				}
 
-			CompoundedString = FString::Printf(TEXT("%s %s"), *CompoundedString, *ValuesToString[I]);
+				if (Settings->bPrefixWithAttributeName)
+				{
+					CompoundedString += FString::Printf(TEXT("[%s]"), *InputSource.GetDisplayText().ToString());
+				}
 
-			DebugStrings.Emplace(CompoundedString.ToString(), Points[I].Transform.GetLocation() + Settings->LocalOffset, Settings->Color);
+				CompoundedString = FString::Printf(TEXT("%s %s"), *CompoundedString, *ValuesToString[I]);
+
+				DebugStrings.Emplace(CompoundedString.ToString(), Points[I].Transform.GetLocation() + Settings->LocalOffset, Settings->Color);
+
+				// Increment the count
+				count += 1;
+			}
 		}
 
 		check(DebugDrawComponent);
diff --git a/Engine/Plugins/PCG/Source/PCG/Public/Elements/PCGVisualizeAttribute.h b/Engine/Plugins/PCG/Source/PCG/Public/Elements/PCGVisualizeAttribute.h
--- a/Engine/Plugins/PCG/Source/PCG/Public/Elements/PCGVisualizeAttribute.h
+++ b/Engine/Plugins/PCG/Source/PCG/Public/Elements/PCGVisualizeAttribute.h
@@ -62,9 +62,13 @@ public:
 	double Duration = 30.0;
 
 	/** The limit of points to draw debug messages. */
-	UPROPERTY(BlueprintReadWrite, EditAnywhere, Category = Settings, meta = (PCG_Overridable, ClampMin = "1", ClampMax="4096"))
+	UPROPERTY(BlueprintReadWrite, EditAnywhere, Category = Settings, meta = (PCG_Overridable, ClampMin = "1"))
 	int32 PointLimit = 4096;
 
+	/** The Max Distance from the camera that the points will show debug info. */
+	UPROPERTY(BlueprintReadWrite, EditAnywhere, Category = Settings, meta = (PCG_Overridable, ClampMin = "0.0"))
+	double MaxDistance = 10000;
+
 	/** The visualizer is enabled. Useful for dynamically overriding. */
 	UPROPERTY(BlueprintReadWrite, EditAnywhere, Category = Settings, meta = (PCG_Overridable))
 	bool bVisualizeEnabled = true;
