[PCG] Fix RuntimeGen ignoring bActivated flag on PCG Components

#rnx

diff --git a/Engine/Plugins/PCG/Source/PCG/Private/RuntimeGen/PCGRuntimeGenScheduler.cpp b/Engine/Plugins/PCG/Source/PCG/Private/RuntimeGen/PCGRuntimeGenScheduler.cpp
--- a/Engine/Plugins/PCG/Source/PCG/Private/RuntimeGen/PCGRuntimeGenScheduler.cpp
+++ b/Engine/Plugins/PCG/Source/PCG/Private/RuntimeGen/PCGRuntimeGenScheduler.cpp
@@ -247,7 +247,7 @@ void FPCGRuntimeGenScheduler::TickQueueComponentsForGeneration(
 	{
 		TRACE_CPUPROFILER_EVENT_SCOPE(FPCGRuntimeGenScheduler::CollectLocalComponents);
 
-		if (!ensure(OriginalComponent) || !OriginalComponent->GetGraph())
+		if (!ensure(OriginalComponent) || !OriginalComponent->GetGraph() || !OriginalComponent->bActivated)
 		{
 			continue;
 		}
@@ -383,7 +383,7 @@ void FPCGRuntimeGenScheduler::TickQueueComponentsForGeneration(
 	{
 		TRACE_CPUPROFILER_EVENT_SCOPE(FPCGRuntimeGenScheduler::CollectNonPartitionedComponents);
 
-		if (!ensure(OriginalComponent) || !OriginalComponent->GetGraph())
+		if (!ensure(OriginalComponent) || !OriginalComponent->GetGraph() || !OriginalComponent->bActivated)
 		{
 			continue;
 		}
@@ -458,6 +458,12 @@ void FPCGRuntimeGenScheduler::TickCleanup(const TSet<IPCGGenSourceBase*>& InGenS
 		// If the Grid is unbounded, we have a non-partitioned or unbounded component.
 		if (Grid == EPCGHiGenGrid::Unbounded)
 		{
+			if (!OriginalComponent->bActivated)
+			{
+				ComponentsToClean.Add({ GenerationKey, OriginalComponent });
+				continue;
+			}
+
 			const FBox Bounds = OriginalComponent->GetGridBounds();
 
 			double MinSquaredDistanceToGenSource = UE_DOUBLE_BIG_NUMBER;
@@ -496,7 +502,7 @@ void FPCGRuntimeGenScheduler::TickCleanup(const TSet<IPCGGenSourceBase*>& InGenS
 		{
 			UPCGComponent* LocalComponent = ActorAndComponentMapping->GetLocalComponent(GridSize, GridCoords, OriginalComponent, /*bRuntimeGenerated=*/true);
 			APCGPartitionActor* PartitionActor = LocalComponent ? Cast<APCGPartitionActor>(LocalComponent->GetOwner()) : nullptr;
-			if (!PartitionActor)
+			if (!PartitionActor || !OriginalComponent->bActivated)
 			{
 				// Attempt to clean even in failure case to avoid leaking resources.
 				ComponentsToClean.Add({ GenerationKey, LocalComponent });
@@ -903,16 +909,20 @@ void FPCGRuntimeGenScheduler::CleanupDelayedRefreshComponents()
 	{
 		const uint32 GridSize = GenerationKey.GetGridSize();
 		const EPCGHiGenGrid Grid = PCGHiGenGrid::GridSizeToGrid(GridSize);
+		UPCGComponent* OriginalComponent = GenerationKey.GetOriginalComponent();
+		const FIntVector& GridCoords = GenerationKey.GetGridCoords();
 
 		// The unbounded grid level will always lie inside the original component, so we can skip it.
 		if (Grid == EPCGHiGenGrid::Unbounded)
 		{
+			if (OriginalComponent && !OriginalComponent->bActivated)
+			{
+				CleanupComponent(GenerationKey, OriginalComponent);
+			}
+
 			continue;
 		}
 
-		const UPCGComponent* OriginalComponent = GenerationKey.GetOriginalComponent();
-		const FIntVector& GridCoords = GenerationKey.GetGridCoords();
-
 		UPCGComponent* LocalComponent = OriginalComponent ? ActorAndComponentMapping->GetLocalComponent(GridSize, GridCoords, OriginalComponent, /*bRuntimeGenerated=*/true) : nullptr;
 		APCGPartitionActor* PartitionActor = LocalComponent ? Cast<APCGPartitionActor>(LocalComponent->GetOwner()) : nullptr;
 
@@ -921,7 +931,7 @@ void FPCGRuntimeGenScheduler::CleanupDelayedRefreshComponents()
 			const FBox OriginalBounds = OriginalComponent->GetGridBounds();
 			const FBox LocalBounds = PartitionActor->GetFixedBounds();
 
-			if (!OriginalBounds.Intersect(LocalBounds))
+			if (!OriginalBounds.Intersect(LocalBounds) || !OriginalComponent->bActivated)
 			{
 				CleanupComponent(GenerationKey, LocalComponent);
 			}
@@ -1001,7 +1011,7 @@ void FPCGRuntimeGenScheduler::RefreshComponent(UPCGComponent* InComponent, bool
 						UE_LOG(LogPCG, Warning, TEXT("[RUNTIMEGEN] SHALLOW REFRESH LOCAL COMPONENT: '%s'"), *LocalComponent->GetOwner()->GetActorNameOrLabel());
 					}
 
-					if (ensure(LocalComponent))
+					if (LocalComponent)
 					{
 						LocalComponent->CleanupLocalImmediate(/*bRemoveComponents=*/false);
 					}
