[PCG] Add new node Get Graph Parameter for extraction parity with GetActorProperty

#rnx

diff --git a/Engine/Plugins/PCG/Source/PCG/Private/Elements/PCGUserParameterGet.cpp b/Engine/Plugins/PCG/Source/PCG/Private/Elements/PCGUserParameterGet.cpp
--- a/Engine/Plugins/PCG/Source/PCG/Private/Elements/PCGUserParameterGet.cpp
+++ b/Engine/Plugins/PCG/Source/PCG/Private/Elements/PCGUserParameterGet.cpp
@@ -139,29 +139,56 @@ FPCGElementPtr UPCGUserParameterGetSettings::CreateElement() const
 	return MakeShared<FPCGUserParameterGetElement>();
 }
 
+TArray<FPCGPinProperties> UPCGGenericUserParameterGetSettings::OutputPinProperties() const
+{
+	TArray<FPCGPinProperties> PinProperties;
+	PinProperties.Emplace(PCGPinConstants::DefaultOutputLabel, EPCGDataType::Param);
+
+	return PinProperties;
+}
+
+FString UPCGGenericUserParameterGetSettings::GetAdditionalTitleInformation() const
+{
+	return PropertyPath;
+}
+
+FPCGElementPtr UPCGGenericUserParameterGetSettings::CreateElement() const
+{
+	return MakeShared<FPCGUserParameterGetElement>();
+}
+
 bool FPCGUserParameterGetElement::ExecuteInternal(FPCGContext* Context) const
 {
 	TRACE_CPUPROFILER_EVENT_SCOPE(FPCGUserParameterGetElement::Execute);
 
 	check(Context);
 
-	const UPCGUserParameterGetSettings* Settings = Context->GetInputSettings<UPCGUserParameterGetSettings>();
-	check(Settings);
-
-	const FName PropertyName = Settings->PropertyName;
-
 	FConstStructView Parameters = PCGUserParameterGetSettings::GetFirstValidLayout(*Context);
 	const UScriptStruct* PropertyBag = Parameters.GetScriptStruct();
 
-	const FProperty* Property = PropertyBag ? PropertyBag->FindPropertyByName(PropertyName) : nullptr;
+	bool bForceObjectAndStructExtraction = false;
+	FPCGAttributePropertySelector InputSelector;
+	FName OutputAttributeName;
 
-	if (!Property)
+	if (const UPCGUserParameterGetSettings* Settings = Context->GetInputSettings<UPCGUserParameterGetSettings>())
+	{
+		bForceObjectAndStructExtraction = Settings->bForceObjectAndStructExtraction;
+		InputSelector = FPCGAttributePropertySelector::CreateAttributeSelector(Settings->PropertyName);
+		OutputAttributeName = Settings->PropertyName;
+	}
+	else if (const UPCGGenericUserParameterGetSettings* GenericSettings = Context->GetInputSettings<UPCGGenericUserParameterGetSettings>())
 	{
-		PCGE_LOG(Error, GraphAndLog, FText::Format(LOCTEXT("InvalidProperty", "Could not find the property '{0}' in the user parameters"), FText::FromName(PropertyName)));
+		bForceObjectAndStructExtraction = GenericSettings->bForceObjectAndStructExtraction;
+		InputSelector = FPCGAttributePropertySelector::CreateSelectorFromString(GenericSettings->PropertyPath);
+		OutputAttributeName = GenericSettings->OutputAttributeName;
+	}
+	else
+	{
+		checkNoEntry();
 		return true;
 	}
 
-	PCGPropertyHelpers::FExtractorParameters ExtractorParameters{ Parameters.GetMemory(), PropertyBag, FPCGAttributePropertySelector::CreateAttributeSelector(PropertyName), PropertyName, Settings->bForceObjectAndStructExtraction, /*bPropertyNeedsToBeVisible=*/false};
+	PCGPropertyHelpers::FExtractorParameters ExtractorParameters{ Parameters.GetMemory(), PropertyBag, InputSelector, OutputAttributeName, bForceObjectAndStructExtraction, /*bPropertyNeedsToBeVisible=*/false};
 
 	// Don't care for object traversed in non-editor build, since it is only useful for tracking.
 	TSet<FSoftObjectPath>* ObjectTraversedPtr = nullptr;
@@ -174,6 +201,11 @@ bool FPCGUserParameterGetElement::ExecuteInternal(FPCGContext* Context) const
 	{
 		Context->OutputData.TaggedData.Emplace_GetRef().Data = ParamData;
 	}
+	else
+	{
+		PCGE_LOG(Error, GraphAndLog, FText::Format(LOCTEXT("InvalidProperty", "Could not find the property '{0}' in the user parameters"), InputSelector.GetDisplayText()));
+		return true;
+	}
 
 	// Register dynamic tracking
 #if WITH_EDITOR
diff --git a/Engine/Plugins/PCG/Source/PCG/Public/Elements/PCGUserParameterGet.h b/Engine/Plugins/PCG/Source/PCG/Public/Elements/PCGUserParameterGet.h
--- a/Engine/Plugins/PCG/Source/PCG/Public/Elements/PCGUserParameterGet.h
+++ b/Engine/Plugins/PCG/Source/PCG/Public/Elements/PCGUserParameterGet.h
@@ -52,6 +52,44 @@ protected:
 	//~End UPCGSettings interface
 };
 
+/**
+* Generic getter for user parameter defined in the PCG Graph, by the user.
+* Will pick up the value from the graph instance.
+* This getter allows to set manually the user parameter they want to get, and add extractor, the same way than GetActorProperty or GetPropertyFromObjectPath
+*/
+UCLASS(MinimalAPI, BlueprintType, ClassGroup = (Procedural))
+class UPCGGenericUserParameterGetSettings : public UPCGSettings
+{
+	GENERATED_BODY()
+
+public:
+	UPROPERTY(BlueprintReadWrite, EditAnywhere, Category = Settings, meta = (PCG_Overridable))
+	FString PropertyPath;
+
+	/** If the property is a struct/object supported by metadata, this option can be toggled to force extracting all (compatible) properties contained in this property. Automatically true if unsupported by metadata. For now, only supports direct child properties (and not deeper). */
+	UPROPERTY(BlueprintReadWrite, EditAnywhere, Category = Settings, meta = (PCG_Overridable))
+	bool bForceObjectAndStructExtraction = false;
+
+	UPROPERTY(BlueprintReadWrite, EditAnywhere, Category = Settings, meta = (PCG_Overridable))
+	FName OutputAttributeName = NAME_None;
+
+	//~Begin UPCGSettings interface
+	virtual bool RequiresDataFromPreTask() const override { return true; }
+	virtual FString GetAdditionalTitleInformation() const override;
+
+protected:
+	virtual TArray<FPCGPinProperties> InputPinProperties() const override { return {}; }
+	virtual TArray<FPCGPinProperties> OutputPinProperties() const override;
+	virtual FPCGElementPtr CreateElement() const override;
+
+#if WITH_EDITOR
+	virtual FName GetDefaultNodeName() const override { return FName("GetGenericGraphParameter"); }
+	virtual FText GetDefaultNodeTitle() const override { return NSLOCTEXT("UPCGGenericUserParameterGetSettings", "NodeTitle", "Get Graph Parameter"); }
+	virtual EPCGSettingsType GetType() const override { return EPCGSettingsType::GraphParameters; }
+#endif
+	//~End UPCGSettings interface
+};
+
 class FPCGUserParameterGetElement : public IPCGElement
 {
 public:
