[PCG] Refresh debug object tree when components are unregistered
* Fixes an issue where stale entries are left in the debug object tree when deleting PCG components
* Add events for PCG Component registration/unregistration

#rnx

diff --git a/Engine/Plugins/PCG/Source/PCG/Private/PCGComponent.cpp b/Engine/Plugins/PCG/Source/PCG/Private/PCGComponent.cpp
--- a/Engine/Plugins/PCG/Source/PCG/Private/PCGComponent.cpp
+++ b/Engine/Plugins/PCG/Source/PCG/Private/PCGComponent.cpp
@@ -1347,6 +1347,8 @@ void UPCGComponent::OnUnregister()
 		{
 			Subsystem->CancelGeneration(this);
 		}
+
+		Subsystem->OnComponentUnregistered.Broadcast();
 	}
 #endif // WITH_EDITOR
 
diff --git a/Engine/Plugins/PCG/Source/PCG/Public/PCGSubsystem.h b/Engine/Plugins/PCG/Source/PCG/Public/PCGSubsystem.h
--- a/Engine/Plugins/PCG/Source/PCG/Public/PCGSubsystem.h
+++ b/Engine/Plugins/PCG/Source/PCG/Public/PCGSubsystem.h
@@ -34,6 +34,7 @@ typedef TSharedPtr<IPCGElement, ESPMode::ThreadSafe> FPCGElementPtr;
 class UWorld;
 
 #if WITH_EDITOR
+DECLARE_MULTICAST_DELEGATE(FPCGOnComponentUnregistered);
 DECLARE_MULTICAST_DELEGATE_OneParam(FPCGOnComponentGenerationCompleteOrCancelled, UPCGSubsystem*);
 #endif // WITH_EDITOR
 
@@ -297,6 +298,7 @@ public:
 	void ClearExecutedStacks(FPCGStack BeginningWithStack);
 	void ClearExecutedStacks(const UPCGGraph* InContainingGraph);
 
+	FPCGOnComponentUnregistered OnComponentUnregistered;
 	FPCGOnComponentGenerationCompleteOrCancelled OnComponentGenerationCompleteOrCancelled;
 
 private:
diff --git a/Engine/Plugins/PCG/Source/PCGEditor/Private/PCGEditor.cpp b/Engine/Plugins/PCG/Source/PCGEditor/Private/PCGEditor.cpp
--- a/Engine/Plugins/PCG/Source/PCGEditor/Private/PCGEditor.cpp
+++ b/Engine/Plugins/PCG/Source/PCGEditor/Private/PCGEditor.cpp
@@ -2988,6 +2988,12 @@ void FPCGEditor::JumpToDefinition(const UClass* Class) const
 	}
 }
 
+void FPCGEditor::OnComponentUnregistered()
+{
+	// Refresh the debug object tree to avoid stale entries from components that have been unregistered.
+	DebugObjectTreeWidget->RequestRefresh();
+}
+
 void FPCGEditor::OnComponentGenerationCompleteOrCancelled(UPCGSubsystem* Subsystem)
 {
 	DebugObjectTreeWidget->RequestRefresh();
@@ -3026,6 +3032,7 @@ void FPCGEditor::RegisterDelegatesForWorld(UWorld* World)
 
 	if (UPCGSubsystem* Subsystem = UPCGSubsystem::GetInstance(World))
 	{
+		Subsystem->OnComponentUnregistered.AddRaw(this, &FPCGEditor::OnComponentUnregistered);
 		Subsystem->OnComponentGenerationCompleteOrCancelled.AddRaw(this, &FPCGEditor::OnComponentGenerationCompleteOrCancelled);
 	}
 }
@@ -3034,6 +3041,7 @@ void FPCGEditor::UnregisterDelegatesForWorld(UWorld* World)
 {
 	if (UPCGSubsystem* Subsystem = UPCGSubsystem::GetInstance(World))
 	{
+		Subsystem->OnComponentUnregistered.RemoveAll(this);
 		Subsystem->OnComponentGenerationCompleteOrCancelled.RemoveAll(this);
 	}
 }
diff --git a/Engine/Plugins/PCG/Source/PCGEditor/Private/PCGEditor.h b/Engine/Plugins/PCG/Source/PCGEditor/Private/PCGEditor.h
--- a/Engine/Plugins/PCG/Source/PCGEditor/Private/PCGEditor.h
+++ b/Engine/Plugins/PCG/Source/PCGEditor/Private/PCGEditor.h
@@ -325,6 +325,9 @@ private:
 	/** To be called everytime we need to replicate our extra nodes to the underlying PCGGraph */
 	void ReplicateExtraNodes() const;
 
+	/** Called when a PCG component unregisters. */
+	void OnComponentUnregistered();
+
 	/** Called when a component finishes executing. Useful for updating debugging tools/UIs. */
 	void OnComponentGenerationCompleteOrCancelled(UPCGSubsystem* Subsystem);
 
