[PCG] Fixed tag parsing error and error in the Source attrbute too.

diff --git a/Engine/Plugins/PCG/Source/PCGEditor/Private/PCGLevelToAsset.cpp b/Engine/Plugins/PCG/Source/PCGEditor/Private/PCGLevelToAsset.cpp
--- a/Engine/Plugins/PCG/Source/PCGEditor/Private/PCGLevelToAsset.cpp
+++ b/Engine/Plugins/PCG/Source/PCGEditor/Private/PCGLevelToAsset.cpp
@@ -143,7 +143,7 @@ bool UPCGLevelToAsset::BP_ExportWorld_Implementation(UWorld* World, const FStrin
 	UPCGMetadata* RootMetadata = RootPointData->MutableMetadata();
 	TArray<FPCGPoint>& Roots = RootPointData->GetMutablePoints();
 	RootMetadata->CreateAttribute<FString>(TEXT("Name"), World->GetName(), /*bAllowsInterpolation=*/false, /*bOverrideParent=*/true);
-	RootMetadata->CreateAttribute<FSoftObjectPath>(TEXT("Source"), PackageName, /*bAllowsInterpolation=*/false, /*bOverrideParent=*/true);
+	RootMetadata->CreateAttribute<FSoftObjectPath>(TEXT("Source"), FSoftObjectPath(World), /*bAllowsInterpolation=*/false, /*bOverrideParent=*/true);
 
 	// Add to data collection
 	{
@@ -343,7 +343,7 @@ bool UPCGLevelToAsset::BP_ExportWorld_Implementation(UWorld* World, const FStrin
 				if (TagString.FindChar(':', DividerPosition))
 				{
 					const FString LeftSide = TagString.Left(DividerPosition);
-					const FString RightSide = TagString.Right(DividerPosition + 1);
+					const FString RightSide = TagString.RightChop(DividerPosition + 1);
 
 					if (LeftSide.IsEmpty() || RightSide.IsEmpty())
 					{
