[PCG] Add start/end/random offset to the spline sampler

diff --git a/Engine/Plugins/PCG/Source/PCG/Private/Data/PCGLandscapeSplineData.cpp b/Engine/Plugins/PCG/Source/PCG/Private/Data/PCGLandscapeSplineData.cpp
--- a/Engine/Plugins/PCG/Source/PCG/Private/Data/PCGLandscapeSplineData.cpp
+++ b/Engine/Plugins/PCG/Source/PCG/Private/Data/PCGLandscapeSplineData.cpp
@@ -240,7 +240,7 @@ const UPCGPointData* UPCGLandscapeSplineData::CreatePointData(FPCGContext* Conte
 	SamplerParams.Mode = EPCGSplineSamplingMode::Distance;
 	SamplerParams.Dimension = EPCGSplineSamplingDimension::OnHorizontal;
 
-	PCGSplineSamplerHelpers::SampleLineData(/*LineData=*/this, /*InBoundingShapeData=*/nullptr, /*InProjectionTarget=*/nullptr, /*InProjectionParams=*/{}, SamplerParams, Data);
+	PCGSplineSamplerHelpers::SampleLineData(Context, /*LineData=*/this, /*InBoundingShapeData=*/nullptr, /*InProjectionTarget=*/nullptr, /*InProjectionParams=*/{}, SamplerParams, Data);
 
 	UE_LOG(LogPCG, Verbose, TEXT("Landscape spline %s generated %d points"), *Spline->GetFName().ToString(), Points.Num());
 
diff --git a/Engine/Plugins/PCG/Source/PCG/Private/Data/PCGSplineData.cpp b/Engine/Plugins/PCG/Source/PCG/Private/Data/PCGSplineData.cpp
--- a/Engine/Plugins/PCG/Source/PCG/Private/Data/PCGSplineData.cpp
+++ b/Engine/Plugins/PCG/Source/PCG/Private/Data/PCGSplineData.cpp
@@ -161,7 +161,7 @@ const UPCGPointData* UPCGSplineData::CreatePointData(FPCGContext* Context) const
 	FPCGSplineSamplerParams SamplerParams;
 	SamplerParams.Mode = EPCGSplineSamplingMode::Distance;
 
-	PCGSplineSamplerHelpers::SampleLineData(this, /*InBoundingShape=*/nullptr, /*InProjectionTarget=*/nullptr, /*InProjectionParams=*/{}, SamplerParams, Data);
+	PCGSplineSamplerHelpers::SampleLineData(Context, this, /*InBoundingShape=*/nullptr, /*InProjectionTarget=*/nullptr, /*InProjectionParams=*/{}, SamplerParams, Data);
 	UE_LOG(LogPCG, Verbose, TEXT("Spline generated %d points"), Data->GetPoints().Num());
 
 	return Data;
diff --git a/Engine/Plugins/PCG/Source/PCG/Private/Data/PCGSplineInteriorSurfaceData.cpp b/Engine/Plugins/PCG/Source/PCG/Private/Data/PCGSplineInteriorSurfaceData.cpp
--- a/Engine/Plugins/PCG/Source/PCG/Private/Data/PCGSplineInteriorSurfaceData.cpp
+++ b/Engine/Plugins/PCG/Source/PCG/Private/Data/PCGSplineInteriorSurfaceData.cpp
@@ -158,7 +158,7 @@ void UPCGSplineInteriorSurfaceData::CacheData()
 
 	// TODO: Ideally SampleLineData could just consume a UPCGSplineInteriorSurfaceData or FPCGSplineStruct, to avoid the extra copy on the spline struct.
 	// TODO: It might be preferable to directly sample the spline than use the spline sampler, since it does a lot of work we don't care about.
-	PCGSplineSamplerHelpers::SampleLineData(SplineData, /*InBoundingShape=*/nullptr, /*InProjectionTarget=*/nullptr, /*InProjectionParams=*/{}, SamplerParams, PointData);
+	PCGSplineSamplerHelpers::SampleLineData(/*Context=*/nullptr, SplineData, /*InBoundingShape=*/nullptr, /*InProjectionTarget=*/nullptr, /*InProjectionParams=*/{}, SamplerParams, PointData);
 
 	// Cache the points which describe the polygon of our spline.
 	for (const FPCGPoint& Point : PointData->GetPoints())
diff --git a/Engine/Plugins/PCG/Source/PCG/Private/Elements/PCGSplineSampler.cpp b/Engine/Plugins/PCG/Source/PCG/Private/Elements/PCGSplineSampler.cpp
--- a/Engine/Plugins/PCG/Source/PCG/Private/Elements/PCGSplineSampler.cpp
+++ b/Engine/Plugins/PCG/Source/PCG/Private/Elements/PCGSplineSampler.cpp
@@ -395,7 +395,7 @@ namespace PCGSplineSamplerHelpers
 
 		const UPCGPolyLineData* LineData = nullptr;
 		int CurrentSegmentIndex = 0;
-		FVector::FReal DistanceToCurrentSegment = 0.0f;
+		FVector::FReal DistanceToCurrentSegment = 0.0;
 		bool bComputeCurvature = false;
 		bool bComputeAlpha = false;
 		bool bComputeDistance = false;
@@ -514,21 +514,61 @@ namespace PCGSplineSamplerHelpers
 
 	struct FDistanceStepSampler : public FStepSampler
 	{
-		FDistanceStepSampler(const UPCGPolyLineData* InLineData, const FPCGSplineSamplerParams& Params)
+		FDistanceStepSampler(const UPCGPolyLineData* InLineData, const FPCGSplineSamplerParams& Params, int Seed)
 			: FStepSampler(InLineData, Params)
 		{
-			DistanceIncrement = Params.DistanceIncrement;
-			CurrentDistance = 0;
+			StartOffset = FMath::Max(0, Params.StartOffset);
+			const FVector::FReal EndOffset = FMath::Max(0, Params.EndOffset);
+
+			CurrentDistance = StartOffset;
+			TotalDistance = InLineData->GetLength();
+			EndDistance = TotalDistance - EndOffset;
+
+			if (Params.Mode == EPCGSplineSamplingMode::NumberOfSamples)
+			{
+				TotalNumSamples = Params.NumSamples;
+
+				if (TotalNumSamples > 0)
+				{
+					// Compute an increment which evenly distributes sample points along the length of the curve.
+					DistanceIncrement = (LineData->GetLength() - StartOffset - EndOffset) / (LineData->IsClosed() ? TotalNumSamples : FMath::Max(1, TotalNumSamples - 1));
+				}
+			}
+			else
+			{
+				DistanceIncrement = Params.DistanceIncrement;
+			}
+
+			MaxRandomOffset = FMath::Max(0.0f, Params.MaxRandomOffsetNormalized) * DistanceIncrement / 2.0f;
+			bUseRandomOffset = !FMath::IsNearlyZero(MaxRandomOffset);
+
+			if (bUseRandomOffset)
+			{
+				RandomSource.Initialize(Seed);
+			}
 		}
 
 		virtual void Step(FSamplerResult& OutResult) override
 		{
+			FVector::FReal OffsetDistance = CurrentDistance + (bUseRandomOffset ? RandomSource.FRandRange(-MaxRandomOffset, MaxRandomOffset) : 0.0);
+
+			// Prevent samples from wrapping around on open splines.
+			if (!LineData->IsClosed() && DistanceToCurrentSegment + OffsetDistance < StartOffset)
+			{
+				OffsetDistance = StartOffset - DistanceToCurrentSegment;
+			}
+
+			if (!LineData->IsClosed() && DistanceToCurrentSegment + OffsetDistance >= EndDistance)
+			{
+				OffsetDistance = EndDistance - DistanceToCurrentSegment;
+			}
+
 			FVector::FReal CurrentSegmentLength = LineData->GetSegmentLength(CurrentSegmentIndex);
 			FTransform& OutTransform = OutResult.LocalTransform;
 			FBox& OutBox = OutResult.Box;
-			OutTransform = LineData->GetTransformAtDistance(CurrentSegmentIndex, CurrentDistance, /*bWorldSpace=*/false, &OutBox);
+			OutTransform = LineData->GetTransformAtDistance(CurrentSegmentIndex, OffsetDistance, /*bWorldSpace=*/false, &OutBox);
 			OutResult.SegmentIndex = CurrentSegmentIndex;
-			OutResult.InputKey = LineData->GetInputKeyAtDistance(CurrentSegmentIndex, CurrentDistance);
+			OutResult.InputKey = LineData->GetInputKeyAtDistance(CurrentSegmentIndex, OffsetDistance);
 
 			// Set min/max to half of extent
 			OutBox.Min.X *= 0.5 * DistanceIncrement / OutTransform.GetScale3D().X;
@@ -536,48 +576,58 @@ namespace PCGSplineSamplerHelpers
 
 			if (bComputeCurvature)
 			{
-				OutResult.Curvature = LineData->GetCurvatureAtDistance(CurrentSegmentIndex, CurrentDistance);
+				OutResult.Curvature = LineData->GetCurvatureAtDistance(CurrentSegmentIndex, OffsetDistance);
 			}
 
 			if (bComputeAlpha)
 			{
-				OutResult.Alpha = LineData->GetAlphaAtDistance(CurrentSegmentIndex, CurrentDistance);
+				OutResult.Alpha = LineData->GetAlphaAtDistance(CurrentSegmentIndex, OffsetDistance);
 			}
 
 			if (bComputeDistance)
 			{
-				OutResult.Distance = DistanceToCurrentSegment + CurrentDistance;
+				OutResult.Distance = DistanceToCurrentSegment + OffsetDistance;
 			}
 
+			// Increment the current distance to get our next sample location. Note that we don't use the offset distance, since the new sample doesn't care
+			// about the previous sample location.
 			CurrentDistance += DistanceIncrement;
+			++CurrentNumSamples;
+
 			while(CurrentDistance > CurrentSegmentLength)
 			{
 				CurrentDistance -= CurrentSegmentLength;
 				++CurrentSegmentIndex;
 
-				if (bComputeDistance)
-				{
-					DistanceToCurrentSegment += CurrentSegmentLength;
-				}
+				DistanceToCurrentSegment += CurrentSegmentLength;
 
-				if (!IsDone())
+				if (IsDone() || CurrentSegmentLength <= 0)
 				{
-					CurrentSegmentLength = LineData->GetSegmentLength(CurrentSegmentIndex);
+					break;
 				}
 				else
 				{
-					break;
+					CurrentSegmentLength = LineData->GetSegmentLength(CurrentSegmentIndex);
 				}
 			}
 		}
 
 		virtual bool IsDone() const override
 		{
-			return CurrentSegmentIndex >= LineData->GetNumSegments();
+			return (DistanceToCurrentSegment + CurrentDistance > EndDistance + UE_DOUBLE_SMALL_NUMBER) || (CurrentNumSamples == TotalNumSamples);
 		}
 
-		FVector::FReal CurrentDistance = 0.0f;
-		FVector::FReal DistanceIncrement = 0.0f;
+		FVector::FReal CurrentDistance = 0.0;
+		FVector::FReal DistanceIncrement = 0.0;
+		FVector::FReal StartOffset = 0.0;
+		FVector::FReal EndDistance = 0.0;
+		FVector::FReal TotalDistance = 0.0;
+		FVector::FReal MaxRandomOffset = 0.0;
+		FRandomStream RandomSource;
+		bool bUseRandomOffset = false;
+
+		int CurrentNumSamples = 0;
+		int TotalNumSamples = -1;
 	};
 
 	struct FDimensionSampler
@@ -867,7 +917,7 @@ namespace PCGSplineSamplerHelpers
 		int NumHeightSteps;
 	};
 
-	void SampleLineData(const UPCGPolyLineData* LineData, const UPCGSpatialData* InBoundingShapeData, const UPCGSpatialData* InProjectionTarget, const FPCGProjectionParams& InProjectionParams, const FPCGSplineSamplerParams& Params, UPCGPointData* OutPointData)
+	void SampleLineData(FPCGContext* Context, const UPCGPolyLineData* LineData, const UPCGSpatialData* InBoundingShapeData, const UPCGSpatialData* InProjectionTarget, const FPCGProjectionParams& InProjectionParams, const FPCGSplineSamplerParams& Params, UPCGPointData* OutPointData)
 	{
 		check(LineData && OutPointData);
 
@@ -884,32 +934,10 @@ namespace PCGSplineSamplerHelpers
 		}
 
 		FSubdivisionStepSampler SubdivisionSampler(LineData, Params);
-		FDistanceStepSampler DistanceSampler(LineData, Params);
+		FDistanceStepSampler DistanceSampler(LineData, Params, Context ? Context->GetSeed() : 42);
 
 		FStepSampler* Sampler = ((Params.Mode == EPCGSplineSamplingMode::Subdivision) ? static_cast<FStepSampler*>(&SubdivisionSampler) : static_cast<FStepSampler*>(&DistanceSampler));
 
-		if (Params.Mode == EPCGSplineSamplingMode::NumberOfSamples)
-		{
-			if (Params.NumSamples <= 0)
-			{
-				return;
-			}
-
-			// Compute an increment which evenly distributes sample points along the length of the curve.
-			DistanceSampler.DistanceIncrement = LineData->GetLength() / (LineData->IsClosed() ? Params.NumSamples : FMath::Max(1, Params.NumSamples - 1));
-
-			if (LineData->IsClosed() || Params.NumSamples == 1)
-			{
-				// If the curve is closed or only has one sample, we nudge the DistanceIncrement slightly to avoid floating point error giving us an extra sample point.
-				DistanceSampler.DistanceIncrement += UE_DOUBLE_SMALL_NUMBER;
-			}
-			else
-			{
-				// If the curve is not closed and has more than one sample, we should nudge DistanceIncrement slightly lower so that we guarantee capturing the last sample point.
-				DistanceSampler.DistanceIncrement -= UE_DOUBLE_SMALL_NUMBER;
-			}
-		}
-
 		FDimensionSampler TrivialDimensionSampler(LineData, InBoundingShapeData, InProjectionTarget, InProjectionParams, Params, OutPointData);
 		FVolumeSampler VolumeSampler(LineData, InBoundingShapeData, InProjectionTarget, InProjectionParams, Params, OutPointData);
 
@@ -1373,6 +1401,11 @@ namespace PCGSplineSamplerHelpers
 	}
 }
 
+UPCGSplineSamplerSettings::UPCGSplineSamplerSettings()
+{
+	bUseSeed = true;
+}
+
 #if WITH_EDITOR
 FText UPCGSplineSamplerSettings::GetNodeTooltipText() const
 {
@@ -1470,7 +1503,7 @@ bool FPCGSplineSamplerElement::ExecuteInternal(FPCGContext* Context) const
 		}
 		else
 		{
-			PCGSplineSamplerHelpers::SampleLineData(LineData, BoundingShape, ProjectionTarget, ProjectionParams, SamplerParams, SampledPointData);
+			PCGSplineSamplerHelpers::SampleLineData(Context, LineData, BoundingShape, ProjectionTarget, ProjectionParams, SamplerParams, SampledPointData);
 		}
 	}
 
diff --git a/Engine/Plugins/PCG/Source/PCG/Public/Elements/PCGSplineSampler.h b/Engine/Plugins/PCG/Source/PCG/Public/Elements/PCGSplineSampler.h
--- a/Engine/Plugins/PCG/Source/PCG/Public/Elements/PCGSplineSampler.h
+++ b/Engine/Plugins/PCG/Source/PCG/Public/Elements/PCGSplineSampler.h
@@ -81,6 +81,18 @@ struct FPCGSplineSamplerParams
 	UPROPERTY(BlueprintReadWrite, EditAnywhere, Category = Settings, meta = (ClampMin = "0", EditCondition = "Dimension==EPCGSplineSamplingDimension::OnVertical||Dimension==EPCGSplineSamplingDimension::OnVolume"))
 	int32 NumHeightSubdivisions = 8;
 
+	/** Distance (in cm) along the spline at which sampling will begin. */
+	UPROPERTY(BlueprintReadWrite, EditAnywhere, Category = Settings, meta = (ClampMin = "0", EditCondition = "Mode!=EPCGSplineSamplingMode::Subdivision&&Dimension!=EPCGSplineSamplingDimension::OnInterior"))
+	float StartOffset = 0.0f;
+
+	/** Distance (in cm) from the end of the spline at which sampling will end. */
+	UPROPERTY(BlueprintReadWrite, EditAnywhere, Category = Settings, meta = (ClampMin = "0", EditCondition = "Mode!=EPCGSplineSamplingMode::Subdivision&&Dimension!=EPCGSplineSamplingDimension::OnInterior"))
+	float EndOffset = 0.0f;
+
+	/** Normalized value for the maximum possible offset for each sample point. 0.0 means no offset, and 1.0 means DistanceIncrement / 2.0. */
+	UPROPERTY(BlueprintReadWrite, EditAnywhere, Category = Settings, meta = (ClampMin = "0", EditCondition = "Mode!=EPCGSplineSamplingMode::Subdivision&&Dimension!=EPCGSplineSamplingDimension::OnInterior"))
+	float MaxRandomOffsetNormalized = 0.0f;
+
 	UPROPERTY(BlueprintReadWrite, EditAnywhere, Category = Settings, meta = (InlineEditConditionToggle))
 	bool bComputeDirectionDelta = false;
 
@@ -196,7 +208,7 @@ struct FPCGSplineSamplerParams
 namespace PCGSplineSamplerHelpers
 {
 	/** Samples on spline or within volume around it. */
-	void SampleLineData(const UPCGPolyLineData* LineData, const UPCGSpatialData* InBoundingShape, const UPCGSpatialData* InProjectionTarget, const FPCGProjectionParams& InProjectionParams, const FPCGSplineSamplerParams& Params, UPCGPointData* OutPointData);
+	void SampleLineData(FPCGContext* Context, const UPCGPolyLineData* LineData, const UPCGSpatialData* InBoundingShape, const UPCGSpatialData* InProjectionTarget, const FPCGProjectionParams& InProjectionParams, const FPCGSplineSamplerParams& Params, UPCGPointData* OutPointData);
 
 	/** Samples 2D region bounded by spline. */
 	void SampleInteriorData(FPCGContext* Context, const UPCGPolyLineData* LineData, const UPCGSpatialData* InBoundingShape, const UPCGSpatialData* InProjectionTarget, const FPCGProjectionParams& InProjectionParams, const FPCGSplineSamplerParams& Params, UPCGPointData* OutPointData);
@@ -216,6 +228,8 @@ class UPCGSplineSamplerSettings : public UPCGSettings
 	GENERATED_BODY()
 
 public:
+	UPCGSplineSamplerSettings();
+
 	// ~Begin UPCGSettings interface
 #if WITH_EDITOR
 	virtual FName GetDefaultNodeName() const override { return FName(TEXT("SplineSampler")); }
