[PCG] Replace Density Remap by [REDACTED_PERSON]

* Node is using Metadata Operation base
* Deprecate old Density Remap
* Add preconfigured info for Density Remap, using Attribute Remap
* Change the Unit Test to use Attribute Remap now

diff --git a/Engine/Plugins/PCG/Source/PCG/Private/Elements/PCGAttributeRemap.cpp b/Engine/Plugins/PCG/Source/PCG/Private/Elements/PCGAttributeRemap.cpp
new file mode 100644
--- /dev/null
+++ b/Engine/Plugins/PCG/Source/PCG/Private/Elements/PCGAttributeRemap.cpp
@@ -0,0 +1,170 @@
+// Copyright Epic Games, Inc. All Rights Reserved.
+
+#include "Elements/PCGAttributeRemap.h"
+
+#include "PCGContext.h"
+#include "Helpers/PCGAsync.h"
+#include "Metadata/Accessors/PCGAttributeAccessorHelpers.h"
+
+#define LOCTEXT_NAMESPACE "PCGAttributeRemapElement"
+
+namespace PCGAttributeRemapElement
+{
+	struct FPCGAttributeRemapParams
+	{
+		double InRangeMin;
+		double InRangeMax;
+		double Slope;
+		double Intercept;
+		bool bClampToUnitRange;
+		bool bIgnoreValuesOutsideInputRange;
+	};
+
+	// Fallback for all unsupported types, does nothing.
+	template <typename T>
+	void Remap(T& InOutValue, const FPCGAttributeRemapParams& InParams) {}
+
+	// Int/Float version, utilizing the concepts.
+	template <typename T> requires (std::is_arithmetic_v<T> && !std::is_same_v<T, bool>)
+	void Remap(T& InOutValue, const FPCGAttributeRemapParams& InParams)
+	{
+		if (!InParams.bIgnoreValuesOutsideInputRange || (InOutValue >= InParams.InRangeMin && InOutValue <= InParams.InRangeMax))
+		{
+			InOutValue = InParams.Slope * (InOutValue - InParams.InRangeMin) + InParams.Intercept;
+			if (InParams.bClampToUnitRange)
+			{
+				InOutValue = FMath::Clamp<T>(InOutValue, T(0.0), T(1.0));
+			}
+		}
+	}
+
+	template <>
+	void Remap<FVector2D>(FVector2D& InOutValue, const FPCGAttributeRemapParams& InParams)
+	{
+		Remap(InOutValue.X, InParams);
+		Remap(InOutValue.Y, InParams);
+	}
+
+	template <>
+	void Remap<FVector>(FVector& InOutValue, const FPCGAttributeRemapParams& InParams)
+	{
+		Remap(InOutValue.X, InParams);
+		Remap(InOutValue.Y, InParams);
+		Remap(InOutValue.Z, InParams);
+	}
+
+	template <>
+	void Remap<FVector4>(FVector4& InOutValue, const FPCGAttributeRemapParams& InParams)
+	{
+		Remap(InOutValue.X, InParams);
+		Remap(InOutValue.Y, InParams);
+		Remap(InOutValue.Z, InParams);
+		Remap(InOutValue.W, InParams);
+	}
+
+	template <>
+	void Remap<FRotator>(FRotator& InOutValue, const FPCGAttributeRemapParams& InParams)
+	{
+		Remap(InOutValue.Roll, InParams);
+		Remap(InOutValue.Pitch, InParams);
+		Remap(InOutValue.Yaw, InParams);
+	}
+}
+
+#if WITH_EDITOR
+FName UPCGAttributeRemapSettings::GetDefaultNodeName() const
+{
+	return FName(TEXT("AttributeRemap"));
+}
+
+FText UPCGAttributeRemapSettings::GetDefaultNodeTitle() const
+{
+	return LOCTEXT("NodeTitle", "Attribute Remap");
+}
+
+TArray<FPCGPreConfiguredSettingsInfo> UPCGAttributeRemapSettings::GetPreconfiguredInfo() const
+{
+	return { FPCGPreConfiguredSettingsInfo{0, LOCTEXT("DensityNodeTitle", "Density Remap")} };
+}
+#endif // WITH_EDITOR
+
+FString UPCGAttributeRemapSettings::GetAdditionalTitleInformation() const
+{
+	return FString::Printf(TEXT("%s -> %s"), *InputSource.ToString(), *OutputTarget.ToString());
+}
+
+void UPCGAttributeRemapSettings::ApplyPreconfiguredSettings(const FPCGPreConfiguredSettingsInfo& PreconfigureInfo)
+{
+	switch (PreconfigureInfo.PreconfiguredIndex)
+	{
+	case 0:
+		InputSource.SetPointProperty(EPCGPointProperties::Density);
+		bClampToUnitRange = true;
+		break;
+	default:
+		break;
+	}
+}
+
+FPCGAttributePropertyInputSelector UPCGAttributeRemapSettings::GetInputSource(uint32 Index) const
+{
+	return Index == 0 ? InputSource : FPCGAttributePropertyInputSelector{};
+}
+
+bool UPCGAttributeRemapSettings::IsSupportedInputType(uint16 TypeId, uint32 InputIndex, bool& bHasSpecialRequirement) const
+{
+	bHasSpecialRequirement = false;
+	return PCG::Private::IsOfTypes<int32, int64, float, double, FVector2D, FVector, FVector4, FRotator, FQuat>(TypeId);
+}
+
+FPCGElementPtr UPCGAttributeRemapSettings::CreateElement() const
+{
+	return MakeShared<FPCGAttributeRemapElement>();
+}
+
+bool FPCGAttributeRemapElement::DoOperation(PCGMetadataOps::FOperationData& OperationData) const
+{
+	TRACE_CPUPROFILER_EVENT_SCOPE(FPCGAttributeRemapElement::Execute);
+
+	const UPCGAttributeRemapSettings* Settings = CastChecked<UPCGAttributeRemapSettings>(OperationData.Settings);
+
+	PCGAttributeRemapElement::FPCGAttributeRemapParams Params{};
+	Params.bClampToUnitRange = Settings->bClampToUnitRange;
+	Params.bIgnoreValuesOutsideInputRange = Settings->bIgnoreValuesOutsideInputRange;
+
+	Params.InRangeMin = FMath::Min(Settings->InRangeMin, Settings->InRangeMax);
+	Params.InRangeMax = FMath::Max(Settings->InRangeMin, Settings->InRangeMax);
+
+	const double OutRangeTrueMin = FMath::Min(Settings->OutRangeMin, Settings->OutRangeMax);
+	const double OutRangeTrueMax = FMath::Max(Settings->OutRangeMin, Settings->OutRangeMax);
+
+	const double InRangeDifference = Params.InRangeMax - Params.InRangeMin;
+	const double OutRangeDifference = OutRangeTrueMax - OutRangeTrueMin;
+
+	// When InRange is a point leave the Slope at 0 so that Density = Intercept
+	if (InRangeDifference == 0)
+	{
+		Params.Slope = 0;
+		Params.Intercept = (OutRangeTrueMin + OutRangeTrueMax) / 2.f;
+	}
+	else
+	{
+		Params.Slope = OutRangeDifference / InRangeDifference;
+		Params.Intercept = OutRangeTrueMin;
+	}
+
+	auto RemapFunc = [this, &Params, &OperationData]<typename AttributeType>(AttributeType DummyValue) -> bool
+	{
+		using FinalType = std::conditional_t<std::is_same_v<AttributeType, FQuat>, FRotator, AttributeType>;
+		return DoUnaryOp<FinalType>(OperationData, [&Params](const FinalType& Value) -> FinalType 
+		{ 
+			FinalType Out = Value;
+			PCGAttributeRemapElement::Remap<FinalType>(Out, Params); 
+			return Out; 
+		});
+	};
+
+	return PCGMetadataAttribute::CallbackWithRightType(OperationData.MostComplexInputType, RemapFunc);
+}
+
+#undef LOCTEXT_NAMESPACE
diff --git a/Engine/Plugins/PCG/Source/PCG/Private/Elements/PCGDensityRemapElement.cpp b/Engine/Plugins/PCG/Source/PCG/Private/Elements/PCGDensityRemapElement.cpp
--- a/Engine/Plugins/PCG/Source/PCG/Private/Elements/PCGDensityRemapElement.cpp
+++ b/Engine/Plugins/PCG/Source/PCG/Private/Elements/PCGDensityRemapElement.cpp
@@ -8,6 +8,13 @@
 
 #include UE_INLINE_GENERATED_CPP_BY_NAME(PCGDensityRemapElement)
 
+UPCGDensityRemapSettings::UPCGDensityRemapSettings()
+{
+#if WITH_EDITOR
+	bExposeToLibrary = false;
+#endif // WITH_EDITOR
+}
+
 FPCGElementPtr UPCGDensityRemapSettings::CreateElement() const
 {
 	return MakeShared<FPCGDensityRemapElement>();
@@ -19,6 +26,7 @@ bool FPCGDensityRemapElement::ExecuteInternal(FPCGContext* InContext) const
 
 	ContextType* Context = static_cast<ContextType*>(InContext);
 
+PRAGMA_DISABLE_DEPRECATION_WARNINGS
 	const UPCGDensityRemapSettings* Settings = Context->GetInputSettings<UPCGDensityRemapSettings>();
 	check(Settings);
 
@@ -27,6 +35,7 @@ bool FPCGDensityRemapElement::ExecuteInternal(FPCGContext* InContext) const
 	const float OutRangeMin = Settings->OutRangeMin;
 	const float OutRangeMax = Settings->OutRangeMax;
 	const bool bExcludeValuesOutsideInputRange = Settings->bExcludeValuesOutsideInputRange;
+PRAGMA_ENABLE_DEPRECATION_WARNINGS
 
 	// used to determine if a density value lies between TrueMin and TrueMax
 	const float InRangeTrueMin = FMath::Min(InRangeMin, InRangeMax);
diff --git a/Engine/Plugins/PCG/Source/PCG/Private/Tests/Elements/PCGDensityRemapTest.cpp b/Engine/Plugins/PCG/Source/PCG/Private/Tests/Elements/PCGDensityRemapTest.cpp
--- a/Engine/Plugins/PCG/Source/PCG/Private/Tests/Elements/PCGDensityRemapTest.cpp
+++ b/Engine/Plugins/PCG/Source/PCG/Private/Tests/Elements/PCGDensityRemapTest.cpp
@@ -6,7 +6,7 @@
 #include "PCGComponent.h"
 
 #include "Data/PCGPointData.h"
-#include "Elements/PCGDensityRemapElement.h"
+#include "Elements/PCGAttributeRemap.h"
 #include "PCGContext.h"
 
 IMPLEMENT_CUSTOM_SIMPLE_AUTOMATION_TEST(FPCGDensityRemapTest, FPCGTestBaseClass, "Plugins.PCG.DensityRemap.Basic", PCGTestsCommon::TestFlags)
@@ -14,10 +14,13 @@ IMPLEMENT_CUSTOM_SIMPLE_AUTOMATION_TEST(FPCGDensityRemapTest, FPCGTestBaseClass,
 bool FPCGDensityRemapTest::RunTest(const FString& Parameters)
 {
 	PCGTestsCommon::FTestData TestData;
-	PCGTestsCommon::GenerateSettings<UPCGDensityRemapSettings>(TestData);
-	UPCGDensityRemapSettings* Settings = CastChecked<UPCGDensityRemapSettings>(TestData.Settings);
+	PCGTestsCommon::GenerateSettings<UPCGAttributeRemapSettings>(TestData);
+	UPCGAttributeRemapSettings* Settings = CastChecked<UPCGAttributeRemapSettings>(TestData.Settings);
 	FPCGElementPtr DensityRemapElement = TestData.Settings->GetElement();
 
+	Settings->InputSource.SetPointProperty(EPCGPointProperties::Density);
+	Settings->bClampToUnitRange = true;
+
 	TObjectPtr<UPCGPointData> PointData = PCGTestsCommon::CreateEmptyPointData();
 	TArray<FPCGPoint>& Points = PointData->GetMutablePoints();
 
@@ -31,6 +34,7 @@ bool FPCGDensityRemapTest::RunTest(const FString& Parameters)
 
 	FPCGTaggedData& TaggedData = TestData.InputData.TaggedData.Emplace_GetRef(FPCGTaggedData());
 	TaggedData.Data = PointData;
+	TaggedData.Pin = PCGPinConstants::DefaultInputLabel;
 
 	auto ValidateDensityRemap = [this, &TestData, DensityRemapElement, Settings](TArray<float> CorrectDensities) -> bool
 	{
@@ -42,7 +46,7 @@ bool FPCGDensityRemapTest::RunTest(const FString& Parameters)
 		const TArray<FPCGTaggedData>& Inputs = Context->InputData.GetInputs();
 		const TArray<FPCGTaggedData>& Outputs = Context->OutputData.GetInputs();
 		
-		if (!TestEqual("Valid number of outputs", Inputs.Num(), Outputs.Num()))
+		if (!TestEqual("Valid number of outputs", Outputs.Num(), Inputs.Num()))
 		{
 			return false;
 		}
@@ -79,7 +83,7 @@ bool FPCGDensityRemapTest::RunTest(const FString& Parameters)
 			const TArray<FPCGPoint>& InPoints = InPointData->GetPoints();
 			const TArray<FPCGPoint>& OutPoints = OutPointData->GetPoints();
 
-			if (!TestEqual("Input and output point counts match", InPoints.Num(), OutPoints.Num()))
+			if (!TestEqual("Input and output point counts match", OutPoints.Num(), InPoints.Num()))
 			{ 
 				bTestPassed = false;
 				continue;
@@ -102,7 +106,7 @@ bool FPCGDensityRemapTest::RunTest(const FString& Parameters)
 		Settings->InRangeMax = 1.f;
 		Settings->OutRangeMin = 0.f;
 		Settings->OutRangeMax = 1.f;
-		Settings->bExcludeValuesOutsideInputRange = true;
+		Settings->bIgnoreValuesOutsideInputRange = true;
 		bTestPassed &= TestTrue("Input and Output are identical when InRange and OutRange are identical", ValidateDensityRemap({ 0.f, 0.2f, 0.4f, 0.6f, 0.8f, 1.f }));
 	}
 
@@ -112,14 +116,14 @@ bool FPCGDensityRemapTest::RunTest(const FString& Parameters)
 		Settings->InRangeMax = 0.4f;
 		Settings->OutRangeMin = 0.f;
 		Settings->OutRangeMax = 1.f;
-		Settings->bExcludeValuesOutsideInputRange = true;
+		Settings->bIgnoreValuesOutsideInputRange = true;
 		bTestPassed &= TestTrue("Valid densities for partial InRange", ValidateDensityRemap({ 0.f, 0.5f, 1.f, 0.6f, 0.8f, 1.f }));
 
 		Settings->InRangeMin = 0.4f;
 		Settings->InRangeMax = 0.f;
 		Settings->OutRangeMin = 1.f;
 		Settings->OutRangeMax = 0.f;
-		Settings->bExcludeValuesOutsideInputRange = true;
+		Settings->bIgnoreValuesOutsideInputRange = true;
 		bTestPassed &= TestTrue("Inverting ranges does not effect output", ValidateDensityRemap({ 0.f, 0.5f, 1.f, 0.6f, 0.8f, 1.f }));
 	}
 
@@ -129,14 +133,14 @@ bool FPCGDensityRemapTest::RunTest(const FString& Parameters)
 		Settings->InRangeMax = 1.f;
 		Settings->OutRangeMin = 0.5f;
 		Settings->OutRangeMax = 1.f;
-		Settings->bExcludeValuesOutsideInputRange = true;
+		Settings->bIgnoreValuesOutsideInputRange = true;
 		bTestPassed &= TestTrue("Valid densities for partial OutRange", ValidateDensityRemap({ 0.f, 0.2f, 3.f / 6.f, 4.f / 6.f, 5.f / 6.f, 1.f }));
 
 		Settings->InRangeMin = 1.f;
 		Settings->InRangeMax = 0.4f;
 		Settings->OutRangeMin = 1.f;
 		Settings->OutRangeMax = 0.5f;
-		Settings->bExcludeValuesOutsideInputRange = true;
+		Settings->bIgnoreValuesOutsideInputRange = true;
 		bTestPassed &= TestTrue("Inverting ranges does not effect output", ValidateDensityRemap({ 0.f, 0.2f, 3.f / 6.f, 4.f / 6.f, 5.f / 6.f, 1.f }));
 	}
 
@@ -146,7 +150,7 @@ bool FPCGDensityRemapTest::RunTest(const FString& Parameters)
 		Settings->InRangeMax = 1.f;
 		Settings->OutRangeMin = 0.5f;
 		Settings->OutRangeMax = 1.f;
-		Settings->bExcludeValuesOutsideInputRange = false;
+		Settings->bIgnoreValuesOutsideInputRange = false;
 		bTestPassed &= TestTrue("All values are remapped when Range Exclusion is disabled", ValidateDensityRemap({ 1.f / 6.f, 2.f / 6.f, 3.f / 6.f, 4.f / 6.f, 5.f / 6.f, 1.f }));
 	}
 
@@ -156,21 +160,21 @@ bool FPCGDensityRemapTest::RunTest(const FString& Parameters)
 		Settings->InRangeMax = 0.2f;
 		Settings->OutRangeMin = 0.f;
 		Settings->OutRangeMax = 1.f;
-		Settings->bExcludeValuesOutsideInputRange = true;
+		Settings->bIgnoreValuesOutsideInputRange = true;
 		bTestPassed &= TestTrue("Point input to Range output", ValidateDensityRemap({ 0.f, 0.5f, 0.4f, 0.6f, 0.8f, 1.f }));
 
 		Settings->InRangeMin = 0.2f;
 		Settings->InRangeMax = 0.2f;
 		Settings->OutRangeMin = 0.5f;
 		Settings->OutRangeMax = 0.5f;
-		Settings->bExcludeValuesOutsideInputRange = true;
+		Settings->bIgnoreValuesOutsideInputRange = true;
 		bTestPassed &= TestTrue("Point input to Point output", ValidateDensityRemap({ 0.f, 0.5f, 0.4f, 0.6f, 0.8f, 1.f }));
 
 		Settings->InRangeMin = 0.2f;
 		Settings->InRangeMax = 1.f;
 		Settings->OutRangeMin = 0.5f;
 		Settings->OutRangeMax = 0.5f;
-		Settings->bExcludeValuesOutsideInputRange = true;
+		Settings->bIgnoreValuesOutsideInputRange = true;
 		bTestPassed &= TestTrue("Range input to Point output", ValidateDensityRemap({ 0.f, 0.5f, 0.5f, 0.5f, 0.5f, 0.5f }));
 	}
 
diff --git a/Engine/Plugins/PCG/Source/PCG/Public/Elements/PCGAttributeRemap.h b/Engine/Plugins/PCG/Source/PCG/Public/Elements/PCGAttributeRemap.h
new file mode 100644
--- /dev/null
+++ b/Engine/Plugins/PCG/Source/PCG/Public/Elements/PCGAttributeRemap.h
@@ -0,0 +1,78 @@
+// Copyright Epic Games, Inc. All Rights Reserved.
+
+#pragma once
+
+#include "PCGSettings.h"
+#include "Elements/Metadata/PCGMetadataOpElementBase.h"
+#include "Elements/PCGTimeSlicedElementBase.h"
+#include "Metadata/Accessors/IPCGAttributeAccessor.h"
+#include "Metadata/Accessors/PCGAttributeAccessorKeys.h"
+
+#include "Templates/UniquePtr.h"
+
+#include "PCGAttributeRemap.generated.h"
+
+/**
+* Remap attribute values from one range to another.
+*/
+UCLASS(MinimalAPI, BlueprintType, ClassGroup = (Procedural))
+class UPCGAttributeRemapSettings : public UPCGMetadataSettingsBase
+{
+	GENERATED_BODY()
+
+public:
+	//~Begin UPCGSettings interface
+#if WITH_EDITOR
+	virtual FName GetDefaultNodeName() const override;
+	virtual FText GetDefaultNodeTitle() const override;
+	virtual TArray<FPCGPreConfiguredSettingsInfo> GetPreconfiguredInfo() const override;
+	virtual bool GroupPreconfiguredSettings() const override { return false; }
+	virtual bool HasFlippedTitleLines() const override { return false; }
+#endif
+
+	virtual FString GetAdditionalTitleInformation() const override;
+	virtual void ApplyPreconfiguredSettings(const FPCGPreConfiguredSettingsInfo& PreconfigureInfo) override;
+
+protected:
+	virtual FPCGElementPtr CreateElement() const override;
+	//~End UPCGSettings interface
+
+public:
+	//~Begin UPCGMetadataSettingsBase interface
+	virtual FPCGAttributePropertyInputSelector GetInputSource(uint32 Index) const override;
+	virtual bool IsSupportedInputType(uint16 TypeId, uint32 InputIndex, bool& bHasSpecialRequirement) const override;
+	//~End UPCGMetadataSettingsBase interface
+
+public:
+	UPROPERTY(BlueprintReadWrite, EditAnywhere, Category = Settings, meta = (PCG_Overridable))
+	FPCGAttributePropertyInputSelector InputSource;
+
+	/** If InRangeMin = InRangeMax, then that attribute value is mapped to the average of OutRangeMin and OutRangeMax */
+	UPROPERTY(BlueprintReadWrite, EditAnywhere, Category = Settings, meta = (PCG_Overridable))
+	double InRangeMin = 0.f;
+
+	/** If InRangeMin = InRangeMax, then that attribute value is mapped to the average of OutRangeMin and OutRangeMax */
+	UPROPERTY(BlueprintReadWrite, EditAnywhere, Category = Settings, meta = (PCG_Overridable))
+	double InRangeMax = 1.f;
+
+	UPROPERTY(BlueprintReadWrite, EditAnywhere, Category = Settings, meta = (PCG_Overridable))
+	double OutRangeMin = 0.f;
+
+	UPROPERTY(BlueprintReadWrite, EditAnywhere, Category = Settings, meta = (PCG_Overridable))
+	double OutRangeMax = 1.f;
+
+	/** If checked, outside values will be clamped between 0 and 1. */
+	UPROPERTY(BlueprintReadWrite, EditAnywhere, Category = Settings, meta = (PCG_Overridable))
+	bool bClampToUnitRange = false;
+
+	/** Attribute values outside of the input range will be unaffected by the remapping */
+	UPROPERTY(BlueprintReadWrite, EditAnywhere, Category = Settings, meta = (PCG_Overridable))
+	bool bIgnoreValuesOutsideInputRange = false;
+};
+
+class FPCGAttributeRemapElement : public FPCGMetadataElementBase
+{
+protected:
+	virtual bool DoOperation(PCGMetadataOps::FOperationData& OperationData) const override;
+};
+
diff --git a/Engine/Plugins/PCG/Source/PCG/Public/Elements/PCGDensityRemapElement.h b/Engine/Plugins/PCG/Source/PCG/Public/Elements/PCGDensityRemapElement.h
--- a/Engine/Plugins/PCG/Source/PCG/Public/Elements/PCGDensityRemapElement.h
+++ b/Engine/Plugins/PCG/Source/PCG/Public/Elements/PCGDensityRemapElement.h
@@ -8,11 +8,13 @@
 #include "PCGDensityRemapElement.generated.h"
 
 UCLASS(MinimalAPI, BlueprintType, ClassGroup = (Procedural))
-class UPCGDensityRemapSettings : public UPCGSettings
+class UE_DEPRECATED(5.5, "Superseded by UPCGAttributeRemapSettings") UPCGDensityRemapSettings : public UPCGSettings
 {
 	GENERATED_BODY()
 
 public:
+	UPCGDensityRemapSettings();
+
 	// ~Begin UPCGSettings interface
 #if WITH_EDITOR
 	virtual FName GetDefaultNodeName() const override { return FName(TEXT("DensityRemap")); }
