[PCG] Improve Attribute Accessors

* Add tools to do metadata/attribute operations in one go instead of one by one:
    - AddEntries
    - GetValueKeys
    - GetValues
    - GetValuesFromItemKeys
    - GetParentKeys

* Use that in Attribute Accessor to avoid locking too many times in read/write
* Add Constructible| Broadcast for convenience
* Add option to create an attribute if it doesn't exist when creating an output accessor.

#rmx

diff --git a/Engine/Plugins/PCG/Source/PCG/Private/Metadata/Accessors/PCGAttributeAccessorHelpers.cpp b/Engine/Plugins/PCG/Source/PCG/Private/Metadata/Accessors/PCGAttributeAccessorHelpers.cpp
--- a/Engine/Plugins/PCG/Source/PCG/Private/Metadata/Accessors/PCGAttributeAccessorHelpers.cpp
+++ b/Engine/Plugins/PCG/Source/PCG/Private/Metadata/Accessors/PCGAttributeAccessorHelpers.cpp
@@ -8,6 +8,8 @@
 #include "Data/PCGPointData.h"
 #include "Data/PCGSpatialData.h"
 #include "Metadata/PCGAttributePropertySelector.h"
+#include "Metadata/PCGMetadata.h"
+#include "Metadata/PCGMetadataAttributeTpl.h"
 #include "Metadata/Accessors/IPCGAttributeAccessorTpl.h"
 #include "Metadata/Accessors/PCGAttributeAccessor.h"
 #include "Metadata/Accessors/PCGCustomAccessor.h"
@@ -630,6 +632,65 @@ TUniquePtr<IPCGAttributeAccessor> PCGAttributeAccessorHelpers::CreateAccessor(FP
 	return CreateAttributeAccessorImpl<IPCGAttributeAccessor>(InAttribute, InMetadata, bQuiet);
 }
 
+TUniquePtr<IPCGAttributeAccessor> PCGAttributeAccessorHelpers::CreateAccessorWithAttributeCreation(UPCGData* InData, const FPCGAttributePropertySelector& InSelector, const IPCGAttributeAccessor* InMatchingAccessor, EPCGAttributeAccessorFlags InTypeMatching, bool bQuiet)
+{
+	TUniquePtr<IPCGAttributeAccessor> Result = CreateAccessor(InData,InSelector, bQuiet);
+
+	if (!InSelector.IsBasicAttribute() || !InMatchingAccessor)
+	{
+		return Result;
+	}
+
+	bool bValid = !!Result;
+
+	if (bValid && !!(InTypeMatching & EPCGAttributeAccessorFlags::StrictType))
+	{
+		bValid &= (InMatchingAccessor->GetUnderlyingType() == Result->GetUnderlyingType());
+	}
+
+	if (bValid && !!(InTypeMatching & EPCGAttributeAccessorFlags::AllowBroadcast))
+	{
+		bValid &= PCG::Private::IsBroadcastable(InMatchingAccessor->GetUnderlyingType(), Result->GetUnderlyingType());
+	}
+
+	if (bValid && !!(InTypeMatching & EPCGAttributeAccessorFlags::AllowConstructible))
+	{
+		bValid &= PCG::Private::IsConstructible(InMatchingAccessor->GetUnderlyingType(), Result->GetUnderlyingType());
+	}
+
+	if (!bValid)
+	{
+		Result.Reset();
+
+		// We didn't find the attribute in the data, or we can't Broadcast/Construct, so create a new one
+		UPCGMetadata* Metadata = InData->MutableMetadata();
+		if (!Metadata)
+		{
+			return Result;
+		}
+
+		const FName AttributeName = InSelector.GetName();
+		if (Metadata->HasAttribute(AttributeName))
+		{
+			Metadata->DeleteAttribute(AttributeName);
+		}
+
+		auto CreateAttributeAndAccessor = [InMatchingAccessor, AttributeName, Metadata](auto&& Dummy) -> TUniquePtr<IPCGAttributeAccessor>
+		{
+			using AttributeType = std::decay_t<decltype(Dummy)>;
+			AttributeType DefaultValue = PCG::Private::MetadataTraits<AttributeType>::ZeroValue();
+			InMatchingAccessor->Get<AttributeType>(DefaultValue, FPCGAttributeAccessorKeysEntries(PCGInvalidEntryKey));
+			FPCGMetadataAttribute<AttributeType>* Attribute = Metadata->CreateAttribute<AttributeType>(AttributeName, DefaultValue, /*bAllowInterpolation=*/true, /*bOverrideParent=*/false);
+
+			return Attribute ? MakeUnique<FPCGAttributeAccessor<AttributeType>>(Attribute, Metadata) : nullptr;
+		};
+
+		Result = PCGMetadataAttribute::CallbackWithRightType(InMatchingAccessor->GetUnderlyingType(), std::move(CreateAttributeAndAccessor));
+	}
+
+	return Result;
+}
+
 TUniquePtr<const IPCGAttributeAccessorKeys> PCGAttributeAccessorHelpers::CreateConstKeys(const UPCGData* InData, const FPCGAttributePropertySelector& InSelector)
 {
 	return CreateKeysImpl<const IPCGAttributeAccessorKeys, const UPCGData>(InData, InSelector);
diff --git a/Engine/Plugins/PCG/Source/PCG/Private/Metadata/PCGMetadata.cpp b/Engine/Plugins/PCG/Source/PCG/Private/Metadata/PCGMetadata.cpp
--- a/Engine/Plugins/PCG/Source/PCG/Private/Metadata/PCGMetadata.cpp
+++ b/Engine/Plugins/PCG/Source/PCG/Private/Metadata/PCGMetadata.cpp
@@ -909,6 +909,31 @@ int64 UPCGMetadata::AddEntry(int64 ParentEntry)
 	return ParentKeys.Add(ParentEntry) + ItemKeyOffset;
 }
 
+TArray<int64> UPCGMetadata::AddEntries(TArrayView<const int64> ParentEntryKeys)
+{
+	TArray<int64> Result;
+	Result.Reserve(ParentEntryKeys.Num());
+
+	FWriteScopeLock ScopeLock(ItemLock);
+	ParentKeys.Reserve(ParentKeys.Num() + ParentEntryKeys.Num());
+	for (const int64 ParentEntry : ParentEntryKeys)
+	{
+		Result.Add(ParentKeys.Add(ParentEntry) + ItemKeyOffset);
+	}
+
+	return Result;
+}
+
+void UPCGMetadata::AddEntriesInPlace(TArrayView<int64*> ParentEntryKeys)
+{
+	FWriteScopeLock ScopeLock(ItemLock);
+	ParentKeys.Reserve(ParentKeys.Num() + ParentEntryKeys.Num());
+	for (int64* ParentEntry : ParentEntryKeys)
+	{
+		*ParentEntry = ParentKeys.Add(*ParentEntry) + ItemKeyOffset;
+	}
+}
+
 int64 UPCGMetadata::AddEntryPlaceholder()
 {
 	FReadScopeLock ScopeLock(ItemLock);
@@ -982,6 +1007,43 @@ PCGMetadataEntryKey UPCGMetadata::GetParentKey(PCGMetadataEntryKey LocalItemKey)
 	}
 }
 
+void UPCGMetadata::GetParentKeys(TArrayView<PCGMetadataEntryKey> LocalItemKeys, const TBitArray<>* Mask) const
+{
+	auto GetParentKey_Unsafe = [this](PCGMetadataEntryKey& LocalItemKey) -> void
+	{
+		if (LocalItemKey < ItemKeyOffset)
+		{
+			// Key is already in parent referential
+			return;
+		}
+		else if (LocalItemKey - ItemKeyOffset < ParentKeys.Num())
+		{
+			LocalItemKey = ParentKeys[LocalItemKey - ItemKeyOffset];
+		}
+		else
+		{
+			UE_LOG(LogPCG, Warning, TEXT("Invalid metadata key - check for entry key not properly initialized"));
+			LocalItemKey = PCGInvalidEntryKey;
+		}
+	};
+
+	FReadScopeLock ScopeLock(ItemLock);
+	if (Mask && ensure(LocalItemKeys.Num() == Mask->Num()))
+	{
+		for (TConstSetBitIterator<> It(*Mask); It; ++It)
+		{
+			GetParentKey_Unsafe(LocalItemKeys[It.GetIndex()]);
+		}
+	}
+	else
+	{
+		for (PCGMetadataEntryKey& LocalItemKey : LocalItemKeys)
+		{
+			GetParentKey_Unsafe(LocalItemKey);
+		}
+	}
+}
+
 void UPCGMetadata::MergePointAttributes(const FPCGPoint& InPointA, const FPCGPoint& InPointB, FPCGPoint& OutPoint, EPCGMetadataOp Op)
 {
 	MergeAttributes(InPointA.MetadataEntry, this, InPointB.MetadataEntry, this, OutPoint.MetadataEntry, Op);
@@ -1291,7 +1353,7 @@ void UPCGMetadata::SetAttributes(const TArrayView<const PCGMetadataEntryKey>& In
 				if (Attribute == OtherAttribute)
 				{
 					TArray<PCGMetadataValueKey> ValueKeys;
-					Attribute->GetValueKeys(InKeys, ValueKeys);
+					Attribute->GetValueKeys(TArrayView<const PCGMetadataEntryKey>(InKeys), ValueKeys);
 					Attribute->SetValuesFromValueKeys(OutKeys, ValueKeys);
 				}
 				else
diff --git a/Engine/Plugins/PCG/Source/PCG/Private/Metadata/PCGMetadataAttribute.cpp b/Engine/Plugins/PCG/Source/PCG/Private/Metadata/PCGMetadataAttribute.cpp
--- a/Engine/Plugins/PCG/Source/PCG/Private/Metadata/PCGMetadataAttribute.cpp
+++ b/Engine/Plugins/PCG/Source/PCG/Private/Metadata/PCGMetadataAttribute.cpp
@@ -162,10 +162,25 @@ void FPCGMetadataAttributeBase::GetValueKeys(const TArrayView<const PCGMetadataE
 	// Bitset with all unset values. If we have any unset value, we will ask the parent for those.
 	TBitArray<> UnsetValues(true, EntryKeys.Num());
 
-	GetValueKeys_Internal(EntryKeys, OutValueKeys, UnsetValues);
+	GetValueKeys_Internal(EntryKeys, OutValueKeys, UnsetValues, /*bOwnerOfEntryKeysView=*/false);
 }
 
-void FPCGMetadataAttributeBase::GetValueKeys_Internal(const TArrayView<const PCGMetadataEntryKey>& EntryKeys, TArrayView<PCGMetadataValueKey> OutValueKeys, TBitArray<>& UnsetValues) const
+void FPCGMetadataAttributeBase::GetValueKeys(TArrayView<PCGMetadataEntryKey> EntryKeys, TArray<PCGMetadataValueKey>& OutValueKeys) const
+{
+	if (EntryKeys.IsEmpty())
+	{
+		return;
+	}
+
+	OutValueKeys.SetNumUninitialized(EntryKeys.Num());
+	// Bitset with all unset values. If we have any unset value, we will ask the parent for those.
+	TBitArray<> UnsetValues(true, EntryKeys.Num());
+
+	// We have ownership on the memory, so we can skip the copy of entries.
+	GetValueKeys_Internal(EntryKeys, OutValueKeys, UnsetValues, /*bOwnerOfEntryKeysView=*/true);
+}
+
+void FPCGMetadataAttributeBase::GetValueKeys_Internal(const TArrayView<const PCGMetadataEntryKey> EntryKeys, TArrayView<PCGMetadataValueKey> OutValueKeys, TBitArray<>& UnsetValues, bool bOwnerOfEntryKeysView) const
 {
 	check(EntryKeys.Num() == OutValueKeys.Num() && OutValueKeys.Num() == UnsetValues.Num());
 
@@ -213,7 +228,26 @@ void FPCGMetadataAttributeBase::GetValueKeys_Internal(const TArrayView<const PCG
 
 	if (Parent && !bFoundAllKeys)
 	{
-		Parent->GetValueKeys_Internal(EntryKeys, OutValueKeys, UnsetValues);
+		auto ParentCall = [this, &OutValueKeys, &UnsetValues](const TArrayView<const PCGMetadataEntryKey> CurrentEntryKeys)
+		{
+			// Before querying the parent, we need to update all our entry keys to get them in the parent referential.
+			// At that point, we are owner of our memory, so it is safe to cast
+			Metadata->GetParentKeys(MakeArrayView(const_cast<PCGMetadataEntryKey*>(CurrentEntryKeys.GetData()), CurrentEntryKeys.Num()), &UnsetValues);
+			Parent->GetValueKeys_Internal(CurrentEntryKeys, OutValueKeys, UnsetValues, /*bOwnerOfEntryKeysView=*/true);
+		};
+
+		// If the input data is coming from outside, we need to copy it to be able to modify it.
+		// Only do it there, because we don't have to pay the cost of the copy if we don't have to check the parent.
+		if (!bOwnerOfEntryKeysView)
+		{
+			TArray<PCGMetadataEntryKey, TInlineAllocator<256>> CopiedKeys;
+			CopiedKeys = EntryKeys;
+			ParentCall(CopiedKeys);
+		}
+		else
+		{
+			ParentCall(EntryKeys);
+		}
 	}
 }
 
diff --git a/Engine/Plugins/PCG/Source/PCG/Public/Metadata/Accessors/IPCGAttributeAccessor.h b/Engine/Plugins/PCG/Source/PCG/Public/Metadata/Accessors/IPCGAttributeAccessor.h
--- a/Engine/Plugins/PCG/Source/PCG/Public/Metadata/Accessors/IPCGAttributeAccessor.h
+++ b/Engine/Plugins/PCG/Source/PCG/Public/Metadata/Accessors/IPCGAttributeAccessor.h
@@ -13,7 +13,7 @@ class IPCGAttributeAccessorKeys;
 
 UENUM(Meta = (Bitflags))
 enum class EPCGAttributeAccessorFlags
-{	
+{
 	// Always require that the underlying type of the accessor match the expected type, 1 for 1.
 	StrictType = 1 << 0,
 
@@ -25,7 +25,9 @@ enum class EPCGAttributeAccessorFlags
 
 	// By default, if the key is a PCGInvalidEntryKey, it will add a new entry. With this set, it will override the default value.
 	// USE WITH CAUTION
-	AllowSetDefaultValue = 1 << 3
+	AllowSetDefaultValue = 1 << 3,
+
+	AllowBroadcastAndConstructible = AllowBroadcast | AllowConstructible
 };
 ENUM_CLASS_FLAGS(EPCGAttributeAccessorFlags);
 
diff --git a/Engine/Plugins/PCG/Source/PCG/Public/Metadata/Accessors/PCGAttributeAccessor.h b/Engine/Plugins/PCG/Source/PCG/Public/Metadata/Accessors/PCGAttributeAccessor.h
--- a/Engine/Plugins/PCG/Source/PCG/Public/Metadata/Accessors/PCGAttributeAccessor.h
+++ b/Engine/Plugins/PCG/Source/PCG/Public/Metadata/Accessors/PCGAttributeAccessor.h
@@ -39,20 +39,20 @@ public:
 
 	bool GetRangeImpl(TArrayView<T> OutValues, int32 Index, const IPCGAttributeAccessorKeys& Keys) const
 	{
-		TArray<const PCGMetadataEntryKey*, TInlineAllocator<256>> EntryKeys;
-		EntryKeys.SetNumUninitialized(OutValues.Num());
+		TArray<const PCGMetadataEntryKey*, TInlineAllocator<256>> EntryKeyPtrs;
+		EntryKeyPtrs.SetNumUninitialized(OutValues.Num());
 
-		TArrayView<const PCGMetadataEntryKey*> EntryKeysView(EntryKeys);
+		TArrayView<const PCGMetadataEntryKey*> EntryKeysView(EntryKeyPtrs);
 		if (!Keys.GetKeys<PCGMetadataEntryKey>(Index, EntryKeysView))
 		{
 			return false;
 		}
 
-		// TODO: Might be good to hase a "GetValuesFromItemKeys" to try locking less often.
-		for (int32 i = 0; i < OutValues.Num(); ++i)
-		{
-			OutValues[i] = Attribute->GetValueFromItemKey(*EntryKeys[i]);
-		}
+		TArray<PCGMetadataEntryKey, TInlineAllocator<256>> EntryKeys;
+		EntryKeys.Reserve(EntryKeyPtrs.Num());
+		Algo::Transform(EntryKeyPtrs, EntryKeys, [](const PCGMetadataEntryKey* KeyPtr) { return *KeyPtr; });
+
+		Attribute->GetValuesFromItemKeys(TArrayView<PCGMetadataEntryKey>(EntryKeys), OutValues);
 
 		return true;
 	}
@@ -69,6 +69,9 @@ public:
 
 		int LastDefaultKeyIndex = INDEX_NONE;
 
+		TArray<PCGMetadataEntryKey*, TInlineAllocator<256>> EntriesToSet;
+		EntriesToSet.Reserve(EntryKeys.Num());
+
 		// Implementation note: this is a stripped down version of UPCGMetadata::InitializeOnSet
 		for(int EntryIndex = 0; EntryIndex < EntryKeys.Num(); ++EntryIndex)
 		{
@@ -77,7 +80,7 @@ public:
 			{
 				if (!(Flags & EPCGAttributeAccessorFlags::AllowSetDefaultValue))
 				{
-					EntryKey = Metadata->AddEntry(); // TODO - replace by AddEntryPlaceholder ?
+					EntriesToSet.Add(&EntryKey);
 				}
 				else
 				{
@@ -86,11 +89,19 @@ public:
 			}
 			else if (EntryKey < Metadata->GetItemKeyCountForParent())
 			{
-				EntryKey = Metadata->AddEntry(EntryKey);
+				EntriesToSet.Add(&EntryKey);
 			}
 		}
 
-		Attribute->SetValues(EntryKeys, InValues);
+		{
+			//TRACE_CPUPROFILER_EVENT_SCOPE(FPCGAttributeAccessor::SetRangeImpl::AddEntriesInPlace);
+			Metadata->AddEntriesInPlace(EntriesToSet);
+		}
+
+		{
+			//TRACE_CPUPROFILER_EVENT_SCOPE(FPCGAttributeAccessor::SetRangeImpl::SetValues);
+			Attribute->SetValues(EntryKeys, InValues);
+		}
 
 		if (LastDefaultKeyIndex != INDEX_NONE)
 		{
diff --git a/Engine/Plugins/PCG/Source/PCG/Public/Metadata/Accessors/PCGAttributeAccessorHelpers.h b/Engine/Plugins/PCG/Source/PCG/Public/Metadata/Accessors/PCGAttributeAccessorHelpers.h
--- a/Engine/Plugins/PCG/Source/PCG/Public/Metadata/Accessors/PCGAttributeAccessorHelpers.h
+++ b/Engine/Plugins/PCG/Source/PCG/Public/Metadata/Accessors/PCGAttributeAccessorHelpers.h
@@ -72,13 +72,22 @@ namespace PCGAttributeAccessorHelpers
 	PCG_API TUniquePtr<const IPCGAttributeAccessor> CreateConstAccessor(const FPCGMetadataAttributeBase* InAttribute, const UPCGMetadata* InMetadata, bool bQuiet = false);
 
 	/**
-	* Creates a accessor to the property or attribute pointed at by the InSelector.
+	* Creates an accessor to the property or attribute pointed at by the InSelector.
 	* Note that InData must not be null if the selector points to an attribute,
 	* but in the case of properties, it either has to be the appropriate type or null.
 	* Make sure to update your selector before-hand if you want to support "@Source". Otherwise the creation will fail.
 	*/
 	PCG_API TUniquePtr<IPCGAttributeAccessor> CreateAccessor(UPCGData* InData, const FPCGAttributePropertySelector& InSelector, bool bQuiet = false);
 
+	/**
+	* Creates an accessor to the property or attribute pointed at by the InSelector. If the selector points to a base attribute and the attribute doesn't exist 
+	* or the types doesn't match with the Matching Accessor, create the attribute using the Matching Accessor.
+	* Note that InData must not be null if the selector points to an attribute,
+	* but in the case of properties, it either has to be the appropriate type or null.
+	* Make sure to update your selector before-hand if you want to support "@Source". Otherwise the creation will fail.
+	*/
+	PCG_API TUniquePtr<IPCGAttributeAccessor> CreateAccessorWithAttributeCreation(UPCGData* InData, const FPCGAttributePropertySelector& InSelector, const IPCGAttributeAccessor* InMatchingAccessor, EPCGAttributeAccessorFlags InTypeMatching = EPCGAttributeAccessorFlags::AllowBroadcastAndConstructible, bool bQuiet = false);
+
 	/**
 	* Creates an accessor to an attribute without requiring a selector.
 	*/
diff --git a/Engine/Plugins/PCG/Source/PCG/Public/Metadata/PCGMetadata.h b/Engine/Plugins/PCG/Source/PCG/Public/Metadata/PCGMetadata.h
--- a/Engine/Plugins/PCG/Source/PCG/Public/Metadata/PCGMetadata.h
+++ b/Engine/Plugins/PCG/Source/PCG/Public/Metadata/PCGMetadata.h
@@ -229,6 +229,12 @@ public:
 	UFUNCTION(BlueprintCallable, Category = "PCG|Metadata")
 	int64 AddEntry(int64 ParentEntryKey = -1);
 
+	/** Adds a unique entry key to the metadata for all the parent entry keys. */
+	TArray<int64> AddEntries(TArrayView<const int64> ParentEntryKeys);
+
+	/** Adds a unique entry key to the metadata for all the parent entry keys, in place. */
+	void AddEntriesInPlace(TArrayView<int64*> ParentEntryKeys);
+
 	/** Advanced method.
 	*   In a MT context, we might not want to add the entry directly (because of write lock). Call this to generate an unique index in the MT context
 	*   And call AddDelayedEntries at the end when you want to add all the entries.
@@ -249,6 +255,9 @@ public:
 	/** Metadata chaining mechanism */
 	PCGMetadataEntryKey GetParentKey(PCGMetadataEntryKey LocalItemKey) const;
 
+	/** Metadata chaining mechanism for bulk version. Can provide a mask to only update only a subset of the passed keys. */
+	void GetParentKeys(TArrayView<PCGMetadataEntryKey> LocalItemKeys, const TBitArray<>* Mask = nullptr) const;
+
 	/** Attributes operations */
 	void MergeAttributes(PCGMetadataEntryKey InKeyA, const UPCGMetadata* InMetadataA, PCGMetadataEntryKey InKeyB, const UPCGMetadata* InMetadataB, PCGMetadataEntryKey& OutKey, EPCGMetadataOp Op);
 	void MergeAttributesSubset(PCGMetadataEntryKey InKeyA, const UPCGMetadata* InMetadataA, const UPCGMetadata* InMetadataSubetA, PCGMetadataEntryKey InKeyB, const UPCGMetadata* InMetadataB, const UPCGMetadata* InMetadataSubsetb, PCGMetadataEntryKey& OutKey, EPCGMetadataOp Op);
diff --git a/Engine/Plugins/PCG/Source/PCG/Public/Metadata/PCGMetadataAttribute.h b/Engine/Plugins/PCG/Source/PCG/Public/Metadata/PCGMetadataAttribute.h
--- a/Engine/Plugins/PCG/Source/PCG/Public/Metadata/PCGMetadataAttribute.h
+++ b/Engine/Plugins/PCG/Source/PCG/Public/Metadata/PCGMetadataAttribute.h
@@ -61,6 +61,9 @@ public:
 	/** Bulk getter, to lock in read only once per parent. */
 	void GetValueKeys(const TArrayView<const PCGMetadataEntryKey>& EntryKeys, TArray<PCGMetadataValueKey>& OutValueKeys) const;
 
+	/** Optimized version that take ownership on the Entries passed.*/
+	void GetValueKeys(TArrayView<PCGMetadataEntryKey> EntryKeys, TArray<PCGMetadataValueKey>& OutValueKeys) const;
+
 	/** Bulk setter to lock in write only once. */
 	void SetValuesFromValueKeys(const TArrayView<const TTuple<PCGMetadataEntryKey, PCGMetadataValueKey>>& EntryValuePairs, bool bResetValueOnDefaultValueKey = true);
 
@@ -89,7 +92,11 @@ private:
 	// Unsafe version, needs to be write lock protected.
 	void SetValueFromValueKey_Unsafe(PCGMetadataEntryKey EntryKey, PCGMetadataValueKey ValueKey, bool bResetValueOnDefaultValueKey, bool bAllowInvalidEntries = false);
 
-	void GetValueKeys_Internal(const TArrayView<const PCGMetadataEntryKey>& EntryKeys, TArrayView<PCGMetadataValueKey> OutValueKeys, TBitArray<>& UnsetValues) const;
+	// Gather the value keys for the list of entry keys.
+	// Because we need to update the entry keys if we look for the value keys in the parent, but because the EntryKeys are coming from the outside, we can't modify them.
+	// (as entry keys need to be put into the parent referential when looking for value keys). 
+	// So we will copy internally the EntryKeys to modify them (and only once) if we are not owner of the memory.
+	void GetValueKeys_Internal(const TArrayView<const PCGMetadataEntryKey> EntryKeys, TArrayView<PCGMetadataValueKey> OutValueKeys, TBitArray<>& UnsetValues, bool bOwnerOfEntryKeysView = false) const;
 
 protected:
 	TMap<PCGMetadataEntryKey, PCGMetadataValueKey> EntryToValueKeyMap;
diff --git a/Engine/Plugins/PCG/Source/PCG/Public/Metadata/PCGMetadataAttributeTpl.h b/Engine/Plugins/PCG/Source/PCG/Public/Metadata/PCGMetadataAttributeTpl.h
--- a/Engine/Plugins/PCG/Source/PCG/Public/Metadata/PCGMetadataAttributeTpl.h
+++ b/Engine/Plugins/PCG/Source/PCG/Public/Metadata/PCGMetadataAttributeTpl.h
@@ -603,6 +603,108 @@ public:
 		}
 	}
 
+protected:
+	void GetValues_Internal(const TArrayView<const PCGMetadataValueKey> ValueKeys, TArrayView<T> OutValues, TBitArray<>& UnretrievedValues) const
+	{
+		check(ValueKeys.Num() == OutValues.Num() && OutValues.Num() == UnretrievedValues.Num());
+
+		bool bFoundAllKeys = true;
+		TConstSetBitIterator<> It(UnretrievedValues);
+		if (!It)
+		{
+			return;
+		}
+
+		const FPCGMetadataAttribute* ThisParent = GetParent();
+
+		ValueLock.ReadLock();
+
+		for (; It; ++It)
+		{
+			const int32 Index = It.GetIndex();
+			const PCGMetadataValueKey ValueKey = ValueKeys[Index];
+
+			auto RetrieveValue = [Index, &OutValues, &UnretrievedValues](T Value)
+			{
+				OutValues[Index] = std::move(Value);
+				UnretrievedValues[Index] = false;
+			};
+
+			if (ValueKey == PCGDefaultValueKey)
+			{
+				RetrieveValue(DefaultValue);
+			}
+			else if (ValueKey >= ValueKeyOffset)
+			{
+				int32 ValueIndex = ValueKey - ValueKeyOffset;
+				RetrieveValue(ValueIndex < Values.Num() ? Values[ValueIndex] : DefaultValue);
+			}
+			else if (!ThisParent)
+			{
+				RetrieveValue(DefaultValue);
+			}
+			else
+			{
+				bFoundAllKeys = false;
+			}
+		}
+
+		ValueLock.ReadUnlock();
+
+		ensure(ThisParent || bFoundAllKeys);
+
+		if (ThisParent && !bFoundAllKeys)
+		{
+			ThisParent->GetValues_Internal(ValueKeys, OutValues, UnretrievedValues);
+		}
+	}
+
+public:
+	/**
+	* Write into pre-allocated OutValues the values associated with the given value keys.
+	*/
+	void GetValues(const TArrayView<const PCGMetadataValueKey> ValueKeys, TArrayView<T> OutValues) const
+	{
+		// Bitset with all unretrieved values. If we have any unretrieved value, we will ask the parent for those.
+		TBitArray<> UnretrievedValues(true, ValueKeys.Num());
+		GetValues_Internal(ValueKeys, OutValues, UnretrievedValues);
+	}
+
+	/** 
+	* Write into pre-allocated OutValues the values associated with the given entry keys. 
+	* Const version on the Entry Keys, where they won't be modified. It will induce a copy of the entry keys
+	* if we ever have to go check the parent attribute, as we need to modify the entry keys for that.
+	* If you don't care if the Entry Keys are modified, use the non-const version of the EntryKeys.
+	*/
+	void GetValuesFromItemKeys(const TArrayView<const PCGMetadataEntryKey> EntryKeys, TArrayView<T> OutValues) const
+	{
+		if (!ensure(EntryKeys.Num() == OutValues.Num()))
+		{
+			return;
+		}
+
+		TArray<PCGMetadataValueKey> ValueKeys;
+		GetValueKeys(EntryKeys, ValueKeys);
+		GetValues(ValueKeys, OutValues);
+	}
+
+	/** 
+	* Write into pre-allocated OutValues the values associated with the given entry keys.
+	* Non-Const version on the Entry Keys, where they can be modified. If you need the Entry Keys to not be modifed,
+	* use the const version of the EntryKeys.
+	*/
+	void GetValuesFromItemKeys(TArrayView<PCGMetadataEntryKey> EntryKeys, TArrayView<T> OutValues) const
+	{
+		if (!ensure(EntryKeys.Num() == OutValues.Num()))
+		{
+			return;
+		}
+
+		TArray<PCGMetadataValueKey> ValueKeys;
+		GetValueKeys(EntryKeys, ValueKeys);
+		GetValues(ValueKeys, OutValues);
+	}
+
 	/** Code related to finding values / compressing data */
 	virtual bool UsesValueKeys() const override
 	{
@@ -887,8 +989,6 @@ namespace PCGMetadataAttribute
 		default:
 			return nullptr;
 		}
-
-#undef AllocatePCGMetadataAttributeOnType
 	}
 
 	template <typename Func, typename... Args>
