[PCG] Wrong assumption in the test, clarify the intent in the copy attribute comment

#rnx

diff --git a/Engine/Plugins/PCG/Source/PCG/Private/Helpers/PCGMetadataHelpers.cpp b/Engine/Plugins/PCG/Source/PCG/Private/Helpers/PCGMetadataHelpers.cpp
--- a/Engine/Plugins/PCG/Source/PCG/Private/Helpers/PCGMetadataHelpers.cpp
+++ b/Engine/Plugins/PCG/Source/PCG/Private/Helpers/PCGMetadataHelpers.cpp
@@ -164,7 +164,7 @@ namespace PCGMetadataHelpers
 
 				if (FPCGMetadataAttributeBase* NewAttr = TargetMetadata->CopyAttribute(SourceAttribute, LocalDestinationAttribute, /*bKeepParent=*/false, /*bCopyEntries=*/true,/*bCopyValues=*/true))
 				{
-					// To keep the previous behavior, we force the copied attribute to have its default value set to the first entry.
+					// To keep the previous behavior for single value attributes, we force the copied attribute to have its default value set to the first entry (only done if there is just a single value).
 					NewAttr->SetDefaultValueToFirstEntry();
 				}
 				else
diff --git a/Engine/Plugins/PCG/Source/PCG/Private/Tests/Elements/PCGCopyAttributeTest.cpp b/Engine/Plugins/PCG/Source/PCG/Private/Tests/Elements/PCGCopyAttributeTest.cpp
--- a/Engine/Plugins/PCG/Source/PCG/Private/Tests/Elements/PCGCopyAttributeTest.cpp
+++ b/Engine/Plugins/PCG/Source/PCG/Private/Tests/Elements/PCGCopyAttributeTest.cpp
@@ -236,7 +236,7 @@ bool FPCGCopyAttributeTests_Points_AttributeToAttribute::RunTest(const FString&
 
 	check(InputAttribute && OutputAttribute);
 
-	UTEST_EQUAL("Output Attribute default value is the same as first entry of Input Attribute", OutputAttribute->GetValueFromItemKey(PCGInvalidEntryKey), InputAttribute->GetValueFromItemKey(PCGMetadataEntryKey(0)));
+	UTEST_EQUAL("Output Attribute default value is the same as default value of Input Attribute", OutputAttribute->GetValueFromItemKey(PCGInvalidEntryKey), InputAttribute->GetValueFromItemKey(PCGInvalidEntryKey));
 
 	// Input attribute was copied in the output attribute correctly
 	for (int i = 0; i < NumOfPoints; ++i)
@@ -373,7 +373,7 @@ bool FPCGCopyAttributeTests_Params_MultiValue::RunTest(const FString& Parameters
 
 	check(InputAttribute && OutputAttribute);
 
-	UTEST_EQUAL("Output Attribute default value is the same as first entry of Input Attribute", OutputAttribute->GetValueFromItemKey(PCGInvalidEntryKey), InputAttribute->GetValueFromItemKey(PCGMetadataEntryKey(0)));
+	UTEST_EQUAL("Output Attribute default value is the same as default value of Input Attribute", OutputAttribute->GetValueFromItemKey(PCGInvalidEntryKey), InputAttribute->GetValueFromItemKey(PCGInvalidEntryKey));
 
 	// Input attribute was copied in the output attribute correctly
 	for (int i = 0; i < NumEntries; ++i)
