[PCG] Allowed include/exclude of landscape hits on the world queries

diff --git a/Engine/Plugins/PCG/Source/PCG/Private/Data/PCGWorldData.cpp b/Engine/Plugins/PCG/Source/PCG/Private/Data/PCGWorldData.cpp
--- a/Engine/Plugins/PCG/Source/PCG/Private/Data/PCGWorldData.cpp
+++ b/Engine/Plugins/PCG/Source/PCG/Private/Data/PCGWorldData.cpp
@@ -48,6 +48,37 @@ void FPCGWorldRayHitQueryParams::Initialize()
 	FPCGWorldCommonQueryParams::Initialize();
 }
 
+#if WITH_EDITOR
+void FPCGWorldCommonQueryParams::CommonPostLoad()
+{
+	if (bIgnoreLandscapeHits_DEPRECATED != false)
+	{
+		SelectLandscapeHits = EPCGWorldQuerySelectLandscapeHits::Exclude;
+		bIgnoreLandscapeHits_DEPRECATED = false;
+	}
+}
+#endif // WITH_EDITOR
+
+void FPCGWorldVolumetricQueryParams::PostSerialize(const FArchive& Ar)
+{
+#if WITH_EDITOR
+	if (Ar.IsLoading() && Ar.IsPersistent() && !Ar.HasAnyPortFlags(PPF_Duplicate | PPF_DuplicateForPIE))
+	{
+		CommonPostLoad();
+	}
+#endif
+}
+
+void FPCGWorldRayHitQueryParams::PostSerialize(const FArchive& Ar)
+{
+#if WITH_EDITOR
+	if (Ar.IsLoading() && Ar.IsPersistent() && !Ar.HasAnyPortFlags(PPF_Duplicate | PPF_DuplicateForPIE))
+	{
+		CommonPostLoad();
+	}
+#endif
+}
+
 void UPCGWorldVolumetricData::Initialize(UWorld* InWorld, const FBox& InBounds)
 {
 	Super::Initialize(InBounds);
@@ -130,9 +161,12 @@ bool UPCGWorldVolumetricData::SamplePoint(const FTransform& InTransform, const F
 			}
 		}
 
-		if (QueryParams.bIgnoreLandscapeHits && OverlappedComponent->GetOwner() && OverlappedComponent->GetOwner()->IsA<ALandscapeProxy>())
+		if (QueryParams.SelectLandscapeHits != EPCGWorldQuerySelectLandscapeHits::Include)
 		{
-			continue;
+			if (OverlappedComponent->GetOwner() && OverlappedComponent->GetOwner()->IsA<ALandscapeProxy>() != (QueryParams.SelectLandscapeHits == EPCGWorldQuerySelectLandscapeHits::Require))
+			{
+				continue;
+			}
 		}
 
 		if (QueryParams.bSearchForOverlap)
@@ -324,12 +358,13 @@ bool UPCGWorldRayHitData::SamplePoint(const FTransform& InTransform, const FBox&
 		}
 
 		bool bHitOnLandscape = false;
-		if (QueryParams.bIgnoreLandscapeHits || QueryParams.bApplyMetadataFromLandscape || QueryParams.bIgnoreBackfaceHits)
+		if(QueryParams.SelectLandscapeHits != EPCGWorldQuerySelectLandscapeHits::Include || QueryParams.bApplyMetadataFromLandscape || QueryParams.bIgnoreBackfaceHits)
 		{
 			bHitOnLandscape = HitComponent->GetOwner() && HitComponent->GetOwner()->IsA<ALandscapeProxy>();
 		}
 
-		if (QueryParams.bIgnoreLandscapeHits && bHitOnLandscape)
+		if((bHitOnLandscape && QueryParams.SelectLandscapeHits == EPCGWorldQuerySelectLandscapeHits::Exclude) ||
+			(!bHitOnLandscape && QueryParams.SelectLandscapeHits == EPCGWorldQuerySelectLandscapeHits::Require))
 		{
 			continue;
 		}
diff --git a/Engine/Plugins/PCG/Source/PCG/Public/Data/PCGWorldData.h b/Engine/Plugins/PCG/Source/PCG/Public/Data/PCGWorldData.h
--- a/Engine/Plugins/PCG/Source/PCG/Public/Data/PCGWorldData.h
+++ b/Engine/Plugins/PCG/Source/PCG/Public/Data/PCGWorldData.h
@@ -23,6 +23,14 @@ enum class EPCGWorldQueryFilterByTag
 	ExcludeTagged
 };
 
+UENUM()
+enum class EPCGWorldQuerySelectLandscapeHits : uint8
+{
+	Exclude = 0 UMETA(ToolTip="Excludes hits from the landscape."),
+	Include UMETA(ToolTip="Will report hits on the landscape."),
+	Require UMETA(ToolTip="Will return only hits on the landscape.")
+};
+
 namespace PCGWorldRayHitConstants
 {
 	const FName PhysicalMaterialReferenceAttribute = TEXT("PhysicalMaterial");
@@ -33,6 +41,10 @@ struct FPCGWorldCommonQueryParams
 {
 	GENERATED_BODY()
 
+#if WITH_EDITOR
+	void CommonPostLoad();
+#endif
+
 	/** If true, will ignore hits/overlaps on content created from PCG. */
 	UPROPERTY(BlueprintReadWrite, EditAnywhere, Category = Data, meta = (PCG_Overridable))
 	bool bIgnorePCGHits = false;
@@ -53,8 +65,13 @@ struct FPCGWorldCommonQueryParams
 	UPROPERTY(BlueprintReadWrite, EditAnywhere, Category = "Data|Filtering", meta = (PCG_Overridable, EditCondition = "ActorTagFilter != EPCGWorldQueryFilterByTag::NoTagFilter"))
 	FString ActorTagsList;
 
-	UPROPERTY(BlueprintReadWrite, EditAnywhere, Category = "Data|Filtering", meta = (PCG_Overridable))
-	bool bIgnoreLandscapeHits = false;
+	UPROPERTY(BlueprintReadWrite, EditAnywhere, Category = "Data|Filtering", meta = (PCG_Overridable, PCG_OverrideAliases="bIgnoreLandscapeHits"))
+	EPCGWorldQuerySelectLandscapeHits SelectLandscapeHits = EPCGWorldQuerySelectLandscapeHits::Include;
+
+#if WITH_EDITORONLY_DATA
+	UPROPERTY(meta=(DeprecatedProperty, DeprecationMessage="IgnoreLandscapeHits has been deprecated in favor of SelectLandscapeHits"))
+	bool bIgnoreLandscapeHits_DEPRECATED = false;
+#endif
 
 	UPROPERTY(BlueprintReadWrite, EditAnywhere, Category = Data, meta = (PCG_Overridable))
 	bool bGetReferenceToActorHit = false;
@@ -74,12 +91,22 @@ struct FPCGWorldVolumetricQueryParams : public FPCGWorldCommonQueryParams
 	GENERATED_BODY()
 
 	void Initialize();
+	void PostSerialize(const FArchive& Ar);
 
 	/** Controls whether we are trying to find an overlap with physical objects (true) or to find empty spaces that do not contain anything (false) */
 	UPROPERTY(BlueprintReadWrite, EditAnywhere, Category = Data, meta = (PCG_Overridable))
 	bool bSearchForOverlap = true;
 };
 
+template<>
+struct TStructOpsTypeTraits<FPCGWorldVolumetricQueryParams> : public TStructOpsTypeTraitsBase2<FPCGWorldVolumetricQueryParams>
+{
+	enum
+	{
+		WithPostSerialize = true,
+	};
+};
+
 /** Queries volume for presence of world collision or not. Can be used to voxelize environment. */
 UCLASS(MinimalAPI, BlueprintType, ClassGroup = (Procedural))
 class UPCGWorldVolumetricData : public UPCGVolumeData
@@ -121,6 +148,7 @@ struct FPCGWorldRayHitQueryParams : public FPCGWorldCommonQueryParams
 	GENERATED_BODY()
 
 	void Initialize();
+	void PostSerialize(const FArchive& Ar);
 
 	/** Set ray parameters including origin, direction and length explicitly rather than deriving these from the generating actor bounds. */
 	UPROPERTY(BlueprintReadWrite, EditAnywhere, Category = Data, meta = (DisplayName = "Set Ray Parameters", PCG_Overridable))
@@ -148,6 +176,15 @@ struct FPCGWorldRayHitQueryParams : public FPCGWorldCommonQueryParams
 	bool bGetReferenceToPhysicalMaterial = false;
 };
 
+template<>
+struct TStructOpsTypeTraits<FPCGWorldRayHitQueryParams> : public TStructOpsTypeTraitsBase2<FPCGWorldRayHitQueryParams>
+{
+	enum
+	{
+		WithPostSerialize = true,
+	};
+};
+
 /** Executes collision queries against world collision. */
 UCLASS(MinimalAPI, BlueprintType, ClassGroup = (Procedural))
 class UPCGWorldRayHitData : public UPCGSurfaceData
