[PCG] Fix PCG Volume not working when brush set to No Collision

#rnx

diff --git a/Engine/Plugins/PCG/Source/PCG/Private/Data/PCGVolumeData.cpp b/Engine/Plugins/PCG/Source/PCG/Private/Data/PCGVolumeData.cpp
--- a/Engine/Plugins/PCG/Source/PCG/Private/Data/PCGVolumeData.cpp
+++ b/Engine/Plugins/PCG/Source/PCG/Private/Data/PCGVolumeData.cpp
@@ -59,7 +59,9 @@ void UPCGVolumeData::Initialize(AVolume* InVolume)
 		FBodyInstance* BodyInstance = BrushComponent->GetBodyInstance();
 		UBodySetup* BodySetup = BrushComponent->GetBodySetup();
 
-		if (BodyInstance && BodySetup && !BodyInstance->IsDynamic())
+		// In some instances, non-collidable bodies will not be initialized, but it's not an issue for PCG so we can continue regardless.
+		// Otherwise, require that the body is not dynamic.
+		if (BodyInstance && BodySetup && (!FPhysicsInterface::IsValid(BodyInstance->ActorHandle) || !BodyInstance->IsDynamic()))
 		{
 			ReleaseInternalBodyInstance();
 
