[PCG] Fix SplineSampler ignoring dependency on PartitionActor volume for the bounding shape

#rnx

diff --git a/Engine/Plugins/PCG/Source/PCG/Private/Elements/PCGSplineSampler.cpp b/Engine/Plugins/PCG/Source/PCG/Private/Elements/PCGSplineSampler.cpp
--- a/Engine/Plugins/PCG/Source/PCG/Private/Elements/PCGSplineSampler.cpp
+++ b/Engine/Plugins/PCG/Source/PCG/Private/Elements/PCGSplineSampler.cpp
@@ -13,6 +13,7 @@
 #include "Data/PCGSpatialData.h"
 #include "Helpers/PCGBlueprintHelpers.h"
 #include "Helpers/PCGHelpers.h"
+#include "Helpers/PCGSettingsHelpers.h"
 
 #include "Async/ParallelFor.h"
 #include "Components/SplineComponent.h"
@@ -1432,6 +1433,29 @@ FPCGElementPtr UPCGSplineSamplerSettings::CreateElement() const
 	return MakeShared<FPCGSplineSamplerElement>();
 }
 
+void FPCGSplineSamplerElement::GetDependenciesCrc(const FPCGDataCollection& InInput, const UPCGSettings* InSettings, UPCGComponent* InComponent, FPCGCrc& OutCrc) const
+{
+	FPCGCrc Crc;
+	IPCGElement::GetDependenciesCrc(InInput, InSettings, InComponent, Crc);
+
+	if (const UPCGSplineSamplerSettings* Settings = Cast<UPCGSplineSamplerSettings>(InSettings))
+	{
+		bool bUnbounded;
+		PCGSettingsHelpers::GetOverrideValue(InInput, Settings, GET_MEMBER_NAME_CHECKED(FPCGSplineSamplerParams, bUnbounded), Settings->SamplerParams.bUnbounded, bUnbounded);
+
+		// If we're operating in bounded mode then we'll use actor bounds, and therefore take dependency on actor data.
+		if (!bUnbounded && InComponent)
+		{
+			if (const UPCGData* Data = InComponent->GetActorPCGData())
+			{
+				Crc.Combine(Data->GetOrComputeCrc(/*bFullDataCrc=*/false));
+			}
+		}
+	}
+
+	OutCrc = Crc;
+}
+
 bool FPCGSplineSamplerElement::ExecuteInternal(FPCGContext* Context) const
 {
 	TRACE_CPUPROFILER_EVENT_SCOPE(FPCGSplineSamplerElement::Execute);
diff --git a/Engine/Plugins/PCG/Source/PCG/Public/Elements/PCGSplineSampler.h b/Engine/Plugins/PCG/Source/PCG/Public/Elements/PCGSplineSampler.h
--- a/Engine/Plugins/PCG/Source/PCG/Public/Elements/PCGSplineSampler.h
+++ b/Engine/Plugins/PCG/Source/PCG/Public/Elements/PCGSplineSampler.h
@@ -254,6 +254,7 @@ public:
 class FPCGSplineSamplerElement : public IPCGElement
 {
 public:
+	virtual void GetDependenciesCrc(const FPCGDataCollection& InInput, const UPCGSettings* InSettings, UPCGComponent* InComponent, FPCGCrc& OutCrc) const override;
 	// Worth computing a full CRC in case we can halt change propagation/re-executions
 	virtual bool ShouldComputeFullOutputDataCrc(FPCGContext* Context) const override { return true; }
 
