[PCG] SplineSampler Distance mode increment clamp/rounding
* Hide irrelevant SplineSampler params properties depending on mode and dimension

diff --git a/Engine/Plugins/PCG/Source/PCG/Private/Elements/PCGSplineSampler.cpp b/Engine/Plugins/PCG/Source/PCG/Private/Elements/PCGSplineSampler.cpp
--- a/Engine/Plugins/PCG/Source/PCG/Private/Elements/PCGSplineSampler.cpp
+++ b/Engine/Plugins/PCG/Source/PCG/Private/Elements/PCGSplineSampler.cpp
@@ -384,6 +384,7 @@ namespace PCGSplineSamplerHelpers
 		FStepSampler(const UPCGPolyLineData* InLineData, const FPCGSplineSamplerParams& Params)
 			: LineData(InLineData)
 			, bComputeCurvature(Params.bComputeCurvature)
+			, bComputeTangents(Params.bComputeTangents)
 			, bComputeAlpha(Params.bComputeAlpha)
 			, bComputeDistance(Params.bComputeDistance)
 		{
@@ -398,6 +399,7 @@ namespace PCGSplineSamplerHelpers
 		int CurrentSegmentIndex = 0;
 		FVector::FReal DistanceToCurrentSegment = 0.0;
 		bool bComputeCurvature = false;
+		bool bComputeTangents = false;
 		bool bComputeAlpha = false;
 		bool bComputeDistance = false;
 	};
@@ -406,7 +408,6 @@ namespace PCGSplineSamplerHelpers
 	{
 		FSubdivisionStepSampler(const UPCGPolyLineData* InLineData, const FPCGSplineSamplerParams& Params)
 			: FStepSampler(InLineData, Params)
-			, bComputeTangents(Params.bComputeTangents)
 		{
 			NumSegments = LineData->GetNumSegments();
 			SubdivisionsPerSegment = Params.SubdivisionsPerSegment;
@@ -510,7 +511,6 @@ namespace PCGSplineSamplerHelpers
 		int NumSegments = 0;
 		int SubdivisionsPerSegment = 0;
 		int SubpointIndex = 0;
-		bool bComputeTangents = false;
 	};
 
 	struct FDistanceStepSampler : public FStepSampler
@@ -525,22 +525,29 @@ namespace PCGSplineSamplerHelpers
 			TotalDistance = InLineData->GetLength();
 			EndDistance = TotalDistance - EndOffset;
 
+			const FVector::FReal TotalLength = EndDistance - StartOffset;
+
 			if (Params.Mode == EPCGSplineSamplingMode::NumberOfSamples)
 			{
 				TotalNumSamples = Params.NumSamples;
+			}
+			else if (Params.Mode == EPCGSplineSamplingMode::Distance && Params.bFitToCurve)
+			{
+				// In Distance mode we can cover the full spline by finding the nearest whole number of samples that would fit, and treating the mode as NumberOfSamples instead.
+				TotalNumSamples = (Params.DistanceIncrement > 0) ? (TotalLength / Params.DistanceIncrement) : 0;
+			}
 
-				if (TotalNumSamples > 0)
-				{
-					// Compute an increment which evenly distributes sample points along the length of the curve.
-					DistanceIncrement = (LineData->GetLength() - StartOffset - EndOffset) / (LineData->IsClosed() ? TotalNumSamples : FMath::Max(1, TotalNumSamples - 1));
-				}
+			if (TotalNumSamples > 0)
+			{
+				// Compute an increment which evenly distributes sample points along the length of the curve.
+				DistanceIncrement = TotalLength / (LineData->IsClosed() ? TotalNumSamples : FMath::Max(1, TotalNumSamples - 1));
 			}
 			else
 			{
 				DistanceIncrement = Params.DistanceIncrement;
 			}
 
-			MaxRandomOffset = FMath::Max(0.0f, Params.MaxRandomOffsetNormalized) * DistanceIncrement / 2.0f;
+			MaxRandomOffset = FMath::Max(0.0, Params.MaxRandomOffsetNormalized) * DistanceIncrement / 2.0;
 			bUseRandomOffset = !FMath::IsNearlyZero(MaxRandomOffset);
 
 			if (bUseRandomOffset)
@@ -580,6 +587,13 @@ namespace PCGSplineSamplerHelpers
 				OutResult.Curvature = LineData->GetCurvatureAtDistance(CurrentSegmentIndex, OffsetDistance);
 			}
 
+			if (bComputeTangents)
+			{
+				const FVector Forward = OutTransform.GetRotation().GetForwardVector();
+				OutResult.ArriveTangent = Forward;
+				OutResult.LeaveTangent = Forward;
+			}
+
 			if (bComputeAlpha)
 			{
 				OutResult.Alpha = LineData->GetAlphaAtDistance(CurrentSegmentIndex, OffsetDistance);
@@ -670,7 +684,7 @@ namespace PCGSplineSamplerHelpers
 					bSetMetadata |= (SegmentIndexAttribute != nullptr);
 				}
 
-				if (Params.bComputeSubsegmentIndex)
+				if (Params.bComputeSubsegmentIndex && Params.Mode == EPCGSplineSamplingMode::Subdivision)
 				{
 					SubsegmentIndexAttribute = Metadata->FindOrCreateAttribute<int>(Params.SubsegmentIndexAttribute, static_cast<int>(DefaultValue));
 					bSetMetadata |= (SubsegmentIndexAttribute != nullptr);
diff --git a/Engine/Plugins/PCG/Source/PCG/Public/Elements/PCGSplineSampler.h b/Engine/Plugins/PCG/Source/PCG/Public/Elements/PCGSplineSampler.h
--- a/Engine/Plugins/PCG/Source/PCG/Public/Elements/PCGSplineSampler.h
+++ b/Engine/Plugins/PCG/Source/PCG/Public/Elements/PCGSplineSampler.h
@@ -60,100 +60,140 @@ struct FPCGSplineSamplerParams
 	UPROPERTY(BlueprintReadWrite, EditAnywhere, Category = Settings)
 	EPCGSplineSamplingDimension Dimension = EPCGSplineSamplingDimension::OnSpline;
 
-	UPROPERTY(BlueprintReadWrite, EditAnywhere, Category = Settings, meta = (EditCondition = "Dimension!=EPCGSplineSamplingDimension::OnInterior"))
+	UPROPERTY(BlueprintReadWrite, EditAnywhere, Category = Settings, meta = (EditCondition = "Dimension!=EPCGSplineSamplingDimension::OnInterior", EditConditionHides))
 	EPCGSplineSamplingMode Mode = EPCGSplineSamplingMode::Subdivision;
 
-	UPROPERTY(BlueprintReadWrite, EditAnywhere, Category = Settings, meta = (EditCondition = "Dimension!=EPCGSplineSamplingDimension::OnSpline&&Dimension!=EPCGSplineSamplingDimension::OnInterior"))
+	UPROPERTY(BlueprintReadWrite, EditAnywhere, Category = Settings, meta = (EditCondition = "Dimension!=EPCGSplineSamplingDimension::OnSpline&&Dimension!=EPCGSplineSamplingDimension::OnInterior", EditConditionHides))
 	EPCGSplineSamplingFill Fill = EPCGSplineSamplingFill::Fill;
 
-	UPROPERTY(BlueprintReadWrite, EditAnywhere, Category = Settings, meta = (ClampMin = "0", EditCondition = "Mode==EPCGSplineSamplingMode::Subdivision&&Dimension!=EPCGSplineSamplingDimension::OnInterior"))
+	UPROPERTY(BlueprintReadWrite, EditAnywhere, Category = Settings, meta = (ClampMin = "0", EditCondition = "Mode==EPCGSplineSamplingMode::Subdivision&&Dimension!=EPCGSplineSamplingDimension::OnInterior", EditConditionHides))
 	int32 SubdivisionsPerSegment = 1;
 
-	UPROPERTY(BlueprintReadWrite, EditAnywhere, Category = Settings, meta = (ClampMin = "0.1", EditCondition = "Mode==EPCGSplineSamplingMode::Distance&&Dimension!=EPCGSplineSamplingDimension::OnInterior"))
+	UPROPERTY(BlueprintReadWrite, EditAnywhere, Category = Settings, meta = (ClampMin = "0.1", EditCondition = "Mode==EPCGSplineSamplingMode::Distance&&Dimension!=EPCGSplineSamplingDimension::OnInterior", EditConditionHides))
 	float DistanceIncrement = 100.0f;
 
-	UPROPERTY(BlueprintReadWrite, EditAnywhere, Category = Settings, meta = (ClampMin = "0", EditCondition = "Mode==EPCGSplineSamplingMode::NumberOfSamples&&Dimension!=EPCGSplineSamplingDimension::OnInterior"))
+	UPROPERTY(BlueprintReadWrite, EditAnywhere, Category = Settings, meta = (ClampMin = "0", EditCondition = "Mode==EPCGSplineSamplingMode::NumberOfSamples&&Dimension!=EPCGSplineSamplingDimension::OnInterior", EditConditionHides))
 	int32 NumSamples = 8;
 
-	UPROPERTY(BlueprintReadWrite, EditAnywhere, Category = Settings, meta = (ClampMin = "0", EditCondition = "Dimension==EPCGSplineSamplingDimension::OnHorizontal||Dimension==EPCGSplineSamplingDimension::OnVolume"))
+	UPROPERTY(BlueprintReadWrite, EditAnywhere, Category = Settings, meta = (ClampMin = "0", EditCondition = "Dimension==EPCGSplineSamplingDimension::OnHorizontal||Dimension==EPCGSplineSamplingDimension::OnVolume", EditConditionHides))
 	int32 NumPlanarSubdivisions = 8;
 
-	UPROPERTY(BlueprintReadWrite, EditAnywhere, Category = Settings, meta = (ClampMin = "0", EditCondition = "Dimension==EPCGSplineSamplingDimension::OnVertical||Dimension==EPCGSplineSamplingDimension::OnVolume"))
+	UPROPERTY(BlueprintReadWrite, EditAnywhere, Category = Settings, meta = (ClampMin = "0", EditCondition = "Dimension==EPCGSplineSamplingDimension::OnVertical||Dimension==EPCGSplineSamplingDimension::OnVolume", EditConditionHides))
 	int32 NumHeightSubdivisions = 8;
 
 	/** Distance (in cm) along the spline at which sampling will begin. */
-	UPROPERTY(BlueprintReadWrite, EditAnywhere, Category = Settings, meta = (ClampMin = "0", EditCondition = "Mode!=EPCGSplineSamplingMode::Subdivision&&Dimension!=EPCGSplineSamplingDimension::OnInterior"))
+	UPROPERTY(BlueprintReadWrite, EditAnywhere, Category = Settings, meta = (ClampMin = "0", EditCondition = "Mode!=EPCGSplineSamplingMode::Subdivision&&Dimension!=EPCGSplineSamplingDimension::OnInterior", EditConditionHides))
 	float StartOffset = 0.0f;
 
 	/** Distance (in cm) from the end of the spline at which sampling will end. */
-	UPROPERTY(BlueprintReadWrite, EditAnywhere, Category = Settings, meta = (ClampMin = "0", EditCondition = "Mode!=EPCGSplineSamplingMode::Subdivision&&Dimension!=EPCGSplineSamplingDimension::OnInterior"))
+	UPROPERTY(BlueprintReadWrite, EditAnywhere, Category = Settings, meta = (ClampMin = "0", EditCondition = "Mode!=EPCGSplineSamplingMode::Subdivision&&Dimension!=EPCGSplineSamplingDimension::OnInterior", EditConditionHides))
 	float EndOffset = 0.0f;
 
 	/** Normalized value for the maximum possible offset for each sample point. 0.0 means no offset, and 1.0 means DistanceIncrement / 2.0. */
-	UPROPERTY(BlueprintReadWrite, EditAnywhere, Category = Settings, meta = (ClampMin = "0", EditCondition = "Mode!=EPCGSplineSamplingMode::Subdivision&&Dimension!=EPCGSplineSamplingDimension::OnInterior"))
+	UPROPERTY(BlueprintReadWrite, EditAnywhere, Category = Settings, meta = (ClampMin = "0", EditCondition = "Mode!=EPCGSplineSamplingMode::Subdivision&&Dimension!=EPCGSplineSamplingDimension::OnInterior", EditConditionHides))
 	float MaxRandomOffsetNormalized = 0.0f;
 
-	UPROPERTY(BlueprintReadWrite, EditAnywhere, Category = Settings, meta = (InlineEditConditionToggle))
+	/**
+	 * If the length of the spline does not divide evenly into the DistanceIncrement, the final sample point will not land on the end of the spline. Enable this to force
+	 * the DistanceIncrement to be rounded up to the nearest value which would yield evenly spaced samples across the entire length of the spline.
+	 */
+	UPROPERTY(BlueprintReadWrite, EditAnywhere, Category = Settings, meta = (EditCondition = "Mode==EPCGSplineSamplingMode::Distance&&Dimension!=EPCGSplineSamplingDimension::OnInterior", EditConditionHides))
+	bool bFitToCurve = false;
+	
+	/** The space between each sample point */
+	UPROPERTY(BlueprintReadWrite, EditAnywhere, Category = Settings, meta = (ClampMin = "0.1", EditCondition = "Dimension==EPCGSplineSamplingDimension::OnInterior", EditConditionHides))
+	float InteriorSampleSpacing = 100.0f;
+
+	/** The space between each sample point on the spline boundary. Used for computation; lower spacing is more expensive but more accurate. */
+	UPROPERTY(BlueprintReadWrite, EditAnywhere, Category = Settings, meta = (ClampMin = "0.1", EditCondition = "Dimension==EPCGSplineSamplingDimension::OnInterior&&!bTreatSplineAsPolyline", EditConditionHides))
+	float InteriorBorderSampleSpacing = 100.0f;
+
+	/** Use the spline points to form a polyline, instead of computing many sample points along the spline. This is more accurate if your spline is linear. */
+	UPROPERTY(BlueprintReadWrite, EditAnywhere, Category = Settings, meta = (EditCondition = "Dimension==EPCGSplineSamplingDimension::OnInterior", EditConditionHides))
+	bool bTreatSplineAsPolyline = false;
+
+	/** Determines the orientation of interior points. */
+	UPROPERTY(BlueprintReadWrite, EditAnywhere, Category = Settings, meta = (EditCondition = "Dimension==EPCGSplineSamplingDimension::OnInterior", EditConditionHides))
+	EPCGSplineSamplingInteriorOrientation InteriorOrientation = EPCGSplineSamplingInteriorOrientation::Uniform;
+
+	/** Project sample points onto one possible surface given by the spline boundary. */
+	UPROPERTY(BlueprintReadWrite, EditAnywhere, Category = Settings, meta = (EditCondition = "Dimension==EPCGSplineSamplingDimension::OnInterior", EditConditionHides))
+	bool bProjectOntoSurface = false;
+
+	// TODO: DirtyCache for OnDependencyChanged when this float curve is an external asset
+	/** Defines the density for each sample based on its distance from the spline. X axis is normalized distance to boundary (0-1), Y axis is density value. */
+	UPROPERTY(BlueprintReadWrite, EditAnywhere, Category = Settings, meta = (EditCondition = "Dimension==EPCGSplineSamplingDimension::OnInterior", EditConditionHides))
+	FRuntimeFloatCurve InteriorDensityFalloffCurve;
+
+	/** Compute the delta angle to the next point on the spline and write it to an attribute. */
+	UPROPERTY(BlueprintReadWrite, EditAnywhere, Category = "Settings|Attributes", meta = (EditCondition = "Dimension!=EPCGSplineSamplingDimension::OnInterior", EditConditionHides))
 	bool bComputeDirectionDelta = false;
 
 	/** Attribute that will contain the delta angle to the next point on the spline w.r.t to the current's point Up vector. */
-	UPROPERTY(BlueprintReadWrite, EditAnywhere, Category = Settings, meta = (EditCondition = "bComputeDirectionDelta"))
+	UPROPERTY(BlueprintReadWrite, EditAnywhere, Category = "Settings|Attributes", meta = (EditCondition = "bComputeDirectionDelta&&Dimension!=EPCGSplineSamplingDimension::OnInterior", EditConditionHides))
 	FName NextDirectionDeltaAttribute = TEXT("NextDirectionDelta");
 
-	UPROPERTY(BlueprintReadWrite, EditAnywhere, Category = Settings, meta = (InlineEditConditionToggle))
+	/** Compute curvature along the spline and write it to an attribute. */
+	UPROPERTY(BlueprintReadWrite, EditAnywhere, Category = "Settings|Attributes", meta = (EditCondition = "Dimension!=EPCGSplineSamplingDimension::OnInterior", EditConditionHides))
 	bool bComputeCurvature = false;
 
-	/** Attribute that will contain the curvature. Note that the radius of curvature is defined as 1/Curvature, and might need you to scale to world units */
-	UPROPERTY(BlueprintReadWrite, EditAnywhere, Category = Settings, meta = (EditCondition = "bComputeCurvature"))
+	/** Attribute that will contain the curvature. Note that the radius of curvature is defined as 1/Curvature, and might need you to scale to world units. */
+	UPROPERTY(BlueprintReadWrite, EditAnywhere, Category = "Settings|Attributes", meta = (EditCondition = "bComputeCurvature&&Dimension!=EPCGSplineSamplingDimension::OnInterior", EditConditionHides))
 	FName CurvatureAttribute = TEXT("Curvature");
 
-	UPROPERTY(BlueprintReadWrite, EditAnywhere, Category = Settings, meta = (InlineEditConditionToggle))
+	/** Compute the spline segment index and write it to an attribute. */
+	UPROPERTY(BlueprintReadWrite, EditAnywhere, Category = "Settings|Attributes", meta = (EditCondition = "Dimension!=EPCGSplineSamplingDimension::OnInterior", EditConditionHides))
 	bool bComputeSegmentIndex = false;
 
 	/** Attribute that will contain the spline segment index. */
-	UPROPERTY(BlueprintReadWrite, EditAnywhere, Category = Settings, meta = (EditCondition = "bComputeSegmentIndex"))
+	UPROPERTY(BlueprintReadWrite, EditAnywhere, Category = "Settings|Attributes", meta = (EditCondition = "bComputeSegmentIndex&&Dimension!=EPCGSplineSamplingDimension::OnInterior", EditConditionHides))
 	FName SegmentIndexAttribute = TEXT("SegmentIndex");
 
-	UPROPERTY(BlueprintReadWrite, EditAnywhere, Category = Settings, meta = (InlineEditConditionToggle))
+	/** Compute the sub-segment index of a point on the spline and write it to an attribute. */
+	UPROPERTY(BlueprintReadWrite, EditAnywhere, Category = "Settings|Attributes", meta = (EditCondition = "Mode==EPCGSplineSamplingMode::Subdivision&&Dimension!=EPCGSplineSamplingDimension::OnInterior", EditConditionHides))
 	bool bComputeSubsegmentIndex = false;
 
 	/** Attribute that will contain the sub-segment index of a point on the spline. When the sub-segment index is 0, the point is a control point on the actual spline. Only applies to Subdivision mode. */
-	UPROPERTY(BlueprintReadWrite, EditAnywhere, Category = Settings, meta = (EditCondition = "bComputeSubsegmentIndex"))
+	UPROPERTY(BlueprintReadWrite, EditAnywhere, Category = "Settings|Attributes", meta = (EditCondition = "bComputeSubsegmentIndex&&Mode==EPCGSplineSamplingMode::Subdivision&&Dimension!=EPCGSplineSamplingDimension::OnInterior", EditConditionHides))
 	FName SubsegmentIndexAttribute = TEXT("SubsegmentIndex");
 
-	UPROPERTY(BlueprintReadWrite, EditAnywhere, Category = Settings)
+	/** Compute arrive and leave tangents along the spline and write them to attributes. */
+	UPROPERTY(BlueprintReadWrite, EditAnywhere, Category = "Settings|Attributes", meta = (EditCondition = "Dimension!=EPCGSplineSamplingDimension::OnInterior", EditConditionHides))
 	bool bComputeTangents = false;
 
-	/** Attribute that will contain the arrive tangent vector. For control points, this will be the actual arrive tangent. For non-control points, this will only be the normalized tangent at this point. Only applies to Subdivision mode. */
-	UPROPERTY(BlueprintReadWrite, EditAnywhere, Category = Settings, meta = (EditCondition = "bComputeTangents"))
+	/** Attribute that will contain the arrive tangent vector. For control points, this will be the actual arrive tangent. For non-control points, this will only be the normalized tangent at this point. */
+	UPROPERTY(BlueprintReadWrite, EditAnywhere, Category = "Settings|Attributes", meta = (EditCondition = "bComputeTangents&&Dimension!=EPCGSplineSamplingDimension::OnInterior", EditConditionHides))
 	FName ArriveTangentAttribute = TEXT("ArriveTangent");
 
-	/** Attribute that will contain the leave tangent vector. For control points, this will be the actual leave tangent. For non-control points, this will only be the normalized tangent at this point. Only applies to Subdivision mode. */
-	UPROPERTY(BlueprintReadWrite, EditAnywhere, Category = Settings, meta = (EditCondition = "bComputeTangents"))
+	/** Attribute that will contain the leave tangent vector. For control points, this will be the actual leave tangent. For non-control points, this will only be the normalized tangent at this point. */
+	UPROPERTY(BlueprintReadWrite, EditAnywhere, Category = "Settings|Attributes", meta = (EditCondition = "bComputeTangents&&Dimension!=EPCGSplineSamplingDimension::OnInterior", EditConditionHides))
 	FName LeaveTangentAttribute = TEXT("LeaveTangent");
 
-	UPROPERTY(BlueprintReadWrite, EditAnywhere, Category = Settings, meta = (InlineEditConditionToggle))
+	/** Compute an alpha value along the spline and write it to an attribute. */
+	UPROPERTY(BlueprintReadWrite, EditAnywhere, Category = "Settings|Attributes", meta = (EditCondition = "Dimension!=EPCGSplineSamplingDimension::OnInterior", EditConditionHides))
 	bool bComputeAlpha = false;
 
 	/**
 	 * Attribute that will contain a value in [0,1] representing how far along the point is to the end of the line. Each segment on the line represents a same-size interval.
 	 * For example, if there are three segments, each segment will take up 0.333... of the interval.
 	 */
-	UPROPERTY(BlueprintReadWrite, EditAnywhere, Category = Settings, meta = (EditCondition = "bComputeAlpha"))
+	UPROPERTY(BlueprintReadWrite, EditAnywhere, Category = "Settings|Attributes", meta = (EditCondition = "bComputeAlpha&&Dimension!=EPCGSplineSamplingDimension::OnInterior", EditConditionHides))
 	FName AlphaAttribute = TEXT("Alpha");
 
-	UPROPERTY(BlueprintReadWrite, EditAnywhere, Category = Settings, meta = (InlineEditConditionToggle))
+	/** Compute distance along the spline and write it to an attribute. */
+	UPROPERTY(BlueprintReadWrite, EditAnywhere, Category = "Settings|Attributes", meta = (EditCondition = "Dimension!=EPCGSplineSamplingDimension::OnInterior", EditConditionHides))
 	bool bComputeDistance = false;
 
 	/** Attribute that will contain the distance along the spline at the sample point. */
-	UPROPERTY(BlueprintReadWrite, EditAnywhere, Category = Settings, meta = (EditCondition = "bComputeDistance"))
+	UPROPERTY(BlueprintReadWrite, EditAnywhere, Category = "Settings|Attributes", meta = (EditCondition = "bComputeDistance&&Dimension!=EPCGSplineSamplingDimension::OnInterior", EditConditionHides))
 	FName DistanceAttribute = TEXT("Distance");
 
-	UPROPERTY(BlueprintReadWrite, EditAnywhere, Category = Settings, meta = (InlineEditConditionToggle))
+	/** Compute an input key for each point along the spline and write the key to an attribute. */
+	UPROPERTY(BlueprintReadWrite, EditAnywhere, Category = "Settings|Attributes", meta = (EditCondition = "Dimension!=EPCGSplineSamplingDimension::OnInterior", EditConditionHides))
 	bool bComputeInputKey = false;
 
 	/** Attribute that will contain the spline input key, a float value between [0, N], where N is the number of control points. Each range [i, i+1] represents an interpolation from 0 to 1 across spline segment i. */
-	UPROPERTY(BlueprintReadWrite, EditAnywhere, Category = Settings, meta = (EditCondition = "bComputeInputKey"))
+	UPROPERTY(BlueprintReadWrite, EditAnywhere, Category = "Settings|Attributes", meta = (EditCondition = "bComputeInputKey&&Dimension!=EPCGSplineSamplingDimension::OnInterior", EditConditionHides))
 	FName InputKeyAttribute = TEXT("InputKey");
 
 	/** If no Bounding Shape input is provided, the actor bounds are used to limit the sample generation domain.
@@ -162,31 +202,6 @@ struct FPCGSplineSamplerParams
 	*/
 	UPROPERTY(BlueprintReadWrite, EditAnywhere, Category = Settings, meta = (PCG_Overridable))
 	bool bUnbounded = false; 
-	
-	/** The space between each sample point */
-	UPROPERTY(BlueprintReadWrite, EditAnywhere, Category = Settings, meta = (ClampMin = "0.1", EditCondition = "Dimension==EPCGSplineSamplingDimension::OnInterior"))
-	float InteriorSampleSpacing = 100.0f;
-
-	/** The space between each sample point on the spline boundary. Used for computation; lower spacing is more expensive but more accurate. */
-	UPROPERTY(BlueprintReadWrite, EditAnywhere, Category = Settings, meta = (ClampMin = "0.1", EditCondition = "Dimension==EPCGSplineSamplingDimension::OnInterior&&!bTreatSplineAsPolyline"))
-	float InteriorBorderSampleSpacing = 100.0f;
-
-	/** Use the spline points to form a polyline, instead of computing many sample points along the spline. This is more accurate if your spline is linear. */
-	UPROPERTY(BlueprintReadWrite, EditAnywhere, Category = Settings, meta = (EditCondition = "Dimension==EPCGSplineSamplingDimension::OnInterior"))
-	bool bTreatSplineAsPolyline = false;
-
-	/** Determines the orientation of interior points. */
-	UPROPERTY(BlueprintReadWrite, EditAnywhere, Category = Settings, meta = (EditCondition = "Dimension==EPCGSplineSamplingDimension::OnInterior"))
-	EPCGSplineSamplingInteriorOrientation InteriorOrientation = EPCGSplineSamplingInteriorOrientation::Uniform;
-
-	/** Project sample points onto one possible surface given by the spline boundary. */
-	UPROPERTY(EditAnywhere, Category = Settings, meta = (EditCondition = "Dimension==EPCGSplineSamplingDimension::OnInterior"))
-	bool bProjectOntoSurface = false;
-
-	// TODO: DirtyCache for OnDependencyChanged when this float curve is an external asset
-	/** Defines the density for each sample based on its distance from the spline. X axis is normalized distance to boundary (0-1), Y axis is density value. */
-	UPROPERTY(EditAnywhere, Category = Settings, meta = (EditCondition = "Dimension==EPCGSplineSamplingDimension::OnInterior"))
-	FRuntimeFloatCurve InteriorDensityFalloffCurve;
 
 	/** Each PCG point represents a discretized, volumetric region of world space. The points' Steepness value [0.0 to
 	 * 1.0] establishes how "hard" or "soft" that volume will be represented. From 0, it will ramp up linearly
