[PCG] Cleared up component behavior when created from BP if they are meant to be managed by PCG.
Added QoL methods to simplify managed resource creation from BP.

diff --git a/Engine/Plugins/PCG/Source/PCG/Private/PCGComponent.cpp b/Engine/Plugins/PCG/Source/PCG/Private/PCGComponent.cpp
--- a/Engine/Plugins/PCG/Source/PCG/Private/PCGComponent.cpp
+++ b/Engine/Plugins/PCG/Source/PCG/Private/PCGComponent.cpp
@@ -170,6 +170,58 @@ void UPCGComponent::AddToManagedResources(UPCGManagedResource* InResource)
 	}
 }
 
+void UPCGComponent::AddComponentsToManagedResources(const TArray<UActorComponent*>& InComponents)
+{
+	for (UActorComponent* Component : InComponents)
+	{
+		if (!Component)
+		{
+			continue;
+		}
+
+		if (!Component->ComponentHasTag(PCGHelpers::DefaultPCGTag))
+		{
+			Component->Modify();
+			Component->ComponentTags.Add(PCGHelpers::DefaultPCGTag);
+		}
+
+		UPCGManagedComponent* ManagedResource = NewObject<UPCGManagedComponent>(this);
+		// Implementation note: we call the setter to make sure that if this is done from BP, the construction method is properly updated
+		ManagedResource->SetGeneratedComponentFromBP(Component);
+
+		AddToManagedResources(ManagedResource);
+	}
+}
+
+void UPCGComponent::AddActorsToManagedResources(const TArray<AActor*>& InActors)
+{
+	TSet<TSoftObjectPtr<AActor>> ValidActors;
+
+	for (AActor* Actor : InActors)
+	{
+		if (Actor)
+		{
+			if (!Actor->Tags.Contains(PCGHelpers::DefaultPCGActorTag))
+			{
+				Actor->Modify();
+				Actor->Tags.Add(PCGHelpers::DefaultPCGActorTag);
+			}
+
+			ValidActors.Add(Actor);
+		}
+	}
+
+	if (ValidActors.IsEmpty())
+	{
+		return;
+	}
+
+	UPCGManagedActors* ManagedResource = NewObject<UPCGManagedActors>(this);
+	ManagedResource->GeneratedActors = ValidActors;
+
+	AddToManagedResources(ManagedResource);
+}
+
 void UPCGComponent::ForEachManagedResource(TFunctionRef<void(UPCGManagedResource*)> Func)
 {
 	FScopeLock ResourcesLock(&GeneratedResourcesLock);
diff --git a/Engine/Plugins/PCG/Source/PCG/Private/PCGManagedResource.cpp b/Engine/Plugins/PCG/Source/PCG/Private/PCGManagedResource.cpp
--- a/Engine/Plugins/PCG/Source/PCG/Private/PCGManagedResource.cpp
+++ b/Engine/Plugins/PCG/Source/PCG/Private/PCGManagedResource.cpp
@@ -461,6 +461,27 @@ void UPCGManagedComponent::MarkAsReused()
 	}
 }
 
+void UPCGManagedComponent::SetGeneratedComponentFromBP(TSoftObjectPtr<UActorComponent> InGeneratedComponent)
+{
+	GeneratedComponent = InGeneratedComponent;
+
+	// Components that are created from blueprint are automatically tagged as "created by construction script",
+	// regardless of whether that is true. This makes sure that the flags on the component are correct and considered an instance component
+	// and will then be properly serialized and managed by PCG.
+	if (UActorComponent* Component = GeneratedComponent.Get())
+	{
+		if (AActor* ComponentOwner = Component->GetOwner())
+		{
+			if (Component->CreationMethod == EComponentCreationMethod::UserConstructionScript)
+			{
+				ComponentOwner->RemoveOwnedComponent(Component);
+				Component->CreationMethod = EComponentCreationMethod::Instance;
+				ComponentOwner->AddOwnedComponent(Component);
+			}
+		}
+	}
+}
+
 #if WITH_EDITOR
 void UPCGManagedComponent::ChangeTransientState(EPCGEditorDirtyMode NewEditingMode)
 {
diff --git a/Engine/Plugins/PCG/Source/PCG/Public/PCGComponent.h b/Engine/Plugins/PCG/Source/PCG/Public/PCGComponent.h
--- a/Engine/Plugins/PCG/Source/PCG/Public/PCGComponent.h
+++ b/Engine/Plugins/PCG/Source/PCG/Public/PCGComponent.h
@@ -140,6 +140,15 @@ public:
 	UFUNCTION(BlueprintCallable, Category = PCG)
 	void AddToManagedResources(UPCGManagedResource* InResource);
 
+	/** Creates a managed component resource and adds it to the current component. Note: in native code, consider using the explicit creation especially if there are special resource objects involved. */
+	UFUNCTION(BlueprintCallable, Category = PCG)
+	void AddComponentsToManagedResources(const TArray<UActorComponent*>& InComponents);
+
+	/** Creates a managed actors resource and adds it to the current component. Note: in native code, consider using the explicit creation especially if there are special resource objects involved. */
+	UFUNCTION(BlueprintCallable, Category = PCG)
+	void AddActorsToManagedResources(const TArray<AActor*>& InActors);
+
+public:
 	void ForEachManagedResource(TFunctionRef<void(UPCGManagedResource*)> InFunction);
 
 	/** Transactionable methods to be called from details UI */
diff --git a/Engine/Plugins/PCG/Source/PCG/Public/PCGManagedResource.h b/Engine/Plugins/PCG/Source/PCG/Public/PCGManagedResource.h
--- a/Engine/Plugins/PCG/Source/PCG/Public/PCGManagedResource.h
+++ b/Engine/Plugins/PCG/Source/PCG/Public/PCGManagedResource.h
@@ -107,6 +107,8 @@ class PCG_API UPCGManagedComponent : public UPCGManagedResource
 {
 	GENERATED_BODY()
 
+	friend class UPCGComponent;
+
 public:
 	//~Begin UObject interface
 	virtual void PostEditImport() override;
@@ -131,8 +133,13 @@ public:
 	virtual void ForgetComponent() { GeneratedComponent.Reset(); }
 
 public:
-	UPROPERTY(BlueprintReadWrite, EditAnywhere, Category = GeneratedData)
+	UPROPERTY(BlueprintReadWrite, EditAnywhere, Category = GeneratedData, BlueprintSetter = SetGeneratedComponentFromBP)
 	TSoftObjectPtr<UActorComponent> GeneratedComponent;
+
+private:
+	// When creating components from BP they will be tagged automatically as created from construction script, which makes them transient and isn't compatible with the PCG workflow.
+	UFUNCTION(BlueprintSetter, meta = (BlueprintInternalUseOnly = "true"))
+	void SetGeneratedComponentFromBP(TSoftObjectPtr<UActorComponent> InGeneratedComponent);
 };
 
 UCLASS(BlueprintType)
