[PCG] Split node resets metadata entries on the output points

diff --git a/Engine/Plugins/PCG/Source/PCG/Private/Elements/PCGSplitPoints.cpp b/Engine/Plugins/PCG/Source/PCG/Private/Elements/PCGSplitPoints.cpp
--- a/Engine/Plugins/PCG/Source/PCG/Private/Elements/PCGSplitPoints.cpp
+++ b/Engine/Plugins/PCG/Source/PCG/Private/Elements/PCGSplitPoints.cpp
@@ -96,12 +96,10 @@ bool FPCGSplitPointsElement::ExecuteInternal(FPCGContext* Context) const
 			}
 
 			// Execution for PointsA portion
-			UPCGMetadataAccessorHelpers::CopyPoint(InPoint, PointA, false);
 			PointA.BoundsMax = PointA.BoundsMax + SplitValues * (MinPlusSplit - PointA.BoundsMax);
 			PointsA[WriteIndex] = PointA;
 
 			// Execution of the PointsB portion
-			UPCGMetadataAccessorHelpers::CopyPoint(InPoint, PointB, false);
 			PointB.BoundsMin = PointB.BoundsMin + SplitValues * (MinPlusSplit - PointB.BoundsMin);
 			PointsB[WriteIndex] = PointB;
 		};
