[PCG] Minor code & ux fixes

diff --git a/Engine/Plugins/PCG/Source/PCG/Private/Elements/PCGPointFromMeshElement.cpp b/Engine/Plugins/PCG/Source/PCG/Private/Elements/PCGPointFromMeshElement.cpp
--- a/Engine/Plugins/PCG/Source/PCG/Private/Elements/PCGPointFromMeshElement.cpp
+++ b/Engine/Plugins/PCG/Source/PCG/Private/Elements/PCGPointFromMeshElement.cpp
@@ -32,7 +32,7 @@ FPCGElementPtr UPCGPointFromMeshSettings::CreateElement() const
 
 bool FPCGPointFromMeshElement::PrepareDataInternal(FPCGContext* Context) const
 {
-	TRACE_CPUPROFILER_EVENT_SCOPE(FPCGPointFromMeshElement::Execute);
+	TRACE_CPUPROFILER_EVENT_SCOPE(FPCGPointFromMeshElement::PrepareData);
 
 	check(Context);
 
diff --git a/Engine/Plugins/PCG/Source/PCG/Private/Metadata/PCGMetadataPartitionCommon.cpp b/Engine/Plugins/PCG/Source/PCG/Private/Metadata/PCGMetadataPartitionCommon.cpp
--- a/Engine/Plugins/PCG/Source/PCG/Private/Metadata/PCGMetadataPartitionCommon.cpp
+++ b/Engine/Plugins/PCG/Source/PCG/Private/Metadata/PCGMetadataPartitionCommon.cpp
@@ -292,6 +292,14 @@ namespace PCGMetadataPartitionCommon
 		}
 	}
 
+	/**
+	* Dispatch the partition according to the data and selector.
+	*/
+	TArray<TArray<int32>> AttributeGenericPartition(const UPCGData* InData, const FPCGAttributePropertySelector& InSelector, FPCGContext* InOptionalContext, bool bSilenceMissingAttributeErrors)
+	{
+		return AttributeGenericPartition<TArray<int32>>(InData, InSelector, InOptionalContext, bSilenceMissingAttributeErrors);
+	}
+
 	/**
 	 * Partition on multiple attributes by first partitioning on the attributes independently. Then take the resultant
 	 * partition and convert them to a BitArray representation of each element's partition. Once in BitArray form,
diff --git a/Engine/Plugins/PCG/Source/PCG/Public/Elements/PCGMatchAndSetAttributes.h b/Engine/Plugins/PCG/Source/PCG/Public/Elements/PCGMatchAndSetAttributes.h
--- a/Engine/Plugins/PCG/Source/PCG/Public/Elements/PCGMatchAndSetAttributes.h
+++ b/Engine/Plugins/PCG/Source/PCG/Public/Elements/PCGMatchAndSetAttributes.h
@@ -96,7 +96,7 @@ public:
 	FPCGAttributePropertyInputSelector MaxDistanceInputAttribute;
 
 	/** Controls whether we will use the attribute provided in the Input Weight Attribute to perform entry selection. */
-	UPROPERTY(BlueprintReadWrite, EditAnywhere, Category = Settings, meta = (PCG_Overridable))
+	UPROPERTY(BlueprintReadWrite, EditAnywhere, Category = Settings, meta = (PCG_Overridable, InlineEditConditionToggle))
 	bool bUseInputWeightAttribute = false;
 
 	/** Input weight from the points, assumed to be in the [0, 1] range. */
