[PCG] Minor: Move Delete and Landscape menus into sub-menus

#rnx

diff --git a/Engine/Plugins/PCG/Source/PCGEditor/Private/PCGEditorModule.cpp b/Engine/Plugins/PCG/Source/PCGEditor/Private/PCGEditorModule.cpp
--- a/Engine/Plugins/PCG/Source/PCGEditor/Private/PCGEditorModule.cpp
+++ b/Engine/Plugins/PCG/Source/PCGEditor/Private/PCGEditorModule.cpp
@@ -218,102 +218,114 @@ void FPCGEditorModule::UnregisterMenuExtensions()
 
 void FPCGEditorModule::PopulateMenuActions(FMenuBuilder& MenuBuilder)
 {
-	MenuBuilder.AddMenuEntry(
-		LOCTEXT("DeletePCGPartitionActors", "Delete all PCG partition actors"),
-		LOCTEXT("DeletePCGPartitionActors_Tooltip", "Deletes all serialized PCG partition actors in the current world"),
-		FSlateIcon(),
-		FUIAction(
-			FExecuteAction::CreateLambda([]() {
-				if (GEditor)
+	MenuBuilder.AddSubMenu(LOCTEXT("PCGSubMenuDelete", "Delete"), FText(), FNewMenuDelegate::CreateLambda([](FMenuBuilder& SubMenuBuilder)
+	{
+		SubMenuBuilder.AddMenuEntry(
+			LOCTEXT("DeletePCGPartitionActors", "All PCG Partition Grid Actors & Generated Actors"),
+			LOCTEXT("DeletePCGPartitionActors_Tooltip", "Deletes all PCG Partition Grid Actors and PCG Partition Generated Actors in the current world"),
+			FSlateIcon(),
+			FUIAction(
+				FExecuteAction::CreateLambda([]() 
 				{
-					if (UPCGSubsystem* PCGSubsystem = UPCGSubsystem::GetInstance(GEditor->GetEditorWorldContext().World()))
+					if (GEditor)
 					{
-						PCGSubsystem->DeleteSerializedPartitionActors(/*bOnlyDeleteUnused=*/false);
+						if (UPCGSubsystem* PCGSubsystem = UPCGSubsystem::GetInstance(GEditor->GetEditorWorldContext().World()))
+						{
+							PCGSubsystem->DeleteSerializedPartitionActors(/*bOnlyDeleteUnused=*/false);
+						}
 					}
-				}
-			})),
-		NAME_None);	
-	
-	MenuBuilder.AddMenuEntry(
-		LOCTEXT("DeletePCGPartitionActorsChildren", "Delete all PCG partition actors children"),
-		LOCTEXT("DeletePCGPartitionActorsChildren_Tooltip", "Deletes all serialized PCG partition actors children in the current world, but not the Partition Actors themselves"),
-		FSlateIcon(),
-		FUIAction(
-			FExecuteAction::CreateLambda([]() {
-				if (GEditor)
+				})),
+			NAME_None);
+
+		SubMenuBuilder.AddMenuEntry(
+			LOCTEXT("DeletePCGPartitionActorsChildren", "All PCG Partition Generated Actors"),
+			LOCTEXT("DeletePCGPartitionActorsChildren_Tooltip", "Deletes all PCG Partition Generated Actors in the current world (not the Partition Grid Actors themselves)"),
+			FSlateIcon(),
+			FUIAction(
+				FExecuteAction::CreateLambda([]() 
 				{
-					if (UPCGSubsystem* PCGSubsystem = UPCGSubsystem::GetInstance(GEditor->GetEditorWorldContext().World()))
+					if (GEditor)
 					{
-						PCGSubsystem->DeleteSerializedPartitionActors(/*bOnlyDeleteUnused=*/false, /*bOnlyChildren=*/true);
+						if (UPCGSubsystem* PCGSubsystem = UPCGSubsystem::GetInstance(GEditor->GetEditorWorldContext().World()))
+						{
+							PCGSubsystem->DeleteSerializedPartitionActors(/*bOnlyDeleteUnused=*/false, /*bOnlyChildren=*/true);
+						}
 					}
-				}
-			})),
-		NAME_None);
+				})),
+			NAME_None);
 
-	MenuBuilder.AddMenuEntry(
-		LOCTEXT("DeleteUnusedPCGPartitionActors", "Delete all unused PCG partition actors"),
-		LOCTEXT("DeleteUnusedPCGPartitionActors_Tooltip", "Deletes all serialized PCG partition actors in the current world that doesn't intersect with any PCG Component."),
-		FSlateIcon(),
-		FUIAction(
-			FExecuteAction::CreateLambda([]() {
-				if (GEditor)
-				{
-					if (UPCGSubsystem* PCGSubsystem = UPCGSubsystem::GetInstance(GEditor->GetEditorWorldContext().World()))
+		SubMenuBuilder.AddMenuEntry(
+			LOCTEXT("DeletePCGWorldActor", "All PCG World Actors"),
+			LOCTEXT("DeletePCGWorldActor_Tooltip", "Deletes all PCG World Actors (This also deletes all PCG Partition Grid Actors & Generated Actors)"),
+			FSlateIcon(),
+			FUIAction(
+				FExecuteAction::CreateLambda([]()
 					{
-						PCGSubsystem->DeleteSerializedPartitionActors(/*bOnlyDeleteUnused=*/true);
-					}
-				}
-			})),
-		NAME_None);
+						if (GEditor)
+						{
+							if (UPCGSubsystem* PCGSubsystem = UPCGSubsystem::GetInstance(GEditor->GetEditorWorldContext().World()))
+							{
+								PCGSubsystem->DestroyAllPCGWorldActors();
+							}
+						}
+					})),
+			NAME_None);
 
-	MenuBuilder.AddMenuEntry(
-		LOCTEXT("DeletePCGWorldActor", "Deletes all PCG World Actors"),
-		LOCTEXT("DeletePCGWorldActor_Tooltip", "Deletes all PCG World Actors"),
-		FSlateIcon(),
-		FUIAction(
-			FExecuteAction::CreateLambda([]() {
-				if (GEditor)
+		SubMenuBuilder.AddMenuEntry(
+			LOCTEXT("DeleteUnusedPCGPartitionActors", "Unused PCG Partition Grid Actors"),
+			LOCTEXT("DeleteUnusedPCGPartitionActors_Tooltip", "Deletes unused PCG Partition Grid Actors in the current world"),
+			FSlateIcon(),
+			FUIAction(
+				FExecuteAction::CreateLambda([]() 
 				{
-					if (UPCGSubsystem* PCGSubsystem = UPCGSubsystem::GetInstance(GEditor->GetEditorWorldContext().World()))
+					if (GEditor)
 					{
-						PCGSubsystem->DestroyAllPCGWorldActors();
+						if (UPCGSubsystem* PCGSubsystem = UPCGSubsystem::GetInstance(GEditor->GetEditorWorldContext().World()))
+						{
+							PCGSubsystem->DeleteSerializedPartitionActors(/*bOnlyDeleteUnused=*/true);
+						}
 					}
-				}
-			})),
-		NAME_None);
+				})),
+			NAME_None);
+	}));
 
-	MenuBuilder.AddMenuEntry(
-		LOCTEXT("BuildLandscapeCache", "Build Landscape Data Cache"),
-		LOCTEXT("BuildLandscapeCache_Tooltip", "Caches the landscape data in the PCG World Actor"),
-		FSlateIcon(),
-		FUIAction(
-			FExecuteAction::CreateLambda([]() {
-				if (GEditor)
+	MenuBuilder.AddSubMenu(LOCTEXT("PCGSubMenuLandscape", "Landscape"), FText(), FNewMenuDelegate::CreateLambda([](FMenuBuilder& SubMenuBuilder)
+	{
+		SubMenuBuilder.AddMenuEntry(
+			LOCTEXT("BuildLandscapeCache", "Build Cache"),
+			LOCTEXT("BuildLandscapeCache_Tooltip", "Caches the landscape data in the PCG World Actor"),
+			FSlateIcon(),
+			FUIAction(
+				FExecuteAction::CreateLambda([]() 
 				{
-					if (UPCGSubsystem* PCGSubsystem = UPCGSubsystem::GetInstance(GEditor->GetEditorWorldContext().World()))
+					if (GEditor)
 					{
-						PCGSubsystem->BuildLandscapeCache();
+						if (UPCGSubsystem* PCGSubsystem = UPCGSubsystem::GetInstance(GEditor->GetEditorWorldContext().World()))
+						{
+							PCGSubsystem->BuildLandscapeCache();
+						}
 					}
-				}
-			})),
-		NAME_None);
+				})),
+			NAME_None);
 
-	MenuBuilder.AddMenuEntry(
-		LOCTEXT("ClearLandscapeCache", "Clear Landscape Data Cache"),
-		LOCTEXT("ClearLandscapeCache_Tooltip", "Clears the landscape data cache in the PCG World Actor"),
-		FSlateIcon(),
-		FUIAction(
-			FExecuteAction::CreateLambda([]() {
-				if (GEditor)
+		SubMenuBuilder.AddMenuEntry(
+			LOCTEXT("ClearLandscapeCache", "Clear Cache"),
+			LOCTEXT("ClearLandscapeCache_Tooltip", "Clears the landscape data cache in the PCG World Actor"),
+			FSlateIcon(),
+			FUIAction(
+				FExecuteAction::CreateLambda([]() 
 				{
-					if (UPCGSubsystem* PCGSubsystem = UPCGSubsystem::GetInstance(GEditor->GetEditorWorldContext().World()))
+					if (GEditor)
 					{
-						PCGSubsystem->ClearLandscapeCache();
+						if (UPCGSubsystem* PCGSubsystem = UPCGSubsystem::GetInstance(GEditor->GetEditorWorldContext().World()))
+						{
+							PCGSubsystem->ClearLandscapeCache();
+						}
 					}
-				}
 				})),
-		NAME_None);
-
+			NAME_None);
+	}));
+	
 	MenuBuilder.AddMenuEntry(
 		LOCTEXT("CancelAllGeneration", "Cancel all PCG tasks"),
 		LOCTEXT("CancelAllGeneration_Tooltip", "Cancels all PCG tasks running"),
