[PCG] Fixed a bug with the Surface Sampler taking full bounds and a bug causing World Ray Hit Query not to cull penetrating hits

#rnx

diff --git a/Engine/Plugins/PCG/Source/PCG/Private/Data/PCGWorldData.cpp b/Engine/Plugins/PCG/Source/PCG/Private/Data/PCGWorldData.cpp
--- a/Engine/Plugins/PCG/Source/PCG/Private/Data/PCGWorldData.cpp
+++ b/Engine/Plugins/PCG/Source/PCG/Private/Data/PCGWorldData.cpp
@@ -373,7 +373,7 @@ bool UPCGWorldRayHitData::SamplePoint(const FTransform& InTransform, const FBox&
 		if (QueryParams.bIgnoreBackfaceHits)
 		{
 			// If its a landscape, we cull if the normal is negative in Z direction (landscape normal is always the +Z axis). If not, then we cull if the impact normal and the ray are headed in the same direction
-			if ((bHitOnLandscape && Hit.ImpactNormal.Z < 0) || (QueryParams.RayDirection).Dot(Hit.ImpactNormal) > 0)
+			if (Hit.bStartPenetrating || (bHitOnLandscape && Hit.ImpactNormal.Z < 0) || (QueryParams.RayDirection).Dot(Hit.ImpactNormal) > 0)
 			{
 				continue;
 			}
diff --git a/Engine/Plugins/PCG/Source/PCG/Private/Elements/PCGSurfaceSampler.cpp b/Engine/Plugins/PCG/Source/PCG/Private/Elements/PCGSurfaceSampler.cpp
--- a/Engine/Plugins/PCG/Source/PCG/Private/Elements/PCGSurfaceSampler.cpp
+++ b/Engine/Plugins/PCG/Source/PCG/Private/Elements/PCGSurfaceSampler.cpp
@@ -466,12 +466,8 @@ bool FPCGSurfaceSamplerElement::PrepareDataInternal(FPCGContext* InContext) cons
 			// If local bounds exists, apply the transform without rotation to get into grid sampling space
 			if (EffectiveGridBounds.IsValid)
 			{
-				// Remove rotation as we don't want to shear an AABB.
-				FTransform NoRotTransform(GeneratingShape->GetTransform());
-				NoRotTransform.SetRotation(FQuat::Identity);
-
 				// Transport the box into semi-local 2D space, where we can get the deterministic grid samples
-				EffectiveGridBounds = EffectiveGridBounds.TransformBy(NoRotTransform);
+				EffectiveGridBounds = PCGHelpers::OverlapBounds(EffectiveGridBounds.TransformBy(GeneratingShape->GetTransform()), BoundingShapeBounds);
 			}
 			else // If no local bounds, try to use the generating shape's bounds
 			{
