[PCG] Collision-backed point data wrapper, useful for differences, etc.

diff --git a/Engine/Plugins/PCG/Source/PCG/Private/Data/PCGCollisionWrapperData.cpp b/Engine/Plugins/PCG/Source/PCG/Private/Data/PCGCollisionWrapperData.cpp
new file mode 100644
--- /dev/null
+++ b/Engine/Plugins/PCG/Source/PCG/Private/Data/PCGCollisionWrapperData.cpp
@@ -0,0 +1,269 @@
+// Copyright Epic Games, Inc. All Rights Reserved.
+
+#include "Data/PCGCollisionWrapperData.h"
+
+#include "Data/PCGPointData.h"
+#include "Elements/Metadata/PCGMetadataElementCommon.h"
+
+#include "Components/StaticMeshComponent.h"
+#include "Engine/StaticMesh.h"
+#include "PhysicsEngine/BodyInstance.h"
+
+#include UE_INLINE_GENERATED_CPP_BY_NAME(PCGCollisionWrapperData)
+
+FPCGCollisionWrapper::~FPCGCollisionWrapper()
+{
+	Uninitialize();
+}
+
+void FPCGCollisionWrapper::Uninitialize()
+{
+	// Implementation note: we do the full uninitialize even if we think we're not initialized due to the default move operators
+	for (FBodyInstance* BodyInstance : BodyInstances)
+	{
+		if (BodyInstance->IsValidBodyInstance())
+		{
+			BodyInstance->TermBody();
+		}
+
+		delete BodyInstance;
+	}
+
+	BodyInstances.Reset();
+	IndexToBodyInstance.Reset();
+
+	bInitialized = false;
+}
+
+bool FPCGCollisionWrapper::Prepare(const IPCGAttributeAccessor* Accessor, const IPCGAttributeAccessorKeys* Keys, TArray<FSoftObjectPath>& MeshPathsToLoad)
+{
+	TRACE_CPUPROFILER_EVENT_SCOPE(FPCGCollisionWrapper::Prepare);
+	MeshPathsToLoad.Reset();
+
+	if (!ensure(!bInitialized))
+	{
+		return true;
+	}
+
+	IndexToBodyInstance.Reset();
+
+	if (!Accessor || !Keys)
+	{
+		return false;
+	}
+
+	auto GatherMeshPathsAndCreateMap = [this, &MeshPathsToLoad](const TArrayView<FSoftObjectPath>& Meshes, int Start, int Range)
+	{
+		for (const FSoftObjectPath& Mesh : Meshes)
+		{
+			IndexToBodyInstance.Add(Mesh.IsNull() ? INDEX_NONE : MeshPathsToLoad.AddUnique(Mesh));
+		}
+	};
+
+	// Implementation note: if this fails, then the mesh paths to load will also be empty, which will behave like we're expecting it to.
+	PCGMetadataElementCommon::ApplyOnAccessorRange<FSoftObjectPath>(*Keys, *Accessor, GatherMeshPathsAndCreateMap, EPCGAttributeAccessorFlags::AllowBroadcast | EPCGAttributeAccessorFlags::AllowConstructible);
+	return true;
+}
+
+void FPCGCollisionWrapper::CreateBodyInstances(const TArray<FSoftObjectPath>& MeshPaths)
+{
+	TRACE_CPUPROFILER_EVENT_SCOPE(FPCGCollisionWrapper::CreateBodyInstances);
+	if (!ensure(!bInitialized))
+	{
+		return;
+	}
+
+	check(BodyInstances.IsEmpty());
+	BodyInstances.Reserve(MeshPaths.Num());
+
+	for (const FSoftObjectPath& MeshPath : MeshPaths)
+	{
+		FBodyInstance* BodyInstance = nullptr;
+
+		TSoftObjectPtr<UStaticMesh> Mesh(MeshPath);
+		UStaticMesh* MeshPtr = nullptr;
+	
+		{
+			//TRACE_CPUPROFILER_EVENT_SCOPE(FPCGCollisionWrapper::CreateBodyInstances::LoadMesh);
+			MeshPtr = Mesh.LoadSynchronous();
+		}
+
+		if(MeshPtr)
+		{
+			//TRACE_CPUPROFILER_EVENT_SCOPE(FPCGCollisionWrapper::CreateBodyInstances::CreateBodyInstance);
+			BodyInstance = new FBodyInstance();
+			BodyInstance->bAutoWeld = false;
+			BodyInstance->bSimulatePhysics = false;
+			BodyInstance->InitBody(MeshPtr->GetBodySetup(), FTransform::Identity, nullptr, nullptr);
+		}
+
+		BodyInstances.Add(BodyInstance);
+	}
+
+	bInitialized = true;
+}
+
+bool FPCGCollisionWrapper::Initialize(const IPCGAttributeAccessor* Accessor, const IPCGAttributeAccessorKeys* Keys)
+{
+	if (!ensure(!bInitialized))
+	{
+		return true;
+	}
+
+	TArray<FSoftObjectPath> MeshPathsToLoad;
+	if (!Prepare(Accessor, Keys, MeshPathsToLoad))
+	{
+		return false;
+	}
+
+	CreateBodyInstances(MeshPathsToLoad);
+	return true;
+}
+
+FBodyInstance* FPCGCollisionWrapper::GetBodyInstance(int32 EntryIndex) const
+{
+	return ((bInitialized && IndexToBodyInstance[EntryIndex] != INDEX_NONE) ? BodyInstances[IndexToBodyInstance[EntryIndex]] : nullptr);
+}
+
+bool UPCGCollisionWrapperData::Initialize(const UPCGPointData* InPointData, const FPCGAttributePropertyInputSelector& InCollisionSelector, bool bInUseComplexCollision)
+{
+	TArray<FSoftObjectPath> MeshesToLoad;
+	if (PreInitializeAndGatherMeshesEx(InPointData, InCollisionSelector, bInUseComplexCollision, MeshesToLoad))
+	{
+		FinalizeInitializationEx(MeshesToLoad);
+		return true;
+	}
+	else
+	{
+		return false;
+	}
+}
+
+bool UPCGCollisionWrapperData::PreInitializeAndGatherMeshesEx(const UPCGPointData* InPointData, const FPCGAttributePropertyInputSelector& InCollisionSelector, bool bInUseComplexCollision, TArray<FSoftObjectPath>& OutMeshesToLoad)
+{
+	CollisionWrapper.Uninitialize();
+
+	check(InPointData);
+	PointData = InPointData;
+#if WITH_EDITOR
+	RawPointData = InPointData;
+#endif
+
+	CollisionSelector = InCollisionSelector;
+	bUseComplexCollision = bInUseComplexCollision;
+
+	// Go through the point data on the attribute we'll use for the shapes
+	// StaticMesh -> see what we're doing in UInstanceStaticMeshComponent::InitInstanceBody, or UStaticMeshComponent::UpdateCollisionFromStaticMesh
+	FPCGAttributePropertyInputSelector InputSelector;
+	InputSelector = InCollisionSelector.CopyAndFixLast(InPointData);
+
+	TUniquePtr<const IPCGAttributeAccessor> InputAccessor = PCGAttributeAccessorHelpers::CreateConstAccessor(InPointData, InputSelector);
+	TUniquePtr<const IPCGAttributeAccessorKeys> InputKeys = PCGAttributeAccessorHelpers::CreateConstKeys(InPointData, InputSelector);
+
+	if (!InputAccessor.IsValid() || !InputKeys.IsValid())
+	{
+		return false;
+	}
+
+	return CollisionWrapper.Prepare(InputAccessor.Get(), InputKeys.Get(), OutMeshesToLoad);
+}
+
+void UPCGCollisionWrapperData::FinalizeInitializationEx(const TArray<FSoftObjectPath>& InMeshPaths)
+{
+	CollisionWrapper.CreateBodyInstances(InMeshPaths);
+}
+
+void UPCGCollisionWrapperData::AddToCrc(FArchiveCrc32& Ar, bool bFullDataCrc) const
+{
+	Super::AddToCrc(Ar, bFullDataCrc);
+	GetPointData()->AddToCrc(Ar, bFullDataCrc);
+}
+
+void UPCGCollisionWrapperData::GetResourceSizeEx(FResourceSizeEx& CumulativeResourceSize)
+{
+	const_cast<UPCGPointData*>(GetPointData())->GetResourceSizeEx(CumulativeResourceSize);
+}
+
+FBox UPCGCollisionWrapperData::GetBounds() const
+{
+	return GetPointData()->GetBounds();
+}
+
+FBox UPCGCollisionWrapperData::GetStrictBounds() const
+{
+	return GetPointData()->GetStrictBounds();
+}
+
+const UPCGPointData* UPCGCollisionWrapperData::ToPointData(FPCGContext* Context, const FBox& InBounds) const
+{
+	return GetPointData()->ToPointData(Context, InBounds);
+}
+
+bool UPCGCollisionWrapperData::SamplePoint(const FTransform& InTransform, const FBox& InBounds, FPCGPoint& OutPoint, UPCGMetadata* OutMetadata) const
+{
+	if(!CollisionWrapper.bInitialized)
+	{
+		return GetPointData()->SamplePoint(InTransform, InBounds, OutPoint, OutMetadata);
+	}
+
+	// Find all points matching this point in the point data octree
+	// For all touching points, check against the actual physics shape we have for that point
+	// E.g. get the point collision attribute value, match to our internal list, get the shape (might require loading), then test against the matching shapes
+	// note that we need to place the stuff in the right referential too, since our shapes will all be at the origin
+	const UPCGPointData::PointOctree& Octree = GetPointData()->GetOctree();
+
+	float Density = 0;
+	FBox TransformedBounds = InBounds.TransformBy(InTransform);
+
+	Octree.FindElementsWithBoundsTest(FBoxCenterAndExtent(TransformedBounds.GetCenter(), TransformedBounds.GetExtent()), [this, &InTransform, &InBounds, &Density](const FPCGPointRef& InPointRef)
+	{
+		if (InPointRef.Point->Density < Density)
+		{
+			return;
+		}
+
+		bool bHasOverlap = true;
+
+		// Compute point index & find body instance
+		// Test against body instance when present
+		if (FBodyInstance* BodyInstance = CollisionWrapper.GetBodyInstance(InPointRef.Point - GetPointData()->GetPoints().GetData()))
+		{
+			// Compute collision shape rotation and transform, knowing that the body instance is at the origin.
+			// e.g. we need to take the InTransform and apply the inverse point transform to it.
+			FTransform RelativeTransform = InTransform.GetRelativeTransform(InPointRef.Point->Transform);
+			FTransform GeomTransform(RelativeTransform.GetRotation(), RelativeTransform.GetLocation());
+
+			FCollisionShape CollisionShape;
+			CollisionShape.SetBox(FVector3f(InBounds.GetExtent() * RelativeTransform.GetScale3D()));
+
+			bHasOverlap = FPhysicsInterface::Overlap_Geom(BodyInstance, CollisionShape, GeomTransform.GetRotation(), GeomTransform, nullptr, bUseComplexCollision);
+		}
+
+		if (bHasOverlap)
+		{
+			Density = InPointRef.Point->Density;
+		}
+	});
+
+	if (Density > 0)
+	{
+		OutPoint = FPCGPoint();
+		OutPoint.Transform = InTransform;
+		OutPoint.SetLocalBounds(InBounds);
+		OutPoint.Density = Density;
+
+		return true;
+	}
+	else
+	{
+		return false;
+	}
+}
+
+UPCGSpatialData* UPCGCollisionWrapperData::CopyInternal() const
+{
+	UPCGCollisionWrapperData* NewCollisionWrapperData = NewObject<UPCGCollisionWrapperData>();
+	NewCollisionWrapperData->Initialize(PointData, CollisionSelector, bUseComplexCollision);
+
+	return NewCollisionWrapperData;
+}
diff --git a/Engine/Plugins/PCG/Source/PCG/Private/Elements/PCGCreateCollisionData.cpp b/Engine/Plugins/PCG/Source/PCG/Private/Elements/PCGCreateCollisionData.cpp
new file mode 100644
--- /dev/null
+++ b/Engine/Plugins/PCG/Source/PCG/Private/Elements/PCGCreateCollisionData.cpp
@@ -0,0 +1,143 @@
+// Copyright Epic Games, Inc. All Rights Reserved.
+
+#include "Elements/PCGCreateCollisionData.h"
+
+#include "PCGContext.h"
+#include "Data/PCGCollisionWrapperData.h"
+#include "Data/PCGPointData.h"
+
+#include UE_INLINE_GENERATED_CPP_BY_NAME(PCGCreateCollisionData)
+
+#define LOCTEXT_NAMESPACE "PCGCreateCollisionDataElement"
+
+#if WITH_EDITOR
+FText UPCGCreateCollisionDataSettings::GetNodeTooltipText() const
+{
+	return LOCTEXT("CreateCollisionTooltip", "Creates a volumetric representation of the points as if they had their selected mesh collision.");
+}
+#endif // WITH_EDITOR
+
+TArray<FPCGPinProperties> UPCGCreateCollisionDataSettings::InputPinProperties() const
+{
+	return Super::DefaultPointInputPinProperties();
+}
+
+TArray<FPCGPinProperties> UPCGCreateCollisionDataSettings::OutputPinProperties() const
+{
+	TArray<FPCGPinProperties> PinProperties;
+	PinProperties.Emplace(PCGPinConstants::DefaultOutputLabel, EPCGDataType::Primitive);
+
+	return PinProperties;
+}
+
+FPCGElementPtr UPCGCreateCollisionDataSettings::CreateElement() const
+{
+	return MakeShared<FPCGCreateCollisionDataElement>();
+}
+
+FPCGCreateCollisionContext::~FPCGCreateCollisionContext()
+{
+	// If the execution is cancelled between the prepare data and the execution,
+	// we might have data that's been rooted temporarily - take care of this here.
+	for (InputMeshData& ThisInputData : PerInputData)
+	{
+		if (ThisInputData.Data && ThisInputData.Data->IsRooted())
+		{
+			ThisInputData.Data->RemoveFromRoot();
+		}
+	}
+}
+
+bool FPCGCreateCollisionDataElement::PrepareDataInternal(FPCGContext* InContext) const
+{
+	TRACE_CPUPROFILER_EVENT_SCOPE(FPCGCreateCollisionDataElement::PrepareData);
+	check(InContext);
+	FPCGCreateCollisionContext* Context = static_cast<FPCGCreateCollisionContext*>(InContext);
+
+	if (!Context->WasLoadRequested())
+	{
+		TArray<FSoftObjectPath> MeshesToLoad;
+
+		const UPCGCreateCollisionDataSettings* Settings = Context->GetInputSettings<UPCGCreateCollisionDataSettings>();
+		check(Settings);
+
+		const TArray<FPCGTaggedData> Inputs = Context->InputData.GetInputsByPin(PCGPinConstants::DefaultInputLabel);
+		for(int InputIndex = 0; InputIndex < Inputs.Num(); ++InputIndex)
+		{
+			const FPCGTaggedData& Input = Inputs[InputIndex];
+
+			if (const UPCGPointData* PointData = Cast<UPCGPointData>(Input.Data))
+			{
+				FPCGCreateCollisionContext::InputMeshData InputMeshData;
+				InputMeshData.InputIndex = InputIndex;
+				InputMeshData.Data = NewObject<UPCGCollisionWrapperData>();
+
+				if (!InputMeshData.Data->PreInitializeAndGatherMeshesEx(PointData, Settings->CollisionAttribute, Settings->bUseComplexCollision, InputMeshData.MeshPaths))
+				{
+					if (Settings->bWarnIfAttributeCouldNotBeUsed)
+					{
+						PCGLog::LogWarningOnGraph(LOCTEXT("FailedToAccessMeshes", "Failed to access the mesh attribute provided - will behave like a PointData."), Context);
+					}
+				}
+				else
+				{
+					// Root object so it can't get deleted until we finish initialize and properly manage it as part of the collection.
+					InputMeshData.Data->AddToRoot();
+
+					for (const FSoftObjectPath& MeshPath : InputMeshData.MeshPaths)
+					{
+						if (!MeshPath.IsNull())
+						{
+							MeshesToLoad.AddUnique(MeshPath);
+						}
+					}
+					
+					Context->PerInputData.Add(MoveTemp(InputMeshData));
+				}
+			}
+		}
+
+		// Finally, request loading for all meshes we gathered
+		return Context->RequestResourceLoad(Context, MoveTemp(MeshesToLoad), !Settings->bSynchronousLoad);
+	}
+
+	return true;
+}
+
+bool FPCGCreateCollisionDataElement::ExecuteInternal(FPCGContext* InContext) const
+{
+	TRACE_CPUPROFILER_EVENT_SCOPE(FPCGCreateCollisionDataElement::Execute);
+
+	check(InContext);
+	FPCGCreateCollisionContext* Context = static_cast<FPCGCreateCollisionContext*>(InContext);
+
+	const UPCGCreateCollisionDataSettings* Settings = Context->GetInputSettings<UPCGCreateCollisionDataSettings>();
+	check(Settings);
+
+	TArray<FPCGTaggedData>& Outputs = Context->OutputData.TaggedData;
+	const TArray<FPCGTaggedData> Inputs = Context->InputData.GetInputsByPin(PCGPinConstants::DefaultInputLabel);
+
+	for (int InputIndex = 0; InputIndex < Inputs.Num(); ++InputIndex)
+	{
+		const FPCGTaggedData& Input = Inputs[InputIndex];
+		FPCGTaggedData& Output = Outputs.Add_GetRef(Input);
+
+		if (FPCGCreateCollisionContext::InputMeshData* MatchingData = Context->PerInputData.FindByPredicate([InputIndex](const FPCGCreateCollisionContext::InputMeshData& Data) { return Data.InputIndex == InputIndex; }))
+		{
+			check(MatchingData->Data);
+			MatchingData->Data->FinalizeInitializationEx(MatchingData->MeshPaths);
+
+			// Unroot collision wrapper data object since it'll be tracked through the collection
+			MatchingData->Data->RemoveFromRoot();
+
+			Output.Data = MatchingData->Data;
+			
+			// Null out the data so the context destructor doesn't interact with the rooting
+			MatchingData->Data = nullptr;
+		}
+	}
+
+	return true;
+}
+
+#undef LOCTEXT_NAMESPACE
\ No newline at end of file
diff --git a/Engine/Plugins/PCG/Source/PCG/Public/Data/PCGCollisionWrapperData.h b/Engine/Plugins/PCG/Source/PCG/Public/Data/PCGCollisionWrapperData.h
new file mode 100644
--- /dev/null
+++ b/Engine/Plugins/PCG/Source/PCG/Public/Data/PCGCollisionWrapperData.h
@@ -0,0 +1,89 @@
+// Copyright Epic Games, Inc. All Rights Reserved.
+
+#pragma once
+
+#include "PCGSpatialData.h"
+#include "Metadata/PCGAttributePropertySelector.h"
+#include "Metadata/Accessors/PCGAttributeAccessorHelpers.h"
+
+#include "PCGCollisionWrapperData.generated.h"
+
+struct FBodyInstance;
+
+struct PCG_API FPCGCollisionWrapper
+{
+	FPCGCollisionWrapper() = default;
+	~FPCGCollisionWrapper();
+	FPCGCollisionWrapper(const FPCGCollisionWrapper&) = delete;
+	FPCGCollisionWrapper(FPCGCollisionWrapper&& Other) = default;
+	FPCGCollisionWrapper& operator=(const FPCGCollisionWrapper&) = delete;
+	FPCGCollisionWrapper& operator=(FPCGCollisionWrapper&& Other) = default;
+
+	// Simple API - does both the prepare & create body instances in a single step
+	bool Initialize(const IPCGAttributeAccessor* Accessor, const IPCGAttributeAccessorKeys* Keys);
+	void Uninitialize();
+
+	// Advanced API - allows to do async loading as we separate the mesh finding part from the body creation part
+	bool Prepare(const IPCGAttributeAccessor* Accessor, const IPCGAttributeAccessorKeys* Keys, TArray<FSoftObjectPath>& MeshPathsToLoad);
+	void CreateBodyInstances(const TArray<FSoftObjectPath>& MeshPaths);
+	
+	// Retrieves the body instance associated to the entry given by its index
+	FBodyInstance* GetBodyInstance(int32 EntryIndex) const;
+
+	TArray<FBodyInstance*> BodyInstances;
+	TArray<int32> IndexToBodyInstance;
+	bool bInitialized = false;
+};
+
+UCLASS(MinimalAPI, BlueprintType, ClassGroup = (Procedural))
+class UPCGCollisionWrapperData : public UPCGSpatialData
+{
+	GENERATED_BODY()
+
+public:
+	/** Inititializes the collision wrapper on a point data based on the provided attribute selector */
+	PCG_API bool Initialize(const UPCGPointData* InPointData, const FPCGAttributePropertyInputSelector& InCollisionSelector, bool bInUseComplexCollision);
+
+	/** Advanced API for async loading */
+	bool PreInitializeAndGatherMeshesEx(const UPCGPointData* InPointData, const FPCGAttributePropertyInputSelector& InCollisionSelector, bool bInUseComplexCollision, TArray<FSoftObjectPath>& OutMeshesToLoad);
+	void FinalizeInitializationEx(const TArray<FSoftObjectPath>& InMeshPaths);
+
+	// ~Begin UPCGData Interface
+	virtual EPCGDataType GetDataType() const override { return EPCGDataType::Primitive; }
+	virtual void AddToCrc(FArchiveCrc32& Ar, bool bFullDataCrc) const override;
+	virtual void GetResourceSizeEx(FResourceSizeEx& CumulativeResourceSize) override;
+	// ~End UPCGData Interface
+
+	// ~Begin UPCGSpatialData interface
+	virtual int GetDimension() const override { return 3; }
+	virtual FBox GetBounds() const override;
+	virtual FBox GetStrictBounds() const override;
+	virtual const UPCGPointData* ToPointData(FPCGContext* Context, const FBox& InBounds = FBox(EForceInit::ForceInit)) const override;
+	virtual bool SamplePoint(const FTransform& Transform, const FBox& Bounds, FPCGPoint& OutPoint, UPCGMetadata* OutMetadata) const override;
+protected:
+	virtual UPCGSpatialData* CopyInternal() const override;
+	// ~End UPCGSpatialData interface
+
+	// For performance reasons, we keep a raw pointer to the point data in editor.
+#if WITH_EDITOR
+	inline const UPCGPointData* GetPointData() const { return RawPointData; }
+#else
+	inline const UPCGPointData* GetPointData() const { return PointData.Get(); }
+#endif
+
+private:
+	UPROPERTY()
+	TObjectPtr<const UPCGPointData> PointData;
+
+	UPROPERTY()
+	FPCGAttributePropertyInputSelector CollisionSelector;
+
+	UPROPERTY()
+	bool bUseComplexCollision = false;
+
+	FPCGCollisionWrapper CollisionWrapper;
+
+#if WITH_EDITOR
+	const UPCGPointData* RawPointData = nullptr;
+#endif
+};
\ No newline at end of file
diff --git a/Engine/Plugins/PCG/Source/PCG/Public/Elements/PCGCreateCollisionData.h b/Engine/Plugins/PCG/Source/PCG/Public/Elements/PCGCreateCollisionData.h
new file mode 100644
--- /dev/null
+++ b/Engine/Plugins/PCG/Source/PCG/Public/Elements/PCGCreateCollisionData.h
@@ -0,0 +1,70 @@
+// Copyright Epic Games, Inc. All Rights Reserved.
+
+#pragma once
+
+#include "PCGContext.h"
+#include "PCGSettings.h"
+#include "Async/PCGAsyncLoadingContext.h"
+#include "Data/PCGCollisionWrapperData.h"
+
+#include "PCGCreateCollisionData.generated.h"
+
+UCLASS(MinimalAPI, BlueprintType, ClassGroup = (Procedural))
+class UPCGCreateCollisionDataSettings : public UPCGSettings
+{
+	GENERATED_BODY()
+
+public:
+	//~Begin UPCGSettings interface
+#if WITH_EDITOR
+	virtual FName GetDefaultNodeName() const override { return FName(TEXT("CreateCollisionData")); }
+	virtual FText GetDefaultNodeTitle() const override { return NSLOCTEXT("PCGCreateCollisionDataSettings", "NodeTitle", "Create Collision Data"); }
+	virtual FText GetNodeTooltipText() const override;
+	virtual EPCGSettingsType GetType() const override { return EPCGSettingsType::Spatial; }
+#endif // WITH_EDITOR
+
+protected:
+	virtual TArray<FPCGPinProperties> InputPinProperties() const override;
+	virtual TArray<FPCGPinProperties> OutputPinProperties() const override;
+	virtual FPCGElementPtr CreateElement() const override;
+	//~End UPCGSettings interface
+
+public:
+	UPROPERTY(BlueprintReadWrite, EditAnywhere, Category = Settings, meta = (PCG_Overridable, PCG_DiscardPropertySelection, PCG_DiscardExtraSelection))
+	FPCGAttributePropertyInputSelector CollisionAttribute;
+
+	/** Queries against complex collision if enabled, performance warning */
+	UPROPERTY(BlueprintReadWrite, EditAnywhere, Category = Settings, meta = (PCG_Overridable))
+	bool bUseComplexCollision = false;
+
+	UPROPERTY(BlueprintReadWrite, EditAnywhere, Category = "Settings|Advanced")
+	bool bWarnIfAttributeCouldNotBeUsed = true;
+
+	UPROPERTY(BlueprintReadWrite, EditAnywhere, Category = "Settings|Debug")
+	bool bSynchronousLoad = false;
+};
+
+struct FPCGCreateCollisionContext : public FPCGContext, public IPCGAsyncLoadingContext
+{
+	~FPCGCreateCollisionContext();
+
+	struct InputMeshData
+	{
+		int InputIndex = INDEX_NONE;
+		TArray<FSoftObjectPath> MeshPaths;
+		UPCGCollisionWrapperData* Data = nullptr;
+	};
+
+	TArray<InputMeshData> PerInputData;
+};
+
+class FPCGCreateCollisionDataElement : public IPCGElementWithCustomContext<FPCGCreateCollisionContext>
+{
+public:
+	// Loading needs to be done on the main thread and accessing objects outside of PCG might not be thread safe, so taking the safe approach
+	virtual bool CanExecuteOnlyOnMainThread(FPCGContext* Context) const override { return true; }
+
+protected:
+	virtual bool PrepareDataInternal(FPCGContext* Context) const override;
+	virtual bool ExecuteInternal(FPCGContext* Context) const override;
+};
\ No newline at end of file
