[PCG] Delete Partition Actors now also deletes unloaded actors

#rnx

diff --git a/Engine/Plugins/PCG/Source/PCG/Private/PCGComponent.cpp b/Engine/Plugins/PCG/Source/PCG/Private/PCGComponent.cpp
--- a/Engine/Plugins/PCG/Source/PCG/Private/PCGComponent.cpp
+++ b/Engine/Plugins/PCG/Source/PCG/Private/PCGComponent.cpp
@@ -3504,6 +3504,24 @@ bool UPCGComponent::GetStackContext(FPCGStackContext& OutStackContext) const
 
 	return false;
 }
+
+TArray<TSoftObjectPtr<AActor>> UPCGComponent::GetManagedActorPaths(AActor* InActor)
+{
+	TSet<TSoftObjectPtr<AActor>> ManagedActorPaths;
+	InActor->ForEachComponent<UPCGComponent>(/*bIncludeFromChildActors=*/true, [&ManagedActorPaths](UPCGComponent* Component)
+	{
+		for (UPCGManagedResource* ManagedResource : Component->GeneratedResources)
+		{
+			if (UPCGManagedActors* ManagedActors = Cast<UPCGManagedActors>(ManagedResource))
+			{
+				ManagedActorPaths.Append(ManagedActors->GeneratedActors);
+			}
+		}
+	});
+
+	return ManagedActorPaths.Array();
+}
+
 #endif // WITH_EDITOR
 
 FPCGComponentInstanceData::FPCGComponentInstanceData(const UPCGComponent* InSourceComponent)
diff --git a/Engine/Plugins/PCG/Source/PCG/Private/PCGSubsystem.cpp b/Engine/Plugins/PCG/Source/PCG/Private/PCGSubsystem.cpp
--- a/Engine/Plugins/PCG/Source/PCG/Private/PCGSubsystem.cpp
+++ b/Engine/Plugins/PCG/Source/PCG/Private/PCGSubsystem.cpp
@@ -25,6 +25,9 @@
 #include "Editor.h"
 #include "PackageSourceControlHelper.h"
 #include "ObjectTools.h"
+#include "Misc/ScopedSlowTask.h"
+#include "WorldPartition/WorldPartition.h"
+#include "WorldPartition/WorldPartitionHelpers.h"
 #else
 #include "Engine/Engine.h"
 #include "Engine/World.h"
@@ -1306,25 +1309,92 @@ void UPCGSubsystem::DeleteSerializedPartitionActors(bool bOnlyDeleteUnused, bool
 		return;
 	}
 
-	auto GatherAndDestroyActors = [this, &PackagesToCleanup, World, bOnlyDeleteUnused, bOnlyChildren](AActor* Actor) -> bool
+	UWorldPartition* WorldPartition = World->GetWorldPartition();
+	TMap<FGuid, TArray<FGuid>> Attachments;
+
+	auto GetAttachments = [WorldPartition](const FGuid& ParentActor, const TMap<FGuid, TArray<FGuid>>& Attachments, auto GetAttachmentsRecursive, TArray<FGuid>& OutAttachedActors) -> void
+	{
+		if (const TArray<FGuid>* Attached = Attachments.Find(ParentActor))
+		{
+			OutAttachedActors.Append(*Attached);
+			for (const FGuid& AttachedGuid : *Attached)
+			{
+				GetAttachmentsRecursive(AttachedGuid, Attachments, GetAttachmentsRecursive, OutAttachedActors);
+			}
+		}
+	};
+
+	auto GatherAndDestroyActors = [this, &PackagesToCleanup, World, WorldPartition, bOnlyDeleteUnused, bOnlyChildren, &GetAttachments, &Attachments](AActor* Actor) -> bool
 	{
 		TObjectPtr<APCGPartitionActor> PartitionActor = CastChecked<APCGPartitionActor>(Actor);
 
 		// Do not delete RuntimeGen PAs or PAs with graph instances if we are only deleting unused PAs.
 		if (!PartitionActor->IsRuntimeGenerated() && (!bOnlyDeleteUnused || !PartitionActor->HasGraphInstances()))
 		{
-			// Gather all child actors
+			TArray<AActor*> ActorsToDelete;
+
+			TArray<FWorldPartitionReference> ActorReferences;
+
+			// Load Generated Resources to delete them
+			TArray<TSoftObjectPtr<AActor>> ManagedActors = UPCGComponent::GetManagedActorPaths(PartitionActor);
+			for (const TSoftObjectPtr<AActor>& ManagedActorPath : ManagedActors)
+			{
+				// Test to see if actor is loaded first to support non World Partition worlds
+				AActor* ManagedActor = ManagedActorPath.Get();
+				if (!ManagedActor && WorldPartition)
+				{
+					if (const FWorldPartitionActorDescInstance* ActorDescInstance = WorldPartition->GetActorDescInstanceByPath(ManagedActorPath.ToSoftObjectPath()))
+					{
+						FWorldPartitionReference& ActorReference = ActorReferences.Emplace_GetRef(FWorldPartitionReference(ActorDescInstance->GetContainerInstance(), ActorDescInstance->GetGuid()));
+						ManagedActor = ActorReference.GetActor();
+					}
+				}
+
+				if (ManagedActor)
+				{
+					ActorsToDelete.Add(ManagedActor);
+				}
+			}
+			
+			// Load Attachments before getting them in the next code block, since loading an actor doesn't load it's attachments (the reference is child to parent)
+			if (WorldPartition)
+			{
+				TArray<FGuid> AttachedActors;
+				GetAttachments(Actor->GetActorGuid(), Attachments, GetAttachments, AttachedActors);
+				for (const FGuid& AttachedActor : AttachedActors)
+				{
+					ActorReferences.Add(FWorldPartitionReference(WorldPartition, AttachedActor));
+				}
+			}
+
+			// We might have actors that weren't saved as managed resources that are attached and have the proper tag
 			TArray<AActor*> AttachedActors;
 			PartitionActor->GetAttachedActors(AttachedActors, /*bResetArray=*/true, /*bRecursivelyIncludeAttachedActors=*/ true);
+			for (AActor* AttachedActor : AttachedActors)
+			{
+				if (AttachedActor->ActorHasTag(PCGHelpers::DefaultPCGActorTag))
+				{
+					if (!ActorsToDelete.Contains(AttachedActor))
+					{
+						ActorsToDelete.Add(AttachedActor);
+					}
+				}
+				else if (WorldPartition && !bOnlyChildren)
+				{
+					// If actor isn't getting deleted but is an attached actor and it's Partition Actor parent 
+					// will get deleted then Pin the actor so it stays loaded and modified for the user to save
+					WorldPartition->PinActors({ AttachedActor->GetActorGuid() });
+				}
+			}
 
 			if (!bOnlyChildren)
 			{
-				AttachedActors.Add(PartitionActor);
+				ActorsToDelete.Add(PartitionActor);
 
 				PCGWorldActor->RemoveSerializedPartitionActorRecord({ PartitionActor->PCGGuid, PartitionActor->GetPCGGridSize(), PartitionActor->GetGridCoord() });
 			}
 
-			for (AActor* ActorToDelete : AttachedActors)
+			for (AActor* ActorToDelete : ActorsToDelete)
 			{
 				if (UPackage* ExternalPackage = ActorToDelete->GetExternalPackage())
 				{
@@ -1344,16 +1414,50 @@ void UPCGSubsystem::DeleteSerializedPartitionActors(bool bOnlyDeleteUnused, bool
 		GEditor->SelectNone(true, true, false);
 	}
 
-	// Attempt destroy on all PAs in the level.
-	if (ULevel* Level = World->GetCurrentLevel())
 	{
-		TRACE_CPUPROFILER_EVENT_SCOPE(UPCGSubsystem::DeleteSerializedPartitionActors::ForEachActorInLevel);
-		UPCGActorHelpers::ForEachActorInLevel<APCGPartitionActor>(Level, GatherAndDestroyActors);
-	}
+		FScopedSlowTask DeleteTask(0, NSLOCTEXT("PCGSubsystem", "DeletePartitionActors", "Deleting PCG Actors..."));
+		DeleteTask.MakeDialog();
 
-	if (PackagesToCleanup.Num() > 0)
-	{
-		ObjectTools::CleanupAfterSuccessfulDelete(PackagesToCleanup.Array(), /*bPerformanceReferenceCheck=*/true);
+		FWorldPartitionHelpers::FForEachActorWithLoadingResult LoadingResult;
+		{
+			TRACE_CPUPROFILER_EVENT_SCOPE(UPCGSubsystem::DeleteSerializedPartitionActors::ForEachActorInLevel);
+			if (WorldPartition)
+			{
+				// Gather Attach Parent information
+				FWorldPartitionHelpers::ForEachActorDescInstance(WorldPartition, [&Attachments](const FWorldPartitionActorDescInstance* ActorDescInstance)
+				{
+					if (ActorDescInstance->GetParentActor().IsValid())
+					{
+						Attachments.FindOrAdd(ActorDescInstance->GetParentActor()).Add(ActorDescInstance->GetGuid());
+					}
+
+					return true;
+				});
+
+				FWorldPartitionHelpers::FForEachActorWithLoadingParams ForEachActorWithLoadingParams;
+				ForEachActorWithLoadingParams.bKeepReferences = true;
+				ForEachActorWithLoadingParams.ActorClasses = { APCGPartitionActor::StaticClass() };
+
+				FWorldPartitionHelpers::ForEachActorWithLoading(WorldPartition, [&](const FWorldPartitionActorDescInstance* ActorDescInstance)
+				{
+					if(AActor* Actor = ActorDescInstance->GetActor())
+					{
+						GatherAndDestroyActors(Actor);
+					}
+					return true;
+				},
+				ForEachActorWithLoadingParams);
+			}
+			else
+			{
+				UPCGActorHelpers::ForEachActorInLevel<APCGPartitionActor>(World->PersistentLevel, GatherAndDestroyActors);
+			}
+		}
+
+		if (PackagesToCleanup.Num() > 0)
+		{
+			ObjectTools::CleanupAfterSuccessfulDelete(PackagesToCleanup.Array(), /*bPerformanceReferenceCheck=*/true);
+		}
 	}
 }
 
diff --git a/Engine/Plugins/PCG/Source/PCG/Public/PCGComponent.h b/Engine/Plugins/PCG/Source/PCG/Public/PCGComponent.h
--- a/Engine/Plugins/PCG/Source/PCG/Public/PCGComponent.h
+++ b/Engine/Plugins/PCG/Source/PCG/Public/PCGComponent.h
@@ -532,6 +532,8 @@ private:
 
 	void MarkResourcesAsTransientOnLoad();
 	bool DeletePreviewResources();
+
+	static TArray<TSoftObjectPtr<AActor>> GetManagedActorPaths(AActor* InActor);
 #endif
 
 	FBox GetGridBounds(const AActor* InActor) const;
