[PCG] SplineSampler SeedingModes
* Add SeedFromIndex to get stable seeds when points change position

diff --git a/Engine/Plugins/PCG/Source/PCG/Private/Elements/PCGSplineSampler.cpp b/Engine/Plugins/PCG/Source/PCG/Private/Elements/PCGSplineSampler.cpp
--- a/Engine/Plugins/PCG/Source/PCG/Private/Elements/PCGSplineSampler.cpp
+++ b/Engine/Plugins/PCG/Source/PCG/Private/Elements/PCGSplineSampler.cpp
@@ -341,6 +341,7 @@ namespace PCGSplineSamplerHelpers
 		FTransform LocalTransform;
 		FBox Box = FBox::BuildAABB(FVector::ZeroVector, FVector::OneVector);
 		FVector::FReal Curvature = 0;
+		int SampleIndex = 0;
 		int SegmentIndex = 0;
 		int SubsegmentIndex = 0;
 		float InputKey = 0.0f;
@@ -352,37 +353,45 @@ namespace PCGSplineSamplerHelpers
 		FVector::FReal Distance = 0.0f;
 	};
 
-	void SetSeed(FPCGPoint& Point, const FVector& LSPosition, const FPCGSplineSamplerParams& Params)
+	void SetSeed(FPCGPoint& Point, const FVector& LSPosition, const FPCGSplineSamplerParams& Params, int ExternalSeed, int SampleIndex)
 	{
-		if (!Params.bSeedFromLocalPosition)
+		if (Params.SeedingMode == EPCGSplineSamplingSeedingMode::SeedFromPosition)
 		{
-			if (!Params.bSeedFrom2DPosition)
+			if (!Params.bSeedFromLocalPosition)
 			{
-				Point.Seed = UPCGBlueprintHelpers::ComputeSeedFromPosition(Point.Transform.GetLocation());
+				if (!Params.bSeedFrom2DPosition)
+				{
+					Point.Seed = UPCGBlueprintHelpers::ComputeSeedFromPosition(Point.Transform.GetLocation());
+				}
+				else
+				{
+					const FVector WSPosition = Point.Transform.GetLocation();
+					Point.Seed = PCGHelpers::ComputeSeed((int)WSPosition.X, (int)WSPosition.Y);
+				}
 			}
-			else
+			else // Use provided local position
 			{
-				const FVector WSPosition = Point.Transform.GetLocation();
-				Point.Seed = PCGHelpers::ComputeSeed((int)WSPosition.X, (int)WSPosition.Y);
+				if (!Params.bSeedFrom2DPosition)
+				{
+					Point.Seed = UPCGBlueprintHelpers::ComputeSeedFromPosition(LSPosition);
+				}
+				else
+				{
+					Point.Seed = PCGHelpers::ComputeSeed((int)LSPosition.X, (int)LSPosition.Y);
+				}
 			}
 		}
-		else // Use provided local position
+		else
 		{
-			if (!Params.bSeedFrom2DPosition)
-			{
-				Point.Seed = UPCGBlueprintHelpers::ComputeSeedFromPosition(LSPosition);
-			}
-			else
-			{
-				Point.Seed = PCGHelpers::ComputeSeed((int)LSPosition.X, (int)LSPosition.Y);
-			}
+			Point.Seed = PCGHelpers::ComputeSeed(ExternalSeed, SampleIndex);
 		}
 	}
 
 	struct FStepSampler
 	{
-		FStepSampler(const UPCGPolyLineData* InLineData, const FPCGSplineSamplerParams& Params)
+		FStepSampler(const UPCGPolyLineData* InLineData, const FPCGSplineSamplerParams& Params, int InSeed)
 			: LineData(InLineData)
+			, Seed(InSeed)
 			, bComputeCurvature(Params.bComputeCurvature)
 			, bComputeTangents(Params.bComputeTangents)
 			, bComputeAlpha(Params.bComputeAlpha)
@@ -396,6 +405,8 @@ namespace PCGSplineSamplerHelpers
 		virtual bool IsDone() const = 0;
 
 		const UPCGPolyLineData* LineData = nullptr;
+		int Seed = 0;
+		int SampleIndex = 0;
 		int CurrentSegmentIndex = 0;
 		FVector::FReal DistanceToCurrentSegment = 0.0;
 		bool bComputeCurvature = false;
@@ -406,8 +417,8 @@ namespace PCGSplineSamplerHelpers
 
 	struct FSubdivisionStepSampler : public FStepSampler
 	{
-		FSubdivisionStepSampler(const UPCGPolyLineData* InLineData, const FPCGSplineSamplerParams& Params)
-			: FStepSampler(InLineData, Params)
+		FSubdivisionStepSampler(const UPCGPolyLineData* InLineData, const FPCGSplineSamplerParams& Params, int InSeed)
+			: FStepSampler(InLineData, Params, InSeed)
 		{
 			NumSegments = LineData->GetNumSegments();
 			SubdivisionsPerSegment = Params.SubdivisionsPerSegment;
@@ -431,6 +442,7 @@ namespace PCGSplineSamplerHelpers
 			FBox& OutBox = OutResult.Box;
 			FTransform& OutTransform = OutResult.LocalTransform;
 			OutTransform = LineData->GetTransformAtDistance(SegmentIndex, DistanceAlongSegment, /*bWorldSpace=*/false, &OutBox);
+			OutResult.SampleIndex = SampleIndex++;
 			OutResult.SegmentIndex = LineData->IsClosed() ? CurrentSegmentIndex : SegmentIndex;
 			OutResult.SubsegmentIndex = SubpointIndex;
 			OutResult.InputKey = LineData->GetInputKeyAtDistance(SegmentIndex, DistanceAlongSegment);
@@ -515,8 +527,8 @@ namespace PCGSplineSamplerHelpers
 
 	struct FDistanceStepSampler : public FStepSampler
 	{
-		FDistanceStepSampler(const UPCGPolyLineData* InLineData, const FPCGSplineSamplerParams& Params, int Seed)
-			: FStepSampler(InLineData, Params)
+		FDistanceStepSampler(const UPCGPolyLineData* InLineData, const FPCGSplineSamplerParams& Params, int InSeed)
+			: FStepSampler(InLineData, Params, InSeed)
 		{
 			StartOffset = FMath::Max(0, Params.StartOffset);
 			const FVector::FReal EndOffset = FMath::Max(0, Params.EndOffset);
@@ -575,6 +587,7 @@ namespace PCGSplineSamplerHelpers
 			FTransform& OutTransform = OutResult.LocalTransform;
 			FBox& OutBox = OutResult.Box;
 			OutTransform = LineData->GetTransformAtDistance(CurrentSegmentIndex, OffsetDistance, /*bWorldSpace=*/false, &OutBox);
+			OutResult.SampleIndex = SampleIndex++;
 			OutResult.SegmentIndex = CurrentSegmentIndex;
 			OutResult.InputKey = LineData->GetInputKeyAtDistance(CurrentSegmentIndex, OffsetDistance);
 
@@ -647,8 +660,9 @@ namespace PCGSplineSamplerHelpers
 
 	struct FDimensionSampler
 	{
-		FDimensionSampler(const UPCGPolyLineData* InLineData, const UPCGSpatialData* InBoundingShapeData, const UPCGSpatialData* InProjectionTarget, const FPCGProjectionParams& InProjectionParams, const FPCGSplineSamplerParams& InParams, UPCGPointData* OutPointData)
+		FDimensionSampler(const UPCGPolyLineData* InLineData, const UPCGSpatialData* InBoundingShapeData, const UPCGSpatialData* InProjectionTarget, const FPCGProjectionParams& InProjectionParams, const FPCGSplineSamplerParams& InParams, int InSeed, UPCGPointData* OutPointData)
 			: Params(InParams)
+			, Seed(InSeed)
 		{
 			check(InLineData);
 			LineData = InLineData;
@@ -805,7 +819,7 @@ namespace PCGSplineSamplerHelpers
 			FPCGPoint BoundsTestPoint;
 			if (bValid && (!BoundingShapeData || BoundingShapeData->SamplePoint(Transform, InResult.Box, BoundsTestPoint, nullptr)))
 			{
-				SetSeed(OutPoint, InResult.LocalTransform.GetLocation(), Params);
+				SetSeed(OutPoint, InResult.LocalTransform.GetLocation(), Params, Seed, InResult.SampleIndex);
 				SetMetadata(InResult, OutPoint, OutPointData->Metadata);
 				OutPointData->GetMutablePoints().Add(OutPoint);
 			}
@@ -816,6 +830,7 @@ namespace PCGSplineSamplerHelpers
 		const UPCGSpatialData* BoundingShapeData = nullptr;
 		const UPCGSpatialData* ProjectionTargetData = nullptr;
 		FPCGProjectionParams ProjectionParams;
+		int Seed = 0;
 
 		bool bSetMetadata = false;
 		bool bHasCustomMetadata = false;
@@ -833,8 +848,8 @@ namespace PCGSplineSamplerHelpers
 	/** Samples in a volume surrounding the poly line. */
 	struct FVolumeSampler : public FDimensionSampler
 	{
-		FVolumeSampler(const UPCGPolyLineData* InLineData, const UPCGSpatialData* InBoundingShapeData, const UPCGSpatialData* InProjectionTarget, const FPCGProjectionParams& InProjectionParams, const FPCGSplineSamplerParams& Params, UPCGPointData* OutPointData)
-			: FDimensionSampler(InLineData, InBoundingShapeData, InProjectionTarget, InProjectionParams, Params, OutPointData)
+		FVolumeSampler(const UPCGPolyLineData* InLineData, const UPCGSpatialData* InBoundingShapeData, const UPCGSpatialData* InProjectionTarget, const FPCGProjectionParams& InProjectionParams, const FPCGSplineSamplerParams& Params, int InSeed, UPCGPointData* OutPointData)
+			: FDimensionSampler(InLineData, InBoundingShapeData, InProjectionTarget, InProjectionParams, Params, InSeed, OutPointData)
 		{
 			Fill = Params.Fill;
 			NumPlanarSteps = 1 + ((Params.Dimension == EPCGSplineSamplingDimension::OnVertical) ? 0 : Params.NumPlanarSubdivisions);
@@ -917,7 +932,7 @@ namespace PCGSplineSamplerHelpers
 						continue;
 					}
 
-					SetSeed(OutPoint, TentativeLocationLS, Params);
+					SetSeed(OutPoint, TentativeLocationLS, Params, Seed, InResult.SampleIndex);
 					OutPoint.Steepness = Params.PointSteepness;
 					SetMetadata(InResult, OutPoint, OutPointData->Metadata);
 					OutPointData->GetMutablePoints().Add(OutPoint);
@@ -948,13 +963,17 @@ namespace PCGSplineSamplerHelpers
 			}
 		}
 
-		FSubdivisionStepSampler SubdivisionSampler(LineData, Params);
-		FDistanceStepSampler DistanceSampler(LineData, Params, Context ? Context->GetSeed() : 42);
+		const int32 Seed = Context ? Context->GetSeed() : 42;
+
+		FSubdivisionStepSampler SubdivisionSampler(LineData, Params, Seed);
+		FDistanceStepSampler DistanceSampler(LineData, Params, Seed);
 
 		FStepSampler* Sampler = ((Params.Mode == EPCGSplineSamplingMode::Subdivision) ? static_cast<FStepSampler*>(&SubdivisionSampler) : static_cast<FStepSampler*>(&DistanceSampler));
 
-		FDimensionSampler TrivialDimensionSampler(LineData, InBoundingShapeData, InProjectionTarget, InProjectionParams, Params, OutPointData);
-		FVolumeSampler VolumeSampler(LineData, InBoundingShapeData, InProjectionTarget, InProjectionParams, Params, OutPointData);
+
+		FDimensionSampler TrivialDimensionSampler(LineData, InBoundingShapeData, InProjectionTarget, InProjectionParams, Params, Seed, OutPointData);
+		FVolumeSampler VolumeSampler(LineData, InBoundingShapeData, InProjectionTarget, InProjectionParams, Params, Seed, OutPointData);
+
 
 		FDimensionSampler* ExtentsSampler = ((Params.Dimension == EPCGSplineSamplingDimension::OnSpline) ? &TrivialDimensionSampler : static_cast<FDimensionSampler*>(&VolumeSampler));
 
@@ -1368,6 +1387,9 @@ namespace PCGSplineSamplerHelpers
 			PointCount += InteriorData.Num();
 		}
 
+		const int32 Seed = Context ? Context->GetSeed() : 42;
+		int32 SampleIndex = 0;
+
 		// TODO: should we parallel for this too?
 		OutPoints.Reserve(PointCount);
 		for (const TArray<TTuple<FTransform, FVector, float>>& InteriorData : InteriorSplinePointData)
@@ -1376,7 +1398,7 @@ namespace PCGSplineSamplerHelpers
 			{
 				FPCGPoint& Point = OutPoints.Emplace_GetRef();
 				Point.Transform = InteriorPoint.Get<0>();
-				SetSeed(Point, InteriorPoint.Get<1>(), Params);
+				SetSeed(Point, InteriorPoint.Get<1>(), Params, Seed, ++SampleIndex);
 				Point.Density = InteriorPoint.Get<2>();
 				Point.BoundsMin = BoundsMin;
 				Point.BoundsMax = BoundsMax;
diff --git a/Engine/Plugins/PCG/Source/PCG/Public/Elements/PCGSplineSampler.h b/Engine/Plugins/PCG/Source/PCG/Public/Elements/PCGSplineSampler.h
--- a/Engine/Plugins/PCG/Source/PCG/Public/Elements/PCGSplineSampler.h
+++ b/Engine/Plugins/PCG/Source/PCG/Public/Elements/PCGSplineSampler.h
@@ -52,6 +52,13 @@ enum class EPCGSplineSamplingInteriorOrientation : uint8
 	FollowCurvature = 1
 };
 
+UENUM()
+enum class EPCGSplineSamplingSeedingMode : uint8
+{
+	SeedFromPosition = 0 UMETA(Tooltip = "Compute seed from sample position."),
+	SeedFromIndex = 1 UMETA(Tooltip = "Compute seed from sample index.")
+};
+
 USTRUCT(BlueprintType)
 struct FPCGSplineSamplerParams
 {
@@ -211,12 +218,16 @@ struct FPCGSplineSamplerParams
 	UPROPERTY(BlueprintReadWrite, EditAnywhere, Category = "Settings|Points", meta=(ClampMin="0", ClampMax="1", PCG_Overridable))
 	float PointSteepness = 0.5f;
 
-	/** Controls whether we will seed the sampled points using the final world position or the local position */
-	UPROPERTY(BlueprintReadWrite, EditAnywhere, Category = "Settings|Seeding")
+	/** Controls the mode for computing a sample point's seed. */
+	UPROPERTY(BlueprintReadWrite, EditAnywhere, Category = "Settings|Seeding");
+	EPCGSplineSamplingSeedingMode SeedingMode = EPCGSplineSamplingSeedingMode::SeedFromPosition;
+
+	/** Controls whether we will seed the sampled points using the final world position or the local position. */
+	UPROPERTY(BlueprintReadWrite, EditAnywhere, Category = "Settings|Seeding", meta=(EditCondition = "SeedingMode==EPCGSplineSamplingSeedingMode::SeedFromPosition", EditConditionHides))
 	bool bSeedFromLocalPosition = false;
 
-	/** Controls whether we will seed the sampled points using the 3D position or the 2D (XY) position */
-	UPROPERTY(BlueprintReadWrite, EditAnywhere, Category = "Settings|Seeding")
+	/** Controls whether we will seed the sampled points using the 3D position or the 2D (XY) position. */
+	UPROPERTY(BlueprintReadWrite, EditAnywhere, Category = "Settings|Seeding", meta=(EditCondition = "SeedingMode==EPCGSplineSamplingSeedingMode::SeedFromPosition", EditConditionHides))
 	bool bSeedFrom2DPosition = false;
 };
 
