Added collision-aware self-pruning.
Added set bounds from mesh node to act as a preamble to the self-pruning

diff --git a/Engine/Plugins/PCG/Source/PCG/Private/Elements/PCGBoundsFromMesh.cpp b/Engine/Plugins/PCG/Source/PCG/Private/Elements/PCGBoundsFromMesh.cpp
new file mode 100644
--- /dev/null
+++ b/Engine/Plugins/PCG/Source/PCG/Private/Elements/PCGBoundsFromMesh.cpp
@@ -0,0 +1,326 @@
+// Copyright Epic Games, Inc. All Rights Reserved.
+
+#include "Elements/PCGBoundsFromMesh.h"
+
+#include "PCGContext.h"
+#include "PCGParamData.h"
+#include "Data/PCGPointData.h"
+#include "Elements/Metadata/PCGMetadataElementCommon.h"
+#include "Helpers/PCGAsync.h"
+#include "Metadata/Accessors/PCGAttributeAccessorHelpers.h"
+#include "Metadata/Accessors/PCGAttributeAccessorKeys.h"
+
+#include "Algo/Transform.h"
+#include "Engine/StaticMesh.h"
+
+#define LOCTEXT_NAMESPACE "PCGBoundsFromMeshElement"
+
+namespace PCGBoundsFromMeshSettings
+{
+	const FName AttributeLabel = TEXT("Attribute");
+	const FText AttributeTooltip = LOCTEXT("AttributeTooltip", "Optional Attribute Set to select the mesh from. Not used if not connected.");
+}
+
+#if WITH_EDITOR
+FName UPCGBoundsFromMeshSettings::GetDefaultNodeName() const
+{
+	return FName(TEXT("BoundsFromMesh"));
+}
+
+FText UPCGBoundsFromMeshSettings::GetDefaultNodeTitle() const
+{
+	return LOCTEXT("NodeTitle", "Bounds From Mesh");
+}
+
+FText UPCGBoundsFromMeshSettings::GetNodeTooltipText() const
+{
+	return LOCTEXT("NodeTooltip", "Sets bounds min/max from mesh(es).");
+}
+
+bool UPCGBoundsFromMeshSettings::IsPinUsedByNodeExecution(const UPCGPin* InPin) const
+{
+	return !InPin || (InPin->Properties.Label != PCGBoundsFromMeshSettings::AttributeLabel) || InPin->IsConnected();
+}
+#endif //WITH_EDITOR
+
+TArray<FPCGPinProperties> UPCGBoundsFromMeshSettings::InputPinProperties() const
+{
+	TArray<FPCGPinProperties> PinProperties;
+	FPCGPinProperties& InputPinProperty = PinProperties.Emplace_GetRef(PCGPinConstants::DefaultInputLabel, EPCGDataType::Point);
+	InputPinProperty.SetRequiredPin();
+
+	PinProperties.Emplace(PCGBoundsFromMeshSettings::AttributeLabel, EPCGDataType::Param, /*bInAllowMultipleConnections=*/false, /*bAllowMultipleData=*/false, PCGBoundsFromMeshSettings::AttributeTooltip);
+
+	return PinProperties;
+}
+
+TArray<FPCGPinProperties> UPCGBoundsFromMeshSettings::OutputPinProperties() const
+{
+	return Super::DefaultPointOutputPinProperties();
+}
+
+// Creates the Element to be used for ExecuteInternal.
+FPCGElementPtr UPCGBoundsFromMeshSettings::CreateElement() const
+{
+	return MakeShared<FPCGBoundsFromMeshElement>();
+}
+
+bool FPCGBoundsFromMeshElement::PrepareDataInternal(FPCGContext* InContext) const
+{
+	TRACE_CPUPROFILER_EVENT_SCOPE(FPCGBoundsFromMeshElement::PrepareData);
+
+	check(InContext);
+	FPCGBoundsFromMeshContext* Context = static_cast<FPCGBoundsFromMeshContext*>(InContext);
+
+	const UPCGBoundsFromMeshSettings* Settings = Context->GetInputSettings<UPCGBoundsFromMeshSettings>();
+	check(Settings);
+
+	if (!Context->bPrepareDone)
+	{
+		// Some background info:
+		// Two common use cases:
+		// 1. Input is a point data, and an attribute set is used to provide the mesh to take bounds from.
+		// 2. Input is a point data, with a mesh attribute.
+		TArray<FPCGTaggedData> SourceAttributeInputs = Context->InputData.GetInputsByPin(PCGBoundsFromMeshSettings::AttributeLabel);
+		const UPCGParamData* SourceAttributeSet = (!SourceAttributeInputs.IsEmpty() ? Cast<UPCGParamData>(SourceAttributeInputs[0].Data) : nullptr);
+
+		if (SourceAttributeSet)
+		{
+			// In this case, we don't need to partition anything, we'll have a single mesh entry.
+			FPCGAttributePropertyInputSelector MeshAttributeSelector = Settings->MeshAttribute.CopyAndFixLast(SourceAttributeSet);
+			TUniquePtr<const IPCGAttributeAccessor> InputAccessor = PCGAttributeAccessorHelpers::CreateConstAccessor(SourceAttributeSet, MeshAttributeSelector);
+			TUniquePtr<const IPCGAttributeAccessorKeys> InputKeys = PCGAttributeAccessorHelpers::CreateConstKeys(SourceAttributeSet, MeshAttributeSelector);
+
+			// Warn if the attribute set has multiple entries, because we'll use only the first value.
+			if (SourceAttributeSet->ConstMetadata()->GetLocalItemCount() > 1)
+			{
+				PCGE_LOG(Warning, GraphAndLog, LOCTEXT("AttributeSetHasMultipleEntries", "Input attribute set has multiple entries - only the first one will be used."));
+			}
+
+			if (!InputAccessor.IsValid() || !InputKeys.IsValid())
+			{
+				if (!Settings->bSilenceAttributeNotFoundErrors)
+				{
+					PCGE_LOG(Warning, GraphAndLog, FText::Format(LOCTEXT("AttributeDoesNotExistOnAttributeSet", "Input attribute/property '{0}' does not exist on attribute set, skipped."), MeshAttributeSelector.GetDisplayText()));
+				}
+			}
+			else
+			{
+				if (InputAccessor->Get<FSoftObjectPath>(Context->SingleMesh, 0, *InputKeys, EPCGAttributeAccessorFlags::AllowConstructible))
+				{
+					if (!Context->SingleMesh.IsNull())
+					{
+						Context->MeshesToLoad.Add(Context->SingleMesh);
+					}
+				}
+				else if (!Settings->bSilenceAttributeNotFoundErrors)
+				{
+					PCGE_LOG(Warning, GraphAndLog, FText::Format(LOCTEXT("AttributeOfTheWrongTypeOnAttributeSet", "Input attribute/property '{0}' on attribute set does not match the expected type, skipped"), MeshAttributeSelector.GetDisplayText()));
+				}
+			}
+
+			Context->bPrepareDone = true;
+		}
+		else
+		{
+			// In this case, we'll visit all points, gather the unique meshes and assign to each point a given mesh index
+			TArray<FPCGTaggedData> Inputs = Context->InputData.GetInputsByPin(PCGPinConstants::DefaultInputLabel);
+
+			while (Context->CurrentPrepareIndex < Inputs.Num())
+			{
+				const FPCGTaggedData& Input = Inputs[Context->CurrentPrepareIndex];
+				const UPCGPointData* PointData = Cast<UPCGPointData>(Input.Data);
+
+				if (!PointData)
+				{
+					++Context->CurrentPrepareIndex;
+					continue;
+				}
+
+				FPCGAttributePropertyInputSelector MeshAttributeSelector = Settings->MeshAttribute.CopyAndFixLast(PointData);
+
+				FPCGBoundsFromMeshContext::InputMeshesData InputMeshData;
+				InputMeshData.InputIndex = Context->CurrentPrepareIndex++;
+
+				// Retrieve meshes
+				TUniquePtr<const IPCGAttributeAccessor> InputAccessor = PCGAttributeAccessorHelpers::CreateConstAccessor(PointData, MeshAttributeSelector);
+				TUniquePtr<const IPCGAttributeAccessorKeys> InputKeys = PCGAttributeAccessorHelpers::CreateConstKeys(PointData, MeshAttributeSelector);
+
+				if (!InputAccessor.IsValid() || !InputKeys.IsValid())
+				{
+					if (!Settings->bSilenceAttributeNotFoundErrors)
+					{
+						PCGE_LOG(Warning, GraphAndLog, FText::Format(LOCTEXT("AttributeDoesNotExistOnInput", "Input attribute/property '{0}' does not exist on input {1}, skipped"), MeshAttributeSelector.GetDisplayText(), FText::AsNumber(Context->CurrentPrepareIndex - 1)));
+					}
+
+					continue;
+				}
+
+				InputMeshData.MeshValueIndex.Reserve(InputKeys->GetNum());
+				auto GatherMeshPathsAndCreateMap = [this, Context, &InputMeshData](const TArrayView<FSoftObjectPath>& Meshes, int Start, int Range)
+				{
+					for (const FSoftObjectPath& Mesh : Meshes)
+					{
+						InputMeshData.MeshValueIndex.Add(Mesh.IsNull() ? INDEX_NONE : Context->MeshesToLoad.AddUnique(Mesh));
+					}
+				};
+
+				if (!PCGMetadataElementCommon::ApplyOnAccessorRange<FSoftObjectPath>(*InputKeys, *InputAccessor, GatherMeshPathsAndCreateMap, EPCGAttributeAccessorFlags::AllowBroadcast | EPCGAttributeAccessorFlags::AllowConstructible))
+				{
+					if (!Settings->bSilenceAttributeNotFoundErrors)
+					{
+						PCGE_LOG(Warning, GraphAndLog, FText::Format(LOCTEXT("AttributeOfTheWrongTypeOnInput", "Input attribute/property '{0}' on input {1} does not match the expected type, skipped"), MeshAttributeSelector.GetDisplayText(), FText::AsNumber(Context->CurrentPrepareIndex - 1)));
+					}
+				}
+				else
+				{
+					Context->PerInputData.Add(MoveTemp(InputMeshData));
+				}
+
+				if (Context->ShouldStop())
+				{
+					break;
+				}
+			}
+
+			if (Context->CurrentPrepareIndex == Inputs.Num())
+			{
+				Context->bPrepareDone = true;
+			}
+
+			// Given partitioning can be expensive, check if we're out of time for this frame
+			if (Context->ShouldStop())
+			{
+				return false;
+			}
+		}
+	}
+
+	if (Context->bPrepareDone && !Context->WasLoadRequested())
+	{
+		TArray<FSoftObjectPath> MeshesToLoad = Context->MeshesToLoad; // because of the move
+		return Context->RequestResourceLoad(Context, MoveTemp(MeshesToLoad), !Settings->bSynchronousLoad);
+	}
+
+	return true;
+}
+
+bool FPCGBoundsFromMeshElement::ExecuteInternal(FPCGContext* InContext) const
+{
+	TRACE_CPUPROFILER_EVENT_SCOPE(FPCGBoundsFromMeshElement::Execute);
+
+	check(InContext);
+	FPCGBoundsFromMeshContext* Context = static_cast<FPCGBoundsFromMeshContext*>(InContext);
+
+	const UPCGBoundsFromMeshSettings* Settings = Context->GetInputSettings<UPCGBoundsFromMeshSettings>();
+	check(Settings);
+
+	TArray<FPCGTaggedData> Inputs = Context->InputData.GetInputsByPin(PCGPinConstants::DefaultInputLabel);
+	TArray<FPCGTaggedData>& Outputs = Context->OutputData.TaggedData;
+
+	if (!Context->bBoundsQueried)
+	{
+		for (int MeshIndex = 0; MeshIndex < Context->MeshesToLoad.Num(); ++MeshIndex)
+		{
+			const FSoftObjectPath& MeshPath = Context->MeshesToLoad[MeshIndex];
+			TSoftObjectPtr<UStaticMesh> Mesh(MeshPath);
+
+			if (Mesh.IsValid() && Mesh.Get())
+			{
+				Context->MeshToBoundsMap.Add(MeshIndex, Mesh->GetBoundingBox());
+			}
+		}
+
+		Context->bBoundsQueried = true;
+	}
+
+	while (Context->CurrentExecuteIndex < Inputs.Num())
+	{
+		const FPCGTaggedData& Input = Inputs[Context->CurrentExecuteIndex++];
+		FPCGTaggedData& Output = Outputs.Add_GetRef(Input);
+
+		// If input is not a point data -> passthrough
+		const UPCGPointData* PointData = Cast<UPCGPointData>(Input.Data);
+
+		if (!PointData)
+		{
+			PCGE_LOG(Error, GraphAndLog, LOCTEXT("InvalidInputData", "Invalid input data"));
+			continue;
+		}
+
+		// Two cases here: either we have a single mesh and need to set its bounds to all points in the input
+		// or we need to compare against our list of meshes
+		if (Context->SingleMesh.IsValid())
+		{
+			TSoftObjectPtr<UStaticMesh> SingleMesh(Context->SingleMesh);
+
+			if (!SingleMesh)
+			{
+				PCGE_LOG(Error, GraphAndLog, LOCTEXT("LoadStaticMeshFailed", "Failed to load StaticMesh"));
+				continue;
+			}
+
+			const FBox StaticMeshBounds = SingleMesh->GetBoundingBox();
+			
+			UPCGPointData* OutputData = NewObject<UPCGPointData>();
+			OutputData->InitializeFromData(PointData);
+
+			const TArray<FPCGPoint>& InputPoints = PointData->GetPoints();
+			TArray<FPCGPoint>& OutputPoints = OutputData->GetMutablePoints();
+			OutputPoints.SetNumUninitialized(InputPoints.Num());
+
+			auto CopyAndSetBounds = [&InputPoints, &OutputPoints, &StaticMeshBounds](const int32 ReadIndex, const int32 WriteIndex)
+			{
+				const FPCGPoint& InputPoint = InputPoints[ReadIndex];
+				FPCGPoint& OutputPoint = OutputPoints[WriteIndex];
+
+				OutputPoint = InputPoint;
+				OutputPoint.BoundsMin = StaticMeshBounds.Min;
+				OutputPoint.BoundsMax = StaticMeshBounds.Max;
+			};
+
+			FPCGAsync::AsyncProcessingOneToOneEx(&Context->AsyncState, InputPoints.Num(), /*Initialize=*/[]() {}, CopyAndSetBounds, /*bEnableTimeSlicing=*/false);
+
+			Output.Data = OutputData;
+		}
+		else
+		{
+			FPCGBoundsFromMeshContext::InputMeshesData* InputData = Context->PerInputData.FindByPredicate([Index = Context->CurrentExecuteIndex - 1](const FPCGBoundsFromMeshContext::InputMeshesData& InputData) { return InputData.InputIndex == Index; });
+			if (InputData)
+			{
+				UPCGPointData* OutputData = NewObject<UPCGPointData>();
+				OutputData->InitializeFromData(PointData);
+
+				const TArray<FPCGPoint>& InputPoints = PointData->GetPoints();
+				TArray<FPCGPoint>& OutputPoints = OutputData->GetMutablePoints();
+				OutputPoints.SetNumUninitialized(InputPoints.Num());
+
+				auto CopyAndSetBounds = [&InputPoints, &OutputPoints, Context, InputData](const int32 ReadIndex, const int32 WriteIndex)
+				{
+					const FPCGPoint& InputPoint = InputPoints[ReadIndex];
+					FPCGPoint& OutputPoint = OutputPoints[WriteIndex];
+
+					OutputPoint = InputPoint;
+					if (FBox* MatchingBounds = Context->MeshToBoundsMap.Find(InputData->MeshValueIndex[ReadIndex]))
+					{
+						OutputPoint.BoundsMin = MatchingBounds->Min;
+						OutputPoint.BoundsMax = MatchingBounds->Max;
+					}
+				};
+
+				FPCGAsync::AsyncProcessingOneToOneEx(&Context->AsyncState, InputPoints.Num(), /*Initialize=*/[]() {}, CopyAndSetBounds, /*bEnableTimeSlicing=*/false);
+
+				Output.Data = OutputData;
+			}
+			else
+			{
+				// The data didn't have the attribute or all meshes were unloadable -> passthrough
+				continue;
+			}
+		}
+	}
+
+	return true;
+}
+
+#undef LOCTEXT_NAMESPACE
\ No newline at end of file
diff --git a/Engine/Plugins/PCG/Source/PCG/Private/Elements/PCGSelfPruning.cpp b/Engine/Plugins/PCG/Source/PCG/Private/Elements/PCGSelfPruning.cpp
--- a/Engine/Plugins/PCG/Source/PCG/Private/Elements/PCGSelfPruning.cpp
+++ b/Engine/Plugins/PCG/Source/PCG/Private/Elements/PCGSelfPruning.cpp
@@ -11,6 +11,10 @@
 #include "Metadata/Accessors/PCGAttributeAccessorKeys.h"
 #include "Metadata/Accessors/PCGAttributeExtractor.h"
 
+#include "CollisionShape.h"
+#include "Chaos/ImplicitObject.h"
+#include "PhysicsEngine/BodyInstance.h"
+
 #include UE_INLINE_GENERATED_CPP_BY_NAME(PCGSelfPruning)
 
 #define LOCTEXT_NAMESPACE "PCGSelfPruningElement"
@@ -153,6 +157,167 @@ namespace PCGSelfPruningElement
 		return true;
 	}
 
+	/** Self-pruning driven by use of collision shapes. Implementation is in practice just a secondary step after the octree query to filter out points if their collisions don`t intersect. */
+	bool CollisionExclusion(FIterationState& IterationState, FPCGContext* InOptionalContext, bool bUseComplexCollision)
+	{
+		TRACE_CPUPROFILER_EVENT_SCOPE(FPCGSelfPruningElement::Execute::CollisionExclusion);
+
+		check(IterationState.InputPointData);
+		const UPCGPointData::PointOctree& Octree = IterationState.InputPointData->GetOctree();
+
+		int32 CheckTimeSlicingCount = 0;
+
+		TArray<const FPCGPointRef*, TInlineAllocator<256>> ElementsToTest;
+
+		while (IterationState.CurrentPointIndex < IterationState.SortedPoints.Num())
+		{
+			// Don't check too many times. Pre-increment will make sure we always at least process "TimeSliceFrequencyCheck" elements at each iteration.
+			if (++CheckTimeSlicingCount >= TimeSliceFrequencyCheck)
+			{
+				if (ShouldStop(InOptionalContext))
+				{
+					return false;
+				}
+
+				CheckTimeSlicingCount = 0;
+			}
+
+			const FPCGPointRef& PointRef = IterationState.SortedPoints[IterationState.CurrentPointIndex++];
+			if (IterationState.ExcludedPoints.Contains(PointRef.Point))
+			{
+				continue; // Point discarded from previous iteration
+			}
+
+			// Select point
+			IterationState.ExclusionPoints.Add(PointRef.Point);
+
+			// 1. Gather point refs to test against - similar to the DensityBoundsExclusion, except we write to an array temporarily
+			ElementsToTest.Reset();
+			Octree.FindElementsWithBoundsTest(FBoxCenterAndExtent(PointRef.Bounds.Origin, PointRef.Bounds.BoxExtent), [&IterationState, &ElementsToTest](const FPCGPointRef& OtherPointRef)
+			{
+				if (!IterationState.ExclusionPoints.Contains(OtherPointRef.Point) && !IterationState.ExcludedPoints.Contains(OtherPointRef.Point))
+				{
+					ElementsToTest.Add(&OtherPointRef);
+				}
+			});
+
+			// For perf reasons, we won't do the body instance setup (scale...) if there's nothing to test against
+			if (ElementsToTest.IsEmpty())
+			{
+				continue;
+			}
+
+			// Implementation note: this is a deconstruction of FBodyInstance::OverlapTestForBodiesImpl
+			// Get this point's body instance, if any. Note the -1 because we already incremented the value.
+			FBodyInstance* ThisInstance = IterationState.CollisionWrapper.GetBodyInstance(IterationState.CurrentPointIndex - 1);
+			FCollisionShape ThisCollisionShape;
+			FTransform TransformNoScale = FTransform(PointRef.Point->Transform.GetRotation(), PointRef.Point->Transform.GetLocation());
+
+			PhysicsInterfaceTypes::FInlineShapeArray TargetShapes;
+
+			if (ThisInstance)
+			{
+				// Since scale can change on a per-point basis and we can't fold it in the relative transform, we must apply scale here.
+				ThisInstance->UpdateBodyScale(PointRef.Point->Transform.GetScale3D());
+				// Implementation note: all of these bodies don't exist in the physics scene, so there's no need to lock
+				FillInlineShapeArray_AssumesLocked(TargetShapes, ThisInstance->GetPhysicsActorHandle());
+			}
+			else
+			{
+				ThisCollisionShape.SetBox(FVector3f(PointRef.Bounds.BoxExtent * PointRef.Point->Transform.GetScale3D()));
+			}
+
+			for (const FPCGPointRef* OtherPointRef : ElementsToTest)
+			{
+				FBodyInstance* OtherInstance = IterationState.CollisionWrapper.GetBodyInstance(OtherPointRef->Point - IterationState.InputPointData->GetPoints().GetData());
+
+				bool bOverlaps = true;
+				if (ThisInstance && OtherInstance)
+				{
+					// If both bodies are the same, then we can't really have them having different scales; we need to do a (temporary) copy when that happens
+					if (ThisInstance == OtherInstance && !OtherPointRef->Point->Transform.GetScale3D().Equals(PointRef.Point->Transform.GetScale3D()))
+					{
+						if (FBodyInstance** TemporaryInstance = IterationState.TemporaryBodyInstances.Find(OtherInstance))
+						{
+							OtherInstance = *TemporaryInstance;
+						}
+						else
+						{
+							OtherInstance = new FBodyInstance();
+							OtherInstance->bAutoWeld = false;
+							OtherInstance->bSimulatePhysics = false;
+							OtherInstance->InitBody(ThisInstance->GetBodySetup(), FTransform::Identity, nullptr, nullptr);
+
+							IterationState.TemporaryBodyInstances.Add(ThisInstance, OtherInstance);
+						}
+
+						OtherInstance->UpdateBodyScale(OtherPointRef->Point->Transform.GetScale3D());
+					}
+
+					FTransform OtherTransformNoScale = FTransform(OtherPointRef->Point->Transform.GetRotation(), OtherPointRef->Point->Transform.GetLocation());
+					FTransform RelativeTransform = TransformNoScale.GetRelativeTransform(OtherTransformNoScale);
+
+					{
+						// Copied more or less verbatim from FBodyInstance::OverlapTestForBodiesImpl
+						//TRACE_CPUPROFILER_EVENT_SCOPE(FPCGSelfPruningElement::Execute::CollisionExclusion::OverlapTest);
+						bOverlaps = false;
+						for (const FPhysicsShapeHandle& Shape : TargetShapes)
+						{
+							if (!Shape.GetGeometry().IsConvex())
+							{
+								continue; //we skip complex shapes - should this respect ComplexAsSimple?
+							}
+
+							bOverlaps = FPhysicsInterface::Overlap_Geom(OtherInstance, FPhysicsInterface::GetGeometryCollection(Shape), RelativeTransform, /*bOutMTD=*/nullptr, bUseComplexCollision);
+							if (bOverlaps)
+							{
+								break;
+							}
+						}
+					}
+				}
+				else if (ThisInstance)
+				{
+					FTransform RelativeTransform = TransformNoScale.GetRelativeTransform(OtherPointRef->Point->Transform);
+					FTransform GeomTransform(RelativeTransform.GetRotation(), RelativeTransform.GetLocation());
+
+					// Create collision shape for other instance
+					FCollisionShape OtherCollisionShape;
+					OtherCollisionShape.SetBox(FVector3f(OtherPointRef->Bounds.BoxExtent * RelativeTransform.GetScale3D()));
+
+					bOverlaps = FPhysicsInterface::Overlap_Geom(ThisInstance, OtherCollisionShape, GeomTransform.GetRotation(), GeomTransform, nullptr, bUseComplexCollision);
+				}
+				else if (OtherInstance)
+				{
+					OtherInstance->UpdateBodyScale(OtherPointRef->Point->Transform.GetScale3D());
+
+					FTransform OtherTransformNoScale = FTransform(OtherPointRef->Point->Transform.GetRotation(), OtherPointRef->Point->Transform.GetLocation());
+					FTransform RelativeTransform = OtherTransformNoScale.GetRelativeTransform(TransformNoScale);
+					FTransform GeomTransform(RelativeTransform.GetRotation(), RelativeTransform.GetLocation());
+
+					bOverlaps = FPhysicsInterface::Overlap_Geom(OtherInstance, ThisCollisionShape, GeomTransform.GetRotation(), GeomTransform, nullptr, bUseComplexCollision);
+				}
+
+				if (bOverlaps)
+				{
+					IterationState.ExcludedPoints.Add(OtherPointRef->Point);
+				}
+			}
+		}
+
+		// Clean up temporary instances
+		if (IterationState.CurrentPointIndex == IterationState.SortedPoints.Num())
+		{
+			for (TPair<FBodyInstance*, FBodyInstance*>& TemporaryInstance : IterationState.TemporaryBodyInstances)
+			{
+				delete TemporaryInstance.Value;
+			}
+			IterationState.TemporaryBodyInstances.Reset();
+		}
+
+		return true;
+	}
+
 	bool DuplicatePointsExclusion(FIterationState& IterationState, FPCGContext* InOptionalContext)
 	{
 		TRACE_CPUPROFILER_EVENT_SCOPE(FPCGSelfPruningElement::Execute::DuplicatePointsExclusion);
@@ -411,6 +576,10 @@ namespace PCGSelfPruningElement
 		{
 			bIsDone = PCGSelfPruningElement::DuplicatePointsExclusion(InState, InOptionalContext);
 		}
+		else if (InParameters.bUseCollisionAttribute)
+		{
+			bIsDone = PCGSelfPruningElement::CollisionExclusion(InState, InOptionalContext, InParameters.bUseComplexCollision);
+		}
 		else
 		{
 			bIsDone = PCGSelfPruningElement::DensityBoundsExclusion(InState, InOptionalContext);
@@ -493,6 +662,9 @@ bool FPCGSelfPruningElement::PrepareDataInternal(FPCGContext* InContext) const
 	FPCGSelfPruningElement::ContextType* TimeSlicedContext = static_cast<FPCGSelfPruningElement::ContextType*>(InContext);
 	check(TimeSlicedContext);
 
+	const UPCGSelfPruningSettings* Settings = TimeSlicedContext->GetInputSettings<UPCGSelfPruningSettings>();
+	check(Settings);
+
 	TimeSlicedContext->SetTimeSliceIsEnabled(true);
 
 	TArray<FPCGTaggedData> Inputs = InContext->InputData.GetInputsByPin(PCGPinConstants::DefaultInputLabel);
@@ -500,7 +672,7 @@ bool FPCGSelfPruningElement::PrepareDataInternal(FPCGContext* InContext) const
 	// No global execution state.
 	EPCGTimeSliceInitResult InitResult = TimeSlicedContext->InitializePerExecutionState();
 
-	TimeSlicedContext->InitializePerIterationStates(Inputs.Num(), [&Inputs, InContext](PCGSelfPruningElement::FIterationState& OutState, const FPCGSelfPruningElement::ExecStateType&, int32 Index) -> EPCGTimeSliceInitResult
+	TimeSlicedContext->InitializePerIterationStates(Inputs.Num(), [&Inputs, Settings, InContext](PCGSelfPruningElement::FIterationState& OutState, const FPCGSelfPruningElement::ExecStateType&, int32 Index) -> EPCGTimeSliceInitResult
 	{
 		const UPCGSpatialData* SpatialData = Cast<UPCGSpatialData>(Inputs[Index].Data);
 		if (!SpatialData)
@@ -516,6 +688,16 @@ bool FPCGSelfPruningElement::PrepareDataInternal(FPCGContext* InContext) const
 			return EPCGTimeSliceInitResult::NoOperation;
 		}
 
+		if (Settings->Parameters.PruningType != EPCGSelfPruningType::RemoveDuplicates && Settings->Parameters.bUseCollisionAttribute)
+		{
+			FPCGAttributePropertyInputSelector InputSelector = Settings->Parameters.CollisionAttribute.CopyAndFixLast(OutState.InputPointData);
+
+			TUniquePtr<const IPCGAttributeAccessor> InputAccessor = PCGAttributeAccessorHelpers::CreateConstAccessor(OutState.InputPointData, InputSelector);
+			TUniquePtr<const IPCGAttributeAccessorKeys> InputKeys = PCGAttributeAccessorHelpers::CreateConstKeys(OutState.InputPointData, InputSelector);
+
+			OutState.CollisionWrapper.Initialize(InputAccessor.Get(), InputKeys.Get());
+		}
+
 		return EPCGTimeSliceInitResult::Success;
 	});
 
diff --git a/Engine/Plugins/PCG/Source/PCG/Public/Elements/PCGBoundsFromMesh.h b/Engine/Plugins/PCG/Source/PCG/Public/Elements/PCGBoundsFromMesh.h
new file mode 100644
--- /dev/null
+++ b/Engine/Plugins/PCG/Source/PCG/Public/Elements/PCGBoundsFromMesh.h
@@ -0,0 +1,80 @@
+// Copyright Epic Games, Inc. All Rights Reserved.
+
+#pragma once
+
+#include "PCGSettings.h"
+#include "PCGContext.h"
+#include "Async/PCGAsyncLoadingContext.h"
+
+#include "PCGBoundsFromMesh.generated.h"
+
+class UStaticMesh;
+
+/**
+* Sets the bounds on the points according to the mesh(es) provided in the mesh pin.
+*/
+UCLASS(MinimalAPI, BlueprintType)
+class UPCGBoundsFromMeshSettings : public UPCGSettings
+{
+	GENERATED_BODY()
+
+public:
+	//~Begin UPCGSettings interface
+#if WITH_EDITOR
+	virtual FName GetDefaultNodeName() const override;
+	virtual FText GetDefaultNodeTitle() const override;
+	virtual FText GetNodeTooltipText() const override;
+	virtual EPCGSettingsType GetType() const override { return EPCGSettingsType::Spatial; }
+	virtual bool IsPinUsedByNodeExecution(const UPCGPin* InPin) const override;
+#endif //WITH_EDITOR
+
+protected:
+	virtual TArray<FPCGPinProperties> InputPinProperties() const override;
+	virtual TArray<FPCGPinProperties> OutputPinProperties() const override;
+	virtual FPCGElementPtr CreateElement() const override;
+	//~End UPCGSettings interface
+
+public:
+	/** Attribute from which to source the meshes to use. */
+	UPROPERTY(BlueprintReadWrite, EditAnywhere, Category = Settings, meta = (PCG_DiscardPropertySelection))
+	FPCGAttributePropertyInputSelector MeshAttribute;
+
+	/** Will not produce warnings when the input data does not have the required attribute. */
+	UPROPERTY(BlueprintReadWrite, EditAnywhere, Category = Settings)
+	bool bSilenceAttributeNotFoundErrors = false;
+
+	/** By default, will use async loading for the meshes. */
+	UPROPERTY(BlueprintReadWrite, EditAnywhere, Category = "Settings|Debug")
+	bool bSynchronousLoad = false;
+};
+
+struct FPCGBoundsFromMeshContext : public FPCGContext, public IPCGAsyncLoadingContext
+{
+	struct InputMeshesData
+	{
+		int InputIndex = 0;
+		TArray<int32> MeshValueIndex;
+	};
+
+	TArray<InputMeshesData> PerInputData;
+
+	FSoftObjectPath SingleMesh;
+	TMap<int, FBox> MeshToBoundsMap;
+
+	TArray<FSoftObjectPath> MeshesToLoad;
+	bool bPrepareDone = false;
+	bool bBoundsQueried = false;
+	int CurrentPrepareIndex = 0;
+	int CurrentExecuteIndex = 0;
+};
+
+class FPCGBoundsFromMeshElement : public IPCGElementWithCustomContext<FPCGBoundsFromMeshContext>
+{
+public:
+	// Loading needs to be done on the main thread and accessing objects outside of PCG might not be thread safe, so taking the safe approach
+	virtual bool CanExecuteOnlyOnMainThread(FPCGContext* Context) const override { return true; }
+
+protected:
+	virtual bool PrepareDataInternal(FPCGContext* Context) const override;
+	virtual bool ExecuteInternal(FPCGContext* Context) const override;
+};
diff --git a/Engine/Plugins/PCG/Source/PCG/Public/Elements/PCGSelfPruning.h b/Engine/Plugins/PCG/Source/PCG/Public/Elements/PCGSelfPruning.h
--- a/Engine/Plugins/PCG/Source/PCG/Public/Elements/PCGSelfPruning.h
+++ b/Engine/Plugins/PCG/Source/PCG/Public/Elements/PCGSelfPruning.h
@@ -4,6 +4,7 @@
 
 #include "PCGContext.h"
 #include "PCGSettings.h"
+#include "Data/PCGCollisionWrapperData.h"
 #include "Data/PCGPointData.h"
 #include "Elements/PCGTimeSlicedElementBase.h"
 
@@ -40,6 +41,17 @@ public:
 
 	UPROPERTY(BlueprintReadWrite, EditAnywhere, Category = Settings, meta = (PCG_Overridable))
 	bool bRandomizedPruning = true;
+
+	UPROPERTY(BlueprintReadWrite, EditAnywhere, Category = Settings, meta = (PCG_Overridable, EditCondition = "PruningType != EPCGSelfPruningType::RemoveDuplicates"))
+	bool bUseCollisionAttribute = false;
+
+	/** Specifies to use the collision from a given mesh instead of the point; note that this will be ignored in the Remove Duplicates mode. */
+	UPROPERTY(BlueprintReadWrite, EditAnywhere, Category = Settings, meta = (PCG_Overridable, PCG_DiscardPropertySelection, PCG_DiscardExtraSelection, EditCondition = "bUseCollisionAttribute && PruningType != EPCGSelfPruningType::RemoveDuplicates", EditConditionHides))
+	FPCGAttributePropertyInputSelector CollisionAttribute;
+
+	/** Queries against complex collision (vs. collection of convex elements in the collision) if enabled, performance warning */
+	UPROPERTY(BlueprintReadWrite, EditAnywhere, Category = Settings, meta = (PCG_Overridable, EditCondition = "bUseCollisionAttribute && PruningType != EPCGSelfPruningType::RemoveDuplicates", EditConditionHides))
+	bool bUseComplexCollision = false;
 };
 
 namespace PCGSelfPruningElement
@@ -65,6 +77,11 @@ namespace PCGSelfPruningElement
 		FPointBitSet ExclusionPoints;
 		int32 CurrentPointIndex = 0;
 		bool bSortDone = false;
+
+		// In the case of the self-pruning using the collision, we'll need instances of the unique bodies that are used in the point data.
+		// Additionally, in the cases where two points to compare use the same instance, we'll have to do a temporary copy.
+		FPCGCollisionWrapper CollisionWrapper;
+		TMap<FBodyInstance*, FBodyInstance*> TemporaryBodyInstances;
 	};
 
 	/** Will do the self pruning on all inputs in the context. Blocking call (no time slicing). Be careful, it can be very costly. */
