Added steepness control on create points grid

diff --git a/Engine/Plugins/PCG/Source/PCG/Private/Elements/PCGCreatePointsGrid.cpp b/Engine/Plugins/PCG/Source/PCG/Private/Elements/PCGCreatePointsGrid.cpp
--- a/Engine/Plugins/PCG/Source/PCG/Private/Elements/PCGCreatePointsGrid.cpp
+++ b/Engine/Plugins/PCG/Source/PCG/Private/Elements/PCGCreatePointsGrid.cpp
@@ -253,6 +253,7 @@ bool FPCGCreatePointsGridElement::ExecuteInternal(FPCGContext* Context) const
 			OutPoint.SetExtents(CellSize * 0.5);
 		}
 
+		OutPoint.Steepness = Settings->PointSteepness;
 		OutPoint.Seed = PCGHelpers::ComputeSeedFromPosition(PointTransform.GetLocation());
 
 		// Discards points outside of the volume
diff --git a/Engine/Plugins/PCG/Source/PCG/Public/Elements/PCGCreatePointsGrid.h b/Engine/Plugins/PCG/Source/PCG/Public/Elements/PCGCreatePointsGrid.h
--- a/Engine/Plugins/PCG/Source/PCG/Public/Elements/PCGCreatePointsGrid.h
+++ b/Engine/Plugins/PCG/Source/PCG/Public/Elements/PCGCreatePointsGrid.h
@@ -55,6 +55,14 @@ public:
 	UPROPERTY(BlueprintReadWrite, EditAnywhere, Category = Settings, meta = (PCG_Overridable, ClampMin = "0.0", UIMin = "0.0"))
 	FVector CellSize = FVector(100.0, 100.0, 100.0);
 
+	/** Each PCG point represents a discretized, volumetric region of world space. The points' Steepness value [0.0 to
+	 * 1.0] establishes how "hard" or "soft" that volume will be represented. From 0, it will ramp up linearly
+	 * increasing its influence over the density from the point's center to up to two times the bounds. At 1, it will
+	 * represent a binary box function with the size of the point's bounds.
+	 */
+	UPROPERTY(BlueprintReadWrite, EditAnywhere, Category = Settings, meta = (ClampMin = "0", ClampMax = "1", PCG_Overridable))
+	float PointSteepness = 0.5f;
+
 #if WITH_EDITORONLY_DATA
 PRAGMA_DISABLE_DEPRECATION_WARNINGS
 	/** Sets the points transform to world or local space*/
