Added replace string operation.
Light refactoring to show these a bit better in the contextual menu

diff --git a/Engine/Plugins/PCG/Source/PCG/Private/Elements/Metadata/PCGMetadataStringOpElement.cpp b/Engine/Plugins/PCG/Source/PCG/Private/Elements/Metadata/PCGMetadataStringOpElement.cpp
--- a/Engine/Plugins/PCG/Source/PCG/Private/Elements/Metadata/PCGMetadataStringOpElement.cpp
+++ b/Engine/Plugins/PCG/Source/PCG/Private/Elements/Metadata/PCGMetadataStringOpElement.cpp
@@ -6,19 +6,50 @@
 
 #include UE_INLINE_GENERATED_CPP_BY_NAME(PCGMetadataStringOpElement)
 
+namespace PCGMetadataStringOpConstants
+{
+	const FName ToFindLabel = TEXT("Search");
+	const FName ToReplaceLabel = TEXT("Replace");
+}
+
 FName UPCGMetadataStringOpSettings::GetInputPinLabel(uint32 Index) const
 {
-	switch (Index)
+	switch (Operation)
 	{
-	case 0: return PCGMetadataSettingsBaseConstants::DoubleInputFirstLabel;
-	case 1: return PCGMetadataSettingsBaseConstants::DoubleInputSecondLabel;
-	default: return NAME_None;
+	case EPCGMetadataStringOperation::Append:
+		switch (Index)
+		{
+		case 0: return PCGMetadataSettingsBaseConstants::DoubleInputFirstLabel;
+		case 1: return PCGMetadataSettingsBaseConstants::DoubleInputSecondLabel;
+		default: break;
+		}
+
+	case EPCGMetadataStringOperation::Replace:
+		switch (Index)
+		{
+		case 0: return PCGPinConstants::DefaultInputLabel;
+		case 1: return PCGMetadataStringOpConstants::ToFindLabel;
+		case 2: return PCGMetadataStringOpConstants::ToReplaceLabel;
+		default: break;
+		}
+
+	default: break;
 	}
+	
+	return NAME_None;
 }
 
 uint32 UPCGMetadataStringOpSettings::GetOperandNum() const
 {
-	return 2;
+	switch (Operation)
+	{
+	case EPCGMetadataStringOperation::Append:
+		return 2;
+	case EPCGMetadataStringOperation::Replace:
+		return 3;
+	default:
+		return 0;
+	}
 }
 
 bool UPCGMetadataStringOpSettings::IsSupportedInputType(uint16 TypeId, uint32 InputIndex, bool& bHasSpecialRequirement) const
@@ -33,6 +64,7 @@ FPCGAttributePropertyInputSelector UPCGMetadataStringOpSettings::GetInputSource(
 	{
 	case 0: return InputSource1;
 	case 1: return InputSource2;
+	case 2: return InputSource3;
 	default: return FPCGAttributePropertyInputSelector();
 	}
 }
@@ -59,8 +91,24 @@ FText UPCGMetadataStringOpSettings::GetDefaultNodeTitle() const
 {
 	return NSLOCTEXT("PCGMetadataStringOpSettings", "NodeTitle", "Attribute String Op");
 }
+
+TArray<FPCGPreConfiguredSettingsInfo> UPCGMetadataStringOpSettings::GetPreconfiguredInfo() const
+{
+	return PCGMetadataElementCommon::FillPreconfiguredSettingsInfoFromEnum<EPCGMetadataStringOperation>();
+}
 #endif // WITH_EDITOR
 
+void UPCGMetadataStringOpSettings::ApplyPreconfiguredSettings(const FPCGPreConfiguredSettingsInfo& PreconfiguredInfo)
+{
+	if (const UEnum* EnumPtr = StaticEnum<EPCGMetadataStringOperation>())
+	{
+		if (EnumPtr->IsValidEnumValue(PreconfiguredInfo.PreconfiguredIndex))
+		{
+			Operation = EPCGMetadataStringOperation(PreconfiguredInfo.PreconfiguredIndex);
+		}
+	}
+}
+
 FPCGElementPtr UPCGMetadataStringOpSettings::CreateElement() const
 {
 	return MakeShared<FPCGMetadataStringOpElement>();
@@ -81,6 +129,13 @@ bool FPCGMetadataStringOpElement::DoOperation(PCGMetadataOps::FOperationData& Op
 	{
 		DoBinaryOp<FString, FString>(OperationData, [](const FString& Value1, const FString& Value2) -> FString { return Value1 + Value2; });
 	}
+	else if (Settings->Operation == EPCGMetadataStringOperation::Replace)
+	{
+		DoTernaryOp<FString, FString, FString>(OperationData, [](const FString& InValue, const FString& InSearch, const FString& InReplace)
+		{
+			return InValue.Replace(*InSearch, *InReplace);
+		});
+	}
 
 	return true;
 }
\ No newline at end of file
diff --git a/Engine/Plugins/PCG/Source/PCG/Public/Elements/Metadata/PCGMetadataStringOpElement.h b/Engine/Plugins/PCG/Source/PCG/Public/Elements/Metadata/PCGMetadataStringOpElement.h
--- a/Engine/Plugins/PCG/Source/PCG/Public/Elements/Metadata/PCGMetadataStringOpElement.h
+++ b/Engine/Plugins/PCG/Source/PCG/Public/Elements/Metadata/PCGMetadataStringOpElement.h
@@ -9,7 +9,8 @@
 UENUM()
 enum class EPCGMetadataStringOperation : uint16
 {
-	Append
+	Append UMETA(DisplayName="Append String"),
+	Replace UMETA(DisplayName="Replace String")
 };
 
 UCLASS(MinimalAPI, BlueprintType, ClassGroup = (Procedural))
@@ -22,8 +23,11 @@ public:
 #if WITH_EDITOR
 	virtual FName GetDefaultNodeName() const override;
 	virtual FText GetDefaultNodeTitle() const override;
+	virtual TArray<FPCGPreConfiguredSettingsInfo> GetPreconfiguredInfo() const override;
+	virtual bool OnlyExposePreconfiguredSettings() const override { return true; }
 #endif
 	virtual FString GetAdditionalTitleInformation() const override;
+	virtual void ApplyPreconfiguredSettings(const FPCGPreConfiguredSettingsInfo& PreconfigureInfo) override;
 
 protected:
 #if WITH_EDITOR
@@ -52,6 +56,9 @@ public:
 
 	UPROPERTY(BlueprintReadWrite, EditAnywhere, Category = Input, meta = (PCG_Overridable))
 	FPCGAttributePropertyInputSelector InputSource2;
+
+	UPROPERTY(BlueprintReadWrite, EditAnywhere, Category = Input, meta = (PCG_Overridable, EditCondition = "Operation==EPCGMetadataStringOperation::Replace", EditConditionHides))
+	FPCGAttributePropertyInputSelector InputSource3;
 };
 
 class FPCGMetadataStringOpElement : public FPCGMetadataElementBase
