[PCG] Allow Surface Sampler to receive multiple inputs on the Bounding Shape, since its creating a Union on that pin anyway

#rnx

diff --git a/Engine/Plugins/PCG/Source/PCG/Private/Elements/PCGSurfaceSampler.cpp b/Engine/Plugins/PCG/Source/PCG/Private/Elements/PCGSurfaceSampler.cpp
--- a/Engine/Plugins/PCG/Source/PCG/Private/Elements/PCGSurfaceSampler.cpp
+++ b/Engine/Plugins/PCG/Source/PCG/Private/Elements/PCGSurfaceSampler.cpp
@@ -288,8 +288,7 @@ TArray<FPCGPinProperties> UPCGSurfaceSamplerSettings::InputPinProperties() const
 	));
 	SurfacePinProperty.SetRequiredPin();
 
-	// Only one connection allowed, user can union multiple shapes
-	PinProperties.Emplace(PCGSurfaceSamplerConstants::BoundingShapeLabel, EPCGDataType::Spatial, /*bInAllowMultipleConnections=*/false, /*bAllowMultipleData=*/false, LOCTEXT("SurfaceSamplerBoundingShapePinTooltip",
+	PinProperties.Emplace(PCGSurfaceSamplerConstants::BoundingShapeLabel, EPCGDataType::Spatial, /*bInAllowMultipleConnections=*/true, /*bAllowMultipleData=*/true, LOCTEXT("SurfaceSamplerBoundingShapePinTooltip",
 		"All sampled points must be contained within this shape. If this input is omitted then bounds will be taken from the actor so that points are contained within actor bounds. "
 		"The Unbounded property disables this and instead generates over the entire bounds of Surface."
 	));
