[PCG] Fix Deleting of PCG Partition Actors not catching unsaved actors (regression from 33170469)

#rnx

diff --git a/Engine/Plugins/PCG/Source/PCG/Private/PCGSubsystem.cpp b/Engine/Plugins/PCG/Source/PCG/Private/PCGSubsystem.cpp
--- a/Engine/Plugins/PCG/Source/PCG/Private/PCGSubsystem.cpp
+++ b/Engine/Plugins/PCG/Source/PCG/Private/PCGSubsystem.cpp
@@ -1434,13 +1434,27 @@ void UPCGSubsystem::DeleteSerializedPartitionActors(bool bOnlyDeleteUnused, bool
 					return true;
 				});
 
+				// Process Loaded Actors first (and unsaved actors that don't have an Actor Desc yet)
+				// Do not use UPCGActorHelpers::ForEachActorInLevel as GatherAndDestroy can end up modifying the Actors array (by loading actors) in a WP World
+				TSet<FGuid> ProcessedActors;
+				const TArray<AActor*> ActorsCopy(World->PersistentLevel->Actors);
+				for (AActor* Actor : ActorsCopy)
+				{
+					if (APCGPartitionActor* PartitionActor = Cast<APCGPartitionActor>(Actor))
+					{
+						ProcessedActors.Add(Actor->GetActorGuid());
+						GatherAndDestroyActors(Actor);
+					}
+				}
+				
 				FWorldPartitionHelpers::FForEachActorWithLoadingParams ForEachActorWithLoadingParams;
 				ForEachActorWithLoadingParams.bKeepReferences = true;
 				ForEachActorWithLoadingParams.ActorClasses = { APCGPartitionActor::StaticClass() };
 
-				FWorldPartitionHelpers::ForEachActorWithLoading(WorldPartition, [&](const FWorldPartitionActorDescInstance* ActorDescInstance)
+				// Load and Process remaining actors
+				FWorldPartitionHelpers::ForEachActorWithLoading(WorldPartition, [&GatherAndDestroyActors, &ProcessedActors](const FWorldPartitionActorDescInstance* ActorDescInstance)
 				{
-					if(AActor* Actor = ActorDescInstance->GetActor())
+					if(AActor* Actor = ActorDescInstance->GetActor(); Actor && !ProcessedActors.Contains(Actor->GetActorGuid()))
 					{
 						GatherAndDestroyActors(Actor);
 					}
