[PCG] Export friend function

#rnx

diff --git a/Engine/Plugins/PCG/Source/PCG/Public/Elements/PCGActorSelector.h b/Engine/Plugins/PCG/Source/PCG/Public/Elements/PCGActorSelector.h
--- a/Engine/Plugins/PCG/Source/PCG/Public/Elements/PCGActorSelector.h
+++ b/Engine/Plugins/PCG/Source/PCG/Public/Elements/PCGActorSelector.h
@@ -60,7 +60,8 @@ struct PCG_API FPCGSelectionKey
 
 	bool operator==(const FPCGSelectionKey& InOther) const;
 
-	friend uint32 GetTypeHash(const FPCGSelectionKey& In);
+	// Friend functions need to be explicitly exported.
+	PCG_API friend uint32 GetTypeHash(const FPCGSelectionKey& In);
 	bool IsMatching(const UObject* InObject, const UPCGComponent* InComponent) const;
 	bool IsMatching(const UObject* InObject, const TSet<FName>& InRemovedTags, const TSet<UPCGComponent*>& InComponents, TSet<UPCGComponent*>* OptionalMatchedComponents = nullptr) const;
 
