[PCG] New World Raycast node

* Should accept point or param data
* Includes Raycast Modes:
  + Infinite
  + Vector with Magnitude
  + Normalized Vector with Ray Length
  + Segments (End Points)

+ Also factored out raycast hit filtering to helpers

diff --git a/Engine/Plugins/PCG/Source/PCG/Private/Data/PCGLandscapeData.cpp b/Engine/Plugins/PCG/Source/PCG/Private/Data/PCGLandscapeData.cpp
--- a/Engine/Plugins/PCG/Source/PCG/Private/Data/PCGLandscapeData.cpp
+++ b/Engine/Plugins/PCG/Source/PCG/Private/Data/PCGLandscapeData.cpp
@@ -6,9 +6,9 @@
 #include "Data/PCGPointData.h"
 #include "Data/PCGSpatialDataTpl.h"
 #include "Data/PCGSurfaceData.h"
-#include "Data/PCGWorldData.h"
 #include "Grid/PCGLandscapeCache.h"
 #include "Helpers/PCGHelpers.h"
+#include "Helpers/PCGWorldQueryHelpers.h"
 
 #include "ChaosInterfaceWrapperCore.h"
 #include "Landscape.h"
@@ -106,7 +106,7 @@ void UPCGLandscapeData::Initialize(const TArray<TWeakObjectPtr<ALandscapeProxy>>
 
 	if (DataProps.bGetPhysicalMaterial)
 	{
-		Metadata->CreateAttribute<FSoftObjectPath>(PCGWorldRayHitConstants::PhysicalMaterialReferenceAttribute, FSoftObjectPath(), /*bAllowInterpolation=*/false, /*bOverrideParent*/false);
+		Metadata->CreateAttribute<FSoftObjectPath>(PCGWorldQueryConstants::PhysicalMaterialReferenceAttribute, FSoftObjectPath(), /*bAllowInterpolation=*/false, /*bOverrideParent*/false);
 	}
 
 	if (DataProps.bGetComponentCoordinates)
@@ -415,7 +415,7 @@ bool UPCGLandscapeData::ProjectPoint(const FTransform& InTransform, const FBox&
 
 	if (DataProps.bGetPhysicalMaterial && OutMetadata && LandscapeCollisionComponent)
 	{
-		if (FPCGMetadataAttribute<FSoftObjectPath>* PhysicalMaterialAttribute = OutMetadata->GetMutableTypedAttribute<FSoftObjectPath>(PCGWorldRayHitConstants::PhysicalMaterialReferenceAttribute))
+		if (FPCGMetadataAttribute<FSoftObjectPath>* PhysicalMaterialAttribute = OutMetadata->GetMutableTypedAttribute<FSoftObjectPath>(PCGWorldQueryConstants::PhysicalMaterialReferenceAttribute))
 		{
 			if(UPhysicalMaterial* PhysicalMaterial = LandscapeCollisionComponent->GetPhysicalMaterial(static_cast<float>(ComponentLocalPoint.X), static_cast<float>(ComponentLocalPoint.Y), EHeightfieldSource::Complex))
 			{
diff --git a/Engine/Plugins/PCG/Source/PCG/Private/Data/PCGWorldData.cpp b/Engine/Plugins/PCG/Source/PCG/Private/Data/PCGWorldData.cpp
--- a/Engine/Plugins/PCG/Source/PCG/Private/Data/PCGWorldData.cpp
+++ b/Engine/Plugins/PCG/Source/PCG/Private/Data/PCGWorldData.cpp
@@ -12,6 +12,7 @@
 #include "Grid/PCGLandscapeCache.h"
 #include "Helpers/PCGBlueprintHelpers.h"
 #include "Helpers/PCGHelpers.h"
+#include "Helpers/PCGWorldQueryHelpers.h"
 
 #include "LandscapeProxy.h"
 #include "Components/BrushComponent.h"
@@ -43,6 +44,11 @@ void FPCGWorldVolumetricQueryParams::Initialize()
 	FPCGWorldCommonQueryParams::Initialize();
 }
 
+void FPCGWorldRaycastQueryParams::Initialize()
+{
+	FPCGWorldCommonQueryParams::Initialize();
+}
+
 void FPCGWorldRayHitQueryParams::Initialize()
 {
 	FPCGWorldCommonQueryParams::Initialize();
@@ -59,6 +65,16 @@ void FPCGWorldCommonQueryParams::CommonPostLoad()
 }
 #endif // WITH_EDITOR
 
+void FPCGWorldRaycastQueryParams::PostSerialize(const FArchive& Ar)
+{
+#if WITH_EDITOR
+	if (Ar.IsLoading() && Ar.IsPersistent() && !Ar.HasAnyPortFlags(PPF_Duplicate | PPF_DuplicateForPIE))
+	{
+		CommonPostLoad();
+	}
+#endif
+}
+
 void FPCGWorldVolumetricQueryParams::PostSerialize(const FArchive& Ar)
 {
 #if WITH_EDITOR
@@ -102,99 +118,20 @@ bool UPCGWorldVolumetricData::SamplePoint(const FTransform& InTransform, const F
 	TArray<FOverlapResult> Overlaps;
 	/*bool bOverlaps =*/ World->OverlapMultiByObjectType(Overlaps, InTransform.TransformPosition(InBounds.GetCenter()), InTransform.GetRotation(), ObjectQueryParams, CollisionShape, Params);
 
-	for (const FOverlapResult& Overlap : Overlaps)
-	{
-		// Skip invisible walls / triggers / volumes
-		const UPrimitiveComponent* OverlappedComponent = Overlap.GetComponent();
-		if (OverlappedComponent->IsA<UBrushComponent>())
-		{
-			continue;
-		}
+	TOptional<FOverlapResult> Overlap = PCGWorldQueryHelpers::FilterOverlapResults(&QueryParams, OriginatingComponent, Overlaps);
 
-		// Skip "no collision" type actors
-		if (!OverlappedComponent->IsQueryCollisionEnabled() || OverlappedComponent->GetCollisionResponseToChannel(QueryParams.CollisionChannel) != ECR_Block)
-		{
-			continue;
-		}
-
-		// Skip to-be-cleaned-up PCG-created objects
-		if (OverlappedComponent->ComponentHasTag(PCGHelpers::MarkedForCleanupPCGTag) || (OverlappedComponent->GetOwner() && OverlappedComponent->GetOwner()->ActorHasTag(PCGHelpers::MarkedForCleanupPCGTag)))
-		{
-			continue;
-		}
-
-		// Optionally skip all PCG created objects
-		if (QueryParams.bIgnorePCGHits && (OverlappedComponent->ComponentHasTag(PCGHelpers::DefaultPCGTag) || (OverlappedComponent->GetOwner() && OverlappedComponent->GetOwner()->ActorHasTag(PCGHelpers::DefaultPCGActorTag))))
-		{
-			continue;
-		}
-
-		// Skip self-generated PCG objects optionally
-		if (QueryParams.bIgnoreSelfHits && OriginatingComponent.IsValid() && OverlappedComponent->ComponentTags.Contains(OriginatingComponent->GetFName()))
-		{
-			continue;
-		}
-
-		// Additional filter as provided in the QueryParams base class
-		if (QueryParams.ActorTagFilter != EPCGWorldQueryFilterByTag::NoTagFilter)
-		{
-			if (AActor* Actor = OverlappedComponent->GetOwner())
-			{
-				bool bFoundMatch = false;
-				for (const FName& Tag : Actor->Tags)
-				{
-					if (QueryParams.ParsedActorTagsList.Contains(Tag))
-					{
-						bFoundMatch = true;
-						break;
-					}
-				}
-
-				if (bFoundMatch != (QueryParams.ActorTagFilter == EPCGWorldQueryFilterByTag::IncludeTagged))
-				{
-					continue;
-				}
-			}
-			else if (QueryParams.ActorTagFilter == EPCGWorldQueryFilterByTag::IncludeTagged)
-			{
-				continue;
-			}
-		}
-
-		if (QueryParams.SelectLandscapeHits != EPCGWorldQuerySelectLandscapeHits::Include)
-		{
-			if (OverlappedComponent->GetOwner() && OverlappedComponent->GetOwner()->IsA<ALandscapeProxy>() != (QueryParams.SelectLandscapeHits == EPCGWorldQuerySelectLandscapeHits::Require))
-			{
-				continue;
-			}
-		}
+	// If searched for overlap and found one, or didn't search and didn't find one, set the point and return true. Otherwise, return false.
+	if (Overlap.IsSet() == QueryParams.bSearchForOverlap)
+	{
+		OutPoint = FPCGPoint(InTransform, 1.0f, UPCGBlueprintHelpers::ComputeSeedFromPosition(InTransform.GetLocation()));
+		OutPoint.SetLocalBounds(InBounds);
 
-		if (QueryParams.bSearchForOverlap)
+		if (Overlap.IsSet() && ActorOverlappedAttribute && Overlap.GetValue().GetActor())
 		{
-			OutPoint = FPCGPoint(InTransform, 1.0f, 0);
-			UPCGBlueprintHelpers::SetSeedFromPosition(OutPoint);
-			OutPoint.SetLocalBounds(InBounds);
-
-			if (ActorOverlappedAttribute && OverlappedComponent->GetOwner())
-			{
-				OutMetadata->InitializeOnSet(OutPoint.MetadataEntry);
-				ActorOverlappedAttribute->SetValue(OutPoint.MetadataEntry, FSoftObjectPath(OverlappedComponent->GetOwner()));
-			}
-
-			return true;
-		}
-		else
-		{
-			return false;
+			OutMetadata->InitializeOnSet(OutPoint.MetadataEntry);
+			ActorOverlappedAttribute->SetValue(OutPoint.MetadataEntry, FSoftObjectPath(Overlap.GetValue().GetActor()));
 		}
-	}
 
-	// No valid hits found
-	if (!QueryParams.bSearchForOverlap)
-	{
-		OutPoint = FPCGPoint(InTransform, 1.0f, 0);
-		UPCGBlueprintHelpers::SetSeedFromPosition(OutPoint);
-		OutPoint.SetLocalBounds(InBounds);
 		return true;
 	}
 	else
@@ -281,137 +218,43 @@ bool UPCGWorldRayHitData::SamplePoint(const FTransform& InTransform, const FBox&
 	// TODO: This seems to be a projection - along a direction. I suspect that UPCGWorldVolumetricData is the SamplePoint(), and this is the ProjectPoint() (in a direction)?
 	check(World.IsValid());
 
-	FPCGMetadataAttribute<FSoftObjectPath>* ActorHitAttribute = ((OutMetadata && QueryParams.bGetReferenceToActorHit) ? OutMetadata->GetMutableTypedAttribute<FSoftObjectPath>(PCGPointDataConstants::ActorReferenceAttribute) : nullptr);
-	FPCGMetadataAttribute<FSoftObjectPath>* PhysicalMaterialAttribute = ((OutMetadata && QueryParams.bGetReferenceToPhysicalMaterial) ? OutMetadata->GetMutableTypedAttribute<FSoftObjectPath>(PCGWorldRayHitConstants::PhysicalMaterialReferenceAttribute) : nullptr);
-
 	// Todo: consider prebuilding this
 	FCollisionObjectQueryParams ObjectQueryParams(QueryParams.CollisionChannel);
 	FCollisionQueryParams Params; // TODO: apply properties from the settings when/if they exist
 	Params.bTraceComplex = QueryParams.bTraceComplex;
-	Params.bReturnPhysicalMaterial = (PhysicalMaterialAttribute != nullptr);
+	Params.bReturnPhysicalMaterial = QueryParams.bGetReferenceToPhysicalMaterial;
 
 	FVector RayOrigin = InTransform.GetLocation() - ((InTransform.GetLocation() - QueryParams.RayOrigin) | QueryParams.RayDirection) * QueryParams.RayDirection;
 	FVector RayEnd = RayOrigin + QueryParams.RayDirection * QueryParams.RayLength;
 
 	TArray<FHitResult> Hits;
-	World->LineTraceMultiByObjectType(Hits, RayOrigin, RayEnd, ObjectQueryParams, Params);
-
-	for (const FHitResult& Hit : Hits)
+	if (World->LineTraceMultiByObjectType(Hits, RayOrigin, RayEnd, ObjectQueryParams, Params))
 	{
-		// Skip invisible walls / triggers / volumes
-		const UPrimitiveComponent* HitComponent = Hit.GetComponent();
-		if (HitComponent->IsA<UBrushComponent>())
-		{
-			continue;
-		}
+		TOptional<FHitResult> HitResult = PCGWorldQueryHelpers::FilterRayHitResults(&QueryParams, OriginatingComponent, Hits);
 
-		// Skip "No collision" type actors
-		if (!HitComponent->IsQueryCollisionEnabled() || HitComponent->GetCollisionResponseToChannel(QueryParams.CollisionChannel) != ECR_Block)
-		{
-			continue;
-		}
-		
-		// Skip to-be-cleaned-up PCG-created objects
-		if (HitComponent->ComponentHasTag(PCGHelpers::MarkedForCleanupPCGTag) || (HitComponent->GetOwner() && HitComponent->GetOwner()->ActorHasTag(PCGHelpers::MarkedForCleanupPCGTag)))
-		{
-			continue;
-		}
+		// TODO: Pre-create attributes within caller or find a better solution than per point
+		PCGWorldQueryHelpers::ApplyRayHitMetadata(HitResult, QueryParams, OutPoint, OutMetadata);
 
-		// Optionally skip all PCG created objects
-		if (QueryParams.bIgnorePCGHits && (HitComponent->ComponentHasTag(PCGHelpers::DefaultPCGTag) || (HitComponent->GetOwner() && HitComponent->GetOwner()->ActorHasTag(PCGHelpers::DefaultPCGActorTag))))
+		if (HitResult.IsSet())
 		{
-			continue;
-		}
+			const FHitResult& Hit = HitResult.GetValue();
 
-		// Skip self-generated PCG objects optionally
-		if (QueryParams.bIgnoreSelfHits && OriginatingComponent.IsValid() && HitComponent->ComponentTags.Contains(OriginatingComponent->GetFName()))
-		{
-			continue;
-		}
+			// Finally, fill in OutPoint - we're done
+			OutPoint = FPCGPoint(PCGWorldQueryHelpers::GetOrthonormalImpactTransform(Hit), 1.0f, UPCGBlueprintHelpers::ComputeSeedFromPosition(Hit.Location));
 
-		// Additional filter as provided in the QueryParams base class
-		if (QueryParams.ActorTagFilter != EPCGWorldQueryFilterByTag::NoTagFilter)
-		{
-			AActor* Actor = HitComponent->GetOwner();
-			
-			if (Actor)
-			{
-				bool bFoundMatch = false;
-				for (const FName& Tag : Actor->Tags)
-				{
-					if (QueryParams.ParsedActorTagsList.Contains(Tag))
-					{
-						bFoundMatch = true;
-						break;
-					}
-				}
+			const bool bApplyMetadataFromLandscape = QueryParams.bApplyMetadataFromLandscape && Hit.GetActor() && Hit.GetActor()->IsA<ALandscapeProxy>();
 
-				if (bFoundMatch != (QueryParams.ActorTagFilter == EPCGWorldQueryFilterByTag::IncludeTagged))
+			// TODO: generalize for other sources of metadata?
+			if (bApplyMetadataFromLandscape && OutMetadata && World->GetSubsystem<UPCGSubsystem>())
+			{
+				if (UPCGLandscapeCache* LandscapeCache = World->GetSubsystem<UPCGSubsystem>()->GetLandscapeCache())
 				{
-					continue;
+					LandscapeCache->SampleMetadataOnPoint(Cast<ALandscapeProxy>(Hit.Component->GetOwner()), OutPoint, OutMetadata);
 				}
 			}
-			else if (QueryParams.ActorTagFilter == EPCGWorldQueryFilterByTag::IncludeTagged)
-			{
-				continue;
-			}
-		}
-
-		bool bHitOnLandscape = false;
-		if(QueryParams.SelectLandscapeHits != EPCGWorldQuerySelectLandscapeHits::Include || QueryParams.bApplyMetadataFromLandscape || QueryParams.bIgnoreBackfaceHits)
-		{
-			bHitOnLandscape = HitComponent->GetOwner() && HitComponent->GetOwner()->IsA<ALandscapeProxy>();
-		}
-
-		if((bHitOnLandscape && QueryParams.SelectLandscapeHits == EPCGWorldQuerySelectLandscapeHits::Exclude) ||
-			(!bHitOnLandscape && QueryParams.SelectLandscapeHits == EPCGWorldQuerySelectLandscapeHits::Require))
-		{
-			continue;
-		}
-
-		// Optionally skip backface hits
-		if (QueryParams.bIgnoreBackfaceHits)
-		{
-			// If its a landscape, we cull if the normal is negative in Z direction (landscape normal is always the +Z axis). If not, then we cull if the impact normal and the ray are headed in the same direction
-			if (Hit.bStartPenetrating || (bHitOnLandscape && Hit.ImpactNormal.Z < 0) || (QueryParams.RayDirection).Dot(Hit.ImpactNormal) > 0)
-			{
-				continue;
-			}
-		}
-
-		// Finally, fill in OutPoint - we're done
-		// Implementation note: this uses the same orthonormalization process as the landscape cache
-		ensure(Hit.ImpactNormal.IsNormalized());
-		const FVector ArbitraryVector = (FMath::Abs(Hit.ImpactNormal.Y) < (1.f - UE_KINDA_SMALL_NUMBER) ? FVector::YAxisVector : FVector::ZAxisVector);
-		const FVector XAxis = (ArbitraryVector ^ Hit.ImpactNormal).GetSafeNormal();
-		const FVector YAxis = (Hit.ImpactNormal ^ XAxis);
-
-		OutPoint = FPCGPoint(FTransform(XAxis, YAxis, Hit.ImpactNormal, Hit.ImpactPoint), 1.0f, 0);
-		UPCGBlueprintHelpers::SetSeedFromPosition(OutPoint);
-
-		// TODO: generalize for other sources of metadata?
-		if (QueryParams.bApplyMetadataFromLandscape && bHitOnLandscape && OutMetadata && World->GetSubsystem<UPCGSubsystem>())
-		{
-			UPCGLandscapeCache* LandscapeCache = World->GetSubsystem<UPCGSubsystem>()->GetLandscapeCache();
-			if (LandscapeCache)
-			{
-				LandscapeCache->SampleMetadataOnPoint(Cast<ALandscapeProxy>(HitComponent->GetOwner()), OutPoint, OutMetadata);
-			}
-		}
 
-		if (ActorHitAttribute && HitComponent->GetOwner())
-		{
-			OutMetadata->InitializeOnSet(OutPoint.MetadataEntry);
-			ActorHitAttribute->SetValue(OutPoint.MetadataEntry, FSoftObjectPath(HitComponent->GetOwner()));
-		}
-
-		if (PhysicalMaterialAttribute && Hit.PhysMaterial.Get())
-		{
-			OutMetadata->InitializeOnSet(OutPoint.MetadataEntry);
-			PhysicalMaterialAttribute->SetValue(OutPoint.MetadataEntry, FSoftObjectPath(Hit.PhysMaterial.Get()));
+			return true;
 		}
-
-		return true;
 	}
 
 	return false;
diff --git a/Engine/Plugins/PCG/Source/PCG/Private/Elements/PCGSurfaceSampler.cpp b/Engine/Plugins/PCG/Source/PCG/Private/Elements/PCGSurfaceSampler.cpp
--- a/Engine/Plugins/PCG/Source/PCG/Private/Elements/PCGSurfaceSampler.cpp
+++ b/Engine/Plugins/PCG/Source/PCG/Private/Elements/PCGSurfaceSampler.cpp
@@ -335,20 +335,6 @@ namespace PCGSurfaceSamplerHelpers
 	using ContextType = FPCGSurfaceSamplerElement::ContextType;
 	using ExecStateType = FPCGSurfaceSamplerElement::ExecStateType;
 
-	const UPCGSpatialData* FindBoundingShape(const FPCGContext* Context, bool& bOutUnionWasCreated)
-	{
-		const UPCGSpatialData* BoundingShape = Context->InputData.GetSpatialUnionOfInputsByPin(PCGSurfaceSamplerConstants::BoundingShapeLabel, bOutUnionWasCreated);
-
-		// Fallback to getting bounds from actor
-		if (!BoundingShape && Context->SourceComponent.IsValid())
-		{
-			check(bOutUnionWasCreated == false);
-			BoundingShape = Cast<UPCGSpatialData>(Context->SourceComponent->GetActorPCGData());
-		}
-
-		return BoundingShape;
-	}
-
 	EPCGTimeSliceInitResult InitializePerExecutionData(ContextType* Context, ExecStateType& OutState)
 	{
 		TRACE_CPUPROFILER_EVENT_SCOPE(FPCGSurfaceSamplerElement::InitializePerExecutionData);
@@ -374,29 +360,27 @@ namespace PCGSurfaceSamplerHelpers
 		}
 
 		// Grab the Bounding Shape input if there is one.
-		const TArray<FPCGTaggedData> BoundingShapeInputs = Context->InputData.GetInputsByPin(PCGSurfaceSamplerConstants::BoundingShapeLabel);
-		FBox& BoundingShapeBounds = OutState.BoundingShapeBounds;
 		if (!Settings->bUnbounded)
 		{
 			bool bUnionWasCreated;
-			OutState.BoundingShape = FindBoundingShape(Context, bUnionWasCreated);
-			if (OutState.BoundingShape && bUnionWasCreated)
-			{
-				Context->TrackObject(OutState.BoundingShape);
-			}
-
+			OutState.BoundingShape = PCGSettingsHelpers::ComputeBoundingShape(Context, PCGSurfaceSamplerConstants::BoundingShapeLabel, bUnionWasCreated);
 			if (OutState.BoundingShape)
 			{
-				BoundingShapeBounds = OutState.BoundingShape->GetBounds();
+				if (bUnionWasCreated)
+				{
+					Context->TrackObject(OutState.BoundingShape);
+				}
+
+				OutState.BoundingShapeBounds = OutState.BoundingShape->GetBounds();
 			}
 
-			if (!BoundingShapeBounds.IsValid)
+			if (!OutState.BoundingShapeBounds.IsValid)
 			{
 				// The bounding shape bounds is invalid, such as an empty intersection, so no operation will need to be performed.
 				return EPCGTimeSliceInitResult::NoOperation;
 			}
 		}
-		else if (BoundingShapeInputs.Num() > 0)
+		else if (Context->InputData.GetInputsByPin(PCGSurfaceSamplerConstants::BoundingShapeLabel).Num() > 0)
 		{
 			PCGE_LOG_C(Verbose, LogOnly, Context, LOCTEXT("BoundsIgnored", "The bounds of the Bounding Shape input pin will be ignored because the Unbounded option is enabled."));
 		}
diff --git a/Engine/Plugins/PCG/Source/PCG/Private/Elements/PCGWorldQuery.cpp b/Engine/Plugins/PCG/Source/PCG/Private/Elements/PCGWorldQuery.cpp
--- a/Engine/Plugins/PCG/Source/PCG/Private/Elements/PCGWorldQuery.cpp
+++ b/Engine/Plugins/PCG/Source/PCG/Private/Elements/PCGWorldQuery.cpp
@@ -8,6 +8,7 @@
 #include "Data/PCGPointData.h"
 #include "Grid/PCGLandscapeCache.h"
 #include "Helpers/PCGHelpers.h"
+#include "Helpers/PCGWorldQueryHelpers.h"
 
 #include "GameFramework/Actor.h"
 #include "LandscapeProxy.h"
@@ -149,7 +150,7 @@ bool FPCGWorldRayHitQueryElement::ExecuteInternal(FPCGContext* Context) const
 
 	if (Data->QueryParams.bGetReferenceToPhysicalMaterial && Data->Metadata)
 	{
-		Data->Metadata->FindOrCreateAttribute(PCGWorldRayHitConstants::PhysicalMaterialReferenceAttribute, FSoftObjectPath(), /*bAllowInterpolation=*/false, /*bOverrideParent=*/false);
+		Data->Metadata->FindOrCreateAttribute(PCGWorldQueryConstants::PhysicalMaterialReferenceAttribute, FSoftObjectPath(), /*bAllowInterpolation=*/false, /*bOverrideParent=*/false);
 	}
 	else
 	{
diff --git a/Engine/Plugins/PCG/Source/PCG/Private/Elements/PCGWorldRaycast.cpp b/Engine/Plugins/PCG/Source/PCG/Private/Elements/PCGWorldRaycast.cpp
new file mode 100644
--- /dev/null
+++ b/Engine/Plugins/PCG/Source/PCG/Private/Elements/PCGWorldRaycast.cpp
@@ -0,0 +1,366 @@
+// Copyright Epic Games, Inc. All Rights Reserved.
+
+#include "Elements/PCGWorldRaycast.h"
+
+#include "PCGComponent.h"
+#include "PCGContext.h"
+#include "PCGPoint.h"
+#include "Data/PCGPointData.h"
+#include "Helpers/PCGBlueprintHelpers.h"
+#include "Helpers/PCGHelpers.h"
+#include "Helpers/PCGSettingsHelpers.h"
+#include "Helpers/PCGWorldQueryHelpers.h"
+#include "Metadata/Accessors/PCGAttributeAccessorHelpers.h"
+
+#include "Engine/HitResult.h"
+#include "Engine/World.h"
+
+#include UE_INLINE_GENERATED_CPP_BY_NAME(PCGWorldRaycast)
+
+#define LOCTEXT_NAMESPACE "PCGWorldRaycastElement"
+
+namespace PCGWorldRaycastElementConstants
+{
+	const FName OriginsInputLabel = TEXT("Origins");
+	const FName EndPointsInputLabel = TEXT("End Points");
+	const FName BoundingShapeLabel = TEXT("Bounding Shape");
+}
+
+UPCGWorldRaycastElementSettings::UPCGWorldRaycastElementSettings()
+{
+	OriginInputAttribute.SetPointProperty(EPCGPointProperties::Position);
+	EndPointAttribute.SetPointProperty(EPCGPointProperties::Position);
+	// Tracing rays along the normal is the common use case
+	RayDirectionAttribute.Update("$Rotation.Up");
+}
+
+TArray<FPCGPinProperties> UPCGWorldRaycastElementSettings::InputPinProperties() const
+{
+	TArray<FPCGPinProperties> PinProperties;
+
+	FPCGPinProperties& OriginPointsInputPin = PinProperties.Emplace_GetRef(PCGWorldRaycastElementConstants::OriginsInputLabel, EPCGDataType::PointOrParam);
+	OriginPointsInputPin.SetRequiredPin();
+#if WITH_EDITOR
+	OriginPointsInputPin.Tooltip = LOCTEXT("OriginPointsInputPinTooltip", "Points acting as the point origin of each individual ray. Only point data from this pin may be forwarded.");
+#endif // WITH_EDITOR
+
+	if (RaycastMode == EPCGWorldRaycastMode::Segments)
+	{
+		FPCGPinProperties& EndPointsInputPin = PinProperties.Emplace_GetRef(PCGWorldRaycastElementConstants::EndPointsInputLabel, EPCGDataType::PointOrParam);
+#if WITH_EDITOR
+		EndPointsInputPin.Tooltip = LOCTEXT("EndPointsInputPinTooltip", "Points acting as the end point of each individual ray.");
+#endif // WITH_EDITOR
+	}
+
+	FPCGPinProperties& BoundingShapeInputPin = PinProperties.Emplace_GetRef(PCGWorldRaycastElementConstants::BoundingShapeLabel, EPCGDataType::Spatial);
+#if WITH_EDITOR
+	BoundingShapeInputPin.Tooltip = LOCTEXT("BoundingShapeInputPinTooltip", "All projected points must be contained within this shape. If this input is omitted then bounds will be taken from the actor so that points are contained within actor bounds.");
+#endif // WITH_EDITOR
+
+	return PinProperties;
+}
+
+TArray<FPCGPinProperties> UPCGWorldRaycastElementSettings::OutputPinProperties() const
+{
+	return DefaultPointOutputPinProperties();
+}
+
+FPCGElementPtr UPCGWorldRaycastElementSettings::CreateElement() const
+{
+	return MakeShared<FPCGWorldRaycastElement>();
+}
+
+bool FPCGWorldRaycastElement::PrepareDataInternal(FPCGContext* InContext) const
+{
+	ContextType* Context = static_cast<ContextType*>(InContext);
+
+	const UPCGWorldRaycastElementSettings* Settings = Context->GetInputSettings<UPCGWorldRaycastElementSettings>();
+	check(Settings);
+
+	const TArray<FPCGTaggedData> OriginsInputData = Context->InputData.GetInputsByPin(PCGWorldRaycastElementConstants::OriginsInputLabel);
+	const TArray<FPCGTaggedData> EndPointsInputData = Context->InputData.GetInputsByPin(PCGWorldRaycastElementConstants::EndPointsInputLabel);
+
+	const EPCGTimeSliceInitResult ExecutionInitResult = Context->InitializePerExecutionState([&OriginsInputData, &EndPointsInputData, Settings](ContextType* Context, ExecStateType& OutState)
+	{
+		// With no origins data, early out
+		if (OriginsInputData.IsEmpty())
+		{
+			return EPCGTimeSliceInitResult::NoOperation;
+		}
+
+		if (!Settings->bUnbounded)
+		{
+			bool bUnionWasCreated;
+			const UPCGSpatialData* BoundingShape = PCGSettingsHelpers::ComputeBoundingShape(Context, PCGWorldRaycastElementConstants::BoundingShapeLabel, bUnionWasCreated);
+			OutState.Bounds = BoundingShape ? BoundingShape->GetBounds() : FBox(EForceInit::ForceInit);
+			if (!OutState.Bounds.IsValid)
+			{
+				// The bounding shape bounds is invalid, such as an empty intersection, so no operation will need to be performed.
+				return EPCGTimeSliceInitResult::NoOperation;
+			}
+		}
+
+		if (Settings->RaycastMode == EPCGWorldRaycastMode::Segments && EndPointsInputData.Num() != OriginsInputData.Num())
+		{
+			PCGLog::LogErrorOnGraph(LOCTEXT("InvalidEndPointDataCount", "End point input data count must match the Origins input data count."), Context);
+			return EPCGTimeSliceInitResult::NoOperation;
+		}
+
+		OutState.CollisionQueryParams.bTraceComplex = Settings->WorldQueryParams.bTraceComplex;
+		OutState.CollisionQueryParams.bReturnPhysicalMaterial = Settings->WorldQueryParams.bGetReferenceToPhysicalMaterial;
+		OutState.CollisionObjectQueryParams = FCollisionObjectQueryParams(Settings->WorldQueryParams.CollisionChannel);
+
+		return EPCGTimeSliceInitResult::Success;
+	});
+
+	if (ExecutionInitResult != EPCGTimeSliceInitResult::Success)
+	{
+		return true;
+	}
+
+	Context->InitializePerIterationStates(OriginsInputData.Num(), [Context, Settings, OriginsInputData, EndPointsInputData](IterStateType& OutState, const ExecStateType&, int32 Index)
+	{
+		const UPCGData* OriginsData = OriginsInputData[Index].Data;
+
+		// Accept only point or param data
+		if (!OriginsData || !OriginsData->ConstMetadata())
+		{
+			PCGLog::LogErrorOnGraph(FText::Format(LOCTEXT("NoDataFoundOnOriginsInput", "No valid data found on Origins input {0}"), Index), Context);
+			return EPCGTimeSliceInitResult::NoOperation;
+		}
+
+		// Given a selector, will populate a given array view with attributes
+		auto CacheArray = [Settings, Context]<typename Type>(const UPCGData* InData, const FPCGAttributePropertyInputSelector& InSelector, TArray<Type>& OutArray)
+		{
+			const FPCGAttributePropertyInputSelector Selector = InSelector.CopyAndFixLast(InData);
+
+			TUniquePtr<const IPCGAttributeAccessor> Accessor = PCGAttributeAccessorHelpers::CreateConstAccessor(InData, Selector);
+			TUniquePtr<const IPCGAttributeAccessorKeys> Keys = PCGAttributeAccessorHelpers::CreateConstKeys(InData, Selector);
+			if (!Accessor.IsValid() || !Keys.IsValid())
+			{
+				PCGLog::LogWarningOnGraph(FText::Format(LOCTEXT("AttributeDoesNotExist", "Attribute/property '{0}' does not exist - skipped."), Selector.GetDisplayText()), Context);
+				return false;
+			}
+
+			OutArray.Empty();
+			OutArray.SetNumUninitialized(Keys->GetNum());
+
+			if (!Accessor->GetRange<Type>(OutArray, 0, *Keys, EPCGAttributeAccessorFlags::AllowBroadcast | EPCGAttributeAccessorFlags::AllowConstructible))
+			{
+				PCGLog::LogErrorOnGraph(FText::Format(LOCTEXT("InvalidAttributeValues", "Could not get the range of values for attribute/property '{0}'."), Selector.GetDisplayText()), Context);
+				return false;
+			}
+
+			return true;
+		};
+
+		// --- Find the ray origin ---
+		if (!CacheArray(OriginsData, Settings->OriginInputAttribute, OutState.CachedRayOrigins))
+		{
+			return EPCGTimeSliceInitResult::NoOperation;
+		}
+
+		// --- Find the ray direction vectors ---
+		// Calculate the vectors from origin to end points
+		if (Settings->RaycastMode == EPCGWorldRaycastMode::Segments)
+		{
+			TObjectPtr<const UPCGData> EndPointsData = EndPointsInputData[Index].Data;
+			if (!EndPointsData)
+			{
+				PCGLog::LogErrorOnGraph(FText::Format(LOCTEXT("NoDataFoundOnEndPointsInput", "No valid data found on End Points input {0}"), Index), Context);
+				return EPCGTimeSliceInitResult::NoOperation;
+			}
+
+			TArray<FVector> EndPoints;
+			if (!CacheArray(EndPointsData, Settings->EndPointAttribute, EndPoints))
+			{
+				return EPCGTimeSliceInitResult::NoOperation;
+			}
+
+			// Support N:N, 1:N, N:1
+			if (OutState.CachedRayOrigins.Num() != 1 && EndPoints.Num() != 1 && EndPoints.Num() != OutState.CachedRayOrigins.Num())
+			{
+				PCGLog::LogErrorOnGraph(LOCTEXT("InvalidEndPointCount", "The End Points input count must be 1 or directly match the Origins input count."), Context);
+				return EPCGTimeSliceInitResult::NoOperation;
+			}
+
+			const int32 PointCount = FMath::Max(OutState.CachedRayOrigins.Num(), EndPoints.Num());
+
+			OutState.CachedRayVectors.Empty();
+			OutState.CachedRayVectors.Reserve(PointCount);
+
+			// Find the ray vector from the origin to end points
+			for (int PointIndex = 0; PointIndex < PointCount; ++PointIndex)
+			{
+				const FVector& Origin = OutState.CachedRayOrigins[OutState.CachedRayOrigins.Num() == 1 ? 0 : PointIndex];
+				const FVector& EndPoint = EndPoints[EndPoints.Num() == 1 ? 0 : PointIndex];
+				OutState.CachedRayVectors.Emplace((EndPoint - Origin));
+			}
+		}
+		else if (Settings->bOverrideRayDirections) // Use a selector to determine the (non-normalized) directions
+		{
+			const FPCGAttributePropertyInputSelector DirectionSelector = Settings->RayDirectionAttribute.CopyAndFixLast(OriginsData);
+
+			if (!CacheArray(OriginsData, DirectionSelector, OutState.CachedRayVectors))
+			{
+				return EPCGTimeSliceInitResult::NoOperation;
+			}
+		}
+		else // Use the single property
+		{
+			OutState.CachedRayVectors.Empty();
+			OutState.CachedRayVectors.Emplace(Settings->RayDirection);
+		}
+
+		// --- Find the ray length ---
+		// If the ray is infinite or user-selected length, we need to normalize the vector
+		if (Settings->RaycastMode == EPCGWorldRaycastMode::Infinite || Settings->RaycastMode == EPCGWorldRaycastMode::NormalizedWithLength)
+		{
+			// Normalize the directions and store locally in the state
+			for (FVector& Ray : OutState.CachedRayVectors)
+			{
+				Ray.Normalize();
+			}
+
+			// If the length is infinite, it will be a fixed value at execution
+			if (Settings->RaycastMode == EPCGWorldRaycastMode::Infinite)
+			{
+				return EPCGTimeSliceInitResult::Success;
+			}
+
+			// Account for user length property. Use a selector to determine the (non-normalized) directions
+			if (Settings->bOverrideRayLengths)
+			{
+				const FPCGAttributePropertyInputSelector LengthSelector = Settings->RayLengthAttribute.CopyAndFixLast(OriginsData);
+
+				if (!CacheArray(OriginsData, LengthSelector, OutState.CachedRayLengths))
+				{
+					return EPCGTimeSliceInitResult::NoOperation;
+				}
+			}
+			else
+			{
+				OutState.CachedRayLengths.Empty();
+				OutState.CachedRayLengths.Emplace(Settings->RayLength);
+			}
+		}
+
+		return EPCGTimeSliceInitResult::Success;
+	});
+
+	return true;
+}
+
+bool FPCGWorldRaycastElement::ExecuteInternal(FPCGContext* InContext) const
+{
+	TRACE_CPUPROFILER_EVENT_SCOPE(FPCGWorldRaycastElement::Execute);
+
+	ContextType* Context = static_cast<ContextType*>(InContext);
+
+	if (!Context->DataIsPreparedForExecution())
+	{
+		return true;
+	}
+
+	return ExecuteSlice(Context, [](ContextType* Context, const ExecStateType& ExecState, IterStateType& IterState, const uint32 IterIndex)
+	{
+		if (Context->GetIterationStateResult(IterIndex) == EPCGTimeSliceInitResult::NoOperation)
+		{
+			return true;
+		}
+
+		const UPCGWorldRaycastElementSettings* Settings = Context->GetInputSettings<UPCGWorldRaycastElementSettings>();
+		check(Settings);
+
+		UWorld* World = Context->SourceComponent->GetWorld();
+		check(World);
+
+		const TArray<FPCGTaggedData> OriginInputs = Context->InputData.GetInputsByPin(PCGWorldRaycastElementConstants::OriginsInputLabel);
+		const FPCGTaggedData& OriginInputData = OriginInputs[IterIndex];
+		const UPCGPointData* OriginInputPointData = Cast<UPCGPointData>(OriginInputData.Data);
+		const TArray<FPCGPoint>* OriginInputPoints = OriginInputPointData ? &OriginInputPointData->GetPoints() : nullptr;
+
+		TArray<FPCGTaggedData>& OutputData = Context->OutputData.TaggedData;
+		UPCGPointData* OutputPointData = NewObject<UPCGPointData>();
+		// TODO: This should initialize from the End Points point data, if Origin point data doesn't exist
+		OutputPointData->InitializeFromData(OriginInputPointData);
+		FPCGTaggedData& Output = OutputData.Emplace_GetRef(OriginInputData);
+		Output.Data = OutputPointData;
+
+		UPCGMetadata* OutMetadata = OutputPointData->Metadata;
+
+		TArray<FPCGPoint>& OutputPoints = OutputPointData->GetMutablePoints();
+		// If all points will be kept, the size can be pre-allocated
+		if (Settings->bKeepOriginalPointOnMiss)
+		{
+			OutputPoints.Reserve(IterState.CachedRayOrigins.Num());
+		}
+
+		const FPCGWorldRaycastQueryParams& WorldQueryParams = Settings->WorldQueryParams;
+
+		const int32 PointCount = FMath::Max(IterState.CachedRayOrigins.Num(), IterState.CachedRayVectors.Num());
+		for (int64 PointIndex = 0; PointIndex < PointCount; ++PointIndex)
+		{
+			const FVector& Origin = IterState.CachedRayOrigins[IterState.CachedRayOrigins.Num() == 1 ? 0 : PointIndex];
+			const FVector& Direction = IterState.CachedRayVectors[IterState.CachedRayVectors.Num() == 1 ? 0 : PointIndex];
+			double Magnitude = 1.0;
+
+			if (Settings->RaycastMode == EPCGWorldRaycastMode::Infinite)
+			{
+				Magnitude = UE_LARGE_WORLD_MAX;
+			}
+			else if (Settings->RaycastMode == EPCGWorldRaycastMode::NormalizedWithLength)
+			{
+				Magnitude = IterState.CachedRayLengths[IterState.CachedRayLengths.Num() == 1 ? 0 : PointIndex];
+			}
+
+			FVector RayVector = Direction * Magnitude;
+			// Discard point if the direction + magnitude vector is the ZeroVector
+			if (RayVector.IsNearlyZero())
+			{
+				continue;
+			}
+
+			TArray<FHitResult> Hits;
+			World->LineTraceMultiByObjectType(Hits, Origin, Origin + RayVector, ExecState.CollisionObjectQueryParams, ExecState.CollisionQueryParams);
+			const TOptional<FHitResult>& HitResult = PCGWorldQueryHelpers::FilterRayHitResults(&WorldQueryParams, Context->SourceComponent, Hits);
+
+			// If there's a hit, and it's within bounds, keep the point
+			if (HitResult.IsSet() && (Settings->bUnbounded || PCGHelpers::IsInsideBounds(ExecState.Bounds, HitResult.GetValue().ImpactPoint)))
+			{
+				const FHitResult& Hit = HitResult.GetValue();
+				// Copy the input point, if it exists
+				FPCGPoint& OutPoint = OriginInputPoints
+					? OutputPoints.Emplace_GetRef((*OriginInputPoints)[OriginInputPoints->Num() == 1 ? 0 : PointIndex])
+					: OutputPoints.Emplace_GetRef();
+
+				OutPoint.Transform = PCGWorldQueryHelpers::GetOrthonormalImpactTransform(Hit);
+				OutPoint.Density = 1.f;
+				UPCGBlueprintHelpers::SetSeedFromPosition(OutPoint);
+			}
+			else if (Settings->bKeepOriginalPointOnMiss)
+			{
+				// Copy the original point if possible
+				if (OriginInputPoints)
+				{
+					OutputPoints.Emplace_GetRef((*OriginInputPoints)[OriginInputPoints->Num() == 1 ? 0 : PointIndex]).Density = 0.f;
+				}
+				else
+				{
+					OutputPoints.Emplace(FTransform(Origin), /*InDensity=*/0.f, UPCGBlueprintHelpers::ComputeSeedFromPosition(Origin));
+				}
+			}
+			else
+			{
+				continue;
+			}
+
+			// TODO: Pre-create the attributes outside the loop once the helper exists in PCGWorldQueryHelpers
+			PCGWorldQueryHelpers::ApplyRayHitMetadata(HitResult, WorldQueryParams, OutputPoints.Last(), OutMetadata);
+		}
+
+		return true;
+	});
+}
+
+#undef LOCTEXT_NAMESPACE
diff --git a/Engine/Plugins/PCG/Source/PCG/Private/Helpers/PCGWorldQueryHelpers.cpp b/Engine/Plugins/PCG/Source/PCG/Private/Helpers/PCGWorldQueryHelpers.cpp
new file mode 100644
--- /dev/null
+++ b/Engine/Plugins/PCG/Source/PCG/Private/Helpers/PCGWorldQueryHelpers.cpp
@@ -0,0 +1,232 @@
+// Copyright Epic Games, Inc. All Rights Reserved.
+
+#include "Helpers/PCGWorldQueryHelpers.h"
+
+#include "PCGComponent.h"
+#include "Data/PCGPointData.h"
+#include "Data/PCGWorldData.h"
+#include "Helpers/PCGHelpers.h"
+
+#include "LandscapeProxy.h"
+#include "Components/BrushComponent.h"
+#include "Components/PrimitiveComponent.h"
+#include "Engine/HitResult.h"
+#include "Engine/OverlapResult.h"
+#include "PhysicalMaterials/PhysicalMaterial.h"
+#include "UObject/SoftObjectPath.h"
+
+namespace PCGWorldQueryHelpers
+{
+	FTransform GetOrthonormalImpactTransform(const FHitResult& Hit)
+	{
+		// Implementation note: this uses the same orthonormalization process as the landscape cache
+		ensure(Hit.ImpactNormal.IsNormalized());
+		const FVector ArbitraryVector = (FMath::Abs(Hit.ImpactNormal.Y) < (1.f - UE_KINDA_SMALL_NUMBER) ? FVector::YAxisVector : FVector::ZAxisVector);
+		const FVector XAxis = (ArbitraryVector ^ Hit.ImpactNormal).GetSafeNormal();
+		const FVector YAxis = (Hit.ImpactNormal ^ XAxis);
+
+		return FTransform(XAxis, YAxis, Hit.ImpactNormal, Hit.ImpactPoint);
+	}
+
+	bool FilterCommonQueryResults(
+		const FPCGWorldCommonQueryParams* QueryParams,
+		const UPrimitiveComponent* TriggeredComponent,
+		const TWeakObjectPtr<UPCGComponent> OriginatingComponent)
+	{
+		check(QueryParams && TriggeredComponent);
+
+		// Skip invisible walls / triggers / volumes
+		if (TriggeredComponent->IsA<UBrushComponent>())
+		{
+			return false;
+		}
+
+		// Skip "No collision" type actors
+		if (!TriggeredComponent->IsQueryCollisionEnabled() || TriggeredComponent->GetCollisionResponseToChannel(QueryParams->CollisionChannel) != ECR_Block)
+		{
+			return false;
+		}
+
+		// Skip to-be-cleaned-up PCG-created objects
+		if (TriggeredComponent->ComponentHasTag(PCGHelpers::MarkedForCleanupPCGTag) || (TriggeredComponent->GetOwner() && TriggeredComponent->GetOwner()->ActorHasTag(PCGHelpers::MarkedForCleanupPCGTag)))
+		{
+			return false;
+		}
+
+		// Optionally skip all PCG created objects
+		if (QueryParams->bIgnorePCGHits && (TriggeredComponent->ComponentHasTag(PCGHelpers::DefaultPCGTag) || (TriggeredComponent->GetOwner() && TriggeredComponent->GetOwner()->ActorHasTag(PCGHelpers::DefaultPCGActorTag))))
+		{
+			return false;
+		}
+
+		// Skip self-generated PCG objects optionally
+		if (QueryParams->bIgnoreSelfHits && OriginatingComponent.IsValid() && TriggeredComponent->ComponentTags.Contains(OriginatingComponent->GetFName()))
+		{
+			return false;
+		}
+
+		// Additional filter as provided in the QueryParams base class
+		if (QueryParams->ActorTagFilter != EPCGWorldQueryFilterByTag::NoTagFilter)
+		{
+			if (AActor* Actor = TriggeredComponent->GetOwner())
+			{
+				bool bFoundMatch = false;
+				for (const FName& Tag : Actor->Tags)
+				{
+					if (QueryParams->ParsedActorTagsList.Contains(Tag))
+					{
+						bFoundMatch = true;
+						break;
+					}
+				}
+
+				if (bFoundMatch != (QueryParams->ActorTagFilter == EPCGWorldQueryFilterByTag::IncludeTagged))
+				{
+					return false;
+				}
+			}
+			else if (QueryParams->ActorTagFilter == EPCGWorldQueryFilterByTag::IncludeTagged)
+			{
+				return false;
+			}
+		}
+
+		// Landscape or not, include it if this is true
+		if (QueryParams->SelectLandscapeHits == EPCGWorldQuerySelectLandscapeHits::Include)
+		{
+			return true;
+		}
+
+		bool bTriggeredOnLandscape = TriggeredComponent->GetOwner() && TriggeredComponent->GetOwner()->IsA<ALandscapeProxy>();
+
+		// If excluding, skip. If requiring, and it's not a landscape, skip.
+		if ((bTriggeredOnLandscape && QueryParams->SelectLandscapeHits == EPCGWorldQuerySelectLandscapeHits::Exclude) ||
+			(!bTriggeredOnLandscape && QueryParams->SelectLandscapeHits == EPCGWorldQuerySelectLandscapeHits::Require))
+		{
+			return false;
+		}
+
+		return true;
+	}
+
+	TOptional<FHitResult> FilterRayHitResults(
+		const FPCGWorldRaycastQueryParams* QueryParams,
+		const TWeakObjectPtr<UPCGComponent> OriginatingComponent,
+		const TArray<FHitResult>& HitResults)
+	{
+		check(QueryParams);
+
+		for (const FHitResult& Hit : HitResults)
+		{
+			const UPrimitiveComponent* HitComponent = Hit.GetComponent();
+
+			if (!FilterCommonQueryResults(QueryParams, HitComponent, OriginatingComponent))
+			{
+				continue;
+			}
+
+			// Optionally skip backface hits
+			if (QueryParams->bIgnoreBackfaceHits)
+			{
+				// If it's a landscape, we cull if the normal is negative in Z direction (landscape normal is always the +Z axis). If not, then we cull if the impact normal and the ray are headed in the same direction
+				if (Hit.bStartPenetrating || (Hit.GetActor() && Hit.GetActor()->IsA<ALandscapeProxy>() && Hit.ImpactNormal.Z < 0) || (Hit.TraceEnd - Hit.TraceStart).Dot(Hit.ImpactNormal) > 0)
+				{
+					continue;
+				}
+			}
+
+			return Hit;
+		}
+
+		return {};
+	}
+
+	TOptional<FOverlapResult> FilterOverlapResults(
+		const FPCGWorldVolumetricQueryParams* QueryParams,
+		const TWeakObjectPtr<UPCGComponent> OriginatingComponent,
+		const TArray<FOverlapResult>& OverlapResults)
+	{
+		check(QueryParams);
+
+		for (const FOverlapResult& Overlap : OverlapResults)
+		{
+			const UPrimitiveComponent* OverlappedComponent = Overlap.GetComponent();
+
+			if (!FilterCommonQueryResults(QueryParams, OverlappedComponent, OriginatingComponent))
+			{
+				continue;
+			}
+
+			return Overlap;
+		}
+
+		return {};
+	}
+
+	// TODO: Add an option to create and apply ranges
+	bool ApplyRayHitMetadata(
+		const TOptional<FHitResult>& HitResult,
+		const FPCGWorldRaycastQueryParams& QueryParams,
+		FPCGPoint& OutPoint,
+		UPCGMetadata* OutMetadata,
+		bool bShouldCreateAttributes)
+	{
+		if (!OutMetadata)
+		{
+			return false;
+		}
+
+		// Check for the attributes are added, if not, then add them all in one go
+		auto CreateAttribute = [OutMetadata]<typename Type>(FName AttributeName, bool bShouldCreate, const Type& DefaultValue)
+		{
+			if (bShouldCreate && !OutMetadata->HasAttribute(AttributeName))
+			{
+				OutMetadata->CreateAttribute<Type>(AttributeName, DefaultValue, /*bAllowsInterpolation=*/true, /*bOverrideParent=*/false);
+			}
+		};
+
+		// TODO: Create a new helper to create attributes outside the point loops
+		if (bShouldCreateAttributes)
+		{
+			CreateAttribute(PCGWorldQueryConstants::ImpactAttribute, QueryParams.bGetImpact, false);
+			CreateAttribute(PCGWorldQueryConstants::DistanceAttribute, QueryParams.bGetDistance, 0.0);
+			CreateAttribute(PCGWorldQueryConstants::ImpactNormalAttribute, QueryParams.bGetImpactNormal, FVector());
+			CreateAttribute(PCGPointDataConstants::ActorReferenceAttribute, QueryParams.bGetReferenceToActorHit, FSoftObjectPath());
+			CreateAttribute(PCGWorldQueryConstants::PhysicalMaterialReferenceAttribute, QueryParams.bGetReferenceToPhysicalMaterial, FSoftObjectPath());
+		}
+
+		auto ApplyAttribute = [&OutPoint, OutMetadata]<typename Type>(FName AttributeName, const Type& Value, bool bShouldApply = true)
+		{
+			if (!bShouldApply)
+			{
+				return true;
+			}
+
+			if (FPCGMetadataAttribute<Type>* Attribute = OutMetadata->GetMutableTypedAttribute<Type>(AttributeName))
+			{
+				OutMetadata->InitializeOnSet(OutPoint.MetadataEntry);
+				Attribute->SetValue(OutPoint.MetadataEntry, Value);
+
+				return true;
+			}
+
+			return false;
+		};
+
+		if (!HitResult.IsSet())
+		{
+			ApplyAttribute(PCGWorldQueryConstants::ImpactAttribute, /*Value=*/false, QueryParams.bGetImpact);
+			return true;
+		}
+
+		const FHitResult& Hit = HitResult.GetValue();
+
+		ApplyAttribute(PCGWorldQueryConstants::ImpactAttribute, /*Value=*/true, QueryParams.bGetImpact);
+		ApplyAttribute(PCGWorldQueryConstants::DistanceAttribute, Hit.Distance, QueryParams.bGetDistance);
+		ApplyAttribute(PCGWorldQueryConstants::ImpactNormalAttribute, FVector(Hit.ImpactNormal), QueryParams.bGetImpactNormal);
+		ApplyAttribute(PCGPointDataConstants::ActorReferenceAttribute, FSoftObjectPath(Hit.GetActor()), QueryParams.bGetReferenceToActorHit);
+		ApplyAttribute(PCGWorldQueryConstants::PhysicalMaterialReferenceAttribute, FSoftObjectPath(Hit.PhysMaterial.Get()), QueryParams.bGetReferenceToPhysicalMaterial);
+
+		return true;
+	}
+}
diff --git a/Engine/Plugins/PCG/Source/PCG/Public/Data/PCGWorldData.h b/Engine/Plugins/PCG/Source/PCG/Public/Data/PCGWorldData.h
--- a/Engine/Plugins/PCG/Source/PCG/Public/Data/PCGWorldData.h
+++ b/Engine/Plugins/PCG/Source/PCG/Public/Data/PCGWorldData.h
@@ -4,6 +4,7 @@
 
 #include "PCGVolumeData.h"
 #include "PCGSurfaceData.h"
+#include "Helpers/PCGWorldQueryHelpers.h"
 
 #include "Engine/EngineTypes.h"
 
@@ -33,7 +34,8 @@
 
 namespace PCGWorldRayHitConstants
 {
-	const FName PhysicalMaterialReferenceAttribute = TEXT("PhysicalMaterial");
+	UE_DEPRECATED(5.5, "Please use 'PCGWorldQueryConstants::PhysicalMaterialReferenceAttribute' instead.")
+	const FName PhysicalMaterialReferenceAttribute = PCGWorldQueryConstants::PhysicalMaterialReferenceAttribute;
 }
 
 USTRUCT(BlueprintType)
@@ -46,10 +48,10 @@
 #endif
 
 	/** If true, will ignore hits/overlaps on content created from PCG. */
-	UPROPERTY(BlueprintReadWrite, EditAnywhere, Category = Data, meta = (PCG_Overridable))
+	UPROPERTY(BlueprintReadWrite, EditAnywhere, Category = "Data|Culling", meta = (PCG_Overridable))
 	bool bIgnorePCGHits = false;
 
-	UPROPERTY(BlueprintReadWrite, EditAnywhere, Category = Data, meta = (PCG_Overridable))
+	UPROPERTY(BlueprintReadWrite, EditAnywhere, Category = "Data|Culling", meta = (PCG_Overridable))
 	bool bIgnoreSelfHits = true;
 
 	UPROPERTY(BlueprintReadWrite, EditAnywhere, Category = "Data|Advanced", meta = (PCG_Overridable))
@@ -73,9 +75,12 @@
 	bool bIgnoreLandscapeHits_DEPRECATED = false;
 #endif
 
-	UPROPERTY(BlueprintReadWrite, EditAnywhere, Category = Data, meta = (PCG_Overridable))
+	UPROPERTY(BlueprintReadWrite, EditAnywhere, Category = "Data|Attributes", meta = (PCG_Overridable))
 	bool bGetReferenceToActorHit = false;
 
+	UPROPERTY(BlueprintReadWrite, EditAnywhere, Category = "Data|Attributes", meta = (PCG_Overridable))
+	bool bGetReferenceToPhysicalMaterial = false;
+
 	// Not exposed, will be filled in when initializing this
 	UPROPERTY()
 	TSet<FName> ParsedActorTagsList;
@@ -86,20 +91,28 @@
 };
 
 USTRUCT(BlueprintType)
-struct FPCGWorldVolumetricQueryParams : public FPCGWorldCommonQueryParams
+struct FPCGWorldRaycastQueryParams : public FPCGWorldCommonQueryParams
 {
 	GENERATED_BODY()
 
 	void Initialize();
 	void PostSerialize(const FArchive& Ar);
 
-	/** Controls whether we are trying to find an overlap with physical objects (true) or to find empty spaces that do not contain anything (false) */
-	UPROPERTY(BlueprintReadWrite, EditAnywhere, Category = Data, meta = (PCG_Overridable))
-	bool bSearchForOverlap = true;
+	UPROPERTY(BlueprintReadWrite, EditAnywhere, Category = "Data|Culling", meta = (PCG_Overridable))
+	bool bIgnoreBackfaceHits = false;
+
+	UPROPERTY(BlueprintReadWrite, EditAnywhere, Category = "Data|Attributes", meta = (PCG_Overridable))
+	bool bGetImpact = false;
+
+	UPROPERTY(BlueprintReadWrite, EditAnywhere, Category = "Data|Attributes", meta = (PCG_Overridable))
+	bool bGetDistance = false;
+
+	UPROPERTY(BlueprintReadWrite, EditAnywhere, Category = "Data|Attributes", meta = (PCG_Overridable))
+	bool bGetImpactNormal = false;
 };
 
 template<>
-struct TStructOpsTypeTraits<FPCGWorldVolumetricQueryParams> : public TStructOpsTypeTraitsBase2<FPCGWorldVolumetricQueryParams>
+struct TStructOpsTypeTraits<FPCGWorldRaycastQueryParams> : public TStructOpsTypeTraitsBase2<FPCGWorldRaycastQueryParams>
 {
 	enum
 	{
@@ -107,43 +120,30 @@
 	};
 };
 
-/** Queries volume for presence of world collision or not. Can be used to voxelize environment. */
-UCLASS(MinimalAPI, BlueprintType, ClassGroup = (Procedural))
-class UPCGWorldVolumetricData : public UPCGVolumeData
+USTRUCT(BlueprintType)
+struct FPCGWorldVolumetricQueryParams : public FPCGWorldCommonQueryParams
 {
 	GENERATED_BODY()
 
-public:
-	PCG_API void Initialize(UWorld* InWorld, const FBox& InBounds = FBox(EForceInit::ForceInit));
+	void Initialize();
+	void PostSerialize(const FArchive& Ar);
 
-	//~Begin UPCGSpatialData interface
-	virtual bool IsBounded() const override { return !!Bounds.IsValid; }
-	virtual bool SamplePoint(const FTransform& Transform, const FBox& Bounds, FPCGPoint& OutPoint, UPCGMetadata* OutMetadata) const override;
-	// TODO not sure what this would mean. Without a direction, this means perhaps finding closest point on any collision surface? Should we implement this disabled?
-	//virtual bool ProjectPoint(const FTransform& InTransform, const FBox& InBounds, const FPCGProjectionParams& InParams, FPCGPoint& OutPoint, UPCGMetadata* OutMetadata) const;
-protected:
-	virtual UPCGSpatialData* CopyInternal() const override;
-	//~End UPCGSpatialData interface
+	/** Controls whether we are trying to find an overlap with physical objects (true) or to find empty spaces that do not contain anything (false) */
+	UPROPERTY(BlueprintReadWrite, EditAnywhere, Category = Data, meta = (PCG_Overridable))
+	bool bSearchForOverlap = true;
+};
 
-public:
-	//~Begin UPCGSpatialDataWithPointCache
-	virtual bool SupportsBoundedPointData() const { return true; }
-	virtual const UPCGPointData* CreatePointData(FPCGContext* Context) const override { return CreatePointData(Context, FBox(EForceInit::ForceInit)); }
-	virtual const UPCGPointData* CreatePointData(FPCGContext* Context, const FBox& InBounds) const override;
-	//~End UPCGSpatialDataWithPointCache
-
-	UPROPERTY()
-	TWeakObjectPtr<UWorld> World = nullptr;
-
-	UPROPERTY()
-	TWeakObjectPtr<UPCGComponent> OriginatingComponent = nullptr;
-
-	UPROPERTY(BlueprintReadWrite, EditAnywhere, Category = Data, meta = (ShowOnlyInnerProperties))
-	FPCGWorldVolumetricQueryParams QueryParams;
+template<>
+struct TStructOpsTypeTraits<FPCGWorldVolumetricQueryParams> : public TStructOpsTypeTraitsBase2<FPCGWorldVolumetricQueryParams>
+{
+	enum
+	{
+		WithPostSerialize = true,
+	};
 };
 
 USTRUCT(BlueprintType)
-struct FPCGWorldRayHitQueryParams : public FPCGWorldCommonQueryParams
+struct FPCGWorldRayHitQueryParams : public FPCGWorldRaycastQueryParams
 {
 	GENERATED_BODY()
 
@@ -167,13 +167,7 @@
 	// examples: bReturnFaceIndex, bReturnPhysicalMaterial, some ignore patterns
 
 	UPROPERTY(BlueprintReadWrite, EditAnywhere, Category = Data, meta = (PCG_Overridable))
-	bool bIgnoreBackfaceHits = false;
-
-	UPROPERTY(BlueprintReadWrite, EditAnywhere, Category = Data, meta = (PCG_Overridable))
 	bool bApplyMetadataFromLandscape = false;
-
-	UPROPERTY(BlueprintReadWrite, EditAnywhere, Category = Data, meta = (PCG_Overridable))
-	bool bGetReferenceToPhysicalMaterial = false;
 };
 
 template<>
@@ -185,6 +179,41 @@
 	};
 };
 
+/** Queries volume for presence of world collision or not. Can be used to voxelize environment. */
+UCLASS(MinimalAPI, BlueprintType, ClassGroup = (Procedural))
+class UPCGWorldVolumetricData : public UPCGVolumeData
+{
+	GENERATED_BODY()
+
+public:
+	PCG_API void Initialize(UWorld* InWorld, const FBox& InBounds = FBox(EForceInit::ForceInit));
+
+	//~Begin UPCGSpatialData interface
+	virtual bool IsBounded() const override { return !!Bounds.IsValid; }
+	virtual bool SamplePoint(const FTransform& Transform, const FBox& Bounds, FPCGPoint& OutPoint, UPCGMetadata* OutMetadata) const override;
+	// TODO not sure what this would mean. Without a direction, this means perhaps finding closest point on any collision surface? Should we implement this disabled?
+	//virtual bool ProjectPoint(const FTransform& InTransform, const FBox& InBounds, const FPCGProjectionParams& InParams, FPCGPoint& OutPoint, UPCGMetadata* OutMetadata) const;
+protected:
+	virtual UPCGSpatialData* CopyInternal() const override;
+	//~End UPCGSpatialData interface
+
+public:
+	//~Begin UPCGSpatialDataWithPointCache
+	virtual bool SupportsBoundedPointData() const { return true; }
+	virtual const UPCGPointData* CreatePointData(FPCGContext* Context) const override { return CreatePointData(Context, FBox(EForceInit::ForceInit)); }
+	virtual const UPCGPointData* CreatePointData(FPCGContext* Context, const FBox& InBounds) const override;
+	//~End UPCGSpatialDataWithPointCache
+
+	UPROPERTY()
+	TWeakObjectPtr<UWorld> World = nullptr;
+
+	UPROPERTY()
+	TWeakObjectPtr<UPCGComponent> OriginatingComponent = nullptr;
+
+	UPROPERTY(BlueprintReadWrite, EditAnywhere, Category = Data, meta = (ShowOnlyInnerProperties))
+	FPCGWorldVolumetricQueryParams QueryParams;
+};
+
 /** Executes collision queries against world collision. */
 UCLASS(MinimalAPI, BlueprintType, ClassGroup = (Procedural))
 class UPCGWorldRayHitData : public UPCGSurfaceData
diff --git a/Engine/Plugins/PCG/Source/PCG/Public/Elements/PCGWorldRaycast.h b/Engine/Plugins/PCG/Source/PCG/Public/Elements/PCGWorldRaycast.h
new file mode 100644
--- /dev/null
+++ b/Engine/Plugins/PCG/Source/PCG/Public/Elements/PCGWorldRaycast.h
@@ -0,0 +1,112 @@
+// Copyright Epic Games, Inc. All Rights Reserved.
+#pragma once
+
+#include "PCGSettings.h"
+#include "PCGTimeSlicedElementBase.h"
+#include "Data/PCGWorldData.h"
+#include "Helpers/PCGWorldQueryHelpers.h"
+
+#include "CollisionQueryParams.h"
+
+#include "PCGWorldRaycast.generated.h"
+
+/**
+ * Casts rays from provided points along a given direction and transform points to the impact point.
+ */
+UCLASS(BlueprintType, ClassGroup = (Procedural))
+class UPCGWorldRaycastElementSettings : public UPCGSettings
+{
+	GENERATED_BODY()
+
+	UPCGWorldRaycastElementSettings();
+
+public:
+	//~Begin UPCGSettings interface
+#if WITH_EDITOR
+	virtual FName GetDefaultNodeName() const override { return FName(TEXT("WorldRaycast")); }
+	virtual FText GetDefaultNodeTitle() const override { return NSLOCTEXT("PCGWorldRaycastElement", "NodeTitle", "World Raycast"); }
+	virtual FText GetNodeTooltipText() const override { return NSLOCTEXT("PCGWorldRaycastElement", "NodeTooltip", "Casts rays from provided points along a given direction and transform points to the impact point."); }
+	virtual EPCGSettingsType GetType() const override { return EPCGSettingsType::Spatial; }
+#endif // WITH_EDITOR
+
+protected:
+	virtual TArray<FPCGPinProperties> InputPinProperties() const override;
+	virtual TArray<FPCGPinProperties> OutputPinProperties() const override;
+	virtual FPCGElementPtr CreateElement() const override;
+	//~End UPCGSettings interface
+
+public:
+	/** Determines how the ray's direction and distance will be calculated. */
+	UPROPERTY(BlueprintReadWrite, EditAnywhere, Category = Settings)
+	EPCGWorldRaycastMode RaycastMode = EPCGWorldRaycastMode::Infinite;
+
+	/** The selected attribute determines the ray origin. */
+	UPROPERTY(BlueprintReadWrite, EditAnywhere, Category = Settings, meta = (PCG_Overridable, EditCondition = "bOverrideOriginPositions", EditConditionHides))
+	FPCGAttributePropertyInputSelector OriginInputAttribute;
+
+	/** Use a selected attribute as the ray direction. */
+	UPROPERTY(BlueprintReadWrite, EditAnywhere, Category = Settings, meta = (PCG_Overridable, EditCondition = "RaycastMode != EPCGWorldRaycastMode::Segments", EditConditionHides))
+	bool bOverrideRayDirections = false;
+
+	/** A ray direction that will be used for all raycasts. */
+	UPROPERTY(BlueprintReadWrite, EditAnywhere, Category = Settings, meta = (PCG_Overridable, EditCondition = "RaycastMode != EPCGWorldRaycastMode::Segments && !bOverrideRayDirections", EditConditionHides))
+	FVector RayDirection = -FVector::UnitZ();
+
+	/** The selected attribute determines the ray direction. */
+	UPROPERTY(BlueprintReadWrite, EditAnywhere, Category = Settings, meta = (PCG_Overridable, EditCondition = "RaycastMode != EPCGWorldRaycastMode::Segments && bOverrideRayDirections", EditConditionHides))
+	FPCGAttributePropertyInputSelector RayDirectionAttribute;
+
+	/** The selected attribute determines the ray terminal point. */
+	UPROPERTY(BlueprintReadWrite, EditAnywhere, Category = Settings, meta = (PCG_Overridable, EditCondition = "RaycastMode == EPCGWorldRaycastMode::Segments", EditConditionHides))
+	FPCGAttributePropertyInputSelector EndPointAttribute;
+
+	/** Use a selected attribute as the ray length. */
+	UPROPERTY(BlueprintReadWrite, EditAnywhere, Category = Settings, meta = (PCG_Overridable, EditCondition = "RaycastMode == EPCGWorldRaycastMode::NormalizedWithLength", EditConditionHides))
+	bool bOverrideRayLengths = false;
+
+	/** A ray length that will be used for all raycasts. */
+	UPROPERTY(BlueprintReadWrite, EditAnywhere, Category = Settings, meta = (PCG_Overridable, EditCondition = "RaycastMode == EPCGWorldRaycastMode::NormalizedWithLength && !bOverrideRayLengths", EditConditionHides))
+	double RayLength = 100000.0;
+
+	/** The selected attribute determines the ray length. */
+	UPROPERTY(BlueprintReadWrite, EditAnywhere, Category = Settings, meta = (PCG_Overridable, EditCondition = "RaycastMode == EPCGWorldRaycastMode::NormalizedWithLength && bOverrideRayLengths", EditConditionHides))
+	FPCGAttributePropertyInputSelector RayLengthAttribute;
+
+	/** World ray trace parameters. */
+	UPROPERTY(BlueprintReadWrite, EditAnywhere, Category = Settings, meta = (PCG_Overridable, ShowOnlyInnerProperties))
+	FPCGWorldRaycastQueryParams WorldQueryParams;
+
+	/** Will keep the original points at their location if the raycast misses or if the hit result is out of bounds. */
+	UPROPERTY(BlueprintReadWrite, EditAnywhere, Category = Settings, meta = (PCG_Overridable))
+	bool bKeepOriginalPointOnMiss = false;
+
+	/** If no Bounding Shape input is provided, the actor bounds are used to limit the sample generation domain. */
+	UPROPERTY(BlueprintReadWrite, EditAnywhere, Category = Settings, meta = (PCG_Overridable))
+	bool bUnbounded = false;
+};
+
+namespace FPCGWorldRaycast
+{
+	struct FExecutionState
+	{
+		FCollisionQueryParams CollisionQueryParams;
+		FCollisionObjectQueryParams CollisionObjectQueryParams;
+		FBox Bounds = FBox(EForceInit::ForceInit);
+	};
+
+	struct FIterationState
+	{
+		TArray<FVector> CachedRayOrigins;
+		TArray<FVector> CachedRayVectors;
+		TArray<double> CachedRayLengths;
+	};
+}
+
+class FPCGWorldRaycastElement : public TPCGTimeSlicedElementBase<FPCGWorldRaycast::FExecutionState, FPCGWorldRaycast::FIterationState>
+{
+protected:
+	virtual bool PrepareDataInternal(FPCGContext* InContext) const override;
+	virtual bool ExecuteInternal(FPCGContext* InContext) const override;
+	// Physics required to run on main thread
+	virtual bool CanExecuteOnlyOnMainThread(FPCGContext* Context) const override { return true; }
+};
diff --git a/Engine/Plugins/PCG/Source/PCG/Public/Helpers/PCGSettingsHelpers.h b/Engine/Plugins/PCG/Source/PCG/Public/Helpers/PCGSettingsHelpers.h
--- a/Engine/Plugins/PCG/Source/PCG/Public/Helpers/PCGSettingsHelpers.h
+++ b/Engine/Plugins/PCG/Source/PCG/Public/Helpers/PCGSettingsHelpers.h
@@ -2,9 +2,12 @@
 
 #pragma once
 
+#include "PCGComponent.h"
+#include "PCGContext.h"
 #include "PCGModule.h"
 #include "PCGParamData.h"
 #include "PCGSettings.h"
+#include "Data/PCGSpatialData.h"
 #include "Metadata/PCGMetadataAttributeTpl.h"
 #include "Metadata/PCGMetadataAttributeTraits.h"
 #include "Metadata/Accessors/PCGAttributeAccessorHelpers.h"
@@ -204,6 +207,20 @@ PRAGMA_ENABLE_DEPRECATION_WARNINGS
 		});
 	}
 
+	inline const UPCGSpatialData* ComputeBoundingShape(const FPCGContext* Context, FName BoundingShapeLabel, bool& bOutUnionWasCreated)
+	{
+		const UPCGSpatialData* BoundingShape = Context->InputData.GetSpatialUnionOfInputsByPin(BoundingShapeLabel, bOutUnionWasCreated);
+
+		// Fallback to getting bounds from actor
+		if (!BoundingShape && Context->SourceComponent.IsValid())
+		{
+			check(bOutUnionWasCreated == false);
+			BoundingShape = Cast<UPCGSpatialData>(Context->SourceComponent->GetActorPCGData());
+		}
+
+		return BoundingShape;
+	}
+
 	struct FPCGGetAllOverridableParamsConfig
 	{
 		// If we don't use the seed, don't add it as override.
diff --git a/Engine/Plugins/PCG/Source/PCG/Public/Helpers/PCGWorldQueryHelpers.h b/Engine/Plugins/PCG/Source/PCG/Public/Helpers/PCGWorldQueryHelpers.h
new file mode 100644
--- /dev/null
+++ b/Engine/Plugins/PCG/Source/PCG/Public/Helpers/PCGWorldQueryHelpers.h
@@ -0,0 +1,64 @@
+// Copyright Epic Games, Inc. All Rights Reserved.
+
+#pragma once
+
+#include "Engine/HitResult.h"
+#include "Misc/Optional.h"
+
+#include "PCGWorldQueryHelpers.generated.h"
+
+class UPCGComponent;
+class UPCGMetadata;
+struct FHitResult;
+struct FOverlapResult;
+struct FPCGPoint;
+struct FPCGWorldCommonQueryParams;
+struct FPCGWorldRaycastQueryParams;
+struct FPCGWorldVolumetricQueryParams;
+
+UENUM()
+enum class EPCGWorldRaycastMode : uint8
+{
+	Infinite UMETA(ToolTip = "Use the direction vector with 'infinite' magnitude."),
+	ScaledVector UMETA(ToolTip = "Use the direction vector 'as-is' for casting the ray with its current magnitude."),
+	NormalizedWithLength UMETA(ToolTip = "Normalize the direction vector and apply the length directly."),
+	Segments UMETA(ToolTip = "User provided end points. Must match input points N:N, N:1, or 1:N.")
+};
+
+namespace PCGWorldQueryConstants
+{
+	const FName ImpactAttribute = TEXT("ImpactResult");
+	const FName DistanceAttribute = TEXT("ImpactDistance");
+	const FName ImpactNormalAttribute = TEXT("ImpactNormal");
+	const FName PhysicalMaterialReferenceAttribute = TEXT("PhysicalMaterial");
+}
+
+namespace PCGWorldQueryHelpers
+{
+	FTransform GetOrthonormalImpactTransform(const FHitResult& Hit);
+
+	/** Returns false if a physics query doesn't meet the requirements of common query parameters. */
+	bool FilterCommonQueryResults(
+		const FPCGWorldCommonQueryParams* QueryParams,
+		const UPrimitiveComponent* TriggeredComponent,
+		const TWeakObjectPtr<UPCGComponent> OriginatingComponent);
+
+	/** Filters through an array of hit results, testing them against raycast query parameters. Optional extra filter predicate. */
+	TOptional<FHitResult> FilterRayHitResults(
+		const FPCGWorldRaycastQueryParams* QueryParams,
+		TWeakObjectPtr<UPCGComponent> OriginatingComponent,
+		const TArray<FHitResult>& HitResults);
+
+	/** Filters through an array of overlap results, testing them against raycast query parameters. Optional extra filter predicate. */
+	TOptional<FOverlapResult> FilterOverlapResults(
+		const FPCGWorldVolumetricQueryParams* QueryParams,
+		TWeakObjectPtr<UPCGComponent> OriginatingComponent,
+		const TArray<FOverlapResult>& OverlapResults);
+
+	/** Applies common world ray hit results to attributes. */
+	bool ApplyRayHitMetadata(const TOptional<FHitResult>& HitResult,
+		const FPCGWorldRaycastQueryParams& QueryParams,
+		FPCGPoint& OutPoint,
+		UPCGMetadata* OutMetadata,
+		bool bShouldCreateAttributes = true);
+}
