[PCG] Sanitize names when getting properties that has invalid names

* Changed the logic in Import/Export Text in the selectors to support this use case.

#rnx

diff --git a/Engine/Plugins/PCG/Source/PCG/Private/Helpers/PCGPropertyHelpers.cpp b/Engine/Plugins/PCG/Source/PCG/Private/Helpers/PCGPropertyHelpers.cpp
--- a/Engine/Plugins/PCG/Source/PCG/Private/Helpers/PCGPropertyHelpers.cpp
+++ b/Engine/Plugins/PCG/Source/PCG/Private/Helpers/PCGPropertyHelpers.cpp
@@ -220,7 +220,7 @@ UPCGParamData* PCGPropertyHelpers::ExtractPropertyAsAttributeSet(const PCGProper
 		Property = ArrayProperty->Inner;
 	}
 
-	using ExtractablePropertyTuple = TTuple<FName, const FProperty*>;
+	using ExtractablePropertyTuple = TTuple<FString, const FProperty*>;
 	TArray<ExtractablePropertyTuple> ExtractableProperties;
 
 	using GetAddressFunc = TFunction<const void* (const void*)>;
@@ -280,8 +280,8 @@ UPCGParamData* PCGPropertyHelpers::ExtractPropertyAsAttributeSet(const PCGProper
 				{
 					// We use authored name as attribute name to avoid issue with noisy property names, like in UUserDefinedStructs, where some random number is appended to the property name.
 					// By default, it will just return the property name anyway.
-					const FString AuthoredName = UnderlyingClass->GetAuthoredNameForField(ChildProperty);
-					ExtractableProperties.Emplace(FName(AuthoredName), ChildProperty);
+					FString AuthoredName = UnderlyingClass->GetAuthoredNameForField(ChildProperty);
+					ExtractableProperties.Emplace(std::move(AuthoredName), ChildProperty);
 				}
 			}
 		}
@@ -290,7 +290,7 @@ UPCGParamData* PCGPropertyHelpers::ExtractPropertyAsAttributeSet(const PCGProper
 	{
 		// For non struct/object, there is just a single property to extract with no shenanigans for address indirection.
 		const FName AttributeName = (Parameters.OutputAttributeName == PCGMetadataAttributeConstants::SourceNameAttributeName) ? Property->GetFName() : Parameters.OutputAttributeName;
-		ExtractableProperties.Emplace(AttributeName, Property);
+		ExtractableProperties.Emplace(AttributeName.ToString(), Property);
 		// Identity
 		AddressFunc = [](const void* InAddress) { return InAddress; };
 	}
@@ -327,10 +327,13 @@ UPCGParamData* PCGPropertyHelpers::ExtractPropertyAsAttributeSet(const PCGProper
 
 		for (ExtractablePropertyTuple& ExtractableProperty : ExtractableProperties)
 		{
-			const FName AttributeName = ExtractableProperty.Get<0>();
+			FString AttributeName = ExtractableProperty.Get<0>();
 			const FProperty* FinalProperty = ExtractableProperty.Get<1>();
 
-			if (!Metadata->SetAttributeFromDataProperty(AttributeName, EntryKey, ContainerPtr, FinalProperty, /*bCreate=*/ true))
+			// Make sure the Attribute name is sanitized, to prevent cases where property names have unsupported characters.
+			FPCGMetadataAttributeBase::SanitizeName(AttributeName);
+
+			if (!Metadata->SetAttributeFromDataProperty(FName(AttributeName), EntryKey, ContainerPtr, FinalProperty, /*bCreate=*/ true))
 			{
 				LogError(FText::Format(LOCTEXT("ErrorCreatingAttribute", "Error while creating an attribute for property '{0}'. Either the property type is not supported by PCG or attribute creation failed."), FText::FromString(FinalProperty->GetName())), OptionalContext);
 				bValidOperation = false;
diff --git a/Engine/Plugins/PCG/Source/PCG/Private/Metadata/PCGAttributePropertySelector.cpp b/Engine/Plugins/PCG/Source/PCG/Private/Metadata/PCGAttributePropertySelector.cpp
--- a/Engine/Plugins/PCG/Source/PCG/Private/Metadata/PCGAttributePropertySelector.cpp
+++ b/Engine/Plugins/PCG/Source/PCG/Private/Metadata/PCGAttributePropertySelector.cpp
@@ -16,21 +16,40 @@ namespace PCGAttributePropertySelectorConstants
 	static const TCHAR* ExtraSeparator = TEXT(".");
 	static const TCHAR PropertyPrefixChar = PropertyPrefix[0];
 	static const TCHAR ExtraSeparatorChar = ExtraSeparator[0];
+
+	static const FString ExportTextLeftSentinel = TEXT("PCGBegin(");
+	static const FString ExportTextRightSentinel = TEXT(")PCGEnd");
 }
 
 bool FPCGAttributePropertySelector::ExportTextItem(FString& ValueStr, FPCGAttributePropertySelector const& DefaultValue, UObject* Parent, int32 PortFlags, UObject* ExportRootScope) const
 {
-	// Quoted string
-	ValueStr = FString(TEXT("\"")) + ToString() + TEXT("\"");
+	// String guarded by sentinels, don't use `"` because it can be used in the selector.
+	TStringBuilder<256> StringBuilder;
+	StringBuilder.Append(PCGAttributePropertySelectorConstants::ExportTextLeftSentinel);
+	StringBuilder.Append(ToString());
+	StringBuilder.Append(PCGAttributePropertySelectorConstants::ExportTextRightSentinel);
+
+	ValueStr = StringBuilder.ToString();
 	return true;
 }
 
 bool FPCGAttributePropertySelector::ImportTextItem(const TCHAR*& Buffer, int32 PortFlags, UObject* Parent, FOutputDevice* ErrorText)
 {
-	FString Token;
-	int NumCharsRead;
-	FParse::QuotedString(Buffer, Token, &NumCharsRead);
-	Update(Token);
+	const FStringView BufferView(Buffer);
+
+	using PCGAttributePropertySelectorConstants::ExportTextLeftSentinel;
+	using PCGAttributePropertySelectorConstants::ExportTextRightSentinel;
+
+	if (!BufferView.StartsWith(ExportTextLeftSentinel) || !BufferView.EndsWith(ExportTextRightSentinel))
+	{
+		// Didn't find our sentinels, abort
+		return false;
+	}
+
+	const int32 Start = ExportTextLeftSentinel.Len();
+	const int32 End = BufferView.Len() - ExportTextRightSentinel.Len();
+
+	Update(FString(BufferView.SubStr(Start, End - Start)));
 	return true;
 }
 
@@ -232,23 +251,7 @@ bool FPCGAttributePropertySelector::Update(const FString& NewValue)
 	}
 	else
 	{
-		TArray<FString> NewValuesTemp;
-		NewValue.ParseIntoArray(NewValuesTemp, PCGAttributePropertySelectorConstants::ExtraSeparator, /*InCullEmpty=*/ false);
-
-		// We can't sanitize the first name because of the `@` and `$` that are valid symbols here, but are not for the sanitizer.
-		// Also if we have parsed nothing, just add a single empty string.
-		NewValues.Emplace(NewValuesTemp.IsEmpty() ? FString() : std::move(NewValuesTemp[0]));
-
-		// Then we make sure extra names are not containing non-alphanumerical chars by sanitizing it
-		for (int i = 1; i < NewValuesTemp.Num(); ++i)
-		{
-			FPCGMetadataAttributeBase::SanitizeName(NewValuesTemp[i]);
-
-			if (!NewValuesTemp[i].IsEmpty())
-			{
-				NewValues.Emplace(std::move(NewValuesTemp[i]));
-			}
-		}
+		NewValue.ParseIntoArray(NewValues, PCGAttributePropertySelectorConstants::ExtraSeparator, /*InCullEmpty=*/ false);
 	}
 
 	const FString& NewName = NewValues[0];
