[PCG] Optional octree to prevent need for builds recomputation in the collision wrapper code

diff --git a/Engine/Plugins/PCG/Source/PCG/Private/Data/PCGCollisionWrapperData.cpp b/Engine/Plugins/PCG/Source/PCG/Private/Data/PCGCollisionWrapperData.cpp
--- a/Engine/Plugins/PCG/Source/PCG/Private/Data/PCGCollisionWrapperData.cpp
+++ b/Engine/Plugins/PCG/Source/PCG/Private/Data/PCGCollisionWrapperData.cpp
@@ -8,6 +8,7 @@
 #include "Components/StaticMeshComponent.h"
 #include "Engine/StaticMesh.h"
 #include "PhysicsEngine/BodyInstance.h"
+#include "Serialization/ArchiveCrc32.h"
 
 #include UE_INLINE_GENERATED_CPP_BY_NAME(PCGCollisionWrapperData)
 
@@ -100,7 +101,7 @@ void FPCGCollisionWrapper::CreateBodyInstances(const TArray<FSoftObjectPath>& Me
 		BodyInstances.Add(BodyInstance);
 	}
 
-	bInitialized = true;
+	bInitialized = !MeshPaths.IsEmpty();
 }
 
 bool FPCGCollisionWrapper::Initialize(const IPCGAttributeAccessor* Accessor, const IPCGAttributeAccessorKeys* Keys)
@@ -125,12 +126,76 @@ FBodyInstance* FPCGCollisionWrapper::GetBodyInstance(int32 EntryIndex) const
 	return ((bInitialized && IndexToBodyInstance[EntryIndex] != INDEX_NONE) ? BodyInstances[IndexToBodyInstance[EntryIndex]] : nullptr);
 }
 
-bool UPCGCollisionWrapperData::Initialize(const UPCGPointData* InPointData, const FPCGAttributePropertyInputSelector& InCollisionSelector, bool bInUseComplexCollision)
+bool FPCGCollisionWrapper::InitializeOctree(const UPCGPointData* InPointData, const TArray<FSoftObjectPath>& InMeshPaths, UPCGPointData::PointOctree& OutOctree, TArray<FPCGPointRef>* OutOctreePointRefs) const
+{
+	TArray<FBox> MeshBoundsList;
+	MeshBoundsList.Reserve(InMeshPaths.Num());
+	bool bHasValidBounds = false;
+
+	for (int32 MeshIndex = 0; MeshIndex < InMeshPaths.Num(); ++MeshIndex)
+	{
+		TSoftObjectPtr<UStaticMesh> Mesh(InMeshPaths[MeshIndex]);
+		if (Mesh.Get()) // should already be loaded by now.
+		{
+			MeshBoundsList.Add(Mesh->GetBoundingBox());
+			bHasValidBounds = true;
+		}
+		else
+		{
+			MeshBoundsList.Emplace(EForceInit::ForceInit);
+		}
+	}
+
+	if (bHasValidBounds)
+	{
+		const TArray<FPCGPoint>& Points = InPointData->GetPoints();
+		TArray<FPCGPointRef> PointRefs;
+		PointRefs.Reserve(Points.Num());
+
+		for (int32 PointIndex = 0; PointIndex < Points.Num(); ++PointIndex)
+		{
+			const FPCGPoint& Point = Points[PointIndex];
+			const FBox& MeshBounds = MeshBoundsList[IndexToBodyInstance[PointIndex]];
+
+			if (MeshBounds.IsValid)
+			{
+				PointRefs.Emplace(Point, MeshBounds);
+			}
+			else
+			{
+				PointRefs.Emplace(Point);
+			}
+		}
+
+		FBox PointsBounds = InPointData->GetBounds();
+		TOctree2<FPCGPointRef, FPCGPointRefSemantics> NewOctree(PointsBounds.GetCenter(), PointsBounds.GetExtent().Length());
+
+		for (const FPCGPointRef& PointRef : PointRefs)
+		{
+			NewOctree.AddElement(PointRef);
+		}
+
+		OutOctree = MoveTemp(NewOctree);
+
+		if (OutOctreePointRefs)
+		{
+			*OutOctreePointRefs = MoveTemp(PointRefs);
+		}
+
+		return true;
+	}
+	else
+	{
+		return false;
+	}
+}
+
+bool UPCGCollisionWrapperData::Initialize(const UPCGPointData* InPointData, const FPCGAttributePropertyInputSelector& InCollisionSelector, bool bInUseComplexCollision, bool bInUseAccurateOctree)
 {
 	TArray<FSoftObjectPath> MeshesToLoad;
 	if (PreInitializeAndGatherMeshesEx(InPointData, InCollisionSelector, bInUseComplexCollision, MeshesToLoad))
 	{
-		FinalizeInitializationEx(MeshesToLoad);
+		FinalizeInitializationEx(InPointData, MeshesToLoad, bInUseAccurateOctree);
 		return true;
 	}
 	else
@@ -168,15 +233,27 @@ bool UPCGCollisionWrapperData::PreInitializeAndGatherMeshesEx(const UPCGPointDat
 	return CollisionWrapper.Prepare(InputAccessor.Get(), InputKeys.Get(), OutMeshesToLoad);
 }
 
-void UPCGCollisionWrapperData::FinalizeInitializationEx(const TArray<FSoftObjectPath>& InMeshPaths)
+void UPCGCollisionWrapperData::FinalizeInitializationEx(const UPCGPointData* InPointData, const TArray<FSoftObjectPath>& InMeshPaths, bool bInUseAccurateOctree)
 {
 	CollisionWrapper.CreateBodyInstances(InMeshPaths);
+
+	if (InPointData && bInUseAccurateOctree)
+	{
+		bUseCollisionAccurateOctree = CollisionWrapper.InitializeOctree(InPointData, InMeshPaths, CollisionAccurateOctree);
+	}
 }
 
 void UPCGCollisionWrapperData::AddToCrc(FArchiveCrc32& Ar, bool bFullDataCrc) const
 {
 	Super::AddToCrc(Ar, bFullDataCrc);
 	GetPointData()->AddToCrc(Ar, bFullDataCrc);
+
+	CollisionSelector.AddToCrc(Ar);
+	uint32 UseCollisionAccurateOctree = bUseCollisionAccurateOctree ? 1 : 0;
+	Ar << UseCollisionAccurateOctree;
+
+	uint32 UseComplexCollision = bUseComplexCollision ? 1 : 0;
+	Ar << UseComplexCollision;
 }
 
 void UPCGCollisionWrapperData::GetResourceSizeEx(FResourceSizeEx& CumulativeResourceSize)
@@ -210,7 +287,7 @@ bool UPCGCollisionWrapperData::SamplePoint(const FTransform& InTransform, const
 	// For all touching points, check against the actual physics shape we have for that point
 	// E.g. get the point collision attribute value, match to our internal list, get the shape (might require loading), then test against the matching shapes
 	// note that we need to place the stuff in the right referential too, since our shapes will all be at the origin
-	const UPCGPointData::PointOctree& Octree = GetPointData()->GetOctree();
+	const UPCGPointData::PointOctree& Octree = (bUseCollisionAccurateOctree ? CollisionAccurateOctree : GetPointData()->GetOctree());
 
 	float Density = 0;
 	FBox TransformedBounds = InBounds.TransformBy(InTransform);
@@ -263,7 +340,7 @@ bool UPCGCollisionWrapperData::SamplePoint(const FTransform& InTransform, const
 UPCGSpatialData* UPCGCollisionWrapperData::CopyInternal() const
 {
 	UPCGCollisionWrapperData* NewCollisionWrapperData = NewObject<UPCGCollisionWrapperData>();
-	NewCollisionWrapperData->Initialize(PointData, CollisionSelector, bUseComplexCollision);
+	NewCollisionWrapperData->Initialize(PointData, CollisionSelector, bUseComplexCollision, bUseCollisionAccurateOctree);
 
 	return NewCollisionWrapperData;
 }
diff --git a/Engine/Plugins/PCG/Source/PCG/Private/Data/PCGPointData.cpp b/Engine/Plugins/PCG/Source/PCG/Private/Data/PCGPointData.cpp
--- a/Engine/Plugins/PCG/Source/PCG/Private/Data/PCGPointData.cpp
+++ b/Engine/Plugins/PCG/Source/PCG/Private/Data/PCGPointData.cpp
@@ -241,6 +241,12 @@ FPCGPointRef::FPCGPointRef(const FPCGPoint& InPoint)
 	Bounds = InPoint.GetDensityBounds();
 }
 
+FPCGPointRef::FPCGPointRef(const FPCGPoint& InPoint, const FBox& InOverrideBounds)
+{
+	Point = &InPoint;
+	Bounds = FBoxSphereBounds(InOverrideBounds.TransformBy(InPoint.Transform));
+}
+
 void UPCGPointData::GetResourceSizeEx(FResourceSizeEx& CumulativeResourceSize)
 {
 	Super::GetResourceSizeEx(CumulativeResourceSize);
@@ -690,7 +696,7 @@ void UPCGPointData::RebuildOctree() const
 		NewOctree.AddElement(FPCGPointRef(Point));
 	}
 
-	Octree = NewOctree;
+	Octree = MoveTemp(NewOctree);
 	bOctreeIsDirty = false;
 }
 
diff --git a/Engine/Plugins/PCG/Source/PCG/Private/Elements/PCGCreateCollisionData.cpp b/Engine/Plugins/PCG/Source/PCG/Private/Elements/PCGCreateCollisionData.cpp
--- a/Engine/Plugins/PCG/Source/PCG/Private/Elements/PCGCreateCollisionData.cpp
+++ b/Engine/Plugins/PCG/Source/PCG/Private/Elements/PCGCreateCollisionData.cpp
@@ -125,7 +125,7 @@ bool FPCGCreateCollisionDataElement::ExecuteInternal(FPCGContext* InContext) con
 		if (FPCGCreateCollisionContext::InputMeshData* MatchingData = Context->PerInputData.FindByPredicate([InputIndex](const FPCGCreateCollisionContext::InputMeshData& Data) { return Data.InputIndex == InputIndex; }))
 		{
 			check(MatchingData->Data);
-			MatchingData->Data->FinalizeInitializationEx(MatchingData->MeshPaths);
+			MatchingData->Data->FinalizeInitializationEx(CastChecked<UPCGPointData>(Input.Data), MatchingData->MeshPaths, Settings->bRecomputeOctreeAccordingToMeshes);
 
 			// Unroot collision wrapper data object since it'll be tracked through the collection
 			MatchingData->Data->RemoveFromRoot();
diff --git a/Engine/Plugins/PCG/Source/PCG/Private/Elements/PCGSelfPruning.cpp b/Engine/Plugins/PCG/Source/PCG/Private/Elements/PCGSelfPruning.cpp
--- a/Engine/Plugins/PCG/Source/PCG/Private/Elements/PCGSelfPruning.cpp
+++ b/Engine/Plugins/PCG/Source/PCG/Private/Elements/PCGSelfPruning.cpp
@@ -163,7 +163,7 @@ namespace PCGSelfPruningElement
 		TRACE_CPUPROFILER_EVENT_SCOPE(FPCGSelfPruningElement::Execute::CollisionExclusion);
 
 		check(IterationState.InputPointData);
-		const UPCGPointData::PointOctree& Octree = IterationState.InputPointData->GetOctree();
+		const UPCGPointData::PointOctree& Octree = IterationState.bUseCollisionAccurateOctree ? IterationState.CollisionAccurateOctree : IterationState.InputPointData->GetOctree();
 
 		int32 CheckTimeSlicingCount = 0;
 
@@ -441,11 +441,17 @@ namespace PCGSelfPruningElement
 		//  then remove this point
 		if (!InState.bSortDone)
 		{
-			InState.SortedPoints.Empty();
-			InState.SortedPoints.Reserve(Points.Num());
-			for (const FPCGPoint& Point : Points)
+			// In the case of the collision-driven self-pruning, we have to populate the sorted points array earlier since we're playing with the bounds.
+			if (!InState.bSortedPointsArrayPopulateDone)
 			{
-				InState.SortedPoints.Add(FPCGPointRef(Point));
+				InState.SortedPoints.Empty();
+				InState.SortedPoints.Reserve(Points.Num());
+				for (const FPCGPoint& Point : Points)
+				{
+					InState.SortedPoints.Emplace(Point);
+				}
+
+				InState.bSortedPointsArrayPopulateDone = true;
 			}
 
 			FPCGAttributePropertySelector ComparisonSource = InParameters.ComparisonSource.CopyAndFixLast(InState.InputPointData);
@@ -695,7 +701,17 @@ bool FPCGSelfPruningElement::PrepareDataInternal(FPCGContext* InContext) const
 			TUniquePtr<const IPCGAttributeAccessor> InputAccessor = PCGAttributeAccessorHelpers::CreateConstAccessor(OutState.InputPointData, InputSelector);
 			TUniquePtr<const IPCGAttributeAccessorKeys> InputKeys = PCGAttributeAccessorHelpers::CreateConstKeys(OutState.InputPointData, InputSelector);
 
-			OutState.CollisionWrapper.Initialize(InputAccessor.Get(), InputKeys.Get());
+			TArray<FSoftObjectPath> Meshes;
+			if (OutState.CollisionWrapper.Prepare(InputAccessor.Get(), InputKeys.Get(), Meshes))
+			{
+				OutState.CollisionWrapper.CreateBodyInstances(Meshes);
+
+				if (Settings->Parameters.bRecomputeOctreeAccordingToMeshes)
+				{
+					OutState.bUseCollisionAccurateOctree = OutState.CollisionWrapper.InitializeOctree(OutState.InputPointData, Meshes, OutState.CollisionAccurateOctree, &OutState.SortedPoints);
+					OutState.bSortedPointsArrayPopulateDone = OutState.bUseCollisionAccurateOctree;
+				}
+			}
 		}
 
 		return EPCGTimeSliceInitResult::Success;
@@ -742,4 +758,16 @@ bool FPCGSelfPruningElement::ExecuteInternal(FPCGContext* Context) const
 	return true;
 }
 
+bool FPCGSelfPruningElement::CanExecuteOnlyOnMainThread(FPCGContext* Context) const
+{
+	if (const UPCGSelfPruningSettings* Settings = (Context ? Context->GetInputSettings<UPCGSelfPruningSettings>() : nullptr))
+	{
+		return Settings->Parameters.bUseCollisionAttribute;
+	}
+	else
+	{
+		return false;
+	}
+}
+
 #undef LOCTEXT_NAMESPACE
diff --git a/Engine/Plugins/PCG/Source/PCG/Public/Data/PCGCollisionWrapperData.h b/Engine/Plugins/PCG/Source/PCG/Public/Data/PCGCollisionWrapperData.h
--- a/Engine/Plugins/PCG/Source/PCG/Public/Data/PCGCollisionWrapperData.h
+++ b/Engine/Plugins/PCG/Source/PCG/Public/Data/PCGCollisionWrapperData.h
@@ -3,6 +3,7 @@
 #pragma once
 
 #include "PCGSpatialData.h"
+#include "Data/PCGPointData.h"
 #include "Metadata/PCGAttributePropertySelector.h"
 #include "Metadata/Accessors/PCGAttributeAccessorHelpers.h"
 
@@ -26,6 +27,7 @@ struct PCG_API FPCGCollisionWrapper
 	// Advanced API - allows to do async loading as we separate the mesh finding part from the body creation part
 	bool Prepare(const IPCGAttributeAccessor* Accessor, const IPCGAttributeAccessorKeys* Keys, TArray<FSoftObjectPath>& MeshPathsToLoad);
 	void CreateBodyInstances(const TArray<FSoftObjectPath>& MeshPaths);
+	bool InitializeOctree(const UPCGPointData* InPointData, const TArray<FSoftObjectPath>& InMeshPaths, UPCGPointData::PointOctree& OutOctree, TArray<FPCGPointRef>* OutOctreePointRefs = nullptr) const;
 	
 	// Retrieves the body instance associated to the entry given by its index
 	FBodyInstance* GetBodyInstance(int32 EntryIndex) const;
@@ -42,11 +44,11 @@ class UPCGCollisionWrapperData : public UPCGSpatialData
 
 public:
 	/** Inititializes the collision wrapper on a point data based on the provided attribute selector */
-	PCG_API bool Initialize(const UPCGPointData* InPointData, const FPCGAttributePropertyInputSelector& InCollisionSelector, bool bInUseComplexCollision);
+	PCG_API bool Initialize(const UPCGPointData* InPointData, const FPCGAttributePropertyInputSelector& InCollisionSelector, bool bInUseComplexCollision, bool bUseAccurateOctree);
 
 	/** Advanced API for async loading */
 	bool PreInitializeAndGatherMeshesEx(const UPCGPointData* InPointData, const FPCGAttributePropertyInputSelector& InCollisionSelector, bool bInUseComplexCollision, TArray<FSoftObjectPath>& OutMeshesToLoad);
-	void FinalizeInitializationEx(const TArray<FSoftObjectPath>& InMeshPaths);
+	void FinalizeInitializationEx(const UPCGPointData* InPointData, const TArray<FSoftObjectPath>& InMeshPaths, bool bUseAccurateOctree);
 
 	// ~Begin UPCGData Interface
 	virtual EPCGDataType GetDataType() const override { return EPCGDataType::Primitive; }
@@ -75,13 +77,19 @@ private:
 	UPROPERTY()
 	TObjectPtr<const UPCGPointData> PointData;
 
+	// Implementation note: in order to be able to duplicate this easily, we're keeping track of the arguments we used when calling Initialize (& derived functions)
 	UPROPERTY()
 	FPCGAttributePropertyInputSelector CollisionSelector;
 
+	/** Uses a new octree based on the mesh bounds, performance warning (does similar work to the BoundsFromMesh node, but will not change the point data). */
+	UPROPERTY()
+	bool bUseCollisionAccurateOctree = false;
+
 	UPROPERTY()
 	bool bUseComplexCollision = false;
 
 	FPCGCollisionWrapper CollisionWrapper;
+	UPCGPointData::PointOctree CollisionAccurateOctree;
 
 #if WITH_EDITOR
 	const UPCGPointData* RawPointData = nullptr;
diff --git a/Engine/Plugins/PCG/Source/PCG/Public/Data/PCGPointData.h b/Engine/Plugins/PCG/Source/PCG/Public/Data/PCGPointData.h
--- a/Engine/Plugins/PCG/Source/PCG/Public/Data/PCGPointData.h
+++ b/Engine/Plugins/PCG/Source/PCG/Public/Data/PCGPointData.h
@@ -28,6 +28,7 @@ namespace PCGPointDataConstants
 struct PCG_API FPCGPointRef
 {
 	FPCGPointRef(const FPCGPoint& InPoint);
+	FPCGPointRef(const FPCGPoint& InPoint, const FBox& InBoundsOverride);
 
 	const FPCGPoint* Point;
 	FBoxSphereBounds Bounds;
diff --git a/Engine/Plugins/PCG/Source/PCG/Public/Elements/PCGCreateCollisionData.h b/Engine/Plugins/PCG/Source/PCG/Public/Elements/PCGCreateCollisionData.h
--- a/Engine/Plugins/PCG/Source/PCG/Public/Elements/PCGCreateCollisionData.h
+++ b/Engine/Plugins/PCG/Source/PCG/Public/Elements/PCGCreateCollisionData.h
@@ -33,6 +33,10 @@ public:
 	UPROPERTY(BlueprintReadWrite, EditAnywhere, Category = Settings, meta = (PCG_Overridable, PCG_DiscardPropertySelection, PCG_DiscardExtraSelection))
 	FPCGAttributePropertyInputSelector CollisionAttribute;
 
+	/** Uses a new octree based on the mesh bounds, performance warning (does similar work to the BoundsFromMesh node, but will not change the point data). */
+	UPROPERTY(BlueprintReadWrite, EditAnywhere, Category = Settings, meta = (PCG_Overridable))
+	bool bRecomputeOctreeAccordingToMeshes = false;
+
 	/** Queries against complex collision if enabled, performance warning */
 	UPROPERTY(BlueprintReadWrite, EditAnywhere, Category = Settings, meta = (PCG_Overridable))
 	bool bUseComplexCollision = false;
diff --git a/Engine/Plugins/PCG/Source/PCG/Public/Elements/PCGSelfPruning.h b/Engine/Plugins/PCG/Source/PCG/Public/Elements/PCGSelfPruning.h
--- a/Engine/Plugins/PCG/Source/PCG/Public/Elements/PCGSelfPruning.h
+++ b/Engine/Plugins/PCG/Source/PCG/Public/Elements/PCGSelfPruning.h
@@ -49,6 +49,10 @@ public:
 	UPROPERTY(BlueprintReadWrite, EditAnywhere, Category = Settings, meta = (PCG_Overridable, PCG_DiscardPropertySelection, PCG_DiscardExtraSelection, EditCondition = "bUseCollisionAttribute && PruningType != EPCGSelfPruningType::RemoveDuplicates", EditConditionHides))
 	FPCGAttributePropertyInputSelector CollisionAttribute;
 
+	/** Uses a new octree based on the mesh bounds, performance warning (does similar work to the BoundsFromMesh node, but will not change the point data). */
+	UPROPERTY(BlueprintReadWrite, EditAnywhere, Category = Settings, meta = (PCG_Overridable, EditCondition = "bUseCollisionAttribute && PruningType != EPCGSelfPruningType::RemoveDuplicates", EditConditionHides))
+	bool bRecomputeOctreeAccordingToMeshes = false;
+
 	/** Queries against complex collision (vs. collection of convex elements in the collision) if enabled, performance warning */
 	UPROPERTY(BlueprintReadWrite, EditAnywhere, Category = Settings, meta = (PCG_Overridable, EditCondition = "bUseCollisionAttribute && PruningType != EPCGSelfPruningType::RemoveDuplicates", EditConditionHides))
 	bool bUseComplexCollision = false;
@@ -77,10 +81,13 @@ namespace PCGSelfPruningElement
 		FPointBitSet ExclusionPoints;
 		int32 CurrentPointIndex = 0;
 		bool bSortDone = false;
+		bool bSortedPointsArrayPopulateDone = false;
+		bool bUseCollisionAccurateOctree = false;
 
 		// In the case of the self-pruning using the collision, we'll need instances of the unique bodies that are used in the point data.
 		// Additionally, in the cases where two points to compare use the same instance, we'll have to do a temporary copy.
 		FPCGCollisionWrapper CollisionWrapper;
+		UPCGPointData::PointOctree CollisionAccurateOctree;
 		TMap<FBodyInstance*, FBodyInstance*> TemporaryBodyInstances;
 	};
 
@@ -141,4 +148,5 @@ class FPCGSelfPruningElement : public TPCGTimeSlicedElementBase<PCGTimeSlice::FE
 protected:
 	virtual bool PrepareDataInternal(FPCGContext* Context) const override;
 	virtual bool ExecuteInternal(FPCGContext* Context) const override;
+	virtual bool CanExecuteOnlyOnMainThread(FPCGContext* Context) const override;
 };
