[PCG] Trigger an error when adding resources that aren't outered properly

diff --git a/Engine/Plugins/PCG/Source/PCG/Private/Helpers/PCGBlueprintHelpers.cpp b/Engine/Plugins/PCG/Source/PCG/Private/Helpers/PCGBlueprintHelpers.cpp
--- a/Engine/Plugins/PCG/Source/PCG/Private/Helpers/PCGBlueprintHelpers.cpp
+++ b/Engine/Plugins/PCG/Source/PCG/Private/Helpers/PCGBlueprintHelpers.cpp
@@ -11,12 +11,26 @@
 #include "Grid/PCGLandscapeCache.h"
 #include "Helpers/PCGHelpers.h"
 
+#include "Blueprint/BlueprintExceptionInfo.h"
 #include "Engine/World.h"
 #include "LandscapeProxy.h"
 #include "LandscapeInfo.h"
 
 #include UE_INLINE_GENERATED_CPP_BY_NAME(PCGBlueprintHelpers)
 
+void UPCGBlueprintHelpers::ThrowBlueprintException(const FText& ErrorMessage)
+{
+	if (FFrame::GetThreadLocalTopStackFrame() && FFrame::GetThreadLocalTopStackFrame()->Object)
+	{
+		const FBlueprintExceptionInfo ExceptionInfo(EBlueprintExceptionType::FatalError, ErrorMessage);
+		FBlueprintCoreDelegates::ThrowScriptException(FFrame::GetThreadLocalTopStackFrame()->Object, *FFrame::GetThreadLocalTopStackFrame(), ExceptionInfo);
+	}
+	else
+	{
+		UE_LOG(LogPCG, Error, TEXT("%s"), *ErrorMessage.ToString());
+	}
+}
+
 int UPCGBlueprintHelpers::ComputeSeedFromPosition(const FVector& InPosition)
 {
 	return PCGHelpers::ComputeSeedFromPosition(InPosition);
diff --git a/Engine/Plugins/PCG/Source/PCG/Private/Helpers/PCGGraphParametersHelpers.cpp b/Engine/Plugins/PCG/Source/PCG/Private/Helpers/PCGGraphParametersHelpers.cpp
--- a/Engine/Plugins/PCG/Source/PCG/Private/Helpers/PCGGraphParametersHelpers.cpp
+++ b/Engine/Plugins/PCG/Source/PCG/Private/Helpers/PCGGraphParametersHelpers.cpp
@@ -4,9 +4,9 @@
 
 #include "PCGModule.h"
 #include "PCGGraph.h"
+#include "Helpers/PCGBlueprintHelpers.h"
 
 #include "PropertyBag.h"
-#include "Blueprint/BlueprintExceptionInfo.h"
 #include "Templates/ValueOrError.h"
 #include "UObject/Script.h"
 #include "UObject/Stack.h"
@@ -20,15 +20,7 @@ namespace PCGGraphParametersHelpersPrivate
 
 	void ThrowBlueprintException(const FText& ErrorMessage)
 	{
-		if (FFrame::GetThreadLocalTopStackFrame() && FFrame::GetThreadLocalTopStackFrame()->Object)
-		{
-			const FBlueprintExceptionInfo ExceptionInfo(EBlueprintExceptionType::FatalError, ErrorMessage);
-			FBlueprintCoreDelegates::ThrowScriptException(FFrame::GetThreadLocalTopStackFrame()->Object, *FFrame::GetThreadLocalTopStackFrame(), ExceptionInfo);
-		}
-		else
-		{
-			UE_LOG(LogPCG, Error, TEXT("%s"), *ErrorMessage.ToString());
-		}
+		UPCGBlueprintHelpers::ThrowBlueprintException(ErrorMessage);
 	}
 
 	void OnException(const EPropertyBagResult Result, const FName PropertyName)
diff --git a/Engine/Plugins/PCG/Source/PCG/Private/Metadata/PCGMetadataAccessor.cpp b/Engine/Plugins/PCG/Source/PCG/Private/Metadata/PCGMetadataAccessor.cpp
--- a/Engine/Plugins/PCG/Source/PCG/Private/Metadata/PCGMetadataAccessor.cpp
+++ b/Engine/Plugins/PCG/Source/PCG/Private/Metadata/PCGMetadataAccessor.cpp
@@ -4,10 +4,9 @@
 
 #include "PCGModule.h"
 #include "PCGPoint.h"
+#include "Helpers/PCGBlueprintHelpers.h"
 #include "Metadata/PCGMetadata.h"
 
-#include "Blueprint/BlueprintExceptionInfo.h"
-
 #include UE_INLINE_GENERATED_CPP_BY_NAME(PCGMetadataAccessor)
 
 #define LOCTEXT_NAMESPACE "PCGMetadataAccessorHelpers"
@@ -18,17 +17,9 @@ namespace PCGMetadataAccessorHelpers
 	static const FText InvalidTypeFormat = LOCTEXT("InvalidTypeFormat", "[PCG] Attribute {0} is of type {1} and it is not the requested type {2}");
 	static const FText NoMetadata = LOCTEXT("NoMetadata", "[PCG] No metadata provided");
 
-	void OnException(FText ErrorMessage)
+	void OnException(const FText& ErrorMessage)
 	{
-		if (FFrame::GetThreadLocalTopStackFrame() && FFrame::GetThreadLocalTopStackFrame()->Object)
-		{
-			const FBlueprintExceptionInfo ExceptionInfo(EBlueprintExceptionType::FatalError, ErrorMessage);
-			FBlueprintCoreDelegates::ThrowScriptException(FFrame::GetThreadLocalTopStackFrame()->Object, *FFrame::GetThreadLocalTopStackFrame(), ExceptionInfo);
-		}
-		else
-		{
-			UE_LOG(LogPCG, Error, TEXT("%s"), *ErrorMessage.ToString());
-		}
+		UPCGBlueprintHelpers::ThrowBlueprintException(ErrorMessage);
 	}
 }
 
diff --git a/Engine/Plugins/PCG/Source/PCG/Private/PCGComponent.cpp b/Engine/Plugins/PCG/Source/PCG/Private/PCGComponent.cpp
--- a/Engine/Plugins/PCG/Source/PCG/Private/PCGComponent.cpp
+++ b/Engine/Plugins/PCG/Source/PCG/Private/PCGComponent.cpp
@@ -26,6 +26,7 @@
 #include "Graph/PCGStackContext.h"
 #include "Grid/PCGPartitionActor.h"
 #include "Helpers/PCGActorHelpers.h"
+#include "Helpers/PCGBlueprintHelpers.h"
 #include "Helpers/PCGHelpers.h"
 #include "RuntimeGen/GenSources/PCGGenSourceBase.h"
 #include "RuntimeGen/SchedulingPolicies/PCGSchedulingPolicyBase.h"
@@ -164,6 +165,11 @@ void UPCGComponent::AddToManagedResources(UPCGManagedResource* InResource)
 
 	if (InResource)
 	{
+		if(!ensure(InResource->GetOuter() == this))
+		{
+			UPCGBlueprintHelpers::ThrowBlueprintException(LOCTEXT("ResourceNotOutered", "Managed resources need to be outered to their PCG component."));
+		}
+
 		FScopeLock ResourcesLock(&GeneratedResourcesLock);
 		check(!GeneratedResourcesInaccessible);
 		GeneratedResources.Add(InResource);
diff --git a/Engine/Plugins/PCG/Source/PCG/Public/Helpers/PCGBlueprintHelpers.h b/Engine/Plugins/PCG/Source/PCG/Public/Helpers/PCGBlueprintHelpers.h
--- a/Engine/Plugins/PCG/Source/PCG/Public/Helpers/PCGBlueprintHelpers.h
+++ b/Engine/Plugins/PCG/Source/PCG/Public/Helpers/PCGBlueprintHelpers.h
@@ -3,9 +3,8 @@
 #pragma once
 
 #include "Kismet/BlueprintFunctionLibrary.h"
-
-
 #include "Math/Box.h"
+
 #include "PCGBlueprintHelpers.generated.h"
 
 class UPCGComponent;
@@ -22,6 +21,8 @@ class PCG_API UPCGBlueprintHelpers : public UBlueprintFunctionLibrary
 	GENERATED_BODY()
 
 public:
+	static void ThrowBlueprintException(const FText& ErrorMessage);
+
 	UFUNCTION(BlueprintCallable, Category = "PCG|Helpers")
 	static int ComputeSeedFromPosition(const FVector& InPosition);
 
