[PCG] Add PCGDataVisualization interface and PCGDataVisualizationFactory factory
* Default collapse to points for Spatial Data
* Custom collapse to control points for Spline Data

diff --git a/Engine/Plugins/PCG/Source/PCG/Private/Elements/PCGDebugElement.cpp b/Engine/Plugins/PCG/Source/PCG/Private/Elements/PCGDebugElement.cpp
--- a/Engine/Plugins/PCG/Source/PCG/Private/Elements/PCGDebugElement.cpp
+++ b/Engine/Plugins/PCG/Source/PCG/Private/Elements/PCGDebugElement.cpp
@@ -2,244 +2,69 @@
 
 #include "Elements/PCGDebugElement.h"
 
-#include "PCGComponent.h"
 #include "PCGContext.h"
-#include "PCGPin.h"
-#include "Data/PCGSpatialData.h"
-#include "Data/PCGPointData.h"
-#include "Helpers/PCGActorHelpers.h"
-#include "Helpers/PCGHelpers.h"
-
-#include "Components/InstancedStaticMeshComponent.h"
-#include "Engine/CollisionProfile.h"
-#include "Engine/StaticMesh.h"
-#include "ISMPartition/ISMComponentDescriptor.h"
-#include "Materials/MaterialInterface.h"
+#include "PCGModule.h"
+
+#if WITH_EDITOR
+#include "PCGDataVisualization.h"
+#endif
 
 #include UE_INLINE_GENERATED_CPP_BY_NAME(PCGDebugElement)
 
 #define LOCTEXT_NAMESPACE "PCGDebugElement"
 
-namespace PCGDebugElement
+FPCGElementPtr UPCGDebugSettings::CreateElement() const
 {
-	void ExecuteDebugDisplay(FPCGContext* Context, AActor* InTargetActor)
-	{
-		TRACE_CPUPROFILER_EVENT_SCOPE(PCGDebugElement::ExecuteDebugDisplay);
-#if WITH_EDITOR
-		// Early validation: if we don't have a valid PCG component, we're not going to add the debug display info
-		if (!Context->SourceComponent.IsValid())
-		{
-			return;
-		}
+	return MakeShared<FPCGDebugElement>();
+}
 
-		const UPCGSettingsInterface* SettingsInterface = Context->GetInputSettingsInterface();
+TArray<FPCGPinProperties> UPCGDebugSettings::InputPinProperties() const
+{
+	TArray<FPCGPinProperties> PinProperties;
+	PinProperties.Emplace(PCGPinConstants::DefaultInputLabel, EPCGDataType::Any, /*bInAllowMultipleConnections=*/true, /*bAllowMultipleData=*/true);
 
-		if (!SettingsInterface)
-		{
-			return;
-		}
+	return PinProperties;
+}
 
-		const FPCGDebugVisualizationSettings& DebugSettings = SettingsInterface->DebugSettings;
+TArray<FPCGPinProperties> UPCGDebugSettings::OutputPinProperties() const
+{
+	return TArray<FPCGPinProperties>();
+}
 
-		UStaticMesh* Mesh = DebugSettings.PointMesh.LoadSynchronous();
+bool FPCGDebugElement::ExecuteInternal(FPCGContext* Context) const
+{
+	TRACE_CPUPROFILER_EVENT_SCOPE(FPCGDebugElement::Execute);
 
-		if (!Mesh)
-		{
-			PCGE_LOG_C(Error, GraphAndLog, Context, FText::Format(LOCTEXT("UnableToLoadMesh", "Debug display was unable to load mesh '{0}'"),
-				FText::FromString(DebugSettings.PointMesh.ToString())));
-			return;
-		}
+#if WITH_EDITOR
+	// Early validation: if we don't have a valid PCG component, we're not going to add the debug display info.
+	if (!Context->SourceComponent.IsValid())
+	{
+		return true;
+	}
 
-		UMaterialInterface* Material = DebugSettings.GetMaterial().LoadSynchronous();
+	const FPCGDataVisualizationRegistry& DataVisRegistry = FPCGModule::GetConstPCGDataVisualizationRegistry();
+	TArray<FPCGTaggedData> Inputs = Context->InputData.GetInputs();
 
-		TArray<UMaterialInterface*> Materials;
-		if (Material)
+	for (const FPCGTaggedData& Input : Inputs)
+	{
+		if (!Input.Data)
 		{
-			Materials.Add(Material);
+			continue;
 		}
 
-		// In the case of a node with multiple output pins, we will select only the inputs from the first non-empty pin.
-		bool bFilterOnPin = false;
-		FName PinFilter = NAME_None;
+		AActor* TargetActor = Context->GetInputSettings<UPCGDebugSettings>()->TargetActor.Get();
 
-		if (Context->Node)
+		if (!TargetActor)
 		{
-			TArray<FPCGPinProperties> OutputPinProperties = Context->Node->OutputPinProperties();
-			for(const FPCGPinProperties& OutPin : OutputPinProperties)
-			{
-				if (Context->Node->IsOutputPinConnected(OutPin.Label))
-				{
-					PinFilter = OutPin.Label;
-					bFilterOnPin = true;
-					break;
-				}
-			}
+			TargetActor = Context->GetTargetActor(nullptr);
 		}
 
-		TArray<FPCGTaggedData> Inputs = Context->InputData.GetInputs();
-		for(const FPCGTaggedData& Input : Inputs)
+		if (const IPCGDataVisualization* DataVis = DataVisRegistry.GetDataVisualization(Input.Data->GetClass()))
 		{
-			// Skip output if we're filtering on the first pin
-			if (bFilterOnPin && Input.Pin != PinFilter)
-			{
-				continue;
-			}
-
-			const UPCGSpatialData* SpatialData = Cast<UPCGSpatialData>(Input.Data);
-
-			if (!SpatialData)
-			{
-				// Data type mismatch
-				continue;
-			}
-
-			AActor* TargetActor = InTargetActor ? InTargetActor : Context->GetTargetActor(nullptr);
-
-			if (!TargetActor)
-			{
-				// No target actor
-				PCGE_LOG_C(Error, GraphAndLog, Context, LOCTEXT("NoTargetActor", "Debug display cannot show data that have no target actor"));
-				continue;
-			}
-
-			const UPCGPointData* PointData = SpatialData->ToPointData(Context);
-
-			if (!PointData)
-			{
-				continue;
-			}
-
-			const TArray<FPCGPoint>& Points = PointData->GetPoints();
-
-			if (Points.Num() == 0)
-			{
-				continue;
-			}
-
-			const int NumCustomData = 8;
-
-			TArray<FTransform> ForwardInstances;
-			TArray<FTransform> ReverseInstances;
-			TArray<float> InstanceCustomData;
-
-			ForwardInstances.Reserve(Points.Num());
-			InstanceCustomData.Reserve(NumCustomData);
-
-			// First, create target instance transforms
-			const float PointScale = DebugSettings.PointScale;
-			const bool bIsAbsolute = DebugSettings.ScaleMethod == EPCGDebugVisScaleMethod::Absolute;
-			const bool bIsRelative = DebugSettings.ScaleMethod == EPCGDebugVisScaleMethod::Relative;
-			const bool bScaleWithExtents = DebugSettings.ScaleMethod == EPCGDebugVisScaleMethod::Extents;
-			const FVector MeshExtents = Mesh->GetBoundingBox().GetExtent();
-			const FVector MeshCenter = Mesh->GetBoundingBox().GetCenter();
-
-			for (const FPCGPoint& Point : Points)
-			{
-				TArray<FTransform>& Instances = ((bIsAbsolute || Point.Transform.GetDeterminant() >= 0) ? ForwardInstances : ReverseInstances);
-				FTransform& InstanceTransform = Instances.Add_GetRef(Point.Transform);
-				if (bIsRelative)
-				{
-					InstanceTransform.SetScale3D(InstanceTransform.GetScale3D() * PointScale);
-				}
-				else if (bScaleWithExtents)
-				{
-					const FVector ScaleWithExtents = Point.GetExtents() / MeshExtents;
-					const FVector TransformedBoxCenterWithOffset = InstanceTransform.TransformPosition(Point.GetLocalCenter()) - InstanceTransform.GetLocation();
-					InstanceTransform.SetTranslation(InstanceTransform.GetTranslation() + TransformedBoxCenterWithOffset);
-					InstanceTransform.SetScale3D(InstanceTransform.GetScale3D() * ScaleWithExtents);
-				}
-				else // absolute scaling only
-				{
-					InstanceTransform.SetScale3D(FVector(PointScale));
-				}
-			}
-
-			FPCGISMCBuilderParameters Params[2];
-			Params[0].NumCustomDataFloats = NumCustomData;
-			Params[0].Descriptor.StaticMesh = Mesh;
-			Params[0].Descriptor.OverrideMaterials = Materials;
-			Params[0].Descriptor.BodyInstance.SetCollisionProfileName(UCollisionProfile::NoCollision_ProfileName);
-			// Note: In the future we may consider enabling culling for performance reasons, but for now culling disabled.
-			Params[0].Descriptor.InstanceStartCullDistance = Params[0].Descriptor.InstanceEndCullDistance = 0;
-			
-			// If the root actor we're binding to is movable, then the ISMC should be movable by default
-			if (USceneComponent* SceneComponent = TargetActor->GetRootComponent())
-			{
-				Params[0].Descriptor.Mobility = SceneComponent->Mobility;
-			}
-
-			Params[1] = Params[0];
-			Params[1].Descriptor.bReverseCulling = true;
-
-			for (int32 Direction = 0; Direction < 2; ++Direction)
-			{
-				TArray<FTransform>& Instances = (Direction == 0 ? ForwardInstances : ReverseInstances);
-
-				if (Instances.IsEmpty())
-				{
-					continue;
-				}
-
-				UInstancedStaticMeshComponent* ISMC = UPCGActorHelpers::GetOrCreateISMC(TargetActor, Context->SourceComponent.Get(), SettingsInterface->GetSettings()->UID, Params[Direction]);
-				check(ISMC && ISMC->NumCustomDataFloats == NumCustomData);
-
-				ISMC->ComponentTags.AddUnique(PCGHelpers::DefaultPCGDebugTag);
-				const int32 PreExistingInstanceCount = ISMC->GetInstanceCount();
-				ISMC->AddInstances(Instances, /*bShouldReturnIndices=*/false, /*bWorldSpace=*/true);
-
-				// Scan all points looking for points that match current direction and add their custom data.
-				int32 PointCounter = 0;
-				for (const FPCGPoint& Point : Points)
-				{
-					const int32 PointDirection = ((bIsAbsolute || Point.Transform.GetDeterminant() >= 0) ? 0 : 1);
-					if (PointDirection != Direction)
-					{
-						continue;
-					}
-
-					InstanceCustomData.Add(Point.Density);
-					const FVector Extents = Point.GetExtents();
-					InstanceCustomData.Add(Extents[0]);
-					InstanceCustomData.Add(Extents[1]);
-					InstanceCustomData.Add(Extents[2]);
-					InstanceCustomData.Add(Point.Color[0]);
-					InstanceCustomData.Add(Point.Color[1]);
-					InstanceCustomData.Add(Point.Color[2]);
-					InstanceCustomData.Add(Point.Color[3]);
-
-					ISMC->SetCustomData(PreExistingInstanceCount + PointCounter, InstanceCustomData);
-
-					InstanceCustomData.Reset();
-
-					++PointCounter;
-				}
-
-				ISMC->UpdateBounds();
-			}
+			DataVis->ExecuteDebugDisplay(Context, Input.Data, TargetActor);
 		}
-#endif
 	}
-}
-
-FPCGElementPtr UPCGDebugSettings::CreateElement() const
-{
-	return MakeShared<FPCGDebugElement>();
-}
-
-TArray<FPCGPinProperties> UPCGDebugSettings::InputPinProperties() const
-{
-	return Super::DefaultPointInputPinProperties();
-}
-
-TArray<FPCGPinProperties> UPCGDebugSettings::OutputPinProperties() const
-{
-	return TArray<FPCGPinProperties>();
-}
-
-bool FPCGDebugElement::ExecuteInternal(FPCGContext* Context) const
-{
-	TRACE_CPUPROFILER_EVENT_SCOPE(FPCGDebugElement::Execute);
-	PCGDebugElement::ExecuteDebugDisplay(Context, Context->GetInputSettings<UPCGDebugSettings>()->TargetActor.Get());
+#endif
 	
 	return true;
 }
diff --git a/Engine/Plugins/PCG/Source/PCG/Private/PCGDataVisualizationRegistry.cpp b/Engine/Plugins/PCG/Source/PCG/Private/PCGDataVisualizationRegistry.cpp
new file mode 100644
--- /dev/null
+++ b/Engine/Plugins/PCG/Source/PCG/Private/PCGDataVisualizationRegistry.cpp
@@ -0,0 +1,52 @@
+// Copyright Epic Games, Inc. All Rights Reserved.
+
+#include "PCGDataVisualizationRegistry.h"
+
+#include "Data/PCGSpatialData.h"
+
+#if WITH_EDITOR
+void FPCGDataVisualizationRegistry::RegisterPCGDataVisualization(const TSubclassOf<UPCGData> DataClass, TUniquePtr<const IPCGDataVisualization> Visualization)
+{
+	// TODO: We can support multiple visualizations in the future, but special care will be needed to avoid multiple types competing to own the ALV visualization.
+	if (ensureMsgf(!ExternalRegistry.Find(DataClass), TEXT("Cannot register multiple data visualizations of the same type.")))
+	{
+		ExternalRegistry.Add(DataClass, std::move(Visualization));
+	}
+}
+
+void FPCGDataVisualizationRegistry::UnregisterPCGDataVisualization(const TSubclassOf<UPCGData> DataClass)
+{
+	ExternalRegistry.Remove(DataClass);
+}
+
+const IPCGDataVisualization* FPCGDataVisualizationRegistry::GetDataVisualization(const TSubclassOf<UPCGData> PCGDataTypeClass) const
+{
+	TSubclassOf<UPCGData> CurrentClass = PCGDataTypeClass;
+	const TSubclassOf<UPCGData> UPCGDataSuperClass = UPCGData::StaticClass()->GetSuperClass();
+
+	// Iterate up the class hierarchy to find the first data class with a visualization implementation.
+	while (CurrentClass && CurrentClass != UPCGDataSuperClass)
+	{
+		// Always try external implementations first, so the user can author their own visualizations for built-in PCG types if they desire.
+		if (const TUniquePtr<const IPCGDataVisualization>* VisualizationPtr = ExternalRegistry.Find(CurrentClass))
+		{
+			return VisualizationPtr->Get();
+		}
+
+		if (const TUniquePtr<const IPCGDataVisualization>* VisualizationPtr = InternalRegistry.Find(CurrentClass))
+		{
+			return VisualizationPtr->Get();
+		}
+
+		CurrentClass = CurrentClass->GetSuperClass();
+
+		// Skip SpatialDataWithPointCache since it will never have a visualization implementation.
+		if (CurrentClass == UPCGSpatialDataWithPointCache::StaticClass())
+		{
+			CurrentClass = CurrentClass->GetSuperClass();
+		}
+	}
+
+	return nullptr;
+}
+#endif
\ No newline at end of file
diff --git a/Engine/Plugins/PCG/Source/PCG/Private/PCGElement.cpp b/Engine/Plugins/PCG/Source/PCG/Private/PCGElement.cpp
--- a/Engine/Plugins/PCG/Source/PCG/Private/PCGElement.cpp
+++ b/Engine/Plugins/PCG/Source/PCG/Private/PCGElement.cpp
@@ -7,7 +7,10 @@
 #include "PCGGraph.h"
 #include "PCGSubsystem.h"
 #include "Data/PCGPointData.h"
-#include "Elements/PCGDebugElement.h"
+
+#if WITH_EDITOR
+#include "PCGDataVisualization.h"
+#endif
 
 #include "HAL/IConsoleManager.h"
 #include "Utils/PCGExtraCapture.h"
@@ -342,7 +345,25 @@ void IPCGElement::DebugDisplay(FPCGContext* Context) const
 	Context->InputData = ElementOutputs;
 	Context->OutputData = FPCGDataCollection();
 
-	PCGDebugElement::ExecuteDebugDisplay(Context);
+	// In the case of a node with multiple output pins, we will select only the inputs from the first non-empty pin.
+	const UPCGPin* FirstOutPin = Context->Node ? Context->Node->GetFirstConnectedOutputPin() : nullptr;
+
+	const FPCGDataVisualizationRegistry& DataVisRegistry = FPCGModule::GetConstPCGDataVisualizationRegistry();
+	TArray<FPCGTaggedData> Inputs = Context->InputData.GetInputs();
+
+	for (const FPCGTaggedData& Input : Inputs)
+	{
+		// Skip output if we're filtering on the first pin or the the data is null.
+		if (!Input.Data || (FirstOutPin && FirstOutPin->Properties.Label != Input.Pin))
+		{
+			continue;
+		}
+
+		if (const IPCGDataVisualization* DataVis = DataVisRegistry.GetDataVisualization(Input.Data->GetClass()))
+		{
+			DataVis->ExecuteDebugDisplay(Context, Input.Data, Context->GetTargetActor(nullptr));
+		}
+	}
 
 	Context->InputData = ElementInputs;
 	Context->OutputData = ElementOutputs;
diff --git a/Engine/Plugins/PCG/Source/PCG/Private/PCGModule.cpp b/Engine/Plugins/PCG/Source/PCG/Private/PCGModule.cpp
--- a/Engine/Plugins/PCG/Source/PCG/Private/PCGModule.cpp
+++ b/Engine/Plugins/PCG/Source/PCG/Private/PCGModule.cpp
@@ -17,30 +17,6 @@
 
 #define LOCTEXT_NAMESPACE "FPCGModule"
 
-class FPCGModule final : public IModuleInterface
-{
-public:
-	//~ IModuleInterface implementation
-
-#if WITH_EDITOR
-	virtual void StartupModule() override;
-	virtual void ShutdownModule() override;
-#endif
-
-	virtual bool SupportsDynamicReloading() override
-	{
-		return true;
-	}
-
-	//~ End IModuleInterface implementation
-
-private:
-#if WITH_EDITOR
-	void RegisterNativeElementDeterminismTests();
-	void DeregisterNativeElementDeterminismTests();
-#endif
-};
-
 #if WITH_EDITOR
 void FPCGModule::StartupModule()
 {
@@ -58,6 +34,11 @@ void FPCGModule::ShutdownModule()
 	PCGDeterminismTests::FNativeTestRegistry::Destroy();
 }
 
+FPCGModule& FPCGModule::GetPCGModuleChecked()
+{
+	return FModuleManager::GetModuleChecked<FPCGModule>("PCG");
+}
+
 void FPCGModule::RegisterNativeElementDeterminismTests()
 {
 	PCGDeterminismTests::FNativeTestRegistry::RegisterTestFunction(UPCGDifferenceSettings::StaticClass(), PCGDeterminismTests::DifferenceElement::RunTestSuite);
@@ -68,7 +49,6 @@ void FPCGModule::DeregisterNativeElementDeterminismTests()
 {
 	PCGDeterminismTests::FNativeTestRegistry::DeregisterTestFunction(UPCGDifferenceSettings::StaticClass());
 }
-
 #endif
 
 IMPLEMENT_MODULE(FPCGModule, PCG);
diff --git a/Engine/Plugins/PCG/Source/PCG/Private/PCGNode.cpp b/Engine/Plugins/PCG/Source/PCG/Private/PCGNode.cpp
--- a/Engine/Plugins/PCG/Source/PCG/Private/PCGNode.cpp
+++ b/Engine/Plugins/PCG/Source/PCG/Private/PCGNode.cpp
@@ -731,6 +731,19 @@ const UPCGPin* UPCGNode::GetFirstConnectedInputPin() const
 	return nullptr;
 }
 
+const UPCGPin* UPCGNode::GetFirstConnectedOutputPin() const
+{
+	for (const UPCGPin* OutputPin : OutputPins)
+	{
+		if (OutputPin && OutputPin->EdgeCount() > 0)
+		{
+			return OutputPin;
+		}
+	}
+
+	return nullptr;
+}
+
 void UPCGNode::SetSettingsInterface(UPCGSettingsInterface* InSettingsInterface, bool bUpdatePins)
 {
 	const bool bDifferentInterface = (SettingsInterface.Get() != InSettingsInterface);
diff --git a/Engine/Plugins/PCG/Source/PCG/Public/Elements/PCGDebugElement.h b/Engine/Plugins/PCG/Source/PCG/Public/Elements/PCGDebugElement.h
--- a/Engine/Plugins/PCG/Source/PCG/Public/Elements/PCGDebugElement.h
+++ b/Engine/Plugins/PCG/Source/PCG/Public/Elements/PCGDebugElement.h
@@ -8,11 +8,6 @@
 
 #include "PCGDebugElement.generated.h"
 
-namespace PCGDebugElement
-{
-	void ExecuteDebugDisplay(FPCGContext* Context, AActor* TargetActor = nullptr);
-}
-
 UCLASS(BlueprintType, ClassGroup = (Procedural))
 class PCG_API UPCGDebugSettings : public UPCGSettings
 {
diff --git a/Engine/Plugins/PCG/Source/PCG/Public/PCGDataVisualization.h b/Engine/Plugins/PCG/Source/PCG/Public/PCGDataVisualization.h
new file mode 100644
--- /dev/null
+++ b/Engine/Plugins/PCG/Source/PCG/Public/PCGDataVisualization.h
@@ -0,0 +1,17 @@
+// Copyright Epic Games, Inc. All Rights Reserved.
+
+#pragma once
+
+#if WITH_EDITOR
+class AActor;
+class UPCGData;
+struct FPCGContext;
+
+/** Implement this interface to provide custom PCGData visualizations. Register your implementation to FPCGModule::FPCGDataVisualizationRegistry to be used automatically. */
+class IPCGDataVisualization
+{
+public:
+	virtual ~IPCGDataVisualization() = default;
+	virtual void ExecuteDebugDisplay(FPCGContext* Context, const UPCGData* Data, AActor* TargetActor) const = 0;
+};
+#endif
diff --git a/Engine/Plugins/PCG/Source/PCG/Public/PCGDataVisualizationRegistry.h b/Engine/Plugins/PCG/Source/PCG/Public/PCGDataVisualizationRegistry.h
new file mode 100644
--- /dev/null
+++ b/Engine/Plugins/PCG/Source/PCG/Public/PCGDataVisualizationRegistry.h
@@ -0,0 +1,39 @@
+// Copyright Epic Games, Inc. All Rights Reserved.
+
+#pragma once
+
+#if WITH_EDITOR
+#include "PCGData.h"
+#include "PCGDataVisualization.h"
+
+/** Register a PCGData type with an IPCGDataVisualization to implement custom visualization behavior. */
+struct PCG_API FPCGDataVisualizationRegistry
+{
+public:
+	FPCGDataVisualizationRegistry() = default;
+	FPCGDataVisualizationRegistry(const FPCGDataVisualizationRegistry&) = delete;
+	FPCGDataVisualizationRegistry(FPCGDataVisualizationRegistry&&) = default;
+	FPCGDataVisualizationRegistry& operator=(const FPCGDataVisualizationRegistry&) = delete;
+	FPCGDataVisualizationRegistry& operator=(FPCGDataVisualizationRegistry&&) = default;
+	virtual ~FPCGDataVisualizationRegistry() = default;
+
+	/** Register an external PCGData visualization. */
+	void RegisterPCGDataVisualization(const TSubclassOf<UPCGData> DataClass, TUniquePtr<const IPCGDataVisualization> Visualization);
+
+	/** Unregister an external PCGData visualization. */
+	void UnregisterPCGDataVisualization(const TSubclassOf<UPCGData> DataClass);
+
+	template <typename PCGDataType, typename = typename std::enable_if_t<std::is_base_of_v<UPCGData, PCGDataType>>>
+	const IPCGDataVisualization* GetDataVisualization() const { return GetDataVisualization(PCGDataType::StaticClass()); }
+	const IPCGDataVisualization* GetDataVisualization(const TSubclassOf<UPCGData> PCGDataTypeClass) const;
+
+private:
+	/** Registry for PCGDataVisualizations defined inside the PCG Plugin. */
+	TMap<const TSubclassOf<UPCGData>, TUniquePtr<const IPCGDataVisualization>> InternalRegistry;
+
+	/** Registry for PCGDataVisualizations defined outside the PCG Plugin. These take priority, allowing the user to override default visualization behavior for internal types. */
+	TMap<const TSubclassOf<UPCGData>, TUniquePtr<const IPCGDataVisualization>> ExternalRegistry;
+
+	friend class FPCGEditorModule;
+};
+#endif
diff --git a/Engine/Plugins/PCG/Source/PCG/Public/PCGModule.h b/Engine/Plugins/PCG/Source/PCG/Public/PCGModule.h
--- a/Engine/Plugins/PCG/Source/PCG/Public/PCGModule.h
+++ b/Engine/Plugins/PCG/Source/PCG/Public/PCGModule.h
@@ -2,12 +2,16 @@
 
 #pragma once
 
+#include "PCGDataVisualizationRegistry.h"
+
+#include "Modules/ModuleInterface.h"
 #include "Stats/Stats.h"
 
 // Logs
 PCG_API DECLARE_LOG_CATEGORY_EXTERN(LogPCG, Log, All);
 
 struct FPCGContext;
+class IPCGDataVisualization;
 
 namespace PCGLog
 {
@@ -27,6 +31,35 @@ DECLARE_STATS_GROUP(TEXT("PCG"), STATGROUP_PCG, STATCAT_Advanced);
 
 // CVars
 
+class FPCGModule final : public IModuleInterface
+{
+public:
+	//~ IModuleInterface implementation
+#if WITH_EDITOR
+	virtual void StartupModule() override;
+	virtual void ShutdownModule() override;
+#endif
+	virtual bool SupportsDynamicReloading() override { return true; }
+	//~ End IModuleInterface implementation
+
+#if WITH_EDITOR
+	PCG_API static FPCGModule& GetPCGModuleChecked();
+
+private:
+	void RegisterNativeElementDeterminismTests();
+	void DeregisterNativeElementDeterminismTests();
+#endif
+
+#if WITH_EDITOR
+public:
+	static const FPCGDataVisualizationRegistry& GetConstPCGDataVisualizationRegistry() { return GetPCGModuleChecked().PCGDataVisualizationRegistry; }
+	static FPCGDataVisualizationRegistry& GetMutablePCGDataVisualizationRegistry() { return GetPCGModuleChecked().PCGDataVisualizationRegistry; }
+
+private:
+	FPCGDataVisualizationRegistry PCGDataVisualizationRegistry;
+#endif
+};
+
 #if UE_ENABLE_INCLUDE_ORDER_DEPRECATED_IN_5_2
 #include "CoreMinimal.h"
 #include "Modules/ModuleManager.h"
diff --git a/Engine/Plugins/PCG/Source/PCG/Public/PCGNode.h b/Engine/Plugins/PCG/Source/PCG/Public/PCGNode.h
--- a/Engine/Plugins/PCG/Source/PCG/Public/PCGNode.h
+++ b/Engine/Plugins/PCG/Source/PCG/Public/PCGNode.h
@@ -142,9 +142,12 @@ public:
 	/** True if the edge is being used by the node. UI will gray out unused pins. */
 	virtual bool IsEdgeUsedByNodeExecution(const UPCGEdge* InEdge) const;
 
-	/** Returns the first connected pin on the node */
+	/** Returns the first connected input pin on the node. */
 	const UPCGPin* GetFirstConnectedInputPin() const;
 
+	/** Returns the first connected output pin on the node. */
+	const UPCGPin* GetFirstConnectedOutputPin() const;
+
 	const TArray<TObjectPtr<UPCGPin>>& GetInputPins() const { return InputPins; }
 	const TArray<TObjectPtr<UPCGPin>>& GetOutputPins() const { return OutputPins; }
 
diff --git a/Engine/Plugins/PCG/Source/PCGEditor/Private/DataVisualizations/PCGSpatialDataVisualization.cpp b/Engine/Plugins/PCG/Source/PCGEditor/Private/DataVisualizations/PCGSpatialDataVisualization.cpp
new file mode 100644
--- /dev/null
+++ b/Engine/Plugins/PCG/Source/PCGEditor/Private/DataVisualizations/PCGSpatialDataVisualization.cpp
@@ -0,0 +1,176 @@
+// Copyright Epic Games, Inc. All Rights Reserved.
+
+#include "DataVisualizations/PCGSpatialDataVisualization.h"
+
+#include "PCGContext.h"
+#include "PCGDebug.h"
+#include "PCGSettings.h"
+#include "Data/PCGPointData.h"
+#include "Data/PCGSpatialData.h"
+#include "Helpers/PCGActorHelpers.h"
+#include "Helpers/PCGHelpers.h"
+
+#include "Components/InstancedStaticMeshComponent.h"
+
+#define LOCTEXT_NAMESPACE "PCGSpatialDataVisualization"
+
+void IPCGSpatialDataVisualization::ExecuteDebugDisplay(FPCGContext* Context, const UPCGData* Data, AActor* TargetActor) const
+{
+	if (!TargetActor)
+	{
+		PCGE_LOG_C(Error, GraphAndLog, Context, LOCTEXT("NoTargetActor", "Cannot execute debug display for spatial data with no target actor."));
+		return;
+	}
+
+	const UPCGSettingsInterface* SettingsInterface = Context ? Context->GetInputSettingsInterface() : nullptr;
+
+	if (!SettingsInterface)
+	{
+		return;
+	}
+
+	const FPCGDebugVisualizationSettings& DebugSettings = SettingsInterface->DebugSettings;
+	UStaticMesh* Mesh = DebugSettings.PointMesh.LoadSynchronous();
+
+	if (!Mesh)
+	{
+		PCGE_LOG_C(Error, GraphAndLog, Context, FText::Format(LOCTEXT("UnableToLoadMesh", "Debug display was unable to load mesh '{0}'."),
+			FText::FromString(DebugSettings.PointMesh.ToString())));
+		return;
+	}
+
+	UMaterialInterface* Material = DebugSettings.GetMaterial().LoadSynchronous();
+
+	TArray<UMaterialInterface*> Materials;
+	if (Material)
+	{
+		Materials.Add(Material);
+	}
+
+	const UPCGPointData* PointData = CollapseToDebugPointData(Context, Data);
+
+	if (!PointData)
+	{
+		return;
+	}
+
+	const TArray<FPCGPoint>& Points = PointData->GetPoints();
+
+	if (Points.Num() == 0)
+	{
+		return;
+	}
+
+	const int NumCustomData = 8;
+
+	TArray<FTransform> ForwardInstances;
+	TArray<FTransform> ReverseInstances;
+	TArray<float> InstanceCustomData;
+
+	ForwardInstances.Reserve(Points.Num());
+	InstanceCustomData.Reserve(NumCustomData);
+
+	// First, create target instance transforms
+	const float PointScale = DebugSettings.PointScale;
+	const bool bIsAbsolute = DebugSettings.ScaleMethod == EPCGDebugVisScaleMethod::Absolute;
+	const bool bIsRelative = DebugSettings.ScaleMethod == EPCGDebugVisScaleMethod::Relative;
+	const bool bScaleWithExtents = DebugSettings.ScaleMethod == EPCGDebugVisScaleMethod::Extents;
+	const FVector MeshExtents = Mesh->GetBoundingBox().GetExtent();
+	const FVector MeshCenter = Mesh->GetBoundingBox().GetCenter();
+
+	for (const FPCGPoint& Point : Points)
+	{
+		TArray<FTransform>& Instances = ((bIsAbsolute || Point.Transform.GetDeterminant() >= 0) ? ForwardInstances : ReverseInstances);
+		FTransform& InstanceTransform = Instances.Add_GetRef(Point.Transform);
+		if (bIsRelative)
+		{
+			InstanceTransform.SetScale3D(InstanceTransform.GetScale3D() * PointScale);
+		}
+		else if (bScaleWithExtents)
+		{
+			const FVector ScaleWithExtents = Point.GetExtents() / MeshExtents;
+			const FVector TransformedBoxCenterWithOffset = InstanceTransform.TransformPosition(Point.GetLocalCenter()) - InstanceTransform.GetLocation();
+			InstanceTransform.SetTranslation(InstanceTransform.GetTranslation() + TransformedBoxCenterWithOffset);
+			InstanceTransform.SetScale3D(InstanceTransform.GetScale3D() * ScaleWithExtents);
+		}
+		else // absolute scaling only
+		{
+			InstanceTransform.SetScale3D(FVector(PointScale));
+		}
+	}
+
+	FPCGISMCBuilderParameters Params[2];
+	Params[0].NumCustomDataFloats = NumCustomData;
+	Params[0].Descriptor.StaticMesh = Mesh;
+	Params[0].Descriptor.OverrideMaterials = Materials;
+	Params[0].Descriptor.BodyInstance.SetCollisionProfileName(UCollisionProfile::NoCollision_ProfileName);
+	// Note: In the future we may consider enabling culling for performance reasons, but for now culling disabled.
+	Params[0].Descriptor.InstanceStartCullDistance = Params[0].Descriptor.InstanceEndCullDistance = 0;
+	
+	// If the root actor we're binding to is movable, then the ISMC should be movable by default
+	if (USceneComponent* SceneComponent = TargetActor->GetRootComponent())
+	{
+		Params[0].Descriptor.Mobility = SceneComponent->Mobility;
+	}
+
+	Params[1] = Params[0];
+	Params[1].Descriptor.bReverseCulling = true;
+
+	for (int32 Direction = 0; Direction < 2; ++Direction)
+	{
+		TArray<FTransform>& Instances = (Direction == 0 ? ForwardInstances : ReverseInstances);
+
+		if (Instances.IsEmpty())
+		{
+			continue;
+		}
+
+		UInstancedStaticMeshComponent* ISMC = UPCGActorHelpers::GetOrCreateISMC(TargetActor, Context->SourceComponent.Get(), SettingsInterface->GetSettings()->UID, Params[Direction]);
+		check(ISMC && ISMC->NumCustomDataFloats == NumCustomData);
+
+		ISMC->ComponentTags.AddUnique(PCGHelpers::DefaultPCGDebugTag);
+		const int32 PreExistingInstanceCount = ISMC->GetInstanceCount();
+		ISMC->AddInstances(Instances, /*bShouldReturnIndices=*/false, /*bWorldSpace=*/true);
+
+		// Scan all points looking for points that match current direction and add their custom data.
+		int32 PointCounter = 0;
+		for (const FPCGPoint& Point : Points)
+		{
+			const int32 PointDirection = ((bIsAbsolute || Point.Transform.GetDeterminant() >= 0) ? 0 : 1);
+			if (PointDirection != Direction)
+			{
+				continue;
+			}
+
+			InstanceCustomData.Add(Point.Density);
+			const FVector Extents = Point.GetExtents();
+			InstanceCustomData.Add(Extents[0]);
+			InstanceCustomData.Add(Extents[1]);
+			InstanceCustomData.Add(Extents[2]);
+			InstanceCustomData.Add(Point.Color[0]);
+			InstanceCustomData.Add(Point.Color[1]);
+			InstanceCustomData.Add(Point.Color[2]);
+			InstanceCustomData.Add(Point.Color[3]);
+
+			ISMC->SetCustomData(PreExistingInstanceCount + PointCounter, InstanceCustomData);
+
+			InstanceCustomData.Reset();
+
+			++PointCounter;
+		}
+
+		ISMC->UpdateBounds();
+	}
+}
+
+const UPCGPointData* IPCGSpatialDataVisualization::CollapseToDebugPointData(FPCGContext* Context, const UPCGData* Data) const
+{
+	if (const UPCGSpatialData* SpatialData = Cast<UPCGSpatialData>(Data))
+	{
+		return SpatialData->ToPointData(Context);
+	}
+
+	return nullptr;
+}
+
+#undef LOCTEXT_NAMESPACE
diff --git a/Engine/Plugins/PCG/Source/PCGEditor/Private/DataVisualizations/PCGSpatialDataVisualization.h b/Engine/Plugins/PCG/Source/PCGEditor/Private/DataVisualizations/PCGSpatialDataVisualization.h
new file mode 100644
--- /dev/null
+++ b/Engine/Plugins/PCG/Source/PCGEditor/Private/DataVisualizations/PCGSpatialDataVisualization.h
@@ -0,0 +1,19 @@
+// Copyright Epic Games, Inc. All Rights Reserved.
+
+#pragma once
+
+#include "PCGDataVisualization.h"
+#include "PCGEditorModule.h"
+
+class AActor;
+class UPCGData;
+class UPCGPointData;
+struct FPCGContext;
+
+/** Default implementation for spatial data. Collapses to a PointData representation. */
+class PCGEDITOR_API IPCGSpatialDataVisualization : public IPCGDataVisualization
+{
+public:
+	virtual void ExecuteDebugDisplay(FPCGContext* Context, const UPCGData* Data, AActor* TargetActor) const override;
+	virtual const UPCGPointData* CollapseToDebugPointData(FPCGContext* Context, const UPCGData* Data) const;
+};
diff --git a/Engine/Plugins/PCG/Source/PCGEditor/Private/DataVisualizations/PCGSplineDataVisualization.cpp b/Engine/Plugins/PCG/Source/PCGEditor/Private/DataVisualizations/PCGSplineDataVisualization.cpp
new file mode 100644
--- /dev/null
+++ b/Engine/Plugins/PCG/Source/PCGEditor/Private/DataVisualizations/PCGSplineDataVisualization.cpp
@@ -0,0 +1,29 @@
+// Copyright Epic Games, Inc. All Rights Reserved.
+
+#include "DataVisualizations/PCGSplineDataVisualization.h"
+
+#include "Data/PCGPointData.h"
+#include "Data/PCGSplineData.h"
+#include "Data/PCGSplineStruct.h"
+
+const UPCGPointData* IPCGSplineDataVisualization::CollapseToDebugPointData(FPCGContext* Context, const UPCGData* Data) const
+{
+	if (const UPCGSplineData* SplineData = Cast<UPCGSplineData>(Data))
+	{
+		UPCGPointData* PointData = NewObject<UPCGPointData>();
+		TArray<FPCGPoint>& Points = PointData->GetMutablePoints();
+
+		const int32 NumControlPoints = SplineData->SplineStruct.SplineCurves.Position.Points.Num();
+
+		for (int ControlPointIndex = 0; ControlPointIndex < NumControlPoints; ++ControlPointIndex)
+		{
+			FPCGPoint& Point = Points.Emplace_GetRef();
+			Point.Transform = SplineData->SplineStruct.GetTransformAtSplineInputKey(ControlPointIndex, ESplineCoordinateSpace::World);
+			Point.SetExtents(FVector::One() * 50.0f); // Set the points to fill up 1 meter of space by default.
+		}
+
+		return PointData;
+	}
+
+	return nullptr;
+}
diff --git a/Engine/Plugins/PCG/Source/PCGEditor/Private/DataVisualizations/PCGSplineDataVisualization.h b/Engine/Plugins/PCG/Source/PCGEditor/Private/DataVisualizations/PCGSplineDataVisualization.h
new file mode 100644
--- /dev/null
+++ b/Engine/Plugins/PCG/Source/PCGEditor/Private/DataVisualizations/PCGSplineDataVisualization.h
@@ -0,0 +1,17 @@
+// Copyright Epic Games, Inc. All Rights Reserved.
+
+#pragma once
+
+#include "PCGSpatialDataVisualization.h"
+#include "PCGEditorModule.h"
+
+class UPCGData;
+class UPCGPointData;
+struct FPCGContext;
+
+class PCGEDITOR_API IPCGSplineDataVisualization : public IPCGSpatialDataVisualization
+{
+public:
+	/** Overrides collapse behavior to show the spline control points. */
+	virtual const UPCGPointData* CollapseToDebugPointData(FPCGContext* Context, const UPCGData* Data) const override;
+};
diff --git a/Engine/Plugins/PCG/Source/PCGEditor/Private/PCGEditorModule.cpp b/Engine/Plugins/PCG/Source/PCGEditor/Private/PCGEditorModule.cpp
--- a/Engine/Plugins/PCG/Source/PCGEditor/Private/PCGEditorModule.cpp
+++ b/Engine/Plugins/PCG/Source/PCGEditor/Private/PCGEditorModule.cpp
@@ -8,8 +8,13 @@
 #include "PCGEditorSettings.h"
 #include "PCGEditorStyle.h"
 #include "PCGEditorUtils.h"
+#include "PCGModule.h"
 #include "PCGSubsystem.h"
 #include "PCGVolumeFactory.h"
+#include "Data/PCGSpatialData.h"
+#include "Data/PCGSplineData.h"
+#include "DataVisualizations/PCGSpatialDataVisualization.h"
+#include "DataVisualizations/PCGSplineDataVisualization.h"
 
 #include "ContentBrowserMenuContexts.h"
 #include "ContentBrowserModule.h"
@@ -37,6 +42,7 @@ void FPCGEditorModule::StartupModule()
 {
 	RegisterDetailsCustomizations();
 	RegisterSettings();
+	RegisterPCGDataVisualizations();
 
 	UToolMenus::RegisterStartupCallback(FSimpleMulticastDelegate::FDelegate::CreateRaw(this, &FPCGEditorModule::RegisterMenuExtensions));
 
@@ -51,6 +57,7 @@ void FPCGEditorModule::StartupModule()
 
 void FPCGEditorModule::ShutdownModule()
 {
+	UnregisterPCGDataVisualizations();
 	UnregisterSettings();
 	UnregisterDetailsCustomizations();
 	UnregisterMenuExtensions();
@@ -407,6 +414,18 @@ void FPCGEditorModule::UnregisterSettings()
 	}
 }
 
+void FPCGEditorModule::RegisterPCGDataVisualizations()
+{
+	FPCGDataVisualizationRegistry& DataVisRegistry = FPCGModule::GetMutablePCGDataVisualizationRegistry();
+	DataVisRegistry.InternalRegistry.Add(UPCGSpatialData::StaticClass(), MakeUnique<const IPCGSpatialDataVisualization>());
+	DataVisRegistry.InternalRegistry.Add(UPCGSplineData::StaticClass(), MakeUnique<const IPCGSplineDataVisualization>());
+}
+
+void FPCGEditorModule::UnregisterPCGDataVisualizations()
+{
+	FPCGModule::GetMutablePCGDataVisualizationRegistry().InternalRegistry.Empty();
+}
+
 IMPLEMENT_MODULE(FPCGEditorModule, PCGEditor);
 
 DEFINE_LOG_CATEGORY(LogPCGEditor);
diff --git a/Engine/Plugins/PCG/Source/PCGEditor/Private/PCGEditorModule.h b/Engine/Plugins/PCG/Source/PCGEditor/Private/PCGEditorModule.h
--- a/Engine/Plugins/PCG/Source/PCGEditor/Private/PCGEditorModule.h
+++ b/Engine/Plugins/PCG/Source/PCGEditor/Private/PCGEditorModule.h
@@ -6,13 +6,12 @@
 #include "Modules/ModuleInterface.h"
 #include "Toolkits/AssetEditorToolkit.h"
 
-class IAssetTypeActions;
-
 // Logs
 DECLARE_LOG_CATEGORY_EXTERN(LogPCGEditor, Log, All);
 
 class FMenuBuilder;
 class FPCGEditorGraphNodeFactory;
+class IAssetTypeActions;
 
 class FPCGEditorModule : public IModuleInterface
 {
@@ -31,6 +30,8 @@ protected:
 	void PopulateMenuActions(FMenuBuilder& MenuBuilder);
 	void RegisterSettings();
 	void UnregisterSettings();
+	void RegisterPCGDataVisualizations();
+	void UnregisterPCGDataVisualizations();
 
 	void OnPostEngineInit();
 	void RegisterOnEditorModeChange();
