[PCG] Add Elevation Isolines node

* Use the Marching Squares algorithm, factorized in its own helper.

diff --git a/Engine/Plugins/PCG/Source/PCG/Private/Elements/PCGElevationIsolines.cpp b/Engine/Plugins/PCG/Source/PCG/Private/Elements/PCGElevationIsolines.cpp
new file mode 100644
--- /dev/null
+++ b/Engine/Plugins/PCG/Source/PCG/Private/Elements/PCGElevationIsolines.cpp
@@ -0,0 +1,188 @@
+// Copyright Epic Games, Inc. All Rights Reserved.
+
+#include "Elements/PCGElevationIsolines.h"
+
+#include "PCGComponent.h"
+#include "PCGContext.h"
+#include "Data/PCGIntersectionData.h"
+#include "Data/PCGPointData.h"
+#include "Data/PCGSplineData.h"
+#include "Data/PCGSurfaceData.h"
+#include "Elements/PCGSurfaceSampler.h"
+#include "Helpers/PCGAsync.h"
+#include "SpatialAlgo/PCGMarchingSquares.h"
+
+#define LOCTEXT_NAMESPACE "PCGElevationIsolinesElement"
+
+#if WITH_EDITOR
+FName UPCGElevationIsolinesSettings::GetDefaultNodeName() const
+{
+	return FName(TEXT("ElevationIsolines"));
+}
+
+FText UPCGElevationIsolinesSettings::GetDefaultNodeTitle() const
+{
+	return LOCTEXT("NodeTitle", "Elevation Isolines");
+}
+#endif // WITH_EDITOR
+
+FPCGElementPtr UPCGElevationIsolinesSettings::CreateElement() const
+{
+	return MakeShared<FPCGElevationIsolinesElement>();
+}
+
+TArray<FPCGPinProperties> UPCGElevationIsolinesSettings::InputPinProperties() const
+{
+	TArray<FPCGPinProperties> Properties;
+	FPCGPinProperties& InputPin = Properties.Emplace_GetRef(PCGPinConstants::DefaultInputLabel, EPCGDataType::Surface);
+	InputPin.SetRequiredPin();
+
+	Properties.Emplace_GetRef(TEXT("BoundingShape"), EPCGDataType::Spatial);
+
+	return Properties;
+}
+
+TArray<FPCGPinProperties> UPCGElevationIsolinesSettings::OutputPinProperties() const
+{ 
+	TArray<FPCGPinProperties> Properties;
+	Properties.Emplace(PCGPinConstants::DefaultOutputLabel, bOutputAsSpline ? EPCGDataType::Spline : EPCGDataType::Point);
+	return Properties;
+}
+
+bool FPCGElevationIsolinesElement::ExecuteInternal(FPCGContext* InContext) const
+{
+	TRACE_CPUPROFILER_EVENT_SCOPE(FPCGElevationIsolinesElement::Execute);
+
+	check(InContext);
+
+	const UPCGElevationIsolinesSettings* Settings = InContext->GetInputSettings<UPCGElevationIsolinesSettings>();
+	check(Settings);
+
+	if (FMath::IsNearlyZero(Settings->ElevationIncrement) || Settings->ElevationEnd < Settings->ElevationStart)
+	{
+		return true;
+	}
+
+	bool bOutUnionDataCreated = false;
+	const UPCGSpatialData* BoundingShape = InContext->InputData.GetSpatialUnionOfInputsByPin(TEXT("BoundingShape"), bOutUnionDataCreated);
+
+	// Fallback to getting bounds from actor
+	if (!BoundingShape && InContext->SourceComponent.IsValid())
+	{
+		check(!bOutUnionDataCreated);
+		BoundingShape = Cast<UPCGSpatialData>(InContext->SourceComponent->GetActorPCGData());
+	}
+
+	if (!BoundingShape)
+	{
+		return true;
+	}
+
+	for (const FPCGTaggedData& InputData : InContext->InputData.GetInputsByPin(PCGPinConstants::DefaultInputLabel))
+	{
+		const UPCGSurfaceData* SurfaceData = Cast<UPCGSurfaceData>(InputData.Data);
+		if (!SurfaceData)
+		{
+			continue;
+		}
+
+		const FBox Bounds = BoundingShape->GetBounds();
+		const FVector CellSize{ Settings->Resolution, Settings->Resolution, 1.0 };
+		const FVector HalfCellSize = CellSize / 2.0;
+		const FBox CellBounds{ -HalfCellSize, HalfCellSize };
+		const FVector Origin{ Bounds.Min.X, Bounds.Min.Y, Bounds.Max.Z };
+		const int32 CellMinX = FMath::CeilToInt((Bounds.Min.X) / CellSize.X);
+		const int32 CellMaxX = FMath::FloorToInt((Bounds.Max.X) / CellSize.X);
+		const int32 CellMinY = FMath::CeilToInt((Bounds.Min.Y) / CellSize.Y);
+		const int32 CellMaxY = FMath::FloorToInt((Bounds.Max.Y) / CellSize.Y);
+		const int32 CellXCount = CellMaxX - CellMinX;
+		const int32 CellYCount = CellMaxY - CellMinY;
+		const int32 CellCount = CellXCount * CellYCount;
+
+		TArray<double> Heightmap;
+		Heightmap.SetNumUninitialized(CellCount);
+
+		int32 NumPointsWritten = 0;
+
+		FPCGProjectionParams ProjectionParams{};
+		ProjectionParams.bProjectPositions = true;
+
+		// TODO: Time-slice
+		FPCGAsync::AsyncProcessingOneToOneEx(&InContext->AsyncState, Heightmap.Num(), []() {}, [CellYCount, Settings, Origin, CellSize, SurfaceData, &Heightmap, &ProjectionParams](int32 ReadIndex, int32 WriteIndex)
+		{
+			const int32 CellX = ReadIndex / CellYCount;
+			const int32 CellY = ReadIndex % CellYCount;
+
+			FVector Cell{ Origin.X + CellSize.X * CellX, Origin.Y + CellSize.Y * CellY, Origin.Z };
+
+			FPCGPoint Projected;
+			bool bSuccess = SurfaceData->ProjectPoint(FTransform(Cell), FBox(EForceInit::ForceInit), ProjectionParams, Projected, nullptr);
+			Heightmap[WriteIndex] = bSuccess ? Projected.Transform.GetLocation().Z : std::numeric_limits<double>::min();
+
+		}, /*bEnableTimeSlicing=*/false);
+
+		auto GetElevation = [CellYCount, &Heightmap](int32 X, int32 Y) -> double
+		{
+			return Heightmap[X * CellYCount + Y];
+		};
+
+		// TODO: Async + Support non-Z up cases.
+		// Run the Marching Squares algorithm on each isoline we want.
+		for (double Elevation = Settings->ElevationStart; Elevation < Settings->ElevationEnd; Elevation += Settings->ElevationIncrement)
+		{
+			for (const PCGSpatialAlgo::FPCGMarchingSquareResult& Result : PCGSpatialAlgo::MarchingSquares(CellXCount, CellYCount, Elevation, GetElevation, /*bUseLinearInterpolation=*/true))
+			{
+				TRACE_CPUPROFILER_EVENT_SCOPE(FPCGElevationIsolinesElement::Execute::CreateOutputData);
+
+				auto TransformFromGridCoordinate = [&Origin, &CellSize, Elevation](const FVector2D& GridCoordinate)
+				{
+					return FTransform(FVector(Origin.X + CellSize.X * GridCoordinate.X, Origin.Y + CellSize.Y * GridCoordinate.Y, Elevation));
+				};
+
+				FPCGTaggedData& OutputData = InContext->OutputData.TaggedData.Emplace_GetRef(InputData);
+				OutputData.Pin = PCGPinConstants::DefaultOutputLabel;
+
+				if (Settings->bOutputAsSpline)
+				{
+					TArray<FSplinePoint> SplinePoints;
+					SplinePoints.Reserve(Result.LinkedGridCoordinates.Num());
+
+					for (int32 Index = 0; Index < Result.LinkedGridCoordinates.Num(); ++Index)
+					{
+						const FTransform Transform = TransformFromGridCoordinate(Result.LinkedGridCoordinates[Index]);
+
+						SplinePoints.Emplace(static_cast<float>(Index),
+							Transform.GetLocation(),
+							FVector::ZeroVector,
+							FVector::ZeroVector,
+							Transform.GetRotation().Rotator(),
+							Transform.GetScale3D(),
+							ESplinePointType::Curve);
+					}
+
+					UPCGSplineData* OutSplineData = NewObject<UPCGSplineData>();
+					OutSplineData->Initialize(SplinePoints, Result.bClosed, FTransform::Identity);
+					OutputData.Data = OutSplineData;
+				}
+				else
+				{
+					UPCGPointData* OutPointData = NewObject<UPCGPointData>();
+					TArray<FPCGPoint>& OutPoints = OutPointData->GetMutablePoints();
+					OutPoints.Reserve(Result.LinkedGridCoordinates.Num());
+					OutputData.Data = OutPointData;
+
+					for (const FVector2D& GridCoordinate : Result.LinkedGridCoordinates)
+					{
+						FPCGPoint& Point = OutPoints.Emplace_GetRef();
+						Point.Transform = TransformFromGridCoordinate(GridCoordinate);
+						Point.SetLocalBounds(CellBounds);
+					}
+				}
+			}
+		}
+	}
+
+	return true;
+}
+
+#undef LOCTEXT_NAMESPACE
diff --git a/Engine/Plugins/PCG/Source/PCG/Private/SpatialAlgo/PCGMarchingSquares.cpp b/Engine/Plugins/PCG/Source/PCG/Private/SpatialAlgo/PCGMarchingSquares.cpp
new file mode 100644
--- /dev/null
+++ b/Engine/Plugins/PCG/Source/PCG/Private/SpatialAlgo/PCGMarchingSquares.cpp
@@ -0,0 +1,450 @@
+// Copyright Epic Games, Inc. All Rights Reserved.
+
+#include "SpatialAlgo/PCGMarchingSquares.h"
+
+#include "Containers/Map.h"
+#include "Math/Vector.h"
+#include "Math/Vector2D.h"
+#include "Misc/EnumClassFlags.h"
+#include "Templates/Tuple.h"
+
+namespace PCGSpatialAlgo
+{
+namespace MarchingSquaresAlgo
+{
+	/**
+	* Edge Ids
+	* *---0---*
+	* |       |
+	* 3       1
+	* |       |
+	* *---2---*
+	*/
+	enum class Edge : uint8
+	{
+		Id_0 = 1 << 0,
+		Id_1 = 1 << 1,
+		Id_2 = 1 << 2,
+		Id_3 = 1 << 3,
+		Invalid = 1 << 7,
+
+		AllEdges = Id_0 | Id_1 | Id_2 | Id_3
+	};
+
+	// Define it as flags so that we can use it as a bitfield for masking seen edges.
+	ENUM_CLASS_FLAGS(Edge);
+
+	struct Params
+	{
+		int32 CellCountX;
+		int32 CellCountY;
+
+		void SplitIndex(int32 Index, int32& OutX, int32& OutY) const
+		{
+			OutX = Index / (CellCountY - 1);
+			OutY = Index % (CellCountY - 1);
+		}
+
+		/**
+		* Return the next index of the grid according to the edge Id
+		* If we go out of the grid, return INDEX_NONE.
+		*/
+		int32 GetNextIndex(int32 X, int32 Y, Edge EdgeId) const
+		{
+			switch (EdgeId)
+			{
+			case Edge::Id_0:
+				--X;
+				break;
+			case Edge::Id_1:
+				++Y;
+				break;
+			case Edge::Id_2:
+				++X;
+				break;
+			case Edge::Id_3:
+				--Y;
+				break;
+			default:
+				checkNoEntry();
+				return INDEX_NONE;
+			}
+
+			if (X < 0 || X >= (CellCountX - 1) || Y < 0 || Y >= (CellCountY - 1))
+			{
+				return INDEX_NONE;
+			}
+			else
+			{
+				return X * (CellCountY - 1) + Y;
+			}
+		}
+	};
+
+	/** Returns the matching edge on the neighbor cell */
+	Edge GetNextEdge(Edge EdgeId)
+	{
+		// The edge Id for 0 will be 2 and vice versa and for 3 it will be 1 and vice versa
+		switch (EdgeId)
+		{
+		case Edge::Id_0:
+			return Edge::Id_2;
+		case Edge::Id_1:
+			return Edge::Id_3;
+		case Edge::Id_2:
+			return Edge::Id_0;
+		case Edge::Id_3:
+			return Edge::Id_1;
+		default:
+			checkNoEntry();
+			return Edge::Invalid;
+		}
+	}
+
+	/**
+	* Return the edge ID corresponding to the input edge id.
+	* cf. header for the different types
+	*/
+	Edge GetEdgeId(uint8 Type, Edge InputEdge)
+	{
+		switch (Type)
+		{
+		case 1: // fall-through
+		case 14:
+			return InputEdge == Edge::Id_2 ? Edge::Id_3 : Edge::Id_2;
+		case 2: // fall-through
+		case 13:
+			return InputEdge == Edge::Id_1 ? Edge::Id_2 : Edge::Id_1;
+		case 3: // fall-through
+		case 12:
+			return InputEdge == Edge::Id_1 ? Edge::Id_3 : Edge::Id_1;
+		case 4: // fall-through
+		case 11:
+			return InputEdge == Edge::Id_0 ? Edge::Id_1 : Edge::Id_0;
+		case 5:
+		{
+			switch (InputEdge)
+			{
+			case Edge::Id_1:
+				return Edge::Id_2;
+			case Edge::Id_2:
+				return Edge::Id_1;
+			case Edge::Id_3:
+				return Edge::Id_0;
+			case Edge::Id_0: // fall-through
+			default:
+				return Edge::Id_3;
+			}
+		}
+		case 6: // fall-through
+		case 9:
+			return InputEdge == Edge::Id_0 ? Edge::Id_2 : Edge::Id_0;
+		case 7: // fall-through
+		case 8:
+			return InputEdge == Edge::Id_3 ? Edge::Id_0 : Edge::Id_3;
+		case 10:
+		{
+			switch (InputEdge)
+			{
+			case Edge::Id_1:
+				return Edge::Id_0;
+			case Edge::Id_2:
+				return Edge::Id_3;
+			case Edge::Id_3:
+				return Edge::Id_2;
+			case Edge::Id_0: // fall-through
+			default:
+				return Edge::Id_1;
+			}
+		}
+		default:
+			checkNoEntry();
+			return Edge::Invalid;
+		}
+	}
+
+	/**
+	* Return the starting edges for a given type, with the first one being the tail, and the second the head.
+	* Since all types have their mirror (like type 1 and 14), we switch the order of the edges, depending if Type < 8 or not,
+	* since it switches the gradient direction.
+	*/
+	TTuple<Edge, Edge> GetStartingEdges(uint8 InType, Edge InEdgeMask)
+	{
+		// For special case 5 and 10. If Edge 0 has not yet been seen, start with 0. Otherwise, start with 2 (in the right order).
+		if (InType == 5)
+		{
+			return EnumHasAnyFlags(InEdgeMask, Edge::Id_0) ? TTuple<Edge, Edge>{ Edge::Id_3, Edge::Id_0 } : TTuple<Edge, Edge>{ Edge::Id_1, Edge::Id_2 };
+		}
+		else if (InType == 10)
+		{
+			return EnumHasAnyFlags(InEdgeMask, Edge::Id_0) ? TTuple<Edge, Edge>{ Edge::Id_0, Edge::Id_1 } : TTuple<Edge, Edge>{ Edge::Id_2, Edge::Id_3 };
+		}
+		else
+		{
+			// Bootstrap with Invalid edge
+			const Edge FirstEdge = GetEdgeId(InType, Edge::Invalid);
+			const Edge SecondEdge = GetEdgeId(InType, FirstEdge);
+
+			return InType >= 8 ? TTuple<Edge, Edge>{ FirstEdge, SecondEdge } : TTuple<Edge, Edge>{ SecondEdge, FirstEdge };
+		}
+	}
+
+	/**
+	* For a given edge id, at coordinate X,Y, return it's actual "real" coordinate on the grid.
+	* If there is no linear interpolation, it will be the middle of the edge.
+	* If there is linear interpolation, we interpolate linearly where the isovalue would sit on the edge.
+	*/
+	FVector2D GetEdgeCoordinate(int32 X, int32 Y, Edge EdgeId, const double& Isovalue, TFunctionRef<double(int32,int32)> ValueQuery, bool bUseLinearInterpolation)
+	{
+		FVector2D Result(X, Y);
+
+		double Ratio = 0.5;
+		if (bUseLinearInterpolation)
+		{
+			double Value1{};
+			double Value2{};
+			switch (EdgeId)
+			{
+			case Edge::Id_0:
+				Value1 = ValueQuery(X, Y);
+				Value2 = ValueQuery(X, Y + 1);
+				break;
+			case Edge::Id_1:
+				Value1 = ValueQuery(X, Y + 1);
+				Value2 = ValueQuery(X + 1, Y + 1);
+				break;
+			case Edge::Id_2:
+				Value1 = ValueQuery(X + 1, Y);
+				Value2 = ValueQuery(X + 1, Y + 1);
+				break;
+			case Edge::Id_3:
+				Value1 = ValueQuery(X, Y);
+				Value2 = ValueQuery(X + 1, Y);
+				break;
+			default:
+				checkNoEntry();
+				break;
+			}
+
+			// If Value1 ~= Value2, we expect 0.5 to be a correct approximation, so keep the default value.
+			if (!FMath::IsNearlyEqual(Value1, Value2))
+			{
+				Ratio = (Isovalue - Value1) / (Value2 - Value1);
+			}
+		}
+
+		check(Ratio >= 0.0 && Ratio <= 1.0);
+
+		switch (EdgeId)
+		{
+		case Edge::Id_0:
+			Result.Y += Ratio;
+			break;
+		case Edge::Id_1:
+			Result.Y += 1;
+			Result.X += Ratio;
+			break;
+		case Edge::Id_2:
+			Result.X += 1;
+			Result.Y += Ratio;
+			break;
+		case Edge::Id_3:
+			Result.X += Ratio;
+			break;
+		default:
+			checkNoEntry();
+			break;
+		}
+
+		return Result;
+	};
+
+	struct MarchingArgs
+	{
+		TArray<uint8>& Types;
+		TMap<int32, Edge>& Candidates;
+		Params IndexParams;
+		double Isovalue;
+		TFunctionRef<double(int32, int32)> ValueQuery;
+		bool bUseLinearInterpolation;
+		TArray<FVector2D>* CurrentLine = nullptr;
+	};
+
+	/** Algorithm that will start at a given index and will march following the contour defined by the types until it reaches the beginning (loop) or the edge of the map (not closed loop) */
+	void Marching(MarchingArgs& Args, int32 StartingIndex, Edge StartingEdgeId)
+	{
+		auto MarkEdgeSeenAndRemove = [&Candidates = Args.Candidates](int32 Index, Edge EdgeId)
+		{
+			check(Candidates.Contains(Index));
+
+			EnumRemoveFlags(Candidates[Index], EdgeId);
+
+			// If we have seen all the edges, remove it.
+			if (!Candidates[Index])
+			{
+				Candidates.Remove(Index);
+			}
+		};
+
+		int32 Index = StartingIndex;
+		Edge EdgeId = StartingEdgeId;
+
+		check(Args.CurrentLine);
+
+		while (true)
+		{
+			Edge NextEdgeId = GetEdgeId(Args.Types[Index], EdgeId);
+
+			int32 X, Y;
+			Args.IndexParams.SplitIndex(Index, X, Y);
+
+			Args.CurrentLine->Add(GetEdgeCoordinate(X, Y, NextEdgeId, Args.Isovalue, Args.ValueQuery, Args.bUseLinearInterpolation));
+
+			MarkEdgeSeenAndRemove(Index, NextEdgeId);
+
+			const int32 NextIndex = Args.IndexParams.GetNextIndex(X, Y, NextEdgeId);
+			if (NextIndex == INDEX_NONE || !Args.Candidates.Contains(NextIndex))
+			{
+				// Getting out of the grid or we reached the beginning (so we are in a loop), stop there on that path
+				return;
+			}
+
+			NextEdgeId = GetNextEdge(NextEdgeId);
+
+			// In case 5 and 10, we can still have the cell in the candidates. But if we've already seen that edge, it means we are in a loop, so we need to stop there.
+			if (!EnumHasAnyFlags(Args.Candidates[NextIndex], NextEdgeId))
+			{
+				break;
+			}
+
+			// Also mark that one as seen
+			MarkEdgeSeenAndRemove(NextIndex, NextEdgeId);
+
+			// If the next index was removed, it means we closed the line, so we can stop there
+			if (!Args.Candidates.Contains(NextIndex))
+			{
+				break;
+			}
+
+			// Otherwise continue marching!
+			Index = NextIndex;
+			EdgeId = NextEdgeId;
+		}
+	};
+}
+}
+
+TArray<PCGSpatialAlgo::FPCGMarchingSquareResult> PCGSpatialAlgo::MarchingSquares(const int32 CellCountX, const int32 CellCountY, const double Isovalue, TFunctionRef<double(int32 X, int32 Y)> ValueQuery, bool bUseLinearInterpolation)
+{
+	if (CellCountX < 1 || CellCountY < 1)
+	{
+		return {};
+	}
+
+	using Edge = MarchingSquaresAlgo::Edge;
+
+	MarchingSquaresAlgo::Params Params{ CellCountX, CellCountY };
+
+	// First generate all the types for each cell. We compute the type depending on the value of the 4 corners of our base map
+	// so the "types" grid is smaller by 1 one on each axis (hence the -1 on cell count)
+	TArray<uint8> Types;
+	Types.SetNumUninitialized((Params.CellCountX - 1) * (Params.CellCountY - 1));
+
+	// Will hold all the indexes where we have an edge, and their seen edge ids stored as a bitfield, as an EdgeMask.
+	// We'll use the map to determine the candidates that have not yet been processed, and also
+	// handle the case 5 and 10 where we have two edges in one cell (so we need to know if one edge was processed but not the other)
+	TMap<int32, Edge> Candidates;
+
+	{
+		TRACE_CPUPROFILER_EVENT_SCOPE(FPCGMarchingSquares::GatheringCandidates);
+
+		for (int32 X = 0; X < Params.CellCountX - 1; ++X)
+		{
+			for (int32 Y = 0; Y < Params.CellCountY - 1; ++Y)
+			{
+				const int32 Index = X * (Params.CellCountY - 1) + Y;
+				uint8& CurrentType = Types[Index];
+				CurrentType = 0;
+
+				// We store our type on a uint8, where we only use 4 bits.
+				// For the list of types, cf. the header.
+				if (ValueQuery(X, Y) >= Isovalue) { CurrentType |= (1 << 3); } // TopLeft, bit 3
+				if (ValueQuery(X, Y + 1) >= Isovalue) { CurrentType |= (1 << 2); } // TopRight, bit 2
+				if (ValueQuery(X + 1, Y + 1) >= Isovalue) { CurrentType |= (1 << 1); } // BottomRight, bit 1
+				if (ValueQuery(X + 1, Y) >= Isovalue) { CurrentType |= (1 << 0); } // BottomLeft, bit 0
+
+				// Types 0 and 15 mean that we have no edges.
+				if (CurrentType == 0 || CurrentType == 15)
+				{
+					continue;
+				}
+
+				if (CurrentType == 5 || CurrentType == 10)
+				{
+					// All edges on
+					Candidates.Add(Index, Edge::AllEdges);
+				}
+				else
+				{
+					// Bootstrapping with Invalid edge, this will give us the first edge
+					// We can then re-use that first edge to get the second edge.
+					const Edge FirstEdge = MarchingSquaresAlgo::GetEdgeId(CurrentType, Edge::Invalid);
+					const Edge SecondEdge = MarchingSquaresAlgo::GetEdgeId(CurrentType, FirstEdge);
+					check(FirstEdge != Edge::Invalid && SecondEdge != Edge::Invalid);
+
+					Candidates.Emplace(Index, FirstEdge | SecondEdge);
+				}
+			}
+		}
+	}
+
+	// Then we will pick up an index at random (first value of the candidate set)
+	// and will go in both directions to connect all the edges.
+	// We will fill 2 arrays, one for tail for for head.
+	// We make sure to have tail going clockwise according to the value gradient (and the head counterclockwise)
+	TRACE_CPUPROFILER_EVENT_SCOPE(FPCGMarchingSquares::ReconnectTheDots);
+
+	TArray<FPCGMarchingSquareResult> FinalResult;
+	while (!Candidates.IsEmpty())
+	{
+		int32 StartingIndex = Candidates.begin()->Key;
+		Edge StartingEdgeMask = Candidates.begin()->Value;
+		check(!!StartingEdgeMask);
+
+		TArray<FVector2D> Head;
+		TArray<FVector2D> Tail;
+
+		const TTuple<Edge, Edge> StartingEdges = MarchingSquaresAlgo::GetStartingEdges(Types[StartingIndex], StartingEdgeMask);
+
+		const Edge FirstEdgeId = StartingEdges.Get<0>();
+		const Edge SecondEdgeId = StartingEdges.Get<1>();
+
+		MarchingSquaresAlgo::MarchingArgs Args{ Types, Candidates, Params, Isovalue, ValueQuery, bUseLinearInterpolation };
+
+		// Run it first as the head, on the first edge.
+		Args.CurrentLine = &Head;
+		MarchingSquaresAlgo::Marching(Args, StartingIndex, FirstEdgeId);
+
+		// Then we check if starting index is not in the candidates or if the first edge was seen, if so it means it is closed (head reached the tail).
+		const bool bIsClosed = !Candidates.Contains(StartingIndex) || !EnumHasAnyFlags(Candidates[StartingIndex], FirstEdgeId);
+		if (bIsClosed)
+		{
+			// Move the head, and pop the last point as it was duplicated.
+			FinalResult.Emplace_GetRef(std::move(Head), true).LinkedGridCoordinates.Pop();
+		}
+		else
+		{
+			// Otherwise, we start again with the second edge, going backward on the tail.
+			Args.CurrentLine = &Tail;
+			MarchingSquaresAlgo::Marching(Args, StartingIndex, SecondEdgeId);
+
+			// Since we went backward, we reverse the tail.
+			Algo::Reverse(Tail);
+			// And add the head.
+			Tail.Append(std::move(Head));
+			FinalResult.Emplace(std::move(Tail), false);
+		}
+	}
+
+	return FinalResult;
+}
\ No newline at end of file
diff --git a/Engine/Plugins/PCG/Source/PCG/Public/Elements/PCGElevationIsolines.h b/Engine/Plugins/PCG/Source/PCG/Public/Elements/PCGElevationIsolines.h
new file mode 100644
--- /dev/null
+++ b/Engine/Plugins/PCG/Source/PCG/Public/Elements/PCGElevationIsolines.h
@@ -0,0 +1,59 @@
+// Copyright Epic Games, Inc. All Rights Reserved.
+
+#pragma once
+
+#include "PCGSettings.h"
+
+#include "PCGElevationIsolines.generated.h"
+
+/**
+* Compute the elevation isolines of a surface, can output either points or splines.
+* Currently only work for Z-up surfaces.
+*/
+UCLASS(MinimalAPI, BlueprintType, ClassGroup = (Procedural))
+class UPCGElevationIsolinesSettings : public UPCGSettings
+{
+	GENERATED_BODY()
+
+public:
+	//~Begin UPCGSettings interface
+#if WITH_EDITOR
+	virtual FName GetDefaultNodeName() const override;
+	virtual FText GetDefaultNodeTitle() const override;
+	virtual EPCGSettingsType GetType() const override { return EPCGSettingsType::Spatial; }
+#endif
+
+protected:
+	virtual TArray<FPCGPinProperties> InputPinProperties() const override;
+	virtual TArray<FPCGPinProperties> OutputPinProperties() const override;
+	virtual FPCGElementPtr CreateElement() const override;
+	//~End UPCGSettings interface
+
+public:
+	/** Minimum elevation of the isolines. */
+	UPROPERTY(BlueprintReadWrite, EditAnywhere, Category=Settings, meta = (PCG_Overridable))
+	double ElevationStart = 0.0;
+
+	/** Maximum elevation of the isolines. */
+	UPROPERTY(BlueprintReadWrite, EditAnywhere, Category = Settings, meta = (PCG_Overridable))
+	double ElevationEnd = 1000.0;
+
+	/** Increment elevation between each isolines. */
+	UPROPERTY(BlueprintReadWrite, EditAnywhere, Category = Settings, meta = (PCG_Overridable))
+	double ElevationIncrement = 100.0;
+
+	/** Resolution of the grid for the discretization of the surface. This is the size of one cell, in cm. */
+	UPROPERTY(BlueprintReadWrite, EditAnywhere, Category = Settings, meta = (PCG_Overridable))
+	double Resolution = 100.0;
+
+	/** Will output splines rather than points. */
+	UPROPERTY(BlueprintReadWrite, EditAnywhere, Category = Settings)
+	bool bOutputAsSpline = false;
+};
+
+class FPCGElevationIsolinesElement : public IPCGElement
+{
+protected:
+	virtual bool ExecuteInternal(FPCGContext* InContext) const override;
+};
+
diff --git a/Engine/Plugins/PCG/Source/PCG/Public/SpatialAlgo/PCGMarchingSquares.h b/Engine/Plugins/PCG/Source/PCG/Public/SpatialAlgo/PCGMarchingSquares.h
new file mode 100644
--- /dev/null
+++ b/Engine/Plugins/PCG/Source/PCG/Public/SpatialAlgo/PCGMarchingSquares.h
@@ -0,0 +1,48 @@
+// Copyright Epic Games, Inc. All Rights Reserved.
+
+#pragma once
+
+#include "Containers/Array.h"
+#include "Math/MathFwd.h"
+#include "Templates/Function.h"
+
+namespace PCGSpatialAlgo
+{
+	struct FPCGMarchingSquareResult
+	{
+		TArray<FVector2D> LinkedGridCoordinates;
+		bool bClosed = false;
+	};
+
+	/**
+	* Marching Squares Types (16 of them). 0 in a corner means value is below Isovalue, and 1 means value is above (or equal) Isovalue.
+	*  0---0   0---0   0---0   0---0   0---1   0---1   0---1   0---1
+	*  |   |   |   |   |   |   |   |   |  \|   |/  |   | | |   |/  |
+	*  |   |   |   |   |   |   |---|   |   |   |   |   | | |   |   |
+	*  |   |   |\  |   |  /|   |   |   |   |   |  /|   | | |   |   |
+	*  0---0   1---0   0---1   1---1   0---0   1---0   0---1   1---1
+	* Type 0  Type 1  Type 2  Type 3  Type 4  Type 5  Type 6  Type 7
+	*
+	*  1---0   1---0   1---0   1---0   1---1   1---1   1---1   1---1
+	*  |/  |   | | |   |  \|   |  \|   |   |   |   |   |   |   |   |
+	*  |   |   | | |   |   |   |   |   |---|   |   |   |   |   |   |
+	*  |   |   | | |   |\  |   |   |   |   |   |  /|   |\  |   |   |
+	*  0---0   1---0   0---1   1---1   0---0   1---0   0---1   1---1
+	* Type 8  Type 9  Type 10 Type 11 Type 12 Type 13 Type 14  Type 15
+	*/
+
+	/** Marching Squares Algorithm (cf. https://en.wikipedia.org/wiki/Marching_squares) 
+	* It takes some kind of map of values (like a heightmap, texture, etc...) and will return a list of coordinates on a grid that will form a line where the value in the map on the line is
+	* the isovalue passed as argument.
+	* It can be used for example to find isolines on the landscape, where the line will be where the landscape is a the "Isovalue" height.
+	* Marching Squares is discretizing the map of values into a 2D grid, of size CellCountX * CellCountY.
+	* 
+	* @param CellCountX               Size of the grid on the X axis
+	* @param CellCountY               Size of the grid on the Y axis
+	* @param Isovalue                 Defines the value where the lines computed will have in the map of values.
+	* @param ValueQuery               Sampling function to get the value for the cell (X,Y).
+	* @param bUseLinearInterpolation  If disabled, a point will be at the middle of cell edges. If we use linear interpolation the point will be linearly interpolated to be placed where the isovalue would be
+	* (using the values on each extremity of the cell edge, so requires 2 more sampling).
+	*/
+	PCG_API TArray<FPCGMarchingSquareResult> MarchingSquares(const int32 CellCountX, const int32 CellCountY, const double Isovalue, TFunctionRef<double(int32 X, int32 Y)> ValueQuery, bool bUseLinearInterpolation);
+}
