[PCG] Fix Mac compilation

#rnx

diff --git a/Engine/Plugins/PCG/Source/PCG/Private/SpatialAlgo/PCGMarchingSquares.cpp b/Engine/Plugins/PCG/Source/PCG/Private/SpatialAlgo/PCGMarchingSquares.cpp
--- a/Engine/Plugins/PCG/Source/PCG/Private/SpatialAlgo/PCGMarchingSquares.cpp
+++ b/Engine/Plugins/PCG/Source/PCG/Private/SpatialAlgo/PCGMarchingSquares.cpp
@@ -430,7 +430,7 @@ TArray<PCGSpatialAlgo::FPCGMarchingSquareResult> PCGSpatialAlgo::MarchingSquares
 		if (bIsClosed)
 		{
 			// Move the head, and pop the last point as it was duplicated.
-			FinalResult.Emplace_GetRef(std::move(Head), true).LinkedGridCoordinates.Pop();
+			FinalResult.Add_GetRef({ std::move(Head), true }).LinkedGridCoordinates.Pop();
 		}
 		else
 		{
@@ -442,7 +442,7 @@ TArray<PCGSpatialAlgo::FPCGMarchingSquareResult> PCGSpatialAlgo::MarchingSquares
 			Algo::Reverse(Tail);
 			// And add the head.
 			Tail.Append(std::move(Head));
-			FinalResult.Emplace(std::move(Tail), false);
+			FinalResult.Add({ std::move(Tail), false });
 		}
 	}
 
