[PCG] Remove 5.4 temp CVar and enable construction script fix (it was disabled temporarly but can be reenabled since 32652670)

#rnx

diff --git a/Engine/Plugins/PCG/Source/PCG/Private/PCGComponent.cpp b/Engine/Plugins/PCG/Source/PCG/Private/PCGComponent.cpp
--- a/Engine/Plugins/PCG/Source/PCG/Private/PCGComponent.cpp
+++ b/Engine/Plugins/PCG/Source/PCG/Private/PCGComponent.cpp
@@ -69,11 +69,6 @@ namespace PCGComponent
 		false,
 		TEXT("Disable refresh for all PCG Components."));
 
-	static TAutoConsoleVariable<bool> CVarConstructionScriptFix(
-		TEXT("pcg.ConstructionScriptFix"),
-		false,
-		TEXT("This CVar will be removed in future releases, it allows disabling this fix if regressions are found."));
-
 	template <typename DelegateType>
 	static void BroadcastDynamicDelegate(const DelegateType& Delegate, UPCGComponent* PCGComponent)
 	{
@@ -1442,7 +1437,7 @@ void UPCGComponent::OnComponentDestroyed(bool bDestroyingHierarchy)
 			Subsystem->UnregisterPCGComponent(this);
 		}
 
-		if (IsCreatedByConstructionScript() && PCGComponent::CVarConstructionScriptFix.GetValueOnAnyThread())
+		if (IsCreatedByConstructionScript())
 		{
 			Subsystem->SetConstructionScriptSourceComponent(this);
 		}
