[PCG] Expand comment slightly after we mistakenly changed Bounding Shape pin config

diff --git a/Engine/Plugins/PCG/Source/PCG/Private/Elements/PCGSplineSampler.cpp b/Engine/Plugins/PCG/Source/PCG/Private/Elements/PCGSplineSampler.cpp
--- a/Engine/Plugins/PCG/Source/PCG/Private/Elements/PCGSplineSampler.cpp
+++ b/Engine/Plugins/PCG/Source/PCG/Private/Elements/PCGSplineSampler.cpp
@@ -1456,7 +1456,7 @@ TArray<FPCGPinProperties> UPCGSplineSamplerSettings::InputPinProperties() const
 	FPCGPinProperties& SplinePinProperty = PinProperties.Emplace_GetRef(PCGSplineSamplerConstants::SplineLabel, EPCGDataType::PolyLine, /*bAllowMultipleConnections=*/true, /*bAllowMultipleData=*/true);
 	SplinePinProperty.SetRequiredPin();
 
-	// Only one connection allowed, user can union multiple shapes.
+	// Only one connection/data allowed. To avoid ambiguity, samplers should require users to union or intersect multiple shapes.
 	PinProperties.Emplace(PCGSplineSamplerConstants::BoundingShapeLabel, EPCGDataType::Spatial, /*bInAllowMultipleConnections=*/false, /*bAllowMultipleData=*/false, LOCTEXT("SplineSamplerBoundingShapePinTooltip",
 		"Optional. All sampled points must be contained within this shape."
 	));
diff --git a/Engine/Plugins/PCG/Source/PCG/Private/Elements/PCGSurfaceSampler.cpp b/Engine/Plugins/PCG/Source/PCG/Private/Elements/PCGSurfaceSampler.cpp
--- a/Engine/Plugins/PCG/Source/PCG/Private/Elements/PCGSurfaceSampler.cpp
+++ b/Engine/Plugins/PCG/Source/PCG/Private/Elements/PCGSurfaceSampler.cpp
@@ -288,7 +288,7 @@ TArray<FPCGPinProperties> UPCGSurfaceSamplerSettings::InputPinProperties() const
 	));
 	SurfacePinProperty.SetRequiredPin();
 
-	// Only one connection allowed, user can union multiple shapes
+	// Only one connection/data allowed. To avoid ambiguity, samplers should require users to union or intersect multiple shapes.
 	PinProperties.Emplace(PCGSurfaceSamplerConstants::BoundingShapeLabel, EPCGDataType::Spatial, /*bInAllowMultipleConnections=*/false, /*bAllowMultipleData=*/false, LOCTEXT("SurfaceSamplerBoundingShapePinTooltip",
 		"All sampled points must be contained within this shape. If this input is omitted then bounds will be taken from the actor so that points are contained within actor bounds. "
 		"The Unbounded property disables this and instead generates over the entire bounds of Surface."
diff --git a/Engine/Plugins/PCG/Source/PCG/Private/Elements/PCGVolumeSampler.cpp b/Engine/Plugins/PCG/Source/PCG/Private/Elements/PCGVolumeSampler.cpp
--- a/Engine/Plugins/PCG/Source/PCG/Private/Elements/PCGVolumeSampler.cpp
+++ b/Engine/Plugins/PCG/Source/PCG/Private/Elements/PCGVolumeSampler.cpp
@@ -151,7 +151,7 @@ TArray<FPCGPinProperties> UPCGVolumeSamplerSettings::InputPinProperties() const
 	));
 	VolumePinProperty.SetRequiredPin();
 
-	// Only one connection allowed, user can union multiple shapes.
+	// Only one connection/data allowed. To avoid ambiguity, samplers should require users to union or intersect multiple shapes.
 	PinProperties.Emplace(PCGVolumeSamplerConstants::BoundingShapeLabel, EPCGDataType::Spatial, /*bInAllowMultipleConnections=*/false, /*bAllowMultipleData=*/false, LOCTEXT("VolumeSamplerBoundingShapePinTooltip",
 		"Optional. All sampled points must be contained within this shape."
 	));
