[PCG] Use FPCGTaskId type instead of int32s and other related tidy up

diff --git a/Engine/Plugins/PCG/Source/PCG/Private/Graph/PCGGraphCompiler.cpp b/Engine/Plugins/PCG/Source/PCG/Private/Graph/PCGGraphCompiler.cpp
--- a/Engine/Plugins/PCG/Source/PCG/Private/Graph/PCGGraphCompiler.cpp
+++ b/Engine/Plugins/PCG/Source/PCG/Private/Graph/PCGGraphCompiler.cpp
@@ -489,7 +489,7 @@ EPCGHiGenGrid FPCGGraphCompiler::CalculateGridRecursive(
 	return Grid;
 }
 
-bool FPCGGraphCompiler::CalculateStaticallyActiveRecursive(FPCGTaskId InTaskId, const TArray<FPCGGraphTask>& InCompiledTasks, TMap<int32, bool>& InOutTaskIdToActiveFlag)
+bool FPCGGraphCompiler::CalculateStaticallyActiveRecursive(FPCGTaskId InTaskId, const TArray<FPCGGraphTask>& InCompiledTasks, TMap<FPCGTaskId, bool>& InOutTaskIdToActiveFlag)
 {
 	if (const bool* bEntry = InOutTaskIdToActiveFlag.Find(InTaskId))
 	{
@@ -633,11 +633,11 @@ void FPCGGraphCompiler::CullTasksStaticInactive(TArray<FPCGGraphTask>& InOutComp
 		return;
 	}
 
-	TMap<int32, bool> NodeIdToActiveFlag;
+	TMap<FPCGTaskId, bool> NodeIdToActiveFlag;
 	// First task is input node task which is active
 	NodeIdToActiveFlag.Add(InOutCompiledTasks[0].NodeId, true);
 
-	for (int32 i = 1; i < InOutCompiledTasks.Num(); ++i)
+	for (int i = 1; i < InOutCompiledTasks.Num(); ++i)
 	{
 		// Results of each call memoized via NodeIdToActiveFlag.
 		CalculateStaticallyActiveRecursive(InOutCompiledTasks[i].NodeId, InOutCompiledTasks, NodeIdToActiveFlag);
@@ -661,20 +661,20 @@ void FPCGGraphCompiler::CullTasks(TArray<FPCGGraphTask>& InOutCompiledTasks, boo
 		return;
 	}
 
-	TArray<int32> TaskRemapping;
+	TArray<FPCGTaskId> TaskRemapping;
 	TaskRemapping.SetNumUninitialized(InOutCompiledTasks.Num());
 
 	// Mark culled tasks by remapping to INDEX_NONE. First task is input task and is never culled.
 	TaskRemapping[0] = 0;
-	for (int32 TaskIndex = 1; TaskIndex < InOutCompiledTasks.Num(); ++TaskIndex)
+	for (int TaskIndex = 1; TaskIndex < InOutCompiledTasks.Num(); ++TaskIndex)
 	{
-		TaskRemapping[TaskIndex] = CullTask(InOutCompiledTasks[TaskIndex]) ? INDEX_NONE : 0;
+		TaskRemapping[TaskIndex] = CullTask(InOutCompiledTasks[TaskIndex]) ? InvalidPCGTaskId : 0;
 	}
 
 	// Optionally add wires that bypass culled nodes.
 	if (bAddPassthroughWires)
 	{
-		for (int32 TaskIndex = 1; TaskIndex < InOutCompiledTasks.Num(); ++TaskIndex)
+		for (int TaskIndex = 1; TaskIndex < InOutCompiledTasks.Num(); ++TaskIndex)
 		{
 			FPCGGraphTask& Task = InOutCompiledTasks[TaskIndex];
 			if (!Task.Node || Task.Node->GetInputPins().IsEmpty())
@@ -682,13 +682,13 @@ void FPCGGraphCompiler::CullTasks(TArray<FPCGGraphTask>& InOutCompiledTasks, boo
 				continue;
 			}
 
-			const int32 InputNumBefore = Task.Inputs.Num();
-			for (int32 InputIndex = 0; InputIndex < InputNumBefore; ++InputIndex)
+			const int InputNumBefore = Task.Inputs.Num();
+			for (int InputIndex = 0; InputIndex < InputNumBefore; ++InputIndex)
 			{
-				const int32 InputTaskId = Task.Inputs[InputIndex].TaskId;
+				const FPCGTaskId InputTaskId = Task.Inputs[InputIndex].TaskId;
 
 				// Is node culled?
-				if (TaskRemapping[InputTaskId] == INDEX_NONE)
+				if (TaskRemapping[InputTaskId] == InvalidPCGTaskId)
 				{
 					const FPCGGraphTask& InputTask = InOutCompiledTasks[InputTaskId];
 					if (InputTask.Node && InputTask.Node->GetInputPins().Num() > 1)
@@ -698,7 +698,7 @@ void FPCGGraphCompiler::CullTasks(TArray<FPCGGraphTask>& InOutCompiledTasks, boo
 					}
 
 					// Upstream node was culled. Wire up the inputs of the culled node to this node.
-					for (int32 InputInputIndex = 0; InputInputIndex < InputTask.Inputs.Num(); ++InputInputIndex)
+					for (int InputInputIndex = 0; InputInputIndex < InputTask.Inputs.Num(); ++InputInputIndex)
 					{
 						FPCGGraphTaskInput& NewInput = Task.Inputs.Add_GetRef(InputTask.Inputs[InputInputIndex]);
 						NewInput.OutPin = Task.Inputs[InputIndex].OutPin;
@@ -709,12 +709,12 @@ void FPCGGraphCompiler::CullTasks(TArray<FPCGGraphTask>& InOutCompiledTasks, boo
 	}
 
 	// Remove all culled tasks by compacting the task array. Never cull first (Input) task.
-	int32 WriteIndex = 1;
-	int32 ReadIndex = 1;
+	int WriteIndex = 1;
+	int ReadIndex = 1;
 	while (ReadIndex < InOutCompiledTasks.Num())
 	{
 		// If not culled, then move the task to it's final remapped position in the task array.
-		if (TaskRemapping[ReadIndex] != INDEX_NONE)
+		if (TaskRemapping[ReadIndex] != InvalidPCGTaskId)
 		{
 			if (WriteIndex != ReadIndex)
 			{
@@ -734,11 +734,11 @@ void FPCGGraphCompiler::CullTasks(TArray<FPCGGraphTask>& InOutCompiledTasks, boo
 	for (FPCGGraphTask& Task : InOutCompiledTasks)
 	{
 		// Remap input task IDs, and remove edges that connect to culled nodes.
-		for (int32 InputIndex = Task.Inputs.Num() - 1; InputIndex >= 0; --InputIndex)
+		for (int InputIndex = Task.Inputs.Num() - 1; InputIndex >= 0; --InputIndex)
 		{
-			const int32 InputTaskId = Task.Inputs[InputIndex].TaskId;
-			const int32 Remap = TaskRemapping[InputTaskId];
-			if (Remap != INDEX_NONE)
+			const FPCGTaskId InputTaskId = Task.Inputs[InputIndex].TaskId;
+			const FPCGTaskId Remap = TaskRemapping[InputTaskId];
+			if (Remap != InvalidPCGTaskId)
 			{
 				Task.Inputs[InputIndex].TaskId = Remap;
 			}
@@ -751,9 +751,9 @@ void FPCGGraphCompiler::CullTasks(TArray<FPCGGraphTask>& InOutCompiledTasks, boo
 		// Remap parent ID if there tasks is in a child scope.
 		if (Task.ParentId != InvalidPCGTaskId)
 		{
-			const int32 RemappedParentId = TaskRemapping[Task.ParentId];
+			const FPCGTaskId RemappedParentId = TaskRemapping[Task.ParentId];
 			// Parent task should not have been culled.
-			ensure(RemappedParentId != INDEX_NONE);
+			ensure(RemappedParentId != InvalidPCGTaskId);
 
 			// Write the remapped ID - even if it's invalid/INDEX_NONE. Hanging parent IDs can cause issues elsewhere.
 			Task.ParentId = RemappedParentId;
diff --git a/Engine/Plugins/PCG/Source/PCG/Private/Graph/PCGGraphCompiler.h b/Engine/Plugins/PCG/Source/PCG/Private/Graph/PCGGraphCompiler.h
--- a/Engine/Plugins/PCG/Source/PCG/Private/Graph/PCGGraphCompiler.h
+++ b/Engine/Plugins/PCG/Source/PCG/Private/Graph/PCGGraphCompiler.h
@@ -71,7 +71,7 @@ private:
 		const FPCGStackContext& InStackContext);
 
 	/** Discovers whether task is on a statically active branch (and needs to be passed to graph executor). */
-	static bool CalculateStaticallyActiveRecursive(FPCGTaskId InTaskId, const TArray<FPCGGraphTask>& InCompiledTasks, TMap<int32, bool>& InTaskIdToActiveFlag);
+	static bool CalculateStaticallyActiveRecursive(FPCGTaskId InTaskId, const TArray<FPCGGraphTask>& InCompiledTasks, TMap<FPCGTaskId, bool>& InTaskIdToActiveFlag);
 
 	/** Culls nodes that are inactive for e.g. missing required inputs or on a statically inactive branch. */
 	static void CullTasksStaticInactive(TArray<FPCGGraphTask>& InOutCompiledTasks);
