[PCG] Replace dummy data we've used for controlling details UI with a virtual method

diff --git a/Engine/Plugins/Experimental/PCGWaterInterop/Source/PCGWaterInterop/Private/Elements/PCGGetWaterSpline.cpp b/Engine/Plugins/Experimental/PCGWaterInterop/Source/PCGWaterInterop/Private/Elements/PCGGetWaterSpline.cpp
--- a/Engine/Plugins/Experimental/PCGWaterInterop/Source/PCGWaterInterop/Private/Elements/PCGGetWaterSpline.cpp
+++ b/Engine/Plugins/Experimental/PCGWaterInterop/Source/PCGWaterInterop/Private/Elements/PCGGetWaterSpline.cpp
@@ -12,7 +12,6 @@
 
 UPCGGetWaterSplineSettings::UPCGGetWaterSplineSettings()
 {
-	bDisplayModeSettings = false;
 	Mode = EPCGGetDataFromActorMode::ParseActorComponents;
 }
 
diff --git a/Engine/Plugins/PCG/Source/PCG/Private/Elements/PCGDataFromActor.cpp b/Engine/Plugins/PCG/Source/PCG/Private/Elements/PCGDataFromActor.cpp
--- a/Engine/Plugins/PCG/Source/PCG/Private/Elements/PCGDataFromActor.cpp
+++ b/Engine/Plugins/PCG/Source/PCG/Private/Elements/PCGDataFromActor.cpp
@@ -179,6 +179,15 @@ TSubclassOf<AActor> UPCGDataFromActorSettings::GetDefaultActorSelectorClass() co
 	return TSubclassOf<AActor>();
 }
 
+#if WITH_EDITOR
+bool UPCGDataFromActorSettings::DisplayModeSettings() const
+{
+	PRAGMA_DISABLE_DEPRECATION_WARNINGS
+	return bDisplayModeSettings;
+	PRAGMA_ENABLE_DEPRECATION_WARNINGS
+}
+#endif
+
 void UPCGDataFromActorSettings::PostLoad()
 {
 	Super::PostLoad();
diff --git a/Engine/Plugins/PCG/Source/PCG/Private/Elements/PCGTypedGetter.cpp b/Engine/Plugins/PCG/Source/PCG/Private/Elements/PCGTypedGetter.cpp
--- a/Engine/Plugins/PCG/Source/PCG/Private/Elements/PCGTypedGetter.cpp
+++ b/Engine/Plugins/PCG/Source/PCG/Private/Elements/PCGTypedGetter.cpp
@@ -14,7 +14,6 @@
 
 UPCGGetLandscapeSettings::UPCGGetLandscapeSettings()
 {
-	bDisplayModeSettings = false;
 	Mode = EPCGGetDataFromActorMode::ParseActorComponents;
 	ActorSelector.bShowActorFilter = false;
 	ActorSelector.bIncludeChildren = false;
@@ -137,7 +136,6 @@ void FPCGGetLandscapeDataElement::ProcessActor(FPCGContext* Context, const UPCGD
 
 UPCGGetSplineSettings::UPCGGetSplineSettings()
 {
-	bDisplayModeSettings = false;
 	Mode = EPCGGetDataFromActorMode::ParseActorComponents;
 }
 
@@ -158,7 +156,6 @@ FText UPCGGetSplineSettings::GetNodeTooltipText() const
 
 UPCGGetVolumeSettings::UPCGGetVolumeSettings()
 {
-	bDisplayModeSettings = false;
 	Mode = EPCGGetDataFromActorMode::ParseActorComponents;
 }
 
@@ -181,7 +178,6 @@ FText UPCGGetVolumeSettings::GetNodeTooltipText() const
 
 UPCGGetPrimitiveSettings::UPCGGetPrimitiveSettings()
 {
-	bDisplayModeSettings = false;
 	Mode = EPCGGetDataFromActorMode::ParseActorComponents;
 }
 
@@ -202,7 +198,6 @@ FText UPCGGetPrimitiveSettings::GetNodeTooltipText() const
 
 UPCGGetPCGComponentSettings::UPCGGetPCGComponentSettings()
 {
-	bDisplayModeSettings = false;
 	Mode = EPCGGetDataFromActorMode::GetDataFromPCGComponent;
 
 	ActorSelector.bShowActorFilter = false;
diff --git a/Engine/Plugins/PCG/Source/PCG/Public/Elements/PCGDataFromActor.h b/Engine/Plugins/PCG/Source/PCG/Public/Elements/PCGDataFromActor.h
--- a/Engine/Plugins/PCG/Source/PCG/Public/Elements/PCGDataFromActor.h
+++ b/Engine/Plugins/PCG/Source/PCG/Public/Elements/PCGDataFromActor.h
@@ -66,12 +66,19 @@ public:
 	/** Override this to change the default value the selector will revert to when changing the actor selection type */
 	virtual TSubclassOf<AActor> GetDefaultActorSelectorClass() const;
 
+protected:
+#if WITH_EDITOR
+	UFUNCTION()
+	virtual bool DisplayModeSettings() const;
+#endif
+
+public:
 	/** Describes which actors to select for data collection. */
 	UPROPERTY(BlueprintReadWrite, EditAnywhere, Category = Settings, meta = (ShowOnlyInnerProperties))
 	FPCGActorSelectorSettings ActorSelector;
 
 	/** Describes what kind of data we will collect from the found actor(s). */
-	UPROPERTY(BlueprintReadWrite, EditAnywhere, Category = "Data Retrieval Settings", meta = (EditCondition = bDisplayModeSettings, EditConditionHides, HideEditConditionToggle))
+	UPROPERTY(BlueprintReadWrite, EditAnywhere, Category = "Data Retrieval Settings", meta = (EditCondition = "DisplayModeSettings()", EditConditionHides, HideEditConditionToggle))
 	EPCGGetDataFromActorMode Mode = EPCGGetDataFromActorMode::ParseActorComponents;
 
 	/** Also produces a single point data at the actor location. */
@@ -94,10 +101,6 @@ public:
 	UPROPERTY(BlueprintReadWrite, EditAnywhere, Category = "Data Retrieval Settings", meta = (EditCondition = "Mode == EPCGGetDataFromActorMode::GetSinglePoint", EditConditionHides))
 	bool bMergeSinglePointData = false;
 
-	// This can be set false by inheriting nodes to hide the 'Mode' property.
-	UPROPERTY(Transient, meta = (EditCondition = false, EditConditionHides))
-	bool bDisplayModeSettings = true;
-
 	/** Provide pin names to match against the found component output pins. Data will automatically be wired to the expected pin if the name comparison succeeds. All unmatched pins will go into the standard out pin. */
 	UPROPERTY(BlueprintReadWrite, EditAnywhere, Category = "Data Retrieval Settings", meta = (EditCondition = "Mode == EPCGGetDataFromActorMode::GetDataFromPCGComponent || Mode == EPCGGetDataFromActorMode::GetDataFromPCGComponentOrParseComponents", EditConditionHides))
 	TArray<FName> ExpectedPins;
@@ -115,6 +118,10 @@ public:
 	UPROPERTY(BlueprintReadOnly, EditAnywhere, Category = "Data Retrieval Settings")
 	bool bTrackActorsOnlyWithinBounds = true;
 #endif // WITH_EDITORONLY_DATA
+
+	UE_DEPRECATED(5.5, "No longer in use, override DisplayModeSettings() instead.")
+	UPROPERTY(Transient, meta = (EditCondition = false, EditConditionHides))
+	bool bDisplayModeSettings = true;
 };
 
 struct FPCGDataFromActorContext : public FPCGContext
diff --git a/Engine/Plugins/PCG/Source/PCG/Public/Elements/PCGTypedGetter.h b/Engine/Plugins/PCG/Source/PCG/Public/Elements/PCGTypedGetter.h
--- a/Engine/Plugins/PCG/Source/PCG/Public/Elements/PCGTypedGetter.h
+++ b/Engine/Plugins/PCG/Source/PCG/Public/Elements/PCGTypedGetter.h
@@ -31,16 +31,21 @@ public:
 
 protected:
 	virtual TArray<FPCGPinProperties> OutputPinProperties() const override;
-
 	virtual FPCGElementPtr CreateElement() const override;
 	//~End UPCGSettings
 
-public:
 	//~Begin UPCGDataFromActorSettings interface
+public:
 	virtual EPCGDataType GetDataFilter() const override { return EPCGDataType::Landscape; }
 	virtual TSubclassOf<AActor> GetDefaultActorSelectorClass() const override;
+
+protected:
+#if WITH_EDITOR
+	virtual bool DisplayModeSettings() const override { return false; }
+#endif
 	//~End UPCGDataFromActorSettings
 
+public:
 	UPROPERTY(BlueprintReadWrite, EditAnywhere, Category = Settings, meta = (PCG_Overridable, ShowOnlyInnerProperties))
 	FPCGLandscapeDataProps SamplingProperties;
 
@@ -80,9 +85,14 @@ protected:
 	virtual TArray<FPCGPinProperties> OutputPinProperties() const override;
 	//~End UPCGSettings
 
-public:
 	//~Begin UPCGDataFromActorSettings interface
+public:
 	virtual EPCGDataType GetDataFilter() const override { return EPCGDataType::PolyLine; }
+
+protected:
+#if WITH_EDITOR
+	virtual bool DisplayModeSettings() const override { return false; }
+#endif
 	//~End UPCGDataFromActorSettings
 };
 
@@ -106,9 +116,14 @@ protected:
 	virtual TArray<FPCGPinProperties> OutputPinProperties() const override;
 	//~End UPCGSettings
 
-public:
 	//~Begin UPCGDataFromActorSettings interface
+public:
 	virtual EPCGDataType GetDataFilter() const override { return EPCGDataType::Volume; }
+
+protected:
+#if WITH_EDITOR
+	virtual bool DisplayModeSettings() const override { return false; }
+#endif
 	//~End UPCGDataFromActorSettings
 };
 
@@ -132,9 +147,14 @@ protected:
 	virtual TArray<FPCGPinProperties> OutputPinProperties() const override;
 	//~End UPCGSettings
 
-public:
 	//~Begin UPCGDataFromActorSettings interface
+public:
 	virtual EPCGDataType GetDataFilter() const override { return EPCGDataType::Primitive; }
+
+protected:
+#if WITH_EDITOR
+	virtual bool DisplayModeSettings() const override { return false; }
+#endif
 	//~End UPCGDataFromActorSettings
 };
 
@@ -163,8 +183,13 @@ protected:
 	virtual TArray<FPCGPinProperties> OutputPinProperties() const override;
 	//~End UPCGSettings
 
-public:
 	//~Begin UPCGDataFromActorSettings interface
+public:
 	virtual EPCGDataType GetDataFilter() const override { return EPCGDataType::Any; }
+
+protected:
+#if WITH_EDITOR
+	virtual bool DisplayModeSettings() const override { return false; }
+#endif
 	//~End UPCGDataFromActorSettings
 };
