[PCG] Extend LevelEditor.ComponentContextMenu: Select Original PCG Component

Possible to select the original component of a PCG Partition Actor component via the Component  menu

#rnx

diff --git a/Engine/Plugins/PCG/Source/PCGEditor/PCGEditor.Build.cs b/Engine/Plugins/PCG/Source/PCGEditor/PCGEditor.Build.cs
--- a/Engine/Plugins/PCG/Source/PCGEditor/PCGEditor.Build.cs
+++ b/Engine/Plugins/PCG/Source/PCGEditor/PCGEditor.Build.cs
@@ -68,6 +68,7 @@ namespace UnrealBuildTool.Rules
 					"TypedElementFramework",
 					"TypedElementRuntime",
 					"UnrealEd",
+					"LevelEditor"
 				}
 			);
 
diff --git a/Engine/Plugins/PCG/Source/PCGEditor/Private/PCGEditorModule.cpp b/Engine/Plugins/PCG/Source/PCGEditor/Private/PCGEditorModule.cpp
--- a/Engine/Plugins/PCG/Source/PCGEditor/Private/PCGEditorModule.cpp
+++ b/Engine/Plugins/PCG/Source/PCGEditor/Private/PCGEditorModule.cpp
@@ -2,6 +2,7 @@
 
 #include "PCGEditorModule.h"
 
+#include "PCGComponent.h"
 #include "PCGEditorCommands.h"
 #include "PCGEditorGraphNodeFactory.h"
 #include "PCGEditorMenuUtils.h"
@@ -15,6 +16,7 @@
 #include "Data/PCGSplineData.h"
 #include "DataVisualizations/PCGSpatialDataVisualization.h"
 #include "DataVisualizations/PCGSplineDataVisualization.h"
+#include "Grid/PCGPartitionActor.h"
 
 #include "ContentBrowserMenuContexts.h"
 #include "ContentBrowserModule.h"
@@ -23,7 +25,9 @@
 #include "EditorModes.h"
 #include "IContentBrowserSingleton.h"
 #include "ISettingsModule.h"
+#include "LevelEditorMenuContext.h"
 #include "PropertyEditorModule.h"
+#include "ScopedTransaction.h"
 #include "ToolMenus.h"
 
 #include "Details/EnumSelectorDetails.h"
@@ -216,6 +220,48 @@ void FPCGEditorModule::RegisterMenuExtensions()
 
 		}), FToolMenuInsert(NAME_None, EToolMenuInsertType::Default));
 	}
+
+	if (UToolMenu* ComponentMenu = UToolMenus::Get()->ExtendMenu("LevelEditor.ComponentContextMenu"))
+	{
+		FToolMenuSection& Section = ComponentMenu->AddDynamicSection("PCGComponent", FNewToolMenuDelegate::CreateLambda([this](UToolMenu* ToolMenu)
+		{
+			if (ULevelEditorContextMenuContext* LevelEditorContext = ToolMenu->FindContext<ULevelEditorContextMenuContext>())
+			{
+				if (LevelEditorContext->SelectedComponents.Num() == 1)
+				{
+					if (UPCGComponent* PCGComponent = Cast<UPCGComponent>(LevelEditorContext->SelectedComponents[0]))
+					{
+						FToolMenuSection& Section = ToolMenu->AddSection("PCGComponentSection", LOCTEXT("PCGComponentSection", "PCG Component"));
+
+						FUIAction SelectOriginalComponentAction(
+							FExecuteAction::CreateLambda([PCGComponent]()
+							{
+								if(UPCGComponent* OriginalComponent = PCGComponent->GetOriginalComponent(); OriginalComponent && GEditor)
+								{
+									FScopedTransaction Transaction(LOCTEXT("SelectionOriginalPCGComponentTransaction", "Select Original PCG Component"));
+									GEditor->SelectNone(true, true);
+									GEditor->SelectComponent(OriginalComponent, true, true);
+								}
+							}),
+							FCanExecuteAction::CreateLambda([PCGComponent]()
+							{
+								return GEditor && PCGComponent->GetOriginalComponent() != nullptr && PCGComponent->GetOriginalComponent() != PCGComponent;
+							})
+						);
+
+						Section.AddMenuEntry(
+							"SelectionOriginalPCGComponent",
+							LOCTEXT("SelectionOriginalPCGComponentLabel", "Select Original PCG Component"),
+							LOCTEXT("SelectionOriginalPCGComponentToolTip", "Selects the original PCG Component for the currently selected PCG Component"),
+							FSlateIcon(),
+							SelectOriginalComponentAction
+						);
+					}
+				}
+			}
+			
+		}), FToolMenuInsert(NAME_None, EToolMenuInsertType::Default));
+	}
 }
 
 void FPCGEditorModule::UnregisterMenuExtensions()
