Fixed issue with caching of the volume sampler results, esp. a problem in partitioned graphs

diff --git a/Engine/Plugins/PCG/Source/PCG/Private/Elements/PCGSplineSampler.cpp b/Engine/Plugins/PCG/Source/PCG/Private/Elements/PCGSplineSampler.cpp
--- a/Engine/Plugins/PCG/Source/PCG/Private/Elements/PCGSplineSampler.cpp
+++ b/Engine/Plugins/PCG/Source/PCG/Private/Elements/PCGSplineSampler.cpp
@@ -1478,9 +1478,10 @@ void FPCGSplineSamplerElement::GetDependenciesCrc(const FPCGDataCollection& InIn
 	{
 		bool bUnbounded;
 		PCGSettingsHelpers::GetOverrideValue(InInput, Settings, GET_MEMBER_NAME_CHECKED(FPCGSplineSamplerParams, bUnbounded), Settings->SamplerParams.bUnbounded, bUnbounded);
+		const bool bBoundsConnected = InInput.GetInputsByPin(PCGSplineSamplerConstants::BoundingShapeLabel).Num() > 0;
 
 		// If we're operating in bounded mode then we'll use actor bounds, and therefore take dependency on actor data.
-		if (!bUnbounded && InComponent)
+		if (!bUnbounded && !bBoundsConnected && InComponent)
 		{
 			if (const UPCGData* Data = InComponent->GetActorPCGData())
 			{
diff --git a/Engine/Plugins/PCG/Source/PCG/Private/Elements/PCGVolumeSampler.cpp b/Engine/Plugins/PCG/Source/PCG/Private/Elements/PCGVolumeSampler.cpp
--- a/Engine/Plugins/PCG/Source/PCG/Private/Elements/PCGVolumeSampler.cpp
+++ b/Engine/Plugins/PCG/Source/PCG/Private/Elements/PCGVolumeSampler.cpp
@@ -11,6 +11,7 @@
 #include "Elements/PCGTimeSlicedElementBase.h"
 #include "Helpers/PCGAsync.h"
 #include "Helpers/PCGHelpers.h"
+#include "Helpers/PCGSettingsHelpers.h"
 
 #include "HAL/UnrealMemory.h"
 
@@ -378,6 +379,31 @@ bool FPCGVolumeSamplerElement::ExecuteInternal(FPCGContext* Context) const
 	});
 }
 
+void FPCGVolumeSamplerElement::GetDependenciesCrc(const FPCGDataCollection& InInput, const UPCGSettings* InSettings, UPCGComponent* InComponent, FPCGCrc& OutCrc) const
+{
+	FPCGCrc Crc;
+	IPCGElement::GetDependenciesCrc(InInput, InSettings, InComponent, Crc);
+
+	if (const UPCGVolumeSamplerSettings* Settings = Cast<UPCGVolumeSamplerSettings>(InSettings))
+	{
+		bool bUnbounded;
+		PCGSettingsHelpers::GetOverrideValue(InInput, Settings, GET_MEMBER_NAME_CHECKED(UPCGVolumeSamplerSettings, bUnbounded), Settings->bUnbounded, bUnbounded);
+		const bool bBoundsConnected = InInput.GetInputsByPin(PCGVolumeSamplerConstants::BoundingShapeLabel).Num() > 0;
+
+		// If we're operating in bounded mode and there is no bounding shape connected then we'll use actor bounds, and therefore take
+		// dependency on actor data.
+		if (!bUnbounded && !bBoundsConnected && InComponent)
+		{
+			if (const UPCGData* Data = InComponent->GetActorPCGData())
+			{
+				Crc.Combine(Data->GetOrComputeCrc(/*bFullDataCrc=*/false));
+			}
+		}
+	}
+
+	OutCrc = Crc;
+}
+
 #if WITH_EDITOR
 void UPCGVolumeSamplerSettings::ApplyDeprecationBeforeUpdatePins(UPCGNode* InOutNode, TArray<TObjectPtr<UPCGPin>>& InputPins, TArray<TObjectPtr<UPCGPin>>& OutputPins)
 {
diff --git a/Engine/Plugins/PCG/Source/PCG/Public/Elements/PCGVolumeSampler.h b/Engine/Plugins/PCG/Source/PCG/Public/Elements/PCGVolumeSampler.h
--- a/Engine/Plugins/PCG/Source/PCG/Public/Elements/PCGVolumeSampler.h
+++ b/Engine/Plugins/PCG/Source/PCG/Public/Elements/PCGVolumeSampler.h
@@ -94,6 +94,7 @@ protected:
 class FPCGVolumeSamplerElement : public TPCGTimeSlicedElementBase<PCGVolumeSampler::FVolumeSamplerExecutionState, PCGVolumeSampler::FVolumeSamplerIterationState>
 {
 public:
+	virtual void GetDependenciesCrc(const FPCGDataCollection& InInput, const UPCGSettings* InSettings, UPCGComponent* InComponent, FPCGCrc& OutCrc) const override;
 	// Might be sampling external data like brush, worth computing a full CRC in case we can halt change propagation/re-executions
 	virtual bool ShouldComputeFullOutputDataCrc(FPCGContext* Context) const override { return true; }
 
