Fix Static Analysis

diff --git a/Engine/Plugins/PCG/Source/PCG/Private/PCGWorldActor.cpp b/Engine/Plugins/PCG/Source/PCG/Private/PCGWorldActor.cpp
--- a/Engine/Plugins/PCG/Source/PCG/Private/PCGWorldActor.cpp
+++ b/Engine/Plugins/PCG/Source/PCG/Private/PCGWorldActor.cpp
@@ -55,6 +55,10 @@ void APCGWorldActor::BeginCacheForCookedPlatformData(const ITargetPlatform* Targ
 	check(LandscapeCacheObject);
 
 	UWorld* World = GetWorld();
+	if (!World)
+	{
+		return;
+	}
 
 	// Cache Partition Actor Records at Cook
 	if (UWorldPartition* WorldPartition = World->GetWorldPartition())
@@ -68,7 +72,7 @@ void APCGWorldActor::BeginCacheForCookedPlatformData(const ITargetPlatform* Targ
 		RuntimePartitionActorRecords = TSet<FPCGPartitionActorRecord>(SerializedPartitionActorRecords);
 	}
 
-	if (World && LandscapeCacheObject->SerializationMode == EPCGLandscapeCacheSerializationMode::SerializeOnlyAtCook)
+	if (LandscapeCacheObject->SerializationMode == EPCGLandscapeCacheSerializationMode::SerializeOnlyAtCook)
 	{
 		// Implementation note: actor references gathered from the world partition helpers will register on creation and unregister on deletion
 		// which is why we need to manage this only in the non-WP case.
