[PCG] Re-adjust UE_DEPRECATED versioning back from the future to the present time period

#rnx

diff --git a/Engine/Plugins/PCG/Source/PCG/Public/Elements/Metadata/PCGMetadataPartition.h b/Engine/Plugins/PCG/Source/PCG/Public/Elements/Metadata/PCGMetadataPartition.h
--- a/Engine/Plugins/PCG/Source/PCG/Public/Elements/Metadata/PCGMetadataPartition.h
+++ b/Engine/Plugins/PCG/Source/PCG/Public/Elements/Metadata/PCGMetadataPartition.h
@@ -45,16 +45,16 @@ public:
 	FString PartitionAttributeNames;
 
 	/** Enables deprecated behavior using spaces as separators. Disable to update the node to current behavior. */
-	UE_DEPRECATED(5.6, "bTokenizeOnWhiteSpace has been deprecated.")
+	UE_DEPRECATED(5.5, "bTokenizeOnWhiteSpace has been deprecated.")
 	UPROPERTY(EditAnywhere, Category = Settings, meta = (EditCondition = "bTokenizeOnWhiteSpace", EditConditionHides, DeprecationMessage = "bTokenizeOnWhiteSpace has been deprecated."))
 	bool bTokenizeOnWhiteSpace = false;
 
 #if WITH_EDITORONLY_DATA
-	UE_DEPRECATED(5.6, "PartitionAttribute has been deprecated.")
+	UE_DEPRECATED(5.5, "PartitionAttribute has been deprecated.")
 	UPROPERTY(meta = (DeprecatedProperty, DeprecationMessage = "PartitionAttribute has been deprecated."))
 	FName PartitionAttribute_DEPRECATED = NAME_None;
 
-	UE_DEPRECATED(5.6, "PartitionAttributeSource has been deprecated.")
+	UE_DEPRECATED(5.5, "PartitionAttributeSource has been deprecated.")
 	UPROPERTY(meta = (DeprecatedProperty, DeprecationMessage = "PartitionAttributeSource has been deprecated."))
 	FPCGAttributePropertyInputSelector PartitionAttributeSource_DEPRECATED;
 #endif // WITH_EDITORONLY_DATA
diff --git a/Engine/Plugins/PCG/Source/PCG/Public/Elements/PCGAddTag.h b/Engine/Plugins/PCG/Source/PCG/Public/Elements/PCGAddTag.h
--- a/Engine/Plugins/PCG/Source/PCG/Public/Elements/PCGAddTag.h
+++ b/Engine/Plugins/PCG/Source/PCG/Public/Elements/PCGAddTag.h
@@ -37,7 +37,7 @@ public:
 	FString TagsToAdd;
 
 	/** Enables deprecated behavior using spaces as separators. Disable to update the node to current behavior. */
-	UE_DEPRECATED(5.6, "bTokenizeOnWhiteSpace has been deprecated.")
+	UE_DEPRECATED(5.5, "bTokenizeOnWhiteSpace has been deprecated.")
 	UPROPERTY(EditAnywhere, Category = Settings, meta = (EditCondition = "bTokenizeOnWhiteSpace", EditConditionHides, DeprecationMessage = "bTokenizeOnWhiteSpace has been deprecated."))
 	bool bTokenizeOnWhiteSpace = false;
 };
diff --git a/Engine/Plugins/PCG/Source/PCG/Public/Elements/PCGDeleteAttributesElement.h b/Engine/Plugins/PCG/Source/PCG/Public/Elements/PCGDeleteAttributesElement.h
--- a/Engine/Plugins/PCG/Source/PCG/Public/Elements/PCGDeleteAttributesElement.h
+++ b/Engine/Plugins/PCG/Source/PCG/Public/Elements/PCGDeleteAttributesElement.h
@@ -62,7 +62,7 @@ public:
 	FString SelectedAttributes;
 
 	/** Enables deprecated behavior using spaces as separators. Disable to update the node to current behavior. */
-	UE_DEPRECATED(5.6, "bTokenizeOnWhiteSpace has been deprecated.")
+	UE_DEPRECATED(5.5, "bTokenizeOnWhiteSpace has been deprecated.")
 	UPROPERTY(EditAnywhere, Category = Settings, meta = (EditCondition = "bTokenizeOnWhiteSpace", EditConditionHides, DeprecationMessage = "bTokenizeOnWhiteSpace has been deprecated."))
 	bool bTokenizeOnWhiteSpace = false;
 
diff --git a/Engine/Plugins/PCG/Source/PCG/Public/Elements/PCGDeleteTags.h b/Engine/Plugins/PCG/Source/PCG/Public/Elements/PCGDeleteTags.h
--- a/Engine/Plugins/PCG/Source/PCG/Public/Elements/PCGDeleteTags.h
+++ b/Engine/Plugins/PCG/Source/PCG/Public/Elements/PCGDeleteTags.h
@@ -49,7 +49,7 @@ public:
 	FString SelectedTags;
 
 	/** Enables deprecated behavior using spaces as separators. Disable to update the node to current behavior. */
-	UE_DEPRECATED(5.6, "bTokenizeOnWhiteSpace has been deprecated.")
+	UE_DEPRECATED(5.5, "bTokenizeOnWhiteSpace has been deprecated.")
 	UPROPERTY(EditAnywhere, Category = Settings, meta = (EditCondition = "bTokenizeOnWhiteSpace", EditConditionHides, DeprecationMessage = "bTokenizeOnWhiteSpace has been deprecated."))
 	bool bTokenizeOnWhiteSpace = false;
 };
diff --git a/Engine/Plugins/PCG/Source/PCG/Public/Elements/PCGFilterByTag.h b/Engine/Plugins/PCG/Source/PCG/Public/Elements/PCGFilterByTag.h
--- a/Engine/Plugins/PCG/Source/PCG/Public/Elements/PCGFilterByTag.h
+++ b/Engine/Plugins/PCG/Source/PCG/Public/Elements/PCGFilterByTag.h
@@ -47,7 +47,7 @@ public:
 	FString SelectedTags;
 
 	/** Enables deprecated behavior using spaces as separators. Disable to update the node to current behavior. */
-	UE_DEPRECATED(5.6, "bTokenizeOnWhiteSpace has been deprecated.")
+	UE_DEPRECATED(5.5, "bTokenizeOnWhiteSpace has been deprecated.")
 	UPROPERTY(EditAnywhere, Category = Settings, meta = (EditCondition = "bTokenizeOnWhiteSpace", EditConditionHides, DeprecationMessage = "bTokenizeOnWhiteSpace has been deprecated."))
 	bool bTokenizeOnWhiteSpace = false;
 };
diff --git a/Engine/Plugins/PCG/Source/PCG/Public/Elements/PCGLoopElement.h b/Engine/Plugins/PCG/Source/PCG/Public/Elements/PCGLoopElement.h
--- a/Engine/Plugins/PCG/Source/PCG/Public/Elements/PCGLoopElement.h
+++ b/Engine/Plugins/PCG/Source/PCG/Public/Elements/PCGLoopElement.h
@@ -61,7 +61,7 @@ public:
 	FString FeedbackPins;
 
 	/** Enables deprecated behavior using spaces as separators. Disable to update the node to current behavior. */
-	UE_DEPRECATED(5.6, "bTokenizeOnWhiteSpace has been deprecated.")
+	UE_DEPRECATED(5.5, "bTokenizeOnWhiteSpace has been deprecated.")
 	UPROPERTY(EditAnywhere, Category = Settings, meta = (EditCondition = "bTokenizeOnWhiteSpace", EditConditionHides, DeprecationMessage = "bTokenizeOnWhiteSpace has been deprecated."))
 	bool bTokenizeOnWhiteSpace = false;
 };
diff --git a/Engine/Plugins/PCG/Source/PCG/Public/Elements/PCGPointProcessingElementBase.h b/Engine/Plugins/PCG/Source/PCG/Public/Elements/PCGPointProcessingElementBase.h
--- a/Engine/Plugins/PCG/Source/PCG/Public/Elements/PCGPointProcessingElementBase.h
+++ b/Engine/Plugins/PCG/Source/PCG/Public/Elements/PCGPointProcessingElementBase.h
@@ -7,7 +7,7 @@
 struct FPCGPoint;
 struct FPCGTaggedData;
 
-class UE_DEPRECATED(5.6, "FPCGPointProcessingElementBase is deprecated. Please use 'FPCGPointOperationElementBase' instead.") PCG_API FPCGPointProcessingElementBase : public IPCGElement
+class UE_DEPRECATED(5.5, "FPCGPointProcessingElementBase is deprecated. Please use 'FPCGPointOperationElementBase' instead.") PCG_API FPCGPointProcessingElementBase : public IPCGElement
 {
 protected:
 	void ProcessPoints(FPCGContext* Context, const TArray<FPCGTaggedData>& Inputs, TArray<FPCGTaggedData>& Outputs, const TFunction<bool(const FPCGPoint&, FPCGPoint&)>& PointFunc) const;
diff --git a/Engine/Plugins/PCG/Source/PCG/Public/Elements/PCGReplaceTags.h b/Engine/Plugins/PCG/Source/PCG/Public/Elements/PCGReplaceTags.h
--- a/Engine/Plugins/PCG/Source/PCG/Public/Elements/PCGReplaceTags.h
+++ b/Engine/Plugins/PCG/Source/PCG/Public/Elements/PCGReplaceTags.h
@@ -45,7 +45,7 @@ public:
 	FString ReplacedTags;
 
 	/** Enables deprecated behavior using spaces as separators. Disable to update the node to current behavior. */
-	UE_DEPRECATED(5.6, "bTokenizeOnWhiteSpace has been deprecated.")
+	UE_DEPRECATED(5.5, "bTokenizeOnWhiteSpace has been deprecated.")
 	UPROPERTY(EditAnywhere, Category = Settings, meta = (EditCondition = "bTokenizeOnWhiteSpace", EditConditionHides, DeprecationMessage = "bTokenizeOnWhiteSpace has been deprecated."))
 	bool bTokenizeOnWhiteSpace = false;
 };
diff --git a/Engine/Plugins/PCG/Source/PCG/Public/Helpers/PCGHelpers.h b/Engine/Plugins/PCG/Source/PCG/Public/Helpers/PCGHelpers.h
--- a/Engine/Plugins/PCG/Source/PCG/Public/Helpers/PCGHelpers.h
+++ b/Engine/Plugins/PCG/Source/PCG/Public/Helpers/PCGHelpers.h
@@ -54,7 +54,7 @@ namespace PCGHelpers
 
 	PCG_API APCGWorldActor* GetPCGWorldActor(UWorld* InWorld);
 
-	UE_DEPRECATED(5.6, "This function was incorrectly parsing by whitespace. Please use 'GetStringArrayFromCommaSeparatedList' instead.")
+	UE_DEPRECATED(5.5, "This function was incorrectly parsing by whitespace. Please use 'GetStringArrayFromCommaSeparatedList' instead.")
 	PCG_API TArray<FString> GetStringArrayFromCommaSeparatedString(const FString& InCommaSeparatedString, const FPCGContext* InOptionalContext = nullptr);
 	PCG_API TArray<FString> GetStringArrayFromCommaSeparatedList(const FString& InCommaSeparatedString);
 
