Fixed timings for nodes that go to sleep/can be executed in parallel

diff --git a/Engine/Plugins/PCG/Source/PCG/Private/Graph/PCGGraphExecutor.cpp b/Engine/Plugins/PCG/Source/PCG/Private/Graph/PCGGraphExecutor.cpp
--- a/Engine/Plugins/PCG/Source/PCG/Private/Graph/PCGGraphExecutor.cpp
+++ b/Engine/Plugins/PCG/Source/PCG/Private/Graph/PCGGraphExecutor.cpp
@@ -824,7 +824,7 @@ void FPCGGraphExecutor::Execute()
 						if (Task.StackIndex != INDEX_NONE)
 						{
 							const FPCGStack* Stack = Task.GetStack();
-							SourceComponent->StoreInspectionData(Stack, Task.Node, TaskInput, CachedOutput, /*bUsedCache=*/true);
+							SourceComponent->StoreInspectionData(Stack, Task.Node, nullptr, TaskInput, CachedOutput, /*bUsedCache=*/true);
 						}
 					}
 #endif
@@ -1003,7 +1003,7 @@ void FPCGGraphExecutor::Execute()
 					if (ActiveTask.StackIndex != INDEX_NONE)
 					{
 						const FPCGStack* Stack = ActiveTask.StackContext->GetStack(ActiveTask.StackIndex);
-						SourceComponent->StoreInspectionData(Stack, ActiveTask.Context->Node, ActiveTask.Context->InputData, ActiveTask.Context->OutputData, /*bUsedCache=*/false);
+						SourceComponent->StoreInspectionData(Stack, ActiveTask.Context->Node, &ActiveTask.Context->Timer, ActiveTask.Context->InputData, ActiveTask.Context->OutputData, /*bUsedCache=*/false);
 					}
 				}
 			}
diff --git a/Engine/Plugins/PCG/Source/PCG/Private/PCGComponent.cpp b/Engine/Plugins/PCG/Source/PCG/Private/PCGComponent.cpp
--- a/Engine/Plugins/PCG/Source/PCG/Private/PCGComponent.cpp
+++ b/Engine/Plugins/PCG/Source/PCG/Private/PCGComponent.cpp
@@ -2156,26 +2156,25 @@ void UPCGComponent::DisableInspection()
 	}
 };
 
-void UPCGComponent::NotifyNodeExecuted(const UPCGNode* InNode, const FPCGStack* InStack, bool bNodeUsedCache)
+void UPCGComponent::NotifyNodeExecuted(const UPCGNode* InNode, const FPCGStack* InStack, const PCGUtils::FCallTime* InTimer, bool bNodeUsedCache)
 {
 	if (!ensure(InStack && InNode))
 	{
 		return;
 	}
 
-	FPCGStack Stack = *InStack;
-
 	// Reset timer information if taken from cache to provide good info in the profiling window
-	if (bNodeUsedCache)
+	PCGUtils::FCallTime Timer;
+	if (InTimer && !bNodeUsedCache)
 	{
-		Stack.Timer = PCGUtils::FCallTime();
+		Timer = *InTimer;
 	}
 
 	FWriteScopeLock Lock(NodeToStacksInWhichNodeExecutedLock);
-	NodeToStacksInWhichNodeExecuted.FindOrAdd(InNode).Add(MoveTemp(Stack));
+	NodeToStacksInWhichNodeExecuted.FindOrAdd(InNode).Add(UPCGComponent::NodeExecutedNotificationData(*InStack, MoveTemp(Timer)));
 }
 
-TMap<TObjectKey<const UPCGNode>, TSet<FPCGStack>> UPCGComponent::GetExecutedNodeStacks() const
+TMap<TObjectKey<const UPCGNode>, TSet<UPCGComponent::NodeExecutedNotificationData>> UPCGComponent::GetExecutedNodeStacks() const
 {
 	FReadScopeLock Lock(NodeToStacksInWhichNodeExecutedLock);
 	return NodeToStacksInWhichNodeExecuted;
@@ -2211,12 +2210,14 @@ void UPCGComponent::NotifyNodeDynamicInactivePins(const UPCGNode* InNode, const
 bool UPCGComponent::WasNodeExecuted(const UPCGNode* InNode, const FPCGStack& Stack) const
 {
 	FReadScopeLock Lock(NodeToStacksInWhichNodeExecutedLock);
-	const TSet<FPCGStack>* FoundStacks = NodeToStacksInWhichNodeExecuted.Find(InNode);
+	const TSet<UPCGComponent::NodeExecutedNotificationData>* FoundNotifications = NodeToStacksInWhichNodeExecuted.Find(InNode);
 
-	return FoundStacks && FoundStacks->Contains(Stack);
+	// Since the operator== & hash functions don't rely on the timer, we can just build a stub from the stack.
+	UPCGComponent::NodeExecutedNotificationData NotificationStub(Stack, PCGUtils::FCallTime());
+	return FoundNotifications && FoundNotifications->Contains(NotificationStub);
 }
 
-void UPCGComponent::StoreInspectionData(const FPCGStack* InStack, const UPCGNode* InNode, const FPCGDataCollection& InInputData, const FPCGDataCollection& InOutputData, bool bUsedCache)
+void UPCGComponent::StoreInspectionData(const FPCGStack* InStack, const UPCGNode* InNode, const PCGUtils::FCallTime* InTimer, const FPCGDataCollection& InInputData, const FPCGDataCollection& InOutputData, bool bUsedCache)
 {
 	if (!InNode || !ensure(InStack))
 	{
@@ -2224,7 +2225,7 @@ void UPCGComponent::StoreInspectionData(const FPCGStack* InStack, const UPCGNode
 	}
 
 	// Notify component that this task executed. Useful for editor visualization.
-	NotifyNodeExecuted(InNode, InStack, bUsedCache);
+	NotifyNodeExecuted(InNode, InStack, InTimer, bUsedCache);
 
 	if (!InOutputData.TaggedData.IsEmpty())
 	{
diff --git a/Engine/Plugins/PCG/Source/PCG/Private/Utils/PCGExtraCapture.cpp b/Engine/Plugins/PCG/Source/PCG/Private/Utils/PCGExtraCapture.cpp
--- a/Engine/Plugins/PCG/Source/PCG/Private/Utils/PCGExtraCapture.cpp
+++ b/Engine/Plugins/PCG/Source/PCG/Private/Utils/PCGExtraCapture.cpp
@@ -26,7 +26,7 @@ void PCGUtils::FExtraCapture::Update(const PCGUtils::FScopedCall& InScopedCall)
 	const double CurrentTime = FPlatformTime::Seconds();
 	const double ThisFrameTime = CurrentTime - InScopedCall.StartTime;
 
-	FCallTime& Timer = const_cast<FPCGStack*>(InScopedCall.Context->Stack)->Timer;
+	FCallTime& Timer = InScopedCall.Context->Timer;
 
 	switch (InScopedCall.Phase)
 	{
@@ -104,7 +104,7 @@ PCGUtils::FCallTreeInfo PCGUtils::FExtraCapture::CalculateCallTreeInfo(const UPC
 	FCallTreeInfo RootInfo;
 
 	// Basically, what we want is - visit all entries in the "NodeToStacksInWhichNodeExecuted" and build our information from there.
-	TMap<TObjectKey<const UPCGNode>, TSet<FPCGStack>> NodeToStacksInWhichNodeExecuted = Component->GetExecutedNodeStacks();
+	TMap<TObjectKey<const UPCGNode>, TSet<UPCGComponent::NodeExecutedNotificationData>> NodeToStacksInWhichNodeExecuted = Component->GetExecutedNodeStacks();
 
 	TArray<const UPCGNode*> NodePath;
 	TArray<int32> NodePathLoop;
@@ -178,10 +178,12 @@ PCGUtils::FCallTreeInfo PCGUtils::FExtraCapture::CalculateCallTreeInfo(const UPC
 	for (const auto& NodeToStacks : NodeToStacksInWhichNodeExecuted)
 	{
 		const UPCGNode* Node = NodeToStacks.Key.ResolveObjectPtr();
-		const TSet<FPCGStack>& Stacks = NodeToStacks.Value;
+		const TSet<UPCGComponent::NodeExecutedNotificationData>& NodeExecutions = NodeToStacks.Value;
 
-		for (const FPCGStack& Stack : Stacks)
+		for(const UPCGComponent::NodeExecutedNotificationData& NodeExecution : NodeExecutions)
 		{
+			const FPCGStack& Stack = NodeExecution.Stack;
+
 			if (!GetNodePath(Stack))
 			{
 				continue;
@@ -205,7 +207,7 @@ PCGUtils::FCallTreeInfo PCGUtils::FExtraCapture::CalculateCallTreeInfo(const UPC
 			FCallTreeInfo* Info = GetCallInfo();
 
 			check(Info);
-			Info->CallTime = Stack.Timer;
+			Info->CallTime = NodeExecution.Timer;
 		}
 	}
 
diff --git a/Engine/Plugins/PCG/Source/PCG/Public/Graph/PCGStackContext.h b/Engine/Plugins/PCG/Source/PCG/Public/Graph/PCGStackContext.h
--- a/Engine/Plugins/PCG/Source/PCG/Public/Graph/PCGStackContext.h
+++ b/Engine/Plugins/PCG/Source/PCG/Public/Graph/PCGStackContext.h
@@ -60,6 +60,17 @@ struct PCG_API FPCGStack
 	friend class FPCGStackContext;
 
 public:
+#if WITH_EDITOR
+	// Disable deprecation warnings on the rule of 5 because of the Timer member. To be removed when that that member is removed.
+	PRAGMA_DISABLE_DEPRECATION_WARNINGS
+	FPCGStack() = default;
+	FPCGStack(const FPCGStack&) = default;
+	FPCGStack(FPCGStack&&) = default;
+	FPCGStack& operator=(const FPCGStack&) = default;
+	FPCGStack& operator=(FPCGStack&&) = default;
+	PRAGMA_ENABLE_DEPRECATION_WARNINGS
+#endif // WITH_EDITOR
+
 	/** Push frame onto top of stack. */
 	void PushFrame(const FPCGStackFrame& Frame) { StackFrames.Add(Frame); }
 	void PushFrame(const UObject* InFrameObject) { StackFrames.Emplace(InFrameObject); }
@@ -115,6 +126,7 @@ public:
 
 #if WITH_EDITOR
 	// Used to store node & hierarchy information
+	UE_DEPRECATED(5.5, "The timer has been moved to the FPCGContext struct.")
 	PCGUtils::FCallTime Timer;
 #endif
 
diff --git a/Engine/Plugins/PCG/Source/PCG/Public/PCGComponent.h b/Engine/Plugins/PCG/Source/PCG/Public/PCGComponent.h
--- a/Engine/Plugins/PCG/Source/PCG/Public/PCGComponent.h
+++ b/Engine/Plugins/PCG/Source/PCG/Public/PCGComponent.h
@@ -341,7 +341,7 @@ public:
 	bool IsInspecting() const;
 	void EnableInspection();
 	void DisableInspection();
-	void StoreInspectionData(const FPCGStack* InStack, const UPCGNode* InNode, const FPCGDataCollection& InInputData, const FPCGDataCollection& InOutputData, bool bUsedCache);
+	void StoreInspectionData(const FPCGStack* InStack, const UPCGNode* InNode, const PCGUtils::FCallTime* InTimer, const FPCGDataCollection& InInputData, const FPCGDataCollection& InOutputData, bool bUsedCache);
 	const FPCGDataCollection* GetInspectionData(const FPCGStack& InStack) const;
 	void ClearInspectionData(bool bClearPerNodeExecutionData = true);
 
@@ -349,10 +349,21 @@ public:
 	bool WasNodeExecuted(const UPCGNode* InNode, const FPCGStack& Stack) const;
 
 	/** Called at execution time each time a node has been executed. */
-	void NotifyNodeExecuted(const UPCGNode* InNode, const FPCGStack* InStack, bool bNodeUsedCache);
+	void NotifyNodeExecuted(const UPCGNode* InNode, const FPCGStack* InStack, const PCGUtils::FCallTime* InTimer, bool bNodeUsedCache);
 
 	/* Retrieves the executed nodes information */
-	TMap<TObjectKey<const UPCGNode>, TSet<FPCGStack>> GetExecutedNodeStacks() const;
+	struct NodeExecutedNotificationData
+	{
+		NodeExecutedNotificationData(const FPCGStack& InStack, const PCGUtils::FCallTime& InTimer) : Stack(InStack), Timer(InTimer) {}
+		// Important implementation note: some logic in WasNodeExecuted relies on the fact we don't use the timer for the operator== and hash functions.
+		friend PCG_API uint32 GetTypeHash(const NodeExecutedNotificationData& NotifData) { return GetTypeHash(NotifData.Stack); }
+		bool operator==(const NodeExecutedNotificationData& OtherNotifData) const { return Stack == OtherNotifData.Stack; }
+
+		FPCGStack Stack;
+		PCGUtils::FCallTime Timer;
+	};
+
+	TMap<TObjectKey<const UPCGNode>, TSet<NodeExecutedNotificationData>> GetExecutedNodeStacks() const;
 
 	/** Retrieve the inactive pin bitmask for the given node and stack in the last execution. */
 	uint64 GetNodeInactivePinMask(const UPCGNode* InNode, const FPCGStack& Stack) const;
@@ -626,7 +637,7 @@ private:
 	mutable FRWLock NodeToStacksThatProducedDataLock;
 
 	/** Map from nodes to all stacks for which a task for the node was executed. */
-	TMap<TObjectKey<const UPCGNode>, TSet<FPCGStack>> NodeToStacksInWhichNodeExecuted;
+	TMap<TObjectKey<const UPCGNode>, TSet<NodeExecutedNotificationData>> NodeToStacksInWhichNodeExecuted;
 	mutable FRWLock NodeToStacksInWhichNodeExecutedLock;
 
 	/** Map from nodes to stacks to mask of output pins that were deactivated during execution. */
diff --git a/Engine/Plugins/PCG/Source/PCG/Public/PCGContext.h b/Engine/Plugins/PCG/Source/PCG/Public/PCGContext.h
--- a/Engine/Plugins/PCG/Source/PCG/Public/PCGContext.h
+++ b/Engine/Plugins/PCG/Source/PCG/Public/PCGContext.h
@@ -5,6 +5,7 @@
 #include "PCGData.h"
 #include "PCGNode.h" // IWYU pragma: keep
 #include "Helpers/PCGAsyncState.h"
+#include "Utils/PCGExtraCapture.h"
 
 #include "UObject/GCObject.h"
 
@@ -49,6 +50,8 @@ struct PCG_API FPCGContext
 {
 	GENERATED_BODY()
 
+	friend class FPCGGraphExecutor;
+
 	virtual ~FPCGContext() = default;
 
 	FPCGDataCollection InputData;
@@ -134,4 +137,9 @@ private:
 
 	// List of params that were in effect overriden
 	TArray<const FPCGSettingsOverridableParam*> OverriddenParams;
+
+#if WITH_EDITOR
+	friend class PCGUtils::FExtraCapture;
+	PCGUtils::FCallTime Timer;
+#endif
 };
\ No newline at end of file
diff --git a/Engine/Plugins/PCG/Source/PCG/Public/Utils/PCGExtraCapture.h b/Engine/Plugins/PCG/Source/PCG/Public/Utils/PCGExtraCapture.h
--- a/Engine/Plugins/PCG/Source/PCG/Public/Utils/PCGExtraCapture.h
+++ b/Engine/Plugins/PCG/Source/PCG/Public/Utils/PCGExtraCapture.h
@@ -18,22 +18,19 @@ namespace PCGUtils
 #if WITH_EDITOR
 	struct PCG_API FCallTime
 	{
-		// sum of all frames
+		double PrepareDataStartTime = MAX_dbl;
+		double PrepareDataTime = 0.0;
+		double PrepareDataEndTime = 0.0;
 		double ExecutionStartTime = MAX_dbl;
 		double ExecutionTime = 0.0;
 		double ExecutionEndTime = 0.0;
-		// how many frames
-		int32 ExecutionFrameCount = 0;
 		double MinExecutionFrameTime = MAX_dbl;
 		double MaxExecutionFrameTime = 0.0;
+		double PostExecuteTime = 0.0;
 
-		double PrepareDataStartTime = MAX_dbl;
-		double PrepareDataTime = 0.0;
-		double PrepareDataEndTime = 0.0;
+		int32 ExecutionFrameCount = 0;
 		int32 PrepareDataFrameCount = 0;
 
-		double PostExecuteTime = 0.0;
-
 		double PrepareDataWallTime() const { return PrepareDataEndTime - PrepareDataStartTime; }
 		double ExecutionWallTime() const { return ExecutionEndTime - ExecutionStartTime; }
 		double TotalTime() const { return ExecutionTime + PrepareDataTime; }
