Optimized static mesh spawner by attribute by collapsing mesh override with the rest

diff --git a/Engine/Plugins/PCG/Source/PCG/Private/Elements/PCGStaticMeshSpawnerContext.cpp b/Engine/Plugins/PCG/Source/PCG/Private/Elements/PCGStaticMeshSpawnerContext.cpp
--- a/Engine/Plugins/PCG/Source/PCG/Private/Elements/PCGStaticMeshSpawnerContext.cpp
+++ b/Engine/Plugins/PCG/Source/PCG/Private/Elements/PCGStaticMeshSpawnerContext.cpp
@@ -11,7 +11,6 @@ void FPCGStaticMeshSpawnerContext::ResetInputIterationData()
 	CurrentOutputPointData = nullptr;
 	MaterialOverrideHelper.Reset();
 	CurrentPointIndex = 0;
-	ValueKeyToMesh.Reset();
 	WeightedMeshInstances.Reset();
 	MeshToValueKey.Reset();
 	CumulativeWeights.Reset();
diff --git a/Engine/Plugins/PCG/Source/PCG/Private/MeshSelectors/PCGMeshSelectorByAttribute.cpp b/Engine/Plugins/PCG/Source/PCG/Private/MeshSelectors/PCGMeshSelectorByAttribute.cpp
--- a/Engine/Plugins/PCG/Source/PCG/Private/MeshSelectors/PCGMeshSelectorByAttribute.cpp
+++ b/Engine/Plugins/PCG/Source/PCG/Private/MeshSelectors/PCGMeshSelectorByAttribute.cpp
@@ -12,6 +12,7 @@
 #include "Metadata/PCGMetadataPartitionCommon.h"
 #include "Metadata/Accessors/PCGAttributeAccessorHelpers.h"
 
+#include "Algo/AnyOf.h"
 #include "Engine/StaticMesh.h"
 
 #include UE_INLINE_GENERATED_CPP_BY_NAME(PCGMeshSelectorByAttribute)
@@ -135,30 +136,8 @@ bool UPCGMeshSelectorByAttribute::SelectInstances(
 	const FPCGMetadataAttributeBase* AttributeBase = InPointData->Metadata->GetConstAttribute(AttributeName);
 	check(AttributeBase);
 
-	// Set up a getter lambda to retrieve the mesh asset in the selection loop below. Returns true if path was valid.
-	TFunction<void(PCGMetadataValueKey, FSoftObjectPath&)> MeshPathGetter;
-	if (PCG::Private::IsOfTypes<FSoftObjectPath>(AttributeBase->GetTypeId()))
-	{
-		MeshPathGetter = [AttributeBase](PCGMetadataValueKey InValueKey, FSoftObjectPath& OutMeshPath)
-		{
-			const FPCGMetadataAttribute<FSoftObjectPath>* Attribute = static_cast<const FPCGMetadataAttribute<FSoftObjectPath>*>(AttributeBase);
-			OutMeshPath = Attribute->GetValue(InValueKey);
-		};
-	}
-	else if (PCG::Private::IsOfTypes<FString>(AttributeBase->GetTypeId()))
-	{
-		MeshPathGetter = [AttributeBase](PCGMetadataValueKey InValueKey, FSoftObjectPath& OutMeshPath)
-		{
-			const FPCGMetadataAttribute<FString>* Attribute = static_cast<const FPCGMetadataAttribute<FString>*>(AttributeBase);
-
-			const FString Path = Attribute->GetValue(InValueKey);
-			if (!Path.IsEmpty() && Path != TEXT("None"))
-			{
-				OutMeshPath = FSoftObjectPath(Path);
-			}
-		};
-	}
-	else
+	// Validate that the "mesh" attribute is of the right type
+	if (!PCG::Private::IsOfTypes<FSoftObjectPath, FString>(AttributeBase->GetTypeId()))
 	{
 		PCGE_LOG_C(Error, GraphAndLog, &Context, FText::Format(LOCTEXT("AttributeInvalidType", "Attribute '{0}' is not of valid type (must be FString or FSoftObjectPath)"), FText::FromName(AttributeName)));
 		return true;
@@ -186,23 +165,25 @@ bool UPCGMeshSelectorByAttribute::SelectInstances(
 
 	TRACE_CPUPROFILER_EVENT_SCOPE(UPCGMeshSelectorByAttribute::SelectEntries);
 
-	// Assign points to entries
-	int32 CurrentPointIndex = Context.CurrentPointIndex;
-	int32 LastCheckpointIndex = CurrentPointIndex;
-	constexpr int32 TimeSlicingCheckFrequency = 1024;
-	TMap<PCGMetadataValueKey, TSoftObjectPtr<UStaticMesh>>& ValueKeyToMesh = Context.ValueKeyToMesh;
-
-	// No partition needed if there are no overrides
-	if (Settings->StaticMeshComponentPropertyOverrides.IsEmpty())
-	{
-		Context.bPartitionDone = true;
-	}
-
 	if (!Context.bPartitionDone)
 	{
+		TRACE_CPUPROFILER_EVENT_SCOPE(UPCGMeshSelectorByAttribute::SelectEntries::BuildingPartition);
+
+		TArray<FPCGObjectPropertyOverrideDescription> PropertyOverrides = Settings->StaticMeshComponentPropertyOverrides;
+		const FString StaticMeshPropertyString = GET_MEMBER_NAME_CHECKED(FSoftISMComponentDescriptor, StaticMesh).ToString();
+
+		// Add the Static mesh override to the list only if not already provided
+		if (!Algo::AnyOf(PropertyOverrides, [&StaticMeshPropertyString](const FPCGObjectPropertyOverrideDescription& PropertyOverride) { return PropertyOverride.PropertyTarget == StaticMeshPropertyString; }))
+		{
+			FPCGAttributePropertyInputSelector MeshSelector;
+			MeshSelector.SetAttributeName(AttributeName);
+			PropertyOverrides.Emplace(MeshSelector, StaticMeshPropertyString);
+		}
+
 		// Validate all the selectors are actual FSoftISMComponentDescriptor properties
 		TArray<FPCGAttributePropertySelector> ValidSelectorOverrides;
-		for (const FPCGObjectPropertyOverrideDescription& PropertyOverride : Settings->StaticMeshComponentPropertyOverrides)
+
+		for (const FPCGObjectPropertyOverrideDescription& PropertyOverride : PropertyOverrides)
 		{
 			if (FSoftISMComponentDescriptor::StaticStruct()->FindPropertyByName(FName(PropertyOverride.PropertyTarget)))
 			{
@@ -228,7 +209,7 @@ bool UPCGMeshSelectorByAttribute::SelectInstances(
 
 			// Use the Object Override to map the user's input selector and property to the descriptor
 			FPCGObjectOverrides Overrides(&Descriptor);
-			Overrides.Initialize(Settings->StaticMeshComponentPropertyOverrides, &Descriptor, InPointData, &Context);
+			Overrides.Initialize(PropertyOverrides, &Descriptor, InPointData, &Context);
 
 			// Since they are already partitioned and identical, we can just use the value on the first point
 			check(!Context.AttributeOverridePartition[I].IsEmpty());
@@ -244,93 +225,104 @@ bool UPCGMeshSelectorByAttribute::SelectInstances(
 		}
 	}
 
+	// Assign points to entries
+	int32 CurrentPartitionIndex = Context.CurrentPointIndex; // misnomer but we're reusing another concept from the context
 	const TArray<FPCGPoint>& Points = InPointData->GetPoints();
-
-	while (CurrentPointIndex < Points.Num())
 	{
-		int32 ThisPointIndex = CurrentPointIndex++;
-		const FPCGPoint& Point = Points[ThisPointIndex];
-
-		const PCGMetadataValueKey ValueKey = AttributeBase->GetValueKey(Point.MetadataEntry);
-		TSoftObjectPtr<UStaticMesh>* NewMesh = ValueKeyToMesh.Find(ValueKey);
-		TSoftObjectPtr<UStaticMesh> Mesh = nullptr;
+		TRACE_CPUPROFILER_EVENT_SCOPE(UPCGMeshSelectorByAttribute::SelectEntries::PushingPointsToInstanceLists);
 
-		// If this ValueKey has not been seen before, let's cache it for the future
-		if (!NewMesh)
+		// TODO: Revisit this when attribute partitioning is returned in a more optimized form
+		// The partition index is used to assign the point to the correct partition's instance
+		while(CurrentPartitionIndex < Context.AttributeOverridePartition.Num())
 		{
-			FSoftObjectPath MeshPath;
-			MeshPathGetter(ValueKey, MeshPath);
+			const int32 PartitionIndex = CurrentPartitionIndex++;
+			TArray<int32>& Partition = Context.AttributeOverridePartition[PartitionIndex];
+			const FSoftISMComponentDescriptor& CurrentPartitionDescriptor = Context.OverriddenDescriptors[PartitionIndex];
+			
+			if (Partition.IsEmpty() || CurrentPartitionDescriptor.StaticMesh.IsNull())
+			{
+				continue;
+			}
 
-			if (!MeshPath.IsNull())
+			// Setup data for mesh bounds computation
+			if (OutPointData && Settings->bApplyMeshBoundsToPoints)
 			{
-				Mesh = TSoftObjectPtr<UStaticMesh>(MeshPath);
+				TRACE_CPUPROFILER_EVENT_SCOPE(UPCGMeshSelectorByAttribute::SelectEntries::PushPointsToApplyMeshBounds);
+				TArray<int32>& PointIndices = Context.MeshToOutPoints.FindOrAdd(CurrentPartitionDescriptor.StaticMesh).FindOrAdd(OutPointData);
+				PointIndices.Append(Partition);
+			}
 
-				if (Mesh.IsNull())
+			// Separate the inverse determinant instances so we can push them to a different ISM
+			TArray<int32> ReverseInstances;
+			{
+				TRACE_CPUPROFILER_EVENT_SCOPE(UPCGMeshSelectorByAttribute::SelectEntries::BuildReverseInstances);
+				int32 WriteIndex = 0;
+				for (int32 ReadIndex = 0; ReadIndex < Partition.Num(); ++ReadIndex)
 				{
-					PCGE_LOG_C(Error, GraphAndLog, &Context, FText::Format(LOCTEXT("InvalidMeshPath", "Invalid mesh path: '{0}'."), FText::FromString(MeshPath.ToString())));
+					const FPCGPoint& Point = Points[Partition[ReadIndex]];
+					if (Point.Transform.GetDeterminant() < 0)
+					{
+						ReverseInstances.Add(Partition[ReadIndex]);
+					}
+					else
+					{
+						if (WriteIndex != ReadIndex)
+						{
+							Partition[WriteIndex] = Partition[ReadIndex];
+						}
+						++WriteIndex;
+					}
 				}
-			}
-			else
-			{
-				PCGE_LOG_C(Warning, LogOnly, &Context, LOCTEXT("TrivialInvalidMeshPath", "Trivially invalid mesh path used."));
-			}
 
-			ValueKeyToMesh.Add(ValueKey, Mesh);
-		}
-		else
-		{
-			Mesh = *NewMesh;
-		}
+				Partition.SetNum(WriteIndex);
+			}
 
-		if (Mesh.IsNull())
-		{
-			continue;
-		}
+			auto AddPointsToInstanceList = [&OutMeshInstances, &CurrentPartitionDescriptor, &MaterialOverrideHelper, &Points, PartitionIndex](const TArray<int32>& PointIndices, bool bReverseTransform)
+			{
+				if (MaterialOverrideHelper.OverridesMaterials())
+				{
+					for (int32 PointIndex = 0; PointIndex < PointIndices.Num(); ++PointIndex)
+					{
+						const FPCGPoint& Point = Points[PointIndices[PointIndex]];
+						FPCGMeshInstanceList& InstanceList = PCGMeshSelectorAttribute::GetInstanceList(OutMeshInstances, CurrentPartitionDescriptor, CurrentPartitionDescriptor.StaticMesh, MaterialOverrideHelper.GetMaterialOverrides(Point.MetadataEntry), bReverseTransform, PartitionIndex);
+						InstanceList.Instances.Emplace(Point.Transform);
+						InstanceList.InstancesMetadataEntry.Emplace(Point.MetadataEntry);
+					}
+				}
+				else
+				{
+					TArray<TSoftObjectPtr<UMaterialInterface>> DummyMaterialList;
+					FPCGMeshInstanceList& InstanceList = PCGMeshSelectorAttribute::GetInstanceList(OutMeshInstances, CurrentPartitionDescriptor, CurrentPartitionDescriptor.StaticMesh, DummyMaterialList, bReverseTransform, PartitionIndex);
+
+					check(InstanceList.Instances.Num() == InstanceList.InstancesMetadataEntry.Num());
+					const int32 InstanceOffset = InstanceList.Instances.Num();
+					InstanceList.Instances.SetNum(InstanceOffset + PointIndices.Num());
+					InstanceList.InstancesMetadataEntry.SetNum(InstanceOffset + PointIndices.Num());
+
+					for(int32 PointIndex = 0; PointIndex < PointIndices.Num(); ++PointIndex)
+					{
+						const FPCGPoint& Point = Points[PointIndices[PointIndex]];
+						InstanceList.Instances[InstanceOffset + PointIndex] = Point.Transform;
+						InstanceList.InstancesMetadataEntry[InstanceOffset + PointIndex] = Point.MetadataEntry;
+					}
+				}
+			};
 
-		// TODO: Revisit this when attribute partitioning is returned in a more optimized form
-		// The partition index is used to assign the point to the correct partition's instance
-		int32 CurrentPointPartitionIndex = INDEX_NONE;
-		for (int PartitionIndex = 0; PartitionIndex < Context.AttributeOverridePartition.Num(); ++PartitionIndex)
-		{
-			TArray<int32>& Partition = Context.AttributeOverridePartition[PartitionIndex];
-			if (Partition.Contains(ThisPointIndex))
 			{
-				CurrentPointPartitionIndex = PartitionIndex;
-				break;
+				TRACE_CPUPROFILER_EVENT_SCOPE(UPCGMeshSelectorByAttribute::SelectEntries::AddPointsToInstanceList);
+				AddPointsToInstanceList(Partition, false);
+				AddPointsToInstanceList(ReverseInstances, true);
 			}
-		}
-
-		// If the point wasn't found in a partition (likely due to no partitions or an invalid property target), just default to the template
-		const FSoftISMComponentDescriptor& CurrentPartitionDescriptor = CurrentPointPartitionIndex != INDEX_NONE ? Context.OverriddenDescriptors[CurrentPointPartitionIndex] : TemplateDescriptor;
-		const bool bReverseTransform = (Point.Transform.GetDeterminant() < 0);
 
-		FPCGMeshInstanceList& InstanceList = PCGMeshSelectorAttribute::GetInstanceList(OutMeshInstances, CurrentPartitionDescriptor, Mesh, MaterialOverrideHelper.GetMaterialOverrides(Point.MetadataEntry), bReverseTransform, CurrentPointPartitionIndex);
-		InstanceList.Instances.Emplace(Point.Transform);
-		InstanceList.InstancesMetadataEntry.Emplace(Point.MetadataEntry);
-
-		if (OutPointData && Settings->bApplyMeshBoundsToPoints)
-		{
-			TArray<int32>& PointIndices = Context.MeshToOutPoints.FindOrAdd(Mesh).FindOrAdd(OutPointData);
-			PointIndices.Emplace(ThisPointIndex);
-		}
-
-		// Check if we should stop here and continue in a subsequent call
-		if (CurrentPointIndex - LastCheckpointIndex >= TimeSlicingCheckFrequency)
-		{
 			if (Context.ShouldStop())
 			{
 				break;
 			}
-			else
-			{
-				LastCheckpointIndex = CurrentPointIndex;
-			}
 		}
 	}
 
-	Context.CurrentPointIndex = CurrentPointIndex;
-
-	return (CurrentPointIndex == Points.Num());
+	Context.CurrentPointIndex = CurrentPartitionIndex; // misnomer, but we're using the same context
+	return CurrentPartitionIndex == Context.AttributeOverridePartition.Num();
 }
 
 #undef LOCTEXT_NAMESPACE
diff --git a/Engine/Plugins/PCG/Source/PCG/Public/Elements/PCGStaticMeshSpawnerContext.h b/Engine/Plugins/PCG/Source/PCG/Public/Elements/PCGStaticMeshSpawnerContext.h
--- a/Engine/Plugins/PCG/Source/PCG/Public/Elements/PCGStaticMeshSpawnerContext.h
+++ b/Engine/Plugins/PCG/Source/PCG/Public/Elements/PCGStaticMeshSpawnerContext.h
@@ -59,9 +59,6 @@ struct FPCGStaticMeshSpawnerContext : public FPCGContext, public IPCGAsyncLoadin
 
 	TMap<TSoftObjectPtr<UStaticMesh>, TMap<UPCGPointData*, TArray<int32>>> MeshToOutPoints;
 
-	// Used in by-attribute selector
-	TMap<PCGMetadataValueKey, TSoftObjectPtr<UStaticMesh>> ValueKeyToMesh;
-
 	// Used in weighted selector
 	TArray<TArray<FPCGMeshInstanceList>> WeightedMeshInstances;
 	TMap<TSoftObjectPtr<UStaticMesh>, PCGMetadataValueKey> MeshToValueKey;
