[PCG] Quick Fix for bug causing newly created points to overwrite attributes in some cases

#rnx

diff --git a/Engine/Plugins/PCG/Source/PCG/Private/Data/PCGWorldData.cpp b/Engine/Plugins/PCG/Source/PCG/Private/Data/PCGWorldData.cpp
--- a/Engine/Plugins/PCG/Source/PCG/Private/Data/PCGWorldData.cpp
+++ b/Engine/Plugins/PCG/Source/PCG/Private/Data/PCGWorldData.cpp
@@ -234,6 +234,7 @@ bool UPCGWorldRayHitData::SamplePoint(const FTransform& InTransform, const FBox&
 
 		// TODO: Pre-create attributes within caller or find a better solution than per point
 		PCGWorldQueryHelpers::ApplyRayHitMetadata(HitResult, QueryParams, OutPoint, OutMetadata);
+		const PCGMetadataEntryKey MetadataEntry = OutPoint.MetadataEntry;
 
 		if (HitResult.IsSet())
 		{
@@ -241,6 +242,7 @@ bool UPCGWorldRayHitData::SamplePoint(const FTransform& InTransform, const FBox&
 
 			// Finally, fill in OutPoint - we're done
 			OutPoint = FPCGPoint(PCGWorldQueryHelpers::GetOrthonormalImpactTransform(Hit), 1.0f, UPCGBlueprintHelpers::ComputeSeedFromPosition(Hit.Location));
+			OutPoint.MetadataEntry = MetadataEntry;
 
 			const bool bApplyMetadataFromLandscape = QueryParams.bApplyMetadataFromLandscape && Hit.GetActor() && Hit.GetActor()->IsA<ALandscapeProxy>();
 
