[PCG] New Create Points Sphere node

+ New Sphere Generation Modes:
  Geodesic
  Angle
  Segments
  Random
  Poisson

Ability to create an angular window

diff --git a/Engine/Plugins/PCG/Source/PCG/Private/Elements/PCGCreatePointsSphere.cpp b/Engine/Plugins/PCG/Source/PCG/Private/Elements/PCGCreatePointsSphere.cpp
new file mode 100644
--- /dev/null
+++ b/Engine/Plugins/PCG/Source/PCG/Private/Elements/PCGCreatePointsSphere.cpp
@@ -0,0 +1,477 @@
+// Copyright Epic Games, Inc. All Rights Reserved.
+
+#include "Elements/PCGCreatePointsSphere.h"
+
+#include "PCGComponent.h"
+#include "Data/PCGPointData.h"
+#include "Helpers/PCGAsync.h"
+#include "Helpers/PCGHelpers.h"
+#include "Helpers/PCGSettingsHelpers.h"
+
+#include "GameFramework/Actor.h"
+
+#include UE_INLINE_GENERATED_CPP_BY_NAME(PCGCreatePointsSphere)
+
+#define LOCTEXT_NAMESPACE "PCGCreatePointsSphereElement"
+
+namespace PCGCreatePointsSphere
+{
+	constexpr double Epsilon = UE_DOUBLE_SMALL_NUMBER;
+
+	struct FCommonParams
+	{
+		// Handle clamping in the constructor for consistency in every helper function
+		explicit FCommonParams(
+			double InLatitudinalStartAngle,
+			double InLatitudinalEndAngle,
+			double InLongitudinalStartAngle,
+			double InLongitudinalEndAngle,
+			double InRadius,
+			const FVector& InOrigin,
+			int32 InSeed = 0,
+			FPCGContext* InContext = nullptr)
+		: LatitudinalStartAngle(FMath::Abs(InLatitudinalStartAngle) > 90 ? FMath::Fmod(InLatitudinalStartAngle, 90) : InLatitudinalStartAngle)
+		, LatitudinalEndAngle(FMath::Abs(InLatitudinalEndAngle) > 90 ? FMath::Fmod(InLatitudinalEndAngle, 90) : InLatitudinalEndAngle)
+		, LongitudinalStartAngle(FMath::Abs(InLongitudinalStartAngle) > 180 ? FMath::Fmod(InLongitudinalStartAngle, 180) : InLongitudinalStartAngle)
+		, LongitudinalEndAngle(FMath::Abs(InLongitudinalEndAngle) > 180 ? FMath::Fmod(InLongitudinalEndAngle, 180) : InLongitudinalEndAngle)
+		, Radius(InRadius)
+		, Origin(InOrigin)
+		, Seed(InSeed)
+		, Context(InContext) {}
+
+		double LatitudinalStartAngle = 0;
+		double LatitudinalEndAngle = 0;
+		double LongitudinalStartAngle = 0;
+		double LongitudinalEndAngle = 0;
+		double Radius = 0;
+		FVector Origin = FVector::ZeroVector;
+		int32 Seed = 0;
+		FPCGContext* Context = nullptr;
+	};
+
+	FVector GeneratePointInWindow(const FCommonParams& Params, FRandomStream& RandomSource)
+	{
+		// Use the inverse cosine's range of [-1,1] to drive the angle for uniform sampling
+		const double NormalizedValue = FMath::GetMappedRangeValueUnclamped<double>({-90, 90}, {-1, 1}, RandomSource.FRandRange(Params.LatitudinalStartAngle, Params.LatitudinalEndAngle));
+		const double Theta = FMath::Acos(NormalizedValue);
+		const double Phi = FMath::DegreesToRadians(Params.LongitudinalEndAngle - Params.LongitudinalStartAngle) * RandomSource.FRand();
+		const double SinTheta = FMath::Sin(Theta);
+
+		return Params.Origin + FVector(SinTheta * FMath::Cos(Phi), SinTheta * FMath::Sin(Phi), FMath::Cos(Theta)) * Params.Radius;
+	}
+
+	void GenerateGeodesicSpherePointLocations(const FCommonParams& Params, int32 SubdivisionCount, TArray<FVector>& OutLocations)
+	{
+		TRACE_CPUPROFILER_EVENT_SCOPE(FPCGCreatePointsSphereElement::Geodesic);
+		OutLocations.Empty();
+
+		// TODO: Factor out the icosahedron and add options like tetrahedron, octahedron, etc
+
+		// Golden ratio = (1 + sqrt(5))/2
+		constexpr double Tau = 1.618033988749894;
+
+		// Icosahedron hardcoded unit vertex locations
+		TArray<FVector> PointLocations = {{0, -1, Tau}, {0, 1, Tau}, {-Tau, 0, 1}, {Tau, 0, 1}, {-1, Tau, 0}, {1, Tau, 0}, {-1, -Tau, 0}, {1, -Tau, 0}, {-Tau, 0, -1}, {Tau, 0, -1}, {0, -1, -Tau}, {0, 1, -Tau}};
+
+		auto SubdivideTriangleRecursive = [&Out = PointLocations](const FVector& V0, const FVector& V1, const FVector& V2, int32 SubdivisionCount, auto&& RecursiveCall)
+		{
+			if (SubdivisionCount < 1)
+			{
+				return;
+			}
+
+			// TODO: This can probably be better optimized to not rely on unique
+			// First calculate and add the midpoints
+			const int32 I0 = Out.AddUnique((V0 + V1) * 0.5);
+			const int32 I1 = Out.AddUnique((V1 + V2) * 0.5);
+			const int32 I2 = Out.AddUnique((V2 + V0) * 0.5);
+
+			// TODO: This is highly inefficient with a small window and high subdivision count. Check if the triangle is outside the window and cull if so.
+			// Then recurse on the four new triangles
+			RecursiveCall(V0, Out[I0], Out[I2], SubdivisionCount - 1, RecursiveCall);
+			RecursiveCall(V1, Out[I0], Out[I1], SubdivisionCount - 1, RecursiveCall);
+			RecursiveCall(V2, Out[I1], Out[I2], SubdivisionCount - 1, RecursiveCall);
+			RecursiveCall(Out[I0], Out[I1], Out[I2], SubdivisionCount - 1, RecursiveCall);
+		};
+
+		// Hardcoded Icosahedron Triangles
+		constexpr int32 InitialTriangles[] = {0, 1, 2, 0, 1, 3, 0, 2, 6, 0, 6, 7, 0, 3, 7, 1, 2, 4, 1, 3, 5, 1, 4, 5, 2, 4, 8, 2, 6, 8, 3, 5, 9, 3, 7, 9, 4, 5, 11, 4, 8, 11, 5, 9, 11, 6, 7, 10, 6, 8, 10, 7, 9, 10, 8, 10, 11, 9, 10, 11};
+
+		// Call the recursive algorithm on the original triangles
+		if (SubdivisionCount > 0)
+		{
+			for (int I = 2; I < std::size(InitialTriangles); I += 3)
+			{
+				SubdivideTriangleRecursive(PointLocations[InitialTriangles[I - 2]], PointLocations[InitialTriangles[I - 1]], PointLocations[InitialTriangles[I]], SubdivisionCount, SubdivideTriangleRecursive);
+			}
+		}
+
+		const double LatStartRad = FMath::DegreesToRadians(Params.LatitudinalStartAngle);
+		const double LatEndRad = FMath::DegreesToRadians(Params.LatitudinalEndAngle);
+		const double LonStartRad = FMath::DegreesToRadians(Params.LongitudinalStartAngle);
+		const double LonEndRad = FMath::DegreesToRadians(Params.LongitudinalEndAngle);
+
+		int32 Index = 0;
+		while (Index < PointLocations.Num())
+		{
+			FVector& PointLocation = PointLocations[Index];
+
+			// Since this algorithm is a recursive subdivision, culling will have to wait until the end
+			const double LatAngle = FMath::Atan2(PointLocation.Z / FVector(PointLocation.X, PointLocation.Y, 0).Length(), 1);
+			const double LonAngle = FMath::Atan2(PointLocation.Y, PointLocation.X);
+			if (LatAngle >= LatStartRad && LatAngle <= LatEndRad && LonAngle >= LonStartRad && LonAngle <= LonEndRad)
+			{
+				// Normalize to bring it to the unit sphere
+				PointLocation = (1.0 / PointLocation.Length()) * PointLocation;
+
+				// Finally, apply the user defined transform parameters
+				PointLocation *= Params.Radius;
+				PointLocation += Params.Origin;
+				++Index;
+			}
+			else // Cull the point
+			{
+				// No need to shrink, since it will append at the end
+				PointLocations.RemoveSingleSwap(PointLocation, EAllowShrinking::No);
+			}
+		}
+
+		OutLocations.Append(PointLocations);
+	}
+
+	void GenerateStandardSpherePointLocations(const FCommonParams& Params, double InTheta, double InPhi, double InJitter, TArray<FVector>& OutLocations)
+	{
+		TRACE_CPUPROFILER_EVENT_SCOPE(FPCGCreatePointsSphereElement::StandardSphere);
+		OutLocations.Empty();
+
+		const double& LatStartAngle = Params.LatitudinalStartAngle;
+		const double& LatEndAngle = Params.LatitudinalEndAngle;
+		const double& LonStartAngle = Params.LongitudinalStartAngle;
+		const double& LonEndAngle = Params.LongitudinalEndAngle;
+
+		const int32 LatSegCount = 1 + FMath::FloorToInt32(FMath::Abs(LatEndAngle - LatStartAngle) / InTheta);
+		const int32 LonSegCount = FMath::FloorToInt32(FMath::Abs(LonEndAngle - LonStartAngle) / InPhi);
+
+		int64 PointCount = 0;
+		// To produce a single pole point, instead of a whole segment
+		int32 ReducedLatSegCount = LatSegCount;
+
+		// Account for the poles
+
+		bool bSouthPole = false;
+		if (FMath::IsNearlyEqual(LatEndAngle, -90.0))
+		{
+			--ReducedLatSegCount;
+			++PointCount;
+			bSouthPole = true;
+		}
+
+		bool bNorthPole = false;
+		if (FMath::IsNearlyEqual(LatStartAngle, 90.0))
+		{
+			--ReducedLatSegCount;
+			++PointCount;
+			bNorthPole = true;
+		}
+
+		PointCount += ReducedLatSegCount * LonSegCount;
+
+		if (PointCount < 1)
+		{
+			return;
+		}
+
+		if (PCGFeatureSwitches::CVarCheckSamplerMemory.GetValueOnAnyThread() && FPlatformMemory::GetStats().AvailablePhysical < sizeof(FPCGPoint) * PointCount)
+		{
+			PCGLog::LogErrorOnGraph(LOCTEXT("MemoryOverflow", "The number of iterations produced is larger than available memory."));
+			return;
+		}
+
+		OutLocations.SetNumUninitialized(PointCount);
+
+		FRandomStream RandomSource(Params.Seed);
+		// Jitter is [0, 0.5) which is the window to not have any overlap possibility
+		const double Jitter = FMath::Clamp(InJitter, 0.0, 0.5 - Epsilon);
+		bool bApplyJitter = !FMath::IsNearlyZero(Jitter);
+
+		int32 CurrentIndex = 0;
+		for (int LatIndex = 0; LatIndex < LatSegCount; ++LatIndex)
+		{
+			// Add the poles under certain conditions
+			if (LatIndex == 0 && bSouthPole)
+			{
+				OutLocations[CurrentIndex++] = Params.Origin - FVector::ZAxisVector * Params.Radius;
+				continue;
+			}
+
+			if (LatIndex == LatSegCount - 1 && bNorthPole)
+			{
+				OutLocations[CurrentIndex++] = Params.Origin + FVector::ZAxisVector * Params.Radius;
+				continue;
+			}
+
+			// -90 to map to a -180 to 0 range for cosine
+			const double Theta = FMath::DegreesToRadians(LatStartAngle - 90 + InTheta * LatIndex);
+			// Pre-calculation of these for the internal loop is possible unless jittered
+			double SinTheta = bApplyJitter ? 0.0 : FMath::Sin(Theta);
+			double CosTheta = bApplyJitter ? 0.0 : FMath::Cos(Theta);
+
+			for (int LonIndex = 0; LonIndex < LonSegCount; ++LonIndex)
+			{
+				double Phi = FMath::DegreesToRadians(LonStartAngle + InPhi * LonIndex);
+
+				if (bApplyJitter)
+				{
+					Phi += FMath::DegreesToRadians(InPhi) * RandomSource.FRandRange(-Jitter, Jitter);
+					const double JitteredTheta = Theta + FMath::DegreesToRadians(InTheta) * RandomSource.FRandRange(-Jitter, Jitter);
+					SinTheta = FMath::Sin(JitteredTheta);
+					CosTheta = FMath::Cos(JitteredTheta);
+				}
+
+				OutLocations[CurrentIndex++] = Params.Origin + FVector(SinTheta * FMath::Cos(Phi), SinTheta * FMath::Sin(Phi), CosTheta) * Params.Radius;
+			}
+		}
+
+		check(CurrentIndex == PointCount);
+	}
+
+	void GenerateRandomSamplingPointLocations(const FCommonParams& Params, int32 SampleCount, TArray<FVector>& OutLocations)
+	{
+		TRACE_CPUPROFILER_EVENT_SCOPE(FPCGCreatePointsSphereElement::RandomSampling);
+		OutLocations.Empty();
+
+		if (PCGFeatureSwitches::CVarCheckSamplerMemory.GetValueOnAnyThread() && FPlatformMemory::GetStats().AvailablePhysical < sizeof(FPCGPoint) * SampleCount)
+		{
+			PCGLog::LogErrorOnGraph(LOCTEXT("MemoryOverflow", "The number of iterations produced is larger than available memory."));
+			return;
+		}
+
+		OutLocations.SetNumUninitialized(SampleCount);
+
+		FRandomStream RandomSource(Params.Seed);
+
+		for (int32 Index = 0; Index < SampleCount; ++Index)
+		{
+			OutLocations[Index] = GeneratePointInWindow(Params, RandomSource);
+		}
+	}
+
+	void GeneratePoissonSamplingPointLocations(const FCommonParams& Params, double Distance, int32 SampleCount, int32 MaxAttempts, TArray<FVector>& OutLocations)
+	{
+		TRACE_CPUPROFILER_EVENT_SCOPE(FPCGCreatePointsSphereElement::Poisson);
+		OutLocations.Empty();
+
+		if (PCGFeatureSwitches::CVarCheckSamplerMemory.GetValueOnAnyThread() && FPlatformMemory::GetStats().AvailablePhysical < sizeof(FPCGPoint) * SampleCount)
+		{
+			PCGLog::LogErrorOnGraph(LOCTEXT("MemoryOverflow", "The number of iterations produced is larger than available memory."));
+			return;
+		}
+
+		OutLocations.Reserve(SampleCount);
+
+		const double DistanceSquared = Distance * Distance;
+
+		FRandomStream RandomSource(Params.Seed);
+
+		for (int I = 0; I < SampleCount; ++I)
+		{
+			for (int32 Attempt = 0; Attempt < MaxAttempts; ++Attempt)
+			{
+				FVector Point = GeneratePointInWindow(Params, RandomSource);
+
+				// TODO: Might be beneficial to check against an octree before the distance check
+				bool bValidPoint = true;
+				for (const FVector& PreviousPoint : OutLocations)
+				{
+					if ((PreviousPoint - Point).SizeSquared() < DistanceSquared)
+					{
+						bValidPoint = false;
+						break;
+					}
+				}
+
+				if (bValidPoint)
+				{
+					OutLocations.Emplace(Point);
+					break;
+				}
+			}
+		}
+	}
+}
+
+TArray<FPCGPinProperties> UPCGCreatePointsSphereSettings::OutputPinProperties() const
+{
+	return DefaultPointOutputPinProperties();
+}
+
+FPCGElementPtr UPCGCreatePointsSphereSettings::CreateElement() const
+{
+	return MakeShared<FPCGCreatePointsSphereElement>();
+}
+
+bool FPCGCreatePointsSphereElement::ExecuteInternal(FPCGContext* InContext) const
+{
+	TRACE_CPUPROFILER_EVENT_SCOPE(FPCGCreatePointsSphereElement::Execute);
+	const UPCGCreatePointsSphereSettings* Settings = InContext->GetInputSettings<UPCGCreatePointsSphereSettings>();
+
+	// Used for culling, regardless of generation coordinate space
+	const UPCGSpatialData* CullingShape = Settings->bCullPointsOutsideVolume ? Cast<UPCGSpatialData>(InContext->SourceComponent->GetActorPCGData()) : nullptr;
+
+	// Early out if the culling shape isn't valid
+	if (!CullingShape && Settings->bCullPointsOutsideVolume)
+	{
+		PCGLog::LogErrorOnGraph(LOCTEXT("CannotCullWithoutAShape", "Unable to cull since the supporting actor has no data."));
+		return true;
+	}
+
+	const PCGCreatePointsSphere::FCommonParams Params(
+		Settings->LatitudinalStartAngle,
+		Settings->LatitudinalEndAngle,
+		Settings->LongitudinalStartAngle,
+		Settings->LongitudinalEndAngle,
+		Settings->Radius,
+		Settings->Origin,
+		InContext->GetSeed(),
+		InContext);
+
+	TArray<FVector> OutLocalPointLocations;
+	switch (Settings->SphereGeneration)
+	{
+		case EPCGSphereGeneration::Geodesic:
+			{
+				PCGCreatePointsSphere::GenerateGeodesicSpherePointLocations(Params, Settings->GeodesicSubdivisions, OutLocalPointLocations);
+			}
+			break;
+		case EPCGSphereGeneration::Angle:
+			{
+				PCGCreatePointsSphere::GenerateStandardSpherePointLocations(Params, Settings->Theta, Settings->Phi, Settings->Jitter, OutLocalPointLocations);
+			}
+			break;
+		case EPCGSphereGeneration::Segments:
+			{
+				if (Settings->LatitudinalSegments < 1 || Settings->LongitudinalSegments < 1)
+				{
+					return true;
+				}
+
+				const double Theta = FMath::Abs(Settings->LatitudinalStartAngle - Settings->LatitudinalEndAngle) / Settings->LatitudinalSegments;
+				const double Phi = FMath::Abs(Settings->LongitudinalStartAngle - Settings->LongitudinalEndAngle) / Settings->LongitudinalSegments;
+				PCGCreatePointsSphere::GenerateStandardSpherePointLocations(Params, Theta, Phi, Settings->Jitter, OutLocalPointLocations);
+			}
+			break;
+		case EPCGSphereGeneration::Random:
+			{
+				PCGCreatePointsSphere::GenerateRandomSamplingPointLocations(Params, Settings->SampleCount, OutLocalPointLocations);
+			}
+			break;
+		case EPCGSphereGeneration::Poisson:
+			{
+				PCGCreatePointsSphere::GeneratePoissonSamplingPointLocations(Params, Settings->PoissonDistance, Settings->SampleCount, Settings->PoissonMaxAttempts, OutLocalPointLocations);
+			}
+			break;
+		default:
+			checkNoEntry();
+			break;
+	}
+
+	TArray<FPCGTaggedData>& Outputs = InContext->OutputData.TaggedData;
+	FPCGTaggedData& Output = Outputs.Emplace_GetRef();
+
+	UPCGPointData* PointData = NewObject<UPCGPointData>();
+	TArray<FPCGPoint>& OutputPoints = PointData->GetMutablePoints();
+	Output.Data = PointData;
+
+	FTransform LocalTransform = FTransform::Identity;
+
+	if (Settings->CoordinateSpace == EPCGCoordinateSpace::OriginalComponent)
+	{
+		check(InContext->SourceComponent->GetOriginalComponent() && InContext->SourceComponent->GetOriginalComponent()->GetOwner());
+		LocalTransform = InContext->SourceComponent->GetOriginalComponent()->GetOwner()->GetActorTransform();
+	}
+	else if (Settings->CoordinateSpace == EPCGCoordinateSpace::LocalComponent)
+	{
+		check(InContext->SourceComponent->GetOwner());
+		LocalTransform = InContext->SourceComponent->GetOwner()->GetActorTransform();
+	}
+
+	// Point scale should not be relative to the coordinate space
+	LocalTransform.SetScale3D(FVector::One());
+
+	OutputPoints.Reserve(OutLocalPointLocations.Num());
+
+	for (int Index = 0; Index < OutLocalPointLocations.Num(); ++Index)
+	{
+		FRotator LocalPointOrientation = FRotator::ZeroRotator;
+
+		// TODO: A second pass would be useful, as we could likely make this more efficient.
+		switch (Settings->PointOrientation)
+		{
+			case EPCGSpherePointOrientation::Radial:
+				LocalPointOrientation = FRotationMatrix::MakeFromZ(OutLocalPointLocations[Index]).Rotator();
+				break;
+			case EPCGSpherePointOrientation::Centric:
+				LocalPointOrientation = FRotationMatrix::MakeFromZ(-OutLocalPointLocations[Index]).Rotator();
+				break;
+			case EPCGSpherePointOrientation::None:
+				break;
+			default:
+				checkNoEntry();
+				break;
+		}
+
+		FTransform PointTransform(LocalPointOrientation, OutLocalPointLocations[Index]);
+
+		if (Settings->CoordinateSpace == EPCGCoordinateSpace::LocalComponent || Settings->CoordinateSpace == EPCGCoordinateSpace::OriginalComponent)
+		{
+			PointTransform *= LocalTransform;
+		}
+
+		if (!CullingShape || (CullingShape->GetDensityAtPosition(PointTransform.GetLocation()) > 0.f))
+		{
+			FPCGPoint& OutPoint = OutputPoints.Emplace_GetRef(PointTransform, 1.f, PCGHelpers::ComputeSeedFromPosition(OutLocalPointLocations[Index]));
+			OutPoint.Steepness = Settings->PointSteepness;
+		}
+	}
+
+	return true;
+}
+
+void FPCGCreatePointsSphereElement::GetDependenciesCrc(const FPCGDataCollection& InInput, const UPCGSettings* InSettings, UPCGComponent* InComponent, FPCGCrc& OutCrc) const
+{
+	FPCGCrc Crc;
+	IPCGElement::GetDependenciesCrc(InInput, InSettings, InComponent, Crc);
+
+	if (const UPCGCreatePointsSphereSettings* Settings = Cast<UPCGCreatePointsSphereSettings>(InSettings))
+	{
+		int CoordinateSpace = static_cast<int>(EPCGCoordinateSpace::World);
+		bool bCullPointsOutsideVolume = false;
+		PCGSettingsHelpers::GetOverrideValue(InInput, Settings, GET_MEMBER_NAME_CHECKED(UPCGCreatePointsSphereSettings, CoordinateSpace), static_cast<int>(Settings->CoordinateSpace), CoordinateSpace);
+		PCGSettingsHelpers::GetOverrideValue(InInput, Settings, GET_MEMBER_NAME_CHECKED(UPCGCreatePointsSphereSettings, bCullPointsOutsideVolume), Settings->bCullPointsOutsideVolume, bCullPointsOutsideVolume);
+
+		EPCGCoordinateSpace EnumCoordinateSpace = static_cast<EPCGCoordinateSpace>(CoordinateSpace);
+
+		if (InComponent)
+		{
+			const UPCGData* Data = nullptr;
+
+			if (bCullPointsOutsideVolume || EnumCoordinateSpace == EPCGCoordinateSpace::LocalComponent)
+			{
+				Data = InComponent->GetActorPCGData();
+			}
+			else if (EnumCoordinateSpace == EPCGCoordinateSpace::OriginalComponent)
+			{
+				Data = InComponent->GetOriginalActorPCGData();
+			}
+
+			if (Data)
+			{
+				Crc.Combine(Data->GetOrComputeCrc(/*bFullDataCrc=*/false));
+			}
+		}
+	}
+
+	OutCrc = Crc;
+}
+
+#undef LOCTEXT_NAMESPACE
diff --git a/Engine/Plugins/PCG/Source/PCG/Public/Elements/PCGCreatePointsSphere.h b/Engine/Plugins/PCG/Source/PCG/Public/Elements/PCGCreatePointsSphere.h
new file mode 100644
--- /dev/null
+++ b/Engine/Plugins/PCG/Source/PCG/Public/Elements/PCGCreatePointsSphere.h
@@ -0,0 +1,147 @@
+// Copyright Epic Games, Inc. All Rights Reserved.
+#pragma once
+
+#include "PCGSettings.h"
+
+#include "PCGCreatePointsSphere.generated.h"
+
+struct ISphereGenerator;
+
+UENUM()
+enum class EPCGSphereGeneration : uint8
+{
+	Geodesic UMETA(ToolTip = "Points generated from subdivided equilateral triangles across the sphere's surface."),
+	Angle UMETA(ToolTip = "Points generated radially along longitude and latitude lines to form a sphere with user defined angles."),
+	Segments UMETA(ToolTip = "Points generated radially along longitude and latitude lines to form a sphere with user defined number of segments."),
+	Random UMETA(ToolTip = "Points are generated with a uniform distribution on the surface of the sphere."),
+	Poisson UMETA(ToolTip = "Sample points along the sphere's surface with a minimal distance between any given two points.")
+};
+
+UENUM()
+enum class EPCGGeodesicSphereRepresentation : uint8
+{
+	Icosahedron UMETA(ToolTip = "Twenty sided polyhedron."),
+	// TODO: Tetrahedron, Octahedron
+};
+
+UENUM()
+enum class EPCGSpherePointOrientation : uint8
+{
+	Radial UMETA(ToolTip = "Points are aligned outward from the origin (the normal of the sphere's surface)."),
+	Centric UMETA(ToolTip = "Points are aligned inward toward the origin (inverse of the normal)."),
+	None UMETA(ToolTip = "Points will not be explicitly aligned.")
+};
+
+/**
+ * Generate a grid of points along the surface of a sphere.
+ */
+UCLASS(BlueprintType, ClassGroup = (Procedural))
+class UPCGCreatePointsSphereSettings : public UPCGSettings
+{
+	GENERATED_BODY()
+
+public:
+	UPCGCreatePointsSphereSettings() { bUseSeed = true; }
+
+	//~Begin UPCGSettings interface
+#if WITH_EDITOR
+	virtual FName GetDefaultNodeName() const override { return FName(TEXT("CreatePointsSphere")); }
+	virtual FText GetDefaultNodeTitle() const override { return NSLOCTEXT("PCGCreatePointsSphereElement", "NodeTitle", "Create Points Sphere"); }
+	virtual FText GetNodeTooltipText() const override { return NSLOCTEXT("PCGCreatePointsSphereElement", "NodeTooltip", "Generate points on the surface of a sphere."); }
+	virtual EPCGSettingsType GetType() const override { return EPCGSettingsType::Spatial; }
+#endif // WITH_EDITOR
+
+protected:
+	virtual TArray<FPCGPinProperties> InputPinProperties() const override { return {}; }
+	virtual TArray<FPCGPinProperties> OutputPinProperties() const override;
+	virtual FPCGElementPtr CreateElement() const override;
+	//~End UPCGSettings interface
+
+public:
+	/** Determines the type of sphere generated. */
+	UPROPERTY(BlueprintReadWrite, EditAnywhere, Category = Settings, meta = (PCG_Overridable))
+	EPCGSphereGeneration SphereGeneration = EPCGSphereGeneration::Geodesic;
+
+	/** Sets the generation referential of the points. */
+	UPROPERTY(BlueprintReadWrite, EditAnywhere, Category = Settings, meta = (PCG_Overridable))
+	EPCGCoordinateSpace CoordinateSpace = EPCGCoordinateSpace::World;
+
+	/** Will determine the points' orientation, once generated. */
+	UPROPERTY(BlueprintReadWrite, EditAnywhere, Category = Settings, meta = (PCG_Overridable))
+	EPCGSpherePointOrientation PointOrientation = EPCGSpherePointOrientation::Radial;
+
+	/** The sphere's origin, around which the points will be generated. */
+	UPROPERTY(BlueprintReadWrite, EditAnywhere, Category = Settings, meta = (PCG_Overridable))
+	FVector Origin = FVector::ZeroVector;
+
+	/** The sphere's radius. */
+	UPROPERTY(BlueprintReadWrite, EditAnywhere, Category = Settings, meta = (PCG_Overridable))
+	double Radius = 100.0;
+
+	/** Determines the number of subdivisions of the geodesic sphere. Becomes exponentially more expensive as it gets higher. */
+	UPROPERTY(BlueprintReadWrite, EditAnywhere, Category = Settings, meta = (PCG_Overridable, EditCondition = "SphereGeneration == EPCGSphereGeneration::Geodesic", EditConditionHides, ClampMin = 0, ClampMax = 8))
+	int32 GeodesicSubdivisions = 2;
+
+	/** The latitudinal angle (in degrees) between generated segments on the standard sphere grid. */
+	UPROPERTY(BlueprintReadWrite, EditAnywhere, DisplayName = "Latitudinal Segment (Degrees)", Category = Settings, meta = (PCG_Overridable, EditCondition = "SphereGeneration == EPCGSphereGeneration::Angle", EditConditionHides, ClampMin = 0.001, ClampMax = 90))
+	double Theta = 12;
+
+	/** The longitudinal angle (in degrees) between generated segments on the standard sphere grid. */
+	UPROPERTY(BlueprintReadWrite, EditAnywhere, DisplayName = "Longitudinal Segment (Degrees)", Category = Settings, meta = (PCG_Overridable, EditCondition = "SphereGeneration == EPCGSphereGeneration::Angle", EditConditionHides, ClampMin = 0.001, ClampMax = 180))
+	double Phi = 12;
+
+	/** Will determine the latitudinal angle between segments needed to generate this number of latitudinal segments. */
+	UPROPERTY(BlueprintReadWrite, EditAnywhere, Category = Settings, meta = (PCG_Overridable, EditCondition = "SphereGeneration == EPCGSphereGeneration::Segments", EditConditionHides, ClampMin = 1, ClampMax = 90000))
+	int32 LatitudinalSegments = 15;
+
+	/** Will determine the longitudinal angle between segments needed to generate this number of longitudinal segments. */
+	UPROPERTY(BlueprintReadWrite, EditAnywhere, Category = Settings, meta = (PCG_Overridable, EditCondition = "SphereGeneration == EPCGSphereGeneration::Segments", EditConditionHides, ClampMin = 1, ClampMax = 180000))
+	int32 LongitudinalSegments = 30;
+
+	/** Determines the number of samples generated for the random generation and the poisson sampling. */
+	UPROPERTY(BlueprintReadWrite, EditAnywhere, Category = Settings, meta = (PCG_Overridable, EditCondition = "SphereGeneration == EPCGSphereGeneration::Random || SphereGeneration == EPCGSphereGeneration::Poisson", EditConditionHides, ClampMin = 1))
+	int32 SampleCount = 100;
+
+	/** The maximum world distance between points sampled on the sphere's surface during a Poisson sampling, in cm. */
+	UPROPERTY(BlueprintReadWrite, EditAnywhere, Category = Settings, meta = (PCG_Overridable, EditCondition = "SphereGeneration == EPCGSphereGeneration::Poisson", EditConditionHides, ClampMin = 1))
+	double PoissonDistance = 100.0;
+
+	/** Poisson sampling will continue to search for open positions until this limit is reached. */
+	UPROPERTY(BlueprintReadWrite, EditAnywhere, Category = Settings, meta = (PCG_Overridable, EditCondition = "SphereGeneration == EPCGSphereGeneration::Poisson", EditConditionHides, ClampMin = 1))
+	int32 PoissonMaxAttempts = 32;
+
+	/** Points will be generated on the sphere's surface beginning at this equatorial angle. */
+	UPROPERTY(BlueprintReadWrite, EditAnywhere, Category = Settings, meta = (PCG_Overridable, ClampMin = -90, ClampMax = 90))
+	double LatitudinalStartAngle = -90.0;
+
+	/** Points will cease to be generated on the sphere's surface after this equatorial angle. */
+	UPROPERTY(BlueprintReadWrite, EditAnywhere, Category = Settings, meta = (PCG_Overridable, ClampMin = -90, ClampMax = 90))
+	double LatitudinalEndAngle = 90.0;
+
+	/** Points will be generated on the sphere's surface beginning at this meridional angle. */
+	UPROPERTY(BlueprintReadWrite, EditAnywhere, Category = Settings, meta = (PCG_Overridable, ClampMin = -180, ClampMax = 180))
+	double LongitudinalStartAngle = -180.0;
+
+	/** Points will cease to be generated on the sphere's surface after this meridional angle. */
+	UPROPERTY(BlueprintReadWrite, EditAnywhere, Category = Settings, meta = (PCG_Overridable, ClampMin = -180, ClampMax = 180))
+	double LongitudinalEndAngle = 180.0;
+
+	/** Adds randomization (in the range of [-Jitter, Jitter]) of the angles of a generated point. */
+	UPROPERTY(BlueprintReadWrite, EditAnywhere, Category = Settings, meta = (PCG_Overridable, EditCondition = "SphereGeneration == EPCGSphereGeneration::Angle || SphereGeneration == EPCGSphereGeneration::Segments", EditConditionHides, ClampMin = 0.0, ClampMax = 0.5))
+	double Jitter = 0.0;
+
+	/** Directly set as the output points' steepness value. */
+	UPROPERTY(BlueprintReadWrite, EditAnywhere, Category = Settings, meta = (ClampMin = "0", ClampMax = "1", PCG_Overridable))
+	float PointSteepness = 0.5f;
+
+	/** Points are removed if they are outside the volume. */
+	UPROPERTY(BlueprintReadWrite, EditAnywhere, Category = Settings, meta = (PCG_Overridable))
+	bool bCullPointsOutsideVolume = false;
+};
+
+class FPCGCreatePointsSphereElement : public IPCGElement
+{
+protected:
+	virtual bool ExecuteInternal(FPCGContext* InContext) const override;
+	virtual void GetDependenciesCrc(const FPCGDataCollection& InInput, const UPCGSettings* InSettings, UPCGComponent* InComponent, FPCGCrc& OutCrc) const override;
+};
