[PCG] Fixed spurious check case in the static mesh spawner

#rnx

diff --git a/Engine/Plugins/PCG/Source/PCG/Private/MeshSelectors/PCGMeshSelectorByAttribute.cpp b/Engine/Plugins/PCG/Source/PCG/Private/MeshSelectors/PCGMeshSelectorByAttribute.cpp
--- a/Engine/Plugins/PCG/Source/PCG/Private/MeshSelectors/PCGMeshSelectorByAttribute.cpp
+++ b/Engine/Plugins/PCG/Source/PCG/Private/MeshSelectors/PCGMeshSelectorByAttribute.cpp
@@ -207,6 +207,12 @@ bool UPCGMeshSelectorByAttribute::SelectInstances(
 		{
 			FSoftISMComponentDescriptor& Descriptor = Context.OverriddenDescriptors.Add_GetRef(TemplateDescriptor);
 
+			// If partition is empty (which can happen, esp. for the default partition on the default value, we'll just skip it here.
+			if (Context.AttributeOverridePartition[I].IsEmpty())
+			{
+				continue;
+			}
+
 			// Use the Object Override to map the user's input selector and property to the descriptor
 			FPCGObjectOverrides Overrides(&Descriptor);
 			Overrides.Initialize(PropertyOverrides, &Descriptor, InPointData, &Context);
