[PCG] Added PCG pause button & cancel mechanism to the toolbar, shown through editor settings.

diff --git a/Engine/Plugins/PCG/Source/PCG/Private/Graph/PCGGraphExecutor.cpp b/Engine/Plugins/PCG/Source/PCG/Private/Graph/PCGGraphExecutor.cpp
--- a/Engine/Plugins/PCG/Source/PCG/Private/Graph/PCGGraphExecutor.cpp
+++ b/Engine/Plugins/PCG/Source/PCG/Private/Graph/PCGGraphExecutor.cpp
@@ -600,6 +600,13 @@ void FPCGGraphExecutor::Execute()
 {
 	TRACE_CPUPROFILER_EVENT_SCOPE(FPCGGraphExecutor::Execute);
 
+#if WITH_EDITOR
+	if (PCGComponent::CVarPausePCGExecution.GetValueOnAnyThread())
+	{
+		return;
+	}
+#endif // WITH_EDITOR
+
 	// Process any newly scheduled graphs to execute
 	ScheduleLock.Lock();
 
diff --git a/Engine/Plugins/PCG/Source/PCG/Private/PCGComponent.cpp b/Engine/Plugins/PCG/Source/PCG/Private/PCGComponent.cpp
--- a/Engine/Plugins/PCG/Source/PCG/Private/PCGComponent.cpp
+++ b/Engine/Plugins/PCG/Source/PCG/Private/PCGComponent.cpp
@@ -69,6 +69,13 @@ namespace PCGComponent
 		false,
 		TEXT("Disable refresh for all PCG Components."));
 
+#if WITH_EDITOR
+	TAutoConsoleVariable<bool> CVarPausePCGExecution(
+		TEXT("pcg.PauseExecution"),
+		false,
+		TEXT("Pauses all execution of PCG but does not cancel tasks."));
+#endif
+
 	template <typename DelegateType>
 	static void BroadcastDynamicDelegate(const DelegateType& Delegate, UPCGComponent* PCGComponent)
 	{
diff --git a/Engine/Plugins/PCG/Source/PCG/Public/PCGComponent.h b/Engine/Plugins/PCG/Source/PCG/Public/PCGComponent.h
--- a/Engine/Plugins/PCG/Source/PCG/Public/PCGComponent.h
+++ b/Engine/Plugins/PCG/Source/PCG/Public/PCGComponent.h
@@ -76,6 +76,13 @@ enum class EPCGComponentDirtyFlag : uint8
 };
 ENUM_CLASS_FLAGS(EPCGComponentDirtyFlag);
 
+namespace PCGComponent
+{
+#if WITH_EDITOR
+	extern PCG_API TAutoConsoleVariable<bool> CVarPausePCGExecution;
+#endif
+}
+
 UCLASS(BlueprintType, ClassGroup = (Procedural), meta = (BlueprintSpawnableComponent, PrioritizeCategories = "PCG"))
 class PCG_API UPCGComponent : public UActorComponent
 {
diff --git a/Engine/Plugins/PCG/Source/PCGEditor/Private/PCGEditorModule.cpp b/Engine/Plugins/PCG/Source/PCGEditor/Private/PCGEditorModule.cpp
--- a/Engine/Plugins/PCG/Source/PCGEditor/Private/PCGEditorModule.cpp
+++ b/Engine/Plugins/PCG/Source/PCGEditor/Private/PCGEditorModule.cpp
@@ -9,6 +9,7 @@
 #include "PCGEditorSettings.h"
 #include "PCGEditorStyle.h"
 #include "PCGEditorUtils.h"
+#include "PCGEngineSettings.h"
 #include "PCGModule.h"
 #include "PCGSubsystem.h"
 #include "PCGVolumeFactory.h"
@@ -262,6 +263,76 @@ void FPCGEditorModule::RegisterMenuExtensions()
 			
 		}), FToolMenuInsert(NAME_None, EToolMenuInsertType::Default));
 	}
+
+	if (UToolMenu* ToolbarMenu = UToolMenus::Get()->ExtendMenu("LevelEditor.LevelEditorToolBar.PlayToolBar"))
+	{
+		FToolMenuSection& Section = ToolbarMenu->FindOrAddSection("Play");
+
+		if (const UPCGEditorSettings* EditorSettings = GetDefault<UPCGEditorSettings>())
+		{
+			if (EditorSettings->bShowPauseButton)
+			{
+				FToolMenuEntry PCGPauseButton = FToolMenuEntry::InitToolBarButton(
+					"PCGPauseButton",
+					FUIAction(
+						FExecuteAction::CreateLambda([]()
+						{
+							const bool bWasPaused = PCGComponent::CVarPausePCGExecution.GetValueOnAnyThread();
+
+							if (bWasPaused)
+							{
+								bool bShouldCancelAll = true;
+
+								if (FSlateApplication::Get().GetModifierKeys().IsControlDown())
+								{
+									bShouldCancelAll = false;
+								}
+								else if (FSlateApplication::Get().GetModifierKeys().IsAltDown())
+								{
+									bShouldCancelAll = true;
+								}
+								else if (const UPCGEditorSettings* EditorSettings = GetDefault<UPCGEditorSettings>())
+								{
+									bShouldCancelAll = EditorSettings->bUnpauseCancelsAll;
+								}
+
+								if (bShouldCancelAll)
+								{
+									if (UPCGSubsystem* PCGSubsystem = UPCGSubsystem::GetInstance(GEditor->GetEditorWorldContext().World()))
+									{
+										PCGSubsystem->CancelAllGeneration();
+									}
+								}
+							}
+
+							PCGComponent::CVarPausePCGExecution->Set(!bWasPaused);
+						}),
+						FIsActionButtonVisible::CreateLambda([](){ return true; }),
+						FIsActionChecked::CreateLambda([]()
+						{
+							return PCGComponent::CVarPausePCGExecution.GetValueOnAnyThread();
+						})
+					),
+					TAttribute<FText>::CreateLambda([]()
+					{
+						if(PCGComponent::CVarPausePCGExecution.GetValueOnAnyThread())
+						{
+							return LOCTEXT("PCGPauseButton_Off", "Paused");
+						}
+						else
+						{
+							return LOCTEXT("PCGPauseButton_On", "PCG");
+						}
+					}),
+					LOCTEXT("PCGPauseButton_Tooltip", "Toggles PCG processing on/off and will cancel tasks depending on settings.\nUse Ctrl to unpause without cancelling all tasks.\nUse Alt to unpause and cancel all tasks."),
+					FSlateIcon(FPCGEditorStyle::Get().GetStyleSetName(), "PCG.EditorIcon"),
+					EUserInterfaceActionType::ToggleButton
+				);
+				PCGPauseButton.StyleNameOverride = "CalloutToolbar"; // used to show the button text
+				Section.AddEntry(PCGPauseButton);
+			}
+		}
+	}
 }
 
 void FPCGEditorModule::UnregisterMenuExtensions()
diff --git a/Engine/Plugins/PCG/Source/PCGEditor/Public/PCGEditorSettings.h b/Engine/Plugins/PCG/Source/PCGEditor/Public/PCGEditorSettings.h
--- a/Engine/Plugins/PCG/Source/PCGEditor/Public/PCGEditorSettings.h
+++ b/Engine/Plugins/PCG/Source/PCGEditor/Public/PCGEditorSettings.h
@@ -172,6 +172,12 @@ public:
 	/** Specify if we want to disable CPU Throttling when a PCG Graph is executing, this will improve execution time when app is out of focus/minimized */
 	UPROPERTY(EditAnywhere, config, Category = "Editor Performance", meta = (DisplayName = "Disable CPU throttling during graph execution"))
 	bool bDisableCPUThrottlingDuringGraphExecution = false;
+
+	UPROPERTY(EditAnywhere, config, Category = "Editor Performance", meta = (DisplayName = "Show PCG pause button in the editor (requires restart)"))
+	bool bShowPauseButton = false;
+
+	UPROPERTY(EditAnywhere, config, Category = "Editor Performance", meta = (DisplayName = "Unpausing cancels all PCG tasks"))
+	bool bUnpauseCancelsAll = true;
 };
 
 #if UE_ENABLE_INCLUDE_ORDER_DEPRECATED_IN_5_2
