[PCG] Grid Descriptor
- There are different properties defining where we output the PCG generated resources (GridSize, 2D or 3D grid, runtime or not and future usages: DataLayers, execution domains, unique grids)
- Introduce FPCGGridDescriptor to generalize this concept and unify code paths / reduce parameters
- Make sure PCGPartitionActors are not using UAIDs so that  2 users can't create the same GridCell actor
- Use new name for new actors so that they don't clash with existing ones
- Move bUse2DGrid to the PCG Graph and deprecate existing data
- Remove GridGuid usage (might be used later for unique grids)

#rnx

diff --git a/Engine/Plugins/PCG/Source/PCG/Private/Graph/PCGGraphExecutor.cpp b/Engine/Plugins/PCG/Source/PCG/Private/Graph/PCGGraphExecutor.cpp
--- a/Engine/Plugins/PCG/Source/PCG/Private/Graph/PCGGraphExecutor.cpp
+++ b/Engine/Plugins/PCG/Source/PCG/Private/Graph/PCGGraphExecutor.cpp
@@ -1995,16 +1995,18 @@ namespace PCGGraphExecutor
 			}
 			else
 			{
-				const APCGWorldActor* PCGWorldActor = Subsystem->GetPCGWorldActor();
 				const AActor* ComponentActor = InContext->SourceComponent->GetOwner();
-				if (PCGWorldActor && ComponentActor)
+				if (ComponentActor)
 				{
+					FPCGGridDescriptor GridDescriptor = FPCGGridDescriptor()
+						.SetGridSize(FromGridSize)
+						.SetIs2DGrid(InContext->SourceComponent->GetOriginalComponent()->Use2DGrid())
+						.SetIsRuntime(InContext->SourceComponent->IsManagedByRuntimeGenSystem());
+
 					// Get grid coords using the parent grid (FromGridSize).
-					const FIntVector CellCoords = UPCGActorHelpers::GetCellCoord(ComponentActor->GetActorLocation(), FromGridSize, PCGWorldActor->bUse2DGrid);
+					const FIntVector CellCoords = UPCGActorHelpers::GetCellCoord(ComponentActor->GetActorLocation(), GridDescriptor.GetGridSize(), GridDescriptor.Is2DGrid());
 
-					// Search for a transient local component if the source component is runtime managed.
-					const bool bTransientComponent = InContext->SourceComponent->IsManagedByRuntimeGenSystem();
-					ComponentWithData = Subsystem->GetLocalComponent(FromGridSize, CellCoords, InContext->SourceComponent->GetOriginalComponent(), bTransientComponent);
+					ComponentWithData = Subsystem->GetLocalComponent(GridDescriptor, CellCoords, InContext->SourceComponent->GetOriginalComponent());
 				}
 			}
 
diff --git a/Engine/Plugins/PCG/Source/PCG/Private/Grid/PCGGridDescriptor.cpp b/Engine/Plugins/PCG/Source/PCG/Private/Grid/PCGGridDescriptor.cpp
new file mode 100644
--- /dev/null
+++ b/Engine/Plugins/PCG/Source/PCG/Private/Grid/PCGGridDescriptor.cpp
@@ -0,0 +1,29 @@
+// Copyright Epic Games, Inc. All Rights Reserved.
+
+#include "Grid/PCGGridDescriptor.h"
+#include "Serialization/ArchiveCrc32.h"
+
+bool FPCGGridDescriptor::operator==(const FPCGGridDescriptor& Other) const
+{
+	return GridSize == Other.GridSize && bIs2DGrid == Other.bIs2DGrid && bIsRuntime == Other.bIsRuntime;
+}
+
+uint32 FPCGGridDescriptor::ComputeHash() const
+{
+	FArchiveCrc32 Ar;
+	FPCGGridDescriptor* NonConstThis = const_cast<FPCGGridDescriptor*>(this);
+
+	Ar << NonConstThis->GridSize << NonConstThis->bIs2DGrid << NonConstThis->bIsRuntime;
+
+	return Ar.GetCrc();
+}
+
+bool FPCGGridCellDescriptor::operator==(const FPCGGridCellDescriptor& InOther) const
+{
+	return Descriptor == InOther.Descriptor && GridCoords == InOther.GridCoords;
+}
+
+uint32 GetTypeHash(const FPCGGridCellDescriptor& In)
+{
+	return HashCombine(GetTypeHash(In.Descriptor), GetTypeHash(In.GridCoords));
+}
diff --git a/Engine/Plugins/PCG/Source/PCG/Private/Grid/PCGPartitionActor.cpp b/Engine/Plugins/PCG/Source/PCG/Private/Grid/PCGPartitionActor.cpp
--- a/Engine/Plugins/PCG/Source/PCG/Private/Grid/PCGPartitionActor.cpp
+++ b/Engine/Plugins/PCG/Source/PCG/Private/Grid/PCGPartitionActor.cpp
@@ -11,6 +11,7 @@
 
 #include "Components/BoxComponent.h"
 #include "Engine/World.h"
+#include "UObject/FortniteMainBranchObjectVersion.h"
 
 #if WITH_EDITOR
 #include "Grid/PCGPartitionActorDesc.h"
@@ -71,6 +72,14 @@ void APCGPartitionActor::PostLoad()
 	}
 #endif
 
+#if WITH_EDITORONLY_DATA
+	// Prior to this version bUse2DGrid was slave of the PCGWorldActor so make sure we update it one last time upon registration
+	if (GetLinkerCustomVersion(FFortniteMainBranchObjectVersion::GUID) < FFortniteMainBranchObjectVersion::PCGGridDescriptor)
+	{
+		bRequiresUse2DGridFixup = true;
+	}
+#endif
+
 	// Safe guard if we ever load a local that was deleted but not removed (like if the user deleted themselves the component)
 	// We can have multiple nullptr entries (if multiple components were removed), hence the while.
 	while (LocalToOriginal.Contains(nullptr))
@@ -262,6 +271,15 @@ bool APCGPartitionActor::IsUserManaged() const
 	return Super::IsUserManaged();
 }
 
+void APCGPartitionActor::UpdateUse2DGridIfNeeded(bool bInUse2DGrid)
+{
+	if (bRequiresUse2DGridFixup)
+	{
+		bUse2DGrid = bInUse2DGrid;
+		bRequiresUse2DGridFixup = false;
+	}
+}
+
 void APCGPartitionActor::SetInvalidForPCG()
 {
 	if (!bIsInvalidForPCG)
@@ -294,6 +312,14 @@ FIntVector APCGPartitionActor::GetGridCoord() const
 	return UPCGActorHelpers::GetCellCoord(Center, PCGGridSize, bUse2DGrid);
 }
 
+FPCGGridDescriptor APCGPartitionActor::GetGridDescriptor() const
+{
+	return FPCGGridDescriptor()
+		.SetGridSize(GetPCGGridSize())
+		.SetIs2DGrid(bUse2DGrid)
+		.SetIsRuntime(IsRuntimeGenerated());
+}
+
 bool APCGPartitionActor::Teleport(const FVector& NewLocation)
 {
 	// We should not be teleporting a PA that is in use. We only teleport empty RuntimeGen PAs to their grid cell before initialization.
@@ -547,22 +573,13 @@ AActor* APCGPartitionActor::GetSceneOutlinerParent() const
 }
 #endif // WITH_EDITOR
 
-void APCGPartitionActor::PostCreation(const FGuid& InGridGUID, uint32 InGridSize)
+void APCGPartitionActor::PostCreation(const FPCGGridDescriptor& GridDescriptor)
 {
-	PCGGuid = InGridGUID;
-	PCGGridSize = InGridSize;
-
-	// Put in cache if we use the 2D grid or not.
-	if (APCGWorldActor* PCGActor = PCGHelpers::GetPCGWorldActor(GetWorld()))
-	{
-		bUse2DGrid = PCGActor->bUse2DGrid;
-	}
-	else
-	{
-		bUse2DGrid = true;
-	}
+	PCGGridSize = GridDescriptor.GetGridSize();
+	bUse2DGrid = GridDescriptor.Is2DGrid();
 
 #if WITH_EDITOR
+	SetGridSize(PCGGridSize);
 	UpdateBoundsComponentExtents();
 #endif // WITH_EDITOR
 
@@ -660,20 +677,36 @@ void APCGPartitionActor::UpdateBoundsComponentExtents()
 #endif // WITH_EDITOR
 
 FString APCGPartitionActor::GetPCGPartitionActorName(uint32 GridSize, const FIntVector& GridCoords, bool bRuntimeGenerated)
+{
+	FPCGGridDescriptor GridDescriptor = FPCGGridDescriptor()
+		.SetGridSize(GridSize)
+		.SetIsRuntime(bRuntimeGenerated);
+	return GetPCGPartitionActorName(GridDescriptor, GridCoords);
+}
+
+FString APCGPartitionActor::GetPCGPartitionActorName(const FPCGGridDescriptor& GridDescriptor, const FIntVector& GridCoords)
 {
 	TStringBuilderWithBuffer<TCHAR, NAME_SIZE> ActorNameBuilder;
 
-	if (bRuntimeGenerated)
+	if (GridDescriptor.IsRuntime())
 	{
-		ActorNameBuilder += TEXT("PCGRuntimeGenPartitionActor_");
+		ActorNameBuilder += TEXT("PCGRuntimePartitionGridActor_");
 	}
 	else
 	{
-		ActorNameBuilder += TEXT("PCGPartitionActor_");
+		ActorNameBuilder += TEXT("PCGPartitionGridActor_");
 	}
 
-	ActorNameBuilder += FString::Printf(TEXT("%d_"), GridSize);
-	ActorNameBuilder += FString::Printf(TEXT("%d_%d_%d"), GridCoords.X, GridCoords.Y, GridCoords.Z);
+	ActorNameBuilder += FString::Printf(TEXT("%d_"), GridDescriptor.GetGridSize());
+	
+	if (GridDescriptor.Is2DGrid())
+	{
+		ActorNameBuilder += FString::Printf(TEXT("%d_%d"), GridCoords.X, GridCoords.Y);
+	}
+	else
+	{
+		ActorNameBuilder += FString::Printf(TEXT("%d_%d_%d"), GridCoords.X, GridCoords.Y, GridCoords.Z);
+	}
 
 	return ActorNameBuilder.ToString();
 }
\ No newline at end of file
diff --git a/Engine/Plugins/PCG/Source/PCG/Private/Grid/PCGPartitionActorDesc.cpp b/Engine/Plugins/PCG/Source/PCG/Private/Grid/PCGPartitionActorDesc.cpp
--- a/Engine/Plugins/PCG/Source/PCG/Private/Grid/PCGPartitionActorDesc.cpp
+++ b/Engine/Plugins/PCG/Source/PCG/Private/Grid/PCGPartitionActorDesc.cpp
@@ -3,11 +3,29 @@
 #include "Grid/PCGPartitionActorDesc.h"
 
 #if WITH_EDITOR
+#include "Grid/PCGPartitionActor.h"
 #include "UObject/FortniteMainBranchObjectVersion.h"
+#include "WorldPartition/WorldPartitionActorDescArchive.h"
 
 void FPCGPartitionActorDesc::Init(const AActor* InActor)
 {
 	FPartitionActorDesc::Init(InActor);
+
+	const APCGPartitionActor* PartitionActor = CastChecked<APCGPartitionActor>(InActor);
+	bUse2DGrid = PartitionActor->IsUsing2DGrid();
+}
+
+FPCGGridDescriptor FPCGPartitionActorDesc::GetGridDescriptor() const
+{
+	return FPCGGridDescriptor()
+		.SetGridSize(GridSize)
+		.SetIs2DGrid(bUse2DGrid)
+		.SetIsRuntime(false);
+}
+
+FIntVector FPCGPartitionActorDesc::GetGridCoord() const
+{
+	return FIntVector(GridIndexX, GridIndexY, GridIndexZ);
 }
 
 void FPCGPartitionActorDesc::Serialize(FArchive& Ar)
@@ -15,36 +33,51 @@ void FPCGPartitionActorDesc::Serialize(FArchive& Ar)
 	Ar.UsingCustomVersion(FFortniteMainBranchObjectVersion::GUID);
 	FPartitionActorDesc::Serialize(Ar);
 
-	if (Ar.IsLoading() && Ar.CustomVer(FFortniteMainBranchObjectVersion::GUID) < FFortniteMainBranchObjectVersion::PCGPartitionActorDesc)
+	if (Ar.IsLoading())
 	{
-		// Try to get Grid Size and Coords from names with following pattern "PCGPartitionActor_{GridSize}_{X}_{Y}_UAID"
-		// 
-		// Naming was broken in 5.4 where "_{Z}" part of the SpawnParams.Name was converted into the FName Number part
-		// in combination with the SpawnParams.NameMode being FActorSpawnParameters::ESpawnActorNameMode::Requested which ends up using FActorSpawnUtils::GetBaseName
-		// which calls FName::GetPlainNameString() which removes the Number part (treats _{Z} as the FName number)
-		//
-		// The label was properly set so if we are dealing with a _UAID_ name then use the label to parse the GridSize/X/Y/Z
-		const FString ActorNameStr = GetActorName().ToString();
-		if (ActorNameStr.Contains("_UAID_"))
+		if (Ar.CustomVer(FFortniteMainBranchObjectVersion::GUID) < FFortniteMainBranchObjectVersion::PCGPartitionActorDesc)
 		{
-			const FString ActorLabelStr = GetActorLabel().ToString();
-			TArray<FString> Tokens;
-			ActorLabelStr.ParseIntoArray(Tokens, TEXT("_"));
-			if (Tokens[0] == TEXT("PCGPartitionActor") && Tokens.Num() == 5)
-			{
-				TTypeFromString<uint32>::FromString(GridSize, *Tokens[1]);
-				TTypeFromString<int64>::FromString(GridIndexX, *Tokens[2]);
-				TTypeFromString<int64>::FromString(GridIndexY, *Tokens[3]);
-				TTypeFromString<int64>::FromString(GridIndexZ, *Tokens[4]);
-				
-				// Try to fixup the GridGuid in the PCGSubsystem using the PCGWorldActor data
-				bRequiresGuidFixup = true;
-			}
-			else
+			// Try to get Grid Size and Coords from names with following pattern "PCGPartitionActor_{GridSize}_{X}_{Y}_UAID"
+			// 
+			// Naming was broken in 5.4 where "_{Z}" part of the SpawnParams.Name was converted into the FName Number part
+			// in combination with the SpawnParams.NameMode being FActorSpawnParameters::ESpawnActorNameMode::Requested which ends up using FActorSpawnUtils::GetBaseName
+			// which calls FName::GetPlainNameString() which removes the Number part (treats _{Z} as the FName number)
+			//
+			// The label was properly set so if we are dealing with a _UAID_ name then use the label to parse the GridSize/X/Y/Z
+			const FString ActorNameStr = GetActorName().ToString();
+			if (ActorNameStr.Contains("_UAID_"))
 			{
-				bInvalid = true;
+				const FString ActorLabelStr = GetActorLabel().ToString();
+				TArray<FString> Tokens;
+				ActorLabelStr.ParseIntoArray(Tokens, TEXT("_"));
+				if (Tokens[0] == TEXT("PCGPartitionActor") && Tokens.Num() == 5)
+				{
+					TTypeFromString<uint32>::FromString(GridSize, *Tokens[1]);
+					TTypeFromString<int64>::FromString(GridIndexX, *Tokens[2]);
+					TTypeFromString<int64>::FromString(GridIndexY, *Tokens[3]);
+					TTypeFromString<int64>::FromString(GridIndexZ, *Tokens[4]);
+				}
+				else
+				{
+					bInvalid = true;
+				}
 			}
 		}
+
+		// For version prior, patch the bUse2DGrid when initializing PCG Subsystem
+		if (Ar.CustomVer(FFortniteMainBranchObjectVersion::GUID) < FFortniteMainBranchObjectVersion::PCGGridDescriptor)
+		{
+			bRequiresUse2DGridFixup = true;
+		}
+	}
+
+	if (!bIsDefaultActorDesc)
+	{
+		// From this version on, serialize bUse2DGrid
+		if (Ar.CustomVer(FFortniteMainBranchObjectVersion::GUID) >= FFortniteMainBranchObjectVersion::PCGGridDescriptor)
+		{
+			Ar << bUse2DGrid;
+		}
 	}
 }
 
diff --git a/Engine/Plugins/PCG/Source/PCG/Private/PCGActorAndComponentMapping.cpp b/Engine/Plugins/PCG/Source/PCG/Private/PCGActorAndComponentMapping.cpp
--- a/Engine/Plugins/PCG/Source/PCG/Private/PCGActorAndComponentMapping.cpp
+++ b/Engine/Plugins/PCG/Source/PCG/Private/PCGActorAndComponentMapping.cpp
@@ -261,11 +261,9 @@ TArray<FPCGTaskId> FPCGActorAndComponentMapping::DispatchToRegisteredLocalCompon
 
 	const bool bIsRuntimeGenerated = OriginalComponent->IsManagedByRuntimeGenSystem();
 
-	const TMap<const UPCGComponent*, TSet<TObjectPtr<APCGPartitionActor>>>& Map = bIsRuntimeGenerated ? ComponentToRuntimeGenPartitionActorsMap : ComponentToPartitionActorsMap;
-
 	// TODO: Might be more interesting to copy the set and release the lock.
-	FReadScopeLock ReadLock(bIsRuntimeGenerated ? ComponentToRuntimeGenPartitionActorsMapLock : ComponentToPartitionActorsMapLock);
-	const TSet<TObjectPtr<APCGPartitionActor>>* PartitionActorsPtr = Map.Find(OriginalComponent);
+	FReadScopeLock ReadLock(ComponentToPartitionActorsMapLock);
+	const TSet<TObjectPtr<APCGPartitionActor>>* PartitionActorsPtr = ComponentToPartitionActorsMap.Find(OriginalComponent);
 
 	if (!PartitionActorsPtr)
 	{
@@ -447,7 +445,6 @@ bool FPCGActorAndComponentMapping::RemapPCGComponent(const UPCGComponent* OldCom
 	};
 
 	RemapPreviousInstances(ComponentToPartitionActorsMap, ComponentToPartitionActorsMapLock);
-	RemapPreviousInstances(ComponentToRuntimeGenPartitionActorsMap, ComponentToRuntimeGenPartitionActorsMapLock);
 
 	// And update the mapping if bounds changed and we want to do actor mapping
 	if (bBoundsChanged && NewComponent->IsPartitioned() && bDoActorMapping)
@@ -596,7 +593,7 @@ void FPCGActorAndComponentMapping::RegisterPartitionActor(APCGPartitionActor* In
 		UE_LOG(LogPCG, Warning, TEXT("[RegisterPartitionActor] Invalid PCG Partiton Actor '%s' (%s). Please delete actor to remove warning."), *InActor->GetName(), *InActor->GetPackage()->GetName());
 		return;
 	} // Invalidate actor if there is a duplicate for the same grid information
-	else if (FGuid* FoundGuid = PartitionActorRecords.Find({ InActor->GetGridGuid(), InActor->GetGridSize(), InActor->GetGridCoord() }); FoundGuid && *FoundGuid != InActor->GetActorGuid())
+	else if (FGuid* FoundGuid = PartitionActorRecords.Find({ InActor->GetGridDescriptor(), InActor->GetGridCoord()}); FoundGuid && *FoundGuid != InActor->GetActorGuid())
 	{
 		InActor->SetInvalidForPCG();
 		UE_LOG(LogPCG, Warning, TEXT("[RegisterPartitionActor] Duplicate PCG Partition Actor '%s' (%s). Please delete actor to remove warning."), *InActor->GetName(), *InActor->GetPackage()->GetName());
@@ -604,18 +601,15 @@ void FPCGActorAndComponentMapping::RegisterPartitionActor(APCGPartitionActor* In
 	}
 #endif
 
-	const uint32 GridSize = InActor->GetPCGGridSize();
+	const FPCGGridDescriptor Descriptor = InActor->GetGridDescriptor();
 	const FIntVector GridCoord = InActor->GetGridCoord();
 
-	check(GridSize > 0);
-
-	const bool bIsRuntimeGenerated = InActor->IsRuntimeGenerated();
-	TMap<uint32, TMap<FIntVector, TObjectPtr<APCGPartitionActor>>>& ActorsMap = bIsRuntimeGenerated ? RuntimeGenPartitionActorsMap : PartitionActorsMap;
+	check(Descriptor.GetGridSize() > 0);
 
 	{
-		FWriteScopeLock WriteLock(bIsRuntimeGenerated ? RuntimeGenPartitionActorsMapLock : PartitionActorsMapLock);
+		FWriteScopeLock WriteLock(PartitionActorsMapLock);
 
-		TMap<FIntVector, TObjectPtr<APCGPartitionActor>>& PartitionActorsMapGrid = ActorsMap.FindOrAdd(GridSize);
+		TMap<FIntVector, TObjectPtr<APCGPartitionActor>>& PartitionActorsMapGrid = PartitionActorsMap.FindOrAdd(Descriptor);
 		if (PartitionActorsMapGrid.Contains(GridCoord))
 		{
 			return;
@@ -627,15 +621,15 @@ void FPCGActorAndComponentMapping::RegisterPartitionActor(APCGPartitionActor* In
 	APCGWorldActor* WorldActor = PCGSubsystem->GetPCGWorldActor();
 	check(WorldActor);
 
-	// For deprecration: bUse2DGrid is now true by default. But if we already have Partition Actors that were created when the flag was false by default,
-	// we keep this flag
-	if (WorldActor->bUse2DGrid != InActor->IsUsing2DGrid())
-	{
-		WorldActor->bUse2DGrid = InActor->IsUsing2DGrid();
-	}
+#if WITH_EDITOR
+PRAGMA_DISABLE_DEPRECATION_WARNINGS
+	// Handle move of the Use2DGrid flag for old partition actors that haven't been resaved
+	InActor->UpdateUse2DGridIfNeeded(WorldActor->bUse2DGrid);
+PRAGMA_ENABLE_DEPRECATION_WARNINGS
+#endif // WITH_EDITOR
 
 	// Register to all the components that intersect with the PA. Ignore for runtime generated, it is handled manually
-	if (!bIsRuntimeGenerated)
+	if (!Descriptor.IsRuntime())
 	{
 		FWriteScopeLock WriteLock(ComponentToPartitionActorsMapLock);
 		ForAllIntersectingPartitionedComponents(FBoxCenterAndExtent(InActor->GetFixedBounds()), [this, InActor, bDoComponentMapping, WorldActor](UPCGComponent* Component)
@@ -683,23 +677,20 @@ void FPCGActorAndComponentMapping::UnregisterPartitionActor(APCGPartitionActor*
 #endif
 
 	const FIntVector GridCoord = Actor->GetGridCoord();
-	const uint32 GridSize = Actor->GetPCGGridSize();
-	if (!ensure(GridSize > 0))
+	const FPCGGridDescriptor GridDescriptor = Actor->GetGridDescriptor();
+	if (!ensure(GridDescriptor.GetGridSize() > 0))
 	{
 		return;
 	}
 
-	const bool bIsRuntimeGenerated = Actor->IsRuntimeGenerated();
-	TMap<uint32, TMap<FIntVector, TObjectPtr<APCGPartitionActor>>>& ActorsMap = bIsRuntimeGenerated ? RuntimeGenPartitionActorsMap : PartitionActorsMap;
-
-	if (TMap<FIntVector, TObjectPtr<APCGPartitionActor>>* PartitionActorsMapGrid = ActorsMap.Find(GridSize))
+	if (TMap<FIntVector, TObjectPtr<APCGPartitionActor>>* PartitionActorsMapGrid = PartitionActorsMap.Find(GridDescriptor))
 	{
-		FWriteScopeLock WriteLock(bIsRuntimeGenerated ? RuntimeGenPartitionActorsMapLock : PartitionActorsMapLock);
+		FWriteScopeLock WriteLock(PartitionActorsMapLock);
 		PartitionActorsMapGrid->Remove(GridCoord);
 	}
 
 	// Unregister from all intersecting components. Ignore for runtime generated, it is handled manually
-	if (!bIsRuntimeGenerated)
+	if (!GridDescriptor.IsRuntime())
 	{
 		FWriteScopeLock WriteLock(ComponentToPartitionActorsMapLock);
 		ForAllIntersectingPartitionedComponents(FBoxCenterAndExtent(Actor->GetFixedBounds()), [this, Actor](UPCGComponent* Component)
@@ -715,28 +706,24 @@ void FPCGActorAndComponentMapping::UnregisterPartitionActor(APCGPartitionActor*
 
 void FPCGActorAndComponentMapping::ForAllIntersectingPartitionActors(const FBox& InBounds, TFunctionRef<void(APCGPartitionActor*)> InFunc) const
 {
-	// No PCGWorldActor just early out. Same for invalid bounds.
-	APCGWorldActor* PCGWorldActor = PCGSubsystem->GetPCGWorldActor();
-
-	if (!PCGWorldActor || !InBounds.IsValid)
+	if (!InBounds.IsValid)
 	{
 		return;
 	}
 
-	auto ForAllIntersectingPartitionActorsOfGridSize = [InFunc, PCGWorldActor, &InBounds](const TMap<uint32, TMap<FIntVector, TObjectPtr<APCGPartitionActor>>>& Map, FRWLock& Lock, uint32 GridSize)
+	auto ForAllIntersectingPartitionActorsOfGridSize = [InFunc, &InBounds, this](const FPCGGridDescriptor& GridDescriptor)
 	{
-		const bool bUse2DGrid = PCGWorldActor->bUse2DGrid;
-		FIntVector MinCellCoords = UPCGActorHelpers::GetCellCoord(InBounds.Min, GridSize, bUse2DGrid);
-		FIntVector MaxCellCoords = UPCGActorHelpers::GetCellCoord(InBounds.Max, GridSize, bUse2DGrid);
+		FReadScopeLock ReadLock(PartitionActorsMapLock);
 
-		FReadScopeLock ReadLock(Lock);
-
-		const TMap<FIntVector, TObjectPtr<APCGPartitionActor>>* PartitionActorsMapGrid = Map.Find(GridSize);
+		const TMap<FIntVector, TObjectPtr<APCGPartitionActor>>* PartitionActorsMapGrid = PartitionActorsMap.Find(GridDescriptor);
 		if (!PartitionActorsMapGrid || PartitionActorsMapGrid->IsEmpty())
 		{
 			return;
 		}
 
+		const FIntVector MinCellCoords = UPCGActorHelpers::GetCellCoord(InBounds.Min, GridDescriptor.GetGridSize(), GridDescriptor.Is2DGrid());
+		const FIntVector MaxCellCoords = UPCGActorHelpers::GetCellCoord(InBounds.Max, GridDescriptor.GetGridSize(), GridDescriptor.Is2DGrid());
+
 		for (int32 z = MinCellCoords.Z; z <= MaxCellCoords.Z; z++)
 		{
 			for (int32 y = MinCellCoords.Y; y <= MaxCellCoords.Y; y++)
@@ -756,20 +743,16 @@ void FPCGActorAndComponentMapping::ForAllIntersectingPartitionActors(const FBox&
 		}
 	};
 
-	PCGHiGenGrid::FSizeToGuidMap GridSizeToGuid;
-	PCGWorldActor->GetSerializedGridGuids(GridSizeToGuid);
-	for (const TPair<uint32, FGuid>& SizeAndGuid : GridSizeToGuid)
+	TArray<FPCGGridDescriptor> GridDescriptors;
 	{
-		const uint32 GridSize = SizeAndGuid.Key;
-		ForAllIntersectingPartitionActorsOfGridSize(PartitionActorsMap, PartitionActorsMapLock, GridSize);
+		FReadScopeLock ReadLock(PartitionActorsMapLock);
+
+		PartitionActorsMap.GenerateKeyArray(GridDescriptors);
 	}
 
-	GridSizeToGuid.Empty();
-	PCGWorldActor->GetTransientGridGuids(GridSizeToGuid);
-	for (const TPair<uint32, FGuid>& SizeAndGuid : GridSizeToGuid)
+	for (const FPCGGridDescriptor& GridDescriptor : GridDescriptors)
 	{
-		const uint32 GridSize = SizeAndGuid.Key;
-		ForAllIntersectingPartitionActorsOfGridSize(RuntimeGenPartitionActorsMap, RuntimeGenPartitionActorsMapLock, GridSize);
+		ForAllIntersectingPartitionActorsOfGridSize(GridDescriptor);
 	}
 }
 
@@ -795,38 +778,40 @@ void FPCGActorAndComponentMapping::UpdateMappingPCGComponentPartitionActor(UPCGC
 		bool bHasUnbounded = false;
 		PCGHiGenGrid::FSizeArray GridSizes;
 		ensure(PCGHelpers::GetGenerationGridSizes(InComponent->GetGraph(), WorldActor, GridSizes, bHasUnbounded));
+				
+		const bool bUse2DGrid = InComponent->Use2DGrid();
+		const bool bIsRuntimeGenerated = InComponent->IsManagedByRuntimeGenSystem();
+		TSet<FPCGGridDescriptor> GridDescriptors;
+		Algo::Transform(GridSizes, GridDescriptors, [bUse2DGrid, bIsRuntimeGenerated](uint32 GridSize) 
+		{ 
+			return FPCGGridDescriptor().SetGridSize(GridSize).SetIs2DGrid(bUse2DGrid).SetIsRuntime(bIsRuntimeGenerated);
+		});
 
-		auto UpdateMapping = [this, InComponent, &Bounds, WorldActor, &GridSizes](bool bIsRuntimeGenerated)
+		auto UpdateMapping = [this, InComponent, &Bounds, WorldActor, &GridDescriptors]()
 		{
-			TMap<const UPCGComponent*, TSet<TObjectPtr<APCGPartitionActor>>>& Map = bIsRuntimeGenerated ? ComponentToRuntimeGenPartitionActorsMap : ComponentToPartitionActorsMap;
 			TSet<TObjectPtr<APCGPartitionActor>> RemovedActors;
 
 			{
-				FWriteScopeLock WriteLock(bIsRuntimeGenerated ? ComponentToRuntimeGenPartitionActorsMapLock : ComponentToPartitionActorsMapLock);
-				TSet<TObjectPtr<APCGPartitionActor>>* PartitionActorsPtr = Map.Find(InComponent);
+				FWriteScopeLock WriteLock(ComponentToPartitionActorsMapLock);
+				TSet<TObjectPtr<APCGPartitionActor>>* PartitionActorsPtr = ComponentToPartitionActorsMap.Find(InComponent);
 
 				if (!PartitionActorsPtr)
 				{
 					// Does not yet exists, add it
-					PartitionActorsPtr = &Map.Emplace(InComponent);
+					PartitionActorsPtr = &ComponentToPartitionActorsMap.Emplace(InComponent);
 					check(PartitionActorsPtr);
 				}
 
 				TSet<TObjectPtr<APCGPartitionActor>> NewMapping;
-				ForAllIntersectingPartitionActors(Bounds, [&NewMapping, InComponent, WorldActor, &GridSizes, bIsRuntimeGenerated](APCGPartitionActor* Actor)
+				ForAllIntersectingPartitionActors(Bounds, [&NewMapping, InComponent, WorldActor, &GridDescriptors](APCGPartitionActor* Actor)
 				{
 					if (!Actor)
 					{
 						return;
 					}
 
-					// TODO: This will need to be revisited when execution domains become a broader concept.
-					const bool bSameDomain = Actor->IsRuntimeGenerated() == bIsRuntimeGenerated;
-
-					// Only add a graph instance to partition actors that are:
-					// * In the same execution domain as the original component.
-					// * On a valid grid for the original component.
-					if (bSameDomain && GridSizes.Contains(Actor->GetPCGGridSize()))
+					// Only add a graph instance to partition actors with compatible Grid Descriptors
+					if (GridDescriptors.Contains(Actor->GetGridDescriptor()))
 					{
 						Actor->AddGraphInstance(InComponent);
 						NewMapping.Add(Actor);
@@ -849,7 +834,7 @@ void FPCGActorAndComponentMapping::UpdateMappingPCGComponentPartitionActor(UPCGC
 			}
 		};
 
-		UpdateMapping(InComponent->IsManagedByRuntimeGenSystem());
+		UpdateMapping();
 	}
 }
 
@@ -878,7 +863,6 @@ void FPCGActorAndComponentMapping::DeleteMappingPCGComponentPartitionActor(UPCGC
 	};
 
 	DeleteMapping(ComponentToPartitionActorsMap, ComponentToPartitionActorsMapLock);
-	DeleteMapping(ComponentToRuntimeGenPartitionActorsMap, ComponentToRuntimeGenPartitionActorsMapLock);
 }
 
 bool FPCGActorAndComponentMapping::IsComponentRegistered(const UPCGComponent* InComponent) const
@@ -924,12 +908,11 @@ TSet<UPCGComponent*> FPCGActorAndComponentMapping::GetAllRegisteredComponents()
 	return Res;
 }
 
-UPCGComponent* FPCGActorAndComponentMapping::GetLocalComponent(uint32 GridSize, const FIntVector& CellCoords, const UPCGComponent* InOriginalComponent, bool bRuntimeGenerated) const
-{
-	const TMap<uint32, TMap<FIntVector, TObjectPtr<APCGPartitionActor>>>& ActorsMap = bRuntimeGenerated ? RuntimeGenPartitionActorsMap : PartitionActorsMap;
-	FReadScopeLock ReadLock(bRuntimeGenerated ? RuntimeGenPartitionActorsMapLock : PartitionActorsMapLock);
+UPCGComponent* FPCGActorAndComponentMapping::GetLocalComponent(const FPCGGridDescriptor& GridDescriptor, const FIntVector& CellCoords, const UPCGComponent* InOriginalComponent) const
+{	
+	FReadScopeLock ReadLock(PartitionActorsMapLock);
 
-	if (const TMap<FIntVector, TObjectPtr<APCGPartitionActor>>* PartitionActorsOnGrid = ActorsMap.Find(GridSize))
+	if (const TMap<FIntVector, TObjectPtr<APCGPartitionActor>>* PartitionActorsOnGrid = PartitionActorsMap.Find(GridDescriptor))
 	{
 		const TObjectPtr<APCGPartitionActor>* PartitionActor = PartitionActorsOnGrid->Find(CellCoords);
 		if (PartitionActor && *PartitionActor)
@@ -941,12 +924,11 @@ UPCGComponent* FPCGActorAndComponentMapping::GetLocalComponent(uint32 GridSize,
 	return nullptr;
 }
 
-APCGPartitionActor* FPCGActorAndComponentMapping::GetPartitionActor(uint32 GridSize, const FIntVector& CellCoords, bool bRuntimeGenerated) const
+APCGPartitionActor* FPCGActorAndComponentMapping::GetPartitionActor(const FPCGGridDescriptor& GridDescriptor, const FIntVector& CellCoords) const
 {
-	const TMap<uint32, TMap<FIntVector, TObjectPtr<APCGPartitionActor>>>& Map = bRuntimeGenerated ? RuntimeGenPartitionActorsMap : PartitionActorsMap;
-	FReadScopeLock ReadLock(bRuntimeGenerated ? RuntimeGenPartitionActorsMapLock : PartitionActorsMapLock);
+	FReadScopeLock ReadLock(PartitionActorsMapLock);
 
-	if (const TMap<FIntVector, TObjectPtr<APCGPartitionActor>>* PartitionActorsOnGrid = Map.Find(GridSize))
+	if (const TMap<FIntVector, TObjectPtr<APCGPartitionActor>>* PartitionActorsOnGrid = PartitionActorsMap.Find(GridDescriptor))
 	{
 		const TObjectPtr<APCGPartitionActor>* PartitionActor = PartitionActorsOnGrid->Find(CellCoords);
 		return PartitionActor ? *PartitionActor : nullptr;
@@ -955,11 +937,11 @@ APCGPartitionActor* FPCGActorAndComponentMapping::GetPartitionActor(uint32 GridS
 	return nullptr;
 }
 
-bool FPCGActorAndComponentMapping::DoesPartitionActorRecordExist(const FGuid& GridGuid, uint32 GridSize, const FIntVector& GridCoords) const
+bool FPCGActorAndComponentMapping::DoesPartitionActorRecordExist(const FPCGGridDescriptor& GridDescriptor, const FIntVector& GridCoords) const
 {
 #if WITH_EDITOR
 	// In Editor this will not be empty (including -game)
-	if (PartitionActorRecords.Contains({ GridGuid, GridSize, GridCoords }))
+	if (PartitionActorRecords.Contains({ GridDescriptor, GridCoords }))
 	{
 		return true;
 	}
@@ -970,7 +952,7 @@ bool FPCGActorAndComponentMapping::DoesPartitionActorRecordExist(const FGuid& Gr
 	// which equates to a different overall behavior (one will be streamed-in/out by WP and other will be spawned and never unspawn). So we need to revisit this as it is not a consistent behavior.
 	if (APCGWorldActor* PCGWorldActor = PCGSubsystem->GetPCGWorldActor())
 	{
-		return PCGWorldActor->RuntimePartitionActorRecords.Contains({ GridGuid, GridSize, GridCoords });
+		return PCGWorldActor->RuntimePartitionActorRecords.Contains({ GridDescriptor, GridCoords });
 	}
 	
 	return false;
@@ -1401,7 +1383,7 @@ void FPCGActorAndComponentMapping::OnActorDescInstanceAdded(FWorldPartitionActor
 	if (InActorDescInstance->GetActorNativeClass()->IsChildOf(APCGPartitionActor::StaticClass()))
 	{
 		const FPCGPartitionActorDesc* ActorDesc = static_cast<const FPCGPartitionActorDesc*>(InActorDescInstance->GetActorDesc());
-		FPCGPartitionActorRecord PartitionActorRecord = { ActorDesc->GridGuid, ActorDesc->GridSize, FIntVector3(ActorDesc->GridIndexX, ActorDesc->GridIndexY, ActorDesc->GridIndexZ) };
+		FPCGGridCellDescriptor PartitionActorRecord = { ActorDesc->GetGridDescriptor(), FIntVector3(ActorDesc->GridIndexX, ActorDesc->GridIndexY, ActorDesc->GridIndexZ)};
 		if (!PartitionActorRecords.Contains(PartitionActorRecord))
 		{
 			PartitionActorRecords.Add(MoveTemp(PartitionActorRecord), ActorDesc->GetGuid());
@@ -1414,7 +1396,7 @@ void FPCGActorAndComponentMapping::OnActorDescInstanceRemoved(FWorldPartitionAct
 	if (InActorDescInstance->GetActorNativeClass()->IsChildOf(APCGPartitionActor::StaticClass()))
 	{
 		const FPCGPartitionActorDesc* ActorDesc = static_cast<const FPCGPartitionActorDesc*>(InActorDescInstance->GetActorDesc());
-		FPCGPartitionActorRecord PartitionActorRecord = { ActorDesc->GridGuid, ActorDesc->GridSize, FIntVector3(ActorDesc->GridIndexX, ActorDesc->GridIndexY, ActorDesc->GridIndexZ) };
+		FPCGGridCellDescriptor PartitionActorRecord = { ActorDesc->GetGridDescriptor(), FIntVector3(ActorDesc->GridIndexX, ActorDesc->GridIndexY, ActorDesc->GridIndexZ)};
 		if (FGuid* FoundGuid = PartitionActorRecords.Find(PartitionActorRecord); FoundGuid && *FoundGuid == InActorDescInstance->GetGuid())
 		{
 			PartitionActorRecords.Remove(PartitionActorRecord);
@@ -1422,33 +1404,32 @@ void FPCGActorAndComponentMapping::OnActorDescInstanceRemoved(FWorldPartitionAct
 	}
 }
 
-void FPCGActorAndComponentMapping::BuildPartitionActorRecords(APCGWorldActor* PCGWorldActor, UWorldPartition* WorldPartition, TMap<FPCGPartitionActorRecord, FGuid>& OutPartitionActorRecords, TSet<FGuid>& OutInvalidPartitionActors)
+void FPCGActorAndComponentMapping::BuildPartitionActorRecords(APCGWorldActor* PCGWorldActor, UWorldPartition* WorldPartition, TMap<FPCGGridCellDescriptor, FGuid>& OutPartitionActorRecords, TSet<FGuid>& OutInvalidPartitionActors)
 {
 	check(WorldPartition);
 	
 	FWorldPartitionHelpers::ForEachActorDescInstance<APCGPartitionActor>(WorldPartition, [PCGWorldActor, &OutPartitionActorRecords, &OutInvalidPartitionActors](const FWorldPartitionActorDescInstance* ActorDescInstance)
 	{
 		const FPCGPartitionActorDesc* ActorDesc = static_cast<const FPCGPartitionActorDesc*>(ActorDescInstance->GetActorDesc());
-		if (ActorDesc->bRequiresGuidFixup && PCGWorldActor)
+		
+		// Fixup data
+		if (ActorDesc->bRequiresUse2DGridFixup && PCGWorldActor)
 		{
-			PCGHiGenGrid::FSizeToGuidMap SizeToGuidMap;
-			PCGWorldActor->GetSerializedGridGuids(SizeToGuidMap);
-			if (FGuid* FoundGuid = SizeToGuidMap.Find(ActorDesc->GridSize))
-			{
-				FPCGPartitionActorDesc* NonConstActorDesc = const_cast<FPCGPartitionActorDesc*>(ActorDesc);
-				NonConstActorDesc->GridGuid = *FoundGuid;
-				NonConstActorDesc->bRequiresGuidFixup = false;
-			}
+			FPCGPartitionActorDesc* NonConstActorDesc = const_cast<FPCGPartitionActorDesc*>(ActorDesc);
+PRAGMA_DISABLE_DEPRECATION_WARNINGS
+			NonConstActorDesc->bUse2DGrid = PCGWorldActor->bUse2DGrid;
+PRAGMA_ENABLE_DEPRECATION_WARNINGS
+			NonConstActorDesc->bRequiresUse2DGridFixup = false;
 		}
 
 		// Add to invalid list if we couldn't properly fixup or if it was already marked invalid
-		if (ActorDesc->bRequiresGuidFixup || ActorDesc->bInvalid)
+		if (ActorDesc->bInvalid)
 		{
 			OutInvalidPartitionActors.Add(ActorDesc->GetGuid());
 		}
 		else // Add to existing valid PCG Partition actors
 		{
-			FPCGPartitionActorRecord PartitionActorRecord = { ActorDesc->GridGuid, ActorDesc->GridSize, FIntVector3(ActorDesc->GridIndexX, ActorDesc->GridIndexY, ActorDesc->GridIndexZ) };
+			FPCGGridCellDescriptor PartitionActorRecord = { ActorDesc->GetGridDescriptor(), FIntVector3(ActorDesc->GridIndexX, ActorDesc->GridIndexY, ActorDesc->GridIndexZ)};
 			if (!OutPartitionActorRecords.Contains(PartitionActorRecord))
 			{
 				OutPartitionActorRecords.Add(MoveTemp(PartitionActorRecord), ActorDesc->GetGuid());
@@ -1479,9 +1460,9 @@ void FPCGActorAndComponentMapping::BuildPartitionActorRecords()
 				{
 					if (UPCGSubsystem* EditorPCGSubsystem = EditorWorld->GetSubsystem<UPCGSubsystem>())
 					{
-						TArray<FPCGPartitionActorRecord> SerializedPartitionActorRecords;
+						TArray<FPCGGridCellDescriptor> SerializedPartitionActorRecords;
 						EditorPCGSubsystem->ActorAndComponentMapping.PartitionActorRecords.GenerateKeyArray(SerializedPartitionActorRecords);
-						PCGWorldActor->RuntimePartitionActorRecords = TSet<FPCGPartitionActorRecord>(SerializedPartitionActorRecords);
+						PCGWorldActor->RuntimePartitionActorRecords = TSet<FPCGGridCellDescriptor>(SerializedPartitionActorRecords);
 					}
 				}
 			}
diff --git a/Engine/Plugins/PCG/Source/PCG/Private/PCGComponent.cpp b/Engine/Plugins/PCG/Source/PCG/Private/PCGComponent.cpp
--- a/Engine/Plugins/PCG/Source/PCG/Private/PCGComponent.cpp
+++ b/Engine/Plugins/PCG/Source/PCG/Private/PCGComponent.cpp
@@ -137,6 +137,12 @@ void UPCGComponent::SetIsPartitioned(bool bIsNowPartitioned)
 	}
 }
 
+bool UPCGComponent::Use2DGrid() const
+{
+	check(GetGraph());
+	return GetGraph()->Use2DGrid();
+}
+
 void UPCGComponent::SetGraph_Implementation(UPCGGraphInterface* InGraph)
 {
 	SetGraphInterfaceLocal(InGraph);
diff --git a/Engine/Plugins/PCG/Source/PCG/Private/PCGGraph.cpp b/Engine/Plugins/PCG/Source/PCG/Private/PCGGraph.cpp
--- a/Engine/Plugins/PCG/Source/PCG/Private/PCGGraph.cpp
+++ b/Engine/Plugins/PCG/Source/PCG/Private/PCGGraph.cpp
@@ -1449,7 +1449,8 @@ void UPCGGraph::PostEditChangeProperty(FPropertyChangedEvent& PropertyChangedEve
 		OnGraphParametersChanged(EPCGGraphParameterEvent::ValueModifiedLocally, MemberPropertyName);
 	}
 	else if (PropertyName == GET_MEMBER_NAME_CHECKED(UPCGGraph, HiGenGridSize)
-		|| PropertyName == GET_MEMBER_NAME_CHECKED(UPCGGraph, bUseHierarchicalGeneration))
+		|| PropertyName == GET_MEMBER_NAME_CHECKED(UPCGGraph, bUseHierarchicalGeneration)
+		|| PropertyName == GET_MEMBER_NAME_CHECKED(UPCGGraph, bUse2DGrid))
 	{
 		// The higen settings change the structure of the graph (presence or absence of links between grid levels).
 		NotifyGraphChanged(EPCGChangeType::Structural | EPCGChangeType::GenerationGrid);
diff --git a/Engine/Plugins/PCG/Source/PCG/Private/PCGSubsystem.cpp b/Engine/Plugins/PCG/Source/PCG/Private/PCGSubsystem.cpp
--- a/Engine/Plugins/PCG/Source/PCG/Private/PCGSubsystem.cpp
+++ b/Engine/Plugins/PCG/Source/PCG/Private/PCGSubsystem.cpp
@@ -556,7 +556,7 @@ FPCGTaskId UPCGSubsystem::ScheduleComponent(UPCGComponent* PCGComponent, EPCGHiG
 			// Note: This is an immediate operation, as we need the PA for the generation.
 			auto ScheduleTask = [](APCGPartitionActor* PCGActor, const FBox& InIntersectedBounds) { return InvalidPCGTaskId; };
 
-			ForAllOverlappingCells(PCGComponent->GetGridBounds(), GridSizes, /*bCanCreateActor=*/true, /*Dependencies=*/{}, ScheduleTask);
+			ForAllOverlappingCells(PCGComponent, PCGComponent->GetGridBounds(), GridSizes, /*bCanCreateActor=*/true, /*Dependencies=*/{}, ScheduleTask);
 		}
 
 		ActorAndComponentMapping.UpdateMappingPCGComponentPartitionActor(PCGComponent);
@@ -964,7 +964,7 @@ void UPCGSubsystem::ForAllOverlappingComponentsInHierarchy(UPCGComponent* InComp
 	});
 }
 
-FPCGTaskId UPCGSubsystem::ForAllOverlappingCells(const FBox& InBounds, const PCGHiGenGrid::FSizeArray& InGridSizes, bool bCanCreateActor, const TArray<FPCGTaskId>& Dependencies, TFunctionRef<FPCGTaskId(APCGPartitionActor*, const FBox&)> InFunc) const
+FPCGTaskId UPCGSubsystem::ForAllOverlappingCells(UPCGComponent* InComponent, const FBox& InBounds, const PCGHiGenGrid::FSizeArray& InGridSizes, bool bCanCreateActor, const TArray<FPCGTaskId>& Dependencies, TFunctionRef<FPCGTaskId(APCGPartitionActor*, const FBox&)> InFunc) const
 {
 	if (!GraphExecutor || !PCGWorldActor)
 	{
@@ -984,25 +984,17 @@ FPCGTaskId UPCGSubsystem::ForAllOverlappingCells(const FBox& InBounds, const PCG
 		return InvalidPCGTaskId;
 	}
 
-	PCGWorldActor->CreateGridGuidsIfNecessary(GridSizes, /*bAreGridsSerialized=*/true);
-
-	PCGHiGenGrid::FSizeToGuidMap GridSizeToGuid;
-	PCGWorldActor->GetSerializedGridGuids(GridSizeToGuid);
-	
 	TArray<FPCGTaskId> CellTasks;
 	for (uint32 GridSize : GridSizes)
 	{
-		const FGuid* GuidPtr = GridSizeToGuid.Find(GridSize);
-		if (!ensure(GuidPtr))
-		{
-			continue;
-		}
-
-		const FGuid GridGuid = *GuidPtr;
+		FPCGGridDescriptor Descriptor = FPCGGridDescriptor()
+			.SetGridSize(GridSize)
+			.SetIs2DGrid(InComponent->Use2DGrid())
+			.SetIsRuntime(false);
 
 		// In case of 2D grid, we are clamping our bounds in Z to be within 0 and GridSize to create a 2D grid instead of 3D.
 		FBox ModifiedInBounds = InBounds;
-		if (PCGWorldActor->bUse2DGrid)
+		if (Descriptor.Is2DGrid())
 		{
 			FVector MinBounds = InBounds.Min;
 			FVector MaxBounds = InBounds.Max;
@@ -1013,13 +1005,13 @@ FPCGTaskId UPCGSubsystem::ForAllOverlappingCells(const FBox& InBounds, const PCG
 		}
 
 		// Helper lambda to apply 'InFunc' to the specified grid cell. Only applies if the cell overlaps the modified bounds, and the actor can be found or created and is not unloaded.
-		auto ApplyOnCell = [this, &CellTasks, ModifiedInBounds, bCanCreateActor, &GridGuid, GridSize, &InFunc](const FIntVector& CellCoord, const FBox& CellBounds)
+		auto ApplyOnCell = [this, &CellTasks, ModifiedInBounds, bCanCreateActor, &Descriptor, &InFunc](const FIntVector& CellCoord, const FBox& CellBounds)
 		{
 			const FBox IntersectedBounds = ModifiedInBounds.Overlap(CellBounds);
 
 			if (IntersectedBounds.IsValid)
 			{
-				if (APCGPartitionActor* Actor = FindOrCreatePCGPartitionActor(GridGuid, GridSize, CellCoord, /*bRuntimeGenerated=*/false, bCanCreateActor))
+				if (APCGPartitionActor* Actor = FindOrCreatePCGPartitionActor(Descriptor, CellCoord, bCanCreateActor))
 				{
 					FPCGTaskId ExecuteTaskId = InFunc(Actor, IntersectedBounds);
 
@@ -1031,8 +1023,8 @@ FPCGTaskId UPCGSubsystem::ForAllOverlappingCells(const FBox& InBounds, const PCG
 			}
 		};
 
-		FIntVector MinCellCoords = UPCGActorHelpers::GetCellCoord(InBounds.Min, GridSize, PCGWorldActor->bUse2DGrid);
-		FIntVector MaxCellCoords = UPCGActorHelpers::GetCellCoord(InBounds.Max, GridSize, PCGWorldActor->bUse2DGrid);
+		FIntVector MinCellCoords = UPCGActorHelpers::GetCellCoord(InBounds.Min, GridSize, Descriptor.Is2DGrid());
+		FIntVector MaxCellCoords = UPCGActorHelpers::GetCellCoord(InBounds.Max, GridSize, Descriptor.Is2DGrid());
 
 		// Apply 'InFunc' to all cells in the provided bounds.
 		for (int32 Z = MinCellCoords.Z; Z <= MaxCellCoords.Z; ++Z)
@@ -1092,17 +1084,48 @@ void UPCGSubsystem::CleanupLocalComponentsImmediate(UPCGComponent* InOriginalCom
 	}
 }
 
+// deprecated
 UPCGComponent* UPCGSubsystem::GetLocalComponent(uint32 GridSize, const FIntVector& CellCoords, const UPCGComponent* InOriginalComponent, bool bTransient) const
 {
-	return ActorAndComponentMapping.GetLocalComponent(GridSize, CellCoords, InOriginalComponent, bTransient);
+	FPCGGridDescriptor GridDescriptor = FPCGGridDescriptor()
+		.SetGridSize(GridSize)
+		.SetIsRuntime(bTransient)
+		.SetIs2DGrid(InOriginalComponent->Use2DGrid());
+
+	return GetLocalComponent(GridDescriptor, CellCoords, InOriginalComponent);
 }
 
+// deprecated
 APCGPartitionActor* UPCGSubsystem::GetRegisteredPCGPartitionActor(uint32 GridSize, const FIntVector& GridCoords, bool bRuntimeGenerated) const
 {
-	return ActorAndComponentMapping.GetPartitionActor(GridSize, GridCoords, bRuntimeGenerated);
+	FPCGGridDescriptor GridDescriptor = FPCGGridDescriptor()
+		.SetGridSize(GridSize)
+		.SetIsRuntime(bRuntimeGenerated);
+
+	return GetRegisteredPCGPartitionActor(GridDescriptor, GridCoords);
 }
 
+// deprecated
 APCGPartitionActor* UPCGSubsystem::FindOrCreatePCGPartitionActor(const FGuid& Guid, uint32 GridSize, const FIntVector& GridCoords, bool bRuntimeGenerated, bool bCanCreateActor) const
+{
+	FPCGGridDescriptor GridDescriptor = FPCGGridDescriptor()
+		.SetGridSize(GridSize)
+		.SetIsRuntime(bRuntimeGenerated);
+
+	return FindOrCreatePCGPartitionActor(GridDescriptor, GridCoords, bCanCreateActor);
+}
+
+UPCGComponent* UPCGSubsystem::GetLocalComponent(const FPCGGridDescriptor& GridDescriptor, const FIntVector& CellCoords, const UPCGComponent* InOriginalComponent) const
+{
+	return ActorAndComponentMapping.GetLocalComponent(GridDescriptor, CellCoords, InOriginalComponent);
+}
+
+APCGPartitionActor* UPCGSubsystem::GetRegisteredPCGPartitionActor(const FPCGGridDescriptor& GridDescriptor, const FIntVector& GridCoords) const
+{
+	return ActorAndComponentMapping.GetPartitionActor(GridDescriptor, GridCoords);
+}
+
+APCGPartitionActor* UPCGSubsystem::FindOrCreatePCGPartitionActor(const FPCGGridDescriptor& GridDescriptor, const FIntVector& GridCoords, bool bCanCreateActor) const
 {
 	TRACE_CPUPROFILER_EVENT_SCOPE(UPCGActorHelpers::FindOrCreatePCGPartitionActor);
 
@@ -1121,14 +1144,14 @@ APCGPartitionActor* UPCGSubsystem::FindOrCreatePCGPartitionActor(const FGuid& Gu
 	}
 
 	// Attempt to find an existing PA.
-	if (APCGPartitionActor* ExistingActor = GetRegisteredPCGPartitionActor(GridSize, GridCoords, bRuntimeGenerated))
+	if (APCGPartitionActor* ExistingActor = GetRegisteredPCGPartitionActor(GridDescriptor, GridCoords))
 	{
 		return ExistingActor;
 	}
-	else if (!bRuntimeGenerated)
+	else if (!GridDescriptor.IsRuntime())
 	{
 		// Check if there is already an unloaded actor for this cell. RuntimeGenerated PAs are never unloaded, so we ignore them.
-		if (ActorAndComponentMapping.DoesPartitionActorRecordExist(Guid, GridSize, GridCoords))
+		if (ActorAndComponentMapping.DoesPartitionActorRecordExist(GridDescriptor, GridCoords))
 		{
 			return nullptr;
 		}
@@ -1141,19 +1164,20 @@ APCGPartitionActor* UPCGSubsystem::FindOrCreatePCGPartitionActor(const FGuid& Gu
 
 	FActorSpawnParameters SpawnParams;
 #if WITH_EDITOR
-	SpawnParams.Name = *APCGPartitionActor::GetPCGPartitionActorName(GridSize, GridCoords, bRuntimeGenerated);
-	SpawnParams.NameMode = FActorSpawnParameters::ESpawnActorNameMode::Requested;
+	SpawnParams.Name = *APCGPartitionActor::GetPCGPartitionActorName(GridDescriptor, GridCoords);
+	SpawnParams.NameMode = FActorSpawnParameters::ESpawnActorNameMode::Required_Fatal;
 #endif
-	if (bRuntimeGenerated)
+	if (GridDescriptor.IsRuntime())
 	{
+		SpawnParams.NameMode = FActorSpawnParameters::ESpawnActorNameMode::Requested;
 		SpawnParams.ObjectFlags |= RF_Transient;
 		SpawnParams.ObjectFlags &= ~RF_Transactional;
 	}
 
-	const FVector CellCenter(FVector(GridCoords.X + 0.5, GridCoords.Y + 0.5, GridCoords.Z + 0.5) * GridSize);
+	const FVector CellCenter(FVector(GridCoords.X + 0.5, GridCoords.Y + 0.5, GridCoords.Z + 0.5) * GridDescriptor.GetGridSize());
 	APCGPartitionActor* NewActor = CastChecked<APCGPartitionActor>(World->SpawnActor(APCGPartitionActor::StaticClass(), &CellCenter, nullptr, SpawnParams));
 
-	if (bRuntimeGenerated)
+	if (GridDescriptor.IsRuntime())
 	{
 		NewActor->SetToRuntimeGenerated();
 	}
@@ -1161,11 +1185,10 @@ APCGPartitionActor* UPCGSubsystem::FindOrCreatePCGPartitionActor(const FGuid& Gu
 #if WITH_EDITOR
 	NewActor->SetLockLocation(true);
 	NewActor->SetActorLabel(SpawnParams.Name.ToString());
-	NewActor->SetGridSize(GridSize);
 #endif
 
 	// Empty GUID if runtime generated, since transient PAs don't need one.
-	NewActor->PostCreation(bRuntimeGenerated ? FGuid() : Guid, GridSize);
+	NewActor->PostCreation(GridDescriptor);
 
 	return NewActor;
 }
@@ -1248,7 +1271,7 @@ void UPCGSubsystem::ClearPCGLink(UPCGComponent* InComponent, const FBox& InBound
 	FPCGTaskId TaskId = InvalidPCGTaskId;
 	if (!GridSizes.IsEmpty())
 	{
-		TaskId = ForAllOverlappingCells(InBounds, GridSizes, /*bCanCreateActor=*/false, /*Dependencies=*/{}, ScheduleTask);
+		TaskId = ForAllOverlappingCells(InComponent, InBounds, GridSizes, /*bCanCreateActor=*/false, /*Dependencies=*/{}, ScheduleTask);
 	}
 
 	// Verify if the NewActor has some components attached to its root or attached actors. If not, destroy it.
@@ -1396,6 +1419,14 @@ void UPCGSubsystem::DeleteSerializedPartitionActors(bool bOnlyDeleteUnused, bool
 			{
 				if (UPackage* ExternalPackage = ActorToDelete->GetExternalPackage())
 				{
+					// Since we aren't in a transaction (and this operation isnt undoable) make sure UPackage objects are no longer marked as RF_Standalone so it they can get properly GCed in (ObjectTools::CleanupAfterSuccessfulDelete)
+					// Without this call we end up spamming in UWorldPartition::OnGCPostReachabilityAnalysis
+					ForEachObjectWithPackage(ExternalPackage, [this, Actor](UObject* Object)
+					{
+						Object->ClearFlags(RF_Standalone);
+						return true;
+					}, false);
+
 					PackagesToCleanup.Add(ExternalPackage);
 				}
 
@@ -1410,6 +1441,8 @@ void UPCGSubsystem::DeleteSerializedPartitionActors(bool bOnlyDeleteUnused, bool
 	if (GEditor)
 	{
 		GEditor->SelectNone(true, true, false);
+		// Any reference in the Transaction buffer to the deleted actors will prevent them from being properly GCed so here we reset the transaction buffer
+		GEditor->ResetTransaction(NSLOCTEXT("PCGSubsystem", "DeletePartitionActorsResetTransaction", "Deleted PCG Actors"));
 	}
 
 	{
diff --git a/Engine/Plugins/PCG/Source/PCG/Private/PCGWorldActor.cpp b/Engine/Plugins/PCG/Source/PCG/Private/PCGWorldActor.cpp
--- a/Engine/Plugins/PCG/Source/PCG/Private/PCGWorldActor.cpp
+++ b/Engine/Plugins/PCG/Source/PCG/Private/PCGWorldActor.cpp
@@ -24,6 +24,7 @@
 
 #include UE_INLINE_GENERATED_CPP_BY_NAME(PCGWorldActor)
 
+PRAGMA_DISABLE_DEPRECATION_WARNINGS
 bool FPCGPartitionActorRecord::operator==(const FPCGPartitionActorRecord& InOther) const
 {
 	return GridGuid == InOther.GridGuid && GridSize == InOther.GridSize && GridCoords == InOther.GridCoords;
@@ -35,6 +36,7 @@ uint32 GetTypeHash(const FPCGPartitionActorRecord& In)
 	HashResult = HashCombine(HashResult, GetTypeHash(In.GridCoords));
 	return HashResult;
 }
+PRAGMA_ENABLE_DEPRECATION_WARNINGS
 
 APCGWorldActor::APCGWorldActor(const FObjectInitializer& ObjectInitializer)
 	: Super(ObjectInitializer)
@@ -63,13 +65,13 @@ void APCGWorldActor::BeginCacheForCookedPlatformData(const ITargetPlatform* Targ
 	// Cache Partition Actor Records at Cook
 	if (UWorldPartition* WorldPartition = World->GetWorldPartition())
 	{
-		TMap<FPCGPartitionActorRecord, FGuid> PartitionActorRecords;
+		TMap<FPCGGridCellDescriptor, FGuid> PartitionActorRecords;
 		TSet<FGuid> InvalidPartitionActors;
 		FPCGActorAndComponentMapping::BuildPartitionActorRecords(this, WorldPartition, PartitionActorRecords, InvalidPartitionActors);
 
-		TArray<FPCGPartitionActorRecord> SerializedPartitionActorRecords;
+		TArray<FPCGGridCellDescriptor> SerializedPartitionActorRecords;
 		PartitionActorRecords.GenerateKeyArray(SerializedPartitionActorRecords);
-		RuntimePartitionActorRecords = TSet<FPCGPartitionActorRecord>(SerializedPartitionActorRecords);
+		RuntimePartitionActorRecords = TSet<FPCGGridCellDescriptor>(SerializedPartitionActorRecords);
 	}
 
 	if (LandscapeCacheObject->SerializationMode == EPCGLandscapeCacheSerializationMode::SerializeOnlyAtCook)
@@ -160,81 +162,6 @@ void APCGWorldActor::EndPlay(const EEndPlayReason::Type EndPlayReason)
 	Super::EndPlay(EndPlayReason);
 }
 
-void APCGWorldActor::CreateGridGuidsIfNecessary(const PCGHiGenGrid::FSizeArray& InGridSizes, bool bAreGridsSerialized)
-{
-	if (InGridSizes.IsEmpty())
-	{
-		return;
-	}
-
-	TMap<uint32, FGuid>& GuidsMap = bAreGridsSerialized ? GridGuids : TransientGridGuids;
-	FRWLock& GuidsLock = bAreGridsSerialized ? GridGuidsLock : TransientGridGuidsLock;
-
-	// Check if any need adding
-	PCGHiGenGrid::FSizeArray GridSizesToAdd;
-	{
-		FReadScopeLock ReadLock(GuidsLock);
-		for (uint32 GridSize : InGridSizes)
-		{
-			if (!GuidsMap.Contains(GridSize))
-			{
-				GridSizesToAdd.Push(GridSize);
-			}
-		}
-	}
-
-	if (GridSizesToAdd.Num() > 0)
-	{
-		FWriteScopeLock WriteLock(GuidsLock);
-
-		bool bModified = false;
-		for (uint32 GridSize : GridSizesToAdd)
-		{
-			if (!GuidsMap.Contains(GridSize))
-			{
-				GuidsMap.Add(GridSize, FGuid::NewGuid());
-				bModified = true;
-			}
-		}
-
-		if (bModified && bAreGridsSerialized)
-		{
-			// Set dirty flag if we added guids. Unfortunately if the guids are not up to date, this will produce save prompts
-			// to users. However, this was needed to ensure the guids are saved - without this guids were lost and PAs were leaked.
-			// The alternative would be to ensure this never happens automatically, but rather only happens when user clicks Generate
-			// or etc. However we take a very proactive approach to creating PAs in editor because they can't be created at runtime.
-
-			// Schedule dirtying rather than do immediately because dirtying is a no-op during level load
-			if (UPCGSubsystem* PCGSubsystem = UPCGSubsystem::GetInstance(GetWorld()))
-			{
-				PCGSubsystem->ScheduleGeneric([this]()
-				{
-					this->MarkPackageDirty();
-					return true;
-				}, nullptr, {});
-			}
-		}
-	}
-}
-
-void APCGWorldActor::GetSerializedGridGuids(PCGHiGenGrid::FSizeToGuidMap& OutSizeToGuidMap) const
-{
-	FReadScopeLock ReadLock(GridGuidsLock);
-	for (const TPair<uint32, FGuid>& SizeGuid : GridGuids)
-	{
-		OutSizeToGuidMap.Add(SizeGuid.Key, SizeGuid.Value);
-	}
-}
-
-void APCGWorldActor::GetTransientGridGuids(PCGHiGenGrid::FSizeToGuidMap& OutSizeToGuidMap) const
-{
-	FReadScopeLock ReadLock(TransientGridGuidsLock);
-	for (const TPair<uint32, FGuid>& SizeGuid : TransientGridGuids)
-	{
-		OutSizeToGuidMap.Add(SizeGuid.Key, SizeGuid.Value);
-	}
-}
-
 void APCGWorldActor::MergeFrom(APCGWorldActor* OtherWorldActor)
 {
 	check(OtherWorldActor && this != OtherWorldActor);
@@ -336,25 +263,11 @@ void APCGWorldActor::OnPartitionGridSizeChanged()
 	}
 }
 
-void APCGWorldActor::PostLoad()
-{
-	Super::PostLoad();
-
-	// Deprecation - PAs used to be placed on a grid with guid=0. If no grid guids are registered,
-	// register a 0 guid now for the current grid size, and this will result in already existing PAs
-	// being reused.
-	if (GridGuids.IsEmpty())
-	{
-		GridGuids.Add(PartitionGridSize, FGuid());
-	}
-}
-
 void APCGWorldActor::PostEditChangeProperty(FPropertyChangedEvent& PropertyChangedEvent)
 {
 	Super::PostEditChangeProperty(PropertyChangedEvent);
 	
-	if (PropertyChangedEvent.GetPropertyName() == GET_MEMBER_NAME_CHECKED(APCGWorldActor, PartitionGridSize)
-		|| PropertyChangedEvent.GetPropertyName() == GET_MEMBER_NAME_CHECKED(APCGWorldActor, bUse2DGrid))
+	if (PropertyChangedEvent.GetPropertyName() == GET_MEMBER_NAME_CHECKED(APCGWorldActor, PartitionGridSize))
 	{
 		OnPartitionGridSizeChanged();
 	}
diff --git a/Engine/Plugins/PCG/Source/PCG/Private/RuntimeGen/PCGRuntimeGenScheduler.cpp b/Engine/Plugins/PCG/Source/PCG/Private/RuntimeGen/PCGRuntimeGenScheduler.cpp
--- a/Engine/Plugins/PCG/Source/PCG/Private/RuntimeGen/PCGRuntimeGenScheduler.cpp
+++ b/Engine/Plugins/PCG/Source/PCG/Private/RuntimeGen/PCGRuntimeGenScheduler.cpp
@@ -10,6 +10,7 @@
 #include "PCGModule.h"
 #include "PCGSubsystem.h"
 #include "PCGWorldActor.h"
+#include "Grid/PCGGridDescriptor.h"
 #include "Grid/PCGPartitionActor.h"
 #include "Helpers/PCGActorHelpers.h"
 #include "Helpers/PCGHelpers.h"
@@ -25,7 +26,7 @@
 
 namespace PCGRuntimeGenSchedulerConstants
 {
-	const FString PooledPartitionActorName = TEXT("PCGRuntimeGenPartitionActor_POOLED");
+	const FString PooledPartitionActorName = TEXT("PCGRuntimePartitionGridActor_POOLED");
 }
 
 namespace PCGRuntimeGenSchedulerHelpers
@@ -56,6 +57,11 @@ namespace PCGRuntimeGenSchedulerHelpers
 		TEXT("Defines the base PartitionActor pool size for the RuntimeGeneration system. Cannot be less than 1."));
 }
 
+FPCGGridDescriptor FPCGRuntimeGenScheduler::FGridGenerationKey::GetGridDescriptor() const
+{
+	return FPCGGridDescriptor().SetGridSize(GetGridSize()).SetIsRuntime(true).SetIs2DGrid(GetOriginalComponent()->Use2DGrid());
+}
+
 FPCGRuntimeGenScheduler::FPCGRuntimeGenScheduler(UWorld* InWorld, FPCGActorAndComponentMapping* InActorAndComponentMapping)
 {
 	check(InWorld && InActorAndComponentMapping);
@@ -279,9 +285,10 @@ void FPCGRuntimeGenScheduler::TickQueueComponentsForGeneration(
 
 				const FVector GenSourcePosition = GenSourcePositionOptional.GetValue();
 				const FBox OriginalComponentBounds = OriginalComponent->GetGridBounds();
+				const bool bIs2DGrid = OriginalComponent->Use2DGrid();
 
 				FVector ModifiedGenSourcePosition = GenSourcePosition;
-				if (InPCGWorldActor->bUse2DGrid)
+				if (bIs2DGrid)
 				{
 					ModifiedGenSourcePosition.Z = OriginalComponentBounds.Min.Z;
 				}
@@ -300,7 +307,8 @@ void FPCGRuntimeGenScheduler::TickQueueComponentsForGeneration(
 					FGridGenerationKey Key(PCGHiGenGrid::UnboundedGridSize(), FIntVector(0), OriginalComponent);
 					if (!GeneratedComponents.Contains(Key))
 					{
-						AddComponentToGenerate(Key, GenSource, Policy, OriginalComponentBounds, InPCGWorldActor->bUse2DGrid);
+						check(Key.GetGridDescriptor().Is2DGrid() == bIs2DGrid);
+						AddComponentToGenerate(Key, GenSource, Policy, OriginalComponentBounds, bIs2DGrid);
 					}
 				}
 
@@ -310,16 +318,16 @@ void FPCGRuntimeGenScheduler::TickQueueComponentsForGeneration(
 				{
 					ensure(PCGHiGenGrid::IsValidGridSize(GridSize));
 
-					const FIntVector GenSourceGridPosition = UPCGActorHelpers::GetCellCoord(GenSourcePosition, GridSize, InPCGWorldActor->bUse2DGrid);
+					const FIntVector GenSourceGridPosition = UPCGActorHelpers::GetCellCoord(GenSourcePosition, GridSize, bIs2DGrid);
 					const double GenerationRadius = OriginalComponent->GetGenerationRadiusFromGrid(PCGHiGenGrid::GridSizeToGrid(GridSize));
 					const int32 GridRadius = FMath::CeilToInt32(GenerationRadius / GridSize); // Radius discretized to # of grid cells.
-					const int32 VerticalGridRadius = InPCGWorldActor->bUse2DGrid ? 0 : GridRadius; // Flatten the vertical grid radius in the 2D case.
+					const int32 VerticalGridRadius = bIs2DGrid ? 0 : GridRadius; // Flatten the vertical grid radius in the 2D case.
 
 					const double HalfGridSize = GridSize / 2.0f;
 					FVector HalfExtent(HalfGridSize, HalfGridSize, HalfGridSize);
 
 					FBox ModifiedBounds = OriginalComponent->GetGridBounds();
-					if (InPCGWorldActor->bUse2DGrid)
+					if (bIs2DGrid)
 					{
 						FVector MinBounds = ModifiedBounds.Min;
 						FVector MaxBounds = ModifiedBounds.Max;
@@ -358,7 +366,7 @@ void FPCGRuntimeGenScheduler::TickQueueComponentsForGeneration(
 									continue;
 								}
 
-								if (InPCGWorldActor->bUse2DGrid)
+								if (Key.GetGridDescriptor().Is2DGrid())
 								{
 									ModifiedGenSourcePosition.Z = IntersectedBounds.Min.Z;
 								}
@@ -368,7 +376,7 @@ void FPCGRuntimeGenScheduler::TickQueueComponentsForGeneration(
 								const double LocalDistanceSquared = CellBounds.ComputeSquaredDistanceToPoint(ModifiedGenSourcePosition);
 								if (LocalDistanceSquared <= GenerationRadius * GenerationRadius)
 								{
-									AddComponentToGenerate(Key, GenSource, Policy, IntersectedBounds, InPCGWorldActor->bUse2DGrid);
+									AddComponentToGenerate(Key, GenSource, Policy, IntersectedBounds, Key.GetGridDescriptor().Is2DGrid());
 								}
 							}
 						}
@@ -414,7 +422,7 @@ void FPCGRuntimeGenScheduler::TickQueueComponentsForGeneration(
 				const FBox OriginalComponentBounds = OriginalComponent->GetGridBounds();
 
 				FVector ModifiedGenSourcePosition = GenSourcePosition;
-				if (InPCGWorldActor->bUse2DGrid)
+				if (OriginalComponent->Use2DGrid())
 				{
 					ModifiedGenSourcePosition.Z = OriginalComponentBounds.Min.Z;
 				}
@@ -448,13 +456,12 @@ void FPCGRuntimeGenScheduler::TickCleanup(const TSet<IPCGGenSourceBase*>& InGenS
 	{
 		TRACE_CPUPROFILER_EVENT_SCOPE(FPCGRuntimeGenScheduler::SelectComponentsToCleanup);
 
-		const uint32 GridSize = GenerationKey.GetGridSize();
+		const FPCGGridDescriptor GridDescriptor = GenerationKey.GetGridDescriptor();
+		const EPCGHiGenGrid Grid = GridDescriptor.GetHiGenGrid();
 		const FIntVector& GridCoords = GenerationKey.GetGridCoords();
 		UPCGComponent* OriginalComponent = GenerationKey.GetOriginalComponent();
 		check(OriginalComponent);
 
-		const EPCGHiGenGrid Grid = PCGHiGenGrid::GridSizeToGrid(GridSize);
-
 		// If the Grid is unbounded, we have a non-partitioned or unbounded component.
 		if (Grid == EPCGHiGenGrid::Unbounded)
 		{
@@ -483,7 +490,7 @@ void FPCGRuntimeGenScheduler::TickCleanup(const TSet<IPCGGenSourceBase*>& InGenS
 				FVector GenSourcePosition = GenSourcePositionOptional.GetValue();
 
 				// Only consider 2D distance when using a 2D grid.
-				if (InPCGWorldActor->bUse2DGrid)
+				if (GridDescriptor.Is2DGrid())
 				{
 					GenSourcePosition.Z = Bounds.Min.Z;
 				}
@@ -500,7 +507,7 @@ void FPCGRuntimeGenScheduler::TickCleanup(const TSet<IPCGGenSourceBase*>& InGenS
 		// Otherwise, we have a local component.
 		else
 		{
-			UPCGComponent* LocalComponent = ActorAndComponentMapping->GetLocalComponent(GridSize, GridCoords, OriginalComponent, /*bRuntimeGenerated=*/true);
+			UPCGComponent* LocalComponent = ActorAndComponentMapping->GetLocalComponent(GridDescriptor, GridCoords, OriginalComponent);
 			APCGPartitionActor* PartitionActor = LocalComponent ? Cast<APCGPartitionActor>(LocalComponent->GetOwner()) : nullptr;
 			if (!PartitionActor || !OriginalComponent->bActivated)
 			{
@@ -528,7 +535,7 @@ void FPCGRuntimeGenScheduler::TickCleanup(const TSet<IPCGGenSourceBase*>& InGenS
 				FVector GenSourcePosition = GenSourcePositionOptional.GetValue();
 
 				// Only consider 2D distance when using a 2D grid.
-				if (InPCGWorldActor->bUse2DGrid)
+				if (GridDescriptor.Is2DGrid())
 				{
 					GenSourcePosition.Z = GridBounds.Min.Z;
 				}
@@ -561,13 +568,13 @@ void FPCGRuntimeGenScheduler::TickScheduleGeneration(TMap<FGridGenerationKey, do
 		const FGridGenerationKey& Key = GenerateEntry.Key;
 		const double Priority = GenerateEntry.Value;
 
-		const uint32 GridSize = GenerateEntry.Key.GetGridSize();
+		const FPCGGridDescriptor GridDescriptor = Key.GetGridDescriptor();
+		const EPCGHiGenGrid Grid = GridDescriptor.GetHiGenGrid();
+
 		const FIntVector GridCoords = GenerateEntry.Key.GetGridCoords();
 		UPCGComponent* OriginalComponent = GenerateEntry.Key.GetOriginalComponent();
 		check(OriginalComponent);
 
-		const EPCGHiGenGrid Grid = PCGHiGenGrid::GridSizeToGrid(GridSize);
-
 		// If the Grid is unbounded, we have a non-partitioned or unbounded component.
 		if (Grid == EPCGHiGenGrid::Unbounded)
 		{
@@ -586,7 +593,7 @@ void FPCGRuntimeGenScheduler::TickScheduleGeneration(TMap<FGridGenerationKey, do
 		else
 		{
 			// Grab local component and PA if they exist already.
-			UPCGComponent* LocalComponent = ActorAndComponentMapping->GetLocalComponent(GridSize, GridCoords, OriginalComponent, /*bRuntimeGenerated=*/true);
+			UPCGComponent* LocalComponent = ActorAndComponentMapping->GetLocalComponent(GridDescriptor, GridCoords, OriginalComponent);
 			TObjectPtr<APCGPartitionActor> PartitionActor = LocalComponent ? Cast<APCGPartitionActor>(LocalComponent->GetOwner()) : nullptr;
 
 			if (!LocalComponent || !ensure(PartitionActor))
@@ -595,12 +602,12 @@ void FPCGRuntimeGenScheduler::TickScheduleGeneration(TMap<FGridGenerationKey, do
 				if (PCGRuntimeGenSchedulerHelpers::CVarRuntimeGenerationEnablePooling.GetValueOnAnyThread())
 				{
 					// Get RuntimeGenPA from pool.
-					PartitionActor = GetPartitionActorFromPool(GridSize, GridCoords);
+					PartitionActor = GetPartitionActorFromPool(GridDescriptor, GridCoords);
 
 					if (PCGRuntimeGenSchedulerHelpers::CVarRuntimeGenerationEnableDebugging.GetValueOnAnyThread())
 					{
 						UE_LOG(LogPCG, Warning, TEXT("[RUNTIMEGEN] UNPOOL PARTITION ACTOR: '%s' (priority %lf, %d remaining out of %d)"),
-							*APCGPartitionActor::GetPCGPartitionActorName(GridSize, GridCoords, /*bRuntimeGenerated=*/true),
+							*APCGPartitionActor::GetPCGPartitionActorName(GridDescriptor, GridCoords),
 							Priority,
 							PartitionActorPool.Num(),
 							PartitionActorPoolSize);
@@ -611,12 +618,12 @@ void FPCGRuntimeGenScheduler::TickScheduleGeneration(TMap<FGridGenerationKey, do
 					if (PCGRuntimeGenSchedulerHelpers::CVarRuntimeGenerationEnableDebugging.GetValueOnAnyThread())
 					{
 						UE_LOG(LogPCG, Warning, TEXT("[RUNTIMEGEN] CREATE PARTITION ACTOR: '%s' (priority %lf)"),
-							*APCGPartitionActor::GetPCGPartitionActorName(GridSize, GridCoords, /*bRuntimeGenerated=*/true),
+							*APCGPartitionActor::GetPCGPartitionActorName(GridDescriptor, GridCoords),
 							Priority);
 					}
 
 					// Find or Create RuntimeGenPA.
-					PartitionActor = Subsystem->FindOrCreatePCGPartitionActor(FGuid(), GridSize, GridCoords, /*bRuntimeGenerated=*/true);
+					PartitionActor = Subsystem->FindOrCreatePCGPartitionActor(GridDescriptor, GridCoords);
 				}
 
 				if (!ensure(PartitionActor))
@@ -626,12 +633,12 @@ void FPCGRuntimeGenScheduler::TickScheduleGeneration(TMap<FGridGenerationKey, do
 
 				// Update component mapping for this PA (add local component).
 				{
-					FWriteScopeLock WriteLock(ActorAndComponentMapping->ComponentToRuntimeGenPartitionActorsMapLock);
-					TSet<TObjectPtr<APCGPartitionActor>>* PartitionActorsPtr = ActorAndComponentMapping->ComponentToRuntimeGenPartitionActorsMap.Find(OriginalComponent);
+					FWriteScopeLock WriteLock(ActorAndComponentMapping->ComponentToPartitionActorsMapLock);
+					TSet<TObjectPtr<APCGPartitionActor>>* PartitionActorsPtr = ActorAndComponentMapping->ComponentToPartitionActorsMap.Find(OriginalComponent);
 
 					if (!PartitionActorsPtr)
 					{
-						PartitionActorsPtr = &ActorAndComponentMapping->ComponentToRuntimeGenPartitionActorsMap.Emplace(OriginalComponent);
+						PartitionActorsPtr = &ActorAndComponentMapping->ComponentToPartitionActorsMap.Emplace(OriginalComponent);
 					}
 
 					// Log this original component before setting up the PA, so that we early out from RefreshComponent if it gets called
@@ -721,8 +728,6 @@ void FPCGRuntimeGenScheduler::OnOriginalComponentRegistered(UPCGComponent* InOri
 		return;
 	}
 
-	CreateGridGuidsForComponent(InOriginalComponent);
-
 	// When an original/non-partitioned component is registered, we need to dirty the state.
 	bAnyRuntimeGenComponentsExistDirty = true;
 }
@@ -751,11 +756,11 @@ void FPCGRuntimeGenScheduler::OnOriginalComponentUnregistered(UPCGComponent* InO
 
 	for (FGridGenerationKey GenerationKey : KeysToCleanup)
 	{
-		const uint32 GridSize = GenerationKey.GetGridSize();
-		const bool bIsOriginalComponent = GridSize == PCGHiGenGrid::UnboundedGridSize();
+		const FPCGGridDescriptor GridDescriptor = GenerationKey.GetGridDescriptor();
+		const bool bIsOriginalComponent = GridDescriptor.GetGridSize() == PCGHiGenGrid::UnboundedGridSize();
 
 		// Get the generated component for this key (might be a local component).
-		UPCGComponent* ComponentToCleanup = bIsOriginalComponent ? InOriginalComponent : ActorAndComponentMapping->GetLocalComponent(GridSize, GenerationKey.GetGridCoords(), InOriginalComponent, /*bRuntimeGenerated=*/true);
+		UPCGComponent* ComponentToCleanup = bIsOriginalComponent ? InOriginalComponent : ActorAndComponentMapping->GetLocalComponent(GridDescriptor, GenerationKey.GetGridCoords(), InOriginalComponent);
 
 		if (ensure(ComponentToCleanup))
 		{
@@ -777,17 +782,16 @@ void FPCGRuntimeGenScheduler::CleanupLocalComponents(const APCGWorldActor* InPCG
 	// Find all generated local components.
 	for (FGridGenerationKey GenerationKey : GeneratedComponents)
 	{
-		const uint32 GridSize = GenerationKey.GetGridSize();
+		const FPCGGridDescriptor GridDescriptor = GenerationKey.GetGridDescriptor();
+		const EPCGHiGenGrid Grid = GridDescriptor.GetHiGenGrid();
 		const FIntVector& GridCoords = GenerationKey.GetGridCoords();
 		UPCGComponent* OriginalComponent = GenerationKey.GetOriginalComponent();
 		check(OriginalComponent);
 
-		const EPCGHiGenGrid Grid = PCGHiGenGrid::GridSizeToGrid(GridSize);
-
 		// Only operate on LocalComponents.
 		if (Grid != EPCGHiGenGrid::Unbounded)
 		{
-			UPCGComponent* LocalComponent = ActorAndComponentMapping->GetLocalComponent(GridSize, GridCoords, OriginalComponent, /*bRuntimeGenerated=*/true);
+			UPCGComponent* LocalComponent = ActorAndComponentMapping->GetLocalComponent(GridDescriptor, GridCoords, OriginalComponent);
 			ComponentsToClean.Add({ GenerationKey, LocalComponent });
 		}
 	}
@@ -804,22 +808,23 @@ void FPCGRuntimeGenScheduler::CleanupComponent(const FGridGenerationKey& Generat
 
 	check(ActorAndComponentMapping);
 
-	const uint32 GridSize = GenerationKey.GetGridSize();
+	const FPCGGridDescriptor GridDescriptor = GenerationKey.GetGridDescriptor();
+	const EPCGHiGenGrid Grid = GridDescriptor.GetHiGenGrid();
+
 	const FIntVector GridCoords = GenerationKey.GetGridCoords();
 
 	APCGPartitionActor* PartitionActor = nullptr;
 
 	if (!GeneratedComponent)
 	{
-		UE_LOG(LogPCG, Warning, TEXT("Runtime generated component could not be recovered on grid %d at (%d, %d, %d). It has been lost or destroyed."), GridSize, GridCoords.X, GridCoords.Y, GridCoords.Z);
+		UE_LOG(LogPCG, Warning, TEXT("Runtime generated component could not be recovered on grid %d at (%d, %d, %d). It has been lost or destroyed."), GridDescriptor.GetGridSize(), GridCoords.X, GridCoords.Y, GridCoords.Z);
 
 		// If the GeneratedComponent has been lost for some reason, get the PA directly from the ActorAndComponentMapping.
-		PartitionActor = ActorAndComponentMapping->GetPartitionActor(GridSize, GridCoords, /*bRuntimeGenerated=*/true);
+		PartitionActor = ActorAndComponentMapping->GetPartitionActor(GridDescriptor, GridCoords);
 	}
 	else // If the GeneratedComponent does exist, we can clean it up!.
 	{
 		PartitionActor = Cast<APCGPartitionActor>(GeneratedComponent->GetOwner());
-		const EPCGHiGenGrid Grid = PCGHiGenGrid::GridSizeToGrid(GridSize);
 
 		// If the Grid is unbounded, we have a non-partitioned or unbounded component.
 		if (Grid == EPCGHiGenGrid::Unbounded)
@@ -846,8 +851,8 @@ void FPCGRuntimeGenScheduler::CleanupComponent(const FGridGenerationKey& Generat
 
 			// Remove component mapping.
 			{
-				FWriteScopeLock WriteLock(ActorAndComponentMapping->ComponentToRuntimeGenPartitionActorsMapLock);
-				TSet<TObjectPtr<APCGPartitionActor>>* PartitionActorsPtr = ActorAndComponentMapping->ComponentToRuntimeGenPartitionActorsMap.Find(OriginalComponent);
+				FWriteScopeLock WriteLock(ActorAndComponentMapping->ComponentToPartitionActorsMapLock);
+				TSet<TObjectPtr<APCGPartitionActor>>* PartitionActorsPtr = ActorAndComponentMapping->ComponentToPartitionActorsMap.Find(OriginalComponent);
 
 				if (PartitionActorsPtr)
 				{
@@ -855,7 +860,7 @@ void FPCGRuntimeGenScheduler::CleanupComponent(const FGridGenerationKey& Generat
 
 					if (PartitionActorsPtr->IsEmpty())
 					{
-						ActorAndComponentMapping->ComponentToRuntimeGenPartitionActorsMap.Remove(OriginalComponent);
+						ActorAndComponentMapping->ComponentToPartitionActorsMap.Remove(OriginalComponent);
 					}
 				}
 			}
@@ -907,8 +912,8 @@ void FPCGRuntimeGenScheduler::CleanupDelayedRefreshComponents()
 	// If it is not, it would be leaked instead of refreshed, so we should force a full cleanup.
 	for (const FGridGenerationKey& GenerationKey : GeneratedComponentsToRemove)
 	{
-		const uint32 GridSize = GenerationKey.GetGridSize();
-		const EPCGHiGenGrid Grid = PCGHiGenGrid::GridSizeToGrid(GridSize);
+		const FPCGGridDescriptor GridDescriptor = GenerationKey.GetGridDescriptor();
+		const EPCGHiGenGrid Grid = GridDescriptor.GetHiGenGrid();
 		UPCGComponent* OriginalComponent = GenerationKey.GetOriginalComponent();
 		const FIntVector& GridCoords = GenerationKey.GetGridCoords();
 
@@ -923,7 +928,7 @@ void FPCGRuntimeGenScheduler::CleanupDelayedRefreshComponents()
 			continue;
 		}
 
-		UPCGComponent* LocalComponent = OriginalComponent ? ActorAndComponentMapping->GetLocalComponent(GridSize, GridCoords, OriginalComponent, /*bRuntimeGenerated=*/true) : nullptr;
+		UPCGComponent* LocalComponent = OriginalComponent ? ActorAndComponentMapping->GetLocalComponent(GridDescriptor, GridCoords, OriginalComponent) : nullptr;
 		APCGPartitionActor* PartitionActor = LocalComponent ? Cast<APCGPartitionActor>(LocalComponent->GetOwner()) : nullptr;
 
 		if (LocalComponent && PartitionActor)
@@ -999,12 +1004,13 @@ void FPCGRuntimeGenScheduler::RefreshComponent(UPCGComponent* InComponent, bool
 			{
 				if (Key.GetOriginalComponent() == InComponent && !GeneratedComponentsToRemove.Contains(Key))
 				{
+					const FPCGGridDescriptor GridDescriptor = Key.GetGridDescriptor();
+					
 					// TODO - clean up local immediate will have a flag in the future to clean up the local components on its own, so this call to CleanupLocalImmediate will not be required
 					UPCGComponent* LocalComponent = ActorAndComponentMapping->GetLocalComponent(
-						Key.GetGridSize(),
+						GridDescriptor,
 						Key.GetGridCoords(),
-						OriginalComponent,
-						/*bRuntimeGenerated=*/true);
+						OriginalComponent);
 
 					if (bLoggingEnabled && LocalComponent && LocalComponent->GetOwner())
 					{
@@ -1072,8 +1078,8 @@ void FPCGRuntimeGenScheduler::RefreshComponent(UPCGComponent* InComponent, bool
 			// Gather all generated components which originated from this original component.
 			for (FGridGenerationKey GenerationKey : GenerationKeys)
 			{
-				const uint32 GridSize = GenerationKey.GetGridSize();
-				const EPCGHiGenGrid Grid = PCGHiGenGrid::GridSizeToGrid(GridSize);
+				const FPCGGridDescriptor GridDescriptor = GenerationKey.GetGridDescriptor();
+				const EPCGHiGenGrid Grid = GridDescriptor.GetHiGenGrid();
 
 				// If the Grid is unbounded, we have a non-partitioned or unbounded component.
 				if (Grid == EPCGHiGenGrid::Unbounded)
@@ -1092,7 +1098,7 @@ void FPCGRuntimeGenScheduler::RefreshComponent(UPCGComponent* InComponent, bool
 				{
 					const FIntVector GridCoords = GenerationKey.GetGridCoords();
 
-					if (UPCGComponent* LocalComponent = ActorAndComponentMapping->GetLocalComponent(GridSize, GridCoords, OriginalComponent, /*bRuntimeGenerated=*/true))
+					if (UPCGComponent* LocalComponent = ActorAndComponentMapping->GetLocalComponent(GridDescriptor, GridCoords, OriginalComponent))
 					{
 						RefreshLocalComponent(LocalComponent);
 					}
@@ -1113,7 +1119,7 @@ void FPCGRuntimeGenScheduler::RefreshComponent(UPCGComponent* InComponent, bool
 	}
 }
 
-APCGPartitionActor* FPCGRuntimeGenScheduler::GetPartitionActorFromPool(uint32 GridSize, const FIntVector& GridCoords)
+APCGPartitionActor* FPCGRuntimeGenScheduler::GetPartitionActorFromPool(const FPCGGridDescriptor& GridDescriptor, const FIntVector& GridCoords)
 {
 	TRACE_CPUPROFILER_EVENT_SCOPE(FPCGRuntimeGenScheduler::GetPartitionActorFromPool);
 
@@ -1126,7 +1132,7 @@ APCGPartitionActor* FPCGRuntimeGenScheduler::GetPartitionActorFromPool(uint32 Gr
 	}
 
 	// Attempt to find an existing RuntimeGen PA.
-	if (APCGPartitionActor* ExistingActor = ActorAndComponentMapping->GetPartitionActor(GridSize, GridCoords, /*bRuntimeGenerated=*/true))
+	if (APCGPartitionActor* ExistingActor = ActorAndComponentMapping->GetPartitionActor(GridDescriptor, GridCoords))
 	{
 		return ExistingActor;
 	}
@@ -1150,13 +1156,13 @@ APCGPartitionActor* FPCGRuntimeGenScheduler::GetPartitionActorFromPool(uint32 Gr
 	APCGPartitionActor* PartitionActor = PartitionActorPool.Pop();
 
 #if WITH_EDITOR
-	const FName ActorName = *APCGPartitionActor::GetPCGPartitionActorName(GridSize, GridCoords, /*bRuntimeGenerated=*/true);
+	const FName ActorName = *APCGPartitionActor::GetPCGPartitionActorName(GridDescriptor, GridCoords);
 
 	PartitionActor->Rename(*ActorName.ToString(), PartitionActor->GetOuter(), REN_NonTransactional | REN_DoNotDirty | REN_ForceNoResetLoaders);
 	PartitionActor->SetActorLabel(ActorName.ToString());
 #endif
 
-	const FVector CellCenter(FVector(GridCoords.X + 0.5, GridCoords.Y + 0.5, GridCoords.Z + 0.5) * GridSize);
+	const FVector CellCenter(FVector(GridCoords.X + 0.5, GridCoords.Y + 0.5, GridCoords.Z + 0.5) * GridDescriptor.GetGridSize());
 	if (!PartitionActor->Teleport(CellCenter))
 	{
 		UE_LOG(LogPCG, Error, TEXT("[RUNTIMEGEN] Could not set the location of RuntimeGen partition actor '%s'."), *PartitionActor->GetActorNameOrLabel());
@@ -1167,7 +1173,7 @@ APCGPartitionActor* FPCGRuntimeGenScheduler::GetPartitionActorFromPool(uint32 Gr
 #endif
 
 	// Empty GUID, RuntimeGen PAs don't need one.
-	PartitionActor->PostCreation(FGuid(), GridSize);
+	PartitionActor->PostCreation(GridDescriptor);
 
 	return PartitionActor;
 }
@@ -1215,17 +1221,4 @@ void FPCGRuntimeGenScheduler::ResetPartitionActorPoolToSize(uint32 NewPoolSize)
 	AddPartitionActorPoolCount(NewPoolSize);
 }
 
-void FPCGRuntimeGenScheduler::CreateGridGuidsForComponent(UPCGComponent* InComponent)
-{
-	if (InComponent && InComponent->IsPartitioned() && InComponent->IsManagedByRuntimeGenSystem())
-	{
-		if (APCGWorldActor* PCGWorldActor = PCGHelpers::GetPCGWorldActor(World))
-		{
-			bool bHasUnbounded;
-			PCGHiGenGrid::FSizeArray GridSizes;
-			ensure(PCGHelpers::GetGenerationGridSizes(InComponent->GetGraph(), PCGWorldActor, GridSizes, bHasUnbounded));
 
-			PCGWorldActor->CreateGridGuidsIfNecessary(GridSizes, /*bAreGridsSerialized=*/false);
-		}
-	}
-}
diff --git a/Engine/Plugins/PCG/Source/PCG/Public/Grid/PCGGridDescriptor.h b/Engine/Plugins/PCG/Source/PCG/Public/Grid/PCGGridDescriptor.h
new file mode 100644
--- /dev/null
+++ b/Engine/Plugins/PCG/Source/PCG/Public/Grid/PCGGridDescriptor.h
@@ -0,0 +1,84 @@
+// Copyright Epic Games, Inc. All Rights Reserved.
+
+#pragma once
+
+#include "PCGCommon.h"
+
+#include "CoreMinimal.h"
+#include "Misc/Optional.h"
+
+#include "PCGGridDescriptor.generated.h"
+
+/**
+* Descriptor struct used to determine where to output generated resources
+*/
+USTRUCT()
+struct FPCGGridDescriptor
+{
+	GENERATED_BODY()
+
+private:
+	UPROPERTY()
+	uint32 GridSize = 0;
+		
+	UPROPERTY()
+	bool bIs2DGrid = true;
+
+	UPROPERTY()
+	bool bIsRuntime = false;
+		
+public:
+	FPCGGridDescriptor& SetGridSize(uint32 InGridSize) { check(!Hash.IsSet()); GridSize = InGridSize; return *this; }
+	uint32 GetGridSize() const { return GridSize; }
+		
+	EPCGHiGenGrid GetHiGenGrid() const { return PCGHiGenGrid::GridSizeToGrid(GetGridSize()); }
+
+	FPCGGridDescriptor& SetIs2DGrid(bool bInIs2DGrid) { check(!Hash.IsSet()); bIs2DGrid = bInIs2DGrid; return *this; }
+	bool Is2DGrid() const { return bIs2DGrid; }
+
+	FPCGGridDescriptor& SetIsRuntime(bool bInIsRuntime) { check(!Hash.IsSet()); bIsRuntime = bInIsRuntime; return *this; }
+	bool IsRuntime() const { return bIsRuntime; }
+
+	bool operator==(const FPCGGridDescriptor& Other) const;
+	bool operator!=(const FPCGGridDescriptor& Other) const { return !(*this == Other); }
+		
+	friend uint32 GetTypeHash(const FPCGGridDescriptor& Descriptor)
+	{
+		if (!Descriptor.Hash.IsSet())
+		{
+			Descriptor.Hash = Descriptor.ComputeHash();
+		}
+
+		return Descriptor.Hash.GetValue();
+	}
+
+private:
+	uint32 ComputeHash() const;
+
+	mutable TOptional<uint32> Hash;
+};
+
+/**
+ * Describes one entry in a Grid
+ */
+USTRUCT()
+struct FPCGGridCellDescriptor
+{
+	GENERATED_BODY()
+
+	FPCGGridCellDescriptor() = default;
+	FPCGGridCellDescriptor(const FPCGGridDescriptor& InDescriptor, const FIntVector& InGridCoords)
+		: Descriptor(InDescriptor), GridCoords(InGridCoords)
+	{}
+
+	UPROPERTY()
+	FPCGGridDescriptor Descriptor;
+
+	/** The specific grid cell this actor lives in. */
+	UPROPERTY(VisibleAnywhere, Category = Debug)
+	FIntVector GridCoords = FIntVector::ZeroValue;
+
+	bool operator==(const FPCGGridCellDescriptor& InOther) const;
+	friend uint32 GetTypeHash(const FPCGGridCellDescriptor& In);
+};
+
diff --git a/Engine/Plugins/PCG/Source/PCG/Public/Grid/PCGPartitionActor.h b/Engine/Plugins/PCG/Source/PCG/Public/Grid/PCGPartitionActor.h
--- a/Engine/Plugins/PCG/Source/PCG/Public/Grid/PCGPartitionActor.h
+++ b/Engine/Plugins/PCG/Source/PCG/Public/Grid/PCGPartitionActor.h
@@ -5,6 +5,7 @@
 #include "ActorPartition/PartitionActor.h"
 
 #include "PCGCommon.h"
+#include "Grid/PCGGridDescriptor.h"
 
 #include "PCGPartitionActor.generated.h"
 
@@ -46,16 +47,19 @@ public:
 #if WITH_EDITOR
 	//~Begin APartitionActor Interface
 	virtual uint32 GetDefaultGridSize(UWorld* InWorld) const override;
-	virtual FGuid GetGridGuid() const override { return PCGGuid; }
 	virtual bool ShouldIncludeGridSizeInLabel() const override { return true; }
 	virtual TUniquePtr<class FWorldPartitionActorDesc> CreateClassActorDesc() const override;
 	virtual bool IsUserManaged() const override;
 	//~End APartitionActor Interface
+
+	// Called from PCGActorAndComponentMapping in case we are dealing with an older PCGPartitionActor
+	void UpdateUse2DGridIfNeeded(bool bInUse2DGrid);
 #endif
 
 	FBox GetFixedBounds() const;
 	FIntVector GetGridCoord() const;
 	uint32 GetPCGGridSize() const { return PCGGridSize; }
+	FPCGGridDescriptor GetGridDescriptor() const;
 
 	bool IsUsing2DGrid() const { return bUse2DGrid; }
 
@@ -81,7 +85,10 @@ public:
 	void RemoveLocalComponent(UPCGComponent* LocalComponent);
 
 	/** To be called after the creation of a new actor to set the grid guid and size. */
-	void PostCreation(const FGuid& InGridGUID, uint32 InGridSize);
+	UE_DEPRECATED(5.5, "Use FPCGGridDescriptor version")
+	void PostCreation(const FGuid& InGridGUID, uint32 InGridSize) { }
+
+	void PostCreation(const FPCGGridDescriptor& GridDescriptor);
 
 	/** [Game thread only] Return if the actor is safe for deletion, meaning no generation is currently running on all original components. */
 	bool IsSafeForDeletion() const;
@@ -106,6 +113,7 @@ public:
 	UFUNCTION(BlueprintCallable, Category = "PCG|PartitionActor")
 	UPCGComponent* GetOriginalComponent(const UPCGComponent* LocalComponent) const;
 
+	UE_DEPRECATED(5.5, "PCGGuid is deprecated")
 	UPROPERTY()
 	FGuid PCGGuid;
 
@@ -147,6 +155,9 @@ private:
 	/** Flag used to ignore some invalid actors so they don't get registered into the PCG Subsystem */
 	UPROPERTY(Transient, NonTransactional)
 	bool bIsInvalidForPCG = false;
+
+	/** Set from PostLoad so that we can know if this actor needs to update its bUse2DGrid from the old APCGWorldActor::bUse2DGrid flag or not */
+	bool bRequiresUse2DGridFixup = false;
 #endif // WITH_EDITORONLY_DATA
 
 	/** Tracks the registration status of this PA with the ActorAndComponentMapping system. Helps us avoid invalid (un)registers. */
@@ -163,8 +174,11 @@ public:
 	 * Gets the name this partition actor should have.
 	 * This does not respect traditional PA name contents like GridGuid, ShouldIncludeGridSizeInName, or ContextHash.
 	 */
+	UE_DEPRECATED(5.5, "Use FPCGGridDescriptor version")
 	static FString GetPCGPartitionActorName(uint32 GridSize, const FIntVector& GridCoords, bool bRuntimeGenerated);
 
+	static FString GetPCGPartitionActorName(const FPCGGridDescriptor& GridDescriptor, const FIntVector& GridCoords);
+
 #if WITH_EDITOR
 	bool IsInvalidForPCG() const { return bIsInvalidForPCG; }
 #endif
diff --git a/Engine/Plugins/PCG/Source/PCG/Public/Grid/PCGPartitionActorDesc.h b/Engine/Plugins/PCG/Source/PCG/Public/Grid/PCGPartitionActorDesc.h
--- a/Engine/Plugins/PCG/Source/PCG/Public/Grid/PCGPartitionActorDesc.h
+++ b/Engine/Plugins/PCG/Source/PCG/Public/Grid/PCGPartitionActorDesc.h
@@ -4,14 +4,18 @@
 
 #if WITH_EDITOR
 #include "CoreMinimal.h"
-#include "UObject/ObjectMacros.h"
 #include "Containers/Map.h"
+#include "Grid/PCGGridDescriptor.h"
+#include "UObject/ObjectMacros.h"
 #include "WorldPartition/ActorPartition/PartitionActorDesc.h"
 
 class FPCGPartitionActorDesc : public FPartitionActorDesc
 {
 public:
 	FPCGPartitionActorDesc() = default;
+
+	FPCGGridDescriptor GetGridDescriptor() const;
+	FIntVector GetGridCoord() const;
 protected:
 	virtual void Init(const AActor* InActor) override;
 	virtual uint32 GetSizeOf() const override { return sizeof(FPCGPartitionActorDesc); }
@@ -20,7 +24,8 @@ protected:
 private:
 	friend class FPCGActorAndComponentMapping;
 
-	bool bRequiresGuidFixup = false;
+	bool bUse2DGrid = true;
+	bool bRequiresUse2DGridFixup = false;
 	bool bInvalid = false;
 };
 #endif
\ No newline at end of file
diff --git a/Engine/Plugins/PCG/Source/PCG/Public/PCGActorAndComponentMapping.h b/Engine/Plugins/PCG/Source/PCG/Public/PCGActorAndComponentMapping.h
--- a/Engine/Plugins/PCG/Source/PCG/Public/PCGActorAndComponentMapping.h
+++ b/Engine/Plugins/PCG/Source/PCG/Public/PCGActorAndComponentMapping.h
@@ -4,6 +4,7 @@
 
 #include "PCGCommon.h"
 #include "Elements/PCGActorSelector.h"
+#include "Grid/PCGGridDescriptor.h"
 #include "Grid/PCGComponentOctree.h"
 #include "RuntimeGen/PCGRuntimeGenScheduler.h"
 
@@ -62,7 +63,7 @@ public:
 	/** Notify that we exited the Landscape edit mode. */
 	void NotifyLandscapeEditModeExited();
 
-	static void BuildPartitionActorRecords(APCGWorldActor* PCGWorldActor, UWorldPartition* WorldPartition, TMap<FPCGPartitionActorRecord, FGuid>& OutPartitionActorRecords, TSet<FGuid>& OutInvalidPartitionActors);
+	static void BuildPartitionActorRecords(APCGWorldActor* PCGWorldActor, UWorldPartition* WorldPartition, TMap<FPCGGridCellDescriptor, FGuid>& OutPartitionActorRecords, TSet<FGuid>& OutInvalidPartitionActors);
 #endif // WITH_EDITOR
 
 	/** Register a new PCG Component or update it. Returns true if it was added/updated. Thread safe */
@@ -89,15 +90,13 @@ public:
 	/** Return a copy of all the registered components. Thread safe */
 	TSet<UPCGComponent*> GetAllRegisteredComponents() const;
 
-	/** Retrieves a local component using grid size and grid coordinates, returns nullptr if no such component is found. */
-	UPCGComponent* GetLocalComponent(uint32 GridSize, const FIntVector& CellCoords, const UPCGComponent* InOriginalComponent, bool bRuntimeGenerated = false) const;
+	/** Retrieves a local component using grid descriptor and grid coordinates, returns nullptr if no such component is found. */
+	UPCGComponent* GetLocalComponent(const FPCGGridDescriptor& GridDescriptor, const FIntVector& CellCoords, const UPCGComponent* InOriginalComponent) const;
 
-	/** Retrieves a partition actor using grid size and grid coordinates, returns nullptr if no such partition actor is found. */
-	APCGPartitionActor* GetPartitionActor(uint32 GridSize, const FIntVector& CellCoords, bool bRuntimeGenerated = false) const;
+	APCGPartitionActor* GetPartitionActor(const FPCGGridDescriptor& GridDescriptor, const FIntVector& CellCoords) const;
 
 	/** Returns true if there is record of a partition actor living in a certain grid cell, regardless of whether or not it is loaded. */
-	bool DoesPartitionActorRecordExist(const FGuid& GridGuid, uint32 GridSize, const FIntVector& GridCoords) const;
-
+	bool DoesPartitionActorRecordExist(const FPCGGridDescriptor& GridDescriptor, const FIntVector& GridCoords) const;
 private:
 
 #if WITH_EDITOR
@@ -218,21 +217,13 @@ private:
 	FPCGComponentOctreeAndMap PartitionedOctree;
 
 	/** Mapping from grid size and grid coords to partition actor. We can only have 1 partition actor per grid cell. */
-	TMap<uint32, TMap<FIntVector, TObjectPtr<APCGPartitionActor>>> PartitionActorsMap;
+	TMap<FPCGGridDescriptor, TMap<FIntVector, TObjectPtr<APCGPartitionActor>>> PartitionActorsMap;
 	mutable FRWLock PartitionActorsMapLock;
 
-	/** Mapping from grid size and grid coords to RuntimeGen partition actor. We can only have 1 RuntimeGen partition actor per grid cell. */
-	TMap<uint32, TMap<FIntVector, TObjectPtr<APCGPartitionActor>>> RuntimeGenPartitionActorsMap;
-	mutable FRWLock RuntimeGenPartitionActorsMapLock;
-
 	/** Mapping between original components and its overlapping partition actors. */
 	TMap<const UPCGComponent*, TSet<TObjectPtr<APCGPartitionActor>>> ComponentToPartitionActorsMap;
 	mutable FRWLock ComponentToPartitionActorsMapLock;
 
-	/** Mapping between original components and its overlapping RuntimeGen partition actors. */
-	TMap<const UPCGComponent*, TSet<TObjectPtr<APCGPartitionActor>>> ComponentToRuntimeGenPartitionActorsMap;
-	mutable FRWLock ComponentToRuntimeGenPartitionActorsMapLock;
-
 	/** Components to be unregister at the next frame. cf. UnregisterComponent for a better understanding on why it is needed. */
 	TSet<UPCGComponent*> DelayedComponentToUnregister;
 	mutable FCriticalSection DelayedComponentToUnregisterLock;
@@ -273,7 +264,7 @@ private:
 	double LastPreviousActorDataCleanup = -1.0;
 
 	// Set of existing PCG Partition Actors (for World Partition worlds)
-	TMap<FPCGPartitionActorRecord, FGuid> PartitionActorRecords;
+	TMap<FPCGGridCellDescriptor, FGuid> PartitionActorRecords;
 	// Previously generated PCG Partition Actors to ignore (will prevent them from getting registered and mark them for deletion)
 	TSet<FGuid> InvalidPartitionActors;
 #endif // WITH_EDITOR
diff --git a/Engine/Plugins/PCG/Source/PCG/Public/PCGCommon.h b/Engine/Plugins/PCG/Source/PCG/Public/PCGCommon.h
--- a/Engine/Plugins/PCG/Source/PCG/Public/PCGCommon.h
+++ b/Engine/Plugins/PCG/Source/PCG/Public/PCGCommon.h
@@ -229,9 +229,6 @@ namespace PCGHiGenGrid
 	// Alias for array which is allocated on the stack (we have a strong idea of the max required elements).
 	using FSizeArray = TArray<uint32, TInlineAllocator<PCGHiGenGrid::NumGridValues>>;
 
-	// Alias for grid size to guid map allocated on the stack, which is unlikely to have a large number of elements.
-	using FSizeToGuidMap = TMap<uint32, FGuid, TInlineSetAllocator<32>>;
-
 	PCG_API bool IsValidGridSize(uint32 InGridSize);
 	PCG_API bool IsValidGrid(EPCGHiGenGrid InGrid);
 	PCG_API bool IsValidGridOrUninitialized(EPCGHiGenGrid InGrid);
@@ -346,9 +343,8 @@ namespace PCGQualityHelpers
 	PCG_API FName GetQualityPinLabel();
 }
 
-/** This record uniquely identifies a partition actor. */
-USTRUCT(BlueprintType)
-struct FPCGPartitionActorRecord
+USTRUCT(meta=(Deprecated = "5.5"))
+struct UE_DEPRECATED(5.5, "FPCGPartitionActorRecord is deprecated.") FPCGPartitionActorRecord
 {
 	GENERATED_BODY()
 		
diff --git a/Engine/Plugins/PCG/Source/PCG/Public/PCGComponent.h b/Engine/Plugins/PCG/Source/PCG/Public/PCGComponent.h
--- a/Engine/Plugins/PCG/Source/PCG/Public/PCGComponent.h
+++ b/Engine/Plugins/PCG/Source/PCG/Public/PCGComponent.h
@@ -414,6 +414,9 @@ public:
 	/** Returns true if the component is managed by the runtime generation system. Nothing else should generate or cleanup this component. */
 	bool IsManagedByRuntimeGenSystem() const { return GenerationTrigger == EPCGComponentGenerationTrigger::GenerateAtRuntime; }
 
+	/** Returns true if component should output on a 2D Grid */
+	bool Use2DGrid() const;
+
 	/* Responsibility of the PCG Partition Actor to mark is local */
 	void MarkAsLocalComponent() { bIsComponentLocal = true; }
 
diff --git a/Engine/Plugins/PCG/Source/PCG/Public/PCGGraph.h b/Engine/Plugins/PCG/Source/PCG/Public/PCGGraph.h
--- a/Engine/Plugins/PCG/Source/PCG/Public/PCGGraph.h
+++ b/Engine/Plugins/PCG/Source/PCG/Public/PCGGraph.h
@@ -206,6 +206,7 @@ public:
 	EPCGHiGenGrid GetDefaultGrid() const { ensure(IsHierarchicalGenerationEnabled()); return HiGenGridSize; }
 	uint32 GetDefaultGridSize() const;
 	bool IsHierarchicalGenerationEnabled() const { return bUseHierarchicalGeneration; }
+	bool Use2DGrid() const { return bUse2DGrid; }
 
 #if WITH_EDITORONLY_DATA
 	UPROPERTY(BlueprintReadWrite, EditAnywhere, Category = AssetInfo, AssetRegistrySearchable)
@@ -388,6 +389,9 @@ protected:
 	UPROPERTY(EditAnywhere, Category = Settings, meta = (DisplayName = "HiGen Default Grid Size", EditCondition = "bUseHierarchicalGeneration"))
 	EPCGHiGenGrid HiGenGridSize = EPCGHiGenGrid::Grid256;
 
+	UPROPERTY(EditAnywhere, Category = Settings, meta = (DisplayName = "2D Grid"))
+	bool bUse2DGrid = true;
+
 	/** Execution grid size for nodes. */
 	mutable TMap<const UPCGNode*, uint32> NodeToGridSize;
 	mutable FRWLock NodeToGridSizeLock;
diff --git a/Engine/Plugins/PCG/Source/PCG/Public/PCGSubsystem.h b/Engine/Plugins/PCG/Source/PCG/Public/PCGSubsystem.h
--- a/Engine/Plugins/PCG/Source/PCG/Public/PCGSubsystem.h
+++ b/Engine/Plugins/PCG/Source/PCG/Public/PCGSubsystem.h
@@ -6,6 +6,7 @@
 
 #include "PCGCommon.h"
 #include "PCGActorAndComponentMapping.h"
+#include "Grid/PCGGridDescriptor.h"
 #include "Grid/PCGComponentOctree.h"
 #include "UObject/ObjectKey.h"
 #include "Utils/PCGNodeVisualLogs.h"
@@ -219,20 +220,32 @@ public:
 	 * Call InFunc to all partition grid cells matching 'InGridSizes' and overlapping with 'InBounds'. 'InFunc' can schedule work or execute immediately.
 	 * 'InGridSizes' should be sorted in descending order. If 'bCanCreateActor' is true, it will create the partition actor at that cell if necessary.
 	 */
-	FPCGTaskId ForAllOverlappingCells(const FBox& InBounds, const PCGHiGenGrid::FSizeArray& InGridSizes, bool bCanCreateActor, const TArray<FPCGTaskId>& Dependencies, TFunctionRef<FPCGTaskId(APCGPartitionActor*, const FBox&)> InFunc) const;
+	UE_DEPRECATED(5.5, "Use version with UPCGComponent")
+	FPCGTaskId ForAllOverlappingCells(const FBox& InBounds, const PCGHiGenGrid::FSizeArray& InGridSizes, bool bCanCreateActor, const TArray<FPCGTaskId>& Dependencies, TFunctionRef<FPCGTaskId(APCGPartitionActor*, const FBox&)> InFunc) const { return InvalidPCGTaskId;  }
 
+	FPCGTaskId ForAllOverlappingCells(UPCGComponent* InPCGComponent, const FBox& InBounds, const PCGHiGenGrid::FSizeArray& InGridSizes, bool bCanCreateActor, const TArray<FPCGTaskId>& Dependencies, TFunctionRef<FPCGTaskId(APCGPartitionActor*, const FBox&)> InFunc) const;
+		
 	/** Immediately cleanup the local components associated with an original component. */
 	void CleanupLocalComponentsImmediate(UPCGComponent* InOriginalComponent, bool bRemoveComponents);
 
-	/** Retrieves a local component using grid size and grid coordinates, returns nullptr if no such component found. */
+	UE_DEPRECATED(5.5, "Use FPCGGridDescriptor version")
 	UPCGComponent* GetLocalComponent(uint32 GridSize, const FIntVector& CellCoords, const UPCGComponent* InOriginalComponent, bool bTransient = false) const;
 
-	/** Retrieves a registered partition actor using grid size and grid coordinates, returns nullptr if no such partition actor is found. */
+	UE_DEPRECATED(5.5, "Use FPCGGridDescriptor version")
 	APCGPartitionActor* GetRegisteredPCGPartitionActor(uint32 GridSize, const FIntVector& GridCoords, bool bRuntimeGenerated = false) const;
 
-	/** Creates a new partition actor if one does not already exist with the same grid size, coords, and generation mode. */
+	UE_DEPRECATED(5.5, "Use FPCGGridDescriptor verison")
 	APCGPartitionActor* FindOrCreatePCGPartitionActor(const FGuid& Guid, uint32 GridSize, const FIntVector& GridCoords, bool bRuntimeGenerated, bool bCanCreateActor = true) const;
 
+	/** Retrieves a local component using grid descriptor and grid coordinates, returns nullptr if no such component is found. */
+	UPCGComponent* GetLocalComponent(const FPCGGridDescriptor& GridDescriptor, const FIntVector& CellCoords, const UPCGComponent* InOriginalComponent) const;
+
+	/** Retrieves a registered partition actor using grid size and grid coordinates, returns nullptr if no such partition actor is found. */
+	APCGPartitionActor* GetRegisteredPCGPartitionActor(const FPCGGridDescriptor& GridDescriptor, const FIntVector& GridCoords) const;
+
+	/** Creates a new partition actor if one does not already exist with the same grid size, coords, and generation mode. */
+	APCGPartitionActor* FindOrCreatePCGPartitionActor(const FPCGGridDescriptor& GridDescriptor, const FIntVector& GridCoords, bool bCanCreateActor = true) const;
+	
 	/** True if graph cache debugging is enabled. */
 	bool IsGraphCacheDebuggingEnabled() const;
 
diff --git a/Engine/Plugins/PCG/Source/PCG/Public/PCGWorldActor.h b/Engine/Plugins/PCG/Source/PCG/Public/PCGWorldActor.h
--- a/Engine/Plugins/PCG/Source/PCG/Public/PCGWorldActor.h
+++ b/Engine/Plugins/PCG/Source/PCG/Public/PCGWorldActor.h
@@ -3,6 +3,7 @@
 #pragma once
 
 #include "PCGCommon.h"
+#include "Grid/PCGGridDescriptor.h"
 
 #include "GameFramework/Actor.h"
 #include "Misc/Guid.h"
@@ -35,21 +36,11 @@ public:
 
 	//~ Begin UObject Interface.
 #if WITH_EDITOR
-	virtual void PostLoad() override;
 	virtual void PostEditChangeProperty(FPropertyChangedEvent& PropertyChangedEvent) override;
 #endif
 	virtual void BeginDestroy() override;
 	//~ End UObject Interface.
 
-	/** Creates guids for unused grid sizes. */
-	void CreateGridGuidsIfNecessary(const PCGHiGenGrid::FSizeArray& InGridSizes, bool bAreGridsSerialized);
-
-	/** Returns the serialized grid GUIDs used for the partitioned actors, one per grid size. */
-	void GetSerializedGridGuids(PCGHiGenGrid::FSizeToGuidMap& OutSizeToGuidMap) const;
-
-	/** Returns the transient grid GUIDs used for the partitioned actors, one per grid size. */
-	void GetTransientGridGuids(PCGHiGenGrid::FSizeToGuidMap& OutSizeToGuidMap) const;
-
 	void MergeFrom(APCGWorldActor* OtherWorldActor);
 
 #if WITH_EDITOR
@@ -66,11 +57,12 @@ public:
 	UPROPERTY(BlueprintReadOnly, VisibleAnywhere, Category = CachedData, meta = (NoResetToDefault, DisplayName="Landscape Cache"))
 	TObjectPtr<UPCGLandscapeCache> LandscapeCacheObject = nullptr;
 
-	/** Disable creation of Partition Actors on the Z axis. Can improve performances if 3D partitioning is not needed. */
-	UPROPERTY(config, EditAnywhere, Category = GenerationSettings)
+#if WITH_EDITORONLY_DATA
+	/** This property was moved to the UPCGGraph but we keep it around to fixup existing data */
+	UE_DEPRECATED(5.5, "bUse2DGrid is deprecated")
+	UPROPERTY()
 	bool bUse2DGrid = true;
 
-#if WITH_EDITORONLY_DATA
 	/** Allows any currently active editor viewport to act as a Runtime Generation Source. */
 	UPROPERTY(EditAnywhere, Category = RuntimeGeneration)
 	bool bTreatEditorViewportAsGenerationSource = false;
@@ -86,19 +78,9 @@ private:
 	void OnPartitionGridSizeChanged();
 #endif
 
-	/** GUIDs of the serialized partitioned actor grids, one per grid size. */
-	UPROPERTY()
-	TMap<uint32, FGuid> GridGuids;
-	mutable FRWLock GridGuidsLock;
-
-	/** GUIDs of the transient partitioned actor grids, one per grid size. */
-	UPROPERTY(Transient)
-	TMap<uint32, FGuid> TransientGridGuids;
-	mutable FRWLock TransientGridGuidsLock;
-
 	/** Keeps a record of what grid cells contain a serialized partition actor for game worlds. */
 	UPROPERTY()
-	TSet<FPCGPartitionActorRecord> RuntimePartitionActorRecords;
+	TSet<FPCGGridCellDescriptor> RuntimePartitionActorRecords;
 };
 
 #if UE_ENABLE_INCLUDE_ORDER_DEPRECATED_IN_5_2
diff --git a/Engine/Plugins/PCG/Source/PCG/Public/RuntimeGen/PCGRuntimeGenScheduler.h b/Engine/Plugins/PCG/Source/PCG/Public/RuntimeGen/PCGRuntimeGenScheduler.h
--- a/Engine/Plugins/PCG/Source/PCG/Public/RuntimeGen/PCGRuntimeGenScheduler.h
+++ b/Engine/Plugins/PCG/Source/PCG/Public/RuntimeGen/PCGRuntimeGenScheduler.h
@@ -3,6 +3,7 @@
 #pragma once
 
 #include "PCGCommon.h"
+#include "Grid/PCGGridDescriptor.h"
 #include "UObject/WeakObjectPtr.h"
 
 class APCGPartitionActor;
@@ -53,6 +54,8 @@ protected:
 		uint32 GetGridSize() const { return Get<0>(); }
 		FIntVector GetGridCoords() const { return Get<1>(); }
 		UPCGComponent* GetOriginalComponent() const { return Get<2>(); }
+
+		FPCGGridDescriptor GetGridDescriptor() const;
 	};
 
 	/** Returns true if the scheduler should tick this frame. */
@@ -88,7 +91,7 @@ protected:
 	/** Grabs an empty RuntimeGen PA from the PartitionActorPool and initializes it at the given GridSize and GridCoords. If no PAs are available in the pool,
 	* the pool capacity will double and new PAs will be created.
 	*/
-	APCGPartitionActor* GetPartitionActorFromPool(uint32 GridSize, const FIntVector& GridCoords);
+	APCGPartitionActor* GetPartitionActorFromPool(const FPCGGridDescriptor& GridDescriptor, const FIntVector& GridCoords);
 
 	/** Adds Count new RuntimeGen PAs to the Runtime PA pool. */
 	void AddPartitionActorPoolCount(int32 Count);
@@ -96,9 +99,6 @@ protected:
 	/** Destroy all pooled partition actors and rebuild with the NewPoolSize. */
 	void ResetPartitionActorPoolToSize(uint32 NewPoolSize);
 
-	/** Create grid guids for the given component (if necessary). Only succeeds on partitioned original components. */
-	void CreateGridGuidsForComponent(UPCGComponent* InComponent);
-
 private:
 	/** Tracks the generated components managed by the RuntimeGenScheduler. For local components, this generation key will hold the original component.
 	* For non-partitioned components, the generation key should have unbounded grid size and (0, 0, 0) grid coordinates.
diff --git a/Engine/Source/Runtime/Core/Public/UObject/FortniteMainBranchObjectVersion.h b/Engine/Source/Runtime/Core/Public/UObject/FortniteMainBranchObjectVersion.h
--- a/Engine/Source/Runtime/Core/Public/UObject/FortniteMainBranchObjectVersion.h
+++ b/Engine/Source/Runtime/Core/Public/UObject/FortniteMainBranchObjectVersion.h
@@ -465,6 +465,9 @@ struct FFortniteMainBranchObjectVersion
 		// Changes for LevelInstance support in StaticLighting
 		LevelInstanceStaticLightingSupport,
 
+		// PCGGridDescriptor
+		PCGGridDescriptor,
+
 		// -----<new versions can be added above this line>-------------------------------------------------
 		VersionPlusOne,
 		LatestVersion = VersionPlusOne - 1
