Fix CIS

diff --git a/Engine/Plugins/PCG/Source/PCG/Private/Grid/PCGPartitionActor.cpp b/Engine/Plugins/PCG/Source/PCG/Private/Grid/PCGPartitionActor.cpp
--- a/Engine/Plugins/PCG/Source/PCG/Private/Grid/PCGPartitionActor.cpp
+++ b/Engine/Plugins/PCG/Source/PCG/Private/Grid/PCGPartitionActor.cpp
@@ -73,7 +73,7 @@ void APCGPartitionActor::PostLoad()
 #endif
 
 #if WITH_EDITORONLY_DATA
-	// Prior to this version bUse2DGrid was slave of the PCGWorldActor so make sure we update it one last time upon registration
+	// Prior to this version bUse2DGrid was dependant of the PCGWorldActor so make sure we update it one last time upon registration
 	if (GetLinkerCustomVersion(FFortniteMainBranchObjectVersion::GUID) < FFortniteMainBranchObjectVersion::PCGGridDescriptor)
 	{
 		bRequiresUse2DGridFixup = true;
