[PCG] Combine the seed with the seed of the first point for Random Choice for every input data.

* Without that, multi data input will all select the same indexes.
* Add unit tests for that case.

#rnx

diff --git a/Engine/Plugins/PCG/Source/PCG/Private/Elements/PCGRandomChoice.cpp b/Engine/Plugins/PCG/Source/PCG/Private/Elements/PCGRandomChoice.cpp
--- a/Engine/Plugins/PCG/Source/PCG/Private/Elements/PCGRandomChoice.cpp
+++ b/Engine/Plugins/PCG/Source/PCG/Private/Elements/PCGRandomChoice.cpp
@@ -6,6 +6,7 @@
 #include "PCGParamData.h"
 #include "PCGPin.h"
 #include "Data/PCGPointData.h"
+#include "Helpers/PCGHelpers.h"
 
 #define LOCTEXT_NAMESPACE "PCGRandomChoiceElement"
 
@@ -93,6 +94,7 @@ bool FPCGRandomChoiceElement::ExecuteInternal(FPCGContext* Context) const
 		}
 
 		const TArray<FPCGPoint>& InPoints = InputPointData->GetPoints();
+
 		int NumOfElementsToKeep = 0;
 
 		if (Settings->bFixedMode)
@@ -151,7 +153,14 @@ bool FPCGRandomChoiceElement::ExecuteInternal(FPCGContext* Context) const
 		// * Pick a number in [0, n[
 		// * Pick a number in[0, n - 1[, then for each previously selected number, if it's larger, add +1
 		// It's O(n) cpu + O(n) memory vs O(s^2) cpu + O(s) memory. (n = total number of points, s = number of points to keep)
-		FRandomStream RandStream(Context->GetSeed());
+		int32 Seed = Context->GetSeed();
+		if (ensure(!InPoints.IsEmpty()))
+		{
+			// Combine the seed with the first point so that multiple data produces different results.
+			Seed = PCGHelpers::ComputeSeed(Seed, InPoints[0].Seed);
+		}
+
+		FRandomStream RandStream(Seed);
 		TArray<int32> ShuffledIndexes;
 		ShuffledIndexes.Reserve(InPoints.Num());
 
diff --git a/Engine/Plugins/PCG/Source/PCG/Private/Tests/Elements/PCGRandomChoiceTest.cpp b/Engine/Plugins/PCG/Source/PCG/Private/Tests/Elements/PCGRandomChoiceTest.cpp
--- a/Engine/Plugins/PCG/Source/PCG/Private/Tests/Elements/PCGRandomChoiceTest.cpp
+++ b/Engine/Plugins/PCG/Source/PCG/Private/Tests/Elements/PCGRandomChoiceTest.cpp
@@ -10,15 +10,17 @@
 
 #include "Elements/PCGRandomChoice.h"
 
-IMPLEMENT_CUSTOM_SIMPLE_AUTOMATION_TEST(FPCGRandomChoiceTest_Fixed, FPCGTestBaseClass, "Plugins.PCG.RandomChoice,Fixed", PCGTestsCommon::TestFlags)
+IMPLEMENT_CUSTOM_SIMPLE_AUTOMATION_TEST(FPCGRandomChoiceTest_Fixed, FPCGTestBaseClass, "Plugins.PCG.RandomChoice.Fixed", PCGTestsCommon::TestFlags)
 IMPLEMENT_CUSTOM_SIMPLE_AUTOMATION_TEST(FPCGRandomChoiceTest_Ratio, FPCGTestBaseClass, "Plugins.PCG.RandomChoice.Ratio", PCGTestsCommon::TestFlags)
 IMPLEMENT_CUSTOM_SIMPLE_AUTOMATION_TEST(FPCGRandomChoiceTest_SelectNone, FPCGTestBaseClass, "Plugins.PCG.RandomChoice.SelectNone", PCGTestsCommon::TestFlags)
 IMPLEMENT_CUSTOM_SIMPLE_AUTOMATION_TEST(FPCGRandomChoiceTest_SelectAll, FPCGTestBaseClass, "Plugins.PCG.RandomChoice.SelectAll", PCGTestsCommon::TestFlags)
 IMPLEMENT_CUSTOM_SIMPLE_AUTOMATION_TEST(FPCGRandomChoiceTest_NoDiscard, FPCGTestBaseClass, "Plugins.PCG.RandomChoice.NoDiscard", PCGTestsCommon::TestFlags)
+IMPLEMENT_CUSTOM_SIMPLE_AUTOMATION_TEST(FPCGRandomChoiceTest_MultiData_SameSeed, FPCGTestBaseClass, "Plugins.PCG.RandomChoice.MultiData.SameSeed", PCGTestsCommon::TestFlags)
+IMPLEMENT_CUSTOM_SIMPLE_AUTOMATION_TEST(FPCGRandomChoiceTest_MultiData_DifferentSeed, FPCGTestBaseClass, "Plugins.PCG.RandomChoice.MultiData.DifferentSeed", PCGTestsCommon::TestFlags)
 
 namespace PCGRandomChoiceTest
 {
-	UPCGPointData* CreateInputPointData(FPCGContext* Context, const int NumPoints)
+	UPCGPointData* CreateInputPointData(FPCGContext* Context, const int NumPoints, const int Seed = 42)
 	{
 		check(Context);
 
@@ -30,6 +32,7 @@ namespace PCGRandomChoiceTest
 		{
 			// Store the index in the density
 			Points[i].Density = i;
+			Points[i].Seed = Seed;
 		}
 
 		FPCGTaggedData& InputData = Context->InputData.TaggedData.Emplace_GetRef();
@@ -250,4 +253,107 @@ bool FPCGRandomChoiceTest_NoDiscard::RunTest(const FString& Parameters)
 
 	return true;
 }
+
+bool FPCGRandomChoiceTest_MultiData_SameSeed::RunTest(const FString& Parameters)
+{
+	PCGTestsCommon::FTestData TestData;
+	UPCGRandomChoiceSettings* Settings = PCGTestsCommon::GenerateSettings<UPCGRandomChoiceSettings>(TestData);
+	check(Settings);
+
+	constexpr int NumOfPoints = 100;
+	constexpr int ExpectedNumElementsChosen = 7;
+	constexpr int ExpectedNumElementsDiscarded = 93;
+
+	Settings->bFixedMode = true;
+	Settings->FixedNumber = ExpectedNumElementsChosen;
+
+	TUniquePtr<FPCGContext> Context = TestData.InitializeTestContext();
+	const UPCGPointData* FirstInputPointData = PCGRandomChoiceTest::CreateInputPointData(Context.Get(), NumOfPoints);
+	const UPCGPointData* SecondInputPointData = PCGRandomChoiceTest::CreateInputPointData(Context.Get(), NumOfPoints);
+
+	FPCGElementPtr TestElement = TestData.Settings->GetElement();
+	while (!TestElement->Execute(Context.Get())) {}
+
+	TArray<FPCGTaggedData> ChosenOutputTagged = Context->OutputData.GetInputsByPin(PCGRandomChoiceConstants::ChosenPointsLabel);
+	TArray<FPCGTaggedData> DiscardedOutputTagged = Context->OutputData.GetInputsByPin(PCGRandomChoiceConstants::DiscardedPointsLabel);
+
+	const UPCGPointData* FirstChosenOutputData = ChosenOutputTagged.Num() == 2 ? Cast<const UPCGPointData>(ChosenOutputTagged[0].Data) : nullptr;
+	const UPCGPointData* FirstDiscardedOutputData = DiscardedOutputTagged.Num() == 2 ? Cast<const UPCGPointData>(DiscardedOutputTagged[0].Data) : nullptr;
+
+	const UPCGPointData* SecondChosenOutputData = ChosenOutputTagged.Num() == 2 ? Cast<const UPCGPointData>(ChosenOutputTagged[1].Data) : nullptr;
+	const UPCGPointData* SecondDiscardedOutputData = DiscardedOutputTagged.Num() == 2 ? Cast<const UPCGPointData>(DiscardedOutputTagged[1].Data) : nullptr;
+
+	UTEST_TRUE("There is 2 point data in chosen points", FirstChosenOutputData && SecondChosenOutputData);
+	UTEST_TRUE("There is 2 point data in discarded points", FirstDiscardedOutputData && SecondDiscardedOutputData);
+
+	UTEST_EQUAL("There is the right number of points in first chosen", FirstChosenOutputData->GetPoints().Num(), ExpectedNumElementsChosen);
+	UTEST_EQUAL("There is the right number of points in second chosen", SecondChosenOutputData->GetPoints().Num(), ExpectedNumElementsChosen);
+
+	UTEST_EQUAL("There is the right number of points in first discarded", FirstDiscardedOutputData->GetPoints().Num(), ExpectedNumElementsDiscarded);
+	UTEST_EQUAL("There is the right number of points in second discarded", SecondDiscardedOutputData->GetPoints().Num(), ExpectedNumElementsDiscarded);
+
+	UTEST_TRUE("All points are there and in the right order for first data", PCGRandomChoiceTest::VerifyAllPointsThere(NumOfPoints, FirstChosenOutputData, FirstDiscardedOutputData));
+	UTEST_TRUE("All points are there and in the right order for second data", PCGRandomChoiceTest::VerifyAllPointsThere(NumOfPoints, SecondChosenOutputData, SecondDiscardedOutputData));
+
+	// If the point seeds are the same, chosen points should be exactly the same
+	for (int32 i = 0; i < ExpectedNumElementsChosen; ++i)
+	{
+		UTEST_TRUE(*FString::Printf(TEXT("Point %d is the same for both chosen inputs"), i), FirstChosenOutputData->GetPoints()[i].Density == SecondChosenOutputData->GetPoints()[i].Density);
+	}
+
+	return true;
+}
+
+bool FPCGRandomChoiceTest_MultiData_DifferentSeed::RunTest(const FString& Parameters)
+{
+	PCGTestsCommon::FTestData TestData;
+	UPCGRandomChoiceSettings* Settings = PCGTestsCommon::GenerateSettings<UPCGRandomChoiceSettings>(TestData);
+	check(Settings);
+
+	constexpr int NumOfPoints = 100;
+	constexpr int ExpectedNumElementsChosen = 5;
+	constexpr int ExpectedNumElementsDiscarded = 95;
+
+	Settings->bFixedMode = true;
+	Settings->FixedNumber = ExpectedNumElementsChosen;
+
+	TUniquePtr<FPCGContext> Context = TestData.InitializeTestContext();
+	const UPCGPointData* FirstInputPointData = PCGRandomChoiceTest::CreateInputPointData(Context.Get(), NumOfPoints);
+	const UPCGPointData* SecondInputPointData = PCGRandomChoiceTest::CreateInputPointData(Context.Get(), NumOfPoints, 4653465);
+
+	FPCGElementPtr TestElement = TestData.Settings->GetElement();
+	while (!TestElement->Execute(Context.Get())) {}
+
+	TArray<FPCGTaggedData> ChosenOutputTagged = Context->OutputData.GetInputsByPin(PCGRandomChoiceConstants::ChosenPointsLabel);
+	TArray<FPCGTaggedData> DiscardedOutputTagged = Context->OutputData.GetInputsByPin(PCGRandomChoiceConstants::DiscardedPointsLabel);
+
+	const UPCGPointData* FirstChosenOutputData = ChosenOutputTagged.Num() == 2 ? Cast<const UPCGPointData>(ChosenOutputTagged[0].Data) : nullptr;
+	const UPCGPointData* FirstDiscardedOutputData = DiscardedOutputTagged.Num() == 2 ? Cast<const UPCGPointData>(DiscardedOutputTagged[0].Data) : nullptr;
+
+	const UPCGPointData* SecondChosenOutputData = ChosenOutputTagged.Num() == 2 ? Cast<const UPCGPointData>(ChosenOutputTagged[1].Data) : nullptr;
+	const UPCGPointData* SecondDiscardedOutputData = DiscardedOutputTagged.Num() == 2 ? Cast<const UPCGPointData>(DiscardedOutputTagged[1].Data) : nullptr;
+
+	UTEST_TRUE("There is 2 point data in chosen points", FirstChosenOutputData && SecondChosenOutputData);
+	UTEST_TRUE("There is 2 point data in discarded points", FirstDiscardedOutputData && SecondDiscardedOutputData);
+
+	UTEST_EQUAL("There is the right number of points in first chosen", FirstChosenOutputData->GetPoints().Num(), ExpectedNumElementsChosen);
+	UTEST_EQUAL("There is the right number of points in second chosen", SecondChosenOutputData->GetPoints().Num(), ExpectedNumElementsChosen);
+
+	UTEST_EQUAL("There is the right number of points in first discarded", FirstDiscardedOutputData->GetPoints().Num(), ExpectedNumElementsDiscarded);
+	UTEST_EQUAL("There is the right number of points in second discarded", SecondDiscardedOutputData->GetPoints().Num(), ExpectedNumElementsDiscarded);
+
+	UTEST_TRUE("All points are there and in the right order for first data", PCGRandomChoiceTest::VerifyAllPointsThere(NumOfPoints, FirstChosenOutputData, FirstDiscardedOutputData));
+	UTEST_TRUE("All points are there and in the right order for second data", PCGRandomChoiceTest::VerifyAllPointsThere(NumOfPoints, SecondChosenOutputData, SecondDiscardedOutputData));
+
+	// If point seeds are different, chosen points should be different. Verify that at least one point is different.
+	bool bIsDifferent = false;
+	for (int32 i = 0; i < ExpectedNumElementsChosen && !bIsDifferent; ++i)
+	{
+		bIsDifferent = FirstChosenOutputData->GetPoints()[i].Density != SecondChosenOutputData->GetPoints()[i].Density;
+	}
+
+	UTEST_TRUE("Chosen points are different", bIsDifferent);
+
+	return true;
+}
 #endif // WITH_EDITOR
