[PCG] Elevation Isolines improvements
* Add linear option for splines
* Add option to project on surface normal.
* Add option to add tags to the output
* Make the control points locally at Z=0 to avoid approximation issues while sampling, and putting the spline itself at Z=Elevation

#rnx

diff --git a/Engine/Plugins/PCG/Source/PCG/Private/Elements/PCGElevationIsolines.cpp b/Engine/Plugins/PCG/Source/PCG/Private/Elements/PCGElevationIsolines.cpp
--- a/Engine/Plugins/PCG/Source/PCG/Private/Elements/PCGElevationIsolines.cpp
+++ b/Engine/Plugins/PCG/Source/PCG/Private/Elements/PCGElevationIsolines.cpp
@@ -99,16 +99,17 @@ bool FPCGElevationIsolinesElement::ExecuteInternal(FPCGContext* InContext) const
 		const int32 CellYCount = CellMaxY - CellMinY;
 		const int32 CellCount = CellXCount * CellYCount;
 
-		TArray<double> Heightmap;
-		Heightmap.SetNumUninitialized(CellCount);
+		TArray<TTuple<double, FQuat>> HeightmapWithRotation;
+		HeightmapWithRotation.SetNumUninitialized(CellCount);
 
 		int32 NumPointsWritten = 0;
 
 		FPCGProjectionParams ProjectionParams{};
 		ProjectionParams.bProjectPositions = true;
+		ProjectionParams.bProjectRotations = Settings->bProjectSurfaceNormal;
 
 		// TODO: Time-slice
-		FPCGAsync::AsyncProcessingOneToOneEx(&InContext->AsyncState, Heightmap.Num(), []() {}, [CellYCount, Settings, Origin, CellSize, SurfaceData, &Heightmap, &ProjectionParams](int32 ReadIndex, int32 WriteIndex)
+		FPCGAsync::AsyncProcessingOneToOneEx(&InContext->AsyncState, HeightmapWithRotation.Num(), []() {}, [CellYCount, Settings, Origin, CellSize, SurfaceData, &HeightmapWithRotation, &ProjectionParams](int32 ReadIndex, int32 WriteIndex)
 		{
 			const int32 CellX = ReadIndex / CellYCount;
 			const int32 CellY = ReadIndex % CellYCount;
@@ -117,51 +118,170 @@ bool FPCGElevationIsolinesElement::ExecuteInternal(FPCGContext* InContext) const
 
 			FPCGPoint Projected;
 			bool bSuccess = SurfaceData->ProjectPoint(FTransform(Cell), FBox(EForceInit::ForceInit), ProjectionParams, Projected, nullptr);
-			Heightmap[WriteIndex] = bSuccess ? Projected.Transform.GetLocation().Z : std::numeric_limits<double>::min();
+			HeightmapWithRotation[WriteIndex].Get<0>() = bSuccess ? Projected.Transform.GetLocation().Z : std::numeric_limits<double>::min();
+
+			if (ProjectionParams.bProjectRotations)
+			{
+				HeightmapWithRotation[WriteIndex].Get<1>() = bSuccess ? Projected.Transform.GetRotation() : FQuat::Identity;
+			}
 
 		}, /*bEnableTimeSlicing=*/false);
 
-		auto GetElevation = [CellYCount, &Heightmap](int32 X, int32 Y) -> double
+		auto GetElevation = [CellYCount, &HeightmapWithRotation](int32 X, int32 Y) -> double
 		{
-			return Heightmap[X * CellYCount + Y];
+			return HeightmapWithRotation[X * CellYCount + Y].Get<0>();
 		};
 
 		// TODO: Async + Support non-Z up cases.
 		// Run the Marching Squares algorithm on each isoline we want.
+		int32 ElevationCount = 0;
 		for (double Elevation = Settings->ElevationStart; Elevation < Settings->ElevationEnd; Elevation += Settings->ElevationIncrement)
 		{
 			for (const PCGSpatialAlgo::FPCGMarchingSquareResult& Result : PCGSpatialAlgo::MarchingSquares(CellXCount, CellYCount, Elevation, GetElevation, /*bUseLinearInterpolation=*/true))
 			{
+				// We can't do anything with less than 2 points.
+				if (Result.LinkedGridCoordinates.Num() < 2)
+				{
+					continue;
+				}
+
 				TRACE_CPUPROFILER_EVENT_SCOPE(FPCGElevationIsolinesElement::Execute::CreateOutputData);
 
-				auto TransformFromGridCoordinate = [&Origin, &CellSize, Elevation](const FVector2D& GridCoordinate)
+				auto TransformFromGridCoordinate = [&Origin, &CellSize, Elevation](const FVector2D& GridCoordinates) -> FTransform
 				{
-					return FTransform(FVector(Origin.X + CellSize.X * GridCoordinate.X, Origin.Y + CellSize.Y * GridCoordinate.Y, Elevation));
+					return FTransform(FVector(Origin.X + CellSize.X * GridCoordinates.X, Origin.Y + CellSize.Y * GridCoordinates.Y, Elevation));
 				};
 
+				// Each point coordinate will have an integer coordinate and a non-integer one. (That's how Marching Square is building it).
+				// So to interpolate the normals, we find the two corner of the cell where the edge is.
+				// Floor and Ceil will return the same value if X/Y is an integer. So depending on the fractional value of X or Y, it will give us
+				// the two wanted corners. Then we do a Slerp with this ratio between the 2 rotations computed before (when we did the projection).
+				auto ReconstructPointRotation = [CellYCount, HeightmapWithRotation](const FVector2D& GridCoordinates) -> FQuat
+				{
+					const int32 Index1 = FMath::FloorToInt(GridCoordinates.X) * CellYCount + FMath::FloorToInt(GridCoordinates.Y);
+					const int32 Index2 = FMath::CeilToInt(GridCoordinates.X) * CellYCount + FMath::CeilToInt(GridCoordinates.Y);
+
+					FQuat Rotation = HeightmapWithRotation[Index1].Get<1>();
+
+					if (Index1 != Index2)
+					{
+						const double RatioX = FMath::Frac(GridCoordinates.X);
+						const double RatioY = FMath::Frac(GridCoordinates.Y);
+						const double Ratio = FMath::IsNearlyZero(RatioX) ? RatioY : RatioX;
+
+						Rotation = FQuat::Slerp(Rotation, HeightmapWithRotation[Index2].Get<1>(), Ratio);
+					}
+
+					return Rotation;
+				};
+
+				// Because we can have points that are very close to each other (the edge is in the corner of the cell), we merge them (by doing an average in the middle)
+				// Threshold = (Resolution * 0.5)^2
+				const double Threshold = (Settings->Resolution * Settings->Resolution * 0.25);
+				auto MergePointsIfTooClose = [Threshold](FTransform& PreviousTransform, const FTransform& CurrentTransform) -> bool
+				{
+					if (FVector::DistSquared(PreviousTransform.GetLocation(), CurrentTransform.GetLocation()) < Threshold)
+					{
+						FQuat LerpQuat = FQuat::Slerp(PreviousTransform.GetRotation(), CurrentTransform.GetRotation(), 0.5);
+						PreviousTransform.SetRotation(LerpQuat);
+						PreviousTransform.LerpTranslationScale3D(PreviousTransform, CurrentTransform, ScalarRegister(0.5));
+						return true;
+					}
+					else
+					{
+						return false;
+					}
+				};
+
+				// Process all the transforms
+				TArray<FTransform> FinalTransforms;
+				FinalTransforms.Reserve(Result.LinkedGridCoordinates.Num());
+
+				for (int32 Index = 0; Index < Result.LinkedGridCoordinates.Num(); ++Index)
+				{
+					const FVector2D GridCoordinates = Result.LinkedGridCoordinates[Index];
+					FTransform Transform = TransformFromGridCoordinate(GridCoordinates);
+
+					// Reconstruct the rotation from the normal if needed
+					if (Settings->bProjectSurfaceNormal)
+					{
+						Transform.SetRotation(ReconstructPointRotation(GridCoordinates));
+					}
+
+					// If we have more than one transform, try to merge with the previous point if they are too close.
+					// The result of the merge will be in the first transform passed as parameter (update in place).
+					// INFO: If more than 2 points are close together, it will first average the first two then average the next point (and again).
+					// It's unlikely to happen though.
+					if (!FinalTransforms.IsEmpty())
+					{
+						FTransform& PreviousTransform = FinalTransforms.Last();
+						const bool bHasMergedWithPrevious = MergePointsIfTooClose(PreviousTransform, Transform);
+						bool bSkipPoint = bHasMergedWithPrevious;
+
+						// Also if we are at the end, and we are a closed loop, try to merge with the first transform.
+						if (Result.bClosed && Index == Result.LinkedGridCoordinates.Num() - 1)
+						{
+							const FTransform& CurrentTransform = bHasMergedWithPrevious ? PreviousTransform : Transform;
+							if (MergePointsIfTooClose(FinalTransforms[0], CurrentTransform))
+							{
+								bSkipPoint = true;
+
+								// If we merged the last and first point, check if we merged it with the penultimate one, if so, pop it.
+								if (bHasMergedWithPrevious)
+								{
+									FinalTransforms.Pop();
+								}
+							}
+						}
+
+						if (bSkipPoint)
+						{
+							continue;
+						}
+					}
+
+					FinalTransforms.Add(std::move(Transform));
+				}
+
+				// Make sure that after the merge, we still have more than 2 points, otherwise, just discard this line.
+				if (FinalTransforms.Num() < 2)
+				{
+					continue;
+				}
+
+				// And finally create the data.
 				FPCGTaggedData& OutputData = InContext->OutputData.TaggedData.Emplace_GetRef(InputData);
 				OutputData.Pin = PCGPinConstants::DefaultOutputLabel;
 
+				if (Settings->bAddTagOnOutputForSameElevation)
+				{
+					OutputData.Tags.Add(FString::Printf(TEXT("ElevationIdx_%d"), ElevationCount));
+				}
+
 				if (Settings->bOutputAsSpline)
 				{
 					TArray<FSplinePoint> SplinePoints;
-					SplinePoints.Reserve(Result.LinkedGridCoordinates.Num());
+					SplinePoints.Reserve(FinalTransforms.Num());
 
-					for (int32 Index = 0; Index < Result.LinkedGridCoordinates.Num(); ++Index)
+					for (const FTransform& Transform : FinalTransforms)
 					{
-						const FTransform Transform = TransformFromGridCoordinate(Result.LinkedGridCoordinates[Index]);
-
-						SplinePoints.Emplace(static_cast<float>(Index),
-							Transform.GetLocation(),
+						// Control points will be locally at 0 in Z.
+						// It is to prevent any approximation errors while interpolating on the spline.
+						// The spline itself will be a Z = Elevation.
+						FVector Location = Transform.GetLocation();
+						Location.Z = 0;
+
+						SplinePoints.Emplace(static_cast<float>(SplinePoints.Num()),
+							Location,
 							FVector::ZeroVector,
 							FVector::ZeroVector,
 							Transform.GetRotation().Rotator(),
 							Transform.GetScale3D(),
-							ESplinePointType::Curve);
+							Settings->bLinearSpline ? ESplinePointType::Linear : ESplinePointType::CurveCustomTangent);
 					}
 
 					UPCGSplineData* OutSplineData = NewObject<UPCGSplineData>();
-					OutSplineData->Initialize(SplinePoints, Result.bClosed, FTransform::Identity);
+					OutSplineData->Initialize(SplinePoints, Result.bClosed, FTransform(FVector::UpVector * Elevation));
 					OutputData.Data = OutSplineData;
 				}
 				else
@@ -171,14 +291,16 @@ bool FPCGElevationIsolinesElement::ExecuteInternal(FPCGContext* InContext) const
 					OutPoints.Reserve(Result.LinkedGridCoordinates.Num());
 					OutputData.Data = OutPointData;
 
-					for (const FVector2D& GridCoordinate : Result.LinkedGridCoordinates)
+					for (const FTransform& Transform : FinalTransforms)
 					{
 						FPCGPoint& Point = OutPoints.Emplace_GetRef();
-						Point.Transform = TransformFromGridCoordinate(GridCoordinate);
+						Point.Transform = Transform;
 						Point.SetLocalBounds(CellBounds);
 					}
 				}
 			}
+
+			++ElevationCount;
 		}
 	}
 
diff --git a/Engine/Plugins/PCG/Source/PCG/Public/Elements/PCGElevationIsolines.h b/Engine/Plugins/PCG/Source/PCG/Public/Elements/PCGElevationIsolines.h
--- a/Engine/Plugins/PCG/Source/PCG/Public/Elements/PCGElevationIsolines.h
+++ b/Engine/Plugins/PCG/Source/PCG/Public/Elements/PCGElevationIsolines.h
@@ -46,9 +46,21 @@ public:
 	UPROPERTY(BlueprintReadWrite, EditAnywhere, Category = Settings, meta = (PCG_Overridable))
 	double Resolution = 100.0;
 
+	/** Can add a tag (integer) to group output data that are at the same elevation. */
+	UPROPERTY(BlueprintReadWrite, EditAnywhere, Category = Settings, meta = (PCG_Overridable))
+	bool bAddTagOnOutputForSameElevation = false;
+
+	/** Option to either have Z up or project the surface normal at this position (similar to project rotations on the projection node). */
+	UPROPERTY(BlueprintReadWrite, EditAnywhere, Category = Settings, meta = (PCG_Overridable))
+	bool bProjectSurfaceNormal = false;
+
 	/** Will output splines rather than points. */
 	UPROPERTY(BlueprintReadWrite, EditAnywhere, Category = Settings)
 	bool bOutputAsSpline = false;
+
+	/** Spline can either be curved or linear. */
+	UPROPERTY(BlueprintReadWrite, EditAnywhere, Category = Settings, meta = (PCG_Overridable, EditCondition = "bOutputAsSpline", EditConditionHides))
+	bool bLinearSpline = false;
 };
 
 class FPCGElevationIsolinesElement : public IPCGElement
