[PCG] Enable extraction of more metadata from Ray Hits

+ Impact (T/F)
+ Impact Point
+ Impact Normal
+ Impact Distance
+ Local Impact Point
+ Physical Material
+ Render Material
+ Static Mesh
+ UV Coords
+ Element Index
+ Face Index

diff --git a/Engine/Plugins/PCG/Source/PCG/Private/Data/PCGWorldData.cpp b/Engine/Plugins/PCG/Source/PCG/Private/Data/PCGWorldData.cpp
--- a/Engine/Plugins/PCG/Source/PCG/Private/Data/PCGWorldData.cpp
+++ b/Engine/Plugins/PCG/Source/PCG/Private/Data/PCGWorldData.cpp
@@ -223,6 +223,7 @@ bool UPCGWorldRayHitData::SamplePoint(const FTransform& InTransform, const FBox&
 	FCollisionQueryParams Params; // TODO: apply properties from the settings when/if they exist
 	Params.bTraceComplex = QueryParams.bTraceComplex;
 	Params.bReturnPhysicalMaterial = QueryParams.bGetReferenceToPhysicalMaterial;
+	Params.bReturnFaceIndex = QueryParams.bGetFaceIndex || QueryParams.bGetUVCoords;
 
 	FVector RayOrigin = InTransform.GetLocation() - ((InTransform.GetLocation() - QueryParams.RayOrigin) | QueryParams.RayDirection) * QueryParams.RayDirection;
 	FVector RayEnd = RayOrigin + QueryParams.RayDirection * QueryParams.RayLength;
@@ -232,29 +233,12 @@ bool UPCGWorldRayHitData::SamplePoint(const FTransform& InTransform, const FBox&
 	{
 		TOptional<FHitResult> HitResult = PCGWorldQueryHelpers::FilterRayHitResults(&QueryParams, OriginatingComponent, Hits);
 
-		// TODO: Pre-create attributes within caller or find a better solution than per point
-		PCGWorldQueryHelpers::ApplyRayHitMetadata(HitResult, QueryParams, OutPoint, OutMetadata);
-		const PCGMetadataEntryKey MetadataEntry = OutPoint.MetadataEntry;
-
 		if (HitResult.IsSet())
 		{
 			const FHitResult& Hit = HitResult.GetValue();
-
-			// Finally, fill in OutPoint - we're done
 			OutPoint = FPCGPoint(PCGWorldQueryHelpers::GetOrthonormalImpactTransform(Hit), 1.0f, UPCGBlueprintHelpers::ComputeSeedFromPosition(Hit.Location));
-			OutPoint.MetadataEntry = MetadataEntry;
-
-			const bool bApplyMetadataFromLandscape = QueryParams.bApplyMetadataFromLandscape && Hit.GetActor() && Hit.GetActor()->IsA<ALandscapeProxy>();
-
-			// TODO: generalize for other sources of metadata?
-			if (bApplyMetadataFromLandscape && OutMetadata && World->GetSubsystem<UPCGSubsystem>())
-			{
-				if (UPCGLandscapeCache* LandscapeCache = World->GetSubsystem<UPCGSubsystem>()->GetLandscapeCache())
-				{
-					LandscapeCache->SampleMetadataOnPoint(Cast<ALandscapeProxy>(Hit.Component->GetOwner()), OutPoint, OutMetadata);
-				}
-			}
-
+			// TODO: Pre-create attributes within caller or find a better solution than per point
+			PCGWorldQueryHelpers::ApplyRayHitMetadata(Hit, QueryParams, OutPoint, OutMetadata, World, /*bShouldCreateAttributes=*/true);
 			return true;
 		}
 	}
diff --git a/Engine/Plugins/PCG/Source/PCG/Private/Elements/PCGWorldRaycast.cpp b/Engine/Plugins/PCG/Source/PCG/Private/Elements/PCGWorldRaycast.cpp
--- a/Engine/Plugins/PCG/Source/PCG/Private/Elements/PCGWorldRaycast.cpp
+++ b/Engine/Plugins/PCG/Source/PCG/Private/Elements/PCGWorldRaycast.cpp
@@ -27,6 +27,7 @@ namespace PCGWorldRaycastElementConstants
 }
 
 UPCGWorldRaycastElementSettings::UPCGWorldRaycastElementSettings()
+	: WorldQueryParams()
 {
 	OriginInputAttribute.SetPointProperty(EPCGPointProperties::Position);
 	EndPointAttribute.SetPointProperty(EPCGPointProperties::Position);
@@ -108,6 +109,7 @@ bool FPCGWorldRaycastElement::PrepareDataInternal(FPCGContext* InContext) const
 
 		OutState.CollisionQueryParams.bTraceComplex = Settings->WorldQueryParams.bTraceComplex;
 		OutState.CollisionQueryParams.bReturnPhysicalMaterial = Settings->WorldQueryParams.bGetReferenceToPhysicalMaterial;
+		OutState.CollisionQueryParams.bReturnFaceIndex = Settings->WorldQueryParams.bGetFaceIndex || Settings->WorldQueryParams.bGetUVCoords;
 		OutState.CollisionObjectQueryParams = FCollisionObjectQueryParams(Settings->WorldQueryParams.CollisionChannel);
 
 		return EPCGTimeSliceInitResult::Success;
@@ -298,6 +300,13 @@ bool FPCGWorldRaycastElement::ExecuteInternal(FPCGContext* InContext) const
 
 		const FPCGWorldRaycastQueryParams& WorldQueryParams = Settings->WorldQueryParams;
 
+		if (!PCGWorldQueryHelpers::CreateRayHitAttributes(WorldQueryParams, OutMetadata))
+		{
+			PCGLog::LogWarningOnGraph(LOCTEXT("UnableToCreateAllAttributes", "One or more attributes were unable to be created."), Context);
+		}
+
+		bool bAttributeSuccess = true;
+
 		const int32 PointCount = FMath::Max(IterState.CachedRayOrigins.Num(), IterState.CachedRayVectors.Num());
 		for (int64 PointIndex = 0; PointIndex < PointCount; ++PointIndex)
 		{
@@ -337,6 +346,8 @@ bool FPCGWorldRaycastElement::ExecuteInternal(FPCGContext* InContext) const
 				OutPoint.Transform = PCGWorldQueryHelpers::GetOrthonormalImpactTransform(Hit);
 				OutPoint.Density = 1.f;
 				UPCGBlueprintHelpers::SetSeedFromPosition(OutPoint);
+
+				bAttributeSuccess &= PCGWorldQueryHelpers::ApplyRayHitMetadata(Hit, WorldQueryParams, OutputPoints.Last(), OutMetadata, World, /*bShouldCreateAttributes=*/false);
 			}
 			else if (Settings->bKeepOriginalPointOnMiss)
 			{
@@ -349,14 +360,14 @@ bool FPCGWorldRaycastElement::ExecuteInternal(FPCGContext* InContext) const
 				{
 					OutputPoints.Emplace(FTransform(Origin), /*InDensity=*/0.f, UPCGBlueprintHelpers::ComputeSeedFromPosition(Origin));
 				}
+
+				bAttributeSuccess &= PCGWorldQueryHelpers::ApplyRayMissMetadata(WorldQueryParams, OutputPoints.Last(), OutMetadata);
 			}
-			else
-			{
-				continue;
-			}
+		}
 
-			// TODO: Pre-create the attributes outside the loop once the helper exists in PCGWorldQueryHelpers
-			PCGWorldQueryHelpers::ApplyRayHitMetadata(HitResult, WorldQueryParams, OutputPoints.Last(), OutMetadata);
+		if (!bAttributeSuccess)
+		{
+			PCGLog::LogWarningOnGraph(LOCTEXT("UnableToApplyAllAttributes", "One or more attributes were unable to be applied."), Context);
 		}
 
 		return true;
diff --git a/Engine/Plugins/PCG/Source/PCG/Private/Helpers/PCGWorldQueryHelpers.cpp b/Engine/Plugins/PCG/Source/PCG/Private/Helpers/PCGWorldQueryHelpers.cpp
--- a/Engine/Plugins/PCG/Source/PCG/Private/Helpers/PCGWorldQueryHelpers.cpp
+++ b/Engine/Plugins/PCG/Source/PCG/Private/Helpers/PCGWorldQueryHelpers.cpp
@@ -3,8 +3,10 @@
 #include "Helpers/PCGWorldQueryHelpers.h"
 
 #include "PCGComponent.h"
+#include "PCGSubsystem.h"
 #include "Data/PCGPointData.h"
 #include "Data/PCGWorldData.h"
+#include "Grid/PCGLandscapeCache.h"
 #include "Helpers/PCGHelpers.h"
 
 #include "LandscapeProxy.h"
@@ -12,7 +14,12 @@
 #include "Components/PrimitiveComponent.h"
 #include "Engine/HitResult.h"
 #include "Engine/OverlapResult.h"
+#include "Engine/StaticMesh.h"
+#include "Kismet/GameplayStatics.h"
+#include "Materials/MaterialInterface.h"
 #include "PhysicalMaterials/PhysicalMaterial.h"
+#include "PhysicsEngine/BodySetup.h"
+#include "PhysicsEngine/PhysicsSettings.h"
 #include "UObject/SoftObjectPath.h"
 
 namespace PCGWorldQueryHelpers
@@ -163,12 +170,73 @@ namespace PCGWorldQueryHelpers
 		return {};
 	}
 
+	bool CreateRayHitAttributes(const FPCGWorldRaycastQueryParams& QueryParams, UPCGMetadata* OutMetadata)
+	{
+		auto CreateAttribute = [OutMetadata]<typename Type>(FName AttributeName, bool bShouldCreate, const Type& DefaultValue)
+		{
+			if (!bShouldCreate)
+			{
+				return true;
+			}
+
+			if (!OutMetadata->HasAttribute(AttributeName))
+			{
+				if (OutMetadata->CreateAttribute<Type>(AttributeName, DefaultValue, /*bAllowsInterpolation=*/true, /*bOverrideParent=*/false))
+				{
+					return true;
+				}
+			}
+
+			return false;
+		};
+
+		bool bResult = true;
+		// Default T/F Impact to true, as most cases misses will be ignored completely
+		bResult &= CreateAttribute(PCGWorldQueryConstants::ImpactAttribute, QueryParams.bGetImpact, true);
+		bResult &= CreateAttribute(PCGWorldQueryConstants::ImpactPointAttribute, QueryParams.bGetImpactPoint, FVector::ZeroVector);
+		bResult &= CreateAttribute(PCGWorldQueryConstants::ImpactNormalAttribute, QueryParams.bGetImpactNormal, FVector::ZeroVector);
+		bResult &= CreateAttribute(PCGWorldQueryConstants::ImpactDistanceAttribute, QueryParams.bGetDistance, 0.0);
+		bResult &= CreateAttribute(PCGWorldQueryConstants::LocalImpactPointAttribute, QueryParams.bGetLocalImpactPoint, FVector::ZeroVector);
+		bResult &= CreateAttribute(PCGPointDataConstants::ActorReferenceAttribute, QueryParams.bGetReferenceToActorHit, FSoftObjectPath());
+		bResult &= CreateAttribute(PCGWorldQueryConstants::PhysicalMaterialReferenceAttribute, QueryParams.bGetReferenceToPhysicalMaterial, FSoftObjectPath());
+		bResult &= CreateAttribute(PCGWorldQueryConstants::RenderMaterialReferenceAttribute, QueryParams.bGetReferenceToRenderMaterial, FSoftObjectPath());
+		bResult &= CreateAttribute(PCGWorldQueryConstants::StaticMeshReferenceAttribute, QueryParams.bGetReferenceToStaticMesh, FSoftObjectPath());
+		bResult &= CreateAttribute(PCGWorldQueryConstants::ElementIndexAttribute, QueryParams.bGetElementIndex, int32{0});
+		bResult &= CreateAttribute(PCGWorldQueryConstants::UVCoordAttribute, QueryParams.bTraceComplex && QueryParams.bGetUVCoords, FVector2D::ZeroVector);
+		bResult &= CreateAttribute(PCGWorldQueryConstants::FaceIndexAttribute, QueryParams.bTraceComplex && QueryParams.bGetFaceIndex, int32{0});
+
+		return bResult;
+	}
+
+	bool ApplyRayMissMetadata(const FPCGWorldRaycastQueryParams& QueryParams, FPCGPoint& OutPoint, UPCGMetadata* OutMetadata)
+	{
+		if (!QueryParams.bGetImpact)
+		{
+			return true;
+		}
+
+		if (!OutMetadata)
+		{
+			return false;
+		}
+
+		if (FPCGMetadataAttribute<bool>* Attribute = OutMetadata->FindOrCreateAttribute<bool>(PCGWorldQueryConstants::ImpactAttribute, true, /*bAllowsInterpolation=*/true, /*bOverrideParent=*/false))
+		{
+			OutMetadata->InitializeOnSet(OutPoint.MetadataEntry);
+			Attribute->SetValue(OutPoint.MetadataEntry, false);
+			return true;
+		}
+
+		return false;
+	}
+
 	// TODO: Add an option to create and apply ranges
 	bool ApplyRayHitMetadata(
-		const TOptional<FHitResult>& HitResult,
+		const FHitResult& HitResult,
 		const FPCGWorldRaycastQueryParams& QueryParams,
 		FPCGPoint& OutPoint,
 		UPCGMetadata* OutMetadata,
+		TWeakObjectPtr<UWorld> World,
 		bool bShouldCreateAttributes)
 	{
 		if (!OutMetadata)
@@ -176,23 +244,9 @@ namespace PCGWorldQueryHelpers
 			return false;
 		}
 
-		// Check for the attributes are added, if not, then add them all in one go
-		auto CreateAttribute = [OutMetadata]<typename Type>(FName AttributeName, bool bShouldCreate, const Type& DefaultValue)
-		{
-			if (bShouldCreate && !OutMetadata->HasAttribute(AttributeName))
-			{
-				OutMetadata->CreateAttribute<Type>(AttributeName, DefaultValue, /*bAllowsInterpolation=*/true, /*bOverrideParent=*/false);
-			}
-		};
-
-		// TODO: Create a new helper to create attributes outside the point loops
 		if (bShouldCreateAttributes)
 		{
-			CreateAttribute(PCGWorldQueryConstants::ImpactAttribute, QueryParams.bGetImpact, false);
-			CreateAttribute(PCGWorldQueryConstants::DistanceAttribute, QueryParams.bGetDistance, 0.0);
-			CreateAttribute(PCGWorldQueryConstants::ImpactNormalAttribute, QueryParams.bGetImpactNormal, FVector());
-			CreateAttribute(PCGPointDataConstants::ActorReferenceAttribute, QueryParams.bGetReferenceToActorHit, FSoftObjectPath());
-			CreateAttribute(PCGWorldQueryConstants::PhysicalMaterialReferenceAttribute, QueryParams.bGetReferenceToPhysicalMaterial, FSoftObjectPath());
+			CreateRayHitAttributes(QueryParams, OutMetadata);
 		}
 
 		auto ApplyAttribute = [&OutPoint, OutMetadata]<typename Type>(FName AttributeName, const Type& Value, bool bShouldApply = true)
@@ -206,27 +260,82 @@ namespace PCGWorldQueryHelpers
 			{
 				OutMetadata->InitializeOnSet(OutPoint.MetadataEntry);
 				Attribute->SetValue(OutPoint.MetadataEntry, Value);
-
 				return true;
 			}
 
 			return false;
 		};
 
-		if (!HitResult.IsSet())
+		const FHitResult& Hit = HitResult;
+
+		bool bResult = true;
+		// Note: The T/F Impact attribute is true by default, so no need to set it directly.
+		bResult &= ApplyAttribute(PCGWorldQueryConstants::ImpactPointAttribute, FVector(Hit.ImpactPoint), QueryParams.bGetImpactPoint);
+		bResult &= ApplyAttribute(PCGWorldQueryConstants::ImpactNormalAttribute, FVector(Hit.ImpactNormal), QueryParams.bGetImpactNormal);
+		bResult &= ApplyAttribute(PCGWorldQueryConstants::ImpactDistanceAttribute, (Hit.ImpactPoint - Hit.TraceStart).Length(), QueryParams.bGetDistance);
+		bResult &= ApplyAttribute(PCGPointDataConstants::ActorReferenceAttribute, FSoftObjectPath(Hit.GetActor()), QueryParams.bGetReferenceToActorHit);
+		bResult &= ApplyAttribute(PCGWorldQueryConstants::PhysicalMaterialReferenceAttribute, FSoftObjectPath(Hit.PhysMaterial.Get()), QueryParams.bGetReferenceToPhysicalMaterial);
+		bResult &= ApplyAttribute(PCGWorldQueryConstants::ElementIndexAttribute, static_cast<int32>(Hit.ElementIndex), QueryParams.bGetElementIndex);
+
+		if (ALandscapeProxy* Landscape = Cast<ALandscapeProxy>(Hit.GetActor()))
 		{
-			ApplyAttribute(PCGWorldQueryConstants::ImpactAttribute, /*Value=*/false, QueryParams.bGetImpact);
-			return true;
+			if (const UMaterialInterface* RenderMaterial = Landscape->GetLandscapeMaterial())
+			{
+				bResult &= ApplyAttribute(PCGWorldQueryConstants::RenderMaterialReferenceAttribute, FSoftObjectPath(RenderMaterial), QueryParams.bGetReferenceToRenderMaterial);
+			}
+
+			if (QueryParams.bApplyMetadataFromLandscape && World.IsValid() && World->GetSubsystem<UPCGSubsystem>())
+			{
+				if (UPCGLandscapeCache* LandscapeCache = World->GetSubsystem<UPCGSubsystem>()->GetLandscapeCache())
+				{
+					// TODO: This is not ideal, but we won't have the Landscape until execution and the LandscapeCache intends on requiring it
+					const TArray<FName> Layers = LandscapeCache->GetLayerNames(Landscape);
+					for (const FName& Layer : Layers)
+					{
+						if (!OutMetadata->HasAttribute(Layer))
+						{
+							OutMetadata->CreateAttribute<float>(Layer, float(), /*bAllowInterpolation=*/true, /*bOverrideParent=*/true);
+						}
+					}
+
+					LandscapeCache->SampleMetadataOnPoint(Landscape, OutPoint, OutMetadata);
+				}
+			}
 		}
+		else if (UPrimitiveComponent* HitComponent = Hit.GetComponent())
+		{
+			if (const UMaterialInterface* RenderMaterial = HitComponent->GetMaterial(Hit.ElementIndex))
+			{
+				bResult &= ApplyAttribute(PCGWorldQueryConstants::RenderMaterialReferenceAttribute, FSoftObjectPath(RenderMaterial), QueryParams.bGetReferenceToRenderMaterial);
+			}
 
-		const FHitResult& Hit = HitResult.GetValue();
+			if (QueryParams.bGetLocalImpactPoint)
+			{
+				const FVector LocalHitLocation = HitComponent->GetComponentToWorld().InverseTransformPosition(Hit.ImpactPoint);
+				bResult &= ApplyAttribute(PCGWorldQueryConstants::LocalImpactPointAttribute, LocalHitLocation);
+			}
 
-		ApplyAttribute(PCGWorldQueryConstants::ImpactAttribute, /*Value=*/true, QueryParams.bGetImpact);
-		ApplyAttribute(PCGWorldQueryConstants::DistanceAttribute, Hit.Distance, QueryParams.bGetDistance);
-		ApplyAttribute(PCGWorldQueryConstants::ImpactNormalAttribute, FVector(Hit.ImpactNormal), QueryParams.bGetImpactNormal);
-		ApplyAttribute(PCGPointDataConstants::ActorReferenceAttribute, FSoftObjectPath(Hit.GetActor()), QueryParams.bGetReferenceToActorHit);
-		ApplyAttribute(PCGWorldQueryConstants::PhysicalMaterialReferenceAttribute, FSoftObjectPath(Hit.PhysMaterial.Get()), QueryParams.bGetReferenceToPhysicalMaterial);
+			if (const UStaticMeshComponent* StaticMeshComponent = Cast<UStaticMeshComponent>(HitComponent))
+			{
+				if (UStaticMesh* StaticMesh = StaticMeshComponent->GetStaticMesh())
+				{
+					bResult &= ApplyAttribute(PCGWorldQueryConstants::StaticMeshReferenceAttribute, FSoftObjectPath(StaticMesh), QueryParams.bGetReferenceToStaticMesh);
+				}
 
-		return true;
+				// Implementation note: FaceIndex will return -1 if complex queries are disabled
+				bResult &= ApplyAttribute(PCGWorldQueryConstants::FaceIndexAttribute, Hit.FaceIndex, QueryParams.bTraceComplex && QueryParams.bGetFaceIndex);
+
+				if (QueryParams.bTraceComplex && QueryParams.bGetUVCoords && UPhysicsSettings::Get()->bSupportUVFromHitResults)
+				{
+					FVector2D UVCoords;
+					if (UGameplayStatics::FindCollisionUV(Hit, QueryParams.UVChannel, UVCoords))
+					{
+						bResult &= ApplyAttribute(PCGWorldQueryConstants::UVCoordAttribute, UVCoords);
+					}
+				}
+			}
+		}
+
+		return bResult;
 	}
 }
diff --git a/Engine/Plugins/PCG/Source/PCG/Public/Data/PCGWorldData.h b/Engine/Plugins/PCG/Source/PCG/Public/Data/PCGWorldData.h
--- a/Engine/Plugins/PCG/Source/PCG/Public/Data/PCGWorldData.h
+++ b/Engine/Plugins/PCG/Source/PCG/Public/Data/PCGWorldData.h
@@ -48,10 +48,10 @@ struct FPCGWorldCommonQueryParams
 #endif
 
 	/** If true, will ignore hits/overlaps on content created from PCG. */
-	UPROPERTY(BlueprintReadWrite, EditAnywhere, Category = "Data|Culling", meta = (PCG_Overridable))
+	UPROPERTY(BlueprintReadWrite, EditAnywhere, Category = "Data|Filtering", meta = (PCG_Overridable))
 	bool bIgnorePCGHits = false;
 
-	UPROPERTY(BlueprintReadWrite, EditAnywhere, Category = "Data|Culling", meta = (PCG_Overridable))
+	UPROPERTY(BlueprintReadWrite, EditAnywhere, Category = "Data|Filtering", meta = (PCG_Overridable))
 	bool bIgnoreSelfHits = true;
 
 	UPROPERTY(BlueprintReadWrite, EditAnywhere, Category = "Data|Advanced", meta = (PCG_Overridable))
@@ -98,17 +98,57 @@ struct FPCGWorldRaycastQueryParams : public FPCGWorldCommonQueryParams
 	void Initialize();
 	void PostSerialize(const FArchive& Ar);
 
-	UPROPERTY(BlueprintReadWrite, EditAnywhere, Category = "Data|Culling", meta = (PCG_Overridable))
-	bool bIgnoreBackfaceHits = false;
+	/** Ignore rays that hit backfaces. */
+	UPROPERTY(BlueprintReadWrite, EditAnywhere, Category = "Data|Filtering", meta = (PCG_Overridable))
+	uint8 bIgnoreBackfaceHits : 1;
+
+	/** Create an attribute for whether the raycast resulted in a hit. */
+	UPROPERTY(BlueprintReadWrite, EditAnywhere, Category = "Data|Attributes", meta = (PCG_Overridable))
+	uint8 bGetImpact : 1;
+
+	/** Create an attribute for the impact location in world space. */
+	UPROPERTY(BlueprintReadWrite, EditAnywhere, Category = "Data|Attributes", meta = (PCG_Overridable))
+	uint8 bGetImpactPoint : 1;
 
+	/** Create an attribute for the impact normal. */
 	UPROPERTY(BlueprintReadWrite, EditAnywhere, Category = "Data|Attributes", meta = (PCG_Overridable))
-	bool bGetImpact = false;
+	uint8 bGetImpactNormal : 1;
 
+	/** Create an attribute for the distance between the ray origin and the impact point. */
 	UPROPERTY(BlueprintReadWrite, EditAnywhere, Category = "Data|Attributes", meta = (PCG_Overridable))
-	bool bGetDistance = false;
+	uint8 bGetDistance : 1;
 
+	/** Create an attribute for the impact point in the hit object's local space. */
 	UPROPERTY(BlueprintReadWrite, EditAnywhere, Category = "Data|Attributes", meta = (PCG_Overridable))
-	bool bGetImpactNormal = false;
+	uint8 bGetLocalImpactPoint : 1;
+
+	/** Create an attribute for the render material. */
+	UPROPERTY(BlueprintReadWrite, EditAnywhere, Category = "Data|Attributes", meta = (PCG_Overridable))
+	uint8 bGetReferenceToRenderMaterial : 1;
+
+	/** Create an attribute for the static mesh. */
+	UPROPERTY(BlueprintReadWrite, EditAnywhere, Category = "Data|Attributes", meta = (PCG_Overridable))
+	uint8 bGetReferenceToStaticMesh : 1;
+
+	/** Create an attribute for index of the hit face. Note: Will only work in complex traces. */
+	UPROPERTY(BlueprintReadWrite, EditAnywhere, Category = "Data|Attributes", meta = (PCG_Overridable, EditCondition = "bTraceComplex"))
+	uint8 bGetFaceIndex : 1;
+
+	/** Create an attribute for UV Coordinates of the surface hit. Note: Will only work in complex traces. */
+	UPROPERTY(BlueprintReadWrite, EditAnywhere, Category = "Data|Attributes", meta = (PCG_Overridable, EditCondition = "bTraceComplex"))
+	uint8 bGetUVCoords : 1;
+
+	/** Create an attribute for the index of the element hit. Unique to the hit primitive. */
+	UPROPERTY(BlueprintReadWrite, EditAnywhere, Category = "Data|Attributes", meta = (PCG_Overridable))
+	uint8 bGetElementIndex : 1;
+
+	/** Will apply landscape layers and their values at the impact point. */
+	UPROPERTY(BlueprintReadWrite, EditAnywhere, Category = "Data|Attributes", meta = (PCG_Overridable, EditCondition = "SelectLandscapeHits != EPCGWorldQuerySelectLandscapeHits::Exclude"))
+	uint8 bApplyMetadataFromLandscape : 1;
+
+	/** This UV Channel will be selected when retrieving UV Coordinates from a raycast query. */
+	UPROPERTY(BlueprintReadWrite, EditAnywhere, Category = "Data|Attributes", meta = (PCG_Overridable, EditCondition = "bTraceComplex && bGetUVCoords", EditConditionHides, DisplayAfter = "bGetUVCoords"))
+	int32 UVChannel = INDEX_NONE;
 };
 
 template<>
@@ -165,9 +205,6 @@ struct FPCGWorldRayHitQueryParams : public FPCGWorldRaycastQueryParams
 
 	// TODO: see in FCollisionQueryParams if there are some flags we want to expose
 	// examples: bReturnFaceIndex, bReturnPhysicalMaterial, some ignore patterns
-
-	UPROPERTY(BlueprintReadWrite, EditAnywhere, Category = Data, meta = (PCG_Overridable))
-	bool bApplyMetadataFromLandscape = false;
 };
 
 template<>
diff --git a/Engine/Plugins/PCG/Source/PCG/Public/Helpers/PCGWorldQueryHelpers.h b/Engine/Plugins/PCG/Source/PCG/Public/Helpers/PCGWorldQueryHelpers.h
--- a/Engine/Plugins/PCG/Source/PCG/Public/Helpers/PCGWorldQueryHelpers.h
+++ b/Engine/Plugins/PCG/Source/PCG/Public/Helpers/PCGWorldQueryHelpers.h
@@ -28,9 +28,16 @@ enum class EPCGWorldRaycastMode : uint8
 namespace PCGWorldQueryConstants
 {
 	const FName ImpactAttribute = TEXT("ImpactResult");
-	const FName DistanceAttribute = TEXT("ImpactDistance");
+	const FName ImpactPointAttribute = TEXT("ImpactPoint");
 	const FName ImpactNormalAttribute = TEXT("ImpactNormal");
+	const FName ImpactDistanceAttribute = TEXT("ImpactDistance");
+	const FName LocalImpactPointAttribute = TEXT("ImpactLocalPoint");
 	const FName PhysicalMaterialReferenceAttribute = TEXT("PhysicalMaterial");
+	const FName RenderMaterialReferenceAttribute = TEXT("ImpactRenderMaterial");
+	const FName StaticMeshReferenceAttribute = TEXT("ImpactStaticMesh");
+	const FName ElementIndexAttribute = TEXT("ImpactElementIndex");
+	const FName UVCoordAttribute = TEXT("ImpactUVCoords");
+	const FName FaceIndexAttribute = TEXT("ImpactFaceIndex");
 }
 
 namespace PCGWorldQueryHelpers
@@ -55,10 +62,17 @@ namespace PCGWorldQueryHelpers
 		TWeakObjectPtr<UPCGComponent> OriginatingComponent,
 		const TArray<FOverlapResult>& OverlapResults);
 
+	/** Creates hit result attributes based off query params. Can be called before ApplyRayHitMetadata. */
+	bool CreateRayHitAttributes(const FPCGWorldRaycastQueryParams& QueryParams, UPCGMetadata* OutMetadata);
+
+	/** Applies a 'miss' hit result to the metadata. To be called if the ray misses the target and the point should be kept. */
+	bool ApplyRayMissMetadata(const FPCGWorldRaycastQueryParams& QueryParams, FPCGPoint& OutPoint, UPCGMetadata* OutMetadata);
+
 	/** Applies common world ray hit results to attributes. */
-	bool ApplyRayHitMetadata(const TOptional<FHitResult>& HitResult,
+	bool ApplyRayHitMetadata(const FHitResult& HitResult,
 		const FPCGWorldRaycastQueryParams& QueryParams,
 		FPCGPoint& OutPoint,
 		UPCGMetadata* OutMetadata,
+		TWeakObjectPtr<UWorld> World,
 		bool bShouldCreateAttributes = true);
 }
