[PCG] Fix OnGraphChangedDelegate not being properly hooked after delete/undo of PCG Volume
- Cleanup the whole registration code and register/unregister delegates in the UPCGComponent::OnRegister/OnUnregister instead

#rnx

diff --git a/Engine/Plugins/PCG/Source/PCG/Private/PCGComponent.cpp b/Engine/Plugins/PCG/Source/PCG/Private/PCGComponent.cpp
--- a/Engine/Plugins/PCG/Source/PCG/Private/PCGComponent.cpp
+++ b/Engine/Plugins/PCG/Source/PCG/Private/PCGComponent.cpp
@@ -1406,6 +1406,7 @@ void UPCGComponent::EndPlay(const EEndPlayReason::Type EndPlayReason)
 
 void UPCGComponent::OnUnregister()
 {
+	TeardownCallbacks();
 #if WITH_EDITOR
 	if (UPCGSubsystem* Subsystem = GetSubsystem())
 	{
@@ -1423,17 +1424,6 @@ void UPCGComponent::OnUnregister()
 
 void UPCGComponent::OnComponentDestroyed(bool bDestroyingHierarchy)
 {
-#if WITH_EDITOR
-	// BeginDestroy is not called immediately when a component is destroyed. Therefore callbacks are not cleaned
-	// until GC is ran, and can stack up with BP reconstruction scripts. Force the removal of callbacks here. If the component
-	// is dead, we don't want to react to callbacks anyway.
-	if (GraphInstance)
-	{
-		GraphInstance->OnGraphChangedDelegate.RemoveAll(this);
-		GraphInstance->TeardownCallbacks();
-	}
-#endif
-
 	// Bookkeeping local components that might be deleted by the user.
 	// Making sure that the corresponding partition actor doesn't keep a dangling references
 	if (APCGPartitionActor* PAOwner = Cast<APCGPartitionActor>(GetOwner()))
@@ -1567,7 +1557,7 @@ void UPCGComponent::PostLoad()
 		Graph_DEPRECATED = nullptr;
 	}
 
-	SetupCallbacksOnCreation();
+	UpdateTrackingCache();
 
 	// Always set the editing mode to Preview when we're in GenerateAtRuntime mode
 	CurrentEditingMode = IsManagedByRuntimeGenSystem() ? EPCGEditorDirtyMode::Preview : SerializedEditingMode;
@@ -1606,19 +1596,6 @@ void UPCGComponent::PostLoad()
 }
 
 #if WITH_EDITOR
-void UPCGComponent::SetupCallbacksOnCreation()
-{
-	UpdateTrackingCache();
-
-	if (GraphInstance)
-	{
-		// We might have already connected in PostInitProperties
-		// To be sure, remove it and re-add it.
-		GraphInstance->OnGraphChangedDelegate.RemoveAll(this);
-		GraphInstance->OnGraphChangedDelegate.AddUObject(this, &UPCGComponent::OnGraphChanged);
-	}
-}
-
 bool UPCGComponent::CanEditChange(const FProperty* InProperty) const
 {
 	// Can't change anything if the component is local
@@ -1629,11 +1606,6 @@ bool UPCGComponent::CanEditChange(const FProperty* InProperty) const
 void UPCGComponent::BeginDestroy()
 {
 #if WITH_EDITOR
-	if (GraphInstance)
-	{
-		GraphInstance->OnGraphChangedDelegate.RemoveAll(this);
-	}
-
 	// For the special case where a component is part of a reconstruction script (from a BP),
 	// but gets destroyed immediately, we need to force the unregistering. 
 	if (UPCGSubsystem* PCGSubsystem = UPCGSubsystem::GetSubsystemForCurrentWorld())
@@ -1652,10 +1624,6 @@ bool UPCGComponent::IsEditorOnly() const
 
 void UPCGComponent::PostInitProperties()
 {
-#if WITH_EDITOR
-	GraphInstance->OnGraphChangedDelegate.AddUObject(this, &UPCGComponent::OnGraphChanged);
-#endif // WITH_EDITOR
-
 #if WITH_EDITOR
 	// Force bIsPartitioned at false for new objects
 	if (PCGHelpers::IsNewObjectAndNotDefault(this, /*bCheckHierarchy=*/true))
@@ -1694,6 +1662,7 @@ void UPCGComponent::OnRegister()
 		}
 	}
 #endif //WITH_EDITOR
+	SetupCallbacks();
 }
 
 TStructOnScope<FActorComponentInstanceData> UPCGComponent::GetComponentInstanceData() const
@@ -1702,6 +1671,25 @@ TStructOnScope<FActorComponentInstanceData> UPCGComponent::GetComponentInstanceD
 	return InstanceData;
 }
 
+void UPCGComponent::SetupCallbacks()
+{
+#if WITH_EDITOR
+	if (!GraphInstance->OnGraphChangedDelegate.IsBoundToObject(this))
+	{
+		GraphInstance->OnGraphChangedDelegate.AddUObject(this, &UPCGComponent::OnGraphChanged);
+		GraphInstance->SetupCallbacks();
+	}
+#endif
+}
+
+void UPCGComponent::TeardownCallbacks()
+{
+#if WITH_EDITOR
+	GraphInstance->TeardownCallbacks();
+	GraphInstance->OnGraphChangedDelegate.RemoveAll(this);
+#endif
+}
+
 void UPCGComponent::OnGraphChanged(UPCGGraphInterface* InGraph, EPCGChangeType ChangeType)
 {
 	RefreshAfterGraphChanged(InGraph, ChangeType);
@@ -1958,7 +1946,7 @@ void UPCGComponent::PostEditImport()
 {
 	Super::PostEditImport();
 
-	SetupCallbacksOnCreation();
+	UpdateTrackingCache();
 }
 
 void UPCGComponent::PreEditUndo()
@@ -3654,16 +3642,6 @@ void FPCGComponentInstanceData::ApplyToComponent(UActorComponent* Component, con
 #endif
 		}
 
-#if WITH_EDITOR
-		// Reconnect callbacks
-		if (PCGComponent->GraphInstance)
-		{
-			PCGComponent->GraphInstance->SetupCallbacks();
-			PCGComponent->GraphInstance->OnGraphChangedDelegate.RemoveAll(PCGComponent);
-			PCGComponent->GraphInstance->OnGraphChangedDelegate.AddUObject(PCGComponent, &UPCGComponent::OnGraphChanged);
-		}
-#endif // WITH_EDITOR
-
 		bool bDoActorMapping = PCGComponent->bGenerated || PCGHelpers::IsRuntimeOrPIE();
 
 		// Also remap
@@ -3674,12 +3652,6 @@ void FPCGComponentInstanceData::ApplyToComponent(UActorComponent* Component, con
 		}
 
 #if WITH_EDITOR
-		// Disconnect callbacks on source.
-		if (ConstructionSourceComponent && ConstructionSourceComponent->GraphInstance)
-		{
-			ConstructionSourceComponent->GraphInstance->TeardownCallbacks();
-		}
-
 		// Finally, start a delayed refresh task (if there is not one already), in editor only
 		// It is important to be delayed, because we cannot spawn Partition Actors within this scope,
 		// because we are in a construction script.
diff --git a/Engine/Plugins/PCG/Source/PCG/Public/PCGComponent.h b/Engine/Plugins/PCG/Source/PCG/Public/PCGComponent.h
--- a/Engine/Plugins/PCG/Source/PCG/Public/PCGComponent.h
+++ b/Engine/Plugins/PCG/Source/PCG/Public/PCGComponent.h
@@ -525,15 +525,16 @@ private:
 	void RefreshAfterGraphChanged(UPCGGraphInterface* InGraph, EPCGChangeType ChangeType);
 	void OnGraphChanged(UPCGGraphInterface* InGraph, EPCGChangeType ChangeType);
 
+	/** Sets up actor, tracking, landscape and graph callbacks */
+	void SetupCallbacks();
+	void TeardownCallbacks();
+
 #if WITH_EDITOR
 	virtual void PreEditChange(FProperty* PropertyAboutToChange) override;
 	virtual void PostEditChangeProperty(FPropertyChangedEvent& PropertyChangedEvent) override;
 	virtual void PreEditUndo() override;
 	virtual void PostEditUndo() override;
 
-	/** Sets up actor, tracking, landscape and graph callbacks */
-	void SetupCallbacksOnCreation();
-
 	/** Returns true if something changed in the tracking. */
 	bool UpdateTrackingCache(TArray<FPCGSelectionKey>* OptionalChangedKeys = nullptr);
 
