[PCG] Tentative fix for Static Analysis

#rnx

diff --git a/Engine/Plugins/PCG/Source/PCG/Private/Tests/Elements/PCGRandomChoiceTest.cpp b/Engine/Plugins/PCG/Source/PCG/Private/Tests/Elements/PCGRandomChoiceTest.cpp
--- a/Engine/Plugins/PCG/Source/PCG/Private/Tests/Elements/PCGRandomChoiceTest.cpp
+++ b/Engine/Plugins/PCG/Source/PCG/Private/Tests/Elements/PCGRandomChoiceTest.cpp
@@ -286,6 +286,8 @@ bool FPCGRandomChoiceTest_MultiData_SameSeed::RunTest(const FString& Parameters)
 	UTEST_TRUE("There is 2 point data in chosen points", FirstChosenOutputData && SecondChosenOutputData);
 	UTEST_TRUE("There is 2 point data in discarded points", FirstDiscardedOutputData && SecondDiscardedOutputData);
 
+	check(FirstChosenOutputData && SecondChosenOutputData && FirstDiscardedOutputData && SecondDiscardedOutputData);
+
 	UTEST_EQUAL("There is the right number of points in first chosen", FirstChosenOutputData->GetPoints().Num(), ExpectedNumElementsChosen);
 	UTEST_EQUAL("There is the right number of points in second chosen", SecondChosenOutputData->GetPoints().Num(), ExpectedNumElementsChosen);
 
@@ -336,6 +338,8 @@ bool FPCGRandomChoiceTest_MultiData_DifferentSeed::RunTest(const FString& Parame
 	UTEST_TRUE("There is 2 point data in chosen points", FirstChosenOutputData && SecondChosenOutputData);
 	UTEST_TRUE("There is 2 point data in discarded points", FirstDiscardedOutputData && SecondDiscardedOutputData);
 
+	check(FirstChosenOutputData && SecondChosenOutputData && FirstDiscardedOutputData && SecondDiscardedOutputData);
+
 	UTEST_EQUAL("There is the right number of points in first chosen", FirstChosenOutputData->GetPoints().Num(), ExpectedNumElementsChosen);
 	UTEST_EQUAL("There is the right number of points in second chosen", SecondChosenOutputData->GetPoints().Num(), ExpectedNumElementsChosen);
 
