[PCG] Add ALV visualization support to the PCGDataVisualization interface

diff --git a/Engine/Plugins/PCG/Source/PCG/Public/PCGDataVisualization.h b/Engine/Plugins/PCG/Source/PCG/Public/PCGDataVisualization.h
--- a/Engine/Plugins/PCG/Source/PCG/Public/PCGDataVisualization.h
+++ b/Engine/Plugins/PCG/Source/PCG/Public/PCGDataVisualization.h
@@ -2,16 +2,53 @@
 
 #pragma once
 
+#include "Metadata/Accessors/IPCGAttributeAccessor.h"
+
 #if WITH_EDITOR
 class AActor;
 class UPCGData;
 struct FPCGContext;
 
+enum class EPCGTableVisualizerColumnSortingMode
+{
+	None,
+	Ascending,
+	Descending
+};
+
+enum class EPCGTableVisualizerCellAlignment
+{
+	Fill,
+	Left,
+	Center,
+	Right,
+};
+
+struct FPCGTableVisualizerColumnInfo
+{
+	FName Id;
+	FText Label;
+	FText Tooltip;
+	float Width = -1.0f; // Will be calculated automatically if < 0
+	EPCGTableVisualizerCellAlignment CellAlignment = EPCGTableVisualizerCellAlignment::Right;
+	TSharedPtr<const IPCGAttributeAccessor> Accessor;
+};
+
+struct FPCGTableVisualizerInfo
+{
+	const UPCGData* Data = nullptr;
+	TArray<FPCGTableVisualizerColumnInfo> ColumnInfos;
+	EPCGTableVisualizerColumnSortingMode SortingMode = EPCGTableVisualizerColumnSortingMode::Ascending;
+	FName SortingColumn = NAME_None;
+	TSharedPtr<const IPCGAttributeAccessorKeys> AccessorKeys;
+};
+
 /** Implement this interface to provide custom PCGData visualizations. Register your implementation to FPCGModule::FPCGDataVisualizationRegistry to be used automatically. */
 class IPCGDataVisualization
 {
 public:
 	virtual ~IPCGDataVisualization() = default;
 	virtual void ExecuteDebugDisplay(FPCGContext* Context, const UPCGData* Data, AActor* TargetActor) const = 0;
+	virtual FPCGTableVisualizerInfo GetTableVisualizerInfo(const UPCGData* Data) const = 0;
 };
 #endif
diff --git a/Engine/Plugins/PCG/Source/PCGEditor/Private/DataVisualizations/PCGDataVisualizationHelpers.cpp b/Engine/Plugins/PCG/Source/PCGEditor/Private/DataVisualizations/PCGDataVisualizationHelpers.cpp
new file mode 100644
--- /dev/null
+++ b/Engine/Plugins/PCG/Source/PCGEditor/Private/DataVisualizations/PCGDataVisualizationHelpers.cpp
@@ -0,0 +1,161 @@
+// Copyright Epic Games, Inc. All Rights Reserved.
+
+#include "DataVisualizations/PCGDataVisualizationHelpers.h"
+
+#include "SPCGEditorGraphAttributeListView.h"
+#include "Metadata/PCGAttributePropertySelector.h"
+#include "Metadata/PCGMetadata.h"
+#include "Metadata/Accessors/PCGAttributeAccessorHelpers.h"
+
+#define LOCTEXT_NAMESPACE "PCGDataVisualizationHelpers"
+
+namespace PCGDataVisualizationConstants
+{
+	const FString NoneAttributeId = TEXT("@None");
+}
+
+namespace PCGDataVisualizationHelpers
+{
+	FPCGTableVisualizerColumnInfo CreateColumnInfo(const UPCGData* Data, FName Id, const FText& Label, TFunction<TSharedPtr<const IPCGAttributeAccessor>()> CreateAccessorFunc)
+	{
+		FPCGTableVisualizerColumnInfo ColumnInfo;
+		ColumnInfo.Id = Id;
+		ColumnInfo.Label = Label;
+
+		FPCGAttributePropertyInputSelector Selector = FPCGAttributePropertySelector::CreateSelectorFromString<FPCGAttributePropertyInputSelector>(Id.ToString());
+		ColumnInfo.Accessor = CreateAccessorFunc ? CreateAccessorFunc() : TSharedPtr<const IPCGAttributeAccessor>(PCGAttributeAccessorHelpers::CreateConstAccessor(Data, Selector).Release());
+
+		if (ensure(ColumnInfo.Accessor))
+		{
+			ColumnInfo.Tooltip = FText::FromString(PCG::Private::GetTypeName(ColumnInfo.Accessor->GetUnderlyingType()));
+		}
+
+		return ColumnInfo;
+	}
+
+	FPCGTableVisualizerColumnInfo CreateMetadataColumnInfo(const UPCGData* Data, const FName& Id, EPCGMetadataTypes MetadataType, const TCHAR* PostFix = nullptr)
+	{
+		FString OriginalId = Id.ToString();
+		if (PostFix)
+		{
+			OriginalId.Append(PostFix);
+		}
+
+		FString IdString(OriginalId);
+		FText Label;
+
+		if (Id == NAME_None)
+		{
+			IdString = FString::Printf(TEXT("%s%s"), *PCGDataVisualizationConstants::NoneAttributeId, PostFix);
+			Label = FText::Format(LOCTEXT("NoneLabelFormat", "{0}{1}"), UEnum::GetDisplayValueAsText(MetadataType), FText::FromString(PostFix));
+		}
+		else if (Label.IsEmpty())
+		{
+			Label = FText::FromString(IdString);
+		}
+
+		FPCGTableVisualizerColumnInfo ColumnInfo;
+		ColumnInfo.Id = FName(IdString);
+		ColumnInfo.Label = Label;
+
+		if (MetadataType == EPCGMetadataTypes::String)
+		{
+			ColumnInfo.CellAlignment = EPCGTableVisualizerCellAlignment::Left;
+			ColumnInfo.Width = PCGEditorGraphAttributeListView::MaxColumnWidth;
+		}
+
+		if (Data)
+		{
+			FPCGAttributePropertyInputSelector Selector;
+			Selector.Update(OriginalId);
+	
+			ColumnInfo.Accessor = TSharedPtr<const IPCGAttributeAccessor>(PCGAttributeAccessorHelpers::CreateConstAccessor(Data, Selector).Release());
+
+			if (ColumnInfo.Accessor)
+			{
+				ColumnInfo.Tooltip = FText::FromString(PCG::Private::GetTypeName(ColumnInfo.Accessor->GetUnderlyingType()));
+			}
+		}
+
+		return ColumnInfo;
+	}
+
+	void CreateMetadataColumnInfos(const UPCGData* Data, TArray<FPCGTableVisualizerColumnInfo>& OutColumnInfos)
+	{
+		const UPCGMetadata* Metadata = Data->ConstMetadata();
+
+		TArray<FName> AttributeNames;
+		TArray<EPCGMetadataTypes> AttributeTypes;
+		Metadata->GetAttributes(AttributeNames, AttributeTypes);
+
+		for (int32 I = 0; I < AttributeNames.Num(); I++)
+		{
+			const FName& AttributeName = AttributeNames[I];
+			const EPCGMetadataTypes AttributeType = AttributeTypes[I];
+			FName ColumnName = AttributeName;
+
+			switch (AttributeType)
+			{
+			case EPCGMetadataTypes::Float:
+			case EPCGMetadataTypes::Double:
+			case EPCGMetadataTypes::Integer32:
+			case EPCGMetadataTypes::Integer64:
+			case EPCGMetadataTypes::Boolean:
+			case EPCGMetadataTypes::String:
+			case EPCGMetadataTypes::Name:
+			case EPCGMetadataTypes::SoftObjectPath:
+			case EPCGMetadataTypes::SoftClassPath:
+			{
+				OutColumnInfos.Add(CreateMetadataColumnInfo(Data, ColumnName, AttributeType, TEXT("")));
+				break;
+			}
+			case EPCGMetadataTypes::Vector2:
+			{
+				OutColumnInfos.Add(CreateMetadataColumnInfo(Data, ColumnName, AttributeType, TEXT(".X")));
+				OutColumnInfos.Add(CreateMetadataColumnInfo(Data, ColumnName, AttributeType, TEXT(".Y")));
+				break;
+			}
+			case EPCGMetadataTypes::Vector:
+			{
+				OutColumnInfos.Add(CreateMetadataColumnInfo(Data, ColumnName, AttributeType, TEXT(".X")));
+				OutColumnInfos.Add(CreateMetadataColumnInfo(Data, ColumnName, AttributeType, TEXT(".Y")));
+				OutColumnInfos.Add(CreateMetadataColumnInfo(Data, ColumnName, AttributeType, TEXT(".Z")));
+				break;
+			}
+			case EPCGMetadataTypes::Vector4:
+			case EPCGMetadataTypes::Quaternion:
+			{
+				OutColumnInfos.Add(CreateMetadataColumnInfo(Data, ColumnName, AttributeType, TEXT(".X")));
+				OutColumnInfos.Add(CreateMetadataColumnInfo(Data, ColumnName, AttributeType, TEXT(".Y")));
+				OutColumnInfos.Add(CreateMetadataColumnInfo(Data, ColumnName, AttributeType, TEXT(".Z")));
+				OutColumnInfos.Add(CreateMetadataColumnInfo(Data, ColumnName, AttributeType, TEXT(".W")));
+				break;
+			}
+			case EPCGMetadataTypes::Rotator:
+			{
+				OutColumnInfos.Add(CreateMetadataColumnInfo(Data, ColumnName, AttributeType, TEXT(".Roll")));
+				OutColumnInfos.Add(CreateMetadataColumnInfo(Data, ColumnName, AttributeType, TEXT(".Pitch")));
+				OutColumnInfos.Add(CreateMetadataColumnInfo(Data, ColumnName, AttributeType, TEXT(".Yaw")));
+				break;
+			}
+			case EPCGMetadataTypes::Transform:
+			{
+				OutColumnInfos.Add(CreateMetadataColumnInfo(Data, ColumnName, AttributeType, TEXT(".Position.X")));
+				OutColumnInfos.Add(CreateMetadataColumnInfo(Data, ColumnName, AttributeType, TEXT(".Position.Y")));
+				OutColumnInfos.Add(CreateMetadataColumnInfo(Data, ColumnName, AttributeType, TEXT(".Position.Z")));
+				OutColumnInfos.Add(CreateMetadataColumnInfo(Data, ColumnName, AttributeType, TEXT(".Rotation.Roll")));
+				OutColumnInfos.Add(CreateMetadataColumnInfo(Data, ColumnName, AttributeType, TEXT(".Rotation.Pitch")));
+				OutColumnInfos.Add(CreateMetadataColumnInfo(Data, ColumnName, AttributeType, TEXT(".Rotation.Yaw")));
+				OutColumnInfos.Add(CreateMetadataColumnInfo(Data, ColumnName, AttributeType, TEXT(".Scale.X")));
+				OutColumnInfos.Add(CreateMetadataColumnInfo(Data, ColumnName, AttributeType, TEXT(".Scale.Y")));
+				OutColumnInfos.Add(CreateMetadataColumnInfo(Data, ColumnName, AttributeType, TEXT(".Scale.Z")));
+				break;
+			}
+			default:
+				break;
+			}
+		}
+	}
+}
+
+#undef LOCTEXT_NAMESPACE
diff --git a/Engine/Plugins/PCG/Source/PCGEditor/Private/DataVisualizations/PCGDataVisualizationHelpers.h b/Engine/Plugins/PCG/Source/PCGEditor/Private/DataVisualizations/PCGDataVisualizationHelpers.h
new file mode 100644
--- /dev/null
+++ b/Engine/Plugins/PCG/Source/PCGEditor/Private/DataVisualizations/PCGDataVisualizationHelpers.h
@@ -0,0 +1,13 @@
+// Copyright Epic Games, Inc. All Rights Reserved.
+
+#pragma once
+
+#include "PCGDataVisualization.h"
+
+class UPCGData;
+
+namespace PCGDataVisualizationHelpers
+{
+	FPCGTableVisualizerColumnInfo CreateColumnInfo(const UPCGData* Data, FName Id, const FText& Label, TFunction<TSharedPtr<const IPCGAttributeAccessor>()> CreateAccessorFunc = nullptr);
+	void CreateMetadataColumnInfos(const UPCGData* Data, TArray<FPCGTableVisualizerColumnInfo>& OutColumnInfos);
+}
diff --git a/Engine/Plugins/PCG/Source/PCGEditor/Private/DataVisualizations/PCGParamDataVisualization.cpp b/Engine/Plugins/PCG/Source/PCGEditor/Private/DataVisualizations/PCGParamDataVisualization.cpp
new file mode 100644
--- /dev/null
+++ b/Engine/Plugins/PCG/Source/PCGEditor/Private/DataVisualizations/PCGParamDataVisualization.cpp
@@ -0,0 +1,34 @@
+// Copyright Epic Games, Inc. All Rights Reserved.
+
+#include "DataVisualizations/PCGParamDataVisualization.h"
+
+#include "PCGData.h"
+#include "DataVisualizations/PCGDataVisualizationHelpers.h"
+#include "Metadata/Accessors/PCGAttributeAccessorHelpers.h"
+
+#define LOCTEXT_NAMESPACE "PCGParamDataVisualization"
+
+namespace PCGParamDataVisualizationConstants
+{
+	/** Names of the columns in the attribute list. */
+	const FName NAME_Index = FName(TEXT("$Index"));
+
+	/** Labels of the columns. */
+	const FText TEXT_Index = LOCTEXT("Index", "$Index");
+}
+
+FPCGTableVisualizerInfo IPCGParamDataVisualization::GetTableVisualizerInfo(const UPCGData* Data) const
+{
+	FPCGTableVisualizerInfo Info;
+	Info.Data = Data;
+	Info.ColumnInfos.Add(PCGDataVisualizationHelpers::CreateColumnInfo(Data, PCGParamDataVisualizationConstants::NAME_Index, PCGParamDataVisualizationConstants::TEXT_Index));
+	Info.SortingColumn = PCGParamDataVisualizationConstants::NAME_Index;
+	Info.AccessorKeys = TSharedPtr<const IPCGAttributeAccessorKeys>(PCGAttributeAccessorHelpers::CreateConstKeys(Data, FPCGAttributePropertyInputSelector()).Release());
+
+	// Add Metadata Columns
+	PCGDataVisualizationHelpers::CreateMetadataColumnInfos(Data, Info.ColumnInfos);
+
+	return Info;
+}
+
+#undef LOCTEXT_NAMESPACE
diff --git a/Engine/Plugins/PCG/Source/PCGEditor/Private/DataVisualizations/PCGParamDataVisualization.h b/Engine/Plugins/PCG/Source/PCGEditor/Private/DataVisualizations/PCGParamDataVisualization.h
new file mode 100644
--- /dev/null
+++ b/Engine/Plugins/PCG/Source/PCGEditor/Private/DataVisualizations/PCGParamDataVisualization.h
@@ -0,0 +1,14 @@
+// Copyright Epic Games, Inc. All Rights Reserved.
+
+#pragma once
+
+#include "PCGDataVisualization.h"
+
+class UPCGData;
+
+class PCGEDITOR_API IPCGParamDataVisualization : public IPCGDataVisualization
+{
+public:
+	virtual void ExecuteDebugDisplay(FPCGContext* Context, const UPCGData* Data, AActor* TargetActor) const { /* Do nothing. */ }
+	virtual FPCGTableVisualizerInfo GetTableVisualizerInfo(const UPCGData* Data) const override;
+};
diff --git a/Engine/Plugins/PCG/Source/PCGEditor/Private/DataVisualizations/PCGSpatialDataVisualization.cpp b/Engine/Plugins/PCG/Source/PCGEditor/Private/DataVisualizations/PCGSpatialDataVisualization.cpp
--- a/Engine/Plugins/PCG/Source/PCGEditor/Private/DataVisualizations/PCGSpatialDataVisualization.cpp
+++ b/Engine/Plugins/PCG/Source/PCGEditor/Private/DataVisualizations/PCGSpatialDataVisualization.cpp
@@ -3,17 +3,80 @@
 #include "DataVisualizations/PCGSpatialDataVisualization.h"
 
 #include "PCGContext.h"
+#include "PCGData.h"
 #include "PCGDebug.h"
+#include "PCGEditorCommon.h"
+#include "PCGPoint.h"
 #include "PCGSettings.h"
 #include "Data/PCGPointData.h"
 #include "Data/PCGSpatialData.h"
+#include "DataVisualizations/PCGDataVisualizationHelpers.h"
 #include "Helpers/PCGActorHelpers.h"
 #include "Helpers/PCGHelpers.h"
+#include "Metadata/Accessors/PCGAttributeAccessorHelpers.h"
+#include "Metadata/Accessors/PCGCustomAccessor.h"
 
 #include "Components/InstancedStaticMeshComponent.h"
 
 #define LOCTEXT_NAMESPACE "PCGSpatialDataVisualization"
 
+namespace PCGPointDataVisualizationConstants
+{
+	/** Names of the columns in the attribute list. */
+	const FName NAME_Index = FName(TEXT("$Index"));
+	const FName NAME_PositionX = FName(TEXT("$Position.X"));
+	const FName NAME_PositionY = FName(TEXT("$Position.Y"));
+	const FName NAME_PositionZ = FName(TEXT("$Position.Z"));
+	const FName NAME_RotationRoll = FName(TEXT("$Rotation.Roll"));
+	const FName NAME_RotationPitch = FName(TEXT("$Rotation.Pitch"));
+	const FName NAME_RotationYaw = FName(TEXT("$Rotation.Yaw"));
+	const FName NAME_ScaleX = FName(TEXT("$Scale.X"));
+	const FName NAME_ScaleY = FName(TEXT("$Scale.Y"));
+	const FName NAME_ScaleZ = FName(TEXT("$Scale.Z"));
+	const FName NAME_BoundsMinX = FName(TEXT("$BoundsMin.X"));
+	const FName NAME_BoundsMinY = FName(TEXT("$BoundsMin.Y"));
+	const FName NAME_BoundsMinZ = FName(TEXT("$BoundsMin.Z"));
+	const FName NAME_BoundsMaxX = FName(TEXT("$BoundsMax.X"));
+	const FName NAME_BoundsMaxY = FName(TEXT("$BoundsMax.Y"));
+	const FName NAME_BoundsMaxZ = FName(TEXT("$BoundsMax.Z"));
+	const FName NAME_ColorR = FName(TEXT("$Color.R"));
+	const FName NAME_ColorG = FName(TEXT("$Color.G"));
+	const FName NAME_ColorB = FName(TEXT("$Color.B"));
+	const FName NAME_ColorA = FName(TEXT("$Color.A"));
+	const FName NAME_Density = FName(TEXT("$Density"));
+	const FName NAME_Steepness = FName(TEXT("$Steepness"));
+	const FName NAME_Seed = FName(TEXT("$Seed"));
+	const FName NAME_MetadataEntry = FName(TEXT("MetadataEntry"));
+	const FName NAME_MetadataEntryParent = FName(TEXT("PointMetadataEntryParent"));
+
+	/** Labels of the columns. */
+	const FText TEXT_Index = LOCTEXT("Index", "$Index");
+	const FText TEXT_PositionX = LOCTEXT("PositionX", "$Position.X");
+	const FText TEXT_PositionY = LOCTEXT("PositionY", "$Position.Y");
+	const FText TEXT_PositionZ = LOCTEXT("PositionZ", "$Position.Z");
+	const FText TEXT_RotationRoll = LOCTEXT("RotationRoll", "$Rotation.Roll");
+	const FText TEXT_RotationPitch = LOCTEXT("RotationPitch", "$Rotation.Pitch");
+	const FText TEXT_RotationYaw = LOCTEXT("RotationYaw", "$Rotation.Yaw");
+	const FText TEXT_ScaleX = LOCTEXT("ScaleX", "$Scale.X");
+	const FText TEXT_ScaleY = LOCTEXT("ScaleY", "$Scale.Y");
+	const FText TEXT_ScaleZ = LOCTEXT("ScaleZ", "$Scale.Z");
+	const FText TEXT_BoundsMinX = LOCTEXT("BoundsMinX", "$BoundsMin.X");
+	const FText TEXT_BoundsMinY = LOCTEXT("BoundsMinY", "$BoundsMin.Y");
+	const FText TEXT_BoundsMinZ = LOCTEXT("BoundsMinZ", "$BoundsMin.Z");
+	const FText TEXT_BoundsMaxX = LOCTEXT("BoundsMaxX", "$BoundsMax.X");
+	const FText TEXT_BoundsMaxY = LOCTEXT("BoundsMaxY", "$BoundsMax.Y");
+	const FText TEXT_BoundsMaxZ = LOCTEXT("BoundsMaxZ", "$BoundsMax.Z");
+	const FText TEXT_ColorR = LOCTEXT("ColorR", "$Color.R");
+	const FText TEXT_ColorG = LOCTEXT("ColorG", "$Color.G");
+	const FText TEXT_ColorB = LOCTEXT("ColorB", "$Color.B");
+	const FText TEXT_ColorA = LOCTEXT("ColorA", "$Color.A");
+	const FText TEXT_Density = LOCTEXT("Density", "$Density");
+	const FText TEXT_Steepness = LOCTEXT("Steepness", "$Steepness");
+	const FText TEXT_Seed = LOCTEXT("Seed", "$Seed");
+	const FText TEXT_MetadataEntry = LOCTEXT("MetadataEntry", "Entry Key");
+	const FText TEXT_MetadataEntryParent = LOCTEXT("MetadataEntryParent", "Parent Key");
+}
+
 void IPCGSpatialDataVisualization::ExecuteDebugDisplay(FPCGContext* Context, const UPCGData* Data, AActor* TargetActor) const
 {
 	if (!TargetActor)
@@ -163,6 +226,75 @@ void IPCGSpatialDataVisualization::ExecuteDebugDisplay(FPCGContext* Context, con
 	}
 }
 
+FPCGTableVisualizerInfo IPCGSpatialDataVisualization::GetTableVisualizerInfo(const UPCGData* Data) const
+{
+	using namespace PCGDataVisualizationHelpers;
+	using namespace PCGPointDataVisualizationConstants;
+
+	// Collapse to point representation for visualization.
+	const UPCGPointData* PointData = CollapseToDebugPointData(/*Context=*/nullptr, Data);
+
+	FPCGTableVisualizerInfo Info;
+	Info.Data = PointData;
+
+	// Columns
+	Info.ColumnInfos.Add(CreateColumnInfo(PointData, NAME_Index,         TEXT_Index));
+	Info.ColumnInfos.Add(CreateColumnInfo(PointData, NAME_PositionX,     TEXT_PositionX));
+	Info.ColumnInfos.Add(CreateColumnInfo(PointData, NAME_PositionY,     TEXT_PositionY));
+	Info.ColumnInfos.Add(CreateColumnInfo(PointData, NAME_PositionZ,     TEXT_PositionZ));
+	Info.ColumnInfos.Add(CreateColumnInfo(PointData, NAME_RotationRoll,  TEXT_RotationRoll));
+	Info.ColumnInfos.Add(CreateColumnInfo(PointData, NAME_RotationPitch, TEXT_RotationPitch));
+	Info.ColumnInfos.Add(CreateColumnInfo(PointData, NAME_RotationYaw,   TEXT_RotationYaw));
+	Info.ColumnInfos.Add(CreateColumnInfo(PointData, NAME_ScaleX,        TEXT_ScaleX));
+	Info.ColumnInfos.Add(CreateColumnInfo(PointData, NAME_ScaleY,        TEXT_ScaleY));
+	Info.ColumnInfos.Add(CreateColumnInfo(PointData, NAME_ScaleZ,        TEXT_ScaleZ));
+	Info.ColumnInfos.Add(CreateColumnInfo(PointData, NAME_BoundsMinX,    TEXT_BoundsMinX));
+	Info.ColumnInfos.Add(CreateColumnInfo(PointData, NAME_BoundsMinY,    TEXT_BoundsMinY));
+	Info.ColumnInfos.Add(CreateColumnInfo(PointData, NAME_BoundsMinZ,    TEXT_BoundsMinZ));
+	Info.ColumnInfos.Add(CreateColumnInfo(PointData, NAME_BoundsMaxX,    TEXT_BoundsMaxX));
+	Info.ColumnInfos.Add(CreateColumnInfo(PointData, NAME_BoundsMaxY,    TEXT_BoundsMaxY));
+	Info.ColumnInfos.Add(CreateColumnInfo(PointData, NAME_BoundsMaxZ,    TEXT_BoundsMaxZ));
+	Info.ColumnInfos.Add(CreateColumnInfo(PointData, NAME_ColorR,        TEXT_ColorR));
+	Info.ColumnInfos.Add(CreateColumnInfo(PointData, NAME_ColorG,        TEXT_ColorG));
+	Info.ColumnInfos.Add(CreateColumnInfo(PointData, NAME_ColorB,        TEXT_ColorB));
+	Info.ColumnInfos.Add(CreateColumnInfo(PointData, NAME_ColorA,        TEXT_ColorA));
+	Info.ColumnInfos.Add(CreateColumnInfo(PointData, NAME_Density,       TEXT_Density));
+	Info.ColumnInfos.Add(CreateColumnInfo(PointData, NAME_Steepness,     TEXT_Steepness));
+	Info.ColumnInfos.Add(CreateColumnInfo(PointData, NAME_Seed,          TEXT_Seed));
+
+	if (FPCGEditorCommon::CVarShowAdvancedAttributesFields.GetValueOnAnyThread())
+	{
+		Info.ColumnInfos.Add(CreateColumnInfo(PointData, NAME_MetadataEntry, TEXT_MetadataEntry, []()
+		{
+			return TSharedPtr<const IPCGAttributeAccessor>(PCGAttributeAccessorHelpers::CreatePropertyAccessor(GET_MEMBER_NAME_CHECKED(FPCGPoint, MetadataEntry), FPCGPoint::StaticStruct()).Release());
+		}));
+
+		Info.ColumnInfos.Add(CreateColumnInfo(PointData, NAME_MetadataEntryParent, TEXT_MetadataEntryParent, [PointData]()
+		{
+			return MakeShared<FPCGCustomPointAccessor<int64>>([PointData](const FPCGPoint& Point, void* OutValue)
+			{
+				if (const UPCGMetadata* Metadata = PointData->ConstMetadata())
+				{
+					*reinterpret_cast<int64*>(OutValue) = Metadata->GetParentKey(Point.MetadataEntry);
+					return true;
+				}
+				return false;
+			}, nullptr);
+		}));
+	}
+
+	// Add Metadata Columns
+	CreateMetadataColumnInfos(PointData, Info.ColumnInfos);
+
+	// Column Sorting
+	Info.SortingColumn = NAME_Index;
+
+	// Accessor Keys
+	Info.AccessorKeys = TSharedPtr<const IPCGAttributeAccessorKeys>(PCGAttributeAccessorHelpers::CreateConstKeys(PointData, FPCGAttributePropertyInputSelector()).Release());
+
+	return Info;
+}
+
 const UPCGPointData* IPCGSpatialDataVisualization::CollapseToDebugPointData(FPCGContext* Context, const UPCGData* Data) const
 {
 	if (const UPCGSpatialData* SpatialData = Cast<UPCGSpatialData>(Data))
diff --git a/Engine/Plugins/PCG/Source/PCGEditor/Private/DataVisualizations/PCGSpatialDataVisualization.h b/Engine/Plugins/PCG/Source/PCGEditor/Private/DataVisualizations/PCGSpatialDataVisualization.h
--- a/Engine/Plugins/PCG/Source/PCGEditor/Private/DataVisualizations/PCGSpatialDataVisualization.h
+++ b/Engine/Plugins/PCG/Source/PCGEditor/Private/DataVisualizations/PCGSpatialDataVisualization.h
@@ -14,6 +14,10 @@ struct FPCGContext;
 class PCGEDITOR_API IPCGSpatialDataVisualization : public IPCGDataVisualization
 {
 public:
+	// ~Begin IPCGDataVisualization interface
 	virtual void ExecuteDebugDisplay(FPCGContext* Context, const UPCGData* Data, AActor* TargetActor) const override;
+	virtual FPCGTableVisualizerInfo GetTableVisualizerInfo(const UPCGData* Data) const override;
+	// ~End IPCGDataVisualization interface
+
 	virtual const UPCGPointData* CollapseToDebugPointData(FPCGContext* Context, const UPCGData* Data) const;
 };
diff --git a/Engine/Plugins/PCG/Source/PCGEditor/Private/PCGEditorCommon.h b/Engine/Plugins/PCG/Source/PCGEditor/Private/PCGEditorCommon.h
--- a/Engine/Plugins/PCG/Source/PCGEditor/Private/PCGEditorCommon.h
+++ b/Engine/Plugins/PCG/Source/PCGEditor/Private/PCGEditorCommon.h
@@ -2,6 +2,8 @@
 
 #pragma once
 
+#include "HAL/IConsoleManager.h"
+
 #include "PCGEditorCommon.generated.h"
 
 UENUM()
@@ -42,4 +44,9 @@ namespace FPCGEditorCommon
 	const FName RenderTargetDataType = FName(TEXT("Render Target Data"));
 	const FName VolumeDataType = FName(TEXT("Volume Data"));
 	const FName PrimitiveDataType = FName(TEXT("Primitive Data"));	
+
+	static TAutoConsoleVariable<bool> CVarShowAdvancedAttributesFields(
+		TEXT("pcg.graph.ShowAdvancedAttributes"),
+		false,
+		TEXT("Control whether advanced attributes/properties are shown in the PCG graph editor"));
 };
diff --git a/Engine/Plugins/PCG/Source/PCGEditor/Private/PCGEditorModule.cpp b/Engine/Plugins/PCG/Source/PCGEditor/Private/PCGEditorModule.cpp
--- a/Engine/Plugins/PCG/Source/PCGEditor/Private/PCGEditorModule.cpp
+++ b/Engine/Plugins/PCG/Source/PCGEditor/Private/PCGEditorModule.cpp
@@ -11,10 +11,12 @@
 #include "PCGEditorUtils.h"
 #include "PCGEngineSettings.h"
 #include "PCGModule.h"
+#include "PCGParamData.h"
 #include "PCGSubsystem.h"
 #include "PCGVolumeFactory.h"
 #include "Data/PCGSpatialData.h"
 #include "Data/PCGSplineData.h"
+#include "DataVisualizations/PCGParamDataVisualization.h"
 #include "DataVisualizations/PCGSpatialDataVisualization.h"
 #include "DataVisualizations/PCGSplineDataVisualization.h"
 #include "Grid/PCGPartitionActor.h"
@@ -534,6 +536,7 @@ void FPCGEditorModule::UnregisterSettings()
 void FPCGEditorModule::RegisterPCGDataVisualizations()
 {
 	FPCGDataVisualizationRegistry& DataVisRegistry = FPCGModule::GetMutablePCGDataVisualizationRegistry();
+	DataVisRegistry.InternalRegistry.Add(UPCGParamData::StaticClass(), MakeUnique<const IPCGParamDataVisualization>());
 	DataVisRegistry.InternalRegistry.Add(UPCGSpatialData::StaticClass(), MakeUnique<const IPCGSpatialDataVisualization>());
 	DataVisRegistry.InternalRegistry.Add(UPCGSplineData::StaticClass(), MakeUnique<const IPCGSplineDataVisualization>());
 }
diff --git a/Engine/Plugins/PCG/Source/PCGEditor/Private/SPCGEditorGraphAttributeListView.cpp b/Engine/Plugins/PCG/Source/PCGEditor/Private/SPCGEditorGraphAttributeListView.cpp
--- a/Engine/Plugins/PCG/Source/PCGEditor/Private/SPCGEditorGraphAttributeListView.cpp
+++ b/Engine/Plugins/PCG/Source/PCGEditor/Private/SPCGEditorGraphAttributeListView.cpp
@@ -4,6 +4,7 @@
 
 #include "PCGComponent.h"
 #include "PCGData.h"
+#include "PCGDataVisualization.h"
 #include "PCGGraph.h"
 #include "PCGNode.h"
 #include "PCGParamData.h"
@@ -34,75 +35,15 @@
 
 #define LOCTEXT_NAMESPACE "SPCGEditorGraphAttributeListView"
 
-static TAutoConsoleVariable<bool> CVarShowAdvancedAttributesFields(
-	TEXT("pcg.graph.ShowAdvancedAttributes"),
-	false,
-	TEXT("Control whether advanced attributes/properties are shown in the PCG graph editor"));
-
 namespace PCGEditorGraphAttributeListView
 {
-	const FString NoneAttributeId = TEXT("@None");
 	const FText NoPinAvailableText = LOCTEXT("NoPinAvailableText", "No pins");
-	const FText LastLabelFormat = LOCTEXT("LastLabelFormat", "Last attribute: {0} / {1}");
+	const FText CrcLabelFormat = LOCTEXT("InfoTextBlockWithCrcFmt", "{0} | CRC: {1}");
+	const FText LastLabelFormat = LOCTEXT("InfoTextBlockWithLastAttributeFmt", "{0} | Last attribute: {1}");
 	const FText NoDataAvailableText = LOCTEXT("NoDataAvailableText", "No data available");
 	const FText NoNodeInspectedText = LOCTEXT("NoNodeInspectedText", "No node being inspected");
 	const FText NoNodeInspectedToolTip = LOCTEXT("NoNodeInspectedToolTip", "Inspect a node using the right click menu");
-
-	/** Names of the columns in the attribute list */
-	const FName NAME_Index = FName(TEXT("$Index"));
-	const FName NAME_PointPositionX = FName(TEXT("$Position.X"));
-	const FName NAME_PointPositionY = FName(TEXT("$Position.Y"));
-	const FName NAME_PointPositionZ = FName(TEXT("$Position.Z"));
-	const FName NAME_PointRotationRoll = FName(TEXT("$Rotation.Roll"));
-	const FName NAME_PointRotationPitch = FName(TEXT("$Rotation.Pitch"));
-	const FName NAME_PointRotationYaw = FName(TEXT("$Rotation.Yaw"));
-	const FName NAME_PointScaleX = FName(TEXT("$Scale.X"));
-	const FName NAME_PointScaleY = FName(TEXT("$Scale.Y"));
-	const FName NAME_PointScaleZ = FName(TEXT("$Scale.Z"));
-	const FName NAME_PointBoundsMinX = FName(TEXT("$BoundsMin.X"));
-	const FName NAME_PointBoundsMinY = FName(TEXT("$BoundsMin.Y"));
-	const FName NAME_PointBoundsMinZ = FName(TEXT("$BoundsMin.Z"));
-	const FName NAME_PointBoundsMaxX = FName(TEXT("$BoundsMax.X"));
-	const FName NAME_PointBoundsMaxY = FName(TEXT("$BoundsMax.Y"));
-	const FName NAME_PointBoundsMaxZ = FName(TEXT("$BoundsMax.Z"));
-	const FName NAME_PointColorR = FName(TEXT("$Color.R"));
-	const FName NAME_PointColorG = FName(TEXT("$Color.G"));
-	const FName NAME_PointColorB = FName(TEXT("$Color.B"));
-	const FName NAME_PointColorA = FName(TEXT("$Color.A"));
-	const FName NAME_PointDensity = FName(TEXT("$Density"));
-	const FName NAME_PointSteepness = FName(TEXT("$Steepness"));
-	const FName NAME_PointSeed = FName(TEXT("$Seed"));
-	const FName NAME_PointMetadataEntry = FName(TEXT("MetadataEntry"));
-	const FName NAME_PointMetadataEntryParent = FName(TEXT("PointMetadataEntryParent"));
-
-	/** Labels of the columns */
 	const FText TEXT_IndexLabel = LOCTEXT("IndexLabel", "Index");
-	const FText TEXT_PointPositionLabelX = LOCTEXT("PointPositionLabelX", "Position.X");
-	const FText TEXT_PointPositionLabelY = LOCTEXT("PointPositionLabelY", "Position.Y");
-	const FText TEXT_PointPositionLabelZ = LOCTEXT("PointPositionLabelZ", "Position.Z");
-	const FText TEXT_PointRotationLabelRoll = LOCTEXT("PointRotationLabelRoll", "Rotation.Roll");
-	const FText TEXT_PointRotationLabelPitch = LOCTEXT("PointRotationLabelPitch", "Rotation.Pitch");
-	const FText TEXT_PointRotationLabelYaw = LOCTEXT("PointRotationLabelYaw", "Rotation.Yaw");
-	const FText TEXT_PointScaleLabelX = LOCTEXT("PointScaleLabelX", "Scale.X");
-	const FText TEXT_PointScaleLabelY = LOCTEXT("PointScaleLabelY", "Scale.Y");
-	const FText TEXT_PointScaleLabelZ = LOCTEXT("PointScaleLabelZ", "Scale.Z");
-	const FText TEXT_PointBoundsLabelMinX = LOCTEXT("PointBoundsMinX", "BoundsMin.X");
-	const FText TEXT_PointBoundsLabelMinY = LOCTEXT("PointBoundsMinY", "BoundsMin.Y");
-	const FText TEXT_PointBoundsLabelMinZ = LOCTEXT("PointBoundsMinZ", "BoundsMin.Z");
-	const FText TEXT_PointBoundsLabelMaxX = LOCTEXT("PointBoundsMaxX", "BoundsMax.X");
-	const FText TEXT_PointBoundsLabelMaxY = LOCTEXT("PointBoundsMaxY", "BoundsMax.Y");
-	const FText TEXT_PointBoundsLabelMaxZ = LOCTEXT("PointBoundsMaxZ", "BoundsMax.Z");
-	const FText TEXT_PointColorLabelR = LOCTEXT("PointColorR", "Color.R");
-	const FText TEXT_PointColorLabelG = LOCTEXT("PointColorG", "Color.G");
-	const FText TEXT_PointColorLabelB = LOCTEXT("PointColorB", "Color.B");
-	const FText TEXT_PointColorLabelA = LOCTEXT("PointColorA", "Color.A");
-	const FText TEXT_PointDensityLabel = LOCTEXT("PointDensityLabel", "Density");
-	const FText TEXT_PointSteepnessLabel = LOCTEXT("PointSteepnessLabel", "Steepness");
-	const FText TEXT_PointSeedLabel = LOCTEXT("PointSeedLabel", "Seed");
-	const FText TEXT_PointMetadataEntryLabel = LOCTEXT("PointMetadataEntryLabel", "Entry Key");
-	const FText TEXT_PointMetadataEntryParentLabel = LOCTEXT("PointMetadataEntryParentLabel", "Parent Key");
-
-	constexpr float MaxColumnWidth = 200.0f;
 
 	bool IsGraphCacheDebuggingEnabled()
 	{
@@ -141,7 +82,7 @@ void FPCGListViewUpdater::Launch()
 
 void FPCGListViewUpdater::AsyncSort()
 {
-	if (const FPCGColumnData* Data = ColumnData.Find(SortMode == EColumnSortMode::None ? PCGEditorGraphAttributeListView::NAME_Index : SortingColumn))
+	if (const FPCGColumnData* Data = ColumnData.Find(SortingColumn))
 	{
 		if (Data->DataAccessor.IsValid() && Data->DataKeys.IsValid() && Data->DataKeys->GetNum() == ListViewItems.Num())
 		{
@@ -320,14 +261,12 @@ SPCGEditorGraphAttributeListView::~SPCGEditorGraphAttributeListView()
 		PCGEditorPtr.Pin()->OnInspectedStackChangedDelegate.RemoveAll(this);
 	}
 
-	CollapsedPointData.Reset();
+	DataStrongPtr.Reset();
 }
 
 void SPCGEditorGraphAttributeListView::Construct(const FArguments& InArgs, TSharedPtr<FPCGEditor> InPCGEditor)
 {
 	PCGEditorPtr = InPCGEditor;
-	SortingColumn = PCGEditorGraphAttributeListView::NAME_Index;
-
 	PCGEditorPtr.Pin()->OnInspectedStackChangedDelegate.AddSP(this, &SPCGEditorGraphAttributeListView::OnInspectedStackChanged);
 
 	TextFilter = MakeShareable(new FTextFilterExpressionEvaluator(ETextFilterExpressionEvaluatorMode::Complex));
@@ -546,6 +485,23 @@ void SPCGEditorGraphAttributeListView::Tick(const FGeometry& AllottedGeometry, c
 		{
 			ListView->SetItemsSource(&FilteredListViewItems);
 			ListView->RequestListRefresh();
+
+			InfoTextBlock->SetText(FText::Format(LOCTEXT("InfoTextBlockFmt", "Showing {0}/{1} entries"), FilteredListViewItems.Num(), ListViewItems.Num()));
+
+			if (PCGEditorGraphAttributeListView::IsGraphCacheDebuggingEnabled())
+			{
+				// If cache debugging enabled, write CRC to help diagnose missed-dependency issues
+				const FPCGDataCollection* InspectionData = GetInspectionData();
+				const int32 DataIndex = GetSelectedDataIndex();
+				const FPCGCrc Crc = (InspectionData && InspectionData->DataCrcs.IsValidIndex(DataIndex)) ? InspectionData->DataCrcs[DataIndex] : FPCGCrc(0);
+				InfoTextBlock->SetText(FText::Format(PCGEditorGraphAttributeListView::CrcLabelFormat, InfoTextBlock->GetText(), Crc.GetValue()));
+			}
+
+			if (DataStrongPtr.IsValid() && DataStrongPtr->HasCachedLastSelector())
+			{
+				const FText LastSelector = DataStrongPtr->GetCachedLastSelector().GetDisplayText();
+				InfoTextBlock->SetText(FText::Format(PCGEditorGraphAttributeListView::LastLabelFormat, InfoTextBlock->GetText(), LastSelector));
+			}
 		}
 
 		CurrentUpdateTask.Reset();
@@ -676,7 +632,7 @@ void SPCGEditorGraphAttributeListView::RefreshAttributeList()
 	ListViewItems.Empty();
 	ListViewHeader->ClearColumns();
 	InfoTextBlock->SetText(FText::GetEmpty());
-	CollapsedPointData.Reset();
+	DataStrongPtr.Reset();
 
 	const FPCGDataCollection* InspectionData = GetInspectionData();
 	if (!InspectionData)
@@ -694,74 +650,53 @@ void SPCGEditorGraphAttributeListView::RefreshAttributeList()
 	const UPCGData* PCGData = TaggedData.Data;
 	const FPCGCrc Crc = InspectionData->DataCrcs.IsValidIndex(DataIndex) ? InspectionData->DataCrcs[DataIndex] : FPCGCrc(0);
 
-	if (const UPCGParamData* PCGParamData = Cast<const UPCGParamData>(PCGData))
+	if (!PCGData)
 	{
-		if (const UPCGMetadata* PCGMetadata = PCGParamData->ConstMetadata())
-		{
-			AddColumn(PCGData, PCGEditorGraphAttributeListView::NAME_Index, PCGEditorGraphAttributeListView::TEXT_IndexLabel);
-			GenerateColumnsFromMetadata(PCGData, PCGMetadata);
+		return;
+	}
 
-			const PCGMetadataEntryKey ItemKeyLowerBound = PCGMetadata->GetItemKeyCountForParent();
-			const PCGMetadataEntryKey ItemKeyUpperBound = PCGMetadata->GetItemCountForChild();
-			for (PCGMetadataEntryKey MetadataItemKey = ItemKeyLowerBound; MetadataItemKey < ItemKeyUpperBound; ++MetadataItemKey)
-			{
-				PCGListviewItemPtr ListViewItem = MakeShared<FPCGListViewItem>();
-				ListViewItem->Index = MetadataItemKey - ItemKeyLowerBound;
-				ListViewItems.Add(ListViewItem);
-			}
+	const FPCGDataVisualizationRegistry& DataVisRegistry = FPCGModule::GetConstPCGDataVisualizationRegistry();
+	const UPCGData* DataToVisualize = PCGData;
 
-			if (!PCGEditorGraphAttributeListView::IsGraphCacheDebuggingEnabled())
-			{
-				InfoTextBlock->SetText(FText::Format(LOCTEXT("MetadataInfoTextBlockFmt", "Number of metadata: {0}"), ItemKeyUpperBound - ItemKeyLowerBound));
-			}
-			else
-			{
-				// If cache debugging enabled, write CRC to help diagnose missed-dependency issues
-				InfoTextBlock->SetText(FText::Format(LOCTEXT("MetadataInfoTextBlockWithCrcFmt", "Number of metadata: {0}  CRC: {1}"), ItemKeyUpperBound - ItemKeyLowerBound, Crc.GetValue()));
-			}
-		}
-	}
-	else if (const UPCGSpatialData* PCGSpatialData = Cast<const UPCGSpatialData>(PCGData))
+	if (const IPCGDataVisualization* DataVisualization = DataVisRegistry.GetDataVisualization(PCGData->GetClass()))
 	{
-		if (const UPCGPointData* PCGPointData = PCGSpatialData->ToPointData())
+		const FPCGTableVisualizerInfo TableVisualizerInfo = DataVisualization->GetTableVisualizerInfo(PCGData);
+		DataToVisualize = TableVisualizerInfo.Data;
+
+		for (const FPCGTableVisualizerColumnInfo& ColumnInfo : TableVisualizerInfo.ColumnInfos)
 		{
-			CollapsedPointData.Reset(PCGPointData);
+			AddColumn(DataToVisualize, ColumnInfo, TableVisualizerInfo.AccessorKeys);
+		}
 
-			const UPCGMetadata* PCGMetadata = PCGPointData->ConstMetadata();
+		SortingColumn = TableVisualizerInfo.SortingColumn;
+		SortMode = (EColumnSortMode::Type)TableVisualizerInfo.SortingMode;
 
-			AddPointDataColumns(PCGPointData);
-			GenerateColumnsFromMetadata(PCGPointData, PCGMetadata);
+		const uint32 NumEntries = TableVisualizerInfo.AccessorKeys.IsValid() ? TableVisualizerInfo.AccessorKeys->GetNum() : 0;
+		ListViewItems.Reserve(NumEntries);
 
-			const TArray<FPCGPoint>& PCGPoints = PCGPointData->GetPoints();
-			const int32 NumPoints = PCGPoints.Num();
-			ListViewItems.Reserve(NumPoints);
-			for (int32 PointIndex = 0; PointIndex < NumPoints; PointIndex++)
-			{
-				const FPCGPoint& PCGPoint = PCGPoints[PointIndex];
-				// TODO: Investigate swapping out the shared ptr's for better performance on huge data sets
-				PCGListviewItemPtr ListViewItem = MakeShared<FPCGListViewItem>();
-				ListViewItem->Index = PointIndex;
-				ListViewItem->PCGPoint = &PCGPoint;
-				ListViewItems.Add(ListViewItem);
-			}
+		const TArray<FPCGPoint>& Points = Cast<UPCGPointData>(DataStrongPtr.Get()) ? Cast<UPCGPointData>(DataStrongPtr.Get())->GetPoints() : TArray<FPCGPoint>();
 
-			if (!PCGEditorGraphAttributeListView::IsGraphCacheDebuggingEnabled())
-			{
-				InfoTextBlock->SetText(FText::Format(LOCTEXT("PointInfoTextBlockFmt", "Number of points: {0}"), NumPoints));
-			}
-			else
+		for (uint32 Index = 0; Index < NumEntries; ++Index)
+		{
+			PCGListviewItemPtr ListViewItem = MakeShared<FPCGListViewItem>();
+			ListViewItem->Index = Index;
+
+			// TODO: It would be preferable to have a custom override for double-click behavior instead of relying on collapsed point data.
+			ListViewItem->DoubleClickCallback = [DataPtr = DataStrongPtr.Get()](int Index)
 			{
-				// If cache debugging enabled, write CRC to help diagnose missed-dependency issues
-				InfoTextBlock->SetText(FText::Format(LOCTEXT("PointInfoTextBlockWithCrcFmt", "Number of points: {0}, CRC: {1}"), NumPoints, Crc.GetValue()));
-			}
+				if (const UPCGPointData* PointData = Cast<UPCGPointData>(DataPtr))
+				{
+					const FPCGPoint& Point = PointData->GetPoints()[Index];
+					const FBox BoundingBox = Point.GetLocalBounds().TransformBy(Point.Transform.ToMatrixWithScale());
+					GEditor->MoveViewportCamerasToBox(BoundingBox, true, 2.5f);
+				}
+			};
+
+			ListViewItems.Add(ListViewItem);
 		}
 	}
 
-	if (PCGData && PCGData->HasCachedLastSelector())
-	{
-		const FText LastSelector = PCGData->GetCachedLastSelector().GetDisplayText();
-		InfoTextBlock->SetText(FText::Format(PCGEditorGraphAttributeListView::LastLabelFormat, LastSelector, InfoTextBlock->GetText()));
-	}
+	DataStrongPtr.Reset(DataToVisualize);
 
 	ListView->SetItemsSource(&ListViewItems);
 	ListView->RequestListRefresh();
@@ -900,11 +835,6 @@ TSharedRef<SWidget> SPCGEditorGraphAttributeListView::OnGenerateFilterMenu()
 
 	for (const SHeaderRow::FColumn& Column : Columns)
 	{
-		if (Column.ColumnId == PCGEditorGraphAttributeListView::NAME_Index)
-		{
-			continue;
-		}
-
 		MenuBuilder.AddMenuEntry(
 			Column.DefaultText,
 			Column.DefaultTooltip,
@@ -982,81 +912,6 @@ int32 SPCGEditorGraphAttributeListView::GetSelectedDataIndex() const
 	return Index;
 }
 
-void SPCGEditorGraphAttributeListView::GenerateColumnsFromMetadata(const UPCGData* InPCGData, const UPCGMetadata* PCGMetadata)
-{
-	TArray<FName> AttributeNames;
-	TArray<EPCGMetadataTypes> AttributeTypes;
-	PCGMetadata->GetAttributes(AttributeNames, AttributeTypes);
-
-	for (int32 I = 0; I < AttributeNames.Num(); I++)
-	{
-		const FName& AttributeName = AttributeNames[I];
-		const EPCGMetadataTypes AttributeType = AttributeTypes[I];
-		FName ColumnName = AttributeName;
-
-		switch (AttributeType)
-		{
-		case EPCGMetadataTypes::Float:
-		case EPCGMetadataTypes::Double:
-		case EPCGMetadataTypes::Integer32:
-		case EPCGMetadataTypes::Integer64:
-		case EPCGMetadataTypes::Boolean:
-		case EPCGMetadataTypes::String:
-		case EPCGMetadataTypes::Name:
-		case EPCGMetadataTypes::SoftObjectPath:
-		case EPCGMetadataTypes::SoftClassPath:
-			{
-				AddMetadataColumn(InPCGData, ColumnName, AttributeType);
-				break;
-			}
-		case EPCGMetadataTypes::Vector2:
-		{
-			AddMetadataColumn(InPCGData, ColumnName, AttributeType, TEXT(".X"));
-			AddMetadataColumn(InPCGData, ColumnName, AttributeType, TEXT(".Y"));
-			break;
-		}
-		case EPCGMetadataTypes::Vector:
-			{
-				AddMetadataColumn(InPCGData, ColumnName, AttributeType, TEXT(".X"));
-				AddMetadataColumn(InPCGData, ColumnName, AttributeType, TEXT(".Y"));
-				AddMetadataColumn(InPCGData, ColumnName, AttributeType, TEXT(".Z"));
-				break;
-			}
-		case EPCGMetadataTypes::Vector4:
-		case EPCGMetadataTypes::Quaternion:
-			{
-				AddMetadataColumn(InPCGData, ColumnName, AttributeType, TEXT(".X"));
-				AddMetadataColumn(InPCGData, ColumnName, AttributeType, TEXT(".Y"));
-				AddMetadataColumn(InPCGData, ColumnName, AttributeType, TEXT(".Z"));
-				AddMetadataColumn(InPCGData, ColumnName, AttributeType, TEXT(".W"));
-				break;
-			}
-		case EPCGMetadataTypes::Rotator:
-			{
-				AddMetadataColumn(InPCGData, ColumnName, AttributeType, TEXT(".Roll"));
-				AddMetadataColumn(InPCGData, ColumnName, AttributeType, TEXT(".Pitch"));
-				AddMetadataColumn(InPCGData, ColumnName, AttributeType, TEXT(".Yaw"));
-				break;
-			}
-		case EPCGMetadataTypes::Transform:
-			{
-				AddMetadataColumn(InPCGData, ColumnName, AttributeType, TEXT(".Position.X"));
-				AddMetadataColumn(InPCGData, ColumnName, AttributeType, TEXT(".Position.Y"));
-				AddMetadataColumn(InPCGData, ColumnName, AttributeType, TEXT(".Position.Z"));
-				AddMetadataColumn(InPCGData, ColumnName, AttributeType, TEXT(".Rotation.Roll"));
-				AddMetadataColumn(InPCGData, ColumnName, AttributeType, TEXT(".Rotation.Pitch"));
-				AddMetadataColumn(InPCGData, ColumnName, AttributeType, TEXT(".Rotation.Yaw"));
-				AddMetadataColumn(InPCGData, ColumnName, AttributeType, TEXT(".Scale.X"));
-				AddMetadataColumn(InPCGData, ColumnName, AttributeType, TEXT(".Scale.Y"));
-				AddMetadataColumn(InPCGData, ColumnName, AttributeType, TEXT(".Scale.Z"));
-				break;
-			}
-		default:
-			break;
-		}
-	}
-}
-
 void SPCGEditorGraphAttributeListView::ToggleAllAttributes()
 {
 	const TArray<FName> HiddenColumns = ListViewHeader->GetHiddenColumnIds();
@@ -1072,10 +927,7 @@ void SPCGEditorGraphAttributeListView::ToggleAllAttributes()
 		const TIndirectArray<SHeaderRow::FColumn>& Columns = ListViewHeader->GetColumns();
 		for (const SHeaderRow::FColumn& Column : Columns)
 		{
-			if (Column.ColumnId != PCGEditorGraphAttributeListView::NAME_Index)
-			{
-				ListViewHeader->SetShowGeneratedColumn(Column.ColumnId, /*InShow=*/false);
-			}
+			ListViewHeader->SetShowGeneratedColumn(Column.ColumnId, /*InShow=*/false);
 		}
 	}
 }
@@ -1092,11 +944,6 @@ ECheckBoxState SPCGEditorGraphAttributeListView::GetAnyAttributeEnabledState() c
 
 	for (const SHeaderRow::FColumn& Column : ListViewHeader->GetColumns())
 	{
-		if (Column.ColumnId == PCGEditorGraphAttributeListView::NAME_Index)
-		{
-			continue;
-		}
-
 		bAllEnabled &= Column.bIsVisible;
 		bAnyEnabled |= Column.bIsVisible;
 	}
@@ -1131,10 +978,9 @@ void SPCGEditorGraphAttributeListView::OnItemDoubleClicked(PCGListviewItemPtr It
 {
 	check(Item);
 
-	if (const FPCGPoint* Point = Item->PCGPoint)
+	if (Item->DoubleClickCallback)
 	{
-		const FBox BoundingBox = Point->GetLocalBounds().TransformBy(Point->Transform.ToMatrixWithScale());
-		GEditor->MoveViewportCamerasToBox(BoundingBox, true, 2.5f);
+		Item->DoubleClickCallback(Item->Index);
 	}
 }
 
@@ -1197,167 +1043,25 @@ FReply SPCGEditorGraphAttributeListView::OnListViewKeyDown(const FGeometry& /*In
 	return FReply::Unhandled();
 }
 
-void SPCGEditorGraphAttributeListView::AddColumn(const UPCGData* InPCGData, const FName& InColumnId, const FText& ColumnLabel)
+void SPCGEditorGraphAttributeListView::AddColumn(const UPCGData* InPCGData, const FPCGTableVisualizerColumnInfo& InColumnInfo, TSharedPtr<const IPCGAttributeAccessorKeys> InAccessorKeys)
 {
-	FText ToolTip;
-
-	if (InPCGData)
-	{
-		const FString ColumnIdString = InColumnId.ToString();
-
-		FPCGAttributePropertyInputSelector TargetSelector;
-		TargetSelector.Update(ColumnIdString);
-
-		FPCGColumnData& ColumnData = PCGColumnData.Add(InColumnId);
-
-		if (InColumnId == PCGEditorGraphAttributeListView::NAME_PointMetadataEntry)
-		{
-			if (const UPCGPointData* PCGPointData = Cast<const UPCGPointData>(InPCGData))
-			{
-				TUniquePtr<const IPCGAttributeAccessor> DataAccessor = PCGAttributeAccessorHelpers::CreatePropertyAccessor(GET_MEMBER_NAME_CHECKED(FPCGPoint, MetadataEntry), FPCGPoint::StaticStruct());
-				ColumnData.DataAccessor = TSharedPtr<const IPCGAttributeAccessor>(DataAccessor.Release());
-				ColumnData.DataKeys = MakeShared<FPCGAttributeAccessorKeysPoints>(PCGPointData->GetPoints());
-			}
-		}
-		else if (InColumnId == PCGEditorGraphAttributeListView::NAME_PointMetadataEntryParent)
-		{
-			if (const UPCGPointData* PCGPointData = Cast<const UPCGPointData>(InPCGData))
-			{
-				ColumnData.DataAccessor = MakeShared<FPCGCustomPointAccessor<int64>>([PCGPointData](const FPCGPoint& Point, void* OutValue)
-				{
-					if (const UPCGMetadata* Metadata = PCGPointData->Metadata)
-					{
-						*reinterpret_cast<int64*>(OutValue) = Metadata->GetParentKey(Point.MetadataEntry);
-						return true;
-					}
-					return false;
-				}, nullptr);
-				ColumnData.DataKeys = MakeShared<FPCGAttributeAccessorKeysPoints>(PCGPointData->GetPoints());
-			}
-		}
-		else
-		{
-			ColumnData.DataAccessor = TSharedPtr<const IPCGAttributeAccessor>(PCGAttributeAccessorHelpers::CreateConstAccessor(InPCGData, TargetSelector).Release());
-			ColumnData.DataKeys = TSharedPtr<const IPCGAttributeAccessorKeys>(PCGAttributeAccessorHelpers::CreateConstKeys(InPCGData, TargetSelector).Release());
-		}
-
-		ToolTip = FText::FromString(PCG::Private::GetTypeName(ColumnData.DataAccessor->GetUnderlyingType()));
-	}
-
-	const float ColumnWidth = PCGEditorGraphAttributeListView::CalculateColumnWidth(ColumnLabel);
+	FPCGColumnData& ColumnData = PCGColumnData.Add(InColumnInfo.Id);
+	ColumnData.DataAccessor = InColumnInfo.Accessor;
+	ColumnData.DataKeys = InAccessorKeys;
 
 	SHeaderRow::FColumn::FArguments Arguments;
-	Arguments.ColumnId(InColumnId);
-	Arguments.DefaultLabel(ColumnLabel);
-	Arguments.DefaultTooltip(ToolTip);
-	Arguments.ManualWidth(ColumnWidth);
+	Arguments.ColumnId(InColumnInfo.Id);
+	Arguments.DefaultLabel(InColumnInfo.Label);
+	Arguments.DefaultTooltip(InColumnInfo.Tooltip);
+	Arguments.ManualWidth(InColumnInfo.Width < 0.0f ? PCGEditorGraphAttributeListView::CalculateColumnWidth(InColumnInfo.Label) : InColumnInfo.Width);
 	Arguments.HAlignHeader(HAlign_Center);
-	Arguments.HAlignCell(HAlign_Right);
-	Arguments.SortMode(this, &SPCGEditorGraphAttributeListView::GetColumnSortMode, InColumnId);
+	Arguments.HAlignCell((EHorizontalAlignment)InColumnInfo.CellAlignment);
+	Arguments.SortMode(this, &SPCGEditorGraphAttributeListView::GetColumnSortMode, InColumnInfo.Id);
 	Arguments.OnSort(this, &SPCGEditorGraphAttributeListView::OnColumnSortModeChanged);
 	Arguments.OverflowPolicy(ETextOverflowPolicy::Ellipsis);
 
 	SHeaderRow::FColumn* NewColumn = new SHeaderRow::FColumn(Arguments);
-	NewColumn->bIsVisible = !HiddenAttributes.Contains(InColumnId);
-	ListViewHeader->AddColumn(*NewColumn);
-}
-
-void SPCGEditorGraphAttributeListView::AddPointDataColumns(const UPCGPointData* InPCGPointData)
-{
-	AddColumn(InPCGPointData, PCGEditorGraphAttributeListView::NAME_Index, PCGEditorGraphAttributeListView::TEXT_IndexLabel);
-	AddColumn(InPCGPointData, PCGEditorGraphAttributeListView::NAME_PointPositionX, PCGEditorGraphAttributeListView::TEXT_PointPositionLabelX);
-	AddColumn(InPCGPointData, PCGEditorGraphAttributeListView::NAME_PointPositionY, PCGEditorGraphAttributeListView::TEXT_PointPositionLabelY);
-	AddColumn(InPCGPointData, PCGEditorGraphAttributeListView::NAME_PointPositionZ, PCGEditorGraphAttributeListView::TEXT_PointPositionLabelZ);
-	AddColumn(InPCGPointData, PCGEditorGraphAttributeListView::NAME_PointRotationRoll, PCGEditorGraphAttributeListView::TEXT_PointRotationLabelRoll);
-	AddColumn(InPCGPointData, PCGEditorGraphAttributeListView::NAME_PointRotationPitch, PCGEditorGraphAttributeListView::TEXT_PointRotationLabelPitch);
-	AddColumn(InPCGPointData, PCGEditorGraphAttributeListView::NAME_PointRotationYaw, PCGEditorGraphAttributeListView::TEXT_PointRotationLabelYaw);
-	AddColumn(InPCGPointData, PCGEditorGraphAttributeListView::NAME_PointScaleX, PCGEditorGraphAttributeListView::TEXT_PointScaleLabelX);
-	AddColumn(InPCGPointData, PCGEditorGraphAttributeListView::NAME_PointScaleY, PCGEditorGraphAttributeListView::TEXT_PointScaleLabelY);
-	AddColumn(InPCGPointData, PCGEditorGraphAttributeListView::NAME_PointScaleZ, PCGEditorGraphAttributeListView::TEXT_PointScaleLabelZ);
-	AddColumn(InPCGPointData, PCGEditorGraphAttributeListView::NAME_PointBoundsMinX, PCGEditorGraphAttributeListView::TEXT_PointBoundsLabelMinX);
-	AddColumn(InPCGPointData, PCGEditorGraphAttributeListView::NAME_PointBoundsMinY, PCGEditorGraphAttributeListView::TEXT_PointBoundsLabelMinY);
-	AddColumn(InPCGPointData, PCGEditorGraphAttributeListView::NAME_PointBoundsMinZ, PCGEditorGraphAttributeListView::TEXT_PointBoundsLabelMinZ);
-	AddColumn(InPCGPointData, PCGEditorGraphAttributeListView::NAME_PointBoundsMaxX, PCGEditorGraphAttributeListView::TEXT_PointBoundsLabelMaxX);
-	AddColumn(InPCGPointData, PCGEditorGraphAttributeListView::NAME_PointBoundsMaxY, PCGEditorGraphAttributeListView::TEXT_PointBoundsLabelMaxY);
-	AddColumn(InPCGPointData, PCGEditorGraphAttributeListView::NAME_PointBoundsMaxZ, PCGEditorGraphAttributeListView::TEXT_PointBoundsLabelMaxZ);
-	AddColumn(InPCGPointData, PCGEditorGraphAttributeListView::NAME_PointColorR, PCGEditorGraphAttributeListView::TEXT_PointColorLabelR);
-	AddColumn(InPCGPointData, PCGEditorGraphAttributeListView::NAME_PointColorG, PCGEditorGraphAttributeListView::TEXT_PointColorLabelG);
-	AddColumn(InPCGPointData, PCGEditorGraphAttributeListView::NAME_PointColorB, PCGEditorGraphAttributeListView::TEXT_PointColorLabelB);
-	AddColumn(InPCGPointData, PCGEditorGraphAttributeListView::NAME_PointColorA, PCGEditorGraphAttributeListView::TEXT_PointColorLabelA);
-	AddColumn(InPCGPointData, PCGEditorGraphAttributeListView::NAME_PointDensity, PCGEditorGraphAttributeListView::TEXT_PointDensityLabel);
-	AddColumn(InPCGPointData, PCGEditorGraphAttributeListView::NAME_PointSteepness, PCGEditorGraphAttributeListView::TEXT_PointSteepnessLabel);
-	AddColumn(InPCGPointData, PCGEditorGraphAttributeListView::NAME_PointSeed, PCGEditorGraphAttributeListView::TEXT_PointSeedLabel);
-
-	if (CVarShowAdvancedAttributesFields.GetValueOnAnyThread())
-	{
-		AddColumn(InPCGPointData, PCGEditorGraphAttributeListView::NAME_PointMetadataEntry, PCGEditorGraphAttributeListView::TEXT_PointMetadataEntryLabel);
-		AddColumn(InPCGPointData, PCGEditorGraphAttributeListView::NAME_PointMetadataEntryParent, PCGEditorGraphAttributeListView::TEXT_PointMetadataEntryParentLabel);
-	}
-}
-
-void SPCGEditorGraphAttributeListView::AddMetadataColumn(const UPCGData* InPCGData, const FName& InColumnId, EPCGMetadataTypes InMetadataType, const TCHAR* PostFix)
-{
-	FString OriginalColumnIdString = InColumnId.ToString();
-	if (PostFix)
-	{
-		OriginalColumnIdString.Append(PostFix);
-	}
-
-	FText ColumnLabel;
-	FText ToolTip;
-	FString ColumnIdString(OriginalColumnIdString);
-	if (InColumnId == NAME_None)
-	{
-		ColumnIdString = FString::Printf(TEXT("%s%s"), *PCGEditorGraphAttributeListView::NoneAttributeId, PostFix);
-		ColumnLabel = FText::Format(LOCTEXT("NoneLabelFormat", "{0}{1}"), UEnum::GetDisplayValueAsText(InMetadataType), FText::FromString(PostFix));
-	}
-
-	const FName ColumnId(ColumnIdString);
-
-	if (InPCGData)
-	{
-		FPCGAttributePropertyInputSelector TargetSelector;
-		TargetSelector.Update(OriginalColumnIdString);
-
-		FPCGColumnData& ColumnData = PCGColumnData.Add(ColumnId);
-		ColumnData.DataAccessor = TSharedPtr<const IPCGAttributeAccessor>(PCGAttributeAccessorHelpers::CreateConstAccessor(InPCGData, TargetSelector).Release());
-		ColumnData.DataKeys = TSharedPtr<const IPCGAttributeAccessorKeys>(PCGAttributeAccessorHelpers::CreateConstKeys(InPCGData, TargetSelector).Release());
-
-		ToolTip = FText::FromString(PCG::Private::GetTypeName(ColumnData.DataAccessor->GetUnderlyingType()));
-	}
-
-	if (ColumnLabel.IsEmpty())
-	{
-		ColumnLabel = FText::FromName(ColumnId);
-	}
-
-	float ColumnWidth = 0.0f;
-
-	EHorizontalAlignment CellAlignment = EHorizontalAlignment::HAlign_Right;
-
-	if (InMetadataType == EPCGMetadataTypes::String)
-	{
-		ColumnWidth = PCGEditorGraphAttributeListView::MaxColumnWidth;
-		CellAlignment = EHorizontalAlignment::HAlign_Left;
-	}
-	else
-	{
-		ColumnWidth = PCGEditorGraphAttributeListView::CalculateColumnWidth(ColumnLabel);
-	}
-
-	SHeaderRow::FColumn::FArguments ColumnArguments;
-	ColumnArguments.ColumnId(ColumnId);
-	ColumnArguments.DefaultLabel(ColumnLabel);
-	ColumnArguments.DefaultTooltip(ToolTip);
-	ColumnArguments.HAlignHeader(EHorizontalAlignment::HAlign_Center);
-	ColumnArguments.HAlignCell(CellAlignment);
-	ColumnArguments.ManualWidth(ColumnWidth);
-	ColumnArguments.SortMode(this, &SPCGEditorGraphAttributeListView::GetColumnSortMode, ColumnId);
-	ColumnArguments.OnSort(this, &SPCGEditorGraphAttributeListView::OnColumnSortModeChanged);
-	ColumnArguments.OverflowPolicy(ETextOverflowPolicy::Ellipsis);
-
-	SHeaderRow::FColumn* NewColumn = new SHeaderRow::FColumn(ColumnArguments);
-	NewColumn->bIsVisible = !HiddenAttributes.Contains(ColumnId);
+	NewColumn->bIsVisible = !HiddenAttributes.Contains(InColumnInfo.Id);
 	ListViewHeader->AddColumn(*NewColumn);
 }
 
@@ -1393,7 +1097,7 @@ void SPCGEditorGraphAttributeListView::CopySelectionToClipboard() const
 
 	// Gather selected rows and sort them to match the displayed order instead of selection order
 	TArray<PCGListviewItemPtr> SelectedListViewItems = ListView->GetSelectedItems();
-	if (const FPCGColumnData* ColumnData = PCGColumnData.Find(SortMode == EColumnSortMode::None ? PCGEditorGraphAttributeListView::NAME_Index : SortingColumn))
+	if (const FPCGColumnData* ColumnData = PCGColumnData.Find(SortingColumn))
 	{
 		if (ColumnData->DataAccessor.IsValid() && ColumnData->DataKeys.IsValid())
 		{
diff --git a/Engine/Plugins/PCG/Source/PCGEditor/Private/SPCGEditorGraphAttributeListView.h b/Engine/Plugins/PCG/Source/PCGEditor/Private/SPCGEditorGraphAttributeListView.h
--- a/Engine/Plugins/PCG/Source/PCGEditor/Private/SPCGEditorGraphAttributeListView.h
+++ b/Engine/Plugins/PCG/Source/PCGEditor/Private/SPCGEditorGraphAttributeListView.h
@@ -27,6 +27,7 @@ class UPCGEditorGraphNodeBase;
 class UPCGMetadata;
 class UPCGParamData;
 class UPCGPointData;
+struct FPCGTableVisualizerColumnInfo;
 struct FPCGDataCollection;
 struct FPCGPoint;
 enum class EPCGMetadataTypes : uint8;
@@ -36,10 +37,15 @@ class SHeaderRow;
 class SSearchBox;
 struct FSlateBrush;
 
+namespace PCGEditorGraphAttributeListView
+{
+	constexpr float MaxColumnWidth = 200.0f;
+}
+
 struct FPCGListViewItem
 {
 	int32 Index = INDEX_NONE;
-	const FPCGPoint* PCGPoint = nullptr;
+	TFunction<void(int)> DoubleClickCallback = nullptr;
 };
 
 struct FPCGColumnData
@@ -186,7 +192,6 @@ private:
 	void OnSelectionChanged(TSharedPtr<FName> Item, ESelectInfo::Type SelectInfo);
 	FText OnGenerateSelectedDataText() const;
 	int32 GetSelectedDataIndex() const;
-	void GenerateColumnsFromMetadata(const UPCGData* InPCGData, const UPCGMetadata* PCGMetadata);
 
 	void ToggleAllAttributes();
 	void ToggleAttribute(FName InAttributeName);
@@ -204,7 +209,7 @@ private:
 
 	FReply OnListViewKeyDown(const FGeometry& InGeometry, const FKeyEvent& InKeyEvent) const;
 
-	void AddColumn(const UPCGData* InPCGData, const FName& InColumnId, const FText& ColumnLabel);
+	void AddColumn(const UPCGData* InPCGData, const FPCGTableVisualizerColumnInfo& InColumnInfo, TSharedPtr<const IPCGAttributeAccessorKeys> InAccessorKeys);
 	void AddPointDataColumns(const UPCGPointData* InPCGPointData);
 	void AddMetadataColumn(const UPCGData* InPCGData, const FName& InColumnId, EPCGMetadataTypes InMetadataType, const TCHAR* PostFix = nullptr);
 
@@ -264,5 +269,5 @@ private:
 	TSharedPtr<FPCGListViewUpdater> CurrentUpdateTask = nullptr;
 
 	/** Used to ensure data collapsed for inspection is kept alive. */
-	TStrongObjectPtr<const UPCGPointData> CollapsedPointData = nullptr;
+	TStrongObjectPtr<const UPCGData> DataStrongPtr = nullptr;
 };
