Fix 'Console object named 'pcg.graph.ShowAdvancedAttributes' already exists but is being registered again'

#rnx

diff --git a/Engine/Plugins/PCG/Source/PCGEditor/Private/PCGEditorCommon.cpp b/Engine/Plugins/PCG/Source/PCGEditor/Private/PCGEditorCommon.cpp
new file mode 100644
--- /dev/null
+++ b/Engine/Plugins/PCG/Source/PCGEditor/Private/PCGEditorCommon.cpp
@@ -0,0 +1,13 @@
+// Copyright Epic Games, Inc. All Rights Reserved.
+
+#pragma once
+
+#include "PCGEditorCommon.h"
+
+namespace FPCGEditorCommon
+{
+	TAutoConsoleVariable<bool> CVarShowAdvancedAttributesFields(
+		TEXT("pcg.graph.ShowAdvancedAttributes"),
+		false,
+		TEXT("Control whether advanced attributes/properties are shown in the PCG graph editor"));
+}
diff --git a/Engine/Plugins/PCG/Source/PCGEditor/Private/PCGEditorCommon.h b/Engine/Plugins/PCG/Source/PCGEditor/Private/PCGEditorCommon.h
--- a/Engine/Plugins/PCG/Source/PCGEditor/Private/PCGEditorCommon.h
+++ b/Engine/Plugins/PCG/Source/PCGEditor/Private/PCGEditorCommon.h
@@ -43,10 +43,7 @@ namespace FPCGEditorCommon
 	const FName TextureDataType = FName(TEXT("Texture Data"));
 	const FName RenderTargetDataType = FName(TEXT("Render Target Data"));
 	const FName VolumeDataType = FName(TEXT("Volume Data"));
-	const FName PrimitiveDataType = FName(TEXT("Primitive Data"));	
+	const FName PrimitiveDataType = FName(TEXT("Primitive Data"));
 
-	static TAutoConsoleVariable<bool> CVarShowAdvancedAttributesFields(
-		TEXT("pcg.graph.ShowAdvancedAttributes"),
-		false,
-		TEXT("Control whether advanced attributes/properties are shown in the PCG graph editor"));
+	extern TAutoConsoleVariable<bool> CVarShowAdvancedAttributesFields;
 };
