[PCG] Fix Isoline tangents in case of curve spline output.

* Leftover of my tests, just need to change CurveCustomTangent to Curve when adding spline points.

#rnx

diff --git a/Engine/Plugins/PCG/Source/PCG/Private/Elements/PCGElevationIsolines.cpp b/Engine/Plugins/PCG/Source/PCG/Private/Elements/PCGElevationIsolines.cpp
--- a/Engine/Plugins/PCG/Source/PCG/Private/Elements/PCGElevationIsolines.cpp
+++ b/Engine/Plugins/PCG/Source/PCG/Private/Elements/PCGElevationIsolines.cpp
@@ -277,7 +277,7 @@ bool FPCGElevationIsolinesElement::ExecuteInternal(FPCGContext* InContext) const
 							FVector::ZeroVector,
 							Transform.GetRotation().Rotator(),
 							Transform.GetScale3D(),
-							Settings->bLinearSpline ? ESplinePointType::Linear : ESplinePointType::CurveCustomTangent);
+							Settings->bLinearSpline ? ESplinePointType::Linear : ESplinePointType::Curve);
 					}
 
 					UPCGSplineData* OutSplineData = NewObject<UPCGSplineData>();
