Remove '#pragma once' from PCGEditorCommon.cpp

#rnx

diff --git a/Engine/Plugins/PCG/Source/PCGEditor/Private/PCGEditorCommon.cpp b/Engine/Plugins/PCG/Source/PCGEditor/Private/PCGEditorCommon.cpp
--- a/Engine/Plugins/PCG/Source/PCGEditor/Private/PCGEditorCommon.cpp
+++ b/Engine/Plugins/PCG/Source/PCGEditor/Private/PCGEditorCommon.cpp
@@ -1,7 +1,5 @@
 // Copyright Epic Games, Inc. All Rights Reserved.
 
-#pragma once
-
 #include "PCGEditorCommon.h"
 
 namespace FPCGEditorCommon
