[PCG] Moved cvars for cleanliness

#rnx

diff --git a/Engine/Plugins/PCG/Source/PCG/Private/Graph/PCGGraphExecutor.cpp b/Engine/Plugins/PCG/Source/PCG/Private/Graph/PCGGraphExecutor.cpp
--- a/Engine/Plugins/PCG/Source/PCG/Private/Graph/PCGGraphExecutor.cpp
+++ b/Engine/Plugins/PCG/Source/PCG/Private/Graph/PCGGraphExecutor.cpp
@@ -601,7 +601,7 @@ void FPCGGraphExecutor::Execute()
 	TRACE_CPUPROFILER_EVENT_SCOPE(FPCGGraphExecutor::Execute);
 
 #if WITH_EDITOR
-	if (PCGComponent::CVarPausePCGExecution.GetValueOnAnyThread())
+	if (PCGSystemSwitches::CVarPausePCGExecution.GetValueOnAnyThread())
 	{
 		return;
 	}
diff --git a/Engine/Plugins/PCG/Source/PCG/Private/PCGCommon.cpp b/Engine/Plugins/PCG/Source/PCG/Private/PCGCommon.cpp
--- a/Engine/Plugins/PCG/Source/PCG/Private/PCGCommon.cpp
+++ b/Engine/Plugins/PCG/Source/PCG/Private/PCGCommon.cpp
@@ -26,6 +26,21 @@ namespace PCGFeatureSwitches
 	};
 }
 
+namespace PCGSystemSwitches
+{
+#if WITH_EDITOR
+	TAutoConsoleVariable<bool> CVarPausePCGExecution(
+		TEXT("pcg.PauseExecution"),
+		false,
+		TEXT("Pauses all execution of PCG but does not cancel tasks."));
+
+	TAutoConsoleVariable<bool> CVarGlobalDisableRefresh(
+			TEXT("pcg.GlobalDisableRefresh"),
+			false,
+			TEXT("Disable refresh for all PCG Components."));
+#endif
+}
+
 namespace PCGHiGenGrid
 {
 	bool IsValidGridSize(uint32 InGridSize)
diff --git a/Engine/Plugins/PCG/Source/PCG/Private/PCGComponent.cpp b/Engine/Plugins/PCG/Source/PCG/Private/PCGComponent.cpp
--- a/Engine/Plugins/PCG/Source/PCG/Private/PCGComponent.cpp
+++ b/Engine/Plugins/PCG/Source/PCG/Private/PCGComponent.cpp
@@ -64,18 +64,6 @@
 
 namespace PCGComponent
 {
-	static TAutoConsoleVariable<bool> CVarGlobalDisableRefresh(
-		TEXT("pcg.GlobalDisableRefresh"),
-		false,
-		TEXT("Disable refresh for all PCG Components."));
-
-#if WITH_EDITOR
-	TAutoConsoleVariable<bool> CVarPausePCGExecution(
-		TEXT("pcg.PauseExecution"),
-		false,
-		TEXT("Pauses all execution of PCG but does not cancel tasks."));
-#endif
-
 	template <typename DelegateType>
 	static void BroadcastDynamicDelegate(const DelegateType& Delegate, UPCGComponent* PCGComponent)
 	{
@@ -2353,7 +2341,7 @@ void UPCGComponent::Refresh(EPCGChangeType ChangeType, bool bCancelExistingRefre
 	}
 
 	// If refresh is disabled, just exit
-	if (PCGComponent::CVarGlobalDisableRefresh.GetValueOnAnyThread() || IsRunningCommandlet())
+	if (PCGSystemSwitches::CVarGlobalDisableRefresh.GetValueOnAnyThread() || IsRunningCommandlet())
 	{
 		return;
 	}
diff --git a/Engine/Plugins/PCG/Source/PCG/Public/PCGCommon.h b/Engine/Plugins/PCG/Source/PCG/Public/PCGCommon.h
--- a/Engine/Plugins/PCG/Source/PCG/Public/PCGCommon.h
+++ b/Engine/Plugins/PCG/Source/PCG/Public/PCGCommon.h
@@ -186,6 +186,14 @@ namespace PCGFeatureSwitches
 	extern PCG_API TAutoConsoleVariable<float> CVarSamplerMemoryThreshold;
 }
 
+namespace PCGSystemSwitches
+{
+#if WITH_EDITOR
+	extern PCG_API TAutoConsoleVariable<bool> CVarPausePCGExecution;
+	extern TAutoConsoleVariable<bool> CVarGlobalDisableRefresh;
+#endif
+}
+
 /** Describes space referential for operations that create data */
 UENUM(BlueprintType)
 enum class EPCGCoordinateSpace : uint8
diff --git a/Engine/Plugins/PCG/Source/PCG/Public/PCGComponent.h b/Engine/Plugins/PCG/Source/PCG/Public/PCGComponent.h
--- a/Engine/Plugins/PCG/Source/PCG/Public/PCGComponent.h
+++ b/Engine/Plugins/PCG/Source/PCG/Public/PCGComponent.h
@@ -76,13 +76,6 @@ enum class EPCGComponentDirtyFlag : uint8
 };
 ENUM_CLASS_FLAGS(EPCGComponentDirtyFlag);
 
-namespace PCGComponent
-{
-#if WITH_EDITOR
-	extern PCG_API TAutoConsoleVariable<bool> CVarPausePCGExecution;
-#endif
-}
-
 UCLASS(BlueprintType, ClassGroup = (Procedural), meta = (BlueprintSpawnableComponent, PrioritizeCategories = "PCG"))
 class PCG_API UPCGComponent : public UActorComponent
 {
diff --git a/Engine/Plugins/PCG/Source/PCGEditor/Private/PCGEditorModule.cpp b/Engine/Plugins/PCG/Source/PCGEditor/Private/PCGEditorModule.cpp
--- a/Engine/Plugins/PCG/Source/PCGEditor/Private/PCGEditorModule.cpp
+++ b/Engine/Plugins/PCG/Source/PCGEditor/Private/PCGEditorModule.cpp
@@ -279,7 +279,7 @@ void FPCGEditorModule::RegisterMenuExtensions()
 					FUIAction(
 						FExecuteAction::CreateLambda([]()
 						{
-							const bool bWasPaused = PCGComponent::CVarPausePCGExecution.GetValueOnAnyThread();
+							const bool bWasPaused = PCGSystemSwitches::CVarPausePCGExecution.GetValueOnAnyThread();
 
 							if (bWasPaused)
 							{
@@ -307,17 +307,17 @@ void FPCGEditorModule::RegisterMenuExtensions()
 								}
 							}
 
-							PCGComponent::CVarPausePCGExecution->Set(!bWasPaused);
+							PCGSystemSwitches::CVarPausePCGExecution->Set(!bWasPaused);
 						}),
 						FIsActionButtonVisible::CreateLambda([](){ return true; }),
 						FIsActionChecked::CreateLambda([]()
 						{
-							return PCGComponent::CVarPausePCGExecution.GetValueOnAnyThread();
+							return PCGSystemSwitches::CVarPausePCGExecution.GetValueOnAnyThread();
 						})
 					),
 					TAttribute<FText>::CreateLambda([]()
 					{
-						if(PCGComponent::CVarPausePCGExecution.GetValueOnAnyThread())
+						if(PCGSystemSwitches::CVarPausePCGExecution.GetValueOnAnyThread())
 						{
 							return LOCTEXT("PCGPauseButton_Off", "Paused");
 						}
