[PCG] Allow for user selected render material index in Ray Hit Metadata Extraction

#rnx

diff --git a/Engine/Plugins/PCG/Source/PCG/Private/Helpers/PCGWorldQueryHelpers.cpp b/Engine/Plugins/PCG/Source/PCG/Private/Helpers/PCGWorldQueryHelpers.cpp
--- a/Engine/Plugins/PCG/Source/PCG/Private/Helpers/PCGWorldQueryHelpers.cpp
+++ b/Engine/Plugins/PCG/Source/PCG/Private/Helpers/PCGWorldQueryHelpers.cpp
@@ -304,9 +304,10 @@ namespace PCGWorldQueryHelpers
 		}
 		else if (UPrimitiveComponent* HitComponent = Hit.GetComponent())
 		{
-			if (const UMaterialInterface* RenderMaterial = HitComponent->GetMaterial(Hit.ElementIndex))
+			if (QueryParams.bGetReferenceToRenderMaterial)
 			{
-				bResult &= ApplyAttribute(PCGWorldQueryConstants::RenderMaterialReferenceAttribute, FSoftObjectPath(RenderMaterial), QueryParams.bGetReferenceToRenderMaterial);
+				const UMaterialInterface* RenderMaterial = HitComponent->GetMaterial(QueryParams.RenderMaterialIndex);
+				bResult &= RenderMaterial && ApplyAttribute(PCGWorldQueryConstants::RenderMaterialReferenceAttribute, FSoftObjectPath(RenderMaterial));
 			}
 
 			if (QueryParams.bGetLocalImpactPoint)
diff --git a/Engine/Plugins/PCG/Source/PCG/Public/Data/PCGWorldData.h b/Engine/Plugins/PCG/Source/PCG/Public/Data/PCGWorldData.h
--- a/Engine/Plugins/PCG/Source/PCG/Public/Data/PCGWorldData.h
+++ b/Engine/Plugins/PCG/Source/PCG/Public/Data/PCGWorldData.h
@@ -146,9 +146,13 @@ struct FPCGWorldRaycastQueryParams : public FPCGWorldCommonQueryParams
 	UPROPERTY(BlueprintReadWrite, EditAnywhere, Category = "Data|Attributes", meta = (PCG_Overridable, EditCondition = "SelectLandscapeHits != EPCGWorldQuerySelectLandscapeHits::Exclude"))
 	uint8 bApplyMetadataFromLandscape : 1;
 
+	/** The index of the render material when it is queried from the hit. */
+	UPROPERTY(BlueprintReadWrite, EditAnywhere, Category = "Data|Attributes", meta = (PCG_Overridable, EditCondition = "bGetReferenceToRenderMaterial", EditConditionHides, DisplayAfter = "bGetReferenceToRenderMaterial"))
+	int32 RenderMaterialIndex = 0;
+
 	/** This UV Channel will be selected when retrieving UV Coordinates from a raycast query. */
 	UPROPERTY(BlueprintReadWrite, EditAnywhere, Category = "Data|Attributes", meta = (PCG_Overridable, EditCondition = "bTraceComplex && bGetUVCoords", EditConditionHides, DisplayAfter = "bGetUVCoords"))
-	int32 UVChannel = INDEX_NONE;
+	int32 UVChannel = 0;
 };
 
 template<>
