[PCG]
Added prefix & suffix tagging on the Add Tags node.
Added substring & matches mode for the Filter Data By Tag and the Delete Tags nodes.

diff --git a/Engine/Plugins/PCG/Source/PCG/Private/Elements/PCGAddTag.cpp b/Engine/Plugins/PCG/Source/PCG/Private/Elements/PCGAddTag.cpp
--- a/Engine/Plugins/PCG/Source/PCG/Private/Elements/PCGAddTag.cpp
+++ b/Engine/Plugins/PCG/Source/PCG/Private/Elements/PCGAddTag.cpp
@@ -55,9 +55,11 @@ bool FPCGAddTagElement::ExecuteInternal(FPCGContext* Context) const
 
 	for (const FString& Tag : TagsArray)
 	{
+		const FString TagToAdd = Settings->Prefix + Tag + Settings->Suffix;
+
 		for (FPCGTaggedData& OutputTaggedData : Context->OutputData.TaggedData)
 		{
-				OutputTaggedData.Tags.Add(Tag);
+			OutputTaggedData.Tags.Add(TagToAdd);
 		}
 	}
 	
diff --git a/Engine/Plugins/PCG/Source/PCG/Private/Elements/PCGDeleteTags.cpp b/Engine/Plugins/PCG/Source/PCG/Private/Elements/PCGDeleteTags.cpp
--- a/Engine/Plugins/PCG/Source/PCG/Private/Elements/PCGDeleteTags.cpp
+++ b/Engine/Plugins/PCG/Source/PCG/Private/Elements/PCGDeleteTags.cpp
@@ -7,6 +7,8 @@
 #include "PCGData.h"
 #include "Helpers/PCGHelpers.h"
 
+#include "Algo/AnyOf.h"
+
 #define LOCTEXT_NAMESPACE "PCGDeleteTagsElement"
 
 #if WITH_EDITOR
@@ -103,7 +105,9 @@ bool FPCGDeleteTagsElement::ExecuteInternal(FPCGContext* Context) const
 
 		for (const FString& Tag : Input.Tags)
 		{
-			if (TagsToFilter.Contains(Tag) == bKeepInFilter)
+			if ((Settings->Operator == EPCGStringMatchingOperator::Equal && TagsToFilter.Contains(Tag) == bKeepInFilter) ||
+				(Settings->Operator == EPCGStringMatchingOperator::Substring && Algo::AnyOf(TagsToFilter, [&Tag](const FString& TagToFilter) { return Tag.Contains(TagToFilter); }) == bKeepInFilter) ||
+				(Settings->Operator == EPCGStringMatchingOperator::Matches && Algo::AnyOf(TagsToFilter, [&Tag](const FString& TagToFilter) { return Tag.MatchesWildcard(TagToFilter); }) == bKeepInFilter))
 			{
 				Output.Tags.Add(Tag);
 			}
diff --git a/Engine/Plugins/PCG/Source/PCG/Private/Elements/PCGFilterByTag.cpp b/Engine/Plugins/PCG/Source/PCG/Private/Elements/PCGFilterByTag.cpp
--- a/Engine/Plugins/PCG/Source/PCG/Private/Elements/PCGFilterByTag.cpp
+++ b/Engine/Plugins/PCG/Source/PCG/Private/Elements/PCGFilterByTag.cpp
@@ -5,6 +5,8 @@
 #include "PCGCustomVersion.h"
 #include "Helpers/PCGHelpers.h"
 
+#include "Algo/AnyOf.h"
+
 #include UE_INLINE_GENERATED_CPP_BY_NAME(PCGFilterByTag)
 
 #define LOCTEXT_NAMESPACE "PCGFilterByTag"
@@ -94,7 +96,9 @@ bool FPCGFilterByTagElement::ExecuteInternal(FPCGContext* Context) const
 
 		for (const FString& Tag : Tags)
 		{
-			if (Input.Tags.Contains(Tag))
+			if((Settings->Operator == EPCGStringMatchingOperator::Equal && Input.Tags.Contains(Tag)) ||
+				(Settings->Operator == EPCGStringMatchingOperator::Substring && Algo::AnyOf(Input.Tags, [&Tag](const FString& InputTag) { return InputTag.Contains(Tag); })) ||
+				(Settings->Operator == EPCGStringMatchingOperator::Matches && Algo::AnyOf(Input.Tags, [&Tag](const FString& InputTag) { return InputTag.MatchesWildcard(Tag); })))
 			{
 				bHasCommonTags = true;
 				break;
diff --git a/Engine/Plugins/PCG/Source/PCG/Public/Elements/PCGAddTag.h b/Engine/Plugins/PCG/Source/PCG/Public/Elements/PCGAddTag.h
--- a/Engine/Plugins/PCG/Source/PCG/Public/Elements/PCGAddTag.h
+++ b/Engine/Plugins/PCG/Source/PCG/Public/Elements/PCGAddTag.h
@@ -32,10 +32,18 @@ protected:
 	//~End UPCGSettings interface
 
 public:
-	/** Comma-separated list of tags to apply to the node */
+	/** Comma-separated list of tags to apply to the node. */
 	UPROPERTY(EditAnywhere, BlueprintReadWrite, Category = Settings, meta = (PCG_Overridable))
 	FString TagsToAdd;
 
+	/** Common prefix to add to all tags, can be left empty. */
+	UPROPERTY(EditAnywhere, BlueprintReadWrite, Category = Settings, meta = (PCG_Overridable))
+	FString Prefix;
+
+	/** Common suffix to add to all tags, can be left empty. */
+	UPROPERTY(EditAnywhere, BlueprintReadWrite, Category = Settings, meta = (PCG_Overridable))
+	FString Suffix;
+
 	/** Enables deprecated behavior using spaces as separators. Disable to update the node to current behavior. */
 	UE_DEPRECATED(5.5, "bTokenizeOnWhiteSpace has been deprecated.")
 	UPROPERTY(EditAnywhere, Category = Settings, meta = (EditCondition = "bTokenizeOnWhiteSpace", EditConditionHides, DeprecationMessage = "bTokenizeOnWhiteSpace has been deprecated."))
diff --git a/Engine/Plugins/PCG/Source/PCG/Public/Elements/PCGDeleteTags.h b/Engine/Plugins/PCG/Source/PCG/Public/Elements/PCGDeleteTags.h
--- a/Engine/Plugins/PCG/Source/PCG/Public/Elements/PCGDeleteTags.h
+++ b/Engine/Plugins/PCG/Source/PCG/Public/Elements/PCGDeleteTags.h
@@ -44,6 +44,9 @@ public:
 	UPROPERTY(BlueprintReadWrite, EditAnywhere, Category = Settings, meta = (PCG_Overridable))
 	EPCGTagFilterOperation Operation = EPCGTagFilterOperation::DeleteSelectedTags;
 
+	UPROPERTY(BlueprintReadWrite, EditAnywhere, Category = Settings, meta = (PCG_Overridable))
+	EPCGStringMatchingOperator Operator = EPCGStringMatchingOperator::Equal;
+
 	/** Comma-separated list of tags to add or remove from the input data. */
 	UPROPERTY(BlueprintReadWrite, EditAnywhere, Category = Settings, meta = (PCG_Overridable))
 	FString SelectedTags;
diff --git a/Engine/Plugins/PCG/Source/PCG/Public/Elements/PCGFilterByTag.h b/Engine/Plugins/PCG/Source/PCG/Public/Elements/PCGFilterByTag.h
--- a/Engine/Plugins/PCG/Source/PCG/Public/Elements/PCGFilterByTag.h
+++ b/Engine/Plugins/PCG/Source/PCG/Public/Elements/PCGFilterByTag.h
@@ -42,6 +42,9 @@ public:
 	UPROPERTY(BlueprintReadWrite, EditAnywhere, Category = Settings, meta = (PCG_Overridable))
 	EPCGFilterByTagOperation Operation = EPCGFilterByTagOperation::KeepTagged;
 
+	UPROPERTY(BlueprintReadWrite, EditAnywhere, Category = Settings, meta = (PCG_Overridable))
+	EPCGStringMatchingOperator Operator = EPCGStringMatchingOperator::Equal;
+
 	/** Comma-separated list of tags */
 	UPROPERTY(BlueprintReadWrite, EditAnywhere, Category = Settings, meta = (PCG_Overridable))
 	FString SelectedTags;
diff --git a/Engine/Plugins/PCG/Source/PCG/Public/PCGCommon.h b/Engine/Plugins/PCG/Source/PCG/Public/PCGCommon.h
--- a/Engine/Plugins/PCG/Source/PCG/Public/PCGCommon.h
+++ b/Engine/Plugins/PCG/Source/PCG/Public/PCGCommon.h
@@ -203,6 +203,14 @@ enum class EPCGCoordinateSpace : uint8
 	LocalComponent
 };
 
+UENUM(BlueprintType)
+enum class EPCGStringMatchingOperator : uint8
+{
+	Equal,
+	Substring,
+	Matches
+};
+
 /** Describes one or more target execution grids. */
 UENUM(meta = (Bitflags, UseEnumValuesAsMaskValuesInEditor = "true"))
 enum class EPCGHiGenGrid : uint32
