[PCG] Initialize properly the FPCGWorldRaycastQueryParams properties

#rnx

diff --git a/Engine/Plugins/PCG/Source/PCG/Public/Data/PCGWorldData.h b/Engine/Plugins/PCG/Source/PCG/Public/Data/PCGWorldData.h
--- a/Engine/Plugins/PCG/Source/PCG/Public/Data/PCGWorldData.h
+++ b/Engine/Plugins/PCG/Source/PCG/Public/Data/PCGWorldData.h
@@ -100,51 +100,51 @@ struct FPCGWorldRaycastQueryParams : public FPCGWorldCommonQueryParams
 
 	/** Ignore rays that hit backfaces. */
 	UPROPERTY(BlueprintReadWrite, EditAnywhere, Category = "Data|Filtering", meta = (PCG_Overridable))
-	uint8 bIgnoreBackfaceHits : 1;
+	uint8 bIgnoreBackfaceHits : 1 = false;
 
 	/** Create an attribute for whether the raycast resulted in a hit. */
 	UPROPERTY(BlueprintReadWrite, EditAnywhere, Category = "Data|Attributes", meta = (PCG_Overridable))
-	uint8 bGetImpact : 1;
+	uint8 bGetImpact : 1 = false;
 
 	/** Create an attribute for the impact location in world space. */
 	UPROPERTY(BlueprintReadWrite, EditAnywhere, Category = "Data|Attributes", meta = (PCG_Overridable))
-	uint8 bGetImpactPoint : 1;
+	uint8 bGetImpactPoint : 1 = false;
 
 	/** Create an attribute for the impact normal. */
 	UPROPERTY(BlueprintReadWrite, EditAnywhere, Category = "Data|Attributes", meta = (PCG_Overridable))
-	uint8 bGetImpactNormal : 1;
+	uint8 bGetImpactNormal : 1 = false;
 
 	/** Create an attribute for the distance between the ray origin and the impact point. */
 	UPROPERTY(BlueprintReadWrite, EditAnywhere, Category = "Data|Attributes", meta = (PCG_Overridable))
-	uint8 bGetDistance : 1;
+	uint8 bGetDistance : 1 = false;
 
 	/** Create an attribute for the impact point in the hit object's local space. */
 	UPROPERTY(BlueprintReadWrite, EditAnywhere, Category = "Data|Attributes", meta = (PCG_Overridable))
-	uint8 bGetLocalImpactPoint : 1;
+	uint8 bGetLocalImpactPoint : 1 = false;
 
 	/** Create an attribute for the render material. */
 	UPROPERTY(BlueprintReadWrite, EditAnywhere, Category = "Data|Attributes", meta = (PCG_Overridable))
-	uint8 bGetReferenceToRenderMaterial : 1;
+	uint8 bGetReferenceToRenderMaterial : 1 = false;
 
 	/** Create an attribute for the static mesh. */
 	UPROPERTY(BlueprintReadWrite, EditAnywhere, Category = "Data|Attributes", meta = (PCG_Overridable))
-	uint8 bGetReferenceToStaticMesh : 1;
+	uint8 bGetReferenceToStaticMesh : 1 = false;
 
 	/** Create an attribute for index of the hit face. Note: Will only work in complex traces. */
 	UPROPERTY(BlueprintReadWrite, EditAnywhere, Category = "Data|Attributes", meta = (PCG_Overridable, EditCondition = "bTraceComplex"))
-	uint8 bGetFaceIndex : 1;
+	uint8 bGetFaceIndex : 1 = false;
 
 	/** Create an attribute for UV Coordinates of the surface hit. Note: Will only work in complex traces. */
 	UPROPERTY(BlueprintReadWrite, EditAnywhere, Category = "Data|Attributes", meta = (PCG_Overridable, EditCondition = "bTraceComplex"))
-	uint8 bGetUVCoords : 1;
+	uint8 bGetUVCoords : 1 = false;
 
 	/** Create an attribute for the index of the element hit. Unique to the hit primitive. */
 	UPROPERTY(BlueprintReadWrite, EditAnywhere, Category = "Data|Attributes", meta = (PCG_Overridable))
-	uint8 bGetElementIndex : 1;
+	uint8 bGetElementIndex : 1 = false;
 
 	/** Will apply landscape layers and their values at the impact point. */
 	UPROPERTY(BlueprintReadWrite, EditAnywhere, Category = "Data|Attributes", meta = (PCG_Overridable, EditCondition = "SelectLandscapeHits != EPCGWorldQuerySelectLandscapeHits::Exclude"))
-	uint8 bApplyMetadataFromLandscape : 1;
+	uint8 bApplyMetadataFromLandscape : 1 = false;
 
 	/** The index of the render material when it is queried from the hit. */
 	UPROPERTY(BlueprintReadWrite, EditAnywhere, Category = "Data|Attributes", meta = (PCG_Overridable, EditCondition = "bGetReferenceToRenderMaterial", EditConditionHides, DisplayAfter = "bGetReferenceToRenderMaterial"))
