PCG - Disable triggering RenderDoc captures from texture readback code when UE is launched with RD attached.

diff --git a/Engine/Plugins/PCG/Source/PCGCompute/Private/PCGTextureReadback.cpp b/Engine/Plugins/PCG/Source/PCGCompute/Private/PCGTextureReadback.cpp
--- a/Engine/Plugins/PCG/Source/PCGCompute/Private/PCGTextureReadback.cpp
+++ b/Engine/Plugins/PCG/Source/PCGCompute/Private/PCGTextureReadback.cpp
@@ -4,7 +4,6 @@
 
 #include "GlobalShader.h"
 #include "RHIGPUReadback.h"
-#include "RenderCaptureInterface.h"
 #include "RenderGraphUtils.h"
 #include "Async/Async.h"
 
@@ -42,8 +41,6 @@ void FPCGTextureReadbackInterface::Dispatch_RenderThread(FRHICommandListImmediat
 {
 	check(Params.SourceTexture && Params.SourceSampler);
 
-	RenderCaptureInterface::FScopedCapture RenderCapture(true, &RHICmdList, TEXT("PCGTextureDataReadback"));
-
 	FPCGTextureReadbackCS::FParameters PassParameters;
 	PassParameters.SourceTextureArray = Params.SourceTexture;
 	PassParameters.SourceSampler = Params.SourceSampler;
