[PCG] Fix crash in Update Dynamic tracking

* Make a pass where we query GetOriginalComponent and make sure we test it to be not null.
* Also move the PostGenerate function to be called at the same time as the other generated callback, because something can happen in this post process call from the user that invalidate our component afterwards, so make sure we do it at the same time as other calls, and protect ourselves if the component ever goes invalid.

#rnx

diff --git a/Engine/Plugins/PCG/Source/PCG/Private/Elements/PCGDataFromActor.cpp b/Engine/Plugins/PCG/Source/PCG/Private/Elements/PCGDataFromActor.cpp
--- a/Engine/Plugins/PCG/Source/PCG/Private/Elements/PCGDataFromActor.cpp
+++ b/Engine/Plugins/PCG/Source/PCG/Private/Elements/PCGDataFromActor.cpp
@@ -424,11 +424,14 @@ bool FPCGDataFromActorElement::ExecuteInternal(FPCGContext* InContext) const
 #if WITH_EDITOR
 		// Remove ignored change origins now that we've completed the wait tasks.
 		UPCGComponent* OriginalComponent = Context->SourceComponent->GetOriginalComponent();
-		for (TObjectKey<UObject>& IgnoredChangeOriginKey : Context->IgnoredChangeOrigins)
+		if (ensure(OriginalComponent))
 		{
-			if (UObject* IgnoredChangeOrigin = IgnoredChangeOriginKey.ResolveObjectPtr())
+			for (TObjectKey<UObject>& IgnoredChangeOriginKey : Context->IgnoredChangeOrigins)
 			{
-				OriginalComponent->StopIgnoringChangeOriginDuringGeneration(IgnoredChangeOrigin);
+				if (UObject* IgnoredChangeOrigin = IgnoredChangeOriginKey.ResolveObjectPtr())
+				{
+					OriginalComponent->StopIgnoringChangeOriginDuringGeneration(IgnoredChangeOrigin);
+				}
 			}
 		}
 #endif
@@ -453,7 +456,7 @@ void FPCGDataFromActorElement::GatherWaitTasks(AActor* FoundActor, FPCGContext*
 	check(Settings);
 
 	UPCGComponent* SourceComponent = Context->SourceComponent.IsValid() ? Context->SourceComponent.Get() : nullptr;
-	const UPCGComponent* SourceOriginalComponent = SourceComponent ? SourceComponent->GetOriginalComponent() : nullptr;
+	UPCGComponent* SourceOriginalComponent = SourceComponent ? SourceComponent->GetOriginalComponent() : nullptr;
 
 	if (!SourceOriginalComponent)
 	{
@@ -492,7 +495,7 @@ void FPCGDataFromActorElement::GatherWaitTasks(AActor* FoundActor, FPCGContext*
 			// Signal that any change notifications from generating upstream component should not trigger re-executions of this component.
 			// Such change notifications can cancel the current execution.
 			// Note: Uses owner because FPCGActorAndComponentMapping::OnPCGGraphGeneratedOrCleaned reports change on owner.
-			SourceComponent->GetOriginalComponent()->StartIgnoringChangeOriginDuringGeneration(Component->GetOwner());
+			SourceOriginalComponent->StartIgnoringChangeOriginDuringGeneration(Component->GetOwner());
 			Context->IgnoredChangeOrigins.Add(Component->GetOwner());
 #endif
 
diff --git a/Engine/Plugins/PCG/Source/PCG/Private/PCGActorAndComponentMapping.cpp b/Engine/Plugins/PCG/Source/PCG/Private/PCGActorAndComponentMapping.cpp
--- a/Engine/Plugins/PCG/Source/PCG/Private/PCGActorAndComponentMapping.cpp
+++ b/Engine/Plugins/PCG/Source/PCG/Private/PCGActorAndComponentMapping.cpp
@@ -1816,7 +1816,9 @@ void FPCGActorAndComponentMapping::OnObjectChanged(UObject* InObject, const FAct
 				continue;
 			}
 
-			if (PCGComponent->GetOriginalComponent()->IsIgnoringChangeOrigin(InObject))
+			UPCGComponent* OriginalComponent = PCGComponent->GetOriginalComponent();
+
+			if (ensure(OriginalComponent) && OriginalComponent->IsIgnoringChangeOrigin(InObject))
 			{
 				PCGGraphExecutionLogging::LogChangeOriginIgnoredForComponent(InObject, PCGComponent);
 				continue;
@@ -1864,7 +1866,9 @@ void FPCGActorAndComponentMapping::OnObjectChanged(UObject* InObject, const FAct
 				return;
 			}
 
-			if (ComponentRef.Component->GetOriginalComponent()->IsIgnoringChangeOrigin(InObject))
+			UPCGComponent* OriginalComponent = ComponentRef.Component->GetOriginalComponent();
+
+			if (ensure(OriginalComponent) && OriginalComponent->IsIgnoringChangeOrigin(InObject))
 			{
 				PCGGraphExecutionLogging::LogChangeOriginIgnoredForComponent(InObject, ComponentRef.Component);
 				return;
@@ -1897,7 +1901,9 @@ void FPCGActorAndComponentMapping::OnObjectChanged(UObject* InObject, const FAct
 				return;
 			}
 
-			if (ComponentRef.Component->GetOriginalComponent()->IsIgnoringChangeOrigin(InObject))
+			UPCGComponent* OriginalComponent = ComponentRef.Component->GetOriginalComponent();
+
+			if (ensure(OriginalComponent) && OriginalComponent->IsIgnoringChangeOrigin(InObject))
 			{
 				PCGGraphExecutionLogging::LogChangeOriginIgnoredForComponent(InObject, ComponentRef.Component);
 				return;
diff --git a/Engine/Plugins/PCG/Source/PCG/Private/PCGComponent.cpp b/Engine/Plugins/PCG/Source/PCG/Private/PCGComponent.cpp
--- a/Engine/Plugins/PCG/Source/PCG/Private/PCGComponent.cpp
+++ b/Engine/Plugins/PCG/Source/PCG/Private/PCGComponent.cpp
@@ -509,8 +509,6 @@ void UPCGComponent::PostProcessGraph(const FBox& InNewBounds, bool bInGenerated,
 			// We don't have the info on the original component here, so forward for all
 			// components.
 			Context->OutputData = Context->InputData;
-
-			CallPostGenerateFunctions(Context);
 		}
 
 #if WITH_EDITOR
@@ -523,7 +521,16 @@ void UPCGComponent::PostProcessGraph(const FBox& InNewBounds, bool bInGenerated,
 		UpdateDynamicTracking();
 #endif // WITH_EDITOR
 
-		PCGComponent::BroadcastDynamicDelegate(OnPCGGraphGeneratedExternal, this);
+		if (IsValid(this) && Context)
+		{
+			CallPostGenerateFunctions(Context);
+		}
+
+		// If Generate function made the component invalid for whatever reason, we re-check if this is valid.
+		if (IsValid(this))
+		{
+			PCGComponent::BroadcastDynamicDelegate(OnPCGGraphGeneratedExternal, this);
+		}
 	}
 
 	// Trigger notification - will be used by other tracking mechanisms
@@ -3191,17 +3198,20 @@ void UPCGComponent::UpdateDynamicTracking()
 		TRACE_CPUPROFILER_EVENT_SCOPE(UPCGComponent::UpdateDynamicTracking::LocalComponent);
 
 		UPCGComponent* OriginalComponent = GetOriginalComponent();
-		FScopeLock Lock(&OriginalComponent->CurrentExecutionDynamicTrackingLock);
-		for (auto& It : CurrentExecutionDynamicTracking)
+		if (ensure(OriginalComponent))
 		{
-			TArray<FPCGSettingsAndCulling>& OriginalSettingsAndCulling = OriginalComponent->CurrentExecutionDynamicTracking.FindOrAdd(It.Key);
-			for (FPCGSettingsAndCulling& SettingsAndCulling : It.Value)
+			FScopeLock Lock(&OriginalComponent->CurrentExecutionDynamicTrackingLock);
+			for (auto& It : CurrentExecutionDynamicTracking)
 			{
-				OriginalSettingsAndCulling.AddUnique(std::move(SettingsAndCulling));
+				TArray<FPCGSettingsAndCulling>& OriginalSettingsAndCulling = OriginalComponent->CurrentExecutionDynamicTracking.FindOrAdd(It.Key);
+				for (FPCGSettingsAndCulling& SettingsAndCulling : It.Value)
+				{
+					OriginalSettingsAndCulling.AddUnique(std::move(SettingsAndCulling));
+				}
 			}
-		}
 
-		OriginalComponent->CurrentExecutionDynamicTrackingSettings.Append(CurrentExecutionDynamicTrackingSettings);
+			OriginalComponent->CurrentExecutionDynamicTrackingSettings.Append(CurrentExecutionDynamicTrackingSettings);
+		}
 
 		CurrentExecutionDynamicTracking.Empty();
 		CurrentExecutionDynamicTrackingSettings.Empty();
