[PCG] Enable multithreading/time slicing in Metadata operations

* Refactor Async processing to be more-cache friendly
* Change the meaning of DoOperation bool value (returns now if it the operation is done or not)

Submitted on behalf of adrien.logut
#rnx

diff --git a/Engine/Plugins/PCG/Source/PCG/Private/Elements/Metadata/PCGMetadataBitwiseOpElement.cpp b/Engine/Plugins/PCG/Source/PCG/Private/Elements/Metadata/PCGMetadataBitwiseOpElement.cpp
--- a/Engine/Plugins/PCG/Source/PCG/Private/Elements/Metadata/PCGMetadataBitwiseOpElement.cpp
+++ b/Engine/Plugins/PCG/Source/PCG/Private/Elements/Metadata/PCGMetadataBitwiseOpElement.cpp
@@ -144,12 +144,10 @@ bool FPCGMetadataBitwiseElement::DoOperation(PCGMetadataOps::FOperationData& Ope
 
 	if (Settings->Operation == EPCGMetadataBitwiseOperation::Not)
 	{
-		DoUnaryOp<int64>(OperationData, [](const int64& Value) -> int64 { return PCGMetadataBitwiseSettings::UnaryOp(Value); });
+		return DoUnaryOp<int64>(OperationData, [](const int64& Value) -> int64 { return PCGMetadataBitwiseSettings::UnaryOp(Value); });
 	}
 	else
 	{
-		DoBinaryOp<int64, int64>(OperationData, [Operation = Settings->Operation](const int64& Value1, const int64& Value2) -> int64 { return PCGMetadataBitwiseSettings::BinaryOp(Value1, Value2, Operation); });
+		return DoBinaryOp<int64, int64>(OperationData, [Operation = Settings->Operation](const int64& Value1, const int64& Value2) -> int64 { return PCGMetadataBitwiseSettings::BinaryOp(Value1, Value2, Operation); });
 	}
-
-	return true;
 }
diff --git a/Engine/Plugins/PCG/Source/PCG/Private/Elements/Metadata/PCGMetadataBooleanOpElement.cpp b/Engine/Plugins/PCG/Source/PCG/Private/Elements/Metadata/PCGMetadataBooleanOpElement.cpp
--- a/Engine/Plugins/PCG/Source/PCG/Private/Elements/Metadata/PCGMetadataBooleanOpElement.cpp
+++ b/Engine/Plugins/PCG/Source/PCG/Private/Elements/Metadata/PCGMetadataBooleanOpElement.cpp
@@ -144,12 +144,10 @@ bool FPCGMetadataBooleanElement::DoOperation(PCGMetadataOps::FOperationData& Ope
 
 	if (Settings->Operation == EPCGMetadataBooleanOperation::Not)
 	{
-		DoUnaryOp<bool>(OperationData, [](const bool& Value) -> bool { return PCGMetadataBooleanSettings::UnaryOp(Value); });
+		return DoUnaryOp<bool>(OperationData, [](const bool& Value) -> bool { return PCGMetadataBooleanSettings::UnaryOp(Value); });
 	}
 	else
 	{
-		DoBinaryOp<bool, bool>(OperationData, [Operation = Settings->Operation](const bool& Value1, const bool& Value2) -> bool { return PCGMetadataBooleanSettings::BinaryOp(Value1, Value2, Operation); });
+		return DoBinaryOp<bool, bool>(OperationData, [Operation = Settings->Operation](const bool& Value1, const bool& Value2) -> bool { return PCGMetadataBooleanSettings::BinaryOp(Value1, Value2, Operation); });
 	}
-
-	return true;
 }
diff --git a/Engine/Plugins/PCG/Source/PCG/Private/Elements/Metadata/PCGMetadataBreakVector.cpp b/Engine/Plugins/PCG/Source/PCG/Private/Elements/Metadata/PCGMetadataBreakVector.cpp
--- a/Engine/Plugins/PCG/Source/PCG/Private/Elements/Metadata/PCGMetadataBreakVector.cpp
+++ b/Engine/Plugins/PCG/Source/PCG/Private/Elements/Metadata/PCGMetadataBreakVector.cpp
@@ -167,13 +167,12 @@ bool FPCGMetadataBreakVectorElement::DoOperation(PCGMetadataOps::FOperationData&
 	const UPCGMetadataBreakVectorSettings* Settings = static_cast<const UPCGMetadataBreakVectorSettings*>(OperationData.Settings);
 	check(Settings);
 
-	auto BreakFunc = [this, &OperationData](auto DummyValue) -> bool
+	auto BreakFunc = [this, &OperationData]<typename AttributeType>(AttributeType) -> bool
 	{
-		using AttributeType = decltype(DummyValue);
-
 		if constexpr (!PCGMetadataBreakVectorSettings::IsValidType<AttributeType>())
 		{
-			return false;
+			ensure(false);
+			return true;
 		}
 		else
 		{
diff --git a/Engine/Plugins/PCG/Source/PCG/Private/Elements/Metadata/PCGMetadataCompareOpElement.cpp b/Engine/Plugins/PCG/Source/PCG/Private/Elements/Metadata/PCGMetadataCompareOpElement.cpp
--- a/Engine/Plugins/PCG/Source/PCG/Private/Elements/Metadata/PCGMetadataCompareOpElement.cpp
+++ b/Engine/Plugins/PCG/Source/PCG/Private/Elements/Metadata/PCGMetadataCompareOpElement.cpp
@@ -168,10 +168,8 @@ bool FPCGMetadataCompareElement::DoOperation(PCGMetadataOps::FOperationData& Ope
 
 	const UPCGMetadataCompareSettings* Settings = CastChecked<UPCGMetadataCompareSettings>(OperationData.Settings);
 
-	auto CompareFunc = [this, Operation = Settings->Operation, Tolerance = Settings->Tolerance, &OperationData](auto DummyValue) -> bool
+	auto CompareFunc = [this, Operation = Settings->Operation, Tolerance = Settings->Tolerance, &OperationData]<typename AttributeType>(AttributeType) -> bool
 	{
-		using AttributeType = decltype(DummyValue);
-
 		return DoBinaryOp<AttributeType, AttributeType>(OperationData, 
 			[Operation, Tolerance](const AttributeType& Value1, const AttributeType& Value2) -> bool { 
 				return PCGMetadataCompareSettings::ApplyCompare(Value1, Value2, Operation, Tolerance); 
diff --git a/Engine/Plugins/PCG/Source/PCG/Private/Elements/Metadata/PCGMetadataMakeRotator.cpp b/Engine/Plugins/PCG/Source/PCG/Private/Elements/Metadata/PCGMetadataMakeRotator.cpp
--- a/Engine/Plugins/PCG/Source/PCG/Private/Elements/Metadata/PCGMetadataMakeRotator.cpp
+++ b/Engine/Plugins/PCG/Source/PCG/Private/Elements/Metadata/PCGMetadataMakeRotator.cpp
@@ -187,9 +187,10 @@ bool FPCGMetadataMakeRotatorElement::DoOperation(PCGMetadataOps::FOperationData&
 		{
 			return FRotator{ Pitch, Yaw, Roll };
 		});
+	default:
+		ensure(false);
+		return true;
 	}
-
-	return true;
 }
 
 #undef LOCTEXT_NAMESPACE
diff --git a/Engine/Plugins/PCG/Source/PCG/Private/Elements/Metadata/PCGMetadataMakeTransform.cpp b/Engine/Plugins/PCG/Source/PCG/Private/Elements/Metadata/PCGMetadataMakeTransform.cpp
--- a/Engine/Plugins/PCG/Source/PCG/Private/Elements/Metadata/PCGMetadataMakeTransform.cpp
+++ b/Engine/Plugins/PCG/Source/PCG/Private/Elements/Metadata/PCGMetadataMakeTransform.cpp
@@ -193,10 +193,8 @@ bool FPCGMetadataMakeTransformElement::DoOperation(PCGMetadataOps::FOperationDat
 
 	const UPCGMetadataMakeTransformSettings* Settings = CastChecked<UPCGMetadataMakeTransformSettings>(OperationData.Settings);
 
-	auto TransformFunc = [this, &OperationData](auto DummyValue) -> bool
+	auto TransformFunc = [this, &OperationData]<typename AttributeType>(AttributeType) -> bool
 	{
-		using AttributeType = decltype(DummyValue);
-
 		if constexpr (PCG::Private::IsOfTypes<AttributeType, FVector2D, FVector, FVector4>())
 		{
 			return DoTernaryOp<AttributeType, FQuat, AttributeType>(OperationData, PCGMetadataMakeTransformSettings::MakeTransform<AttributeType>);
@@ -207,7 +205,8 @@ bool FPCGMetadataMakeTransformElement::DoOperation(PCGMetadataOps::FOperationDat
 		}
 		else
 		{
-			return false;
+			ensure(false);
+			return true;
 		}
 	};
 
diff --git a/Engine/Plugins/PCG/Source/PCG/Private/Elements/Metadata/PCGMetadataMakeVector.cpp b/Engine/Plugins/PCG/Source/PCG/Private/Elements/Metadata/PCGMetadataMakeVector.cpp
--- a/Engine/Plugins/PCG/Source/PCG/Private/Elements/Metadata/PCGMetadataMakeVector.cpp
+++ b/Engine/Plugins/PCG/Source/PCG/Private/Elements/Metadata/PCGMetadataMakeVector.cpp
@@ -257,17 +257,17 @@ bool FPCGMetadataMakeVectorElement::DoOperation(PCGMetadataOps::FOperationData&
 
 	if (Settings->OutputType == EPCGMetadataTypes::Vector2)
 	{
-		DoBinaryOp<double, double>(OperationData, PCGMetadataMakeVectorSettings::MakeVector2);
+		return DoBinaryOp<double, double>(OperationData, PCGMetadataMakeVectorSettings::MakeVector2);
 	}
 	else if (Settings->OutputType == EPCGMetadataTypes::Vector)
 	{
 		if (Settings->MakeVector3Op == EPCGMetadataMakeVector3::ThreeValues)
 		{
-			DoTernaryOp<double, double, double>(OperationData, PCGMetadataMakeVectorSettings::MakeVector3);
+			return DoTernaryOp<double, double, double>(OperationData, PCGMetadataMakeVectorSettings::MakeVector3);
 		}
 		else
 		{
-			DoBinaryOp<FVector2D, double>(OperationData, PCGMetadataMakeVectorSettings::MakeVector3Vec2);
+			return DoBinaryOp<FVector2D, double>(OperationData, PCGMetadataMakeVectorSettings::MakeVector3Vec2);
 		}
 	}
 	else
@@ -275,21 +275,16 @@ bool FPCGMetadataMakeVectorElement::DoOperation(PCGMetadataOps::FOperationData&
 		switch (Settings->MakeVector4Op)
 		{
 		case EPCGMetadataMakeVector4::FourValues:
-			DoQuaternaryOp<double, double, double, double>(OperationData, PCGMetadataMakeVectorSettings::MakeVector4);
-			break;
+			return DoQuaternaryOp<double, double, double, double>(OperationData, PCGMetadataMakeVectorSettings::MakeVector4);
 		case EPCGMetadataMakeVector4::TwoVector2:
-			DoBinaryOp<FVector2D, FVector2D>(OperationData, PCGMetadataMakeVectorSettings::MakeVector4TwoVec2);
-			break;
+			return DoBinaryOp<FVector2D, FVector2D>(OperationData, PCGMetadataMakeVectorSettings::MakeVector4TwoVec2);
 		case EPCGMetadataMakeVector4::Vector2AndTwoValues:
-			DoTernaryOp<FVector2D, double, double>(OperationData, PCGMetadataMakeVectorSettings::MakeVector4Vec2);
-			break;
+			return DoTernaryOp<FVector2D, double, double>(OperationData, PCGMetadataMakeVectorSettings::MakeVector4Vec2);
 		case EPCGMetadataMakeVector4::Vector3AndValue:
-			DoBinaryOp<FVector, double>(OperationData, PCGMetadataMakeVectorSettings::MakeVector4Vec3);
-			break;
+			return DoBinaryOp<FVector, double>(OperationData, PCGMetadataMakeVectorSettings::MakeVector4Vec3);
 		default:
-			break;
+			ensure(false);
+			return true;
 		}
 	}
-
-	return true;
 }
diff --git a/Engine/Plugins/PCG/Source/PCG/Private/Elements/Metadata/PCGMetadataMathsOpElement.cpp b/Engine/Plugins/PCG/Source/PCG/Private/Elements/Metadata/PCGMetadataMathsOpElement.cpp
--- a/Engine/Plugins/PCG/Source/PCG/Private/Elements/Metadata/PCGMetadataMathsOpElement.cpp
+++ b/Engine/Plugins/PCG/Source/PCG/Private/Elements/Metadata/PCGMetadataMathsOpElement.cpp
@@ -121,14 +121,12 @@ namespace PCGMetadataMathsSettings
 			return Value1 * Value2;
 		case EPCGMetadataMathsOperation::Divide:
 			return (Value2 != T{0}) ? (Value1 / Value2) : T{0}; // To mirror FMath
-		case EPCGMetadataMathsOperation::Max:
-			return PCGMetadataMaths::Max(Value1, Value2);
-		case EPCGMetadataMathsOperation::Min:
-			return PCGMetadataMaths::Min(Value1, Value2);
+		case EPCGMetadataMathsOperation::Max: // fall-through
 		case EPCGMetadataMathsOperation::ClampMin:
-			return PCGMetadataMaths::Clamp(Value1, Value2, Value1);
+			return PCGMetadataMaths::Max(Value1, Value2);
+		case EPCGMetadataMathsOperation::Min: // fall-through
 		case EPCGMetadataMathsOperation::ClampMax:
-			return PCGMetadataMaths::Clamp(Value1, Value1, Value2);
+			return PCGMetadataMaths::Min(Value1, Value2);
 		case EPCGMetadataMathsOperation::Pow:
 			return PCGMetadataMaths::Pow(Value1, Value2);
 		case EPCGMetadataMathsOperation::Modulo:
@@ -324,14 +322,13 @@ bool FPCGMetadataMathsElement::DoOperation(PCGMetadataOps::FOperationData& Opera
 
 	const UPCGMetadataMathsSettings* Settings = CastChecked<UPCGMetadataMathsSettings>(OperationData.Settings);
 
-	auto MathFunc = [this, Operation = Settings->Operation, &OperationData](auto DummyOutValue) -> void
+	auto MathFunc = [this, Operation = Settings->Operation, &OperationData]<typename AttributeType>(AttributeType) -> bool
 	{
-		using AttributeType = decltype(DummyOutValue);
-
 		// Need to remove types that would not compile
 		if constexpr (!PCG::Private::IsOfTypes<AttributeType, float, double, int32, int64, FVector2D, FVector, FVector4>())
 		{
-			return;
+			ensure(false);
+			return true;
 		}
 		else
 		{
@@ -343,22 +340,27 @@ bool FPCGMetadataMathsElement::DoOperation(PCGMetadataOps::FOperationData& Opera
 
 				if (OperationData.OutputType == PCG::Private::MetadataTypes<int64>::Id)
 				{
-					DoUnaryOp<AttributeType>(OperationData, [Operation](const AttributeType& Value) -> OverriddenOutputType { return static_cast<OverriddenOutputType>(PCGMetadataMathsSettings::UnaryOp<AttributeType>(Value, Operation)); });
+					return DoUnaryOp<AttributeType>(OperationData, [Operation](const AttributeType& Value) -> OverriddenOutputType { return static_cast<OverriddenOutputType>(PCGMetadataMathsSettings::UnaryOp<AttributeType>(Value, Operation)); });
 				}
 				else
 				{
-					DoUnaryOp<AttributeType>(OperationData, [Operation](const AttributeType& Value) -> AttributeType { return PCGMetadataMathsSettings::UnaryOp(Value, Operation); });
+					return DoUnaryOp<AttributeType>(OperationData, [Operation](const AttributeType& Value) -> AttributeType { return PCGMetadataMathsSettings::UnaryOp(Value, Operation); });
 				}
 			}
 			else if (PCGMetadataMathsSettings::IsBinaryOp(Operation))
 			{
 				TRACE_CPUPROFILER_EVENT_SCOPE(FPCGMetadataMathsElement::ExecuteInternal::BinaryOp);
-				DoBinaryOp<AttributeType, AttributeType>(OperationData, [Operation](const AttributeType& Value1, const AttributeType& Value2) -> AttributeType { return PCGMetadataMathsSettings::BinaryOp(Value1, Value2, Operation); });
+				return DoBinaryOp<AttributeType, AttributeType>(OperationData, [Operation](const AttributeType& Value1, const AttributeType& Value2) -> AttributeType { return PCGMetadataMathsSettings::BinaryOp(Value1, Value2, Operation); });
 			}
 			else if (PCGMetadataMathsSettings::IsTernaryOp(Operation))
 			{
 				TRACE_CPUPROFILER_EVENT_SCOPE(FPCGMetadataMathsElement::ExecuteInternal::TernaryOp);
-				DoTernaryOp<AttributeType, AttributeType, AttributeType>(OperationData, [Operation](const AttributeType& Value1, const AttributeType& Value2, const AttributeType& Value3) -> AttributeType { return PCGMetadataMathsSettings::TernaryOp(Value1, Value2, Value3, Operation); });
+				return DoTernaryOp<AttributeType, AttributeType, AttributeType>(OperationData, [Operation](const AttributeType& Value1, const AttributeType& Value2, const AttributeType& Value3) -> AttributeType { return PCGMetadataMathsSettings::TernaryOp(Value1, Value2, Value3, Operation); });
+			}
+			else
+			{
+				ensure(false);
+				return true;
 			}
 		}
 	};
@@ -366,12 +368,10 @@ bool FPCGMetadataMathsElement::DoOperation(PCGMetadataOps::FOperationData& Opera
 	// If the output is double, force all to double.
 	if (OperationData.OutputType == PCG::Private::MetadataTypes<double>::Id)
 	{
-		MathFunc(double{});
+		return MathFunc(double{});
 	}
 	else
 	{
-		PCGMetadataAttribute::CallbackWithRightType(OperationData.MostComplexInputType, MathFunc);
+		return PCGMetadataAttribute::CallbackWithRightType(OperationData.MostComplexInputType, MathFunc);
 	}
-
-	return true;
 }
diff --git a/Engine/Plugins/PCG/Source/PCG/Private/Elements/Metadata/PCGMetadataOpElementBase.cpp b/Engine/Plugins/PCG/Source/PCG/Private/Elements/Metadata/PCGMetadataOpElementBase.cpp
--- a/Engine/Plugins/PCG/Source/PCG/Private/Elements/Metadata/PCGMetadataOpElementBase.cpp
+++ b/Engine/Plugins/PCG/Source/PCG/Private/Elements/Metadata/PCGMetadataOpElementBase.cpp
@@ -422,6 +422,8 @@ bool FPCGMetadataElementBase::PrepareDataInternal(FPCGContext* Context) const
 		TArray<FPCGTaggedData>& Outputs = Context->OutputData.TaggedData;
 		const uint32 NumberOfResults = Settings->GetResultNum();
 
+		OutState.Context = Context;
+
 		// Gathering all the inputs metadata
 		TArray<const UPCGMetadata*> SourceMetadata;
 		TArray<const FPCGMetadataAttributeBase*> SourceAttribute;
@@ -587,61 +589,47 @@ bool FPCGMetadataElementBase::PrepareDataInternal(FPCGContext* Context) const
 		// Use implicit capture, since we capture a lot
 		auto CreateAttribute = [&]<typename AttributeType>(uint32 OutputIndex, AttributeType DummyOutValue) -> bool
 		{
-			FPCGTaggedData& OutputData = Outputs.Add_GetRef(InputTaggedData[PrimaryPinIndex]);
-			OutputData.Pin = Settings->GetOutputPinLabel(OutputIndex);
-
-			UPCGMetadata* OutMetadata = nullptr;
+			FPCGTaggedData& OutputTaggedData = Outputs.Add_GetRef(InputTaggedData[PrimaryPinIndex]);
+			OutputTaggedData.Pin = Settings->GetOutputPinLabel(OutputIndex);
 
-			const FName OutputName = OutputTarget.GetName();
-			const FText OutputTargetText = OutputTarget.GetDisplayText();
+			// In case of property or attribute with extra accessor, we need to validate that the property/attribute can accept the output type.
+			// Verify this before duplicating, because an extra allocation is certainly less costly than duplicating the data.
+			// Do it with a const accessor, since OutputTaggedData.Data is still pointing on the const input data.
 
-			if (OutputTarget.GetSelection() == EPCGAttributePropertySelection::Attribute && OutputTarget.GetExtraNames().IsEmpty())
+			if (!OutputTarget.IsBasicAttribute())
 			{
-				// In case of an attribute, we check if we have extra selectors. If not, we can just delete the attribute
-				// and create a new one of the right type.
-				// But if we have extra selectors, we need to handle it the same way as properties.
-				// There is no point of failure before duplicating. So duplicate, create the attribute and then the accessor.
-				PCGMetadataElementCommon::DuplicateTaggedData(InputTaggedData[PrimaryPinIndex], OutputData, OutMetadata);
-				FPCGMetadataAttributeBase* OutputAttribute = PCGMetadataElementCommon::ClearOrCreateAttribute<AttributeType>(OutMetadata, OutputName);
-				if (!OutputAttribute)
+				const TUniquePtr<const IPCGAttributeAccessor> TempConstAccessor = PCGAttributeAccessorHelpers::CreateConstAccessor(OutputTaggedData.Data.Get(), OutputTarget);
+
+				const bool bIsValid = TempConstAccessor.IsValid() && PCG::Private::IsBroadcastable(PCG::Private::MetadataTypes<AttributeType>::Id, TempConstAccessor->GetUnderlyingType());
+
+				if (!bIsValid)
 				{
+					PCGE_LOG(Error, GraphAndLog, FText::Format(LOCTEXT("AttributeTypeBroadcastFailed_Updated", "Attribute/Property '{0}' ({1}) is not compatible with operation output type ({2})."),
+						OutputTarget.GetDisplayText(),
+						PCG::Private::GetTypeNameText(TempConstAccessor->GetUnderlyingType()),
+						PCG::Private::GetTypeNameText<AttributeType>()));
 					return false;
 				}
+			}
 
-				// And copy the mapping from the original attribute, if it is not points
-				if (!InputTaggedData[PrimaryPinIndex].Data->IsA<UPCGPointData>() && SourceMetadata[PrimaryPinIndex] && SourceAttribute[PrimaryPinIndex])
-				{
-					PCGMetadataElementCommon::CopyEntryToValueKeyMap(SourceMetadata[PrimaryPinIndex], SourceAttribute[PrimaryPinIndex], OutputAttribute);
-				}
+			check(InputTaggedData[PrimaryPinIndex].Data);
+			UPCGData* OutputData = InputTaggedData[PrimaryPinIndex].Data->DuplicateData();
+			check(OutputData);
+			OutputTaggedData.Data = OutputData;
 
-				OperationData.OutputAccessors[OutputIndex] = PCGAttributeAccessorHelpers::CreateAccessor(const_cast<UPCGData*>(OutputData.Data.Get()), OutputTarget);
-			}
-			else
+			if (OutputTarget.IsBasicAttribute())
 			{
-				// In case of property or attribute with extra accessor, we need to validate that the property/attribute can accept the output type.
-				// Verify this before duplicating, because an extra allocation is certainly less costly than duplicating the data.
-				// Do it with a const accessor, since OutputData.Data is still pointing on the const input data.
-				const TUniquePtr<const IPCGAttributeAccessor> TempConstAccessor = PCGAttributeAccessorHelpers::CreateConstAccessor(OutputData.Data.Get(), OutputTarget);
+				const FName OutputName = OutputTarget.GetName();
 
-				if (TempConstAccessor.IsValid())
+				FPCGMetadataAttributeBase* OutputAttribute = PCGMetadataElementCommon::ClearOrCreateAttribute<AttributeType>(OutputData->MutableMetadata(), OutputName);
+				if (!OutputAttribute)
 				{
-					// We matched an attribute/property, check if the output type is valid.
-					if (!PCG::Private::IsBroadcastable(PCG::Private::MetadataTypes<AttributeType>::Id, TempConstAccessor->GetUnderlyingType()))
-					{
-						PCGE_LOG(Error, GraphAndLog, FText::Format(LOCTEXT("AttributeTypeBroadcastFailed_Updated", "Attribute/Property '{0}' ({1}) is not compatible with operation output type ({2})."),
-							OutputTargetText,
-							PCG::Private::GetTypeNameText(TempConstAccessor->GetUnderlyingType()),
-							PCG::Private::GetTypeNameText<AttributeType>()));
-						return false;
-					}
-
-					PCGMetadataElementCommon::DuplicateTaggedData(InputTaggedData[PrimaryPinIndex], OutputData, OutMetadata);
-
-					// Re-create the accessor to point to the right data (since we just duplicated the data)
-					OperationData.OutputAccessors[OutputIndex] = PCGAttributeAccessorHelpers::CreateAccessor(const_cast<UPCGData*>(OutputData.Data.Get()), OutputTarget);
+					return false;
 				}
 			}
 
+			OperationData.OutputAccessors[OutputIndex] = PCGAttributeAccessorHelpers::CreateAccessor(OutputData, OutputTarget);
+
 			if (!OperationData.OutputAccessors[OutputIndex].IsValid())
 			{
 				return false;
@@ -649,11 +637,11 @@ bool FPCGMetadataElementBase::PrepareDataInternal(FPCGContext* Context) const
 
 			if (OperationData.OutputAccessors[OutputIndex]->IsReadOnly())
 			{
-				PCGE_LOG(Error, GraphAndLog, FText::Format(LOCTEXT("OutputAccessorIsReadOnly", "Attribute/Property '{0}' is read only."), OutputTargetText));
+				PCGE_LOG(Error, GraphAndLog, FText::Format(LOCTEXT("OutputAccessorIsReadOnly", "Attribute/Property '{0}' is read only."), OutputTarget.GetDisplayText()));
 				return false;
 			}
 
-			OperationData.OutputKeys[OutputIndex] = PCGAttributeAccessorHelpers::CreateKeys(const_cast<UPCGData*>(OutputData.Data.Get()), OutputTarget);
+			OperationData.OutputKeys[OutputIndex] = PCGAttributeAccessorHelpers::CreateKeys(OutputData, OutputTarget);
 
 			return OperationData.OutputKeys[OutputIndex].IsValid();
 		};
@@ -729,13 +717,7 @@ bool FPCGMetadataElementBase::ExecuteInternal(FPCGContext* Context) const
 			return true;
 		}
 
-		// TODO: Add range-based async evaluation to the DoOperation function in the future
-		if (!DoOperation(IterState))
-		{
-			PCGE_LOG(Error, GraphAndLog, LOCTEXT("ErrorOccurred", "Error while performing the metadata operation, check logs for more information"));
-		}
-
-		return true;
+		return DoOperation(IterState);
 	});
 }
 
diff --git a/Engine/Plugins/PCG/Source/PCG/Private/Elements/Metadata/PCGMetadataRotatorOpElement.cpp b/Engine/Plugins/PCG/Source/PCG/Private/Elements/Metadata/PCGMetadataRotatorOpElement.cpp
--- a/Engine/Plugins/PCG/Source/PCG/Private/Elements/Metadata/PCGMetadataRotatorOpElement.cpp
+++ b/Engine/Plugins/PCG/Source/PCG/Private/Elements/Metadata/PCGMetadataRotatorOpElement.cpp
@@ -245,38 +245,37 @@ bool FPCGMetadataRotatorElement::DoOperation(PCGMetadataOps::FOperationData& Ope
 
 	const UPCGMetadataRotatorSettings* Settings = CastChecked<UPCGMetadataRotatorSettings>(OperationData.Settings);
 
-	auto RotatorFunc = [this, &OperationData, Operation = Settings->Operation](auto DummyValue)
+	auto RotatorFunc = [this, &OperationData, Operation = Settings->Operation]<typename AttributeType>(AttributeType) -> bool
 	{
-		using AttributeType = decltype(DummyValue);
-
 		if constexpr (!PCG::Private::IsOfTypes<AttributeType, FQuat, FRotator>())
 		{
-			return;
+			ensure(false);
+			return true;
 		}
 		else
 		{
 			if (PCGMetadataRotatorSettings::IsTransfromOp(Operation))
 			{
-				DoBinaryOp<AttributeType, FTransform>(OperationData, [Operation](const AttributeType& Value, const FTransform& Transform)->AttributeType {
+				return DoBinaryOp<AttributeType, FTransform>(OperationData, [Operation](const AttributeType& Value, const FTransform& Transform)->AttributeType {
 					return PCGMetadataRotatorSettings::ApplyTransformOperation(Value, Transform, Operation);
 					});
 			}
 			else if (PCGMetadataRotatorSettings::IsUnaryOp(Operation))
 			{
-				DoUnaryOp<AttributeType>(OperationData, [Operation](const AttributeType& Value)->AttributeType {
+				return DoUnaryOp<AttributeType>(OperationData, [Operation](const AttributeType& Value)->AttributeType {
 					double DummyDouble = 0.0;
 					return PCGMetadataRotatorSettings::ApplyRotatorOperation(Value, AttributeType{}, DummyDouble, Operation);
 				});
 			}
 			else if (PCGMetadataRotatorSettings::IsTernaryOp(Operation))
 			{
-				DoTernaryOp<AttributeType, AttributeType, double>(OperationData, [Operation](const AttributeType& Value1, const AttributeType& Value2, const double& Ratio)->AttributeType {
+				return DoTernaryOp<AttributeType, AttributeType, double>(OperationData, [Operation](const AttributeType& Value1, const AttributeType& Value2, const double& Ratio)->AttributeType {
 					return PCGMetadataRotatorSettings::ApplyRotatorOperation(Value1, Value2, Ratio, Operation);
 				});
 			}
 			else
 			{
-				DoBinaryOp<AttributeType, AttributeType>(OperationData, [Operation](const AttributeType& Value1, const AttributeType& Value2)->AttributeType {
+				return DoBinaryOp<AttributeType, AttributeType>(OperationData, [Operation](const AttributeType& Value1, const AttributeType& Value2)->AttributeType {
 					double DummyDouble = 0.0;
 					return PCGMetadataRotatorSettings::ApplyRotatorOperation(Value1, Value2, DummyDouble, Operation);
 				});
@@ -284,7 +283,5 @@ bool FPCGMetadataRotatorElement::DoOperation(PCGMetadataOps::FOperationData& Ope
 		}
 	};
 
-	PCGMetadataAttribute::CallbackWithRightType(OperationData.OutputType, RotatorFunc);
-
-	return true;
+	return PCGMetadataAttribute::CallbackWithRightType(OperationData.OutputType, RotatorFunc);
 }
diff --git a/Engine/Plugins/PCG/Source/PCG/Private/Elements/Metadata/PCGMetadataStringOpElement.cpp b/Engine/Plugins/PCG/Source/PCG/Private/Elements/Metadata/PCGMetadataStringOpElement.cpp
--- a/Engine/Plugins/PCG/Source/PCG/Private/Elements/Metadata/PCGMetadataStringOpElement.cpp
+++ b/Engine/Plugins/PCG/Source/PCG/Private/Elements/Metadata/PCGMetadataStringOpElement.cpp
@@ -127,15 +127,18 @@ bool FPCGMetadataStringOpElement::DoOperation(PCGMetadataOps::FOperationData& Op
 
 	if (Settings->Operation == EPCGMetadataStringOperation::Append)
 	{
-		DoBinaryOp<FString, FString>(OperationData, [](const FString& Value1, const FString& Value2) -> FString { return Value1 + Value2; });
+		return DoBinaryOp<FString, FString>(OperationData, [](const FString& Value1, const FString& Value2) -> FString { return Value1 + Value2; });
 	}
 	else if (Settings->Operation == EPCGMetadataStringOperation::Replace)
 	{
-		DoTernaryOp<FString, FString, FString>(OperationData, [](const FString& InValue, const FString& InSearch, const FString& InReplace)
+		return DoTernaryOp<FString, FString, FString>(OperationData, [](const FString& InValue, const FString& InSearch, const FString& InReplace)
 		{
 			return InValue.Replace(*InSearch, *InReplace);
 		});
 	}
-
-	return true;
+	else
+	{
+		ensure(false);
+		return true;
+	}
 }
\ No newline at end of file
diff --git a/Engine/Plugins/PCG/Source/PCG/Private/Elements/Metadata/PCGMetadataTransformOpElement.cpp b/Engine/Plugins/PCG/Source/PCG/Private/Elements/Metadata/PCGMetadataTransformOpElement.cpp
--- a/Engine/Plugins/PCG/Source/PCG/Private/Elements/Metadata/PCGMetadataTransformOpElement.cpp
+++ b/Engine/Plugins/PCG/Source/PCG/Private/Elements/Metadata/PCGMetadataTransformOpElement.cpp
@@ -210,24 +210,22 @@ bool FPCGMetadataTransformElement::DoOperation(PCGMetadataOps::FOperationData& O
 
 	if (PCGMetadataTransfromSettings::IsUnaryOp(Settings->Operation))
 	{
-		DoUnaryOp<FTransform>(OperationData, [Operation = Settings->Operation, Mode = Settings->TransformLerpMode](const FTransform& Value)->FTransform {
+		return DoUnaryOp<FTransform>(OperationData, [Operation = Settings->Operation, Mode = Settings->TransformLerpMode](const FTransform& Value)->FTransform {
 			double DummyDouble = 0.0;
 			return PCGMetadataTransfromSettings::ApplyTransformOperation(Value, FTransform{}, DummyDouble, Operation, Mode);
 			});
 	}
 	else if (PCGMetadataTransfromSettings::IsTernaryOp(Settings->Operation))
 	{
-		DoTernaryOp<FTransform, FTransform, double>(OperationData, [Operation = Settings->Operation, Mode = Settings->TransformLerpMode](const FTransform& Value1, const FTransform& Value2, const double& Ratio)->FTransform {
+		return DoTernaryOp<FTransform, FTransform, double>(OperationData, [Operation = Settings->Operation, Mode = Settings->TransformLerpMode](const FTransform& Value1, const FTransform& Value2, const double& Ratio)->FTransform {
 			return PCGMetadataTransfromSettings::ApplyTransformOperation(Value1, Value2, Ratio, Operation, Mode);
 			});
 	}
 	else
 	{
-		DoBinaryOp<FTransform, FTransform>(OperationData, [Operation = Settings->Operation, Mode = Settings->TransformLerpMode](const FTransform& Value1, const FTransform& Value2)->FTransform {
+		return DoBinaryOp<FTransform, FTransform>(OperationData, [Operation = Settings->Operation, Mode = Settings->TransformLerpMode](const FTransform& Value1, const FTransform& Value2)->FTransform {
 			double DummyDouble = 0.0;
 			return PCGMetadataTransfromSettings::ApplyTransformOperation(Value1, Value2, DummyDouble, Operation, Mode);
 			});
 	}
-
-	return true;
 }
diff --git a/Engine/Plugins/PCG/Source/PCG/Private/Elements/Metadata/PCGMetadataTrigOpElement.cpp b/Engine/Plugins/PCG/Source/PCG/Private/Elements/Metadata/PCGMetadataTrigOpElement.cpp
--- a/Engine/Plugins/PCG/Source/PCG/Private/Elements/Metadata/PCGMetadataTrigOpElement.cpp
+++ b/Engine/Plugins/PCG/Source/PCG/Private/Elements/Metadata/PCGMetadataTrigOpElement.cpp
@@ -156,12 +156,10 @@ bool FPCGMetadataTrigElement::DoOperation(PCGMetadataOps::FOperationData& Operat
 
 	if (Settings->Operation == EPCGMetadataTrigOperation::Atan2)
 	{
-		DoBinaryOp<double, double>(OperationData, [](const double& Value1, const double& Value2) -> double { return PCGMetadataTrigSettings::BinaryOp(Value1, Value2); });
+		return DoBinaryOp<double, double>(OperationData, [](const double& Value1, const double& Value2) -> double { return PCGMetadataTrigSettings::BinaryOp(Value1, Value2); });
 	}
 	else
 	{
-		DoUnaryOp<double>(OperationData, [Operation = Settings->Operation](const double& Value) -> double { return PCGMetadataTrigSettings::UnaryOp(Value, Operation); });
+		return DoUnaryOp<double>(OperationData, [Operation = Settings->Operation](const double& Value) -> double { return PCGMetadataTrigSettings::UnaryOp(Value, Operation); });
 	}
-
-	return true;
 }
diff --git a/Engine/Plugins/PCG/Source/PCG/Private/Elements/Metadata/PCGMetadataVectorOpElement.cpp b/Engine/Plugins/PCG/Source/PCG/Private/Elements/Metadata/PCGMetadataVectorOpElement.cpp
--- a/Engine/Plugins/PCG/Source/PCG/Private/Elements/Metadata/PCGMetadataVectorOpElement.cpp
+++ b/Engine/Plugins/PCG/Source/PCG/Private/Elements/Metadata/PCGMetadataVectorOpElement.cpp
@@ -387,25 +387,24 @@ bool FPCGMetadataVectorElement::DoOperation(PCGMetadataOps::FOperationData& Oper
 
 	const UPCGMetadataVectorSettings* Settings = CastChecked<UPCGMetadataVectorSettings>(OperationData.Settings);
 
-	auto VectorFunc = [this, Operation = Settings->Operation, &OperationData](auto DummyValue)
+	auto VectorFunc = [this, Operation = Settings->Operation, &OperationData]<typename AttributeType>(AttributeType) -> bool
 	{
-		using AttributeType = decltype(DummyValue);
-
 		if constexpr (!PCG::Private::IsOfTypes<AttributeType, FVector2D, FVector, FVector4>())
 		{
-			return;
+			ensure(false);
+			return true;
 		}
 		else
 		{
 			if (PCGMetadataVectorSettings::IsTransformOp(Operation))
 			{
-				DoBinaryOp<AttributeType, FTransform>(OperationData, [Operation](const AttributeType& Value, const FTransform& Transform) -> AttributeType { return PCGMetadataVectorSettings::ApplyTransformOperation(Value, Transform, Operation);});
+				return DoBinaryOp<AttributeType, FTransform>(OperationData, [Operation](const AttributeType& Value, const FTransform& Transform) -> AttributeType { return PCGMetadataVectorSettings::ApplyTransformOperation(Value, Transform, Operation);});
 			}
 			else if (PCG::Private::IsOfTypes<double>(OperationData.OutputType))
 			{
 				if (PCGMetadataVectorSettings::IsUnaryOp(Operation))
 				{
-					DoUnaryOp<AttributeType>(OperationData, [Operation](const AttributeType& Value) -> double {
+					return DoUnaryOp<AttributeType>(OperationData, [Operation](const AttributeType& Value) -> double {
 						double Res = 0.0;
 						PCGMetadataVectorSettings::ApplyVectorOperation(Value, AttributeType{}, Res, Operation);
 						return Res;
@@ -413,7 +412,7 @@ bool FPCGMetadataVectorElement::DoOperation(PCGMetadataOps::FOperationData& Oper
 				}
 				else
 				{
-					DoBinaryOp<AttributeType, AttributeType>(OperationData, [Operation](const AttributeType& Value1, const AttributeType& Value2) -> double {
+					return DoBinaryOp<AttributeType, AttributeType>(OperationData, [Operation](const AttributeType& Value1, const AttributeType& Value2) -> double {
 						double Res = 0.0;
 						PCGMetadataVectorSettings::ApplyVectorOperation(Value1, Value2, Res, Operation);
 						return Res;
@@ -424,21 +423,21 @@ bool FPCGMetadataVectorElement::DoOperation(PCGMetadataOps::FOperationData& Oper
 			{
 				if (PCGMetadataVectorSettings::IsUnaryOp(Operation))
 				{
-					DoUnaryOp<AttributeType>(OperationData, [Operation](const AttributeType& Value) -> AttributeType {
+					return DoUnaryOp<AttributeType>(OperationData, [Operation](const AttributeType& Value) -> AttributeType {
 						double DummyDouble = 0.0;
 						return PCGMetadataVectorSettings::ApplyVectorOperation(Value, AttributeType{}, DummyDouble, Operation);
 						});
 				}
 				else if (PCGMetadataVectorSettings::IsTernaryOp(Operation))
 				{
-					DoTernaryOp<AttributeType, AttributeType, double>(OperationData, [Operation](const AttributeType& Value1, const AttributeType& Value2, const double& DoubleValue) -> AttributeType {
+					return DoTernaryOp<AttributeType, AttributeType, double>(OperationData, [Operation](const AttributeType& Value1, const AttributeType& Value2, const double& DoubleValue) -> AttributeType {
 						double DoubleCopy = DoubleValue;
 						return PCGMetadataVectorSettings::ApplyVectorOperation(Value1, Value2, DoubleCopy, Operation);
 						});
 				}
 				else
 				{
-					DoBinaryOp<AttributeType, AttributeType>(OperationData, [Operation](const AttributeType& Value1, const AttributeType& Value2) -> AttributeType {
+					return DoBinaryOp<AttributeType, AttributeType>(OperationData, [Operation](const AttributeType& Value1, const AttributeType& Value2) -> AttributeType {
 						double DummyDouble = 0.0;
 						return PCGMetadataVectorSettings::ApplyVectorOperation(Value1, Value2, DummyDouble, Operation);
 						});
@@ -447,7 +446,5 @@ bool FPCGMetadataVectorElement::DoOperation(PCGMetadataOps::FOperationData& Oper
 		}
 	};
 
-	PCGMetadataAttribute::CallbackWithRightType(OperationData.MostComplexInputType, VectorFunc);
-
-	return true;
+	return PCGMetadataAttribute::CallbackWithRightType(OperationData.MostComplexInputType, VectorFunc);
 }
diff --git a/Engine/Plugins/PCG/Source/PCG/Private/Elements/PCGAttributeRemap.cpp b/Engine/Plugins/PCG/Source/PCG/Private/Elements/PCGAttributeRemap.cpp
--- a/Engine/Plugins/PCG/Source/PCG/Private/Elements/PCGAttributeRemap.cpp
+++ b/Engine/Plugins/PCG/Source/PCG/Private/Elements/PCGAttributeRemap.cpp
@@ -153,7 +153,7 @@ bool FPCGAttributeRemapElement::DoOperation(PCGMetadataOps::FOperationData& Oper
 		Params.Intercept = OutRangeTrueMin;
 	}
 
-	auto RemapFunc = [this, &Params, &OperationData]<typename AttributeType>(AttributeType DummyValue) -> bool
+	auto RemapFunc = [this, &Params, &OperationData]<typename AttributeType>(AttributeType) -> bool
 	{
 		using FinalType = std::conditional_t<std::is_same_v<AttributeType, FQuat>, FRotator, AttributeType>;
 		return DoUnaryOp<FinalType>(OperationData, [&Params](const FinalType& Value) -> FinalType 
diff --git a/Engine/Plugins/PCG/Source/PCG/Private/Helpers/PCGAsyncState.cpp b/Engine/Plugins/PCG/Source/PCG/Private/Helpers/PCGAsyncState.cpp
--- a/Engine/Plugins/PCG/Source/PCG/Private/Helpers/PCGAsyncState.cpp
+++ b/Engine/Plugins/PCG/Source/PCG/Private/Helpers/PCGAsyncState.cpp
@@ -7,4 +7,14 @@
 bool FPCGAsyncState::ShouldStop() const
 {
 	return FPlatformTime::Seconds() > EndTime;
+}
+
+void FPCGAsyncState::Reset()
+{
+	AsyncCurrentWriteIndex = 0;
+	TasksChunksStartEnd.Reset();
+	CurrentChunkToCollapse = 0;
+	ChunkToNumElementsWrittenMap.Reset();
+
+	bStarted = false;
 }
\ No newline at end of file
diff --git a/Engine/Plugins/PCG/Source/PCG/Public/Elements/Metadata/PCGMetadataOpElementBase.h b/Engine/Plugins/PCG/Source/PCG/Public/Elements/Metadata/PCGMetadataOpElementBase.h
--- a/Engine/Plugins/PCG/Source/PCG/Public/Elements/Metadata/PCGMetadataOpElementBase.h
+++ b/Engine/Plugins/PCG/Source/PCG/Public/Elements/Metadata/PCGMetadataOpElementBase.h
@@ -4,6 +4,8 @@
 
 #include "PCGSettings.h"
 
+#include "Data/PCGPointData.h"
+#include "Helpers/PCGAsync.h"
 #include "Metadata/PCGAttributePropertySelector.h"
 #include "Metadata/Accessors/IPCGAttributeAccessor.h"
 #include "Metadata/Accessors/PCGAttributeAccessorKeys.h"
@@ -15,6 +17,17 @@
 
 class FPCGMetadataAttributeBase;
 
+// FIXME: To be removed when we are confident Metadata is stable in MT.
+static inline TAutoConsoleVariable<bool> CVarMetadataOperationInMT(
+	TEXT("pcg.MetadataOperationInMT"),
+	true,
+	TEXT("Metadata operations are now multithreaded."));
+
+static inline TAutoConsoleVariable<int> CVarMetadataOperationChunkSize(
+	TEXT("pcg.MetadataOperationChunkSize"),
+	256,
+	TEXT("Metadata operations chunk size."));
+
 namespace PCGMetadataSettingsBaseConstants
 {
 	const FName DoubleInputFirstLabel = TEXT("InA");
@@ -174,6 +187,8 @@ namespace PCGMetadataOps
 		uint16 OutputType;
 		const UPCGMetadataSettingsBase* Settings = nullptr;
 
+		FPCGContext* Context;
+
 		TArray<FPCGAttributePropertyInputSelector> InputSources;
 
 		TArray<TUniquePtr<const IPCGAttributeAccessorKeys>> InputKeys;
@@ -424,24 +439,40 @@ inline bool FPCGMetadataElementBase::DoNAryOp(PCGMetadataOps::FOperationData& In
 
 	EPCGAttributeAccessorFlags Flags = EPCGAttributeAccessorFlags::AllowBroadcast;
 
-	// First set the default value
+	// First set the default value (only on first pass)
 	PCG::Private::NAryOperation::Options Options{ Flags, Flags | EPCGAttributeAccessorFlags::AllowSetDefaultValue, true };
-	PCG::Private::NAryOperation::Operation<InputTypes...>(InOperationData, /*StartIndex=*/0, /*Range=*/1, Options, InCallbacks);
+	if (!InOperationData.Context->AsyncState.bStarted)
+	{
+		PCG::Private::NAryOperation::Operation<InputTypes...>(InOperationData, /*StartIndex=*/0, /*Range=*/1, Options, InCallbacks);
+	}
 
 	// Then iterate over all the values
 	Options.SetFlags = Flags;
 	Options.bUseDefaultKey = false;
 
-	const int32 NumberOfIterations = (InOperationData.NumberOfElementsToProcess + PCG::Private::NAryOperation::DefaultChunkSize - 1) / PCG::Private::NAryOperation::DefaultChunkSize;
+	const int32 ChunkSize = CVarMetadataOperationChunkSize.GetValueOnAnyThread();
 
-	for (int32 i = 0; i < NumberOfIterations; ++i)
+	if (CVarMetadataOperationInMT.GetValueOnAnyThread())
 	{
-		int32 StartIndex = i * PCG::Private::NAryOperation::DefaultChunkSize;
-		int32 Range = FMath::Min(InOperationData.NumberOfElementsToProcess - StartIndex, PCG::Private::NAryOperation::DefaultChunkSize);
-		PCG::Private::NAryOperation::Operation<InputTypes...>(InOperationData, StartIndex, Range, Options, InCallbacks);
+		return FPCGAsync::AsyncProcessingOneToOneRangeEx(&InOperationData.Context->AsyncState, InOperationData.NumberOfElementsToProcess, []() {},
+			[&InOperationData, &InCallbacks, &Options](int32 StartReadIndex, int32 StartWriteIndex, int32 Count)
+		{
+			PCG::Private::NAryOperation::Operation<InputTypes...>(InOperationData, StartReadIndex, Count, Options, InCallbacks);
+			return Count;
+		}, /*bEnableTimeSlicing=*/true, ChunkSize);
 	}
+	else
+	{
+		const int32 NumberOfIterations = (InOperationData.NumberOfElementsToProcess + PCG::Private::NAryOperation::DefaultChunkSize - 1) / PCG::Private::NAryOperation::DefaultChunkSize;
+		for (int32 i = 0; i < NumberOfIterations; ++i)
+		{
+			int32 StartIndex = i * PCG::Private::NAryOperation::DefaultChunkSize;
+			int32 Range = FMath::Min(InOperationData.NumberOfElementsToProcess - StartIndex, PCG::Private::NAryOperation::DefaultChunkSize);
+			PCG::Private::NAryOperation::Operation<InputTypes...>(InOperationData, StartIndex, Range, Options, InCallbacks);
+		}
 
-	return true;
+		return true;
+	}
 }
 
 template <typename InType, typename... Callbacks>
diff --git a/Engine/Plugins/PCG/Source/PCG/Public/Helpers/PCGAsync.h b/Engine/Plugins/PCG/Source/PCG/Public/Helpers/PCGAsync.h
--- a/Engine/Plugins/PCG/Source/PCG/Public/Helpers/PCGAsync.h
+++ b/Engine/Plugins/PCG/Source/PCG/Public/Helpers/PCGAsync.h
@@ -101,39 +101,63 @@ namespace FPCGAsync
 		const int32 OverrideChunkSize = ConsoleVar::CVarAsyncOverrideChunkSize.GetValueOnAnyThread();
 		const int32 ChunkSize = OverrideChunkSize > 0 ? OverrideChunkSize : InChunkSize;
 
-		const int32 StartIndex = AsyncState.AsyncCurrentReadIndex;
 		const bool bEnableTimeSlicing = bInEnableTimeSlicing && !ConsoleVar::CVarDisableAsyncTimeSlicing.GetValueOnAnyThread();
 
-		if (AsyncState.NumAvailableTasks <= 0 || ChunkSize <= 0 || NumIterations <= 0 || StartIndex < 0)
+		if (AsyncState.NumAvailableTasks <= 0 || ChunkSize <= 0 || NumIterations <= 0)
 		{
 			// Invalid request
 			return true;
 		}
 
-		if (StartIndex >= NumIterations)
+		const int32 ChunksNumber = 1 + ((NumIterations - 1) / ChunkSize);
+
+		if (AsyncState.CurrentChunkToCollapse >= ChunksNumber)
 		{
 			// Nothing left to do
 			return true;
 		}
 
-		const int32 RemainingIterations = NumIterations - StartIndex;
-		const int32 ChunksNumber = 1 + ((RemainingIterations - 1) / ChunkSize);
-		const int32 NumTasks = FMath::Min(AsyncState.NumAvailableTasks, ChunksNumber);
-		const int32 NumFutures = NumTasks - 1;
-
-		if (StartIndex == 0)
+		if (!AsyncState.bStarted)
 		{
 			TRACE_CPUPROFILER_EVENT_SCOPE(FPCGAsync::Private::AsyncProcessing::Initialize);
 
 			Initialize();
+
+			AsyncState.bStarted = true;
+
+			// At the beginning, dispatch the chunks that each task needs to process.
+			// It will be stored on the AsyncState so that it can be reused.
+			// Each task will update its current counter directly in the AsyncState.
+			// It's OK to not protect it since they all will read/write from different place in memory.
+
+			int32 NumTasks = FMath::Min(AsyncState.NumAvailableTasks, ChunksNumber);
+			const int32 IterationsPerTask = ChunksNumber / NumTasks;
+			
+			// Some tasks will have an extra chunk of work if NumTasks doesn't divide ChunksNumber.
+			const int32 ExtraChunkLimit = ChunksNumber - (IterationsPerTask * NumTasks);
+
+			// Main thread is also doing work.
+			const int32 NumFutures = NumTasks - 1;
+			if (NumFutures != 0)
+			{
+				int32 Count = 0;
+				AsyncState.TasksChunksStartEnd.Reserve(NumTasks);
+				for (int32 i = 0; i < NumTasks; ++i)
+				{
+					const int32 NumChunkForThisTask = i < ExtraChunkLimit ? IterationsPerTask + 1 : IterationsPerTask;
+					AsyncState.TasksChunksStartEnd.Emplace_GetRef(Count, Count + NumChunkForThisTask);
+					Count += NumChunkForThisTask;
+				}
+
+				check(Count == ChunksNumber);
+			}
 		}
 
+		const int32 NumFutures = FMath::Max(0, AsyncState.TasksChunksStartEnd.Num() - 1);
+
 		// Synchronisation structure to be shared between async tasks and collapsing main thread.
 		struct FSynchroStruct 
 		{
-			// Atomic counter to know which chunk a given task should process.
-			std::atomic<int32> CurrentChunkToProcess = 0;
-
 			// Atomic to indicate tasks to stop processing new chunks.
 			std::atomic<bool> bQuit = false;
 
@@ -143,22 +167,21 @@ namespace FPCGAsync
 		};
 
 		// We won't stop if we are not time slicing.
-		auto ShouldStop = [AsyncState, bEnableTimeSlicing]() -> bool
+		auto ShouldStop = [&AsyncState, bEnableTimeSlicing]() -> bool
 		{
 			return bEnableTimeSlicing && AsyncState.ShouldStop();
 		};
 
-		// Main thread will either work if there is no future, or collapse arrays if there are some.
+		// Main thread will either work if there is no future, or work and collapse arrays if there are some.
 		if (NumFutures == 0)
 		{
 			// Main thread is working
-			int32 CurrentChunkToProcess = 0;
 			check(ChunksNumber > 0);
 
 			// Do at least one run, as we could fall into an infinite loop if we always have to stop before doing anything.
 			do
 			{
-				const int32 StartReadIndex = AsyncState.AsyncCurrentReadIndex;
+				const int32 StartReadIndex = AsyncState.CurrentChunkToCollapse * ChunkSize;
 				const int32 StartWriteIndex = AsyncState.AsyncCurrentWriteIndex;
 				const int32 Count = FMath::Min(ChunkSize, NumIterations - StartReadIndex);
 
@@ -173,9 +196,7 @@ namespace FPCGAsync
 				}
 
 				AsyncState.AsyncCurrentWriteIndex += NumElementsWritten;
-				AsyncState.AsyncCurrentReadIndex += ChunkSize;
-				++CurrentChunkToProcess;
-			} while (CurrentChunkToProcess < ChunksNumber && !ShouldStop());
+			} while (++AsyncState.CurrentChunkToCollapse < ChunksNumber && !ShouldStop());
 		}
 		else
 		{
@@ -187,23 +208,15 @@ namespace FPCGAsync
 			TArray<UE::Tasks::TTask<void>> AsyncTasks;
 			AsyncTasks.Reserve(NumFutures);
 
-			auto JobTask = [&IterationInnerLoop, &SynchroStruct, ChunksNumber, StartIndex, ChunkSize, NumIterations]()
+			// A job task will pick up the chunk to process, do the work on it, and push in the queue info that it processed the chunk and how many points it has written.
+			auto JobTask = [&IterationInnerLoop, &SynchroStruct, ChunkSize, NumIterations](const int32 CurrentChunkToProcess)
 			{
-				const int32 CurrentChunkToProcess = SynchroStruct.CurrentChunkToProcess++;
-				if (CurrentChunkToProcess >= ChunksNumber)
-				{
-					// If we reached the end, notify eveyone it's done
-					SynchroStruct.bQuit = true;
-					return;
-				}
-
-				const int32 StartReadIndex = StartIndex + CurrentChunkToProcess * ChunkSize;
+				const int32 StartReadIndex = CurrentChunkToProcess * ChunkSize;
 				// We write in the same "memory space" with the async task (in preallocated output array) and will be collapsed by the main thread.
 				const int32 StartWriteIndex = StartReadIndex;
 				const int32 Count = FMath::Min(ChunkSize, NumIterations - StartReadIndex);
 
 				check(Count > 0);
-
 				int32 NumElementsWritten;
 
 				{
@@ -215,88 +228,117 @@ namespace FPCGAsync
 				SynchroStruct.ChunkProcessedIndexAndNumElementsWrittenQueue.Enqueue(CurrentChunkToProcess, NumElementsWritten);
 			};
 
+			// Starting all the tasks, only if the associated task has still some work to do.
 			for (int32 TaskIndex = 0; TaskIndex < NumFutures; ++TaskIndex)
 			{
 				TRACE_CPUPROFILER_EVENT_SCOPE(FPCGAsync::Private::AsyncProcessing::StartingTasks);
-				AsyncTasks.Emplace(UE::Tasks::Launch(UE_SOURCE_LOCATION, [&JobTask, &SynchroStruct]() -> void
+
+				if (AsyncState.TasksChunksStartEnd[TaskIndex].Get<0>() == AsyncState.TasksChunksStartEnd[TaskIndex].Get<1>())
+				{
+					// No work left to do.
+					continue;
+				}
+
+				AsyncTasks.Emplace(UE::Tasks::Launch(UE_SOURCE_LOCATION, [&JobTask, &SynchroStruct, &StartEnd = AsyncState.TasksChunksStartEnd[TaskIndex]]() -> void
 				{
 					// Do at least one run, as we could fall into an infinite loop if we always have to stop before doing anything.
+					// Quit if we have no more work to do or we were told to stop.
 					do
 					{
-						JobTask();
-					} while (!SynchroStruct.bQuit);
+						JobTask(StartEnd.Get<0>()++);
+					} while (!SynchroStruct.bQuit && StartEnd.Get<0>() != StartEnd.Get<1>());
 				}));
 			}
 
+			auto FlushQueue = [&SynchroStruct, &AsyncState]()
+			{
+				TPair<int32, int32> QueueItem;
+				while (SynchroStruct.ChunkProcessedIndexAndNumElementsWrittenQueue.Dequeue(QueueItem))
+				{
+					AsyncState.ChunkToNumElementsWrittenMap.Add(QueueItem);
+				}
+			};
+
 			// Collapsing needs to be done in order, so we have a map between chunk index and the number of elements written for this chunk.
 			// Note that we need to collpase because points can be discarded, so we can have less points "kept" than the chunk size.
-			TMap<int32, int32> ChunkToNumElementsWrittenMap;
-			bool bLocalStop = false;
-			int32 CurrentChunkToCollapse = 0;
 			while (true)
 			{
-				// Either we should stop because the time has elapsed or all workloads have been dispatched
-				if (!bLocalStop && (ShouldStop() || SynchroStruct.bQuit))
+				// Either we should stop because the time has elapsed or all workloads have been processed
+				if (ShouldStop() || SynchroStruct.bQuit)
 				{
 					SynchroStruct.bQuit = true;
-					bLocalStop = true;
 
 					// Wait for all futures to finish their job
 					UE::Tasks::Wait(AsyncTasks);
+
+					// Make sure to dequeue everything
+					FlushQueue();
+
+					// Exit
+					break;
 				}
 
-				TPair<int32, int32> QueueItem;
-				// Try to unqueue a processed item.
-				if (!SynchroStruct.ChunkProcessedIndexAndNumElementsWrittenQueue.Dequeue(QueueItem) && !bLocalStop)
+				// Try to flush queue before trying anything.
+				FlushQueue();
+				
+				// If we have work to do, do it.
+				if (AsyncState.TasksChunksStartEnd.Last().Get<0>() != AsyncState.TasksChunksStartEnd.Last().Get<1>())
 				{
-					// If there is nothing to unqueue (nothing to collapse), we take this wait opportunity to
-					// process a chunk on the main thread.
-					JobTask();
+					JobTask(AsyncState.TasksChunksStartEnd.Last().Get<0>()++);
 					continue;
 				}
-				else
-				{
-					// If we successfully unqueue, it means we can try to collapse.
-					TRACE_CPUPROFILER_EVENT_SCOPE(FPCGAsync::AsyncProcessing::CollapsingNewData);
 
-					ChunkToNumElementsWrittenMap.Add(QueueItem);
+				// At that point, we are on collapse duty, don't stop until we finish the collapse or we have to stop.
+				// Don't check stop too often.
+				constexpr int32 NumCheckStop = 20;
+				int32 CheckStop = 0;
+				while (AsyncState.CurrentChunkToCollapse != ChunksNumber)
+				{
+					if (++CheckStop == NumCheckStop)
+					{
+						CheckStop = 0;
+						if (ShouldStop())
+						{
+							SynchroStruct.bQuit = true;
+							break;
+						}
+					}
 
-					// We always collapse in order, so try to collapse chunk in order while we have some in the map.
+					// Try to remove the current chunk we need to process. If it is not in the map it means it has not yet been processed (or dequeued)
+					// so flush the queue and try again.
 					int32 NumberOfElementsWritten = 0;
-					
-					while (ChunkToNumElementsWrittenMap.RemoveAndCopyValue(CurrentChunkToCollapse, NumberOfElementsWritten))
+					if (!AsyncState.ChunkToNumElementsWrittenMap.RemoveAndCopyValue(AsyncState.CurrentChunkToCollapse, NumberOfElementsWritten))
 					{
-						const int32 ReadStartIndexForCurrentChunkToCollapse = AsyncState.AsyncCurrentReadIndex;
-						check(NumberOfElementsWritten <= ChunkSize);
+						FlushQueue();
+						continue;
+					}
 
-						// If ReadStartIndexForCurrentChunkToCollapse == AsyncState.AsyncCurrentWriteIndex, no need to copy, all elements
-						// are already at the right place
-						if (ReadStartIndexForCurrentChunkToCollapse == AsyncState.AsyncCurrentWriteIndex)
-						{
-							AsyncState.AsyncCurrentWriteIndex += NumberOfElementsWritten;
-						}
-						else
-						{
-							// Otherwise collapse
-							MoveDataRange(ReadStartIndexForCurrentChunkToCollapse, AsyncState.AsyncCurrentWriteIndex, NumberOfElementsWritten);
-							AsyncState.AsyncCurrentWriteIndex += NumberOfElementsWritten;
-						}
+					TRACE_CPUPROFILER_EVENT_SCOPE(FPCGAsync::AsyncProcessing::CollapsingNewData);
+					const int32 ReadStartIndexForCurrentChunkToCollapse = AsyncState.CurrentChunkToCollapse * ChunkSize;
+					check(NumberOfElementsWritten <= ChunkSize);
 
-						AsyncState.AsyncCurrentReadIndex += ChunkSize;
-						++CurrentChunkToCollapse;
+					// If ReadStartIndexForCurrentChunkToCollapse == AsyncState.AsyncCurrentWriteIndex, no need to copy, all elements
+					// are already at the right place
+					if (ReadStartIndexForCurrentChunkToCollapse == AsyncState.AsyncCurrentWriteIndex)
+					{
+						AsyncState.AsyncCurrentWriteIndex += NumberOfElementsWritten;
+					}
+					else
+					{
+						// Otherwise collapse
+						MoveDataRange(ReadStartIndexForCurrentChunkToCollapse, AsyncState.AsyncCurrentWriteIndex, NumberOfElementsWritten);
+						AsyncState.AsyncCurrentWriteIndex += NumberOfElementsWritten;
 					}
-				}
 
-				// If we locally stoped (meaning we already wait for async tasks to finish) and we collapse all chunks processed
-				// we can exit the loop
-				if (bLocalStop && CurrentChunkToCollapse == FMath::Min((int32)SynchroStruct.CurrentChunkToProcess, ChunksNumber))
-				{
-					break;
+					++AsyncState.CurrentChunkToCollapse;
 				}
+
+				// If we reach that point, it means we finished the collapse or we ran out of time. So indicate to stop.
+				SynchroStruct.bQuit = true;
 			}
 		}
 
-		const bool bIsDone = (AsyncState.AsyncCurrentReadIndex >= NumIterations);
+		const bool bIsDone = AsyncState.CurrentChunkToCollapse == ChunksNumber;
 
 		check(bIsDone || bEnableTimeSlicing);
 
@@ -305,10 +347,7 @@ namespace FPCGAsync
 			TRACE_CPUPROFILER_EVENT_SCOPE(FPCGAsync::Private::AsyncProcessing::Finished);
 
 			Finished(AsyncState.AsyncCurrentWriteIndex);
-
-			// Reset indexes
-			AsyncState.AsyncCurrentReadIndex = 0;
-			AsyncState.AsyncCurrentWriteIndex = 0;
+			AsyncState.Reset();
 		}
 
 		return bIsDone;
@@ -357,6 +396,52 @@ namespace FPCGAsync
 		}
 	}
 
+	/**
+	* A Helper for generic parallel loops, with support for timeslicing, specialized to work on ranges. This version only uses indexes allowing more flexible usage
+	* to process many arrays in parallel or use it for other batch updating.
+	* 
+	* This version also does not support removing elements and is expects a 1:1 mapping. This just removes some lambdas required to support
+	* that use case.
+	* 
+	* Work will be separated in chunks, that will be processed in parallel. Main thread will then collapse incoming data from async tasks.
+	* Will use AsyncState.ShouldStop() to stop execution if timeslicing is enabled.
+	* Important info: 
+	*   - We will finish to process and collapse data for all data already in process, even if we need to stop. To mitigate this, try to use small chunk sizes.
+	*   - To avoid infinite loops (when we should stop even before starting working), we will at least process 1 chunk of data per thread.
+	*   - To have async tasks, you need to have at least 3 available threads (main thread + 2 futures). Otherwise, we will only process on the main thread, without collapse.
+	* 
+	* @param AsyncState - The context containing the information about how many tasks we can launch, async read/write index for the current job and a function to know if we need to stop processing.
+	* @param NumIterations - The number of calls that will be done to the provided function, also an upper bound on the number of data generated.
+	* @param Initialize - Signature: void(). A function that will be called once on the first timeslice, where you can reserve data for processing
+	* @param Finished - Signature: void(int32 Count). Called once on finished, and tells you the total count of points written.
+	* @param bEnableTimeSlicing - If false, we will not stop until all the processing is done.
+	* @param ChunkSize - Size of the chunks to cut the input data with
+	* @returns true if the processing is done, false otherwise. Use this to know if you need to reschedule the task.
+	*/
+	template <typename InitializeFunc, typename ProcessRangeFunc>
+	bool AsyncProcessingOneToOneRangeEx(FPCGAsyncState* AsyncState, int32 NumIterations, InitializeFunc&& Initialize, ProcessRangeFunc&& ProcessRange, const bool bEnableTimeSlicing, const int32 ChunkSize = 64)
+	{
+		auto MoveDataRange = [](int32, int32, int32) { ensure(false); };
+		auto Finished = [NumIterations](int32 Count) { ensure(NumIterations == Count); };
+
+		if (AsyncState && !AsyncState->bIsRunningAsyncCall)
+		{
+			AsyncState->bIsRunningAsyncCall = true;
+			const bool bIsDone = Private::AsyncProcessing(*AsyncState, NumIterations, Initialize, ProcessRange, MoveDataRange, Finished, bEnableTimeSlicing, ChunkSize);
+			AsyncState->bIsRunningAsyncCall = false;
+			return bIsDone;
+		}
+		else
+		{
+			// Can't use time slicing without an async state or while running in another async call (it will mess up with async indexes). 
+			// We also force using one thread (the current one).
+			FPCGAsyncState DummyState;
+			DummyState.NumAvailableTasks = 1;
+			DummyState.bIsRunningAsyncCall = true;
+			return Private::AsyncProcessing(DummyState, NumIterations, Initialize, ProcessRange, MoveDataRange, Finished, /*bEnableTimeSlicing=*/false, ChunkSize);
+		}
+	}
+
 	/**
 	* A Helper for generic parallel loops, with support for timeslicing. This version only uses indexes allowing more flexible usage
 	* to process many arrays in parallel or use it for other batch updating.
diff --git a/Engine/Plugins/PCG/Source/PCG/Public/Helpers/PCGAsyncState.h b/Engine/Plugins/PCG/Source/PCG/Public/Helpers/PCGAsyncState.h
--- a/Engine/Plugins/PCG/Source/PCG/Public/Helpers/PCGAsyncState.h
+++ b/Engine/Plugins/PCG/Source/PCG/Public/Helpers/PCGAsyncState.h
@@ -2,29 +2,60 @@
 
 #pragma once
 
-#include "HAL/Platform.h"
+#include "Containers/Array.h"
+#include "Containers/Map.h"
+#include "Templates/Tuple.h"
 
 /**
 * Helper class to gather information about asynchronus execution. Will be held in the PCGContext.
 */
 struct FPCGAsyncState
 {
+	// Deprecating a property requires to have deprecation warning disable for the default copy ctor...
+	PRAGMA_DISABLE_DEPRECATION_WARNINGS
+	FPCGAsyncState() = default;
+	~FPCGAsyncState() = default;
+
+	FPCGAsyncState(const FPCGAsyncState&) = default;
+	FPCGAsyncState(FPCGAsyncState&&) = default;
+
+	FPCGAsyncState& operator=(const FPCGAsyncState&) = default;
+	FPCGAsyncState& operator=(FPCGAsyncState&&) = default;
+	PRAGMA_ENABLE_DEPRECATION_WARNINGS
+
 	/** For timeslicing, it will be set by the graph executor to know when we should stop. */
 	double EndTime = 0.0;
 
 	/** How many tasks are available to run async. */
 	int32 NumAvailableTasks = 0;
 
-	/** For async process, keep track of where the processing is, in the input data(read) and output data(write). */
+	UE_DEPRECATED(5.5, "Not used anymore")
 	int32 AsyncCurrentReadIndex = 0;
+
+	/** For async process, keep track of where the write processing is in output data. */
 	int32 AsyncCurrentWriteIndex = 0;
 
-	/** For multithreading, track if the current element is run on the main thread.*/
+	/** List of the current chunk and last chunk a given task have to process. There is one entry per future, and the last one is for the main thread. */
+	TArray<TTuple<int32, int32>> TasksChunksStartEnd;
+
+	/** Keep information if the async has started or not. */
+	bool bStarted = false;
+
+	/** Keep track of the current chunk we need to collapse. Important to collapse in order. */
+	int32 CurrentChunkToCollapse = 0;
+
+	/** Map between the processed chunks but not yet processed and the number of elements written for that chunk. */
+	TMap<int32, int32> ChunkToNumElementsWrittenMap;
+
+	/** For multithreading, track if the current element is run on the main thread. */
 	bool bIsRunningOnMainThread = true;
 
-	/** True if currently inside a PCGAsync scope - will prevent further async processing */
+	/** True if currently inside a PCGAsync scope - will prevent further async processing. */
 	bool bIsRunningAsyncCall = false;
 
 	/** Returns true if we reached end time. */
 	PCG_API bool ShouldStop() const;
+
+	/** To be called at the end of an async processing. */
+	PCG_API void Reset();
 };
\ No newline at end of file
