[PCG] Build fix for warning around cvars
#rnx

diff --git a/Engine/Plugins/PCG/Source/PCG/Private/Elements/Metadata/PCGMetadataOpElementBase.cpp b/Engine/Plugins/PCG/Source/PCG/Private/Elements/Metadata/PCGMetadataOpElementBase.cpp
--- a/Engine/Plugins/PCG/Source/PCG/Private/Elements/Metadata/PCGMetadataOpElementBase.cpp
+++ b/Engine/Plugins/PCG/Source/PCG/Private/Elements/Metadata/PCGMetadataOpElementBase.cpp
@@ -14,6 +14,19 @@
 
 #define LOCTEXT_NAMESPACE "PCGMetadataElementBaseElement"
 
+namespace PCGMetadataBase
+{
+	TAutoConsoleVariable<bool> CVarMetadataOperationInMT(
+		TEXT("pcg.MetadataOperationInMT"),
+		true,
+		TEXT("Metadata operations are now multithreaded."));
+
+	TAutoConsoleVariable<int> CVarMetadataOperationChunkSize(
+		TEXT("pcg.MetadataOperationChunkSize"),
+		256,
+		TEXT("Metadata operations chunk size."));
+}
+
 void UPCGMetadataSettingsBase::PostLoad()
 {
 	Super::PostLoad();
diff --git a/Engine/Plugins/PCG/Source/PCG/Public/Elements/Metadata/PCGMetadataOpElementBase.h b/Engine/Plugins/PCG/Source/PCG/Public/Elements/Metadata/PCGMetadataOpElementBase.h
--- a/Engine/Plugins/PCG/Source/PCG/Public/Elements/Metadata/PCGMetadataOpElementBase.h
+++ b/Engine/Plugins/PCG/Source/PCG/Public/Elements/Metadata/PCGMetadataOpElementBase.h
@@ -18,15 +18,11 @@
 class FPCGMetadataAttributeBase;
 
 // FIXME: To be removed when we are confident Metadata is stable in MT.
-static inline TAutoConsoleVariable<bool> CVarMetadataOperationInMT(
-	TEXT("pcg.MetadataOperationInMT"),
-	true,
-	TEXT("Metadata operations are now multithreaded."));
-
-static inline TAutoConsoleVariable<int> CVarMetadataOperationChunkSize(
-	TEXT("pcg.MetadataOperationChunkSize"),
-	256,
-	TEXT("Metadata operations chunk size."));
+namespace PCGMetadataBase
+{
+	extern TAutoConsoleVariable<bool> CVarMetadataOperationInMT;
+	extern TAutoConsoleVariable<int> CVarMetadataOperationChunkSize;
+}
 
 namespace PCGMetadataSettingsBaseConstants
 {
@@ -450,9 +446,9 @@ inline bool FPCGMetadataElementBase::DoNAryOp(PCGMetadataOps::FOperationData& In
 	Options.SetFlags = Flags;
 	Options.bUseDefaultKey = false;
 
-	const int32 ChunkSize = CVarMetadataOperationChunkSize.GetValueOnAnyThread();
+	const int32 ChunkSize = PCGMetadataBase::CVarMetadataOperationChunkSize.GetValueOnAnyThread();
 
-	if (CVarMetadataOperationInMT.GetValueOnAnyThread())
+	if (PCGMetadataBase::CVarMetadataOperationInMT.GetValueOnAnyThread())
 	{
 		return FPCGAsync::AsyncProcessingOneToOneRangeEx(&InOperationData.Context->AsyncState, InOperationData.NumberOfElementsToProcess, []() {},
 			[&InOperationData, &InCallbacks, &Options](int32 StartReadIndex, int32 StartWriteIndex, int32 Count)
