[PCG] Deprecate the labels of param overrides from localized names to authored names

* Settings with localized property names where breaking nodes that were overriding those properties.
* Now use the AuthoredName for all labels
* Add deprecation to reconnect old labels to new labels.

diff --git a/Engine/Plugins/PCG/Source/PCG/Private/Helpers/PCGSettingsHelpers.cpp b/Engine/Plugins/PCG/Source/PCG/Private/Helpers/PCGSettingsHelpers.cpp
--- a/Engine/Plugins/PCG/Source/PCG/Private/Helpers/PCGSettingsHelpers.cpp
+++ b/Engine/Plugins/PCG/Source/PCG/Private/Helpers/PCGSettingsHelpers.cpp
@@ -407,7 +407,7 @@ namespace PCGSettingsHelpers
 					{
 						// If we have a clash, we will use the full path, so mark this param and the other that clashed to use the full path.
 						Res[CachedLabelIndex].bHasNameClash = true;
-						Res[CachedLabelIndex].Label = FName(Res[CachedLabelIndex].GetDisplayPropertyPath());
+						Res[CachedLabelIndex].Label = FName(Res[CachedLabelIndex].GetPropertyPath());
 						bHasNameClash = true;
 					}
 
@@ -427,7 +427,7 @@ namespace PCGSettingsHelpers
 					else
 					{
 #if WITH_EDITOR
-						Param.Label = FName(Param.GetDisplayPropertyPath());
+						Param.Label = FName(Param.GetPropertyPath());
 #endif // WITH_EDITOR
 						Param.bHasNameClash = true;
 					}
@@ -462,14 +462,13 @@ namespace PCGSettingsHelpers
 				FName Label = NAME_None;
 				bool bHasNameClash = false;
 #if WITH_EDITOR
-				// GetDisplayNameText is not available in non-editor build.
-				Label = *Property->GetDisplayNameText().ToString();
+				Label = *Property->GetAuthoredName();
 				const int32 CachedLabelIndex = LabelCache.IndexOfByKey(Label);
 				if (CachedLabelIndex != INDEX_NONE)
 				{
 					// If we have a clash, we will use the full path, so mark this param and the other that clashed to use the full path.
 					Res[CachedLabelIndex].bHasNameClash = true;
-					Res[CachedLabelIndex].Label = FName(Res[CachedLabelIndex].GetDisplayPropertyPath());
+					Res[CachedLabelIndex].Label = FName(Res[CachedLabelIndex].GetPropertyPath());
 					bHasNameClash = true;
 				}
 
@@ -482,7 +481,7 @@ namespace PCGSettingsHelpers
 				Param.PropertyClass = InClass;
 				Param.bHasNameClash = bHasNameClash;
 #if WITH_EDITOR
-				Param.Label = bHasNameClash ? FName(Param.GetDisplayPropertyPath()) : Label;
+				Param.Label = bHasNameClash ? FName(Param.GetPropertyPath()) : Label;
 				GatherAliases(Property, Param);
 #endif // WITH_EDITOR
 			}
diff --git a/Engine/Plugins/PCG/Source/PCG/Private/PCGSettings.cpp b/Engine/Plugins/PCG/Source/PCG/Private/PCGSettings.cpp
--- a/Engine/Plugins/PCG/Source/PCG/Private/PCGSettings.cpp
+++ b/Engine/Plugins/PCG/Source/PCG/Private/PCGSettings.cpp
@@ -28,6 +28,16 @@
 
 #define LOCTEXT_NAMESPACE "PCGSettings"
 
+namespace PCGSettings
+{
+	// For the deprecation of overridable params as display name to authored names
+	
+	// Arbitrary invalid index
+	static constexpr int DeprecationAliasIndex = -2;
+	// For property path concatenation
+	static constexpr const TCHAR* PropertyPathSeparator = TEXT("/");
+}
+
 /** Custom Crc computation that ignores properties that will not affect the computed result of a node. */
 class FPCGSettingsObjectCrc32 : public FArchiveObjectCrc32
 {
@@ -64,7 +74,7 @@ public:
 
 FString FPCGSettingsOverridableParam::GetPropertyPath() const
 {
-	return FString::JoinBy(Properties, TEXT("/"), [](const FProperty* InProperty) { return InProperty ? InProperty->GetName() : FString(); });
+	return FString::JoinBy(Properties, PCGSettings::PropertyPathSeparator, [](const FProperty* InProperty) { return InProperty ? InProperty->GetAuthoredName() : FString(); });
 }
 
 TArray<FName> FPCGSettingsOverridableParam::GenerateAllPossibleAliases() const
@@ -138,7 +148,21 @@ TArray<FName> FPCGSettingsOverridableParam::GenerateAllPossibleAliases() const
 #if WITH_EDITOR
 FString FPCGSettingsOverridableParam::GetDisplayPropertyPath() const
 {
-	return FString::JoinBy(Properties, TEXT("/"), [](const FProperty* InProperty) { return InProperty ? InProperty->GetDisplayNameText().ToString() : FString(); });
+	return GetDisplayPropertyPathText().ToString();
+}
+
+FText FPCGSettingsOverridableParam::GetDisplayPropertyPathText() const
+{
+	if (bHasNameClash)
+	{
+		TArray<FText> PropertyNames;
+		Algo::Transform(Properties, PropertyNames, [](const FProperty* InProperty) { return InProperty ? InProperty->GetDisplayNameText() : FText(); });
+		return FText::Join(FText::FromString(PCGSettings::PropertyPathSeparator), PropertyNames);
+	}
+	else
+	{
+		return !Properties.IsEmpty() && Properties.Last() ? Properties.Last()->GetDisplayNameText() : FText();
+	}
 }
 #endif // WITH_EDITOR
 
@@ -638,6 +662,19 @@ void UPCGSettings::ApplyDeprecationBeforeUpdatePins(UPCGNode* InOutNode, TArray<
 	{
 		InOutNode->RenameInputPin(PCGPinConstants::Private::OldDefaultParamsLabel, PCGPinConstants::DefaultParamsLabel);
 	}
+
+	// FIXME: Should be behind a version number
+	// Deprecation, renaming all the pins generated with the old labels to the new labels
+	for (const FPCGSettingsOverridableParam& Param : CachedOverridableParams)
+	{
+		if (const FPCGPropertyAliases* OldLabels = Param.MapOfAliases.Find(PCGSettings::DeprecationAliasIndex))
+		{
+			for (const FName& OldLabel : OldLabels->Aliases)
+			{
+				InOutNode->RenameInputPin(OldLabel, Param.Label);
+			}
+		}
+	}
 }
 
 void UPCGSettings::ApplyStructuralDeprecation(UPCGNode* InOutNode)
@@ -826,7 +863,47 @@ void UPCGSettings::InitializeCachedOverridableParams(bool bReset)
 #if WITH_EDITOR
 	if (CachedOverridableParams.IsEmpty() || bReset)
 	{
+		// Deprecation: Previously params labels were using the GetDisplayTextName, which is localized and will break existing pins.
+		// Now it will be the authored name, but we still need to make sure the pins are not broken. 
+		// For that matter, we will match names for properties between serialized and newly found ones. If their label don't match, we add the previous label as alias.
+		// FIXME: Should be depending on a version, but we can't do that as an hotfix.
+		TMap<FString, FName> PropertyPathToLabel;
+
+		if (!CachedOverridableParams.IsEmpty())
+		{
+			for (const FPCGSettingsOverridableParam& Param : CachedOverridableParams)
+			{
+				if (Param.PropertiesNames.IsEmpty())
+				{
+					continue;
+				}
+
+				FString Path = FString::JoinBy(Param.PropertiesNames, PCGSettings::PropertyPathSeparator, [](const FName PropertyName) { return PropertyName.ToString();});
+				PropertyPathToLabel.Emplace(Path, Param.Label);
+			}
+		}
+
 		CachedOverridableParams = GatherOverridableParams();
+
+		if (!PropertyPathToLabel.IsEmpty())
+		{
+			for (FPCGSettingsOverridableParam& Param : CachedOverridableParams)
+			{
+				if (Param.PropertiesNames.IsEmpty())
+				{
+					continue;
+				}
+
+				FString Path = FString::JoinBy(Param.PropertiesNames, PCGSettings::PropertyPathSeparator, [](const FName PropertyName) { return PropertyName.ToString();});
+				if (FName* OldLabel = PropertyPathToLabel.Find(Path))
+				{
+					if (*OldLabel != Param.Label)
+					{
+						Param.MapOfAliases.FindOrAdd(PCGSettings::DeprecationAliasIndex).Aliases.Emplace(*OldLabel);
+					}
+				}
+			}
+		}
 	}
 #endif // WITH_EDITOR
 
diff --git a/Engine/Plugins/PCG/Source/PCG/Public/PCGSettings.h b/Engine/Plugins/PCG/Source/PCG/Public/PCGSettings.h
--- a/Engine/Plugins/PCG/Source/PCG/Public/PCGSettings.h
+++ b/Engine/Plugins/PCG/Source/PCG/Public/PCGSettings.h
@@ -112,6 +112,7 @@ struct FPCGSettingsOverridableParam
 
 #if WITH_EDITOR
 	FString GetDisplayPropertyPath() const;
+	PCG_API FText GetDisplayPropertyPathText() const;
 #endif // WITH_EDITOR
 };
 
diff --git a/Engine/Plugins/PCG/Source/PCGEditor/Private/PCGEditorGraphNodeBase.cpp b/Engine/Plugins/PCG/Source/PCGEditor/Private/PCGEditorGraphNodeBase.cpp
--- a/Engine/Plugins/PCG/Source/PCGEditor/Private/PCGEditorGraphNodeBase.cpp
+++ b/Engine/Plugins/PCG/Source/PCGEditor/Private/PCGEditorGraphNodeBase.cpp
@@ -1188,7 +1188,26 @@ void UPCGEditorGraphNodeBase::CreatePins(const TArray<UPCGPin*>& InInputPins, co
 FText UPCGEditorGraphNodeBase::GetPinFriendlyName(const UPCGPin* InPin) const
 {
 	check(InPin);
-	return FText::FromString(FName::NameToDisplayString(InPin->Properties.Label.ToString(), false));
+
+	// For overridable params, use the display name of properties (for localized version or overridden display name in metadata).
+	if (InPin->Properties.IsAdvancedPin() && InPin->Properties.AllowedTypes == EPCGDataType::Param)
+	{
+		const UPCGSettings* Settings = InPin->Node ? InPin->Node->GetSettings() : nullptr;
+		if (Settings)
+		{
+			const FPCGSettingsOverridableParam* Param = Settings->OverridableParams().FindByPredicate([Label = InPin->Properties.Label](const FPCGSettingsOverridableParam& ParamToCheck)
+			{
+				return ParamToCheck.Label == Label;
+			});
+
+			if (Param)
+			{
+				return Param->GetDisplayPropertyPathText();
+			}
+		}
+	}
+
+	return FText::FromString(FName::NameToDisplayString(InPin->Properties.Label.ToString(), /*bIsBool=*/false));
 }
 
 #undef LOCTEXT_NAMESPACE
