[PCG] Added missing BP getter/setter on world to export on the level to asset exporter

#rnx

diff --git a/Engine/Plugins/PCG/Source/PCG/Public/PCGDataAsset.h b/Engine/Plugins/PCG/Source/PCG/Public/PCGDataAsset.h
--- a/Engine/Plugins/PCG/Source/PCG/Public/PCGDataAsset.h
+++ b/Engine/Plugins/PCG/Source/PCG/Public/PCGDataAsset.h
@@ -20,7 +20,7 @@ public:
 	// TODO: additional embedded graph?
 
 	/** Contains direct data owned by this data asset */
-	UPROPERTY(BlueprintReadOnly, VisibleAnywhere, Category = Data, meta = (NoResetToDefault))
+	UPROPERTY(BlueprintReadWrite, VisibleAnywhere, Category = Data, meta = (NoResetToDefault))
 	FPCGDataCollection Data;
 
 	/** Alternative name (instead of asset name), can be left empty */
diff --git a/Engine/Plugins/PCG/Source/PCGEditor/Private/PCGLevelToAsset.cpp b/Engine/Plugins/PCG/Source/PCGEditor/Private/PCGLevelToAsset.cpp
--- a/Engine/Plugins/PCG/Source/PCGEditor/Private/PCGLevelToAsset.cpp
+++ b/Engine/Plugins/PCG/Source/PCGEditor/Private/PCGLevelToAsset.cpp
@@ -420,3 +420,13 @@ bool UPCGLevelToAsset::BP_ExportWorld_Implementation(UWorld* World, const FStrin
 
 	return true;
 }
+
+void UPCGLevelToAsset::SetWorld(UWorld* World)
+{
+	WorldToExport = World;
+}
+
+UWorld* UPCGLevelToAsset::GetWorld() const
+{
+	return WorldToExport;
+}
\ No newline at end of file
diff --git a/Engine/Plugins/PCG/Source/PCGEditor/Public/PCGLevelToAsset.h b/Engine/Plugins/PCG/Source/PCGEditor/Public/PCGLevelToAsset.h
--- a/Engine/Plugins/PCG/Source/PCGEditor/Public/PCGLevelToAsset.h
+++ b/Engine/Plugins/PCG/Source/PCGEditor/Public/PCGLevelToAsset.h
@@ -29,6 +29,12 @@ public:
 	UFUNCTION(BlueprintNativeEvent, meta = (DisplayName = "Export World", ForceAsFunction))
 	bool BP_ExportWorld(UWorld* World, const FString& PackageName, UPCGDataAsset* Asset);
 
+	UFUNCTION(BlueprintCallable, Category="PCG|IO")
+	void SetWorld(UWorld* World);
+
+	UFUNCTION(BlueprintCallable, Category = "PCG|IO")
+	UWorld* GetWorld() const;
+
 protected:
 	//~Being UPCGAssetExporter interface
 	virtual bool ExportAsset(const FString& PackageName, UPCGDataAsset* Asset) override;
