[PCG] Remove overly aggressive check. Call on UPCGGraph CDO instead for null graphs.

- Still need to investigate what happens when UPCGComponent with no graph registers to the system to avoid doing unecessary work

#rnx

diff --git a/Engine/Plugins/PCG/Source/PCG/Private/PCGComponent.cpp b/Engine/Plugins/PCG/Source/PCG/Private/PCGComponent.cpp
--- a/Engine/Plugins/PCG/Source/PCG/Private/PCGComponent.cpp
+++ b/Engine/Plugins/PCG/Source/PCG/Private/PCGComponent.cpp
@@ -127,8 +127,12 @@ void UPCGComponent::SetIsPartitioned(bool bIsNowPartitioned)
 
 bool UPCGComponent::Use2DGrid() const
 {
-	check(GetGraph());
-	return GetGraph()->Use2DGrid();
+	if (UPCGGraph* PCGGraph = GetGraph())
+	{
+		return PCGGraph->Use2DGrid();
+	}
+
+	return GetDefault<UPCGGraph>()->Use2DGrid();
 }
 
 void UPCGComponent::SetGraph_Implementation(UPCGGraphInterface* InGraph)
