[PCG] PCGWorldActor override ActorTypeSupportsDataLayer/ActorTypeSupportsExternalDataLayer to return false (PCGWorldActor in a DL isn't supported)

#rnx

diff --git a/Engine/Plugins/PCG/Source/PCG/Public/PCGWorldActor.h b/Engine/Plugins/PCG/Source/PCG/Public/PCGWorldActor.h
--- a/Engine/Plugins/PCG/Source/PCG/Public/PCGWorldActor.h
+++ b/Engine/Plugins/PCG/Source/PCG/Public/PCGWorldActor.h
@@ -26,6 +26,8 @@ public:
 	virtual void BeginPlay() override;
 	virtual void EndPlay(const EEndPlayReason::Type EndPlayReason) override;
 #if WITH_EDITOR
+	virtual bool ActorTypeSupportsDataLayer() const override { return false; }
+	virtual bool ActorTypeSupportsExternalDataLayer() const override { return false; }
 	virtual bool CanChangeIsSpatiallyLoadedFlag() const override { return false; }
 	virtual bool IsUserManaged() const override { return false; }
 	virtual bool ShouldExport() override { return false; }
