[PCG] Added approximate data size in MB in the profiling window

diff --git a/Engine/Plugins/PCG/Source/PCG/Private/Utils/PCGExtraCapture.cpp b/Engine/Plugins/PCG/Source/PCG/Private/Utils/PCGExtraCapture.cpp
--- a/Engine/Plugins/PCG/Source/PCG/Private/Utils/PCGExtraCapture.cpp
+++ b/Engine/Plugins/PCG/Source/PCG/Private/Utils/PCGExtraCapture.cpp
@@ -5,6 +5,7 @@
 #include "PCGComponent.h"
 #include "PCGContext.h"
 #include "PCGGraph.h"
+#include "PCGInputOutputSettings.h"
 #include "PCGSubgraph.h"
 #include "Graph/PCGGraphCompiler.h"
 #include "Graph/PCGGraphExecutor.h"
@@ -58,6 +59,44 @@ void PCGUtils::FExtraCapture::Update(const PCGUtils::FScopedCall& InScopedCall)
 		break;
 	case EPCGExecutionPhase::PostExecute:
 		Timer.PostExecuteTime = ThisFrameTime;
+
+		if (InScopedCall.Context->GetInputSettings<UPCGGraphInputOutputSettings>())
+		{
+			// Ignore input/output nodes as they don't create any data and could count towards memory costs this way (not in totals, but locally).
+		}
+		else
+		{
+			ensure(Timer.OutputDataMemorySize == 0);
+			// Count memory here
+			// FIXME / TODO: This is sensitive to representation of data. If/when we change data representation to separate data from metadata, this might not report memory accurately
+			auto VisitAllData = [](const FPCGDataCollection& InCollection, TSet<const UPCGData*>& OutData)
+			{
+				for (const FPCGTaggedData& Data : InCollection.TaggedData)
+				{
+					if (Data.Data)
+					{
+						Data.Data->VisitDataNetwork([&OutData](const UPCGData* VisitedData) { OutData.Add(VisitedData); });
+					}
+				}
+			};
+
+			TSet<const UPCGData*> InputData;
+			TSet<const UPCGData*> OutputData;
+			VisitAllData(InScopedCall.Context->InputData, InputData);
+			VisitAllData(InScopedCall.Context->OutputData, OutputData);
+
+			TSet<const UPCGData*> NewData = OutputData.Difference(InputData);
+
+			Timer.OutputDataMemorySize = 0;
+			for (const UPCGData* Data : NewData)
+			{
+				FResourceSizeEx ResSize = FResourceSizeEx(EResourceSizeMode::Exclusive);
+				const_cast<UPCGData*>(Data)->GetResourceSizeEx(ResSize);
+				const SIZE_T DataSize = ResSize.GetDedicatedSystemMemoryBytes();
+				Timer.OutputDataMemorySize += DataSize;
+			}
+		}
+
 		break;
 	}
 
@@ -80,6 +119,11 @@ namespace PCGUtils
 			Info.Name = FName::NameToDisplayString((*Subgraph)->GetName(), /*bIsBool=*/false);
 		}
 
+		if (!Info.Children.IsEmpty())
+		{
+			Info.CallTime.OutputDataMemorySize = 0;
+		}
+
 		for (FCallTreeInfo& Child : Info.Children)
 		{
 			BuildTreeInfo(Child, SubgraphNodeToGraphMap);
@@ -95,6 +139,8 @@ namespace PCGUtils
 
 			Info.CallTime.MinExecutionFrameTime = FMath::Min(Info.CallTime.MinExecutionFrameTime, Child.CallTime.MinExecutionFrameTime);
 			Info.CallTime.MaxExecutionFrameTime = FMath::Max(Info.CallTime.MaxExecutionFrameTime, Child.CallTime.MaxExecutionFrameTime);
+
+			Info.CallTime.OutputDataMemorySize += Child.CallTime.OutputDataMemorySize;
 		}
 	}
 }
diff --git a/Engine/Plugins/PCG/Source/PCG/Public/Utils/PCGExtraCapture.h b/Engine/Plugins/PCG/Source/PCG/Public/Utils/PCGExtraCapture.h
--- a/Engine/Plugins/PCG/Source/PCG/Public/Utils/PCGExtraCapture.h
+++ b/Engine/Plugins/PCG/Source/PCG/Public/Utils/PCGExtraCapture.h
@@ -30,6 +30,7 @@ namespace PCGUtils
 
 		int32 ExecutionFrameCount = 0;
 		int32 PrepareDataFrameCount = 0;
+		uint64 OutputDataMemorySize = 0;
 
 		double PrepareDataWallTime() const { return PrepareDataEndTime - PrepareDataStartTime; }
 		double ExecutionWallTime() const { return ExecutionEndTime - ExecutionStartTime; }
diff --git a/Engine/Plugins/PCG/Source/PCGEditor/Private/SPCGEditorGraphProfilingView.cpp b/Engine/Plugins/PCG/Source/PCGEditor/Private/SPCGEditorGraphProfilingView.cpp
--- a/Engine/Plugins/PCG/Source/PCGEditor/Private/SPCGEditorGraphProfilingView.cpp
+++ b/Engine/Plugins/PCG/Source/PCGEditor/Private/SPCGEditorGraphProfilingView.cpp
@@ -36,6 +36,7 @@ namespace PCGEditorGraphProfilingView
 	const FName NAME_NbExecutionFrames = FName(TEXT("NbExecutionFrames"));
 	const FName NAME_TotalTime = FName(TEXT("TotalTime"));
 	const FName NAME_TotalWallTime = FName(TEXT("Total_WallTime"));
+	const FName NAME_MemoryInMB = FName(TEXT("Data size (MB)"));
 
 	/** Labels of the columns */
 	const FText TEXT_NodeLabel = LOCTEXT("NodeLabel", "Node");
@@ -49,6 +50,7 @@ namespace PCGEditorGraphProfilingView
 	const FText TEXT_NbExecutionFramesLabel = LOCTEXT("NbExecutionFramesLabel", "Exec Frames");
 	const FText TEXT_TotalTimeLabel = LOCTEXT("TotalTimeLabel", "Total");
 	const FText TEXT_TotalWallTimeLabel = LOCTEXT("TotalWallTimeLabel", "Total WallTime");
+	const FText TEXT_MemoryInMBLabel = LOCTEXT("MemoryInMBLabel", "Data size (MB)");
 
 	/** Tooltips */
 	const FText TEXT_PrepareDataTimeTooltip = LOCTEXT("PrepareDataTimeTooltip", "Cost of the PrepareData execution phase which some nodes use to process the incoming data, in ms.");
@@ -61,6 +63,7 @@ namespace PCGEditorGraphProfilingView
 	const FText TEXT_NbExecutionFramesTooltip = LOCTEXT("NbExecutionFramesTooltip", "The number of frames in which one or more execution phases were executed.");
 	const FText TEXT_TotalTimeTooltip = LOCTEXT("TotalTimeTooltip", "The total time spent in this node, summed over all execution and prepare data frames, in ms.");
 	const FText TEXT_TotalWallTimeTooltip = LOCTEXT("TotalWallTimeTooltip", "Total real time elapsed between the first call until completion, including any wait/sleep time, in ms.");
+	const FText TEXT_MemoryInMBTooltip = LOCTEXT("MemoryInMBTooltip", "Size in MB of newly created data.");
 }
 
 void SPCGProfilingListViewItemRow::Construct(const FArguments& InArgs, const TSharedRef<STableViewBase>& InOwnerTableView, const PCGProfilingListViewItemPtr& Item)
@@ -140,6 +143,11 @@ FText FPCGProfilingListViewItem::GetTextForColumn(FName ColumnId, bool bNoGroupi
 		// In ms
 		return FText::AsNumber((CallTime.TotalWallTime()) * 1000.0, NumberFormattingOptions);
 	}
+	else if (ColumnId == PCGEditorGraphProfilingView::NAME_MemoryInMB)
+	{
+		// In MB
+		return FText::AsNumber(CallTime.OutputDataMemorySize / (1024.0 * 1024.0), NumberFormattingOptions);
+	}
 	else
 	{
 		return LOCTEXT("ItemColumnError", "Unrecognized Column");
@@ -306,6 +314,22 @@ void SPCGEditorGraphProfilingView::Construct(const FArguments& InArgs, TSharedPt
 				.Text(this, &SPCGEditorGraphProfilingView::GetTotalWallTimeLabel)
 				.MinDesiredWidth(50.0f)
 			]
+			+SHorizontalBox::Slot()
+			.Padding(FMargin(30.0f, 0.0f, 2.0f, 0.0f))
+			.VAlign(EVerticalAlignment::VAlign_Center)
+			.AutoWidth()
+			[
+				SNew(STextBlock)
+				.Text(LOCTEXT("TotalMemoryInM", "Total data size (MB):"))
+			]
+			+SHorizontalBox::Slot()
+			.VAlign(EVerticalAlignment::VAlign_Center)
+			.AutoWidth()
+			[
+				SNew(STextBlock)
+				.Text(this, &SPCGEditorGraphProfilingView::GetTotalDataSizeInMB)
+				.MinDesiredWidth(50.0f)
+			]
 		]
 		+SVerticalBox::Slot()
 		.FillHeight(1.0f)
@@ -454,7 +478,16 @@ TSharedRef<SHeaderRow> SPCGEditorGraphProfilingView::CreateHeaderRowWidget()
 		.SortMode(this, &SPCGEditorGraphProfilingView::GetColumnSortMode, PCGEditorGraphProfilingView::NAME_TotalWallTime)
 		.OnSort(this, &SPCGEditorGraphProfilingView::OnSortColumnHeader)
 		.InitialSortMode(EColumnSortMode::Descending)
-		.DefaultTooltip(PCGEditorGraphProfilingView::TEXT_TotalWallTimeTooltip);
+		.DefaultTooltip(PCGEditorGraphProfilingView::TEXT_TotalWallTimeTooltip)
+		+ SHeaderRow::Column(PCGEditorGraphProfilingView::NAME_MemoryInMB)
+		.ManualWidth(100.0f)
+		.DefaultLabel(PCGEditorGraphProfilingView::TEXT_MemoryInMBLabel)
+		.HAlignHeader(HAlign_Center)
+		.HAlignCell(HAlign_Right)
+		.SortMode(this, &SPCGEditorGraphProfilingView::GetColumnSortMode, PCGEditorGraphProfilingView::NAME_MemoryInMB)
+		.OnSort(this, &SPCGEditorGraphProfilingView::OnSortColumnHeader)
+		.InitialSortMode(EColumnSortMode::Descending)
+		.DefaultTooltip(PCGEditorGraphProfilingView::TEXT_MemoryInMBTooltip);
 }
 
 FText SPCGEditorGraphProfilingView::GetTotalTimeLabel() const
@@ -467,6 +500,11 @@ FText SPCGEditorGraphProfilingView::GetTotalWallTimeLabel() const
 	return FText::Format(LOCTEXT("GraphTotalWallTimeLabel", "{0} s"), TotalWallTime);
 }
 
+FText SPCGEditorGraphProfilingView::GetTotalDataSizeInMB() const
+{
+	return FText::Format(LOCTEXT("GraphTotalDataSizeInMBLabel", "{0} MB"), TotalDataSizeInMB);
+}
+
 void SPCGEditorGraphProfilingView::OnSortColumnHeader(const EColumnSortPriority::Type SortPriority, const FName& ColumnId, const EColumnSortMode::Type NewSortMode)
 {
 	if (SortingColumn == ColumnId)
@@ -636,6 +674,7 @@ FReply SPCGEditorGraphProfilingView::Refresh()
 {
 	TotalTime = 0;
 	TotalWallTime = 0;
+	TotalDataSizeInMB = 0;
 
 	ListViewItems.Empty();
 	ListView->RequestListRefresh();
@@ -674,6 +713,7 @@ FReply SPCGEditorGraphProfilingView::Refresh()
 	{
 		TotalTime = TreeInfo.CallTime.TotalTime();
 		TotalWallTime = TreeInfo.CallTime.TotalWallTime();
+		TotalDataSizeInMB = TreeInfo.CallTime.OutputDataMemorySize / (1024.0 * 1024.0);
 	}
 
 	//TODO: could turn this into a tree instead of expanding into a list
@@ -728,6 +768,10 @@ FReply SPCGEditorGraphProfilingView::Refresh()
 				{
 					isLess = A->CallTime.TotalWallTime() < B->CallTime.TotalWallTime();
 				}
+				else if (SortingColumn == PCGEditorGraphProfilingView::NAME_MemoryInMB)
+				{
+					isLess = A->CallTime.OutputDataMemorySize < B->CallTime.OutputDataMemorySize;
+				}
 
 				return SortMode == EColumnSortMode::Ascending ? isLess : !isLess;
 			});
diff --git a/Engine/Plugins/PCG/Source/PCGEditor/Private/SPCGEditorGraphProfilingView.h b/Engine/Plugins/PCG/Source/PCGEditor/Private/SPCGEditorGraphProfilingView.h
--- a/Engine/Plugins/PCG/Source/PCGEditor/Private/SPCGEditorGraphProfilingView.h
+++ b/Engine/Plugins/PCG/Source/PCGEditor/Private/SPCGEditorGraphProfilingView.h
@@ -77,6 +77,7 @@ private:
 	EColumnSortMode::Type GetColumnSortMode(const FName ColumnId) const;
 	FText GetTotalTimeLabel() const;
 	FText GetTotalWallTimeLabel() const;
+	FText GetTotalDataSizeInMB() const;
 
 	FReply OnListViewKeyDown(const FGeometry& InGeometry, const FKeyEvent& InKeyEvent) const;
 	void CopySelectionToClipboard() const;
@@ -111,6 +112,7 @@ private:
 
 	double TotalTime = 0.0;
 	double TotalWallTime = 0.0;
+	double TotalDataSizeInMB = 0.0;
 
 	bool bNeedsRefresh = false;
 
