[PCG] Optimization Pass on Create Points Sphere - Geodesic

#rnx

diff --git a/Engine/Plugins/PCG/Source/PCG/Private/Elements/PCGCreatePointsSphere.cpp b/Engine/Plugins/PCG/Source/PCG/Private/Elements/PCGCreatePointsSphere.cpp
--- a/Engine/Plugins/PCG/Source/PCG/Private/Elements/PCGCreatePointsSphere.cpp
+++ b/Engine/Plugins/PCG/Source/PCG/Private/Elements/PCGCreatePointsSphere.cpp
@@ -69,29 +69,104 @@ namespace PCGCreatePointsSphere
 
 		// Golden ratio = (1 + sqrt(5))/2
 		constexpr double Tau = 1.618033988749894;
+		// Norm = sqrt(Tau^2 + 1^2 + 0^2)
+		constexpr double Norm = 1.902113032590306;
+		// To bring the icosahedron to unit sphere, normalize
+		constexpr double OneNorm = 1 / Norm;
+		constexpr double TauNorm = Tau / Norm;
+
+		// Icosahedron hardcoded unit sphere vertex locations.
+		const TStaticArray<FVector, 12> InitialPointLocations(
+			FVector(0, -OneNorm, TauNorm),  // 0
+			FVector(0, OneNorm, TauNorm),   // 1
+			FVector(-TauNorm, 0, OneNorm),  // 2
+			FVector(TauNorm, 0, OneNorm),   // 3
+			FVector(-OneNorm, TauNorm, 0),  // 4
+			FVector(OneNorm, TauNorm, 0),   // 5
+			FVector(-OneNorm, -TauNorm, 0), // 6
+			FVector(OneNorm, -TauNorm, 0),  // 7
+			FVector(-TauNorm, 0, -OneNorm), // 8
+			FVector(TauNorm, 0, -OneNorm),  // 9
+			FVector(0, -OneNorm, -TauNorm), // 10
+			FVector(0, OneNorm, -TauNorm)); // 11
+
+		// Faster to calculate the sin here and compare it directly later
+		const double LatStartSin = FMath::Sin(FMath::DegreesToRadians(Params.LatitudinalStartAngle));
+		const double LatEndSin = FMath::Sin(FMath::DegreesToRadians(Params.LatitudinalEndAngle));
+		const double LonStartRad = FMath::DegreesToRadians(Params.LongitudinalStartAngle);
+		const double LonEndRad = FMath::DegreesToRadians(Params.LongitudinalEndAngle);
+
+		// If all the new indices are outside the window (on the same side), no need to recurse. Use -1 for <, 0 for inside, and 1 for >.
+		auto CheckLatWindow = [LatStartSin, LatEndSin](const FVector& PointLocation)
+		{
+			if (PointLocation.Z < LatStartSin) { return -1; }
+			else if (PointLocation.Z > LatEndSin) { return 1; }
+			else return 0;
+		};
+
+		auto CheckLonWindow = [LonStartRad, LonEndRad](const FVector& PointLocation)
+		{
+			const double LonAngle = FMath::Atan2(PointLocation.Y, PointLocation.X);
+			if (LonAngle < LonStartRad) { return -1; }
+			else if (LonAngle > LonEndRad) { return 1; }
+			else return 0;
+		};
 
-		// Icosahedron hardcoded unit vertex locations
-		TArray<FVector> PointLocations = {{0, -1, Tau}, {0, 1, Tau}, {-Tau, 0, 1}, {Tau, 0, 1}, {-1, Tau, 0}, {1, Tau, 0}, {-1, -Tau, 0}, {1, -Tau, 0}, {-Tau, 0, -1}, {Tau, 0, -1}, {0, -1, -Tau}, {0, 1, -Tau}};
+		// First subdivision starts with hardcoded locations. So, prime output with culled initial locations.
+		TArray<FVector> PointLocations;
+		PointLocations.Reserve(InitialPointLocations.Num());
+		for (const FVector& PointLocation : InitialPointLocations)
+		{
+			if (CheckLatWindow(PointLocation) == 0 && CheckLonWindow(PointLocation) == 0)
+			{
+				PointLocations.Emplace(PointLocation);
+			}
+		}
 
-		auto SubdivideTriangleRecursive = [&Out = PointLocations](const FVector& V0, const FVector& V1, const FVector& V2, int32 SubdivisionCount, auto&& RecursiveCall)
+		auto SubdivideTriangleRecursive = [&Out = PointLocations, CheckLatWindow, CheckLonWindow](const FVector& V0, const FVector& V1, const FVector& V2, int32 RemainingSubdivisionCount, auto&& RecursiveCall)
 		{
-			if (SubdivisionCount < 1)
+			if (RemainingSubdivisionCount < 1)
 			{
 				return;
 			}
 
-			// TODO: This can probably be better optimized to not rely on unique
 			// First calculate and add the midpoints
-			const int32 I0 = Out.AddUnique((V0 + V1) * 0.5);
-			const int32 I1 = Out.AddUnique((V1 + V2) * 0.5);
-			const int32 I2 = Out.AddUnique((V2 + V0) * 0.5);
+			const FVector M0 = (V0 + V1) * 0.5;
+			const FVector M1 = (V1 + V2) * 0.5;
+			const FVector M2 = (V2 + V0) * 0.5;
+			// Normalize
+			const FVector N0 = (1.0 / M0.Length()) * M0;
+			const FVector N1 = (1.0 / M1.Length()) * M1;
+			const FVector N2 = (1.0 / M2.Length()) * M2;
+
+			// Check if these midpoints are in the window (0), inside (-1), or outside (1)
+			const int Lat0 = CheckLatWindow(N0);
+			const int Lat1 = CheckLatWindow(N1);
+			const int Lat2 = CheckLatWindow(N2);
+			const int Lon0 = CheckLonWindow(N0);
+			const int Lon1 = CheckLonWindow(N1);
+			const int Lon2 = CheckLonWindow(N2);
+
+			// TODO: This could likely be better optimized to not rely on unique
+			// Add the points as long as they're in both windows
+			if (Lat0 == 0 && Lon0 == 0) { Out.AddUnique(N0); }
+			if (Lat1 == 0 && Lon1 == 0) { Out.AddUnique(N1); }
+			if (Lat2 == 0 && Lon2 == 0) { Out.AddUnique(N2); }
+
+			// Cull subdivisions if the points are outside the windows, but only if both sets are on the same side, since they could straddle.
+			if (Lat0 && Lat0 == Lat1 && Lat0 == Lat2)
+			{
+				if (Lon0 && Lon0 == Lon1 && Lon0 == Lon2)
+				{
+					return;
+				}
+			}
 
-			// TODO: This is highly inefficient with a small window and high subdivision count. Check if the triangle is outside the window and cull if so.
 			// Then recurse on the four new triangles
-			RecursiveCall(V0, Out[I0], Out[I2], SubdivisionCount - 1, RecursiveCall);
-			RecursiveCall(V1, Out[I0], Out[I1], SubdivisionCount - 1, RecursiveCall);
-			RecursiveCall(V2, Out[I1], Out[I2], SubdivisionCount - 1, RecursiveCall);
-			RecursiveCall(Out[I0], Out[I1], Out[I2], SubdivisionCount - 1, RecursiveCall);
+			RecursiveCall(V0, M0, M2, RemainingSubdivisionCount - 1, RecursiveCall);
+			RecursiveCall(V1, M0, M1, RemainingSubdivisionCount - 1, RecursiveCall);
+			RecursiveCall(V2, M1, M2, RemainingSubdivisionCount - 1, RecursiveCall);
+			RecursiveCall(M0, M1, M2, RemainingSubdivisionCount - 1, RecursiveCall);
 		};
 
 		// Hardcoded Icosahedron Triangles
@@ -102,38 +177,14 @@ namespace PCGCreatePointsSphere
 		{
 			for (int I = 2; I < std::size(InitialTriangles); I += 3)
 			{
-				SubdivideTriangleRecursive(PointLocations[InitialTriangles[I - 2]], PointLocations[InitialTriangles[I - 1]], PointLocations[InitialTriangles[I]], SubdivisionCount, SubdivideTriangleRecursive);
+				SubdivideTriangleRecursive(InitialPointLocations[InitialTriangles[I - 2]], InitialPointLocations[InitialTriangles[I - 1]], InitialPointLocations[InitialTriangles[I]], SubdivisionCount, SubdivideTriangleRecursive);
 			}
 		}
 
-		const double LatStartRad = FMath::DegreesToRadians(Params.LatitudinalStartAngle);
-		const double LatEndRad = FMath::DegreesToRadians(Params.LatitudinalEndAngle);
-		const double LonStartRad = FMath::DegreesToRadians(Params.LongitudinalStartAngle);
-		const double LonEndRad = FMath::DegreesToRadians(Params.LongitudinalEndAngle);
-
-		int32 Index = 0;
-		while (Index < PointLocations.Num())
+		for (FVector& PointLocation : PointLocations)
 		{
-			FVector& PointLocation = PointLocations[Index];
-
-			// Since this algorithm is a recursive subdivision, culling will have to wait until the end
-			const double LatAngle = FMath::Atan2(PointLocation.Z / FVector(PointLocation.X, PointLocation.Y, 0).Length(), 1);
-			const double LonAngle = FMath::Atan2(PointLocation.Y, PointLocation.X);
-			if (LatAngle >= LatStartRad && LatAngle <= LatEndRad && LonAngle >= LonStartRad && LonAngle <= LonEndRad)
-			{
-				// Normalize to bring it to the unit sphere
-				PointLocation = (1.0 / PointLocation.Length()) * PointLocation;
-
-				// Finally, apply the user defined transform parameters
-				PointLocation *= Params.Radius;
-				PointLocation += Params.Origin;
-				++Index;
-			}
-			else // Cull the point
-			{
-				// No need to shrink, since it will append at the end
-				PointLocations.RemoveSingleSwap(PointLocation, EAllowShrinking::No);
-			}
+			// Finally, apply the user defined transform parameters
+			PointLocation = (PointLocation * Params.Radius) + Params.Origin;
 		}
 
 		OutLocations.Append(PointLocations);
