[PCG] Implement custom progress notification with Cancel button

#rnx

diff --git a/Engine/Plugins/PCG/Source/PCG/Private/Graph/PCGGraphExecutor.cpp b/Engine/Plugins/PCG/Source/PCG/Private/Graph/PCGGraphExecutor.cpp
--- a/Engine/Plugins/PCG/Source/PCG/Private/Graph/PCGGraphExecutor.cpp
+++ b/Engine/Plugins/PCG/Source/PCG/Private/Graph/PCGGraphExecutor.cpp
@@ -24,6 +24,8 @@
 
 #if WITH_EDITOR
 #include "Editor.h"
+#include "Editor/IPCGEditorModule.h"
+#include "Modules/ModuleManager.h"
 #else
 #include "GameFramework/Actor.h"
 #endif
@@ -131,10 +133,8 @@ bool FPCGGraphTask::IsApproximatelyEqual(const FPCGGraphTask& Other) const
 }
 #endif // WITH_EDITOR
 
-FPCGGraphExecutor::FPCGGraphExecutor()
-#if WITH_EDITOR
-	: GenerationProgressNotification(GetNotificationTextFormat())
-#endif
+FPCGGraphExecutor::FPCGGraphExecutor(UWorld* InWorld)
+	: World(InWorld)
 {
 }
 
@@ -147,7 +147,7 @@ FPCGGraphExecutor::~FPCGGraphExecutor()
 #if WITH_EDITOR
 	// Cleanup + clear notification
 	ClearAllTasks();
-	UpdateGenerationNotification();
+	ReleaseGenerationNotification();
 #endif
 }
 
@@ -1751,22 +1751,59 @@ void FPCGGraphExecutor::NotifyGraphChanged(UPCGGraph* InGraph, EPCGChangeType Ch
 void FPCGGraphExecutor::UpdateGenerationNotification()
 {
 	const int32 RemainingTaskNum = GetNonScheduledRemainingTaskCount();
+	if (RemainingTaskNum == 0)
+	{
+		ReleaseGenerationNotification();
+		return;
+	}
 
-	if (RemainingTaskNum > 0)
+	if (!GenerationProgressNotification.IsValid())
 	{
-		GenerationProgressNotification.Update(RemainingTaskNum);
+		IPCGEditorModule* EditorModule = FModuleManager::GetModulePtr<IPCGEditorModule>("PCGEditor");
+		if (!EditorModule)
+		{
+			return;
+		}
+
+		GenerationProgressNotification = EditorModule->CreateProgressNotification(GetNotificationTextFormat(), /*bCanCancel=*/true);
+		if (GenerationProgressNotification.IsValid())
+		{
+			GenerationProgressNotification.Pin()->OnCancelTasks().AddRaw(this, &FPCGGraphExecutor::OnNotificationCancel);
+		}
 	}
-	else
+
+	if (GenerationProgressNotification.IsValid())
+	{
+		GenerationProgressNotification.Pin()->Update(RemainingTaskNum);
+	}
+}
+
+void FPCGGraphExecutor::ReleaseGenerationNotification()
+{
+	if (GenerationProgressNotification.IsValid())
+	{
+		GenerationProgressNotification.Pin()->OnCancelTasks().RemoveAll(this);
+		if (IPCGEditorModule* EditorModule = FModuleManager::GetModulePtr<IPCGEditorModule>("PCGEditor"))
+		{
+			EditorModule->ReleaseProgressNotification(GenerationProgressNotification);
+		}
+		GenerationProgressNotification = nullptr;
+	}
+}
+
+void FPCGGraphExecutor::OnNotificationCancel()
+{
+	ReleaseGenerationNotification();
+
+	if (UPCGSubsystem* PCGSubsystem = UWorld::GetSubsystem<UPCGSubsystem>(World))
 	{
-		GenerationProgressNotification.Update(0);
-		// To reset the UI notification
-		GenerationProgressNotification = FAsyncCompilationNotification(GetNotificationTextFormat());
+		PCGSubsystem->CancelAllGeneration();
 	}
 }
 
 FTextFormat FPCGGraphExecutor::GetNotificationTextFormat()
 {
-	return NSLOCTEXT("PCG", "PCGGenerationNotificationFormat", "{0}|plural(one=PCG Task,other=PCG Tasks)");
+	return NSLOCTEXT("PCG", "PCGGenerationNotificationFormat", "Executing PCG {0}|plural(one=Task,other=Tasks)");
 }
 
 #endif // WITH_EDITOR
diff --git a/Engine/Plugins/PCG/Source/PCG/Private/Graph/PCGGraphExecutor.h b/Engine/Plugins/PCG/Source/PCG/Private/Graph/PCGGraphExecutor.h
--- a/Engine/Plugins/PCG/Source/PCG/Private/Graph/PCGGraphExecutor.h
+++ b/Engine/Plugins/PCG/Source/PCG/Private/Graph/PCGGraphExecutor.h
@@ -13,7 +13,7 @@
 #include "UObject/GCObject.h"
 
 #if WITH_EDITOR
-#include "AsyncCompilationHelpers.h"
+#include "Editor/IPCGEditorProgressNotification.h"
 #include "WorldPartition/WorldPartitionHandle.h" // Needed for FWorldPartitionReference
 #endif
 
@@ -118,7 +118,9 @@ struct FPCGGraphActiveTask
 class FPCGGraphExecutor : public FGCObject
 {
 public:
-	FPCGGraphExecutor();
+	// Default constructor used by unittests
+	FPCGGraphExecutor() = default;
+	FPCGGraphExecutor(UWorld* InWorld);
 	~FPCGGraphExecutor();
 
 	/** Compile (and cache) a graph for later use. This call is threadsafe */
@@ -241,6 +243,8 @@ private:
 	void ReleaseUnusedActors();
 
 	void UpdateGenerationNotification();
+	void ReleaseGenerationNotification();
+	void OnNotificationCancel();
 	static FTextFormat GetNotificationTextFormat();
 #endif
 
@@ -276,10 +280,12 @@ private:
 	TSet<FWorldPartitionReference> ActorsToRelease;
 
 	int32 ReleaseActorsCountUntilGC = 30;
-	FAsyncCompilationNotification GenerationProgressNotification;
+	TWeakPtr<IPCGEditorProgressNotification> GenerationProgressNotification;
 
 	int32 TidyCacheCountUntilGC = 100;
 #endif
+
+	TObjectPtr<UWorld> World = nullptr;
 };
 
 class FPCGFetchInputElement : public IPCGElement
diff --git a/Engine/Plugins/PCG/Source/PCG/Private/PCGSubsystem.cpp b/Engine/Plugins/PCG/Source/PCG/Private/PCGSubsystem.cpp
--- a/Engine/Plugins/PCG/Source/PCG/Private/PCGSubsystem.cpp
+++ b/Engine/Plugins/PCG/Source/PCG/Private/PCGSubsystem.cpp
@@ -186,7 +186,7 @@ void UPCGSubsystem::PostInitialize()
 
 	// Initialize graph executor
 	check(!GraphExecutor);
-	GraphExecutor = new FPCGGraphExecutor();
+	GraphExecutor = new FPCGGraphExecutor(GetWorld());
 
 	// Initialize runtime generation scheduler
 	check(!RuntimeGenScheduler);
diff --git a/Engine/Plugins/PCG/Source/PCG/Public/Editor/IPCGEditorModule.h b/Engine/Plugins/PCG/Source/PCG/Public/Editor/IPCGEditorModule.h
new file mode 100644
--- /dev/null
+++ b/Engine/Plugins/PCG/Source/PCG/Public/Editor/IPCGEditorModule.h
@@ -0,0 +1,15 @@
+// Copyright Epic Games, Inc. All Rights Reserved.
+
+#pragma once
+
+#include "Editor/IPCGEditorProgressNotification.h"
+
+#include "Internationalization/Text.h"
+#include "Modules/ModuleInterface.h"
+
+class IPCGEditorModule : public IModuleInterface
+{
+public:
+	virtual TWeakPtr<IPCGEditorProgressNotification> CreateProgressNotification(const FTextFormat& TextFormat, bool bCanCancel) = 0;
+	virtual void ReleaseProgressNotification(TWeakPtr<IPCGEditorProgressNotification> InNotification) = 0;
+};
\ No newline at end of file
diff --git a/Engine/Plugins/PCG/Source/PCG/Public/Editor/IPCGEditorProgressNotification.h b/Engine/Plugins/PCG/Source/PCG/Public/Editor/IPCGEditorProgressNotification.h
new file mode 100644
--- /dev/null
+++ b/Engine/Plugins/PCG/Source/PCG/Public/Editor/IPCGEditorProgressNotification.h
@@ -0,0 +1,24 @@
+// Copyright Epic Games, Inc. All Rights Reserved.
+
+#pragma once
+
+#include "CoreMinimal.h"
+#include "Templates/SharedPointer.h"
+
+class IPCGEditorProgressNotification : public TSharedFromThis<IPCGEditorProgressNotification>
+{
+public:
+	virtual ~IPCGEditorProgressNotification() {}
+	
+	virtual void Update(uint32 NumRemainingTasks) = 0;
+
+	virtual uint32 GetNumTasks() const = 0;
+	virtual uint32 GetNumTasksDone() const = 0;
+	virtual const FTextFormat& GetTextFormat() const = 0;
+
+	virtual void Cancel() = 0;
+	virtual bool CanCancel() const = 0;
+
+	DECLARE_EVENT(IPCGEditorNotification, FCancelTasksEvent);
+	virtual FCancelTasksEvent& OnCancelTasks() = 0;
+};
\ No newline at end of file
diff --git a/Engine/Plugins/PCG/Source/PCGEditor/Private/PCGEditorModule.cpp b/Engine/Plugins/PCG/Source/PCGEditor/Private/PCGEditorModule.cpp
--- a/Engine/Plugins/PCG/Source/PCGEditor/Private/PCGEditorModule.cpp
+++ b/Engine/Plugins/PCG/Source/PCGEditor/Private/PCGEditorModule.cpp
@@ -6,6 +6,7 @@
 #include "PCGEditorCommands.h"
 #include "PCGEditorGraphNodeFactory.h"
 #include "PCGEditorMenuUtils.h"
+#include "PCGEditorProgressNotification.h"
 #include "PCGEditorSettings.h"
 #include "PCGEditorStyle.h"
 #include "PCGEditorUtils.h"
@@ -140,6 +141,24 @@ bool FPCGEditorModule::SupportsDynamicReloading()
 	return true;
 }
 
+TWeakPtr<IPCGEditorProgressNotification> FPCGEditorModule::CreateProgressNotification(const FTextFormat& TextFormat, bool bCanCancel)
+{
+	TSharedPtr<IPCGEditorProgressNotification> NewNotification = MakeShared<FPCGEditorProgressNotification>(TextFormat, bCanCancel);
+	ActiveNotifications.Add(NewNotification);
+	return NewNotification.ToWeakPtr();
+}
+
+void FPCGEditorModule::ReleaseProgressNotification(TWeakPtr<IPCGEditorProgressNotification> InNotification)
+{
+	if (InNotification.IsValid())
+	{
+		if (TSharedPtr<IPCGEditorProgressNotification> SharedPtr = InNotification.Pin(); ActiveNotifications.Contains(SharedPtr))
+		{
+			ActiveNotifications.Remove(SharedPtr);
+		}
+	}
+}
+
 void FPCGEditorModule::RegisterOnEditorModeChange()
 {
 	// Have a callback that catches changes in the Editor modes, to catch when we exit the landscape edit mode.
diff --git a/Engine/Plugins/PCG/Source/PCGEditor/Private/PCGEditorModule.h b/Engine/Plugins/PCG/Source/PCGEditor/Private/PCGEditorModule.h
--- a/Engine/Plugins/PCG/Source/PCGEditor/Private/PCGEditorModule.h
+++ b/Engine/Plugins/PCG/Source/PCGEditor/Private/PCGEditorModule.h
@@ -2,6 +2,8 @@
 
 #pragma once
 
+#include "Editor/IPCGEditorModule.h"
+
 #include "AssetTypeCategories.h"
 #include "Modules/ModuleInterface.h"
 #include "Toolkits/AssetEditorToolkit.h"
@@ -13,7 +15,7 @@ class FMenuBuilder;
 class FPCGEditorGraphNodeFactory;
 class IAssetTypeActions;
 
-class FPCGEditorModule : public IModuleInterface
+class FPCGEditorModule : public IPCGEditorModule
 {
 public:
 	// ~IModuleInterface implementation
@@ -22,6 +24,11 @@ public:
 	virtual bool SupportsDynamicReloading() override;
 	// ~End IModuleInterface implementation
 
+	// ~IPCGEditorModule implementation
+	virtual TWeakPtr<IPCGEditorProgressNotification> CreateProgressNotification(const FTextFormat& TextFormat, bool bCanCancel) override;
+	virtual void ReleaseProgressNotification(TWeakPtr<IPCGEditorProgressNotification> InNotification) override;
+	// ~End IPCGEditorModule implementation
+
 protected:
 	void RegisterDetailsCustomizations();
 	void UnregisterDetailsCustomizations();
@@ -48,4 +55,6 @@ protected:
 	TSharedPtr<FPCGEditorGraphNodeFactory> GraphNodeFactory;
 
 	FDelegateHandle ShouldDisableCPUThrottlingDelegateHandle;
+
+	TSet<TSharedPtr<IPCGEditorProgressNotification>> ActiveNotifications;
 };
diff --git a/Engine/Plugins/PCG/Source/PCGEditor/Private/PCGEditorProgressNotification.cpp b/Engine/Plugins/PCG/Source/PCGEditor/Private/PCGEditorProgressNotification.cpp
new file mode 100644
--- /dev/null
+++ b/Engine/Plugins/PCG/Source/PCGEditor/Private/PCGEditorProgressNotification.cpp
@@ -0,0 +1,215 @@
+// Copyright Epic Games, Inc. All Rights Reserved.
+
+#include "PCGEditorProgressNotification.h"
+#include "PCGEditorStyle.h"
+
+#include "Framework/Notifications/NotificationManager.h"
+#include "Widgets/Images/SImage.h"
+#include "Widgets/Input/SButton.h"
+#include "Widgets/Layout/SBox.h"
+#include "Widgets/Notifications/INotificationWidget.h"
+#include "Widgets/Notifications/SProgressBar.h"
+#include "Widgets/SBoxPanel.h"
+#include "Widgets/SCompoundWidget.h"
+#include "Widgets/SOverlay.h"
+#include "Widgets/Text/STextBlock.h"
+
+namespace PCGEditorProgressNotification
+{
+	class SProgressNotificationWidget : public SCompoundWidget, public INotificationWidget
+	{
+	public:
+		SLATE_BEGIN_ARGS(SProgressNotificationWidget) {}
+			SLATE_ATTRIBUTE(TWeakPtr<IPCGEditorProgressNotification>, ProgressNotification)
+		SLATE_END_ARGS()
+
+	public:
+		void Construct(const FArguments & InArgs)
+		{
+			ProgressNotification = InArgs._ProgressNotification;
+			if (TWeakPtr<IPCGEditorProgressNotification> ProgressNotificationWeakPtr = ProgressNotification.Get(); ProgressNotificationWeakPtr.IsValid())
+			{
+				if (TSharedPtr<IPCGEditorProgressNotification> ProgressNotificationPtr = ProgressNotificationWeakPtr.Pin())
+				{
+					TextFormat = ProgressNotificationPtr->GetTextFormat();
+					NumTasks = ProgressNotificationPtr->GetNumTasks();
+				}
+			}
+
+			ChildSlot
+			[
+				SNew(SVerticalBox)
+				+ SVerticalBox::Slot()
+				[
+					SNew(SHorizontalBox)
+					+ SHorizontalBox::Slot()
+					.AutoWidth()
+					.HAlign(HAlign_Center)
+					.VAlign(VAlign_Center)
+					.Padding(0.0f, 0.0f, 4.0f, 0.0f)
+					[
+						SNew(SImage)
+						.Image(FPCGEditorStyle::Get().GetBrush("PCG.EditorIcon"))
+					]
+					+ SHorizontalBox::Slot()
+					.Padding(0.0f, 0.0f, 4.0f, 0.0f)
+					[
+						SNew(SVerticalBox)
+						+ SVerticalBox::Slot()
+						.AutoHeight()
+						.Padding(0.0f, 3.0f, 0.0f, 2.0f)
+						[
+							SAssignNew(ProgressTextWidget, STextBlock)
+							.Font(FAppStyle::Get().GetFontStyle("NotificationList.FontBold"))
+						]
+						+ SVerticalBox::Slot()
+						[
+							SNew(SBox)
+							.HeightOverride(8.f)
+							[
+								SNew(SOverlay)
+								+ SOverlay::Slot()
+								.VAlign(VAlign_Center)
+								.Padding(1.0f, 0.0f)
+								[
+									SAssignNew(ProgressBar, SProgressBar)
+									.Percent(0.0f)
+								]
+								+ SOverlay::Slot()
+								[
+									SNew(SImage)
+									.Image(FAppStyle::Get().GetBrush("StatusBar.ProgressOverlay"))
+									.Visibility(EVisibility::HitTestInvisible)
+								]
+							]
+						]
+					]
+					+ SHorizontalBox::Slot()
+					.AutoWidth()
+					.HAlign(HAlign_Right)
+					.VAlign(VAlign_Center)
+					[
+						SAssignNew(PercentText, STextBlock)
+					]
+				]
+				+ SVerticalBox::Slot()
+				.AutoHeight()
+				.Padding(0.0f, 5.0f, 0.0f, 0.0f)
+				.VAlign(VAlign_Center)
+				.HAlign(HAlign_Right)
+				[
+					SNew(SButton)
+					.Text(NSLOCTEXT("PCGEditorProgressNotification", "CancelButtonText", "Cancel"))
+					.Visibility_Lambda([this]()
+					{
+						TWeakPtr<IPCGEditorProgressNotification> ProgressNotificationPtr = ProgressNotification.Get();
+						if (ProgressNotificationPtr.IsValid() && ProgressNotificationPtr.Pin()->CanCancel())
+						{
+							return EVisibility::Visible;
+						}
+						return EVisibility::Collapsed;
+					})
+					.OnClicked_Lambda([this]()
+					{
+						TWeakPtr<IPCGEditorProgressNotification> ProgressNotificationPtr = ProgressNotification.Get();
+						if (ProgressNotificationPtr.IsValid())
+						{
+							ProgressNotificationPtr.Pin()->Cancel();
+						}
+						return FReply::Handled();
+					})
+				]
+			];
+		}
+
+		virtual void Tick(const FGeometry & AllottedGeometry, const double InCurrentTime, const float InDeltaTime) override
+		{						
+			uint32 NumTasksDone = NumTasks;
+			float PercentDone = 1.0f;
+
+			// Check if Notification is still valid
+			TWeakPtr<IPCGEditorProgressNotification> ProgressNotificationWeakPtr = ProgressNotification.Get();
+			TSharedPtr<IPCGEditorProgressNotification> ProgressNotificationPtr = ProgressNotificationWeakPtr.IsValid() ? ProgressNotificationWeakPtr.Pin() : nullptr;
+			if (ProgressNotificationPtr)
+			{
+				NumTasks = ProgressNotificationPtr->GetNumTasks();
+				NumTasksDone = ProgressNotificationPtr->GetNumTasksDone();
+				PercentDone = NumTasks == 0 ? 1.0f : FMath::Clamp((float)NumTasksDone / NumTasks, 0.0f, 1.0f);
+			}
+
+			FFormatNamedArguments Args;
+			
+			Args.Add(TEXT("ProgressMessage"), FText::Format(TextFormat, FText::AsNumber(NumTasks)));
+			Args.Add(TEXT("NumTasks"), FText::AsNumber(NumTasks));
+			Args.Add(TEXT("NumTasksDone"), FText::AsNumber(NumTasksDone));
+
+			ProgressTextWidget->SetText(FText::Format(NSLOCTEXT("PCGEditorProgressNotification", "ProgressText", "{ProgressMessage} ({NumTasksDone}/{NumTasks})"), Args));
+			PercentText->SetText(FText::AsPercent(PercentDone));
+			ProgressBar->SetPercent(PercentDone);
+		}
+
+		void SetProgressText(FText ProgressText)
+		{
+			ProgressTextWidget->SetText(ProgressText);
+		}
+
+		void SetProgressPercent(float Percent)
+		{
+			ProgressBar->SetPercent(Percent);
+			PercentText->SetText(FText::AsPercent(Percent));
+		}
+
+		/** INotificationWidget interface */
+		virtual void OnSetCompletionState(SNotificationItem::ECompletionState) override {}
+
+		TSharedRef<SWidget> AsWidget() override
+		{
+			return AsShared();
+		}
+		
+	private:
+		TAttribute<TWeakPtr<IPCGEditorProgressNotification>> ProgressNotification;
+		TSharedPtr<SProgressBar> ProgressBar;
+		TSharedPtr<STextBlock> PercentText;
+		TSharedPtr<STextBlock> ProgressTextWidget;
+
+		// Cached values
+		FTextFormat TextFormat;
+		uint32 NumTasks = 0;
+	};
+}
+
+FPCGEditorProgressNotification::FPCGEditorProgressNotification(const FTextFormat& InTextFormat, bool bInCanCancel)
+	: NumTasks(0), NumTasksDone(0), bCanCancel(bInCanCancel), TextFormat(InTextFormat)
+{
+}
+
+FPCGEditorProgressNotification::~FPCGEditorProgressNotification()
+{
+	if (NotificationItem.IsValid())
+	{
+		NotificationItem->ExpireAndFadeout();
+		NotificationItem = nullptr;
+	}
+}
+
+void FPCGEditorProgressNotification::Update(uint32 InNumRemainingTasks)
+{
+	if (InNumRemainingTasks > NumTasks)
+	{
+		NumTasks = InNumRemainingTasks;
+	}
+	NumTasksDone = NumTasks - InNumRemainingTasks;
+
+	if (NumTasks > 0 && !NotificationItem)
+	{
+		FNotificationInfo ProgressNotification(FText::GetEmpty());
+
+		ProgressNotification.ContentWidget = SNew(PCGEditorProgressNotification::SProgressNotificationWidget).ProgressNotification(AsWeak());
+		ProgressNotification.bAllowThrottleWhenFrameRateIsLow = false;
+		ProgressNotification.bFireAndForget = false;
+		ProgressNotification.ExpireDuration = 1.0f;
+		ProgressNotification.FadeOutDuration = 0.5f;
+		NotificationItem = FSlateNotificationManager::Get().AddNotification(ProgressNotification);
+	}
+}
diff --git a/Engine/Plugins/PCG/Source/PCGEditor/Private/PCGEditorProgressNotification.h b/Engine/Plugins/PCG/Source/PCGEditor/Private/PCGEditorProgressNotification.h
new file mode 100644
--- /dev/null
+++ b/Engine/Plugins/PCG/Source/PCGEditor/Private/PCGEditorProgressNotification.h
@@ -0,0 +1,39 @@
+// Copyright Epic Games, Inc. All Rights Reserved.
+
+#pragma once
+
+#include "Editor/IPCGEditorProgressNotification.h"
+
+#include "CoreMinimal.h"
+#include "Internationalization/Text.h"
+#include "Widgets/Notifications/SNotificationList.h"
+
+class FPCGEditorProgressNotification : public IPCGEditorProgressNotification
+{
+public:
+	FPCGEditorProgressNotification(const FTextFormat& InTextFormat, bool bInCanCancel);
+	virtual ~FPCGEditorProgressNotification();
+
+	virtual void Update(uint32 InNumRemainingTasks) override;
+		
+	DECLARE_DERIVED_EVENT(FPCGEditorProgressNotification, IPCGEditorProgressNotification::FCancelTasksEvent, FCancelTasksEvent);
+	virtual FCancelTasksEvent& OnCancelTasks() override { return CancelTasksEvent; }
+
+	virtual uint32 GetNumTasks() const override { return NumTasks; }
+	virtual uint32 GetNumTasksDone() const override { return NumTasksDone; }
+	virtual const FTextFormat& GetTextFormat() const override { return TextFormat; }
+
+	virtual void Cancel() override { OnCancelTasks().Broadcast(); }
+	virtual bool CanCancel() const override { return bCanCancel; }
+	
+private:
+	uint32 NumTasks = 0;
+	uint32 NumTasksDone = 0;
+	
+	bool bCanCancel = false;
+	bool bIsCompleted = false;
+	
+	FTextFormat TextFormat;
+	FCancelTasksEvent CancelTasksEvent;
+	TSharedPtr<SNotificationItem> NotificationItem;
+};
\ No newline at end of file
