[PCG] External Data Layers / Data Layers support
- Add DataLayerAssets to FPCGGridDescriptor
- PCGPartitionActors are created respecting this (EDL & DL)
- PCGSpawnActor and PCGCreateTargetActor node will create actors in proper EDL & DL based on their TargetActors
- PCGSpawnActor node now uses UPCGHelpers::SpawnDefaultActor
- UPCGHelpers::SpawnDefaultActor has a new signature taking a struct param to allow easier adding of new params in the future
- Removal of RuntimePartitionActorRecords as we don't allow Partitioned PCG Volumes to Generate on load (responsability is on the PAs if they exist)

#rnx

diff --git a/Engine/Plugins/PCG/Source/PCG/Private/Elements/PCGCreateTargetActor.cpp b/Engine/Plugins/PCG/Source/PCG/Private/Elements/PCGCreateTargetActor.cpp
--- a/Engine/Plugins/PCG/Source/PCG/Private/Elements/PCGCreateTargetActor.cpp
+++ b/Engine/Plugins/PCG/Source/PCG/Private/Elements/PCGCreateTargetActor.cpp
@@ -18,6 +18,10 @@
 #include "GameFramework/Actor.h"
 #include "UObject/Package.h"
 
+#if WITH_EDITOR
+#include "WorldPartition/DataLayer/DataLayerInstance.h"
+#endif
+
 #include UE_INLINE_GENERATED_CPP_BY_NAME(PCGCreateTargetActor)
 
 #define LOCTEXT_NAMESPACE "PCGCreateTargetActor"
@@ -292,6 +296,7 @@ bool FPCGCreateTargetActorElement::ExecuteInternal(FPCGContext* Context) const
 	FActorSpawnParameters SpawnParams;
 	SpawnParams.Template = TemplateActor;
 	SpawnParams.SpawnCollisionHandlingOverride = ESpawnActorCollisionHandlingMethod::AlwaysSpawn;
+	SpawnParams.OverrideLevel = TargetActor->GetLevel();
 
 	if (PCGHelpers::IsRuntimeOrPIE() || (Context->SourceComponent.IsValid() && Context->SourceComponent->IsInPreviewMode()))
 	{
@@ -304,7 +309,13 @@ bool FPCGCreateTargetActorElement::ExecuteInternal(FPCGContext* Context) const
 		Transform = Settings->ActorPivot;
 	}
 
-	AActor* GeneratedActor = UPCGActorHelpers::SpawnDefaultActor(TargetActor->GetWorld(), TargetActor->GetLevel(), Settings->TemplateActorClass, Transform, SpawnParams);
+	UPCGActorHelpers::FSpawnDefaultActorParams SpawnDefaultActorParams(TargetActor->GetWorld(), Settings->TemplateActorClass, Transform, SpawnParams);
+	
+#if WITH_EDITOR
+	SpawnDefaultActorParams.DataLayerInstances = TargetActor->GetDataLayerInstances();
+#endif
+
+	AActor* GeneratedActor = UPCGActorHelpers::SpawnDefaultActor(SpawnDefaultActorParams);
 
 	if (!GeneratedActor)
 	{
diff --git a/Engine/Plugins/PCG/Source/PCG/Private/Elements/PCGSpawnActor.cpp b/Engine/Plugins/PCG/Source/PCG/Private/Elements/PCGSpawnActor.cpp
--- a/Engine/Plugins/PCG/Source/PCG/Private/Elements/PCGSpawnActor.cpp
+++ b/Engine/Plugins/PCG/Source/PCG/Private/Elements/PCGSpawnActor.cpp
@@ -799,6 +799,12 @@ void FPCGSpawnActorElement::SpawnActors(FPCGSubgraphContext* Context, AActor* Ta
 		SpawnParams.ObjectFlags |= RF_Transient;
 	}
 
+	UPCGActorHelpers::FSpawnDefaultActorParams SpawnDefaultActorParams(TargetActor->GetWorld(), InTemplateActorClass, FTransform::Identity, SpawnParams);
+
+#if WITH_EDITOR
+	SpawnDefaultActorParams.DataLayerInstances = TargetActor->GetDataLayerInstances();
+#endif
+
 	const bool bForceCallGenerate = (Settings->GenerationTrigger == EPCGSpawnActorGenerationTrigger::ForceGenerate);
 #if WITH_EDITOR
 	const bool bOnLoadCallGenerate = (Settings->GenerationTrigger == EPCGSpawnActorGenerationTrigger::Default);
@@ -890,7 +896,8 @@ void FPCGSpawnActorElement::SpawnActors(FPCGSubgraphContext* Context, AActor* Ta
 
 			bAllActorOverridesSucceeded &= ActorOverrides.Apply(i);
 
-			AActor* GeneratedActor = TargetActor->GetWorld()->SpawnActor(InTemplateActorClass, &Point.Transform, SpawnParams);
+			SpawnDefaultActorParams.Transform = Point.Transform;
+			AActor* GeneratedActor = UPCGActorHelpers::SpawnDefaultActor(SpawnDefaultActorParams);
 
 			if (!GeneratedActor)
 			{
diff --git a/Engine/Plugins/PCG/Source/PCG/Private/Graph/PCGGraphExecutor.cpp b/Engine/Plugins/PCG/Source/PCG/Private/Graph/PCGGraphExecutor.cpp
--- a/Engine/Plugins/PCG/Source/PCG/Private/Graph/PCGGraphExecutor.cpp
+++ b/Engine/Plugins/PCG/Source/PCG/Private/Graph/PCGGraphExecutor.cpp
@@ -2035,10 +2035,7 @@ namespace PCGGraphExecutor
 				const AActor* ComponentActor = InContext->SourceComponent->GetOwner();
 				if (ComponentActor)
 				{
-					FPCGGridDescriptor GridDescriptor = FPCGGridDescriptor()
-						.SetGridSize(FromGridSize)
-						.SetIs2DGrid(InContext->SourceComponent->GetOriginalComponent()->Use2DGrid())
-						.SetIsRuntime(InContext->SourceComponent->IsManagedByRuntimeGenSystem());
+					FPCGGridDescriptor GridDescriptor = InContext->SourceComponent->GetGridDescriptor(FromGridSize);
 
 					// Get grid coords using the parent grid (FromGridSize).
 					const FIntVector CellCoords = UPCGActorHelpers::GetCellCoord(ComponentActor->GetActorLocation(), GridDescriptor.GetGridSize(), GridDescriptor.Is2DGrid());
diff --git a/Engine/Plugins/PCG/Source/PCG/Private/Grid/PCGGridDescriptor.cpp b/Engine/Plugins/PCG/Source/PCG/Private/Grid/PCGGridDescriptor.cpp
--- a/Engine/Plugins/PCG/Source/PCG/Private/Grid/PCGGridDescriptor.cpp
+++ b/Engine/Plugins/PCG/Source/PCG/Private/Grid/PCGGridDescriptor.cpp
@@ -1,23 +1,151 @@
 // Copyright Epic Games, Inc. All Rights Reserved.
 
 #include "Grid/PCGGridDescriptor.h"
+
+#include "Algo/Transform.h"
+#include "Misc/StringBuilder.h"
 #include "Serialization/ArchiveCrc32.h"
+#include "WorldPartition/WorldPartition.h"
+#include "WorldPartition/DataLayer/DataLayerManager.h"
+
+namespace PCGGridDescriptor
+{
+	void Serialize(FArchive& Ar, const TArray<TSoftObjectPtr<UDataLayerAsset>>& InDataLayerAssets)
+	{
+		for (const TSoftObjectPtr<UDataLayerAsset>& DataLayerAsset : InDataLayerAssets)
+		{
+			FString DataLayerAssetString = DataLayerAsset.ToString();
+			Ar << DataLayerAssetString;
+		}
+	}
+}
 
 bool FPCGGridDescriptor::operator==(const FPCGGridDescriptor& Other) const
 {
-	return GridSize == Other.GridSize && bIs2DGrid == Other.bIs2DGrid && bIsRuntime == Other.bIsRuntime;
+	// We can always compare RuntimeHash (it will be 0 in non game worlds)
+	if (GridSize == Other.GridSize && bIs2DGrid == Other.bIs2DGrid && bIsRuntime == Other.bIsRuntime && RuntimeHash == Other.RuntimeHash)
+	{
+#if WITH_EDITORONLY_DATA
+		return DataLayerAssets == Other.DataLayerAssets;
+#else
+		return true;
+#endif
+	}
+
+	return false;
 }
 
 uint32 FPCGGridDescriptor::ComputeHash() const
 {
-	FArchiveCrc32 Ar;
+	// Initialize Archive with RuntimeHash, in Editor worlds this value will be 0 and the editor only properties will be hashed instead
+	FArchiveCrc32 Ar(RuntimeHash);
 	FPCGGridDescriptor* NonConstThis = const_cast<FPCGGridDescriptor*>(this);
 
 	Ar << NonConstThis->GridSize << NonConstThis->bIs2DGrid << NonConstThis->bIsRuntime;
 
+#if WITH_EDITORONLY_DATA	
+	// Hashing of Editor only properties
+	PCGGridDescriptor::Serialize(Ar, DataLayerAssets);
+#endif
+
+	return Ar.GetCrc();
+}
+
+#if WITH_EDITOR
+uint32 FPCGGridDescriptor::GetRuntimeHash() const
+{
+	if (DataLayerAssets.IsEmpty())
+	{
+		return 0;
+	}
+
+	FArchiveCrc32 Ar;
+	PCGGridDescriptor::Serialize(Ar, DataLayerAssets);
 	return Ar.GetCrc();
 }
 
+FPCGGridDescriptor& FPCGGridDescriptor::SetDataLayerAssets(const TArray<TSoftObjectPtr<UDataLayerAsset>>& InDataLayerAssets)
+{
+	check(!Hash.IsSet());
+	DataLayerAssets = InDataLayerAssets;
+	DataLayerAssets.Sort([](const TSoftObjectPtr<UDataLayerAsset>& A, const TSoftObjectPtr<UDataLayerAsset>& B) { return A.ToString() < B.ToString(); });
+	return *this;
+}
+
+FPCGGridDescriptor& FPCGGridDescriptor::SetDataLayerAssets(const TArray<const UDataLayerAsset*>& InDataLayerAssets)
+{
+	TArray<TSoftObjectPtr<UDataLayerAsset>> DataLayerAssetPtrs;
+	Algo::Transform(InDataLayerAssets, DataLayerAssetPtrs, [](const UDataLayerAsset* DataLayerAsset) { return TSoftObjectPtr<UDataLayerAsset>(DataLayerAsset); });
+	return SetDataLayerAssets(DataLayerAssetPtrs);
+}
+
+void FPCGGridDescriptor::GetDataLayerAssets(TArray<TSoftObjectPtr<UDataLayerAsset>>& OutDataLayerAssets, const UExternalDataLayerAsset*& OutExternalDataLayerAsset) const
+{
+	OutDataLayerAssets.Empty();
+	OutExternalDataLayerAsset = nullptr;
+
+	for (const TSoftObjectPtr<UDataLayerAsset>& DataLayerAssetPtr : DataLayerAssets)
+	{
+		if (UDataLayerAsset* ResolvedAsset = DataLayerAssetPtr.Get())
+		{
+			if (UExternalDataLayerAsset* ExternalDataLayerAsset = Cast<UExternalDataLayerAsset>(ResolvedAsset))
+			{
+				check(!OutExternalDataLayerAsset);
+				OutExternalDataLayerAsset = ExternalDataLayerAsset;
+			}
+			else
+			{
+				OutDataLayerAssets.Add(DataLayerAssetPtr);
+			}
+		}
+	}
+}
+#endif // WITH_EDITOR
+
+FString FPCGGridDescriptor::GetPartitionActorName(const FIntVector& GridCoords) const
+{
+	TStringBuilderWithBuffer<TCHAR, NAME_SIZE> ActorNameBuilder;
+
+	if (IsRuntime())
+	{
+		ActorNameBuilder += TEXT("PCGRuntimePartitionGridActor_");
+	}
+	else
+	{
+		ActorNameBuilder += TEXT("PCGPartitionGridActor_");
+	}
+
+	ActorNameBuilder += FString::Printf(TEXT("%d_"), GetGridSize());
+
+	if (Is2DGrid())
+	{
+		ActorNameBuilder += FString::Printf(TEXT("%d_%d"), GridCoords.X, GridCoords.Y);
+	}
+	else
+	{
+		ActorNameBuilder += FString::Printf(TEXT("%d_%d_%d"), GridCoords.X, GridCoords.Y, GridCoords.Z);
+	}
+
+#if WITH_EDITORONLY_DATA
+	// For Data Layer Instances, build a hash to add as a suffix
+	{
+		FArchiveCrc32 Ar;
+		if (DataLayerAssets.Num() > 0)
+		{
+			PCGGridDescriptor::Serialize(Ar, DataLayerAssets);
+		}
+
+		uint32 NameHash = Ar.GetCrc();
+		if (NameHash != 0)
+		{
+			ActorNameBuilder += FString::Printf(TEXT("_%X"), NameHash);
+		}
+	}
+#endif
+
+	return ActorNameBuilder.ToString();
+}
+
 bool FPCGGridCellDescriptor::operator==(const FPCGGridCellDescriptor& InOther) const
 {
 	return Descriptor == InOther.Descriptor && GridCoords == InOther.GridCoords;
diff --git a/Engine/Plugins/PCG/Source/PCG/Private/Grid/PCGPartitionActor.cpp b/Engine/Plugins/PCG/Source/PCG/Private/Grid/PCGPartitionActor.cpp
--- a/Engine/Plugins/PCG/Source/PCG/Private/Grid/PCGPartitionActor.cpp
+++ b/Engine/Plugins/PCG/Source/PCG/Private/Grid/PCGPartitionActor.cpp
@@ -314,10 +314,25 @@ FIntVector APCGPartitionActor::GetGridCoord() const
 
 FPCGGridDescriptor APCGPartitionActor::GetGridDescriptor() const
 {
-	return FPCGGridDescriptor()
+	FPCGGridDescriptor GridDescriptor = FPCGGridDescriptor()
 		.SetGridSize(GetPCGGridSize())
 		.SetIs2DGrid(bUse2DGrid)
 		.SetIsRuntime(IsRuntimeGenerated());
+
+#if WITH_EDITORONLY_DATA
+	if (GetWorld() && GetWorld()->IsPlayInEditor())
+	{
+		GridDescriptor.SetRuntimeHash(RuntimeGridDescriptorHash);
+	}
+	else
+	{
+		GridDescriptor.SetDataLayerAssets(GetDataLayerAssets());
+	}
+#else
+	GridDescriptor.SetRuntimeHash(RuntimeGridDescriptorHash);
+#endif
+
+	return GridDescriptor;
 }
 
 bool APCGPartitionActor::Teleport(const FVector& NewLocation)
@@ -579,8 +594,21 @@ void APCGPartitionActor::PostCreation(const FPCGGridDescriptor& GridDescriptor)
 	bUse2DGrid = GridDescriptor.Is2DGrid();
 
 #if WITH_EDITOR
+	// Fetch Non External Assets and assign
+	TArray<TSoftObjectPtr<UDataLayerAsset>> DescriptorDataLayerAssets;
+	const UExternalDataLayerAsset* DescriptorExternalDataLayerAsset = nullptr;
+	GridDescriptor.GetDataLayerAssets(DescriptorDataLayerAssets, DescriptorExternalDataLayerAsset);
+
+	// External DataLayer should have been assigned on Spawn
+	ensure(ExternalDataLayerAsset == DescriptorExternalDataLayerAsset);
+
+	DataLayerAssets = DescriptorDataLayerAssets;
+	// Set only once upon creation, can't change for a Partition Actor
+	RuntimeGridDescriptorHash = GridDescriptor.GetRuntimeHash();
 	SetGridSize(PCGGridSize);
 	UpdateBoundsComponentExtents();
+
+	ensure(GetGridDescriptor() == GridDescriptor);
 #endif // WITH_EDITOR
 
 	RegisterPCG();
@@ -686,27 +714,5 @@ FString APCGPartitionActor::GetPCGPartitionActorName(uint32 GridSize, const FInt
 
 FString APCGPartitionActor::GetPCGPartitionActorName(const FPCGGridDescriptor& GridDescriptor, const FIntVector& GridCoords)
 {
-	TStringBuilderWithBuffer<TCHAR, NAME_SIZE> ActorNameBuilder;
-
-	if (GridDescriptor.IsRuntime())
-	{
-		ActorNameBuilder += TEXT("PCGRuntimePartitionGridActor_");
-	}
-	else
-	{
-		ActorNameBuilder += TEXT("PCGPartitionGridActor_");
-	}
-
-	ActorNameBuilder += FString::Printf(TEXT("%d_"), GridDescriptor.GetGridSize());
-	
-	if (GridDescriptor.Is2DGrid())
-	{
-		ActorNameBuilder += FString::Printf(TEXT("%d_%d"), GridCoords.X, GridCoords.Y);
-	}
-	else
-	{
-		ActorNameBuilder += FString::Printf(TEXT("%d_%d_%d"), GridCoords.X, GridCoords.Y, GridCoords.Z);
-	}
-
-	return ActorNameBuilder.ToString();
+	return GridDescriptor.GetPartitionActorName(GridCoords);
 }
\ No newline at end of file
diff --git a/Engine/Plugins/PCG/Source/PCG/Private/Grid/PCGPartitionActorDesc.cpp b/Engine/Plugins/PCG/Source/PCG/Private/Grid/PCGPartitionActorDesc.cpp
--- a/Engine/Plugins/PCG/Source/PCG/Private/Grid/PCGPartitionActorDesc.cpp
+++ b/Engine/Plugins/PCG/Source/PCG/Private/Grid/PCGPartitionActorDesc.cpp
@@ -4,8 +4,11 @@
 
 #if WITH_EDITOR
 #include "Grid/PCGPartitionActor.h"
+
+#include "Algo/Transform.h"
 #include "UObject/FortniteMainBranchObjectVersion.h"
 #include "WorldPartition/WorldPartitionActorDescArchive.h"
+#include "WorldPartition/WorldPartitionActorDescInstance.h"
 
 void FPCGPartitionActorDesc::Init(const AActor* InActor)
 {
@@ -15,12 +18,26 @@ void FPCGPartitionActorDesc::Init(const AActor* InActor)
 	bUse2DGrid = PartitionActor->IsUsing2DGrid();
 }
 
-FPCGGridDescriptor FPCGPartitionActorDesc::GetGridDescriptor() const
+FPCGGridDescriptor FPCGPartitionActorDesc::GetGridDescriptor(const FWorldPartitionActorDescInstance* InActorDescInstance) const
 {
-	return FPCGGridDescriptor()
+	check(InActorDescInstance);
+	FPCGGridDescriptor GridDescriptor = FPCGGridDescriptor()
 		.SetGridSize(GridSize)
 		.SetIs2DGrid(bUse2DGrid)
 		.SetIsRuntime(false);
+
+	if (IsUsingDataLayerAsset())
+	{
+		TArray<TSoftObjectPtr<UDataLayerAsset>> DataLayerAssets;
+		// Because of deprecation DataLayerAssets in ActorDescs are stored in FNames (we can convert them to SoftObjectPtrs if IsUsingDataLayerAsset() returns true)
+		Algo::Transform(GetDataLayers(), DataLayerAssets, [](const FName& InDataLayerAsset) 
+		{
+			return TSoftObjectPtr<UDataLayerAsset>(FSoftObjectPath(InDataLayerAsset.ToString()));
+		});
+		GridDescriptor.SetDataLayerAssets(DataLayerAssets);
+	}
+
+	return GridDescriptor;
 }
 
 FIntVector FPCGPartitionActorDesc::GetGridCoord() const
diff --git a/Engine/Plugins/PCG/Source/PCG/Private/Helpers/PCGActorHelpers.cpp b/Engine/Plugins/PCG/Source/PCG/Private/Helpers/PCGActorHelpers.cpp
--- a/Engine/Plugins/PCG/Source/PCG/Private/Helpers/PCGActorHelpers.cpp
+++ b/Engine/Plugins/PCG/Source/PCG/Private/Helpers/PCGActorHelpers.cpp
@@ -17,6 +17,10 @@
 #include UE_INLINE_GENERATED_CPP_BY_NAME(PCGActorHelpers)
 
 #if WITH_EDITOR
+#include "Engine/Level.h"
+#include "WorldPartition/DataLayer/DataLayerInstance.h"
+#include "WorldPartition/DataLayer/ExternalDataLayerAsset.h"
+#include "WorldPartition/DataLayer/ExternalDataLayerInstance.h"
 #endif
 
 UInstancedStaticMeshComponent* UPCGActorHelpers::GetOrCreateISMC(AActor* InTargetActor, UPCGComponent* InSourceComponent, uint64 SettingsUID, const FPCGISMCBuilderParameters& InParams)
@@ -414,15 +418,30 @@ AActor* UPCGActorHelpers::SpawnDefaultActor(UWorld* World, ULevel* Level, TSubcl
 
 AActor* UPCGActorHelpers::SpawnDefaultActor(UWorld* World, ULevel* Level, TSubclassOf<AActor> ActorClass, const FTransform& Transform, const FActorSpawnParameters& InSpawnParams, AActor* Parent)
 {
-	if (!World || !ActorClass)
+	FActorSpawnParameters ActorSpawnParams(InSpawnParams);
+	if (Level)
+	{
+		ActorSpawnParams.OverrideLevel = Level;
+	}
+
+	FSpawnDefaultActorParams SpawnDefaultActorParams(World, ActorClass, Transform, ActorSpawnParams);
+	
+	SpawnDefaultActorParams.Parent = Parent;
+	
+	return SpawnDefaultActor(SpawnDefaultActorParams);
+}
+
+AActor* UPCGActorHelpers::SpawnDefaultActor(const FSpawnDefaultActorParams& Params)
+{
+	if (!Params.World || !Params.ActorClass)
 	{
 		return nullptr;
 	}
 
-	FActorSpawnParameters SpawnParams = InSpawnParams;
-	if (Level)
+	FActorSpawnParameters SpawnParams = Params.SpawnParams;
+	if (!SpawnParams.OverrideLevel)
 	{
-		SpawnParams.OverrideLevel = Level;
+		SpawnParams.OverrideLevel = Params.World->PersistentLevel;
 	}
 
 	if (PCGHelpers::IsRuntimeOrPIE())
@@ -430,7 +449,22 @@ AActor* UPCGActorHelpers::SpawnDefaultActor(UWorld* World, ULevel* Level, TSubcl
 		SpawnParams.ObjectFlags |= RF_Transient;
 	}
 
-	AActor* NewActor = World->SpawnActor(*ActorClass, &Transform, SpawnParams);
+#if WITH_EDITOR
+	// Find External Data Layer if it exists so we can create scope that will allow actor to be properly created
+	const UExternalDataLayerInstance* ExternalDataLayerInstance = nullptr;
+	for (const UDataLayerInstance* DataLayerInstance : Params.DataLayerInstances)
+	{
+		if (const UExternalDataLayerInstance* Found = Cast<UExternalDataLayerInstance>(DataLayerInstance))
+		{
+			ExternalDataLayerInstance = Found;
+			break;
+		}
+	}
+
+	FScopedOverrideSpawningLevelMountPointObject EDLScope(ExternalDataLayerInstance ? ExternalDataLayerInstance->GetExternalDataLayerAsset() : nullptr);
+#endif
+
+	AActor* NewActor = Params.World->SpawnActor(*Params.ActorClass, &Params.Transform, SpawnParams);
 	
 	if (!NewActor)
 	{
@@ -438,9 +472,18 @@ AActor* UPCGActorHelpers::SpawnDefaultActor(UWorld* World, ULevel* Level, TSubcl
 	}
 
 	// HACK: until UE-62747 is fixed, we have to force set the scale after spawning the actor
-	NewActor->SetActorRelativeScale3D(Transform.GetScale3D());
+	NewActor->SetActorRelativeScale3D(Params.Transform.GetScale3D());
 
 #if WITH_EDITOR
+	// Add remaining DataLayers (except External which was done on spawn)
+	for (const UDataLayerInstance* DataLayerInstance : Params.DataLayerInstances)
+	{
+		if (DataLayerInstance != ExternalDataLayerInstance)
+		{
+			DataLayerInstance->AddActor(NewActor);
+		}
+	}
+	
 	if (SpawnParams.Name != NAME_None)
 	{
 		NewActor->SetActorLabel(SpawnParams.Name.ToString());
@@ -451,7 +494,7 @@ AActor* UPCGActorHelpers::SpawnDefaultActor(UWorld* World, ULevel* Level, TSubcl
 	if (!RootComponent)
 	{
 		RootComponent = NewObject<USceneComponent>(NewActor, USceneComponent::GetDefaultSceneRootVariableName(), RF_Transactional);
-		RootComponent->SetWorldTransform(Transform);
+		RootComponent->SetWorldTransform(Params.Transform);
 
 		NewActor->SetRootComponent(RootComponent);
 		NewActor->AddInstanceComponent(RootComponent);
@@ -465,9 +508,9 @@ AActor* UPCGActorHelpers::SpawnDefaultActor(UWorld* World, ULevel* Level, TSubcl
 	RootComponent->bVisualizeComponent = true;
 #endif // WITH_EDITOR
 
-	if (Parent)
+	if (Params.Parent)
 	{
-		NewActor->AttachToActor(Parent, FAttachmentTransformRules::KeepWorldTransform);
+		NewActor->AttachToActor(Params.Parent, FAttachmentTransformRules::KeepWorldTransform);
 	}
 
 	return NewActor;
diff --git a/Engine/Plugins/PCG/Source/PCG/Private/PCGActorAndComponentMapping.cpp b/Engine/Plugins/PCG/Source/PCG/Private/PCGActorAndComponentMapping.cpp
--- a/Engine/Plugins/PCG/Source/PCG/Private/PCGActorAndComponentMapping.cpp
+++ b/Engine/Plugins/PCG/Source/PCG/Private/PCGActorAndComponentMapping.cpp
@@ -585,7 +585,11 @@ void FPCGActorAndComponentMapping::RegisterPartitionActor(APCGPartitionActor* In
 		// No need to log actor was already flagged to be invalid (and logged a message for it), we can ignore it.
 		return;
 	}
+#endif
+
+	const FPCGGridDescriptor PartitionActorDescriptor = InActor->GetGridDescriptor();
 
+#if WITH_EDITOR
 	// Ignore Invalid Actors 
 	if (InvalidPartitionActors.Contains(InActor->GetActorGuid()))
 	{
@@ -593,23 +597,21 @@ void FPCGActorAndComponentMapping::RegisterPartitionActor(APCGPartitionActor* In
 		UE_LOG(LogPCG, Warning, TEXT("[RegisterPartitionActor] Invalid PCG Partiton Actor '%s' (%s). Please delete actor to remove warning."), *InActor->GetName(), *InActor->GetPackage()->GetName());
 		return;
 	} // Invalidate actor if there is a duplicate for the same grid information
-	else if (FGuid* FoundGuid = PartitionActorRecords.Find({ InActor->GetGridDescriptor(), InActor->GetGridCoord()}); FoundGuid && *FoundGuid != InActor->GetActorGuid())
+	else if (FGuid* FoundGuid = PartitionActorRecords.Find({ PartitionActorDescriptor, InActor->GetGridCoord()}); FoundGuid && *FoundGuid != InActor->GetActorGuid())
 	{
 		InActor->SetInvalidForPCG();
 		UE_LOG(LogPCG, Warning, TEXT("[RegisterPartitionActor] Duplicate PCG Partition Actor '%s' (%s). Please delete actor to remove warning."), *InActor->GetName(), *InActor->GetPackage()->GetName());
 		return;
 	}
 #endif
-
-	const FPCGGridDescriptor Descriptor = InActor->GetGridDescriptor();
 	const FIntVector GridCoord = InActor->GetGridCoord();
 
-	check(Descriptor.GetGridSize() > 0);
+	check(PartitionActorDescriptor.GetGridSize() > 0);
 
 	{
 		FWriteScopeLock WriteLock(PartitionActorsMapLock);
 
-		TMap<FIntVector, TObjectPtr<APCGPartitionActor>>& PartitionActorsMapGrid = PartitionActorsMap.FindOrAdd(Descriptor);
+		TMap<FIntVector, TObjectPtr<APCGPartitionActor>>& PartitionActorsMapGrid = PartitionActorsMap.FindOrAdd(PartitionActorDescriptor);
 		if (PartitionActorsMapGrid.Contains(GridCoord))
 		{
 			return;
@@ -629,28 +631,24 @@ PRAGMA_ENABLE_DEPRECATION_WARNINGS
 #endif // WITH_EDITOR
 
 	// Register to all the components that intersect with the PA. Ignore for runtime generated, it is handled manually
-	if (!Descriptor.IsRuntime())
+	if (!PartitionActorDescriptor.IsRuntime())
 	{
 		FWriteScopeLock WriteLock(ComponentToPartitionActorsMapLock);
-		ForAllIntersectingPartitionedComponents(FBoxCenterAndExtent(InActor->GetFixedBounds()), [this, InActor, bDoComponentMapping, WorldActor](UPCGComponent* Component)
+		ForAllIntersectingPartitionedComponents(FBoxCenterAndExtent(InActor->GetFixedBounds()), [this, InActor, bDoComponentMapping, WorldActor, PartitionActorDescriptor](UPCGComponent* Component)
 		{
 			// For each component, do the mapping if we ask it explicitly, or if the component is generated
 			if (Component && (bDoComponentMapping || Component->bGenerated))
 			{
-				if (TSet<TObjectPtr<APCGPartitionActor>>* PartitionActorsPtr = ComponentToPartitionActorsMap.Find(Component))
+				FPCGGridDescriptor ComponentGridDescriptor = Component->GetGridDescriptor(InActor->GetPCGGridSize());
+				if (ComponentGridDescriptor == PartitionActorDescriptor)
 				{
-					// TODO: This will need to be revisited when execution domains become a broader concept.
-					const bool bSameDomain = !Component->IsManagedByRuntimeGenSystem();
-
-					// Only create the component mapping if actor and component are from the same execution domain.
-					if (bSameDomain)
+					if (TSet<TObjectPtr<APCGPartitionActor>>* PartitionActorsPtr = ComponentToPartitionActorsMap.Find(Component))
 					{
 						bool bHasUnbounded = false;
 						PCGHiGenGrid::FSizeArray GridSizes;
 						ensure(PCGHelpers::GetGenerationGridSizes(Component->GetGraph(), WorldActor, GridSizes, bHasUnbounded));
 
-						// Only create the component mapping if the original component demands the same grid the actor lives on.
-						if (GridSizes.Contains(InActor->GetPCGGridSize()))
+						if (GridSizes.Contains(PartitionActorDescriptor.GetGridSize()))
 						{
 							// In editor we might load/create partition actors while the component is registering. Because of that,
 							// the mapping might not already exists, even if the component is marked generated.
@@ -779,12 +777,10 @@ void FPCGActorAndComponentMapping::UpdateMappingPCGComponentPartitionActor(UPCGC
 		PCGHiGenGrid::FSizeArray GridSizes;
 		ensure(PCGHelpers::GetGenerationGridSizes(InComponent->GetGraph(), WorldActor, GridSizes, bHasUnbounded));
 				
-		const bool bUse2DGrid = InComponent->Use2DGrid();
-		const bool bIsRuntimeGenerated = InComponent->IsManagedByRuntimeGenSystem();
 		TSet<FPCGGridDescriptor> GridDescriptors;
-		Algo::Transform(GridSizes, GridDescriptors, [bUse2DGrid, bIsRuntimeGenerated](uint32 GridSize) 
+		Algo::Transform(GridSizes, GridDescriptors, [InComponent](uint32 GridSize)
 		{ 
-			return FPCGGridDescriptor().SetGridSize(GridSize).SetIs2DGrid(bUse2DGrid).SetIsRuntime(bIsRuntimeGenerated);
+			return InComponent->GetGridDescriptor(GridSize);
 		});
 
 		auto UpdateMapping = [this, InComponent, &Bounds, WorldActor, &GridDescriptors]()
@@ -937,28 +933,18 @@ APCGPartitionActor* FPCGActorAndComponentMapping::GetPartitionActor(const FPCGGr
 	return nullptr;
 }
 
+#if WITH_EDITOR
 bool FPCGActorAndComponentMapping::DoesPartitionActorRecordExist(const FPCGGridDescriptor& GridDescriptor, const FIntVector& GridCoords) const
 {
-#if WITH_EDITOR
 	// In Editor this will not be empty (including -game)
 	if (PartitionActorRecords.Contains({ GridDescriptor, GridCoords }))
 	{
 		return true;
 	}
-#endif
-
-	// In Game Worlds (including PIE) the Subsystem is responsible for transferring the partition actor records to the PCG World Actor so that it can be queried
-	// @todo_pcg: If the partition actor existed on the editor side it will get streamed in with its owning WP Cell, if it didn't it is probably going to get spawned in the persistent level 
-	// which equates to a different overall behavior (one will be streamed-in/out by WP and other will be spawned and never unspawn). So we need to revisit this as it is not a consistent behavior.
-	if (APCGWorldActor* PCGWorldActor = PCGSubsystem->GetPCGWorldActor())
-	{
-		return PCGWorldActor->RuntimePartitionActorRecords.Contains({ GridDescriptor, GridCoords });
-	}
 	
 	return false;
 }
 
-#if WITH_EDITOR
 void FPCGActorAndComponentMapping::RegisterTracking(UPCGComponent* InComponent)
 {
 	// Discard BP templates, local components and invalid component
@@ -1383,7 +1369,7 @@ void FPCGActorAndComponentMapping::OnActorDescInstanceAdded(FWorldPartitionActor
 	if (InActorDescInstance->GetActorNativeClass()->IsChildOf(APCGPartitionActor::StaticClass()))
 	{
 		const FPCGPartitionActorDesc* ActorDesc = static_cast<const FPCGPartitionActorDesc*>(InActorDescInstance->GetActorDesc());
-		FPCGGridCellDescriptor PartitionActorRecord = { ActorDesc->GetGridDescriptor(), FIntVector3(ActorDesc->GridIndexX, ActorDesc->GridIndexY, ActorDesc->GridIndexZ)};
+		FPCGGridCellDescriptor PartitionActorRecord = { ActorDesc->GetGridDescriptor(InActorDescInstance), FIntVector3(ActorDesc->GridIndexX, ActorDesc->GridIndexY, ActorDesc->GridIndexZ)};
 		if (!PartitionActorRecords.Contains(PartitionActorRecord))
 		{
 			PartitionActorRecords.Add(MoveTemp(PartitionActorRecord), ActorDesc->GetGuid());
@@ -1396,7 +1382,7 @@ void FPCGActorAndComponentMapping::OnActorDescInstanceRemoved(FWorldPartitionAct
 	if (InActorDescInstance->GetActorNativeClass()->IsChildOf(APCGPartitionActor::StaticClass()))
 	{
 		const FPCGPartitionActorDesc* ActorDesc = static_cast<const FPCGPartitionActorDesc*>(InActorDescInstance->GetActorDesc());
-		FPCGGridCellDescriptor PartitionActorRecord = { ActorDesc->GetGridDescriptor(), FIntVector3(ActorDesc->GridIndexX, ActorDesc->GridIndexY, ActorDesc->GridIndexZ)};
+		FPCGGridCellDescriptor PartitionActorRecord = { ActorDesc->GetGridDescriptor(InActorDescInstance), FIntVector3(ActorDesc->GridIndexX, ActorDesc->GridIndexY, ActorDesc->GridIndexZ)};
 		if (FGuid* FoundGuid = PartitionActorRecords.Find(PartitionActorRecord); FoundGuid && *FoundGuid == InActorDescInstance->GetGuid())
 		{
 			PartitionActorRecords.Remove(PartitionActorRecord);
@@ -1408,9 +1394,9 @@ void FPCGActorAndComponentMapping::BuildPartitionActorRecords(APCGWorldActor* PC
 {
 	check(WorldPartition);
 	
-	FWorldPartitionHelpers::ForEachActorDescInstance<APCGPartitionActor>(WorldPartition, [PCGWorldActor, &OutPartitionActorRecords, &OutInvalidPartitionActors](const FWorldPartitionActorDescInstance* ActorDescInstance)
+	FWorldPartitionHelpers::ForEachActorDescInstance<APCGPartitionActor>(WorldPartition, [PCGWorldActor, &OutPartitionActorRecords, &OutInvalidPartitionActors](const FWorldPartitionActorDescInstance* InActorDescInstance)
 	{
-		const FPCGPartitionActorDesc* ActorDesc = static_cast<const FPCGPartitionActorDesc*>(ActorDescInstance->GetActorDesc());
+		const FPCGPartitionActorDesc* ActorDesc = static_cast<const FPCGPartitionActorDesc*>(InActorDescInstance->GetActorDesc());
 		
 		// Fixup data
 		if (ActorDesc->bRequiresUse2DGridFixup && PCGWorldActor)
@@ -1429,7 +1415,7 @@ PRAGMA_ENABLE_DEPRECATION_WARNINGS
 		}
 		else // Add to existing valid PCG Partition actors
 		{
-			FPCGGridCellDescriptor PartitionActorRecord = { ActorDesc->GetGridDescriptor(), FIntVector3(ActorDesc->GridIndexX, ActorDesc->GridIndexY, ActorDesc->GridIndexZ)};
+			FPCGGridCellDescriptor PartitionActorRecord = { ActorDesc->GetGridDescriptor(InActorDescInstance), FIntVector3(ActorDesc->GridIndexX, ActorDesc->GridIndexY, ActorDesc->GridIndexZ)};
 			if (!OutPartitionActorRecords.Contains(PartitionActorRecord))
 			{
 				OutPartitionActorRecords.Add(MoveTemp(PartitionActorRecord), ActorDesc->GetGuid());
@@ -1449,24 +1435,6 @@ void FPCGActorAndComponentMapping::BuildPartitionActorRecords()
 		InvalidPartitionActors.Empty();
 		PartitionActorRecords.Empty();
 		BuildPartitionActorRecords(PCGWorldActor, WorldPartition, PartitionActorRecords, InvalidPartitionActors);
-
-		// If we are a PIE world copy cache from the Editor World
-		if (World->IsPlayInEditor() && PCGWorldActor)
-		{
-			FString SourceWorldPath, RemappedWorldPath;
-			if (World->GetSoftObjectPathMapping(SourceWorldPath, RemappedWorldPath))
-			{
-				if (UWorld* EditorWorld = (UWorld*)FindObject<UWorld>(nullptr, *SourceWorldPath))
-				{
-					if (UPCGSubsystem* EditorPCGSubsystem = EditorWorld->GetSubsystem<UPCGSubsystem>())
-					{
-						TArray<FPCGGridCellDescriptor> SerializedPartitionActorRecords;
-						EditorPCGSubsystem->ActorAndComponentMapping.PartitionActorRecords.GenerateKeyArray(SerializedPartitionActorRecords);
-						PCGWorldActor->RuntimePartitionActorRecords = TSet<FPCGGridCellDescriptor>(SerializedPartitionActorRecords);
-					}
-				}
-			}
-		}
 	}
 }
 
diff --git a/Engine/Plugins/PCG/Source/PCG/Private/PCGComponent.cpp b/Engine/Plugins/PCG/Source/PCG/Private/PCGComponent.cpp
--- a/Engine/Plugins/PCG/Source/PCG/Private/PCGComponent.cpp
+++ b/Engine/Plugins/PCG/Source/PCG/Private/PCGComponent.cpp
@@ -46,6 +46,7 @@
 #include "Components/SplineComponent.h"
 #include "Components/SplineMeshComponent.h"
 #include "Engine/Engine.h"
+#include "Engine/Level.h"
 #include "GameFramework/Volume.h"
 #include "Kismet/GameplayStatics.h"
 #include "LandscapeSplinesComponent.h"
@@ -135,6 +136,41 @@ bool UPCGComponent::Use2DGrid() const
 	return GetDefault<UPCGGraph>()->Use2DGrid();
 }
 
+FPCGGridDescriptor UPCGComponent::GetGridDescriptor(uint32 GridSize) const
+{
+	return GetGridDescriptorInternal(GridSize, /*bRuntimeHashUpdate=*/false);
+}
+
+FPCGGridDescriptor UPCGComponent::GetGridDescriptorInternal(uint32 GridSize, bool bRuntimeHashUpdate) const
+{
+	// Return owner descriptor in case of Partition Actors
+	if (APCGPartitionActor* PartitionActorOwner = Cast<APCGPartitionActor>(GetOwner()))
+	{
+		return PartitionActorOwner->GetGridDescriptor();
+	}
+
+	FPCGGridDescriptor PCGGridDescriptor = FPCGGridDescriptor()
+		.SetGridSize(GridSize)
+		.SetIs2DGrid(Use2DGrid())
+		.SetIsRuntime(IsManagedByRuntimeGenSystem());
+
+#if WITH_EDITORONLY_DATA
+	// Only return the RuntimeGridDescriptorHash for PIE Worlds and not when we are updating the Runtime Hash
+	if (GetWorld() && GetWorld()->IsPlayInEditor() && !bRuntimeHashUpdate)
+	{
+		PCGGridDescriptor.SetRuntimeHash(RuntimeGridDescriptorHash);
+	}
+	else
+	{
+		PCGGridDescriptor.SetDataLayerAssets(GetOwner()->GetDataLayerAssets());
+	}
+#else
+	PCGGridDescriptor.SetRuntimeHash(RuntimeGridDescriptorHash);
+#endif
+
+	return PCGGridDescriptor;
+}
+
 void UPCGComponent::SetGraph_Implementation(UPCGGraphInterface* InGraph)
 {
 	SetGraphInterfaceLocal(InGraph);
@@ -842,7 +878,7 @@ void UPCGComponent::NotifyPropertiesChangedFromBlueprint()
 #endif
 }
 
-AActor* UPCGComponent::ClearPCGLink(UClass* TemplateActor)
+AActor* UPCGComponent::ClearPCGLink(UClass* TemplateActorClass)
 {
 	if (!bGenerated || !GetOwner() || !GetWorld())
 	{
@@ -856,9 +892,20 @@ AActor* UPCGComponent::ClearPCGLink(UClass* TemplateActor)
 	}
 
 	UWorld* World = GetWorld();
+	
+	FActorSpawnParameters ActorSpawnParams;
+	UClass* SpawnClass = TemplateActorClass ? TemplateActorClass : AActor::StaticClass();
+	ActorSpawnParams.Name = MakeUniqueObjectName(GetOwner()->GetLevel(), SpawnClass, TEXT("PCGStamp"));
+	ActorSpawnParams.OverrideLevel = GetOwner()->GetLevel();
+
+	UPCGActorHelpers::FSpawnDefaultActorParams SpawnDefaultActorParams(World, SpawnClass, GetOwner()->GetTransform(), ActorSpawnParams);
+
+#if WITH_EDITOR
+	SpawnDefaultActorParams.DataLayerInstances = GetOwner()->GetDataLayerInstances();
+#endif
 
 	// First create a new actor that will be the new owner of all the resources
-	AActor* NewActor = UPCGActorHelpers::SpawnDefaultActor(World, GetOwner()->GetLevel(), TemplateActor ? TemplateActor : AActor::StaticClass(), TEXT("PCGStamp"), GetOwner()->GetTransform());
+	AActor* NewActor = UPCGActorHelpers::SpawnDefaultActor(SpawnDefaultActorParams);
 
 	// Then move all resources linked to this component to this actor
 	bool bHasMovedResources = MoveResourcesToNewActor(NewActor, /*bCreateChild=*/false);
@@ -1020,7 +1067,17 @@ bool UPCGComponent::MoveResourcesToNewActor(AActor* InNewActor, bool bCreateChil
 
 	if (bCreateChild)
 	{
-		NewActor = UPCGActorHelpers::SpawnDefaultActor(GetWorld(), GetOwner()->GetLevel(), NewActor->GetClass(), TEXT("PCGStampChild"), Owner->GetTransform());
+		FActorSpawnParameters ActorSpawnParams;
+		ActorSpawnParams.Name = MakeUniqueObjectName(Owner->GetLevel(), NewActor->GetClass(), TEXT("PCGStampChild"));
+		ActorSpawnParams.OverrideLevel = Owner->GetLevel();
+
+		UPCGActorHelpers::FSpawnDefaultActorParams SpawnDefaultActorParams(GetWorld(), NewActor->GetClass(), Owner->GetTransform(), ActorSpawnParams);
+
+#if WITH_EDITOR
+		SpawnDefaultActorParams.DataLayerInstances = Owner->GetDataLayerInstances();
+#endif
+
+		NewActor = UPCGActorHelpers::SpawnDefaultActor(SpawnDefaultActorParams);
 		NewActor->AttachToActor(InNewActor, FAttachmentTransformRules::KeepWorldTransform);
 		check(NewActor);
 	}
@@ -1469,6 +1526,12 @@ void UPCGComponent::Serialize(FArchive& Ar)
 		GeneratedResourcesCopy = GeneratedResources;
 		GeneratedResources = LoadedPreviewResources;
 	}
+
+	// When duplicating for PIE, we need to Update the RuntimeGridDescriptorHash before duplication for unsaved changes
+	if (Ar.IsSaving() && (Ar.GetPortFlags() & PPF_DuplicateForPIE))
+	{
+		UpdateRuntimeGridDescriptorHash();
+	}
 #endif // WITH_EDITOR
 
 	Ar.UsingCustomVersion(FPCGCustomVersion::GUID);
@@ -1960,6 +2023,27 @@ void UPCGComponent::PostEditImport()
 	SetupCallbacksOnCreation();
 }
 
+void UPCGComponent::PreSave(FObjectPreSaveContext ObjectSaveContext)
+{
+	Super::PreSave(ObjectSaveContext);
+
+	// Update RuntimeGridDescriptorHash on Save (Actor might have changed DataLayers and we need to update)
+	if (!ObjectSaveContext.IsCooking())
+	{
+		UpdateRuntimeGridDescriptorHash();
+	}
+}
+
+void UPCGComponent::UpdateRuntimeGridDescriptorHash()
+{
+	// No need to maintain RuntimeGridDescriptorHash for PCGComponents owned by Partition Actors
+	if (!IsTemplate() && GetOwner() && !GetOwner()->IsA<APCGPartitionActor>())
+	{
+		FPCGGridDescriptor GridDescriptor = GetGridDescriptorInternal(0, /*bRuntimeHashUpdate=*/true);
+		RuntimeGridDescriptorHash = GridDescriptor.GetRuntimeHash();
+	}
+}
+
 void UPCGComponent::PreEditUndo()
 {
 	Super::PreEditUndo();
diff --git a/Engine/Plugins/PCG/Source/PCG/Private/PCGSubsystem.cpp b/Engine/Plugins/PCG/Source/PCG/Private/PCGSubsystem.cpp
--- a/Engine/Plugins/PCG/Source/PCG/Private/PCGSubsystem.cpp
+++ b/Engine/Plugins/PCG/Source/PCG/Private/PCGSubsystem.cpp
@@ -28,6 +28,7 @@
 #include "Misc/ScopedSlowTask.h"
 #include "WorldPartition/WorldPartition.h"
 #include "WorldPartition/WorldPartitionHelpers.h"
+#include "WorldPartition/DataLayer/DataLayerManager.h"
 #else
 #include "Engine/Engine.h"
 #include "Engine/World.h"
@@ -987,10 +988,8 @@ FPCGTaskId UPCGSubsystem::ForAllOverlappingCells(UPCGComponent* InComponent, con
 	TArray<FPCGTaskId> CellTasks;
 	for (uint32 GridSize : GridSizes)
 	{
-		FPCGGridDescriptor Descriptor = FPCGGridDescriptor()
-			.SetGridSize(GridSize)
-			.SetIs2DGrid(InComponent->Use2DGrid())
-			.SetIsRuntime(false);
+		FPCGGridDescriptor Descriptor = InComponent->GetGridDescriptor(GridSize);
+		check(!Descriptor.IsRuntime());
 
 		// In case of 2D grid, we are clamping our bounds in Z to be within 0 and GridSize to create a 2D grid instead of 3D.
 		FBox ModifiedInBounds = InBounds;
@@ -1150,11 +1149,20 @@ APCGPartitionActor* UPCGSubsystem::FindOrCreatePCGPartitionActor(const FPCGGridD
 	}
 	else if (!GridDescriptor.IsRuntime())
 	{
+		// In a Game World PAs need to be Pre-Existing. 
+		// The Original PCG Component will be marked as Generated on load see UPCGComponent::BeginPlay
+		if (GetWorld()->IsGameWorld())
+		{
+			return nullptr;
+		}
+
+#if WITH_EDITOR
 		// Check if there is already an unloaded actor for this cell. RuntimeGenerated PAs are never unloaded, so we ignore them.
 		if (ActorAndComponentMapping.DoesPartitionActorRecordExist(GridDescriptor, GridCoords))
 		{
 			return nullptr;
 		}
+#endif
 	}
 
 	if (!bCanCreateActor)
@@ -1173,6 +1181,15 @@ APCGPartitionActor* UPCGSubsystem::FindOrCreatePCGPartitionActor(const FPCGGridD
 		SpawnParams.ObjectFlags |= RF_Transient;
 		SpawnParams.ObjectFlags &= ~RF_Transactional;
 	}
+		
+#if WITH_EDITOR
+	TArray<TSoftObjectPtr<UDataLayerAsset>> DataLayerAssets;
+	const UExternalDataLayerAsset* ExternalDataLayerAsset = nullptr;
+			
+	GridDescriptor.GetDataLayerAssets(DataLayerAssets, ExternalDataLayerAsset);
+
+	FScopedOverrideSpawningLevelMountPointObject EDLScope(ExternalDataLayerAsset);
+#endif
 
 	const FVector CellCenter(FVector(GridCoords.X + 0.5, GridCoords.Y + 0.5, GridCoords.Z + 0.5) * GridDescriptor.GetGridSize());
 	APCGPartitionActor* NewActor = CastChecked<APCGPartitionActor>(World->SpawnActor(APCGPartitionActor::StaticClass(), &CellCenter, nullptr, SpawnParams));
diff --git a/Engine/Plugins/PCG/Source/PCG/Private/PCGWorldActor.cpp b/Engine/Plugins/PCG/Source/PCG/Private/PCGWorldActor.cpp
--- a/Engine/Plugins/PCG/Source/PCG/Private/PCGWorldActor.cpp
+++ b/Engine/Plugins/PCG/Source/PCG/Private/PCGWorldActor.cpp
@@ -62,18 +62,6 @@ void APCGWorldActor::BeginCacheForCookedPlatformData(const ITargetPlatform* Targ
 		return;
 	}
 
-	// Cache Partition Actor Records at Cook
-	if (UWorldPartition* WorldPartition = World->GetWorldPartition())
-	{
-		TMap<FPCGGridCellDescriptor, FGuid> PartitionActorRecords;
-		TSet<FGuid> InvalidPartitionActors;
-		FPCGActorAndComponentMapping::BuildPartitionActorRecords(this, WorldPartition, PartitionActorRecords, InvalidPartitionActors);
-
-		TArray<FPCGGridCellDescriptor> SerializedPartitionActorRecords;
-		PartitionActorRecords.GenerateKeyArray(SerializedPartitionActorRecords);
-		RuntimePartitionActorRecords = TSet<FPCGGridCellDescriptor>(SerializedPartitionActorRecords);
-	}
-
 	if (LandscapeCacheObject->SerializationMode == EPCGLandscapeCacheSerializationMode::SerializeOnlyAtCook)
 	{
 		// Implementation note: actor references gathered from the world partition helpers will register on creation and unregister on deletion
diff --git a/Engine/Plugins/PCG/Source/PCG/Public/Grid/PCGGridDescriptor.h b/Engine/Plugins/PCG/Source/PCG/Public/Grid/PCGGridDescriptor.h
--- a/Engine/Plugins/PCG/Source/PCG/Public/Grid/PCGGridDescriptor.h
+++ b/Engine/Plugins/PCG/Source/PCG/Public/Grid/PCGGridDescriptor.h
@@ -6,9 +6,14 @@
 
 #include "CoreMinimal.h"
 #include "Misc/Optional.h"
+#include "UObject/SoftObjectPtr.h"
+#include "WorldPartition/DataLayer/DataLayerAsset.h"
 
 #include "PCGGridDescriptor.generated.h"
 
+class FArchive;
+class UExternalDataLayerAsset;
+
 /**
 * Descriptor struct used to determine where to output generated resources
 */
@@ -26,7 +31,16 @@ private:
 
 	UPROPERTY()
 	bool bIsRuntime = false;
+
+#if WITH_EDITORONLY_DATA
+	UPROPERTY()
+	TArray<TSoftObjectPtr<UDataLayerAsset>> DataLayerAssets;
+#endif // WITH_EDITORONLY_DATA
 		
+	// Some descriptor properties do not exist in non editor builds / game worlds. This is the hash that represents those editor only properties.
+	UPROPERTY()
+	uint32 RuntimeHash = 0;
+
 public:
 	FPCGGridDescriptor& SetGridSize(uint32 InGridSize) { check(!Hash.IsSet()); GridSize = InGridSize; return *this; }
 	uint32 GetGridSize() const { return GridSize; }
@@ -38,6 +52,17 @@ public:
 
 	FPCGGridDescriptor& SetIsRuntime(bool bInIsRuntime) { check(!Hash.IsSet()); bIsRuntime = bInIsRuntime; return *this; }
 	bool IsRuntime() const { return bIsRuntime; }
+		
+	FPCGGridDescriptor& SetRuntimeHash(uint32 InRuntimeHash) { check(!Hash.IsSet()); RuntimeHash = InRuntimeHash; return *this; }
+#if WITH_EDITOR
+	uint32 GetRuntimeHash() const;
+
+	FPCGGridDescriptor& SetDataLayerAssets(const TArray<TSoftObjectPtr<UDataLayerAsset>>& InDataLayerAssets);
+	FPCGGridDescriptor& SetDataLayerAssets(const TArray<const UDataLayerAsset*>& InDataLayerAssets);
+	const TArray<TSoftObjectPtr<UDataLayerAsset>>& GetDataLayerAssets() const { return DataLayerAssets; }
+
+	void GetDataLayerAssets(TArray<TSoftObjectPtr<UDataLayerAsset>>& OutDataLayerAssets, const UExternalDataLayerAsset*& OutExternalDataLayerAsset) const;
+#endif // WITH_EDITOR
 
 	bool operator==(const FPCGGridDescriptor& Other) const;
 	bool operator!=(const FPCGGridDescriptor& Other) const { return !(*this == Other); }
@@ -52,6 +77,8 @@ public:
 		return Descriptor.Hash.GetValue();
 	}
 
+	FString GetPartitionActorName(const FIntVector& GridCoords) const;
+
 private:
 	uint32 ComputeHash() const;
 
diff --git a/Engine/Plugins/PCG/Source/PCG/Public/Grid/PCGPartitionActor.h b/Engine/Plugins/PCG/Source/PCG/Public/Grid/PCGPartitionActor.h
--- a/Engine/Plugins/PCG/Source/PCG/Public/Grid/PCGPartitionActor.h
+++ b/Engine/Plugins/PCG/Source/PCG/Public/Grid/PCGPartitionActor.h
@@ -147,6 +147,9 @@ private:
 	UPROPERTY(VisibleAnywhere, Category = WorldPartition)
 	bool bUse2DGrid;
 
+	UPROPERTY()
+	uint32 RuntimeGridDescriptorHash = 0;
+
 #if WITH_EDITORONLY_DATA
 	/** Box component to draw the Partition actor bounds in the Editor viewport */
 	UPROPERTY(Transient)
diff --git a/Engine/Plugins/PCG/Source/PCG/Public/Grid/PCGPartitionActorDesc.h b/Engine/Plugins/PCG/Source/PCG/Public/Grid/PCGPartitionActorDesc.h
--- a/Engine/Plugins/PCG/Source/PCG/Public/Grid/PCGPartitionActorDesc.h
+++ b/Engine/Plugins/PCG/Source/PCG/Public/Grid/PCGPartitionActorDesc.h
@@ -9,12 +9,14 @@
 #include "UObject/ObjectMacros.h"
 #include "WorldPartition/ActorPartition/PartitionActorDesc.h"
 
+class FWorldPartitionActorDescInstance;
+
 class FPCGPartitionActorDesc : public FPartitionActorDesc
 {
 public:
 	FPCGPartitionActorDesc() = default;
 
-	FPCGGridDescriptor GetGridDescriptor() const;
+	FPCGGridDescriptor GetGridDescriptor(const FWorldPartitionActorDescInstance* ActorDescInstance) const;
 	FIntVector GetGridCoord() const;
 protected:
 	virtual void Init(const AActor* InActor) override;
diff --git a/Engine/Plugins/PCG/Source/PCG/Public/Helpers/PCGActorHelpers.h b/Engine/Plugins/PCG/Source/PCG/Public/Helpers/PCGActorHelpers.h
--- a/Engine/Plugins/PCG/Source/PCG/Public/Helpers/PCGActorHelpers.h
+++ b/Engine/Plugins/PCG/Source/PCG/Public/Helpers/PCGActorHelpers.h
@@ -15,6 +15,7 @@
 
 class AActor;
 class UActorComponent;
+class UDataLayerInstance;
 class UInstancedStaticMeshComponent;
 class ULevel;
 class UMaterialInterface;
@@ -127,6 +128,32 @@ public:
 	*/
 	static AActor* SpawnDefaultActor(UWorld* World, ULevel* Level, TSubclassOf<AActor> ActorClass, const FTransform& Transform, const FActorSpawnParameters& SpawnParams, AActor* Parent = nullptr);
 
+	/**
+	* Struct containing all parameters needed to spawn the actor
+	*/
+	struct FSpawnDefaultActorParams
+	{
+		FSpawnDefaultActorParams(UWorld* InWorld, TSubclassOf<AActor> InActorClass, const FTransform& InTransform, const FActorSpawnParameters& InSpawnParams)
+			: World(InWorld), ActorClass(InActorClass), Transform(InTransform), SpawnParams(InSpawnParams)
+		{
+		}
+
+		UWorld* World = nullptr;
+		TSubclassOf<AActor> ActorClass;
+		FTransform Transform;
+		FActorSpawnParameters SpawnParams;
+		AActor* Parent = nullptr;
+#if WITH_EDITOR
+		TArray<const UDataLayerInstance*> DataLayerInstances;
+#endif
+	};
+
+	/**
+	* Spawn a new actor
+	* @param Params struct containing all the parameters needed to spawn the actor
+	*/
+	static AActor* SpawnDefaultActor(const FSpawnDefaultActorParams& Params);
+	
 	/**
 	 * Return the grid cell coordinates on the PCG partition grid given a position and the grid size.
 	 */
diff --git a/Engine/Plugins/PCG/Source/PCG/Public/PCGActorAndComponentMapping.h b/Engine/Plugins/PCG/Source/PCG/Public/PCGActorAndComponentMapping.h
--- a/Engine/Plugins/PCG/Source/PCG/Public/PCGActorAndComponentMapping.h
+++ b/Engine/Plugins/PCG/Source/PCG/Public/PCGActorAndComponentMapping.h
@@ -95,8 +95,10 @@ public:
 
 	APCGPartitionActor* GetPartitionActor(const FPCGGridDescriptor& GridDescriptor, const FIntVector& CellCoords) const;
 
+#if WITH_EDITOR
 	/** Returns true if there is record of a partition actor living in a certain grid cell, regardless of whether or not it is loaded. */
 	bool DoesPartitionActorRecordExist(const FPCGGridDescriptor& GridDescriptor, const FIntVector& GridCoords) const;
+#endif
 private:
 
 #if WITH_EDITOR
diff --git a/Engine/Plugins/PCG/Source/PCG/Public/PCGComponent.h b/Engine/Plugins/PCG/Source/PCG/Public/PCGComponent.h
--- a/Engine/Plugins/PCG/Source/PCG/Public/PCGComponent.h
+++ b/Engine/Plugins/PCG/Source/PCG/Public/PCGComponent.h
@@ -6,6 +6,7 @@
 #include "PCGNode.h"
 #include "PCGSettings.h"
 #include "Graph/PCGStackContext.h"
+#include "Grid/PCGGridDescriptor.h"
 #include "Utils/PCGExtraCapture.h"
 
 #include "ComponentInstanceDataCache.h"
@@ -98,6 +99,7 @@ public:
 #if WITH_EDITOR
 	virtual bool CanEditChange(const FProperty* InProperty) const override;
 	virtual void PostEditImport() override;
+	virtual void PreSave(FObjectPreSaveContext ObjectSaveContext) override;
 #endif
 	/** ~End UObject interface */
 
@@ -201,7 +203,7 @@ public:
 
 	/** Move all generated resources under a new actor, following a template (AActor if not provided), clearing all link to this PCG component. Returns the new actor.*/
 	UFUNCTION(BlueprintCallable, Category = PCG)
-	AActor* ClearPCGLink(UClass* TemplateActor = nullptr);
+	AActor* ClearPCGLink(UClass* TemplateActorClass = nullptr);
 
 	uint32 GetGenerationGridSize() const { return GenerationGridSize; }
 	void SetGenerationGridSize(uint32 InGenerationGridSize) { GenerationGridSize = InGenerationGridSize; }
@@ -410,6 +412,9 @@ public:
 	/** Returns true if component should output on a 2D Grid */
 	bool Use2DGrid() const;
 
+	/** Returns a GridDescriptor based on this component for the specified grid size */
+	FPCGGridDescriptor GetGridDescriptor(uint32 GridSize) const;
+
 	/* Responsibility of the PCG Partition Actor to mark is local */
 	void MarkAsLocalComponent() { bIsComponentLocal = true; }
 
@@ -441,6 +446,8 @@ public:
 	static UPCGData* CreateActorPCGData(AActor* Actor, const UPCGComponent* Component, bool bParseActor = true);
 
 protected:
+	FPCGGridDescriptor GetGridDescriptorInternal(uint32 GridSize, bool bRuntimeHashUpdate) const;
+	
 	void RefreshSchedulingPolicy();
 
 	/** Purges Actors and Components generated by PCG but are no longer managed by any PCG Component */
@@ -483,6 +490,9 @@ private:
 	/** Track if component should disable 'bIsComponentPartitioned'. Used to deprecate partitioned components in non-WP levels before support for partitioning in non-WP levels. */
 	bool bDisableIsComponentPartitionedOnLoad = false;
 #endif // WITH_EDITORONLY_DATA
+	
+	UPROPERTY()
+	uint32 RuntimeGridDescriptorHash = 0;
 
 private:
 	UPCGData* CreatePCGData();
@@ -548,6 +558,8 @@ private:
 	bool DeletePreviewResources();
 
 	static TArray<TSoftObjectPtr<AActor>> GetManagedActorPaths(AActor* InActor);
+
+	void UpdateRuntimeGridDescriptorHash();
 #endif
 
 	FBox GetGridBounds(const AActor* InActor) const;
diff --git a/Engine/Plugins/PCG/Source/PCG/Public/PCGWorldActor.h b/Engine/Plugins/PCG/Source/PCG/Public/PCGWorldActor.h
--- a/Engine/Plugins/PCG/Source/PCG/Public/PCGWorldActor.h
+++ b/Engine/Plugins/PCG/Source/PCG/Public/PCGWorldActor.h
@@ -79,10 +79,6 @@ private:
 #if WITH_EDITOR
 	void OnPartitionGridSizeChanged();
 #endif
-
-	/** Keeps a record of what grid cells contain a serialized partition actor for game worlds. */
-	UPROPERTY()
-	TSet<FPCGGridCellDescriptor> RuntimePartitionActorRecords;
 };
 
 #if UE_ENABLE_INCLUDE_ORDER_DEPRECATED_IN_5_2
