[PCG]
- Centralized tags to attributes processing
- Added node to convert tags to attributes
- Added node to get list of tags in an attribute set
- Added node to get list of attribute
- Added some flexibility around prefix/suffix on add tags to work with tag values

diff --git a/Engine/Plugins/PCG/Source/PCG/Private/Data/PCGPointData.cpp b/Engine/Plugins/PCG/Source/PCG/Private/Data/PCGPointData.cpp
--- a/Engine/Plugins/PCG/Source/PCG/Private/Data/PCGPointData.cpp
+++ b/Engine/Plugins/PCG/Source/PCG/Private/Data/PCGPointData.cpp
@@ -3,6 +3,7 @@
 #include "Data/PCGPointData.h"
 
 #include "Helpers/PCGHelpers.h"
+#include "Helpers/PCGTagHelpers.h"
 #include "Metadata/PCGMetadataAccessor.h"
 #include "Metadata/Accessors/PCGAttributeAccessorHelpers.h"
 #include "Metadata/Accessors/PCGAttributeAccessorKeys.h"
@@ -408,11 +409,6 @@ void UPCGPointData::AddSinglePointFromActor(AActor* InActor, bool* bOutOptionalS
 {
 	check(InActor);
 
-	if (bOutOptionalSanitizedTagAttributeName)
-	{
-		*bOutOptionalSanitizedTagAttributeName = false;
-	}
-
 	FPCGPoint& Point = GetMutablePoints().Emplace_GetRef();
 	Point.Steepness = 1.0f;
 	Point.Transform = InActor->GetActorTransform();
@@ -432,67 +428,21 @@ void UPCGPointData::AddSinglePointFromActor(AActor* InActor, bool* bOutOptionalS
 		ActorReferenceAttribute->SetValue(Point.MetadataEntry, FSoftObjectPath(InActor));
 	}
 
+	bool bSanitizedAttributeNames = false;
+
 	// Parse tags as well
 	for (FName Tag : InActor->Tags)
 	{
-		FString TagString = Tag.ToString();
-		int32 EqualPosition = INDEX_NONE;
-		
-		// Tags that contain a colon will be consider a field:value pair; we'll try to read a number first then default to a string
-		if(TagString.FindChar(':', EqualPosition))
+		PCG::Private::FParseTagResult TagData(Tag);
+		if (PCG::Private::SetAttributeFromTag(TagData, Metadata, Point.MetadataEntry, /*bCanCreateAttribute=*/true))
 		{
-			FString LeftSide = TagString.Left(EqualPosition);
-			FString RightSide = TagString.RightChop(EqualPosition+1);
-
-			if (LeftSide.IsEmpty() || RightSide.IsEmpty())
-			{
-				continue;
-			}
-
-			const bool bSanitized = FPCGMetadataAttributeBase::SanitizeName(LeftSide);
-			if (bOutOptionalSanitizedTagAttributeName)
-			{
-				*bOutOptionalSanitizedTagAttributeName |= bSanitized;
-			}
-
-			if (RightSide.IsNumeric())
-			{
-				if (FPCGMetadataAttribute<double>* Attribute = Metadata->FindOrCreateAttribute<double>(FName(LeftSide), 0.0, /*bAllowsInterpolation=*/false, /*bOverrideParent=*/false, /*bOverwriteIfTypeMismatch=*/false))
-				{
-					double RightSideValue = FCString::Atod(*RightSide);
-					Attribute->SetValue(Point.MetadataEntry, RightSideValue);
-				}
-			}
-			else
-			{
-				if (FPCGMetadataAttribute<FString>* Attribute = Metadata->FindOrCreateAttribute<FString>(FName(LeftSide), FString(), /*bAllowsInterpolation=*/false, /*bOverrideParent=*/false, /*bOverwriteIfTypeMismatch=*/false))
-				{
-					Attribute->SetValue(Point.MetadataEntry, RightSide);
-				}
-			}
+			bSanitizedAttributeNames |= TagData.HasBeenSanitized();
 		}
-		else // Otherwise, consider that the tag is a boolean value
-		{
-			FName SanitizedAttributeName = NAME_None;
-			if (FPCGMetadataAttributeBase::SanitizeName(TagString))
-			{
-				SanitizedAttributeName = FName(TagString);
-
-				if (bOutOptionalSanitizedTagAttributeName)
-				{
-					*bOutOptionalSanitizedTagAttributeName = true;
-				}
-			}
-			else
-			{
-				SanitizedAttributeName = Tag;
-			}
+	}
 
-			if (FPCGMetadataAttribute<bool>* Attribute = Metadata->FindOrCreateAttribute<bool>(SanitizedAttributeName, false, /*bAllowsInterpolation=*/false, /*bOverrideParent=*/false, /*bOverwriteIfTypeMismatch=*/false))
-			{
-				Attribute->SetValue(Point.MetadataEntry, true);
-			}
-		}
+	if (bOutOptionalSanitizedTagAttributeName)
+	{
+		*bOutOptionalSanitizedTagAttributeName = bSanitizedAttributeNames;
 	}
 }
 
diff --git a/Engine/Plugins/PCG/Source/PCG/Private/Elements/PCGAddTag.cpp b/Engine/Plugins/PCG/Source/PCG/Private/Elements/PCGAddTag.cpp
--- a/Engine/Plugins/PCG/Source/PCG/Private/Elements/PCGAddTag.cpp
+++ b/Engine/Plugins/PCG/Source/PCG/Private/Elements/PCGAddTag.cpp
@@ -55,7 +55,20 @@ bool FPCGAddTagElement::ExecuteInternal(FPCGContext* Context) const
 
 	for (const FString& Tag : TagsArray)
 	{
-		const FString TagToAdd = Settings->Prefix + Tag + Settings->Suffix;
+		FString TagToAdd;
+
+		int32 DividerPosition = INDEX_NONE;
+		if (!Settings->bIgnoreTagValueParsing && Tag.FindChar(':', DividerPosition))
+		{
+			FString LeftSide = Tag.Left(DividerPosition);
+			FString RightSide = Tag.RightChop(DividerPosition); // not +1 because we want to keep the ':'
+
+			TagToAdd = Settings->Prefix + LeftSide + Settings->Suffix + RightSide;
+		}
+		else
+		{
+			TagToAdd = Settings->Prefix + Tag + Settings->Suffix;
+		}
 
 		for (FPCGTaggedData& OutputTaggedData : Context->OutputData.TaggedData)
 		{
diff --git a/Engine/Plugins/PCG/Source/PCG/Private/Elements/PCGAttributeReduceElement.cpp b/Engine/Plugins/PCG/Source/PCG/Private/Elements/PCGAttributeReduceElement.cpp
--- a/Engine/Plugins/PCG/Source/PCG/Private/Elements/PCGAttributeReduceElement.cpp
+++ b/Engine/Plugins/PCG/Source/PCG/Private/Elements/PCGAttributeReduceElement.cpp
@@ -125,6 +125,35 @@ namespace PCGAttributeReduceElement
 			});
 		}
 	}
+
+	template<typename T, typename OT>
+	bool Join(const IPCGAttributeAccessorKeys& Keys, const IPCGAttributeAccessor& Accessor, const FString& Delimiter, OT& OutValue)
+	{
+		if constexpr (std::is_constructible_v<OT, FString>)
+		{
+			TArray<FString> StringsToJoin;
+			StringsToJoin.Reserve(Keys.GetNum());
+
+			bool bAppliedAccessor = PCGMetadataElementCommon::ApplyOnAccessor<T>(Keys, Accessor, [&StringsToJoin](const T& InValue, int32)
+			{
+				StringsToJoin.Add(PCG::Private::MetadataTraits<T>::ToString(InValue));
+			});
+
+			if (bAppliedAccessor)
+			{
+				OutValue = OT(FString::Join(StringsToJoin, *Delimiter));
+				return true;
+			}
+			else
+			{
+				return false;
+			}
+		}
+		else
+		{
+			return false;
+		}
+	}
 }
 
 #if WITH_EDITOR
@@ -266,55 +295,96 @@ bool FPCGAttributeReduceElement::ExecuteInternal(FPCGContext* Context) const
 			continue;
 		}
 
-		auto DoOperation = [&Accessor, &Keys, Operation = Settings->Operation, bMergeOutputAttributes = Settings->bMergeOutputAttributes, &OutputParams, &NewAttribute, OutputAttributeName](auto DummyValue) -> bool
+		auto DoOperation = [&Accessor, &Keys, Settings, &OutputParams, &NewAttribute, OutputAttributeName](auto DummyValue) -> bool
 		{
 			using AttributeType = decltype(DummyValue);
 
-			bool bSuccess = false;
-
-			AttributeType OutputValue = PCG::Private::MetadataTraits<AttributeType>::ZeroValue();
+			const EPCGAttributeReduceOperation Operation = Settings->Operation;
+			const bool bMergeOutputAttributes = Settings->bMergeOutputAttributes;
 
+			bool bCreatedNewParams = false;
 			if (!OutputParams || !bMergeOutputAttributes)
 			{
 				OutputParams = NewObject<UPCGParamData>();
-				NewAttribute = OutputParams->Metadata->CreateAttribute<AttributeType>(OutputAttributeName, OutputValue, /*bAllowInterpolation=*/ true, /*bOverrideParent=*/false);
+				bCreatedNewParams = true;
+			}
+
+			auto DoOperationInternal = [&OutputParams, &Accessor, &Keys, Settings, bCreatedNewParams, &NewAttribute, OutputAttributeName](auto DummyOutputValue) -> bool
+			{
+				using OutAttributeType = decltype(DummyOutputValue);
+				const EPCGAttributeReduceOperation Operation = Settings->Operation;
+				const bool bMergeOutputAttributes = Settings->bMergeOutputAttributes;
+				const FString& JoinDelimiter = Settings->JoinDelimiter;
+
+				OutAttributeType OutputValue = PCG::Private::MetadataTraits<OutAttributeType>::ZeroValue();
 
-				if (!NewAttribute)
+				if (bCreatedNewParams)
 				{
-					OutputParams = nullptr;
-					return false;
+					NewAttribute = OutputParams->Metadata->CreateAttribute<OutAttributeType>(OutputAttributeName, OutputValue, /*bAllowInterpolation=*/ true, /*bOverrideParent=*/false);
+
+					if (!NewAttribute)
+					{
+						OutputParams = nullptr;
+						return false;
+					}
 				}
-			}
 
-			FPCGMetadataAttribute<AttributeType>* TypedNewAttribute = static_cast<FPCGMetadataAttribute<AttributeType>*>(NewAttribute);
-			check(TypedNewAttribute);
+				bool bSuccess = false;
 
-			switch (Operation)
+				if constexpr (std::is_same_v<AttributeType, OutAttributeType>)
+				{
+					switch (Operation)
+					{
+					case EPCGAttributeReduceOperation::Average:
+						bSuccess = PCGAttributeReduceElement::Average<AttributeType>(*Keys, *Accessor, OutputValue);
+						break;
+					case EPCGAttributeReduceOperation::Max:
+						bSuccess = PCGAttributeReduceElement::MinMax<AttributeType, /*bIsMin*/false>(*Keys, *Accessor, OutputValue);
+						break;
+					case EPCGAttributeReduceOperation::Min:
+						bSuccess = PCGAttributeReduceElement::MinMax<AttributeType, /*bIsMin*/true>(*Keys, *Accessor, OutputValue);
+						break;
+					case EPCGAttributeReduceOperation::Sum:
+						bSuccess = PCGAttributeReduceElement::Sum<AttributeType>(*Keys, *Accessor, OutputValue);
+						break;
+					case EPCGAttributeReduceOperation::Join:
+						bSuccess = PCGAttributeReduceElement::Join<AttributeType>(*Keys, *Accessor, Settings->JoinDelimiter, OutputValue);
+					default:
+						break;
+					}
+				}
+				else
+				{
+					switch (Operation)
+					{
+					case EPCGAttributeReduceOperation::Join:
+						bSuccess = PCGAttributeReduceElement::Join<AttributeType>(*Keys, *Accessor, Settings->JoinDelimiter, OutputValue);
+					default:
+						break;
+					}
+				}
+
+				if (bSuccess)
+				{
+					FPCGMetadataAttribute<OutAttributeType>* TypedNewAttribute = static_cast<FPCGMetadataAttribute<OutAttributeType>*>(NewAttribute);
+					check(TypedNewAttribute);
+
+					// Implementation note: since the default value does not match the value computed here
+					// and because we might have multiple entries, we need to set it in the attribute
+					TypedNewAttribute->SetValue(OutputParams->Metadata->AddEntry(), OutputValue);
+				}
+
+				return bSuccess;
+			};
+
+			if (Operation == EPCGAttributeReduceOperation::Join)
 			{
-			case EPCGAttributeReduceOperation::Average:
-				bSuccess = PCGAttributeReduceElement::Average<AttributeType>(*Keys, *Accessor, OutputValue);
-				break;
-			case EPCGAttributeReduceOperation::Max:
-				bSuccess = PCGAttributeReduceElement::MinMax<AttributeType, /*bIsMin*/false>(*Keys, *Accessor, OutputValue);
-				break;
-			case EPCGAttributeReduceOperation::Min:
-				bSuccess = PCGAttributeReduceElement::MinMax<AttributeType, /*bIsMin*/true>(*Keys, *Accessor, OutputValue);
-				break;
-			case EPCGAttributeReduceOperation::Sum:
-				bSuccess = PCGAttributeReduceElement::Sum<AttributeType>(*Keys, *Accessor, OutputValue);
-				break;
-			default:
-				break;
+				return DoOperationInternal(FString{});
 			}
-
-			if (bSuccess)
+			else
 			{
-				// Implementation note: since the default value does not match the value computed here
-				// and because we might have multiple entries, we need to set it in the attribute
-				TypedNewAttribute->SetValue(OutputParams->Metadata->AddEntry(), OutputValue);
+				return DoOperationInternal(AttributeType{});
 			}
-
-			return bSuccess;
 		};
 
 		if (!PCGMetadataAttribute::CallbackWithRightType(Accessor->GetUnderlyingType(), DoOperation))
diff --git a/Engine/Plugins/PCG/Source/PCG/Private/Elements/PCGConvertToAttributeSet.cpp b/Engine/Plugins/PCG/Source/PCG/Private/Elements/PCGConvertToAttributeSet.cpp
--- a/Engine/Plugins/PCG/Source/PCG/Private/Elements/PCGConvertToAttributeSet.cpp
+++ b/Engine/Plugins/PCG/Source/PCG/Private/Elements/PCGConvertToAttributeSet.cpp
@@ -5,6 +5,9 @@
 #include "PCGContext.h"
 #include "PCGParamData.h"
 #include "Data/PCGPointData.h"
+#include "Helpers/PCGTagHelpers.h"
+#include "Metadata/PCGMetadata.h"
+#include "Metadata/PCGMetadataAttributeTpl.h"
 
 #include UE_INLINE_GENERATED_CPP_BY_NAME(PCGConvertToAttributeSet)
 
@@ -78,5 +81,50 @@ bool FPCGConvertToAttributeSetElement::ExecuteInternal(FPCGContext* Context) con
 		Output.Data = ParamData;
 	}
 
+	return true;
+}
+
+TArray<FPCGPinProperties> UPCGTagsToAttributeSetSettings::OutputPinProperties() const
+{
+	TArray<FPCGPinProperties> PinProperties;
+	PinProperties.Emplace(PCGPinConstants::DefaultOutputLabel, EPCGDataType::Param);
+
+	return PinProperties;
+}
+
+FPCGElementPtr UPCGTagsToAttributeSetSettings::CreateElement() const
+{
+	return MakeShared<FPCGTagsToAttributeSetElement>();
+}
+
+bool FPCGTagsToAttributeSetElement::ExecuteInternal(FPCGContext* Context) const
+{
+	TRACE_CPUPROFILER_EVENT_SCOPE(FPCGTagsToAttributeSetElement::Execute);
+	check(Context);
+
+	TArray<FPCGTaggedData> Inputs = Context->InputData.GetInputs();
+	TArray<FPCGTaggedData>& Outputs = Context->OutputData.TaggedData;
+
+	for (const FPCGTaggedData& Input : Inputs)
+	{
+		if (!Input.Data)
+		{
+			continue;
+		}
+
+		UPCGParamData* ParamData = NewObject<UPCGParamData>();
+		check(ParamData && ParamData->Metadata);
+
+		FPCGTaggedData& Output = Outputs.Emplace_GetRef(Input);
+		Output.Data = ParamData;
+
+		PCGMetadataEntryKey EntryKey = ParamData->Metadata->AddEntry();
+
+		for (const FString& DataTag : Input.Tags)
+		{
+			PCG::Private::SetAttributeFromTag(DataTag, ParamData->Metadata, EntryKey, /*bCanCreateAttribute=*/true);
+		}
+	}
+
 	return true;
 }
\ No newline at end of file
diff --git a/Engine/Plugins/PCG/Source/PCG/Private/Elements/PCGGetDataInfo.cpp b/Engine/Plugins/PCG/Source/PCG/Private/Elements/PCGGetDataInfo.cpp
new file mode 100644
--- /dev/null
+++ b/Engine/Plugins/PCG/Source/PCG/Private/Elements/PCGGetDataInfo.cpp
@@ -0,0 +1,174 @@
+// Copyright Epic Games, Inc. All Rights Reserved.
+
+#include "Elements/PCGGetDataInfo.h"
+
+#include "PCGContext.h"
+#include "PCGParamData.h"
+#include "Helpers/PCGTagHelpers.h"
+#include "Metadata/PCGMetadata.h"
+#include "Metadata/Accessors/PCGAttributeAccessorHelpers.h"
+#include "Metadata/Accessors/PCGAttributeAccessorKeys.h"
+
+#include UE_INLINE_GENERATED_CPP_BY_NAME(PCGGetDataInfo)
+
+TArray<FPCGPinProperties> UPCGGetTagsSettings::OutputPinProperties() const
+{
+	TArray<FPCGPinProperties> PinProperties;
+	PinProperties.Emplace(PCGPinConstants::DefaultOutputLabel, EPCGDataType::Param);
+
+	return PinProperties;
+}
+
+FPCGElementPtr UPCGGetTagsSettings::CreateElement() const
+{
+	return MakeShared<FPCGGetTagsElement>();
+}
+
+bool FPCGGetTagsElement::ExecuteInternal(FPCGContext* Context) const
+{
+	TRACE_CPUPROFILER_EVENT_SCOPE(UPCGGetTagsSettings::Execute);
+	check(Context);
+
+	const UPCGGetTagsSettings* Settings = Context->GetInputSettings<UPCGGetTagsSettings>();
+
+	TArray<FPCGTaggedData> Inputs = Context->InputData.GetInputsByPin(PCGPinConstants::DefaultInputLabel);
+	TArray<FPCGTaggedData>& Outputs = Context->OutputData.TaggedData;
+
+	for (const FPCGTaggedData& Input : Inputs)
+	{
+		if (!Input.Data)
+		{
+			continue;
+		}
+
+		UPCGParamData* ParamData = NewObject<UPCGParamData>();
+		check(ParamData && ParamData->Metadata);
+
+		FPCGTaggedData& Output = Outputs.Emplace_GetRef(Input);
+		Output.Data = ParamData;
+
+		FPCGMetadataAttribute<FString>* Attribute = ParamData->Metadata->FindOrCreateAttribute<FString>(FName("Attribute"), FString(), /*bAllowsInterpolation=*/false, /*bOverrideParent=*/false, /*bOverwriteIfTypeMismatch=*/true);
+		FPCGMetadataAttribute<FString>* Values = nullptr;
+		if (Settings->bExtractTagValues)
+		{
+			Values = ParamData->Metadata->FindOrCreateAttribute<FString>(FName("Values"), FString(), false, false, true);
+		}
+
+		for (const FString& Tag : Input.Tags)
+		{
+			PCG::Private::FParseTagResult TagData = PCG::Private::ParseTag(Tag);
+			if (TagData.IsValid())
+			{
+				PCGMetadataEntryKey EntryKey = ParamData->Metadata->AddEntry();
+				Attribute->SetValue(EntryKey, TagData.GetOriginalAttribute());
+
+				if (Values)
+				{
+					if (TagData.HasValue())
+					{
+						Values->SetValue(EntryKey, TagData.Value.GetValue());
+					}
+					else // Assume this was a boolean attribute then
+					{
+						Values->SetValue(EntryKey, FString("true"));
+					}
+				}
+			}
+		}
+	}
+
+	return true;
+}
+
+TArray<FPCGPinProperties> UPCGGetAttributesSettings::OutputPinProperties() const
+{
+	TArray<FPCGPinProperties> PinProperties;
+	PinProperties.Emplace(PCGPinConstants::DefaultOutputLabel, EPCGDataType::Param);
+
+	return PinProperties;
+}
+
+FPCGElementPtr UPCGGetAttributesSettings::CreateElement() const
+{
+	return MakeShared<FPCGGetAttributesElement>();
+}
+
+bool FPCGGetAttributesElement::ExecuteInternal(FPCGContext* Context) const
+{
+	TRACE_CPUPROFILER_EVENT_SCOPE(UPCGGetTagsSettings::Execute);
+	check(Context);
+
+	const UPCGGetAttributesSettings* Settings = Context->GetInputSettings<UPCGGetAttributesSettings>();
+
+	TArray<FPCGTaggedData> Inputs = Context->InputData.GetInputsByPin(PCGPinConstants::DefaultInputLabel);
+	TArray<FPCGTaggedData>& Outputs = Context->OutputData.TaggedData;
+
+	for (const FPCGTaggedData& Input : Inputs)
+	{
+		if (!Input.Data)
+		{
+			continue;
+		}
+
+		const UPCGMetadata* InputMetadata = Input.Data->ConstMetadata();
+
+		if (!InputMetadata)
+		{
+			continue;
+		}
+
+		UPCGParamData* ParamData = NewObject<UPCGParamData>();
+		check(ParamData && ParamData->Metadata);
+
+		FPCGTaggedData& Output = Outputs.Emplace_GetRef(Input);
+		Output.Data = ParamData;
+
+		TArray<FName> AttributeNames;
+		TArray<EPCGMetadataTypes> AttributeTypes;
+		InputMetadata->GetAttributes(AttributeNames, AttributeTypes);
+		check(AttributeNames.Num() == AttributeTypes.Num());
+
+		const UEnum* TypesEnum = StaticEnum<EPCGMetadataTypes>();
+
+		FPCGMetadataAttribute<FName>* Attribute = ParamData->Metadata->FindOrCreateAttribute<FName>(FName("Attribute"), NAME_None, /*bAllowsInterpolation=*/false, /*bOverrideParent=*/false, /*bOverwriteIfTypeMismatch=*/true);
+		FPCGMetadataAttribute<FName>* Types = nullptr;
+		if (TypesEnum && Settings->bGetType)
+		{
+			Types = ParamData->Metadata->FindOrCreateAttribute<FName>(FName("Type"), NAME_None, /*bAllowsInterpolation=*/false, /*bOverrideParent=*/false, /*bOverwriteIfTypeMismatch=*/true);
+		}
+
+		FPCGMetadataAttribute<FString>* DefaultValues = nullptr;
+		if (Settings->bGetDefaultValue)
+		{
+			DefaultValues = ParamData->Metadata->FindOrCreateAttribute<FString>(FName("DefaultValue"), FString(), /*bAllowsInterpolation=*/false, /*bOverrideParent=*/false, /*bOverwriteIfTypeMismatch=*/true);
+		}
+
+		for (int AttributeIndex = 0; AttributeIndex < AttributeNames.Num(); ++AttributeIndex)
+		{
+			FName AttributeName = AttributeNames[AttributeIndex];
+			EPCGMetadataTypes AttributeType = AttributeTypes[AttributeIndex];
+
+			PCGMetadataEntryKey EntryKey = ParamData->Metadata->AddEntry();
+			Attribute->SetValue(EntryKey, AttributeName);
+
+			if (Types && TypesEnum)
+			{
+				Types->SetValue(EntryKey, TypesEnum->GetNameByValue(static_cast<__underlying_type(EPCGMetadataTypes)>(AttributeType)));
+			}
+
+			if (DefaultValues)
+			{
+				TUniquePtr<const IPCGAttributeAccessor> AttributeAccessor = PCGAttributeAccessorHelpers::CreateConstAccessor(InputMetadata->GetConstAttribute(AttributeName), InputMetadata, /*bQuiet=*/true);
+				FPCGAttributeAccessorKeysEntries AttributeKeys(PCGInvalidEntryKey);
+				FString DefaultValueString;
+
+				if (AttributeAccessor && AttributeAccessor->Get<FString>(DefaultValueString, 0, AttributeKeys, EPCGAttributeAccessorFlags::AllowBroadcastAndConstructible))
+				{
+					DefaultValues->SetValue(EntryKey, DefaultValueString);
+				}
+			}
+		}
+	}
+
+	return true;
+}
\ No newline at end of file
diff --git a/Engine/Plugins/PCG/Source/PCG/Private/Helpers/PCGTagHelpers.cpp b/Engine/Plugins/PCG/Source/PCG/Private/Helpers/PCGTagHelpers.cpp
new file mode 100644
--- /dev/null
+++ b/Engine/Plugins/PCG/Source/PCG/Private/Helpers/PCGTagHelpers.cpp
@@ -0,0 +1,153 @@
+// Copyright Epic Games, Inc. All Rights Reserved.
+
+#include "Helpers/PCGTagHelpers.h"
+#include "Metadata/PCGMetadataAttributeTpl.h"
+
+namespace PCG
+{
+namespace Private
+{
+	FParseTagResult::FParseTagResult(const FString& InTag)
+	{
+		*this = ParseTag(InTag);
+	}
+
+	FParseTagResult::FParseTagResult(FName InTag)
+	{
+		*this = ParseTag(InTag);
+	}
+
+	FParseTagResult ParseTag(const FString& InTag)
+	{
+		FParseTagResult Result;
+
+		int32 DividerPosition = INDEX_NONE;
+		if (InTag.FindChar(':', DividerPosition))
+		{
+			FString LeftSide = InTag.Left(DividerPosition);
+			FString RightSide = InTag.RightChop(DividerPosition + 1);
+
+			if (LeftSide.IsEmpty())
+			{
+				// Tag doesn't have an attribute - ignore
+			}
+			else if (RightSide.IsEmpty())
+			{
+				// Tag doesn't have a value, ignore the right hand side
+				Result.Attribute = MoveTemp(LeftSide);
+			}
+			else
+			{
+				if (RightSide.IsNumeric())
+				{
+					Result.NumericValue = FCString::Atod(*RightSide);
+				}
+
+				Result.Attribute = MoveTemp(LeftSide);
+				Result.Value = MoveTemp(RightSide);
+			}
+		}
+		else
+		{
+			Result.Attribute = InTag;
+		}
+
+		// Finally, sanitize the attribute name if needed
+		FString OriginalAttribute = Result.Attribute;
+		if (FPCGMetadataAttributeBase::SanitizeName(Result.Attribute))
+		{
+			Result.OriginalAttribute = MoveTemp(OriginalAttribute);
+		}
+
+		return Result;
+	}
+
+	FParseTagResult ParseTag(FName InTag)
+	{
+		return ParseTag(InTag.ToString());
+	}
+
+	bool CreateAttributeFromTag(const FString& InTag, UPCGMetadata* InMetadata, FParseTagResult* OutResult)
+	{
+		FParseTagResult TagData(InTag);
+		const bool bCreateSuccess = CreateAttributeFromTag(TagData, InMetadata);
+
+		if (OutResult)
+		{
+			*OutResult = MoveTemp(TagData);
+		}
+
+		return bCreateSuccess;
+	}
+
+	bool CreateAttributeFromTag(const FParseTagResult& InTagData, UPCGMetadata* InMetadata)
+	{
+		return SetAttributeFromTag(InTagData, InMetadata, PCGInvalidEntryKey, /*bCanCreateAttribute=*/true);
+	}
+
+	bool SetAttributeFromTag(const FString& InTag, UPCGMetadata* InMetadata, PCGMetadataEntryKey InEntryKey, bool bInCanCreateAttribute, FParseTagResult* OutResult)
+	{
+		FParseTagResult TagData = ParseTag(InTag);
+		const bool bSetSuccess = SetAttributeFromTag(TagData, InMetadata, InEntryKey, bInCanCreateAttribute);
+
+		if (OutResult)
+		{
+			*OutResult = MoveTemp(TagData);
+		}
+
+		return bSetSuccess;
+	}
+
+	bool SetAttributeFromTag(const FParseTagResult& TagData, UPCGMetadata* InMetadata, PCGMetadataEntryKey InEntryKey, bool bInCanCreateAttribute)
+	{
+		check(InMetadata);
+
+		if (TagData.IsValid())
+		{
+			FName AttributeName(TagData.Attribute);
+
+			auto GetAttribute = [InMetadata, bInCanCreateAttribute, &AttributeName](auto DefaultValue)
+			{
+				using AttributeType = decltype(DefaultValue);
+				if (bInCanCreateAttribute)
+				{
+					return InMetadata->FindOrCreateAttribute<AttributeType>(AttributeName, DefaultValue, /*bAllowsInterpolation=*/false, /*bOverrideParent=*/false, /*bOverwriteIfTypeMismatch=*/true);
+				}
+				else
+				{
+					return InMetadata->GetMutableTypedAttribute<AttributeType>(AttributeName);
+				}
+			};
+
+			if (TagData.HasNumericValue())
+			{
+				FPCGMetadataAttribute<double>* Attribute = GetAttribute(0.0);
+				if (Attribute && InEntryKey != PCGInvalidEntryKey)
+				{
+					Attribute->SetValue(InEntryKey, TagData.NumericValue.GetValue());
+				}
+			}
+			else if (TagData.HasValue())
+			{
+				FPCGMetadataAttribute<FString>* Attribute = GetAttribute(FString());
+				if (Attribute && InEntryKey != PCGInvalidEntryKey)
+				{
+					Attribute->SetValue(InEntryKey, TagData.Value.GetValue());
+				}
+			}
+			else
+			{
+				FPCGMetadataAttribute<bool>* Attribute = GetAttribute(false);
+				if (Attribute && InEntryKey != PCGInvalidEntryKey)
+				{
+					Attribute->SetValue(InEntryKey, true);
+				}
+			}
+
+			return true;
+		}
+		
+		return false;
+	}
+}
+}
\ No newline at end of file
diff --git a/Engine/Plugins/PCG/Source/PCG/Public/Elements/PCGAddTag.h b/Engine/Plugins/PCG/Source/PCG/Public/Elements/PCGAddTag.h
--- a/Engine/Plugins/PCG/Source/PCG/Public/Elements/PCGAddTag.h
+++ b/Engine/Plugins/PCG/Source/PCG/Public/Elements/PCGAddTag.h
@@ -44,6 +44,10 @@ public:
 	UPROPERTY(EditAnywhere, BlueprintReadWrite, Category = Settings, meta = (PCG_Overridable))
 	FString Suffix;
 
+	/** Controls whether tags are not considered to be key-value pairs, e.g. that the prefix/suffix will be added before the ':' (if any) or not. */
+	UPROPERTY(EditAnywhere, BlueprintReadWrite, Category = Settings, meta = (PCG_Overridable))
+	bool bIgnoreTagValueParsing = false;
+
 	/** Enables deprecated behavior using spaces as separators. Disable to update the node to current behavior. */
 	UE_DEPRECATED(5.5, "bTokenizeOnWhiteSpace has been deprecated.")
 	UPROPERTY(EditAnywhere, Category = Settings, meta = (EditCondition = "bTokenizeOnWhiteSpace", EditConditionHides, DeprecationMessage = "bTokenizeOnWhiteSpace has been deprecated."))
diff --git a/Engine/Plugins/PCG/Source/PCG/Public/Elements/PCGAttributeReduceElement.h b/Engine/Plugins/PCG/Source/PCG/Public/Elements/PCGAttributeReduceElement.h
--- a/Engine/Plugins/PCG/Source/PCG/Public/Elements/PCGAttributeReduceElement.h
+++ b/Engine/Plugins/PCG/Source/PCG/Public/Elements/PCGAttributeReduceElement.h
@@ -14,7 +14,8 @@ enum class EPCGAttributeReduceOperation
 	Average,
 	Max,
 	Min,
-	Sum
+	Sum,
+	Join
 };
 
 /**
@@ -66,6 +67,9 @@ public:
 	UPROPERTY(BlueprintReadWrite, EditAnywhere, Category = Settings, meta = (PCG_Overridable))
 	EPCGAttributeReduceOperation Operation = EPCGAttributeReduceOperation::Average;
 
+	UPROPERTY(BlueprintReadWrite, EditAnywhere, Category = Settings, meta = (PCG_Overridable, EditCondition = "Operation==EPCGAttributeReduceOperation::Join", EditConditionHides))
+	FString JoinDelimiter = FString(", ");
+
 	/** Option to merge all results into a single attribute set with multiple entries, instead of multiple attribute sets with a single value in them.*/
 	UPROPERTY(BlueprintReadWrite, EditAnywhere, Category = Settings, meta = (PCG_Overridable))
 	bool bMergeOutputAttributes = false;
diff --git a/Engine/Plugins/PCG/Source/PCG/Public/Elements/PCGConvertToAttributeSet.h b/Engine/Plugins/PCG/Source/PCG/Public/Elements/PCGConvertToAttributeSet.h
--- a/Engine/Plugins/PCG/Source/PCG/Public/Elements/PCGConvertToAttributeSet.h
+++ b/Engine/Plugins/PCG/Source/PCG/Public/Elements/PCGConvertToAttributeSet.h
@@ -22,7 +22,6 @@ public:
 	virtual bool ShouldDrawNodeCompact() const override { return true; }
 	virtual bool GetCompactNodeIcon(FName& OutCompactNodeIcon) const override;
 #endif
-	
 
 protected:
 	virtual TArray<FPCGPinProperties> InputPinProperties() const override;
@@ -37,6 +36,33 @@ public:
 	// Since this is downstream from points, possibly sampling, it's worth to compute a full CRC so if there are no changes to the attributes, we can maximize cache usage
 	virtual bool ShouldComputeFullOutputDataCrc(FPCGContext* Context) const override { return true; }
 
+protected:
+	virtual bool ExecuteInternal(FPCGContext* Context) const override;
+};
+
+/** Converts tags on data collections to an attribute set (per input data), akin to how it's done on Get Actor Data with the Single Point option. */
+UCLASS(MinimalAPI, BlueprintType, ClassGroup = (Procedural))
+class UPCGTagsToAttributeSetSettings : public UPCGSettings
+{
+	GENERATED_BODY()
+
+public:
+	//~Begin UPCGSettings interface
+#if WITH_EDITOR
+	virtual FName GetDefaultNodeName() const override { return FName(TEXT("TagsToAttributeSet")); }
+	virtual FText GetDefaultNodeTitle() const override { return NSLOCTEXT("PCGTagsToAttributeSetElement", "NodeTitle", "Data Tags To Attribute Set"); }
+	virtual FText GetNodeTooltipText() const override { return NSLOCTEXT("PCGTagsToAttributeSetElement", "NodeTooltip", "Extracts the tags on the data to an attribute set."); }
+	virtual EPCGSettingsType GetType() const override { return EPCGSettingsType::Param; }
+#endif
+
+protected:
+	virtual TArray<FPCGPinProperties> OutputPinProperties() const override;
+	virtual FPCGElementPtr CreateElement() const override;
+	//~End UPCGSettings interface
+};
+
+class FPCGTagsToAttributeSetElement : public IPCGElement
+{
 protected:
 	virtual bool ExecuteInternal(FPCGContext* Context) const override;
 };
\ No newline at end of file
diff --git a/Engine/Plugins/PCG/Source/PCG/Public/Elements/PCGGetDataInfo.h b/Engine/Plugins/PCG/Source/PCG/Public/Elements/PCGGetDataInfo.h
new file mode 100644
--- /dev/null
+++ b/Engine/Plugins/PCG/Source/PCG/Public/Elements/PCGGetDataInfo.h
@@ -0,0 +1,75 @@
+// Copyright Epic Games, Inc. All Rights Reserved.
+
+#pragma once
+
+#include "PCGSettings.h"
+
+#include "PCGGetDataInfo.generated.h"
+
+/** Extracts tags on data to an attribute set with one entry per tag. */
+UCLASS(MinimalAPI, BlueprintType, ClassGroup = (Procedural))
+class UPCGGetTagsSettings : public UPCGSettings
+{
+	GENERATED_BODY()
+
+public:
+	//~Begin UPCGSettings interface
+#if WITH_EDITOR
+	virtual FName GetDefaultNodeName() const override { return FName(TEXT("GetTags")); }
+	virtual FText GetDefaultNodeTitle() const override { return NSLOCTEXT("PCGGetTagsElement", "NodeTitle", "Get Tags"); }
+	virtual FText GetNodeTooltipText() const override { return NSLOCTEXT("PCGGetTagsElement", "NodeTooltip", "Creates an attribute set with one entry per tag"); }
+	virtual EPCGSettingsType GetType() const override { return EPCGSettingsType::Param; }
+#endif
+
+protected:
+	virtual TArray<FPCGPinProperties> OutputPinProperties() const override;
+	virtual FPCGElementPtr CreateElement() const override;
+	//~End UPCGSettings interface
+
+public:
+	/** Creates a 'Values' attribute and stores the values of the valued tags (e.g. 'Tag:Value' tags) as a string. */
+	UPROPERTY(BlueprintReadWrite, EditAnywhere, Category = Settings)
+	bool bExtractTagValues = false;
+};
+
+class FPCGGetTagsElement : public IPCGElement
+{
+protected:
+	virtual bool ExecuteInternal(FPCGContext* Context) const override;
+};
+
+/** Extracts the list of attribute on data to an attribute set with one entry per attribute. */
+UCLASS(MinimalAPI, BlueprintType, ClassGroup = (Procedural))
+class UPCGGetAttributesSettings : public UPCGSettings
+{
+	GENERATED_BODY()
+
+public:
+	//~Begin UPCGSettings interface
+#if WITH_EDITOR
+	virtual FName GetDefaultNodeName() const override { return FName(TEXT("GetAttributeList")); }
+	virtual FText GetDefaultNodeTitle() const override { return NSLOCTEXT("PCGGetAttributesElement", "NodeTitle", "Get Attribute List"); }
+	virtual FText GetNodeTooltipText() const override { return NSLOCTEXT("PCGGetAttributesElement", "NodeTooltip", "Creates an attribute set with one entry per attribute on the input data."); }
+	virtual EPCGSettingsType GetType() const override { return EPCGSettingsType::Param; }
+#endif
+
+protected:
+	virtual TArray<FPCGPinProperties> OutputPinProperties() const override;
+	virtual FPCGElementPtr CreateElement() const override;
+	//~End UPCGSettings interface
+
+public:
+	/** Controls whether the attribute types will be extracted into the 'Type' attribute (as a name). */
+	UPROPERTY(BlueprintReadWrite, EditAnywhere, Category = Settings)
+	bool bGetType = false;
+
+	/** Controls whether the default value for the attribute will be extracted into the 'DefaultValue' attribute (as a string). */
+	UPROPERTY(BlueprintReadWrite, EditAnywhere, Category = Settings)
+	bool bGetDefaultValue = false;
+};
+
+class FPCGGetAttributesElement : public IPCGElement
+{
+protected:
+	virtual bool ExecuteInternal(FPCGContext* Context) const override;
+};
\ No newline at end of file
diff --git a/Engine/Plugins/PCG/Source/PCG/Public/Helpers/PCGTagHelpers.h b/Engine/Plugins/PCG/Source/PCG/Public/Helpers/PCGTagHelpers.h
new file mode 100644
--- /dev/null
+++ b/Engine/Plugins/PCG/Source/PCG/Public/Helpers/PCGTagHelpers.h
@@ -0,0 +1,42 @@
+// Copyright Epic Games, Inc. All Rights Reserved.
+
+#pragma once
+
+#include "PCGCommon.h"
+#include "Metadata/PCGMetadata.h"
+
+namespace PCG
+{
+namespace Private
+{
+	struct PCG_API FParseTagResult
+	{
+		FParseTagResult() = default;
+		explicit FParseTagResult(const FString& InTag);
+		explicit FParseTagResult(FName InTag);
+
+		bool IsValid() const { return !Attribute.IsEmpty(); }
+		bool HasBeenSanitized() const { return OriginalAttribute.IsSet(); }
+		bool HasValue() const { return Value.IsSet(); }
+		bool HasNumericValue() const { return NumericValue.IsSet(); }
+		const FString& GetOriginalAttribute() const { return HasBeenSanitized() ? OriginalAttribute.GetValue() : Attribute; }
+
+		FString Attribute;
+		TOptional<FString> OriginalAttribute;
+		TOptional<FString> Value;
+		TOptional<double> NumericValue;
+	};
+
+	// Builds the tag result structure from the provided tag.
+	PCG_API FParseTagResult ParseTag(const FString& InTag);
+	PCG_API FParseTagResult ParseTag(FName InTag);
+
+	// Parses a tag and creates the corresponding attribute on the provided metadata. If the name is invalid or the predicate rejects it, will return false.
+	PCG_API bool CreateAttributeFromTag(const FString& InTag, UPCGMetadata* InMetadata, FParseTagResult* OutResult = nullptr);
+	PCG_API bool CreateAttributeFromTag(const FParseTagResult& InTagData, UPCGMetadata* InMetadata);
+
+	// Parses a tag, optionally creates the corresponding attribute on the provided metadata and sets the value. If the name is invalid or the predicate rejects it, will return false.
+	PCG_API bool SetAttributeFromTag(const FString& InTag, UPCGMetadata* InMetadata, PCGMetadataEntryKey InKey, bool bInCanCreateAttribute, FParseTagResult* OutResult = nullptr);
+	PCG_API bool SetAttributeFromTag(const FParseTagResult& InTagData, UPCGMetadata* InMetadata, PCGMetadataEntryKey InKey, bool bInCanCreateAttribute);
+}
+}
\ No newline at end of file
diff --git a/Engine/Plugins/PCG/Source/PCG/Public/Metadata/PCGMetadataAttributeTraits.h b/Engine/Plugins/PCG/Source/PCG/Public/Metadata/PCGMetadataAttributeTraits.h
--- a/Engine/Plugins/PCG/Source/PCG/Public/Metadata/PCGMetadataAttributeTraits.h
+++ b/Engine/Plugins/PCG/Source/PCG/Public/Metadata/PCGMetadataAttributeTraits.h
@@ -946,6 +946,11 @@ namespace PCG
 			enum { CanSearchString = true };
 			enum { NeedsConstruction = true };
 
+			static FString ToString(const FString& A)
+			{
+				return A;
+			}
+
 			static bool Equal(const FString& A, const FString& B)
 			{
 				return A == B;
diff --git a/Engine/Plugins/PCG/Source/PCGEditor/Private/PCGLevelToAsset.cpp b/Engine/Plugins/PCG/Source/PCGEditor/Private/PCGLevelToAsset.cpp
--- a/Engine/Plugins/PCG/Source/PCGEditor/Private/PCGLevelToAsset.cpp
+++ b/Engine/Plugins/PCG/Source/PCGEditor/Private/PCGLevelToAsset.cpp
@@ -9,6 +9,7 @@
 #include "Data/PCGPointData.h"
 #include "Helpers/PCGActorHelpers.h"
 #include "Helpers/PCGHelpers.h"
+#include "Helpers/PCGTagHelpers.h"
 #include "Metadata/PCGMetadata.h"
 
 #include "ContentBrowserModule.h"
@@ -235,67 +236,36 @@ bool UPCGLevelToAsset::BP_ExportWorld_Implementation(UWorld* World, const FStrin
 
 		for (FName ActorTag : Actor->Tags)
 		{
-			if (ReservedTags.Contains(ActorTag) || TagToAttributeMap.Contains(ActorTag))
+			PCG::Private::FParseTagResult TagData(ActorTag);
+
+			if (!TagData.IsValid())
 			{
 				continue;
 			}
 
-			FString TagString = ActorTag.ToString();
-			int32 DividerPosition = INDEX_NONE;
-			
-			if (TagString.FindChar(':', DividerPosition))
-			{
-				const FString LeftSide = TagString.Left(DividerPosition);
-				const FString RightSide = TagString.RightChop(DividerPosition + 1);
-
-				if (LeftSide.IsEmpty() || RightSide.IsEmpty())
-				{
-					continue;
-				}
-
-				const FName TagName(LeftSide);
-				FString SanitizedAttributeNameString = LeftSide;
-				const bool bSanitized = FPCGMetadataAttributeBase::SanitizeName(SanitizedAttributeNameString);
-				const FName SanitizedAttributeName(SanitizedAttributeNameString);
-
-				// Once sanitized, multiple tags can map to a single attribute name. The first tag will be used, remaining will be ignored.
-				if (ReservedTags.Contains(SanitizedAttributeName) || TagToAttributeMap.Contains(TagName) || SanitizedAttributeNames.Contains(SanitizedAttributeName))
-				{
-					continue;
-				}
-
-				if (bSanitized)
-				{
-					UE_LOG(LogPCGEditor, Warning, TEXT("Sanitized tag string on actor '%s' to remove invalid characters: '%s' -> '%s'"), *Actor->GetName(), *LeftSide, *SanitizedAttributeNameString);
-				}
-
-				// Otherwise, create the attribute based on the type of the data after the colon.
-				if (RightSide.IsNumeric())
-				{
-					TagToAttributeMap.Add(TagName, PointMetadata->CreateAttribute<double>(SanitizedAttributeName, 0.0f, /*bAllowsInterpolation=*/false, /*bOverrideParent=*/true));
-				}
-				else
-				{
-					TagToAttributeMap.Add(TagName, PointMetadata->CreateAttribute<FString>(SanitizedAttributeName, FString(), /*bAllowsInterpolation=*/false, /*bOverrideParent=*/true));
-				}
+			FName OriginalAttributeName(TagData.GetOriginalAttribute());
+			FName SanitizedAttributeName(TagData.Attribute);
 
-				SanitizedAttributeNames.Add(SanitizedAttributeName);
-			}
-			else
+			// Check if we can safely skip that tag
+			if (ReservedTags.Contains(SanitizedAttributeName) || TagToAttributeMap.Contains(OriginalAttributeName) || SanitizedAttributeNames.Contains(SanitizedAttributeName))
 			{
-				// Simple boolean attribute
-				FString& SanitizedTagString = TagString;
-				FPCGMetadataAttributeBase::SanitizeName(SanitizedTagString);
-				const FName SanitizedTagName(SanitizedTagString);
+				continue;
+			}
 
-				// Once sanitized, multiple tags can map to a single attribute name. The first tag will be used, remaining will be ignored.
-				if (!ReservedTags.Contains(SanitizedTagName) && !SanitizedAttributeNames.Contains(SanitizedTagName))
-				{
-					TagToAttributeMap.Add(ActorTag, PointMetadata->CreateAttribute<bool>(SanitizedTagName, false, /*bAllowsInterpolation=*/false, /*bOverrideParent=*/true));
+			// Try to create the attribute
+			if (!PCG::Private::CreateAttributeFromTag(TagData, PointMetadata))
+			{
+				continue;
+			}
 
-					SanitizedAttributeNames.Add(SanitizedTagName);
-				}
+			// Log warning if we sanitized some values
+			if (TagData.HasBeenSanitized())
+			{
+				UE_LOG(LogPCGEditor, Warning, TEXT("Sanitized tag string on actor '%s' to remove invalid characters: '%s' -> '%s'"), *Actor->GetName(), *TagData.GetOriginalAttribute(), *TagData.Attribute);
 			}
+
+			TagToAttributeMap.Add(OriginalAttributeName, PointMetadata->GetMutableAttribute(SanitizedAttributeName));
+			SanitizedAttributeNames.Add(SanitizedAttributeName);
 		}
 
 		// Prepare actor-level data that's propagated to all points
@@ -337,40 +307,12 @@ bool UPCGLevelToAsset::BP_ExportWorld_Implementation(UWorld* World, const FStrin
 			// create a boolean attribute with the name given by the sanitized tag string.
 			for (FName ActorTag : Actor->Tags)
 			{
-				const FString TagString = ActorTag.ToString();
-				int32 DividerPosition = INDEX_NONE;
+				// Implementation note: we don't set values directly from the tag in the eventuality that there are some name collisions
+				PCG::Private::FParseTagResult TagData(ActorTag);
 
-				if (TagString.FindChar(':', DividerPosition))
-				{
-					const FString LeftSide = TagString.Left(DividerPosition);
-					const FString RightSide = TagString.RightChop(DividerPosition + 1);
-
-					if (LeftSide.IsEmpty() || RightSide.IsEmpty())
-					{
-						continue;
-					}
-
-					const FName TagName(LeftSide);
-					if (FPCGMetadataAttributeBase** Attribute = TagToAttributeMap.Find(TagName))
-					{
-						check(*Attribute);
-						if (RightSide.IsNumeric() && (*Attribute)->GetTypeId() == PCG::Private::MetadataTypes<double>::Id)
-						{
-							double RightSideValue = FCString::Atod(*RightSide);
-							static_cast<FPCGMetadataAttribute<double>*>(*Attribute)->SetValue(Point.MetadataEntry, RightSideValue);
-						}
-						else if ((*Attribute)->GetTypeId() == PCG::Private::MetadataTypes<FString>::Id)
-						{
-							static_cast<FPCGMetadataAttribute<FString>*>(*Attribute)->SetValue(Point.MetadataEntry, RightSide);
-						}
-					}
-				}
-				else if (FPCGMetadataAttributeBase** Attribute = TagToAttributeMap.Find(ActorTag))
+				if (TagData.IsValid() && TagToAttributeMap.Contains(FName(TagData.GetOriginalAttribute())))
 				{
-					if ((*Attribute)->GetTypeId() == PCG::Private::MetadataTypes<bool>::Id)
-					{
-						static_cast<FPCGMetadataAttribute<bool>*>(*Attribute)->SetValue(Point.MetadataEntry, true);
-					}
+					PCG::Private::SetAttributeFromTag(TagData, PointMetadata, Point.MetadataEntry, /*bCanCreateAttribute=*/false);
 				}
 			}
 		};
