[PCG] Avoid trying to delete actors that have no worlds (happens with level reuse and was exposed by the fix 33545539)

Related to UE-215065 where some specific use cases shouldn't even be trying to delete actors (GenerateOnDemand volumes at runtime)

#rnx

diff --git a/Engine/Plugins/PCG/Source/PCG/Private/Helpers/PCGActorHelpers.cpp b/Engine/Plugins/PCG/Source/PCG/Private/Helpers/PCGActorHelpers.cpp
--- a/Engine/Plugins/PCG/Source/PCG/Private/Helpers/PCGActorHelpers.cpp
+++ b/Engine/Plugins/PCG/Source/PCG/Private/Helpers/PCGActorHelpers.cpp
@@ -342,7 +342,8 @@ bool UPCGActorHelpers::DeleteActors(UWorld* World, const TArray<TSoftObjectPtr<A
 		// Not in editor, really unlikely to happen but might be slow
 		for (const TSoftObjectPtr<AActor>& ManagedActor : ActorsToDelete)
 		{
-			if (AActor* Actor = ManagedActor.Get())
+			// @todo_pcg: Revisit this GetWorld() check when fixing UE-215065
+			if (AActor* Actor = ManagedActor.Get(); Actor && Actor->GetWorld())
 			{
 				if (!ensure(World->DestroyActor(Actor)))
 				{
