[PCG] Change back template actor parameter name to prevent issues with BP.

#rnx

diff --git a/Engine/Plugins/PCG/Source/PCG/Public/PCGComponent.h b/Engine/Plugins/PCG/Source/PCG/Public/PCGComponent.h
--- a/Engine/Plugins/PCG/Source/PCG/Public/PCGComponent.h
+++ b/Engine/Plugins/PCG/Source/PCG/Public/PCGComponent.h
@@ -203,7 +203,7 @@ public:
 
 	/** Move all generated resources under a new actor, following a template (AActor if not provided), clearing all link to this PCG component. Returns the new actor.*/
 	UFUNCTION(BlueprintCallable, Category = PCG)
-	AActor* ClearPCGLink(UClass* TemplateActorClass = nullptr);
+	AActor* ClearPCGLink(UClass* TemplateActor = nullptr);
 
 	uint32 GetGenerationGridSize() const { return GenerationGridSize; }
 	void SetGenerationGridSize(uint32 InGenerationGridSize) { GenerationGridSize = InGenerationGridSize; }
