[PCG] Properly support Changelist PCGGraph reverts

- PackageReload was already doing the reference updates but the missing Teardown/SetupCallbacks was the reason the reload and further changes to graph weren't regenerating the graph properly

diff --git a/Engine/Plugins/PCG/Source/PCG/Private/PCGActorAndComponentMapping.cpp b/Engine/Plugins/PCG/Source/PCG/Private/PCGActorAndComponentMapping.cpp
--- a/Engine/Plugins/PCG/Source/PCG/Private/PCGActorAndComponentMapping.cpp
+++ b/Engine/Plugins/PCG/Source/PCG/Private/PCGActorAndComponentMapping.cpp
@@ -31,10 +31,11 @@
 #include "Materials/MaterialInterface.h"
 
 #if WITH_EDITOR
+#include "Grid/PCGPartitionActorDesc.h"
+#include "UObject/PackageReload.h"
 #include "WorldPartition/WorldPartition.h"
 #include "WorldPartition/WorldPartitionActorDescInstance.h"
 #include "WorldPartition/WorldPartitionHelpers.h"
-#include "Grid/PCGPartitionActorDesc.h"
 #endif
 
 namespace PCGActorAndComponentMapping
@@ -1193,6 +1194,7 @@ void FPCGActorAndComponentMapping::RegisterDelegates()
 	FCoreUObjectDelegates::OnObjectPropertyChanged.AddRaw(this, &FPCGActorAndComponentMapping::OnObjectPropertyChanged);
 	FCoreUObjectDelegates::OnObjectModified.AddRaw(this, &FPCGActorAndComponentMapping::OnObjectModified);
 	FCoreUObjectDelegates::OnObjectPreSave.AddRaw(this, &FPCGActorAndComponentMapping::OnObjectSaved);
+	FCoreUObjectDelegates::OnPackageReloaded.AddRaw(this, &FPCGActorAndComponentMapping::OnPackageReloaded);
 
 	UWorld* World = GetWorld();
 	check(World);
@@ -1236,6 +1238,7 @@ void FPCGActorAndComponentMapping::UnregisterDelegates()
 	FCoreUObjectDelegates::OnObjectPropertyChanged.RemoveAll(this);
 	FCoreUObjectDelegates::OnObjectModified.RemoveAll(this);
 	FCoreUObjectDelegates::OnObjectPreSave.RemoveAll(this);
+	FCoreUObjectDelegates::OnPackageReloaded.RemoveAll(this);
 
 	UWorld* World = GetWorld();
 	check(World);
@@ -1364,6 +1367,43 @@ void FPCGActorAndComponentMapping::OnActorAdded(AActor* InActor)
 	OnActorAdded_Internal(InActor, /*bShouldDirty=*/ true);
 }
 
+void FPCGActorAndComponentMapping::OnPackageReloaded(const EPackageReloadPhase InPackageReloadPhase, FPackageReloadedEvent* InPackageReloadedEvent)
+{
+	if (InPackageReloadPhase == EPackageReloadPhase::PrePackageFixup)
+	{
+		for (const auto& RepointedObjectPair : InPackageReloadedEvent->GetRepointedObjects())
+		{
+			if (RepointedObjectPair.Key && RepointedObjectPair.Key->IsAsset() && RepointedObjectPair.Key->IsA<UPCGGraphInterface>())
+			{
+				ForEachObjectOfClass(UPCGGraphInstance::StaticClass(), [OldGraphInterface = RepointedObjectPair.Key](UObject* InObj)
+				{
+					if (UPCGGraphInstance* GraphInstance = Cast<UPCGGraphInstance>(InObj); GraphInstance && GraphInstance->Graph == OldGraphInterface)
+					{
+						GraphInstance->TeardownCallbacks();
+					}
+				});
+			}
+		}
+	}
+	else if (InPackageReloadPhase == EPackageReloadPhase::PostPackageFixup)
+	{
+		for (const auto& RepointedObjectPair : InPackageReloadedEvent->GetRepointedObjects())
+		{
+			if (RepointedObjectPair.Value && RepointedObjectPair.Value->IsAsset() && RepointedObjectPair.Value->IsA<UPCGGraphInterface>())
+			{
+				ForEachObjectOfClass(UPCGGraphInstance::StaticClass(), [NewGraphInterface = RepointedObjectPair.Value](UObject* InObj)
+				{
+					if (UPCGGraphInstance* GraphInstance = Cast<UPCGGraphInstance>(InObj); GraphInstance && GraphInstance->Graph == NewGraphInterface)
+					{
+						GraphInstance->SetupCallbacks();
+						GraphInstance->OnGraphParametersChanged(EPCGGraphParameterEvent::GraphChanged, NAME_None);
+					}
+				});
+			}
+		}
+	}
+}
+
 void FPCGActorAndComponentMapping::OnActorDescInstanceAdded(FWorldPartitionActorDescInstance* InActorDescInstance)
 {
 	if (InActorDescInstance->GetActorNativeClass()->IsChildOf(APCGPartitionActor::StaticClass()))
diff --git a/Engine/Plugins/PCG/Source/PCG/Public/PCGActorAndComponentMapping.h b/Engine/Plugins/PCG/Source/PCG/Public/PCGActorAndComponentMapping.h
--- a/Engine/Plugins/PCG/Source/PCG/Public/PCGActorAndComponentMapping.h
+++ b/Engine/Plugins/PCG/Source/PCG/Public/PCGActorAndComponentMapping.h
@@ -30,6 +30,9 @@ class UPCGSubsystem;
 class UWorld;
 class UWorldPartition;
 class FWorldPartitionActorDescInstance;
+class FPackageReloadedEvent;
+
+enum class EPackageReloadPhase : uint8;
 
 /**
 * This class handle any necessary mapping between actors and pcg components.
@@ -167,6 +170,7 @@ private:
 	/** Return true if the actor is tracked.*/
 	bool IsActorTracked(const AActor* InActor) const;
 
+	void OnPackageReloaded(const EPackageReloadPhase InPackageReloadPhase, FPackageReloadedEvent* InPackageReloadedEvent);
 	void OnActorDescInstanceAdded(FWorldPartitionActorDescInstance* InActorDescInstance);
 	void OnActorDescInstanceRemoved(FWorldPartitionActorDescInstance* InActorDescInstance);
 	void OnActorAdded(AActor* InActor);
