[PCG] Fixed order of inputs being preserved for culled nodes in the graph compiler

#rnx

diff --git a/Engine/Plugins/PCG/Source/PCG/Private/Graph/PCGGraphCompiler.cpp b/Engine/Plugins/PCG/Source/PCG/Private/Graph/PCGGraphCompiler.cpp
--- a/Engine/Plugins/PCG/Source/PCG/Private/Graph/PCGGraphCompiler.cpp
+++ b/Engine/Plugins/PCG/Source/PCG/Private/Graph/PCGGraphCompiler.cpp
@@ -682,28 +682,35 @@ void FPCGGraphCompiler::CullTasks(TArray<FPCGGraphTask>& InOutCompiledTasks, boo
 				continue;
 			}
 
-			const int InputNumBefore = Task.Inputs.Num();
-			for (int InputIndex = 0; InputIndex < InputNumBefore; ++InputIndex)
+			int InputIndex = 0;
+			while (InputIndex < Task.Inputs.Num())
 			{
 				const FPCGTaskId InputTaskId = Task.Inputs[InputIndex].TaskId;
 
 				// Is node culled?
 				if (TaskRemapping[InputTaskId] == InvalidPCGTaskId)
 				{
-					const FPCGGraphTask& InputTask = InOutCompiledTasks[InputTaskId];
-					if (InputTask.Node && InputTask.Node->GetInputPins().Num() > 1)
+					const FPCGGraphTask& CulledInputTask = InOutCompiledTasks[InputTaskId];
+					if (CulledInputTask.Node && CulledInputTask.Node->GetInputPins().Num() > 1)
 					{
 						ensureMsgf(false, TEXT("Task culling currently only supports nodes with a single input pin which are trivial to unwire."));
 						continue;
 					}
 
-					// Upstream node was culled. Wire up the inputs of the culled node to this node.
-					for (int InputInputIndex = 0; InputInputIndex < InputTask.Inputs.Num(); ++InputInputIndex)
+					// Upstream node was culled. To preserve order, insert new wires just after the wire to be culled--which will be removed later.
+					Task.Inputs.Insert(CulledInputTask.Inputs, InputIndex + 1);
+
+					const int NewWireCount = CulledInputTask.Inputs.Num();
+					for (int I = 0; I < NewWireCount; ++I)
 					{
-						FPCGGraphTaskInput& NewInput = Task.Inputs.Add_GetRef(InputTask.Inputs[InputInputIndex]);
-						NewInput.OutPin = Task.Inputs[InputIndex].OutPin;
+						Task.Inputs[InputIndex + 1 + I].OutPin = Task.Inputs[InputIndex].OutPin;
 					}
+
+					// Skip evaluating the newly wired inputs and increment
+					InputIndex += NewWireCount;
 				}
+
+				++InputIndex;
 			}
 		}
 	}
