[PCG] Notification improvements
- Avoid showing notification in Game worlds (except in Simulation)
- Add a Delay for Creating/Releasing the notification to avoid spamming or stacking of Notifications (CVar: pcg.EditorNotificationDelayInSeconds defaults to 3seconds like prior notif system)
- Respect the FSlateNotificationManager::Get().AreNotificationsAllowed() flag

#rnx

diff --git a/Engine/Plugins/PCG/Source/PCG/Private/Graph/PCGGraphExecutor.cpp b/Engine/Plugins/PCG/Source/PCG/Private/Graph/PCGGraphExecutor.cpp
--- a/Engine/Plugins/PCG/Source/PCG/Private/Graph/PCGGraphExecutor.cpp
+++ b/Engine/Plugins/PCG/Source/PCG/Private/Graph/PCGGraphExecutor.cpp
@@ -62,6 +62,11 @@ namespace PCGGraphExecutor
 		TEXT("pcg.EditorFrameTime"),
 		1000.0f / 20.0f,
 		TEXT("Allocated time in ms per frame when running in editor (non pie)"));
+
+	TAutoConsoleVariable<float> CVarEditorNotificationDelayInSeconds(
+		TEXT("pcg.EditorNotificationDelayInSeconds"),
+		3.0f,
+		TEXT("Time in seconds to wait before showing a notification to avoid spamming for short tasks"));
 #endif
 
 	TAutoConsoleVariable<bool> CVarDynamicTaskCulling(
@@ -1750,25 +1755,59 @@ void FPCGGraphExecutor::NotifyGraphChanged(UPCGGraph* InGraph, EPCGChangeType Ch
 
 void FPCGGraphExecutor::UpdateGenerationNotification()
 {
+	// Avoid Notifications for test executors / commandlets / non simulating game world executors
+	if (!World || IsRunningCommandlet() || (World->IsGameWorld() && GEditor && !GEditor->IsSimulateInEditorInProgress()))
+	{
+		return;
+	}
+
 	const int32 RemainingTaskNum = GetNonScheduledRemainingTaskCount();
 	if (RemainingTaskNum == 0)
 	{
-		ReleaseGenerationNotification();
+		// If we had tasks on the last frame, start a timer to avoid releasing the existing notification too soon
+		if (GenerationProgressLastTaskNum != 0)
+		{
+			// Start Release timer
+			GenerationProgressNotificationStartTime = FPlatformTime::Seconds();
+		}
+
+		if ((FPlatformTime::Seconds() - GenerationProgressNotificationStartTime) > PCGGraphExecutor::CVarEditorNotificationDelayInSeconds.GetValueOnAnyThread())
+		{
+			ReleaseGenerationNotification();
+		}
+		else if(GenerationProgressNotification.IsValid())
+		{
+			GenerationProgressNotification.Pin()->Update(0);
+		}
+
+		GenerationProgressLastTaskNum = 0;
 		return;
 	}
 
 	if (!GenerationProgressNotification.IsValid())
 	{
-		IPCGEditorModule* EditorModule = FModuleManager::GetModulePtr<IPCGEditorModule>("PCGEditor");
-		if (!EditorModule)
+		// If we didn't have tasks on the prior frame, start a timer to create the new one to cause a delay and avoid spamming on and off for short tasks
+		if (GenerationProgressLastTaskNum == 0)
 		{
-			return;
+			// Start Create timer
+			GenerationProgressNotificationStartTime = FPlatformTime::Seconds();
 		}
 
-		GenerationProgressNotification = EditorModule->CreateProgressNotification(GetNotificationTextFormat(), /*bCanCancel=*/true);
-		if (GenerationProgressNotification.IsValid())
+		if ((FPlatformTime::Seconds() - GenerationProgressNotificationStartTime) > PCGGraphExecutor::CVarEditorNotificationDelayInSeconds.GetValueOnAnyThread())
 		{
-			GenerationProgressNotification.Pin()->OnCancelTasks().AddRaw(this, &FPCGGraphExecutor::OnNotificationCancel);
+			GenerationProgressNotificationStartTime = 0.0;
+
+			IPCGEditorModule* EditorModule = FModuleManager::GetModulePtr<IPCGEditorModule>("PCGEditor");
+			if (!EditorModule)
+			{
+				return;
+			}
+
+			GenerationProgressNotification = EditorModule->CreateProgressNotification(GetNotificationTextFormat(), /*bCanCancel=*/true);
+			if (GenerationProgressNotification.IsValid())
+			{
+				GenerationProgressNotification.Pin()->OnCancelTasks().AddRaw(this, &FPCGGraphExecutor::OnNotificationCancel);
+			}
 		}
 	}
 
@@ -1776,6 +1815,8 @@ void FPCGGraphExecutor::UpdateGenerationNotification()
 	{
 		GenerationProgressNotification.Pin()->Update(RemainingTaskNum);
 	}
+
+	GenerationProgressLastTaskNum = RemainingTaskNum;
 }
 
 void FPCGGraphExecutor::ReleaseGenerationNotification()
@@ -1789,6 +1830,7 @@ void FPCGGraphExecutor::ReleaseGenerationNotification()
 		}
 		GenerationProgressNotification = nullptr;
 	}
+	GenerationProgressNotificationStartTime = 0.0;
 }
 
 void FPCGGraphExecutor::OnNotificationCancel()
diff --git a/Engine/Plugins/PCG/Source/PCG/Private/Graph/PCGGraphExecutor.h b/Engine/Plugins/PCG/Source/PCG/Private/Graph/PCGGraphExecutor.h
--- a/Engine/Plugins/PCG/Source/PCG/Private/Graph/PCGGraphExecutor.h
+++ b/Engine/Plugins/PCG/Source/PCG/Private/Graph/PCGGraphExecutor.h
@@ -281,6 +281,8 @@ private:
 
 	int32 ReleaseActorsCountUntilGC = 30;
 	TWeakPtr<IPCGEditorProgressNotification> GenerationProgressNotification;
+	double GenerationProgressNotificationStartTime = 0.0;
+	int32 GenerationProgressLastTaskNum = 0;
 
 	int32 TidyCacheCountUntilGC = 100;
 #endif
diff --git a/Engine/Plugins/PCG/Source/PCGEditor/Private/PCGEditorModule.cpp b/Engine/Plugins/PCG/Source/PCGEditor/Private/PCGEditorModule.cpp
--- a/Engine/Plugins/PCG/Source/PCGEditor/Private/PCGEditorModule.cpp
+++ b/Engine/Plugins/PCG/Source/PCGEditor/Private/PCGEditorModule.cpp
@@ -43,6 +43,7 @@
 #include "Details/PCGInstancedPropertyBagOverrideDetails.h"
 #include "Details/PCGVolumeDetails.h"
 #include "Framework/MultiBox/MultiBoxBuilder.h"
+#include "Framework/Notifications/NotificationManager.h"
 
 #define LOCTEXT_NAMESPACE "FPCGEditorModule"
 
@@ -143,6 +144,11 @@ bool FPCGEditorModule::SupportsDynamicReloading()
 
 TWeakPtr<IPCGEditorProgressNotification> FPCGEditorModule::CreateProgressNotification(const FTextFormat& TextFormat, bool bCanCancel)
 {
+	if (!FSlateNotificationManager::Get().AreNotificationsAllowed())
+	{
+		return nullptr;
+	}
+
 	TSharedPtr<IPCGEditorProgressNotification> NewNotification = MakeShared<FPCGEditorProgressNotification>(TextFormat, bCanCancel);
 	ActiveNotifications.Add(NewNotification);
 	return NewNotification.ToWeakPtr();
