[PCG] PCG Execute Blueprint Node
- Avoid registration to FCoreUObjectDelegates::OnObjectPropertyChanged when element has no dependencies
- Avoid building dependencies for execution context duplicated BlueprintElements
- Fix BP Recompile not rebuilding dependencies
- Hold dependencies as Weak ptrs as we don't control their lifetime

diff --git a/Engine/Plugins/PCG/Source/PCG/Private/Elements/PCGExecuteBlueprint.cpp b/Engine/Plugins/PCG/Source/PCG/Private/Elements/PCGExecuteBlueprint.cpp
--- a/Engine/Plugins/PCG/Source/PCG/Private/Elements/PCGExecuteBlueprint.cpp
+++ b/Engine/Plugins/PCG/Source/PCG/Private/Elements/PCGExecuteBlueprint.cpp
@@ -14,6 +14,10 @@
 #include "UObject/Package.h"
 #include "UObject/UObjectThreadContext.h"
 
+#if WITH_EDITOR
+#include "Algo/Transform.h"
+#endif
+
 #include UE_INLINE_GENERATED_CPP_BY_NAME(PCGExecuteBlueprint)
 
 #define LOCTEXT_NAMESPACE "PCGBlueprintElement"
@@ -21,14 +25,18 @@
 #if WITH_EDITOR
 namespace PCGBlueprintHelper
 {
-	TSet<TObjectPtr<UObject>> GetDataDependencies(UPCGBlueprintElement* InElement, int32 MaxDepth)
+	TSet<TWeakObjectPtr<UObject>> GetDataDependencies(UPCGBlueprintElement* InElement, int32 MaxDepth)
 	{
 		check(InElement && InElement->GetClass());
 		UClass* BPClass = InElement->GetClass();
 
 		TSet<TObjectPtr<UObject>> Dependencies;
 		PCGHelpers::GatherDependencies(InElement, Dependencies, MaxDepth);
-		return Dependencies;
+		TSet<TWeakObjectPtr<UObject>> WeakDependencies;
+
+		Algo::Transform(Dependencies, WeakDependencies, [](const TObjectPtr<UObject>& InObjectPtr) { return TWeakObjectPtr<UObject>(InObjectPtr); });
+
+		return WeakDependencies;
 	}
 }
 #endif // WITH_EDITOR
@@ -91,7 +99,11 @@ void UPCGBlueprintElement::PostLoad()
 void UPCGBlueprintElement::BeginDestroy()
 {
 #if WITH_EDITOR
-	FCoreUObjectDelegates::OnObjectPropertyChanged.RemoveAll(this);
+	if (!DataDependencies.IsEmpty())
+	{
+		FCoreUObjectDelegates::OnObjectPropertyChanged.RemoveAll(this);
+		DataDependencies.Empty();
+	}
 #endif
 
 	Super::BeginDestroy();
@@ -105,8 +117,7 @@ void UPCGBlueprintElement::ExecuteWithContext_Implementation(FPCGContext& InCont
 void UPCGBlueprintElement::Initialize()
 {
 #if WITH_EDITOR
-	FCoreUObjectDelegates::OnObjectPropertyChanged.AddUObject(this, &UPCGBlueprintElement::OnDependencyChanged);
-	DataDependencies = PCGBlueprintHelper::GetDataDependencies(this, DependencyParsingDepth);
+	UpdateDependencies();
 #endif
 }
 
@@ -127,10 +138,37 @@ void UPCGBlueprintElement::PostEditChangeProperty(FPropertyChangedEvent& Propert
 {
 	Super::PostEditChangeProperty(PropertyChangedEvent);
 
+	UpdateDependencies();
+
+	OnBlueprintChangedDelegate.Broadcast(this);
+}
+
+void UPCGBlueprintElement::UpdateDependencies()
+{
+	// Avoid calculating dependencies for graph execution element
+	if (!GetOuter()->IsA<UPCGBlueprintSettings>())
+	{
+		return;
+	}
+
+	// Backup to know if we need to unregister from the Delegate or not
+	const bool bHadDependencies = DataDependencies.Num() > 0;
+
 	// Since we don't really know what changed, let's just rebuild our data dependencies
 	DataDependencies = PCGBlueprintHelper::GetDataDependencies(this, DependencyParsingDepth);
 
-	OnBlueprintChangedDelegate.Broadcast(this);
+	// Only Bind to event if we do have dependencies
+	if (!DataDependencies.IsEmpty())
+	{
+		if (!bHadDependencies)
+		{
+			FCoreUObjectDelegates::OnObjectPropertyChanged.AddUObject(this, &UPCGBlueprintElement::OnDependencyChanged);
+		}
+	}
+	else if(bHadDependencies)
+	{
+		FCoreUObjectDelegates::OnObjectPropertyChanged.RemoveAll(this);
+	}
 }
 
 void UPCGBlueprintElement::OnDependencyChanged(UObject* Object, FPropertyChangedEvent& PropertyChangedEvent)
@@ -444,6 +482,12 @@ void UPCGBlueprintSettings::OnBlueprintChanged(UBlueprint* InBlueprint)
 	TeardownBlueprintElementEvent();
 	SetupBlueprintElementEvent();
 
+	// We are responsible for calling Initialize on the newly recreated element
+	if (BlueprintElementInstance)
+	{
+		BlueprintElementInstance->Initialize();
+	}
+
 	// Also, reconstruct overrides
 	InitializeCachedOverridableParams(/*bReset=*/true);
 
diff --git a/Engine/Plugins/PCG/Source/PCG/Public/Elements/PCGExecuteBlueprint.h b/Engine/Plugins/PCG/Source/PCG/Public/Elements/PCGExecuteBlueprint.h
--- a/Engine/Plugins/PCG/Source/PCG/Public/Elements/PCGExecuteBlueprint.h
+++ b/Engine/Plugins/PCG/Source/PCG/Public/Elements/PCGExecuteBlueprint.h
@@ -281,7 +281,8 @@ public:
 protected:
 #if WITH_EDITOR
 	void OnDependencyChanged(UObject* Object, FPropertyChangedEvent& PropertyChangedEvent);
-	TSet<TObjectPtr<UObject>> DataDependencies;
+	void UpdateDependencies();
+	TSet<TWeakObjectPtr<UObject>> DataDependencies;
 #endif
 
 #if !WITH_EDITORONLY_DATA
