[PCG] Buildfix

#rnx

diff --git a/Engine/Plugins/PCG/Source/PCG/Private/Elements/PCGGetDataInfo.cpp b/Engine/Plugins/PCG/Source/PCG/Private/Elements/PCGGetDataInfo.cpp
--- a/Engine/Plugins/PCG/Source/PCG/Private/Elements/PCGGetDataInfo.cpp
+++ b/Engine/Plugins/PCG/Source/PCG/Private/Elements/PCGGetDataInfo.cpp
@@ -162,7 +162,7 @@ bool FPCGGetAttributesElement::ExecuteInternal(FPCGContext* Context) const
 				FPCGAttributeAccessorKeysEntries AttributeKeys(PCGInvalidEntryKey);
 				FString DefaultValueString;
 
-				if (AttributeAccessor && AttributeAccessor->Get<FString>(DefaultValueString, 0, AttributeKeys, EPCGAttributeAccessorFlags::AllowBroadcastAndConstructible))
+				if (AttributeAccessor && AttributeAccessor->Get<FString>(DefaultValueString, 0, AttributeKeys, EPCGAttributeAccessorFlags::AllowBroadcast | EPCGAttributeAccessorFlags::AllowConstructible))
 				{
 					DefaultValues->SetValue(EntryKey, DefaultValueString);
 				}
