[PCG] Fixed major inefficiency in checking for loops when trying to connect edges

#rnx

diff --git a/Engine/Plugins/PCG/Source/PCGEditor/Private/PCGEditorGraphNodeBase.cpp b/Engine/Plugins/PCG/Source/PCGEditor/Private/PCGEditorGraphNodeBase.cpp
--- a/Engine/Plugins/PCG/Source/PCGEditor/Private/PCGEditorGraphNodeBase.cpp
+++ b/Engine/Plugins/PCG/Source/PCGEditor/Private/PCGEditorGraphNodeBase.cpp
@@ -711,7 +711,8 @@ bool UPCGEditorGraphNodeBase::IsCompatible(const UPCGPin* InputPin, const UPCGPi
 	if (PCGEditorGraphSwitches::CVarCheckConnectionCycles.GetValueOnAnyThread() && InputPin && OutputPin && InputPin->Node == PCGNode)
 	{
 		// Upstream Visitor
-		auto Visitor = [ThisPCGNode = PCGNode](const UPCGNode* InNode, auto VisitorLambda) -> bool
+		TSet<const UPCGNode*> VisitedNodes;
+		auto Visitor = [&VisitedNodes, ThisPCGNode = PCGNode](const UPCGNode* InNode, auto VisitorLambda) -> bool
 		{
 			if (InNode)
 			{
@@ -719,6 +720,12 @@ bool UPCGEditorGraphNodeBase::IsCompatible(const UPCGPin* InputPin, const UPCGPi
 				{
 					return false;
 				}
+				else if (VisitedNodes.Contains(InNode))
+				{
+					return true;
+				}
+
+				VisitedNodes.Add(InNode);
 
 				for (const TObjectPtr<UPCGPin>& InputPin : InNode->GetInputPins())
 				{
