[PCG] Fixed subgraph collapse with respect to label changes. Also adjusted subgraph collapse test.

#rnx

diff --git a/Engine/Plugins/PCG/Source/PCG/Private/Helpers/PCGSubgraphHelpers.cpp b/Engine/Plugins/PCG/Source/PCG/Private/Helpers/PCGSubgraphHelpers.cpp
--- a/Engine/Plugins/PCG/Source/PCG/Private/Helpers/PCGSubgraphHelpers.cpp
+++ b/Engine/Plugins/PCG/Source/PCG/Private/Helpers/PCGSubgraphHelpers.cpp
@@ -509,7 +509,7 @@ namespace PCGSubgraphHelpersExtra
 			check(InNode->GetOutputPins().Num() == 1);
 			UPCGPin* GraphParameterPin = InNode->GetOutputPins()[0];
 			check(GraphParameterPin);
-			OutsideSubgraphPinToSubgraphPinLabel.Emplace(GraphParameterPin, FName(FName::NameToDisplayString(GraphParameterPin->Properties.Label.ToString(), /*bIsBool=*/ false)));
+			OutsideSubgraphPinToSubgraphPinLabel.Emplace(GraphParameterPin, GraphParameterPin->Properties.Label);
 
 			return true;
 		}
