[PCG] By default hide runtime gen actors from outliner, as outliner updates can incur significant costs in Slate code.

[RN] By default hide runtime generated partition actors from Scene Outliner to fix significant CPU costs incurred in updating the tree view.

diff --git a/Engine/Plugins/PCG/Source/PCG/Private/PCGSubsystem.cpp b/Engine/Plugins/PCG/Source/PCG/Private/PCGSubsystem.cpp
--- a/Engine/Plugins/PCG/Source/PCG/Private/PCGSubsystem.cpp
+++ b/Engine/Plugins/PCG/Source/PCG/Private/PCGSubsystem.cpp
@@ -1124,7 +1124,7 @@ APCGPartitionActor* UPCGSubsystem::GetRegisteredPCGPartitionActor(const FPCGGrid
 	return ActorAndComponentMapping.GetPartitionActor(GridDescriptor, GridCoords);
 }
 
-APCGPartitionActor* UPCGSubsystem::FindOrCreatePCGPartitionActor(const FPCGGridDescriptor& GridDescriptor, const FIntVector& GridCoords, bool bCanCreateActor) const
+APCGPartitionActor* UPCGSubsystem::FindOrCreatePCGPartitionActor(const FPCGGridDescriptor& GridDescriptor, const FIntVector& GridCoords, bool bCanCreateActor, bool bHideFromOutliner) const
 {
 	TRACE_CPUPROFILER_EVENT_SCOPE(UPCGActorHelpers::FindOrCreatePCGPartitionActor);
 
@@ -1174,6 +1174,7 @@ APCGPartitionActor* UPCGSubsystem::FindOrCreatePCGPartitionActor(const FPCGGridD
 #if WITH_EDITOR
 	SpawnParams.Name = *APCGPartitionActor::GetPCGPartitionActorName(GridDescriptor, GridCoords);
 	SpawnParams.NameMode = FActorSpawnParameters::ESpawnActorNameMode::Required_Fatal;
+	SpawnParams.bHideFromSceneOutliner = bHideFromOutliner;
 #endif
 	if (GridDescriptor.IsRuntime())
 	{
diff --git a/Engine/Plugins/PCG/Source/PCG/Private/RuntimeGen/PCGRuntimeGenScheduler.cpp b/Engine/Plugins/PCG/Source/PCG/Private/RuntimeGen/PCGRuntimeGenScheduler.cpp
--- a/Engine/Plugins/PCG/Source/PCG/Private/RuntimeGen/PCGRuntimeGenScheduler.cpp
+++ b/Engine/Plugins/PCG/Source/PCG/Private/RuntimeGen/PCGRuntimeGenScheduler.cpp
@@ -55,6 +55,29 @@ namespace PCGRuntimeGenSchedulerHelpers
 		TEXT("pcg.RuntimeGeneration.BasePoolSize"),
 		100,
 		TEXT("Defines the base PartitionActor pool size for the RuntimeGeneration system. Cannot be less than 1."));
+
+	static FAutoConsoleCommand CommandFlushActorPool(
+		TEXT("pcg.RuntimeGeneration.FlushActorPool"),
+		TEXT("Flushes all pooled actors and regenerates all components."),
+		FConsoleCommandDelegate::CreateLambda([]()
+		{
+			if (UPCGSubsystem* PCGSubsystem = UPCGSubsystem::GetSubsystemForCurrentWorld())
+			{
+				PCGSubsystem->GetRuntimeGenScheduler()->FlushAllGeneratedActors();
+			}
+		}));
+
+	static TAutoConsoleVariable<bool> CVarHideActorsFromOutliner(
+		TEXT("pcg.RuntimeGeneration.HideActorsFromOutliner"),
+		true,
+		TEXT("Hides partition actors from Scene Outliner."),
+		FConsoleVariableDelegate::CreateLambda([](IConsoleVariable*)
+		{
+			if (UPCGSubsystem* PCGSubsystem = UPCGSubsystem::GetSubsystemForCurrentWorld())
+			{
+				PCGSubsystem->GetRuntimeGenScheduler()->FlushAllGeneratedActors();
+			}
+		}));
 }
 
 FPCGGridDescriptor FPCGRuntimeGenScheduler::FGridGenerationKey::GetGridDescriptor() const
@@ -623,7 +646,11 @@ void FPCGRuntimeGenScheduler::TickScheduleGeneration(TMap<FGridGenerationKey, do
 					}
 
 					// Find or Create RuntimeGenPA.
-					PartitionActor = Subsystem->FindOrCreatePCGPartitionActor(GridDescriptor, GridCoords);
+					PartitionActor = Subsystem->FindOrCreatePCGPartitionActor(
+						GridDescriptor,
+						GridCoords,
+						/*bCanCreateActor=*/true,
+						PCGRuntimeGenSchedulerHelpers::CVarHideActorsFromOutliner.GetValueOnAnyThread());
 				}
 
 				if (!ensure(PartitionActor))
@@ -693,6 +720,13 @@ void FPCGRuntimeGenScheduler::TickScheduleGeneration(TMap<FGridGenerationKey, do
 
 void FPCGRuntimeGenScheduler::TickCVars(const APCGWorldActor* InPCGWorldActor)
 {
+	if (bActorFlushRequested && Subsystem && Subsystem->GetPCGWorldActor())
+	{
+		CleanupLocalComponents(Subsystem->GetPCGWorldActor());
+		ResetPartitionActorPoolToSize(PCGRuntimeGenSchedulerHelpers::CVarRuntimeGenerationBasePoolSize.GetValueOnAnyThread());
+	}
+	bActorFlushRequested = false;
+
 	// If pooling has been disabled since last frame, we should destroy the pool.
 	const bool bPoolingEnabled = PCGRuntimeGenSchedulerHelpers::CVarRuntimeGenerationEnablePooling.GetValueOnAnyThread();
 
@@ -1188,7 +1222,11 @@ void FPCGRuntimeGenScheduler::AddPartitionActorPoolCount(int32 Count)
 #if WITH_EDITOR
 	SpawnParams.NameMode = FActorSpawnParameters::ESpawnActorNameMode::Requested;
 	SpawnParams.Name = *PCGRuntimeGenSchedulerConstants::PooledPartitionActorName;
+
+	// Always hide pooled actors from outliner. Note that outliner tree view updates can incur significant costs in Slate code.
+	SpawnParams.bHideFromSceneOutliner = PCGRuntimeGenSchedulerHelpers::CVarHideActorsFromOutliner.GetValueOnAnyThread();
 #endif
+
 	SpawnParams.ObjectFlags |= RF_Transient;
 	SpawnParams.ObjectFlags &= ~RF_Transactional;
 
@@ -1220,5 +1258,3 @@ void FPCGRuntimeGenScheduler::ResetPartitionActorPoolToSize(uint32 NewPoolSize)
 	PartitionActorPoolSize = 0;
 	AddPartitionActorPoolCount(NewPoolSize);
 }
-
-
diff --git a/Engine/Plugins/PCG/Source/PCG/Public/PCGSubsystem.h b/Engine/Plugins/PCG/Source/PCG/Public/PCGSubsystem.h
--- a/Engine/Plugins/PCG/Source/PCG/Public/PCGSubsystem.h
+++ b/Engine/Plugins/PCG/Source/PCG/Public/PCGSubsystem.h
@@ -168,6 +168,8 @@ public:
 	/** Asks the runtime generation scheduler to refresh a given GenerateAtRuntime component. bRemovePartitionActors will also perform a full cleanup of PAs and local components. */
 	void RefreshRuntimeGenComponent(UPCGComponent* RuntimeComponent, EPCGChangeType ChangeType = EPCGChangeType::None);
 
+	FPCGRuntimeGenScheduler* GetRuntimeGenScheduler() const { return RuntimeGenScheduler; }
+
 	/** Returns true if there are any tasks for this graph currently scheduled or executing. */
 	bool IsGraphCurrentlyExecuting(UPCGGraph* Graph);
 
@@ -244,7 +246,7 @@ public:
 	APCGPartitionActor* GetRegisteredPCGPartitionActor(const FPCGGridDescriptor& GridDescriptor, const FIntVector& GridCoords) const;
 
 	/** Creates a new partition actor if one does not already exist with the same grid size, coords, and generation mode. */
-	APCGPartitionActor* FindOrCreatePCGPartitionActor(const FPCGGridDescriptor& GridDescriptor, const FIntVector& GridCoords, bool bCanCreateActor = true) const;
+	APCGPartitionActor* FindOrCreatePCGPartitionActor(const FPCGGridDescriptor& GridDescriptor, const FIntVector& GridCoords, bool bCanCreateActor = true, bool bHideFromOutliner = false) const;
 	
 	/** True if graph cache debugging is enabled. */
 	bool IsGraphCacheDebuggingEnabled() const;
diff --git a/Engine/Plugins/PCG/Source/PCG/Public/RuntimeGen/PCGRuntimeGenScheduler.h b/Engine/Plugins/PCG/Source/PCG/Public/RuntimeGen/PCGRuntimeGenScheduler.h
--- a/Engine/Plugins/PCG/Source/PCG/Public/RuntimeGen/PCGRuntimeGenScheduler.h
+++ b/Engine/Plugins/PCG/Source/PCG/Public/RuntimeGen/PCGRuntimeGenScheduler.h
@@ -45,6 +45,9 @@ public:
 	void OnOriginalComponentRegistered(UPCGComponent* InOriginalComponent);
 	void OnOriginalComponentUnregistered(UPCGComponent* InOriginalComponent);
 
+	/** Destroy all runtime gen partition actors (both generated and pooled). Executed in next tick. */
+	void FlushAllGeneratedActors() { bActorFlushRequested = true; }
+
 protected:
 	struct FGridGenerationKey : TTuple<uint32, FIntVector, UPCGComponent*>
 	{
@@ -128,6 +131,9 @@ private:
 	bool bPoolingWasEnabledLastFrame = true;
 	uint32 BasePoolSizeLastFrame = 0;
 
+	/** Requests to flush all actors are deferred so they can be handled at a known time during tick. */
+	bool bActorFlushRequested = false;
+
 	/** Track the existence of runtime gen components to avoid unnecessary computation when there is no work to do. */
 	bool bAnyRuntimeGenComponentsExist = false;
 	bool bAnyRuntimeGenComponentsExistDirty = false;
