[PCG] Fixed a bug preventing the Landscape Cache from updating correctly

Implementation Note:
When constructed with a non-const ALandscapeProxy, it would implicitly and unintentionally convert to FObjectKey from the other constructor.

There may be some merit in a stricter set of type traits for the constructors, but as of now, this should fix the issue.

#rnx

diff --git a/Engine/Plugins/PCG/Source/PCG/Public/Grid/PCGLandscapeCache.h b/Engine/Plugins/PCG/Source/PCG/Public/Grid/PCGLandscapeCache.h
--- a/Engine/Plugins/PCG/Source/PCG/Public/Grid/PCGLandscapeCache.h
+++ b/Engine/Plugins/PCG/Source/PCG/Public/Grid/PCGLandscapeCache.h
@@ -183,11 +183,12 @@ private:
 		FIntPoint Coordinate;
 		FObjectKey WorldKey;
 
+		// Implementation note: FObjectKey will implicitly convert UObjects, including AActor.
 		CacheMapKey(const FGuid& InLandscapeGuid, const FIntPoint& InCoordinate, const FObjectKey& InWorldKey)
 			: LandscapeGuid(InLandscapeGuid), Coordinate(InCoordinate), WorldKey(InWorldKey)
 		{}
 
-		CacheMapKey(const FGuid& InLandscapeGuid, const FIntPoint& InCoordinate, AActor* InHintActor)
+		CacheMapKey(const FGuid& InLandscapeGuid, const FIntPoint& InCoordinate, const AActor* InHintActor)
 			: LandscapeGuid(InLandscapeGuid), Coordinate(InCoordinate)
 		{
 			if (InHintActor && InHintActor->GetLevel())
@@ -232,4 +233,4 @@ private:
 #if WITH_EDITOR
 	FRWLock CacheLock;
 #endif
-};
\ No newline at end of file
+};
