[PCG] Fix soft freeze of graph generation when higen components depend on eachother.

diff --git a/Engine/Plugins/PCG/Source/PCG/Private/Graph/PCGGraphExecutor.cpp b/Engine/Plugins/PCG/Source/PCG/Private/Graph/PCGGraphExecutor.cpp
--- a/Engine/Plugins/PCG/Source/PCG/Private/Graph/PCGGraphExecutor.cpp
+++ b/Engine/Plugins/PCG/Source/PCG/Private/Graph/PCGGraphExecutor.cpp
@@ -2077,12 +2077,18 @@ namespace PCGGraphExecutor
 				const AActor* ComponentActor = InContext->SourceComponent->GetOwner();
 				if (ComponentActor)
 				{
-					FPCGGridDescriptor GridDescriptor = InContext->SourceComponent->GetGridDescriptor(FromGridSize);
+					const UPCGComponent* OriginalComponent = InContext->SourceComponent->GetOriginalComponent();
+					if (!ensure(OriginalComponent))
+					{
+						return false;
+					}
+
+					const FPCGGridDescriptor GridDescriptor = OriginalComponent->GetGridDescriptor(FromGridSize);
 
 					// Get grid coords using the parent grid (FromGridSize).
 					const FIntVector CellCoords = UPCGActorHelpers::GetCellCoord(ComponentActor->GetActorLocation(), GridDescriptor.GetGridSize(), GridDescriptor.Is2DGrid());
 
-					ComponentWithData = Subsystem->GetLocalComponent(GridDescriptor, CellCoords, InContext->SourceComponent->GetOriginalComponent());
+					ComponentWithData = Subsystem->GetLocalComponent(GridDescriptor, CellCoords, OriginalComponent);
 				}
 			}
 
diff --git a/Engine/Plugins/PCG/Source/PCG/Private/PCGComponent.cpp b/Engine/Plugins/PCG/Source/PCG/Private/PCGComponent.cpp
--- a/Engine/Plugins/PCG/Source/PCG/Private/PCGComponent.cpp
+++ b/Engine/Plugins/PCG/Source/PCG/Private/PCGComponent.cpp
@@ -146,7 +146,11 @@ FPCGGridDescriptor UPCGComponent::GetGridDescriptorInternal(uint32 GridSize, boo
 	// Return owner descriptor in case of Partition Actors
 	if (APCGPartitionActor* PartitionActorOwner = Cast<APCGPartitionActor>(GetOwner()))
 	{
-		return PartitionActorOwner->GetGridDescriptor();
+		const FPCGGridDescriptor GridDescriptor = PartitionActorOwner->GetGridDescriptor();
+		// If this is a local component, we only serve grid descriptors of the same grid size.
+		check(GridSize == GridDescriptor.GetGridSize());
+
+		return GridDescriptor;
 	}
 
 	FPCGGridDescriptor PCGGridDescriptor = FPCGGridDescriptor()
diff --git a/Engine/Plugins/PCG/Source/PCG/Private/PCGSubsystem.cpp b/Engine/Plugins/PCG/Source/PCG/Private/PCGSubsystem.cpp
--- a/Engine/Plugins/PCG/Source/PCG/Private/PCGSubsystem.cpp
+++ b/Engine/Plugins/PCG/Source/PCG/Private/PCGSubsystem.cpp
@@ -1086,6 +1086,8 @@ void UPCGSubsystem::CleanupLocalComponentsImmediate(UPCGComponent* InOriginalCom
 // deprecated
 UPCGComponent* UPCGSubsystem::GetLocalComponent(uint32 GridSize, const FIntVector& CellCoords, const UPCGComponent* InOriginalComponent, bool bTransient) const
 {
+	check(InOriginalComponent);
+
 	FPCGGridDescriptor GridDescriptor = FPCGGridDescriptor()
 		.SetGridSize(GridSize)
 		.SetIsRuntime(bTransient)
@@ -1116,6 +1118,7 @@ APCGPartitionActor* UPCGSubsystem::FindOrCreatePCGPartitionActor(const FGuid& Gu
 
 UPCGComponent* UPCGSubsystem::GetLocalComponent(const FPCGGridDescriptor& GridDescriptor, const FIntVector& CellCoords, const UPCGComponent* InOriginalComponent) const
 {
+	check(InOriginalComponent);
 	return ActorAndComponentMapping.GetLocalComponent(GridDescriptor, CellCoords, InOriginalComponent);
 }
 
