[PCG] Create a dummy transaction so that deleted actor packages are kept alive (to be saved later).

Needs to be revisited but for now this  fix works.

#rnx

diff --git a/Engine/Plugins/PCG/Source/PCG/Private/Helpers/PCGActorHelpers.cpp b/Engine/Plugins/PCG/Source/PCG/Private/Helpers/PCGActorHelpers.cpp
--- a/Engine/Plugins/PCG/Source/PCG/Private/Helpers/PCGActorHelpers.cpp
+++ b/Engine/Plugins/PCG/Source/PCG/Private/Helpers/PCGActorHelpers.cpp
@@ -18,6 +18,7 @@
 
 #if WITH_EDITOR
 #include "Engine/Level.h"
+#include "ScopedTransaction.h"
 #include "WorldPartition/DataLayer/DataLayerInstance.h"
 #include "WorldPartition/DataLayer/ExternalDataLayerAsset.h"
 #include "WorldPartition/DataLayer/ExternalDataLayerInstance.h"
@@ -339,6 +340,20 @@ bool UPCGActorHelpers::DeleteActors(UWorld* World, const TArray<TSoftObjectPtr<A
 	//else
 #endif
 	{
+#if WITH_EDITOR
+		// Create TX so that dirty actor packages are tracked
+		// 
+		// Without tracking deleted actor packages will get unloaded on the next GC with no chance to save them first 
+		// See FWorldPartitionExternalDirtyActorsTracker::OnAddDirtyActor (Since CL 32133290)
+		//
+		// The Actors not being referenced by the Dirty Tracker will prevent them from being collected in UWorldPartition::AddReferencedObjects
+		// then in UWorldPartition::OnGCPostReachabilityAnalysis all unreachable actors will get processed to remove RF_Standalone flags allowing 
+		// the next GC to collect those packages
+		//
+		// The fix here is to create a dummy transaction so that the deleted actors are tracked but since we don't want to actually push a transaction, we cancel it after the DestroyActor calls
+		FScopedTransaction DummyTransaction(NSLOCTEXT("PCGActorHelpers", "DummyTransaction", "DummyTransaction"), World && !World->IsGameWorld());
+#endif
+
 		// Not in editor, really unlikely to happen but might be slow
 		for (const TSoftObjectPtr<AActor>& ManagedActor : ActorsToDelete)
 		{
@@ -351,6 +366,11 @@ bool UPCGActorHelpers::DeleteActors(UWorld* World, const TArray<TSoftObjectPtr<A
 				}
 			}
 		}
+
+#if WITH_EDITOR
+		// Cancel the Dummy transaction so that it can't be undone.
+		DummyTransaction.Cancel();
+#endif
 	}
 
 	return true;
