[PCG] Add Blur element and add an octree query library.

* Library done because the blur element needed some octree query
* Will allow users to do queries in BP, which will allow them to design nodes like the blur.

#rnx

diff --git a/Engine/Plugins/PCG/Source/PCG/Private/Elements/PCGBlurElement.cpp b/Engine/Plugins/PCG/Source/PCG/Private/Elements/PCGBlurElement.cpp
new file mode 100644
--- /dev/null
+++ b/Engine/Plugins/PCG/Source/PCG/Private/Elements/PCGBlurElement.cpp
@@ -0,0 +1,294 @@
+// Copyright Epic Games, Inc. All Rights Reserved.
+
+#include "Elements/PCGBlurElement.h"
+
+#include "PCGContext.h"
+#include "Data/PCGPointData.h"
+#include "Helpers/PCGAsync.h"
+#include "Helpers/PCGHelpers.h"
+#include "Metadata/Accessors/PCGAttributeAccessorHelpers.h"
+#include "SpatialAlgo/PCGOctreeQueries.h"
+
+#define LOCTEXT_NAMESPACE "PCGBlurElement"
+
+#if WITH_EDITOR
+FName UPCGBlurSettings::GetDefaultNodeName() const
+{
+	return FName(TEXT("BlurElement"));
+}
+
+FText UPCGBlurSettings::GetDefaultNodeTitle() const
+{
+	return LOCTEXT("NodeTitle", "Blur");
+}
+#endif // WITH_EDITOR
+
+FPCGElementPtr UPCGBlurSettings::CreateElement() const
+{
+	return MakeShared<FPCGBlurElement>();
+}
+
+bool FPCGBlurElement::PrepareDataInternal(FPCGContext* InContext) const
+{
+	TRACE_CPUPROFILER_EVENT_SCOPE(FPCGBlurElement::PrepareData);
+
+	const UPCGBlurSettings* Settings = InContext->GetInputSettings<UPCGBlurSettings>();
+	check(Settings);
+
+	ContextType* Context = static_cast<ContextType*>(InContext);
+	check(Context);
+
+	TArray<FPCGTaggedData> Inputs = InContext->InputData.GetInputsByPin(PCGPinConstants::DefaultInputLabel);
+	TArray<FPCGTaggedData>& Outputs = InContext->OutputData.TaggedData;
+
+	// Making sure the parameters are not negative or zero. Extra carefulness with search distance.
+	if (Settings->NumIterations <= 0 || Settings->SearchDistance <= 0 || FMath::IsNearlyZero(Settings->SearchDistance))
+	{
+		// Nothing to do, so forward the input
+		Outputs = Inputs;
+		return true;
+	}
+
+	// Verify that standard deviation is valid
+	if (Settings->BlurMode == EPCGBlurElementMode::Gaussian && Settings->bUseCustomStandardDeviation && (Settings->CustomStandardDeviation <= 0 || FMath::IsNearlyZero(Settings->CustomStandardDeviation)))
+	{
+		PCGLog::LogErrorOnGraph(LOCTEXT("InvalidStdDev", "Blur mode was set to Gaussian but Standard Deviation was less or equal than zero. Abort."));
+		return true;
+	}
+
+	// No exec init
+	Context->InitializePerExecutionState();
+
+	Context->InitializePerIterationStates(Inputs.Num(), [&Inputs, &Outputs, InContext, Settings](FPCGBlurIterState& OutState, const ExecStateType&, const uint32 IterationIndex) -> EPCGTimeSliceInitResult
+	{
+		FPCGTaggedData& OutputData = Outputs.Emplace_GetRef(Inputs[IterationIndex]);
+		OutputData.Pin = PCGPinConstants::DefaultOutputLabel;
+
+		OutState.InputPointData = Cast<UPCGPointData>(Inputs[IterationIndex].Data);
+		if (!OutState.InputPointData)
+		{
+			return EPCGTimeSliceInitResult::NoOperation;
+		}
+
+		FPCGAttributePropertyInputSelector InputSelector = Settings->InputSource.CopyAndFixLast(OutState.InputPointData);
+
+		OutState.InputAccessor = PCGAttributeAccessorHelpers::CreateConstAccessor(OutState.InputPointData, InputSelector);
+		OutState.InputKeys = PCGAttributeAccessorHelpers::CreateConstKeys(OutState.InputPointData, InputSelector);
+
+		if (!OutState.InputAccessor || !OutState.InputKeys)
+		{
+			PCGLog::LogErrorOnGraph(FText::Format(LOCTEXT("InputAttributeNotFound", "Attribute {0} was not found in input {1}"), InputSelector.GetDisplayText(), IterationIndex), InContext);
+			return EPCGTimeSliceInitResult::NoOperation;
+		}
+
+		// Nothing to do
+		if (OutState.InputKeys->GetNum() == 0)
+		{
+			return EPCGTimeSliceInitResult::NoOperation;
+		}
+
+		if (!PCG::Private::IsOfTypes<int32, int64, float, double, FVector2D, FVector, FVector4, FRotator, FQuat>(OutState.InputAccessor->GetUnderlyingType()))
+		{
+			PCGLog::LogErrorOnGraph(FText::Format(LOCTEXT("InputNotNumerical", "Attribute {0} is not numerical for input {1}"), InputSelector.GetDisplayText(), IterationIndex), InContext);
+			return EPCGTimeSliceInitResult::NoOperation;
+		}
+
+		UPCGPointData* OutputPointData = CastChecked<UPCGPointData>(OutState.InputPointData->DuplicateData());
+
+		FPCGAttributePropertyOutputSelector OutputSelector = Settings->OutputTarget.CopyAndFixSource(&InputSelector, OutState.InputPointData);
+
+		OutState.OutputAccessor = PCGAttributeAccessorHelpers::CreateAccessorWithAttributeCreation(OutputPointData, OutputSelector, OutState.InputAccessor.Get());
+		OutState.OutputKeys = PCGAttributeAccessorHelpers::CreateKeys(OutputPointData, OutputSelector);
+
+		if (!OutState.OutputAccessor || !OutState.OutputKeys)
+		{
+			PCGLog::LogErrorOnGraph(FText::Format(LOCTEXT("OutputAttributeFailiure", "Attribute {0} was not found, not compatible or failed to be created for input {1}"), OutputSelector.GetDisplayText(), IterationIndex), InContext);
+			return EPCGTimeSliceInitResult::NoOperation;
+		}
+
+		OutputData.Data = OutputPointData;
+
+		return EPCGTimeSliceInitResult::Success;
+	});
+
+	return true;
+}
+
+bool FPCGBlurElement::ExecuteInternal(FPCGContext* InContext) const
+{
+	TRACE_CPUPROFILER_EVENT_SCOPE(FPCGBlurElement::Execute);
+
+	check(InContext);
+
+	ContextType* Context = static_cast<ContextType*>(InContext);
+	check(Context);
+
+	if (!Context->DataIsPreparedForExecution())
+	{
+		return true;
+	}
+
+	const UPCGBlurSettings* Settings = InContext->GetInputSettings<UPCGBlurSettings>();
+	check(Settings);
+
+	check(Settings->BlurMode != EPCGBlurElementMode::Gaussian || !Settings->bUseCustomStandardDeviation || Settings->CustomStandardDeviation > 0);
+
+	const double StdDeviation = Settings->BlurMode == EPCGBlurElementMode::Gaussian && !Settings->bUseCustomStandardDeviation 
+		? Settings->SearchDistance / 9.0
+		: Settings->CustomStandardDeviation;
+	const double GaussianExpRatio = Settings->BlurMode == EPCGBlurElementMode::Gaussian ? 1.0 / (2.0 * StdDeviation * StdDeviation) : 0.0;
+	const double GaussianRatio = Settings->BlurMode == EPCGBlurElementMode::Gaussian ? GaussianExpRatio / UE_DOUBLE_PI : 0.0;
+
+	const bool bDone = ExecuteSlice(Context, [this, Settings, GaussianExpRatio, GaussianRatio](ContextType* Context, const ExecStateType& ExecState, IterStateType& IterState, const uint32 IterIndex)
+	{
+		if (Context->GetIterationStateResult(IterIndex) != EPCGTimeSliceInitResult::Success)
+		{
+			return true;
+		}
+
+		auto Process = [&IterState, Settings, Context, GaussianExpRatio, GaussianRatio]<typename AttributeType>(AttributeType) -> bool
+		{
+			if (IterState.WorkingBuffer1.IsEmpty())
+			{
+				ensure(IterState.WorkingBuffer2.IsEmpty());
+				IterState.WorkingBuffer1.SetNumUninitialized(sizeof(AttributeType) * IterState.InputKeys->GetNum());
+				IterState.WorkingBuffer2.SetNumUninitialized(sizeof(AttributeType) * IterState.InputKeys->GetNum());
+				IterState.InputAccessor->GetRange<AttributeType>(MakeArrayView(reinterpret_cast<AttributeType*>(IterState.WorkingBuffer1.GetData()), IterState.InputKeys->GetNum()), 0, *IterState.InputKeys);
+			}
+
+			while (IterState.CurrentIteration < Settings->NumIterations)
+			{
+				const uint8* ReadBuffer = (IterState.CurrentIteration % 2 == 0) ? IterState.WorkingBuffer1.GetData() : IterState.WorkingBuffer2.GetData();
+				uint8* WriteBuffer = (IterState.CurrentIteration % 2 == 1) ? IterState.WorkingBuffer1.GetData() : IterState.WorkingBuffer2.GetData();
+
+				TArrayView<const AttributeType> TypedReadBuffer = MakeArrayView(reinterpret_cast<const AttributeType*>(ReadBuffer), IterState.InputKeys->GetNum());
+				TArrayView<AttributeType> TypedWriteBuffer = MakeArrayView(reinterpret_cast<AttributeType*>(WriteBuffer), IterState.InputKeys->GetNum());
+
+				const bool bIsDone = FPCGAsync::AsyncProcessingOneToOneEx(&Context->AsyncState, IterState.InputKeys->GetNum(), []() {}, [&IterState, Settings, &TypedReadBuffer, TypedWriteBuffer, GaussianExpRatio, GaussianRatio](int32 ReadIndex, int32 WriteIndex)
+				{
+					const FPCGPoint& Point = IterState.InputPointData->GetPoints()[ReadIndex];
+					const FVector PointLocation = Point.Transform.GetLocation();
+					TArray<TTuple<const FPCGPoint*, double>> NeighborsAndDistance;
+
+					UPCGOctreeQueries::ForEachPointInsideSphere(IterState.InputPointData, PointLocation, Settings->SearchDistance, [&NeighborsAndDistance](const FPCGPoint& Neighbor, double SquaredDistance)
+					{
+						NeighborsAndDistance.Emplace(&Neighbor, SquaredDistance);
+					});
+
+					if (!ensure(!NeighborsAndDistance.IsEmpty()))
+					{
+						return true;
+					}
+
+					auto GetNeighborValue = [TypedReadBuffer, &IterState](const FPCGPoint* Neighbor) -> AttributeType
+					{
+						// Pointer arithmetic to know the neighbor index.
+						const int32 NeighborIndex = static_cast<int32>(Neighbor - IterState.InputPointData->GetPoints().GetData());
+						return TypedReadBuffer[NeighborIndex];
+					};
+
+
+					// If it is weighted by distance, we compute all the weights first, to have the sum of all of them (for normalization)
+					bool bValidWeights = false;
+					TArray<double> Weights;
+					if (Settings->BlurMode != EPCGBlurElementMode::Constant)
+					{
+						double WeightSum = 0.0;
+						bValidWeights = true;
+						Weights.Reserve(NeighborsAndDistance.Num());
+						for (const auto& [Neighbor, SquaredDistance] : NeighborsAndDistance)
+						{
+							double Weight = 0.0;
+							switch (Settings->BlurMode)
+							{
+							case EPCGBlurElementMode::Linear:
+								Weight = 1.0 - FVector::Distance(PointLocation, Neighbor->Transform.GetLocation()) / Settings->SearchDistance;
+								break;
+							case EPCGBlurElementMode::Gaussian:
+								Weight = GaussianRatio * FMath::Exp(-SquaredDistance * GaussianExpRatio);
+								break;
+							default:
+								checkNoEntry();
+								break;
+							}
+
+							Weights.Add(Weight);
+							WeightSum += Weight;
+						}
+
+						if (!FMath::IsNearlyZero(WeightSum))
+						{
+							for (double& Weight : Weights)
+							{
+								Weight /= WeightSum;
+							}
+						}
+						else
+						{
+							bValidWeights = false;
+						}
+					}
+
+					const double DefaultWeight = 1.0 / NeighborsAndDistance.Num();
+					const bool bUseDefaultWeight = !bValidWeights;
+
+					TypedWriteBuffer[WriteIndex] = PCG::Private::MetadataTraits<AttributeType>::ZeroValueForWeightedSum();
+
+					for (int32 i = 0; i < NeighborsAndDistance.Num(); ++i)
+					{
+						AttributeType Value = GetNeighborValue(NeighborsAndDistance[i].Get<0>());
+						const double Weight = bUseDefaultWeight ? DefaultWeight : Weights[i];
+						TypedWriteBuffer[WriteIndex] = PCG::Private::MetadataTraits<AttributeType>::WeightedSum(TypedWriteBuffer[WriteIndex], Value, Weight);
+					}
+
+					// For Quaternions, we need to normalize them
+					if constexpr (std::is_same_v<AttributeType, FQuat>)
+					{
+						TypedWriteBuffer[WriteIndex].Normalize();
+					}
+
+					return true;
+				},/*bEnableTimeSlicing=*/true);
+
+				if (!bIsDone)
+				{
+					return false;
+				}
+
+				++IterState.CurrentIteration;
+			}
+
+			// Copy back
+			const uint8* ReadBuffer = (IterState.CurrentIteration % 2 == 0) ? IterState.WorkingBuffer1.GetData() : IterState.WorkingBuffer2.GetData();
+			IterState.OutputAccessor->SetRange<AttributeType>(MakeArrayView(reinterpret_cast<const AttributeType*>(ReadBuffer), IterState.InputKeys->GetNum()), 0, *IterState.OutputKeys, EPCGAttributeAccessorFlags::AllowBroadcastAndConstructible);
+
+			IterState.WorkingBuffer1.Reset();
+			IterState.WorkingBuffer2.Reset();
+			return true;
+		};
+
+		auto Dispatch = [&IterState, Settings, Context, Process]<typename T>(T Dummy) -> bool
+		{
+			// Discard any non-numerical types
+			if constexpr (!PCG::Private::IsOfTypes<T, int32, int64, float, double, FVector2D, FVector, FVector4, FRotator, FQuat>())
+			{
+				// We should not enter there.
+				ensure(false);
+				return true;
+			}
+			else
+			{
+				// Work on double for ints and floats to avoid approximation/rounding errors
+				// Also Rotators are treated as Quaternions
+				using AttributeType = std::conditional_t<std::is_arithmetic_v<T>, double, std::conditional_t<std::is_same_v<FRotator, T>, FQuat, T>>; 
+				return Process(AttributeType{});
+			}
+		};
+
+		return PCGMetadataAttribute::CallbackWithRightType(IterState.InputAccessor->GetUnderlyingType(), Dispatch);
+	});
+
+	return bDone;
+}
+
+#undef LOCTEXT_NAMESPACE
diff --git a/Engine/Plugins/PCG/Source/PCG/Private/SpatialAlgo/PCGOctreeQueries.cpp b/Engine/Plugins/PCG/Source/PCG/Private/SpatialAlgo/PCGOctreeQueries.cpp
new file mode 100644
--- /dev/null
+++ b/Engine/Plugins/PCG/Source/PCG/Private/SpatialAlgo/PCGOctreeQueries.cpp
@@ -0,0 +1,162 @@
+// Copyright Epic Games, Inc. All Rights Reserved.
+
+#include "SpatialAlgo/PCGOctreeQueries.h"
+
+namespace PCGOctreeQueries
+{
+	/** Factorize the logic for all GetClosest/GetFarthest, as the only thing that changes is the StartingDistance and the Condition. */
+	template <typename Func = TFunctionRef<bool(double ChosenPointDistance, const FPCGPoint& CurrentPoint, double SquaredDistance)>>
+	const FPCGPoint* GetPointInSphereUnderCondition(const UPCGPointData* InPointData, const FVector& InCenter, const double InSearchDistance, const double InStartingDistance, Func Condition)
+	{
+		if (!InPointData)
+		{
+			return nullptr;
+		}
+
+		const FPCGPoint* ChosenPoint = nullptr;
+		const double SearchDistance = FMath::Max(InSearchDistance, 0.0);
+
+		double ChosenPointDistance = InStartingDistance;
+		UPCGOctreeQueries::ForEachPointInsideSphere(InPointData, InCenter, SearchDistance, [&ChosenPointDistance, &ChosenPoint, Condition = std::move(Condition)](const FPCGPoint& Point, double SquaredDistance)
+		{
+			if (Condition(ChosenPointDistance, Point, SquaredDistance))
+			{
+				ChosenPointDistance = SquaredDistance;
+				ChosenPoint = &Point;
+			}
+		});
+
+		return ChosenPoint;
+	}
+}
+
+TArray<FPCGPoint> UPCGOctreeQueries::GetPointsInsideBounds(const UPCGPointData* InPointData, const FBox& InBounds)
+{
+	if (!InPointData)
+	{
+		return {};
+	}
+
+	TArray<FPCGPoint> Result;
+	InPointData->GetOctree().FindElementsWithBoundsTest(InBounds, [&Result](const FPCGPointRef& PointRef)
+	{
+		if (ensure(PointRef.Point))
+		{
+			Result.Add(*PointRef.Point);
+		}
+	});
+
+	return Result;
+}
+
+TArray<FPCGPoint> UPCGOctreeQueries::GetPointsInsideSphere(const UPCGPointData* InPointData, const FVector& InCenter, const double InRadius)
+{
+	if (!InPointData)
+	{
+		return {};
+	}
+
+	TArray<FPCGPoint> Result;
+	ForEachPointInsideSphere(InPointData, InCenter, InRadius, [&Result](const FPCGPoint& Point, double)
+	{
+		Result.Add(Point);
+	});
+
+	return Result;
+}
+
+void UPCGOctreeQueries::GetClosestPoint(const UPCGPointData* InPointData, const FVector& InCenter, const bool bInDiscardCenter, bool& bOutFound, FPCGPoint& OutPoint, const double InSearchDistance)
+{
+	bOutFound = false;
+	if (const FPCGPoint* ClosestPoint = GetClosestPoint(InPointData, InCenter, bInDiscardCenter, InSearchDistance))
+	{
+		bOutFound = true;
+		OutPoint = *ClosestPoint;
+	}
+}
+
+const FPCGPoint* UPCGOctreeQueries::GetClosestPoint(const UPCGPointData* InPointData, const FVector& InCenter, const bool bInDiscardCenter, const double InSearchDistance)
+{
+	return PCGOctreeQueries::GetPointInSphereUnderCondition(InPointData, InCenter, InSearchDistance, std::numeric_limits<double>::max(),
+		[bInDiscardCenter](double ChosenPointDistance, const FPCGPoint& Point, double SquaredDistance)
+		{
+			return SquaredDistance <= ChosenPointDistance && (!bInDiscardCenter || SquaredDistance > UE_DOUBLE_SMALL_NUMBER);
+		});
+}
+
+void UPCGOctreeQueries::GetClosestPointFromOtherPoint(const UPCGPointData* InPointData, const int32 InPointIndex, bool& bOutFound, FPCGPoint& OutPoint, const double InSearchDistance)
+{
+	bOutFound = false;
+	if (!InPointData)
+	{
+		return;
+	}
+
+	if (!InPointData->GetPoints().IsValidIndex(InPointIndex))
+	{
+		return;
+	}
+
+	if (const FPCGPoint* ClosestPoint = GetClosestPointFromOtherPoint(InPointData, InPointData->GetPoints()[InPointIndex], InSearchDistance))
+	{
+		bOutFound = true;
+		OutPoint = *ClosestPoint;
+	}
+}
+
+const FPCGPoint* UPCGOctreeQueries::GetClosestPointFromOtherPoint(const UPCGPointData* InPointData, const FPCGPoint& InPoint, const double InSearchDistance)
+{
+	return PCGOctreeQueries::GetPointInSphereUnderCondition(InPointData, InPoint.Transform.GetLocation(), InSearchDistance, std::numeric_limits<double>::max(),
+		[&InPoint](double ChosenPointDistance, const FPCGPoint& Point, double SquaredDistance)
+		{
+			return SquaredDistance <= ChosenPointDistance && &InPoint != &Point;
+		});
+}
+
+void UPCGOctreeQueries::GetFarthestPoint(const UPCGPointData* InPointData, const FVector& InCenter, bool& bOutFound, FPCGPoint& OutPoint, const double InSearchDistance)
+{
+	bOutFound = false;
+	if (const FPCGPoint* FarthestPoint = GetFarthestPoint(InPointData, InCenter, InSearchDistance))
+	{
+		bOutFound = true;
+		OutPoint = *FarthestPoint;
+	}
+}
+
+const FPCGPoint* UPCGOctreeQueries::GetFarthestPoint(const UPCGPointData* InPointData, const FVector& InCenter, const double InSearchDistance)
+{
+	return PCGOctreeQueries::GetPointInSphereUnderCondition(InPointData, InCenter, InSearchDistance, std::numeric_limits<double>::min(),
+		[](double ChosenPointDistance, const FPCGPoint& Point, double SquaredDistance)
+		{
+			return SquaredDistance >= ChosenPointDistance;
+		});
+}
+
+void UPCGOctreeQueries::GetFarthestPointFromOtherPoint(const UPCGPointData* InPointData, const int32 InPointIndex, bool& bOutFound, FPCGPoint& OutPoint, const double InSearchDistance)
+{
+	bOutFound = false;
+	if (!InPointData)
+	{
+		return;
+	}
+
+	if (!InPointData->GetPoints().IsValidIndex(InPointIndex))
+	{
+		return;
+	}
+
+	if (const FPCGPoint* ClosestPoint = GetFarthestPointFromOtherPoint(InPointData, InPointData->GetPoints()[InPointIndex], InSearchDistance))
+	{
+		bOutFound = true;
+		OutPoint = *ClosestPoint;
+	}
+}
+
+const FPCGPoint* UPCGOctreeQueries::GetFarthestPointFromOtherPoint(const UPCGPointData* InPointData, const FPCGPoint& InPoint, const double InSearchDistance)
+{
+	return PCGOctreeQueries::GetPointInSphereUnderCondition(InPointData, InPoint.Transform.GetLocation(), InSearchDistance, std::numeric_limits<double>::min(),
+		[&InPoint](double ChosenPointDistance, const FPCGPoint& Point, double SquaredDistance)
+		{
+			return SquaredDistance >= ChosenPointDistance && &InPoint != &Point;
+		});
+}
\ No newline at end of file
diff --git a/Engine/Plugins/PCG/Source/PCG/Private/Tests/Elements/PCGBlurElementTest.cpp b/Engine/Plugins/PCG/Source/PCG/Private/Tests/Elements/PCGBlurElementTest.cpp
new file mode 100644
--- /dev/null
+++ b/Engine/Plugins/PCG/Source/PCG/Private/Tests/Elements/PCGBlurElementTest.cpp
@@ -0,0 +1,81 @@
+// Copyright Epic Games, Inc. All Rights Reserved.
+
+#if WITH_EDITOR
+
+#include "Tests/PCGTestsCommon.h"
+
+#include "PCGComponent.h"
+#include "PCGContext.h"
+#include "Data/PCGPointData.h"
+#include "Elements/PCGBlurElement.h"
+
+IMPLEMENT_CUSTOM_SIMPLE_AUTOMATION_TEST(FPCGBlurElementTest_Density, FPCGTestBaseClass, "Plugins.PCG.BlurElement.Density", PCGTestsCommon::TestFlags)
+
+/**
+* Point data with 4 points, placed on a square.
+* Starting with 1 in the top left corner, 0 elsewhere. Distance is the length of the square. Indexes are shown below.
+* 0    1
+* 
+* 
+* 2    3
+* 
+* Start values {1, 0, 0, 0}
+* First Iteration: {0.333333, 0.333333, 0.333333, 0}
+* Second Iteration: {0.333333, 0.222222, 0.222222, 0.222222}
+*/
+bool FPCGBlurElementTest_Density::RunTest(const FString& Parameters)
+{
+	PCGTestsCommon::FTestData TestData;
+	UPCGBlurSettings* Settings = PCGTestsCommon::GenerateSettings<UPCGBlurSettings>(TestData);
+	check(Settings);
+
+	constexpr double SquareSize = 100.0;
+
+	Settings->InputSource.SetPointProperty(EPCGPointProperties::Density);
+	Settings->NumIterations = 2;
+	Settings->SearchDistance = SquareSize + 1.0; //+ 1 for approximation errors
+
+
+	UPCGPointData* InputPointData = PCGTestsCommon::CreateEmptyPointData();
+	TArray<FPCGPoint>& InputPoints = InputPointData->GetMutablePoints();
+
+	for (int i = 0; i < 4; ++i)
+	{
+		const int X = i / 2;
+		const int Y = i % 2;
+		const FVector Location = FVector(SquareSize * X, SquareSize * Y, 0.0);
+		const float Density = i == 0 ? 1.0f : 0.0f;
+		InputPoints.Emplace(FTransform(Location), Density, 42);
+	}
+
+	FPCGTaggedData& TaggedData = TestData.InputData.TaggedData.Emplace_GetRef(FPCGTaggedData());
+	TaggedData.Data = InputPointData;
+	TaggedData.Pin = PCGPinConstants::DefaultInputLabel;
+
+	FPCGElementPtr BlurElement = TestData.Settings->GetElement();
+	check(BlurElement);
+
+	TUniquePtr<FPCGContext> Context = TestData.InitializeTestContext();
+	while (!BlurElement->Execute(Context.Get())) {}
+
+	TArray<FPCGTaggedData> Outputs = Context->OutputData.GetInputsByPin(PCGPinConstants::DefaultOutputLabel);
+	UTEST_EQUAL("There is one output", Outputs.Num(), 1);
+
+	const UPCGPointData* OutputPointData = Cast<UPCGPointData>(Outputs[0].Data);
+	UTEST_NOT_NULL("Output is a point data", OutputPointData);
+
+	check(OutputPointData);
+
+	const TArray<FPCGPoint>& OutputPoints = OutputPointData->GetPoints();
+	UTEST_EQUAL("Output has as many points as input", OutputPoints.Num(), InputPoints.Num());
+
+	constexpr float DensityValues[] = { 0.333333f, 0.222222f, 0.222222f, 0.222222f };
+	for (int32 i = 0; i < 4; ++i)
+	{
+		UTEST_EQUAL(*FString::Printf(TEXT("Point %d has the expected density."), i), OutputPoints[i].Density, DensityValues[i]);
+	}
+
+	return true;
+}
+
+#endif // WITH_EDITOR
\ No newline at end of file
diff --git a/Engine/Plugins/PCG/Source/PCG/Private/Tests/SpatialAlgo/PCGOctreeQueriesTest.cpp b/Engine/Plugins/PCG/Source/PCG/Private/Tests/SpatialAlgo/PCGOctreeQueriesTest.cpp
new file mode 100644
--- /dev/null
+++ b/Engine/Plugins/PCG/Source/PCG/Private/Tests/SpatialAlgo/PCGOctreeQueriesTest.cpp
@@ -0,0 +1,134 @@
+// Copyright Epic Games, Inc. All Rights Reserved.
+
+#if WITH_EDITOR
+
+#include "Tests/PCGTestsCommon.h"
+
+#include "Data/PCGPointData.h"
+#include "SpatialAlgo/PCGOctreeQueries.h"
+
+IMPLEMENT_CUSTOM_SIMPLE_AUTOMATION_TEST(FPCGOctreeQueriesTest_Sphere, FPCGTestBaseClass, "Plugins.PCG.OctreeQueries.Sphere", PCGTestsCommon::TestFlags)
+IMPLEMENT_CUSTOM_SIMPLE_AUTOMATION_TEST(FPCGOctreeQueriesTest_ClosestPoint, FPCGTestBaseClass, "Plugins.PCG.OctreeQueries.ClosestPoint", PCGTestsCommon::TestFlags)
+IMPLEMENT_CUSTOM_SIMPLE_AUTOMATION_TEST(FPCGOctreeQueriesTest_ClosestPointDiscardCenter, FPCGTestBaseClass, "Plugins.PCG.OctreeQueries.ClosestPointDiscardCenter", PCGTestsCommon::TestFlags)
+IMPLEMENT_CUSTOM_SIMPLE_AUTOMATION_TEST(FPCGOctreeQueriesTest_ClosestPointFromOtherPoint, FPCGTestBaseClass, "Plugins.PCG.OctreeQueries.ClosestPointFromOtherPoint", PCGTestsCommon::TestFlags)
+IMPLEMENT_CUSTOM_SIMPLE_AUTOMATION_TEST(FPCGOctreeQueriesTest_FarthestPoint, FPCGTestBaseClass, "Plugins.PCG.OctreeQueries.FarthestPoint", PCGTestsCommon::TestFlags)
+IMPLEMENT_CUSTOM_SIMPLE_AUTOMATION_TEST(FPCGOctreeQueriesTest_FarthestPointFromOtherPoint, FPCGTestBaseClass, "Plugins.PCG.OctreeQueries.FarthestPointFromOtherPoint", PCGTestsCommon::TestFlags)
+
+/**
+* For all the tests, the point data will be points scattered on a straight line (Direction {1, 1, 1}), and evenly spaced.
+* Query will be on the first point, at the origin.
+*/
+namespace PCGOctreeQueriesTest
+{
+	static constexpr double Distance = 100.0;
+	static constexpr int32 NumPoints = 10;
+
+	UPCGPointData* CreatePointData()
+	{
+		UPCGPointData* InputPointData = PCGTestsCommon::CreateEmptyPointData();
+		TArray<FPCGPoint>& InputPoints = InputPointData->GetMutablePoints();
+		InputPoints.Reserve(NumPoints);
+
+		for (int32 i = 0; i < NumPoints; ++i)
+		{
+			FVector Location = FVector::OneVector * Distance * i;
+			InputPoints.Emplace(FTransform(Location), 1.0f, 42);
+		}
+
+		return InputPointData;
+	}
+}
+
+
+bool FPCGOctreeQueriesTest_Sphere::RunTest(const FString& Parameters)
+{
+	const UPCGPointData* InputPointData = PCGOctreeQueriesTest::CreatePointData();
+	check(InputPointData);
+
+	TArray<int32> ExpectedIndexes = { 0, 1, 2 };
+	int32 CountFound = 0;
+	UPCGOctreeQueries::ForEachPointInsideSphere(InputPointData, FVector::ZeroVector, 350.0, [&ExpectedIndexes, &CountFound, &InputPoints = InputPointData->GetPoints()](const FPCGPoint& InPoint, double)
+	{
+		const int32 PointIndex = static_cast<int32>(&InPoint - InputPoints.GetData());
+		if (ExpectedIndexes.Contains(PointIndex))
+		{
+			++CountFound;
+		}
+	});
+
+	UTEST_EQUAL("We found the expected number of points", CountFound, ExpectedIndexes.Num());
+
+	return true;
+}
+
+bool FPCGOctreeQueriesTest_ClosestPoint::RunTest(const FString& Parameters)
+{
+	const UPCGPointData* InputPointData = PCGOctreeQueriesTest::CreatePointData();
+	check(InputPointData);
+
+	const FPCGPoint* ClosestPoint = UPCGOctreeQueries::GetClosestPoint(InputPointData, FVector::ZeroVector, /*bDiscardCenter=*/false, 350.0);
+	UTEST_NOT_NULL("Closest point was found", ClosestPoint);
+
+	const int32 PointIndex = static_cast<int32>(ClosestPoint - InputPointData->GetPoints().GetData());
+
+	UTEST_EQUAL("Closest point is the right index", PointIndex, 0);
+	return true;
+}
+
+bool FPCGOctreeQueriesTest_ClosestPointDiscardCenter::RunTest(const FString& Parameters)
+{
+	const UPCGPointData* InputPointData = PCGOctreeQueriesTest::CreatePointData();
+	check(InputPointData);
+
+	const FPCGPoint* ClosestPoint = UPCGOctreeQueries::GetClosestPoint(InputPointData, FVector::ZeroVector, /*bDiscardCenter=*/true, 350.0);
+	UTEST_NOT_NULL("Closest point was found", ClosestPoint);
+
+	const int32 PointIndex = static_cast<int32>(ClosestPoint - InputPointData->GetPoints().GetData());
+
+	UTEST_EQUAL("Closest point is the right index", PointIndex, 1);
+	return true;
+}
+
+bool FPCGOctreeQueriesTest_ClosestPointFromOtherPoint::RunTest(const FString& Parameters)
+{
+	const UPCGPointData* InputPointData = PCGOctreeQueriesTest::CreatePointData();
+	check(InputPointData);
+
+	const FPCGPoint* ClosestPoint = UPCGOctreeQueries::GetClosestPointFromOtherPoint(InputPointData, InputPointData->GetPoints()[0], 350.0);
+	UTEST_NOT_NULL("Closest point was found", ClosestPoint);
+
+	const int32 PointIndex = static_cast<int32>(ClosestPoint - InputPointData->GetPoints().GetData());
+
+	UTEST_EQUAL("Closest point is the right index", PointIndex, 1);
+	return true;
+}
+
+bool FPCGOctreeQueriesTest_FarthestPoint::RunTest(const FString& Parameters)
+{
+	const UPCGPointData* InputPointData = PCGOctreeQueriesTest::CreatePointData();
+	check(InputPointData);
+
+	const FPCGPoint* FarthestPoint = UPCGOctreeQueries::GetFarthestPoint(InputPointData, FVector::ZeroVector, 350.0);
+	UTEST_NOT_NULL("Farthest point was found", FarthestPoint);
+
+	const int32 PointIndex = static_cast<int32>(FarthestPoint - InputPointData->GetPoints().GetData());
+
+	UTEST_EQUAL("Farthest point is the right index", PointIndex, 2);
+	return true;
+}
+
+bool FPCGOctreeQueriesTest_FarthestPointFromOtherPoint::RunTest(const FString& Parameters)
+{
+	const UPCGPointData* InputPointData = PCGOctreeQueriesTest::CreatePointData();
+	check(InputPointData);
+
+	const FPCGPoint* FarthestPoint = UPCGOctreeQueries::GetFarthestPointFromOtherPoint(InputPointData, InputPointData->GetPoints()[0], 10000.0);
+	UTEST_NOT_NULL("Farthest point was found", FarthestPoint);
+
+	const int32 PointIndex = static_cast<int32>(FarthestPoint - InputPointData->GetPoints().GetData());
+
+	UTEST_EQUAL("Farthest point is the right index", PointIndex, 9);
+	return true;
+}
+
+#endif // WITH_EDITOR
\ No newline at end of file
diff --git a/Engine/Plugins/PCG/Source/PCG/Public/Elements/PCGBlurElement.h b/Engine/Plugins/PCG/Source/PCG/Public/Elements/PCGBlurElement.h
new file mode 100644
--- /dev/null
+++ b/Engine/Plugins/PCG/Source/PCG/Public/Elements/PCGBlurElement.h
@@ -0,0 +1,96 @@
+// Copyright Epic Games, Inc. All Rights Reserved.
+
+#pragma once
+
+#include "PCGSettings.h"
+
+#include "Elements/PCGTimeSlicedElementBase.h"
+#include "Metadata/Accessors/IPCGAttributeAccessor.h"
+#include "Metadata/Accessors/PCGAttributeAccessorKeys.h"
+
+#include "PCGBlurElement.generated.h"
+
+class UPCGPointData;
+
+UENUM()
+enum class EPCGBlurElementMode
+{
+	/** Same weight for each point, which will be 1 / N, N being the number of points found in the given distance. */
+	Constant,
+	/** Weight for Point will be 1 - Dist(Point, Center) / SearchDistance. */
+	Linear,
+	/** Weight will be a gaussian distribution. */
+	Gaussian
+};
+
+/**
+* Select an attribute on a point data and blur it using the values from neighbors within some distance, center to center, and can be done over multiple iterations.
+*/
+UCLASS(MinimalAPI, BlueprintType, ClassGroup = (Procedural))
+class UPCGBlurSettings : public UPCGSettings
+{
+	GENERATED_BODY()
+
+public:
+	//~Begin UPCGSettings interface
+#if WITH_EDITOR
+	virtual FName GetDefaultNodeName() const override;
+	virtual FText GetDefaultNodeTitle() const override;
+	virtual EPCGSettingsType GetType() const override { return EPCGSettingsType::PointOps; }
+#endif
+
+protected:
+	virtual TArray<FPCGPinProperties> InputPinProperties() const override { return Super::DefaultPointInputPinProperties(); }
+	virtual TArray<FPCGPinProperties> OutputPinProperties() const override { return Super::DefaultPointOutputPinProperties(); }
+	virtual FPCGElementPtr CreateElement() const override;
+	//~End UPCGSettings interface
+
+public:
+	/** Attribute to use as a base value. Needs to be numerical. */
+	UPROPERTY(BlueprintReadWrite, EditAnywhere, Category = Settings, meta = (PCG_Overridable))
+	FPCGAttributePropertyInputSelector InputSource;
+
+	/** Attribute where the result will be written. */
+	UPROPERTY(BlueprintReadWrite, EditAnywhere, Category = Settings, meta = (PCG_Overridable))
+	FPCGAttributePropertyOutputSelector OutputTarget;
+
+	/** Number of iterations to apply the blur. */
+	UPROPERTY(BlueprintReadWrite, EditAnywhere, Category = Settings, meta = (PCG_Overridable, ClampMin = 1))
+	int NumIterations = 1;
+
+	/** Radius for search. */
+	UPROPERTY(BlueprintReadWrite, EditAnywhere, Category = Settings, meta = (PCG_Overridable, ClampMin = 0))
+	double SearchDistance = 1000.0;
+
+	UPROPERTY(BlueprintReadWrite, EditAnywhere, Category = Settings, meta = (PCG_Overridable))
+	EPCGBlurElementMode BlurMode = EPCGBlurElementMode::Constant;
+
+	/** By default, the standard deviation will be SearchDistance / 3, so that at search distance from the point it corresponds to 3 std deviation. */
+	UPROPERTY(BlueprintReadWrite, EditAnywhere, Category = "Settings|Gaussian", meta = (PCG_Overridable, EditConditionHides, EditCondition = "BlurMode == EPCGBlurElementMode::Gaussian"))
+	bool bUseCustomStandardDeviation = false;
+
+	UPROPERTY(BlueprintReadWrite, EditAnywhere, Category = "Settings|Gaussian", meta = (PCG_Overridable, EditConditionHides, EditCondition = "BlurMode == EPCGBlurElementMode::Gaussian && bUseCustomStandardDeviation"))
+	double CustomStandardDeviation = 1.0;
+};
+
+struct FPCGBlurIterState
+{
+	TUniquePtr<const IPCGAttributeAccessor> InputAccessor;
+	TUniquePtr<IPCGAttributeAccessor> OutputAccessor;
+	TUniquePtr<const IPCGAttributeAccessorKeys> InputKeys;
+	TUniquePtr<IPCGAttributeAccessorKeys> OutputKeys;
+
+	const UPCGPointData* InputPointData = nullptr;
+	// Need two working buffers as we are working in parallel. One for read, one for write.
+	TArray<uint8> WorkingBuffer1;
+	TArray<uint8> WorkingBuffer2;
+	int32 CurrentIteration = 0;
+};
+
+class FPCGBlurElement : public TPCGTimeSlicedElementBase<PCGTimeSlice::FEmptyStruct, FPCGBlurIterState>
+{
+protected:
+	virtual bool PrepareDataInternal(FPCGContext* InContext) const override;
+	virtual bool ExecuteInternal(FPCGContext* InContext) const override;
+};
+
diff --git a/Engine/Plugins/PCG/Source/PCG/Public/SpatialAlgo/PCGOctreeQueries.h b/Engine/Plugins/PCG/Source/PCG/Public/SpatialAlgo/PCGOctreeQueries.h
new file mode 100644
--- /dev/null
+++ b/Engine/Plugins/PCG/Source/PCG/Public/SpatialAlgo/PCGOctreeQueries.h
@@ -0,0 +1,92 @@
+// Copyright Epic Games, Inc. All Rights Reserved.
+
+#pragma once
+
+#include "PCGPoint.h"
+#include "Data/PCGPointData.h"
+
+#include "Math/Vector.h"
+#include "Templates/Function.h"
+
+#include "PCGOctreeQueries.generated.h"
+
+UCLASS(MinimalAPI)
+class UPCGOctreeQueries : public UBlueprintFunctionLibrary
+{
+	GENERATED_BODY()
+
+private:
+	// Can only be used with not UFUNCTION, as UFUNCTION requires the value to be explicitly defined.
+	static inline constexpr double DefaultSearchDistance = 20000.0;
+
+public:
+	/** Query the internal octree to call a func on all points within some sphere. Callback takes the point (as const ref) and its squared distance to the center. */
+	template <typename Func = TFunctionRef<void(const FPCGPoint& Point, double SquaredDistance)>>
+	static void ForEachPointInsideSphere(const UPCGPointData* InPointData, const FVector& InCenter, const double InRadius, Func Callback);
+
+	/** Get the closest point to a given position within the search distance, If bInDiscardCenter is true, will reject any points that are exactly at the center. Returns null if not found. */
+	static PCG_API const FPCGPoint* GetClosestPoint(const UPCGPointData* InPointData, const FVector& InCenter, const bool bInDiscardCenter, const double InSearchDistance = DefaultSearchDistance);
+
+	/** Get the nearest point to a given point (that is not itself) within the search distance. Returns null if not found. */
+	static PCG_API const FPCGPoint* GetClosestPointFromOtherPoint(const UPCGPointData* InPointData, const FPCGPoint& InPoint, const double InSearchDistance = DefaultSearchDistance);
+
+	/** Get the farthest point from a given position within the search distance. Returns null if not found. */
+	static PCG_API const FPCGPoint* GetFarthestPoint(const UPCGPointData* InPointData, const FVector& InCenter, const double InSearchDistance = DefaultSearchDistance);
+
+	/** Get the farthest point from a given point (that is not itself) within the search distance. Returns null if not found. */
+	static PCG_API const FPCGPoint* GetFarthestPointFromOtherPoint(const UPCGPointData* InPointData, const FPCGPoint& InPoint, const double InSearchDistance = DefaultSearchDistance);
+
+private:
+	// Blueprint functions are not exposed, use the functions above or the Octree directly
+
+	/** Query the internal octree to return all the points within some bounds. */
+	UFUNCTION(BlueprintCallable, Category = PointData)
+	static TArray<FPCGPoint> GetPointsInsideBounds(const UPCGPointData* InPointData, const FBox& InBounds);
+
+	/** Query the internal octree to return all the points within some sphere. */
+	UFUNCTION(BlueprintCallable, Category = PointData)
+	static TArray<FPCGPoint> GetPointsInsideSphere(const UPCGPointData* InPointData, const FVector& InCenter, const double InRadius);
+
+	/** Get the closest point to a given position within the search distance. If bInDiscardCenter is true, will reject any points that is at the center exactly.*/
+	UFUNCTION(BlueprintCallable, Category = PointData)
+	static void GetClosestPoint(const UPCGPointData* InPointData, const FVector& InCenter, const bool bInDiscardCenter, bool& bOutFound, FPCGPoint& OutPoint, const double InSearchDistance = 20000.0);
+
+	/** Get the nearest point to a given point within the search distance. The point is referenced by its point index in this point data. */
+	UFUNCTION(BlueprintCallable, Category = PointData)
+	static void GetClosestPointFromOtherPoint(const UPCGPointData* InPointData, const int32 InPointIndex, bool& bOutFound, FPCGPoint& OutPoint, const double InSearchDistance = 20000.0);
+
+	/** Get the farthest point from a given position, within the search distance. */
+	UFUNCTION(BlueprintCallable, Category = PointData)
+	static void GetFarthestPoint(const UPCGPointData* InPointData, const FVector& InCenter, bool& bOutFound, FPCGPoint& OutPoint, const double InSearchDistance = 20000.0);
+
+	/** Get the farthest point from a given point within the search distance. The point is referenced by its point index in this point data. */
+	UFUNCTION(BlueprintCallable, Category = PointData)
+	static void GetFarthestPointFromOtherPoint(const UPCGPointData* InPointData, const int32 InPointIndex, bool& bOutFound, FPCGPoint& OutPoint, const double InSearchDistance = 20000.0);
+};
+
+template <typename Func>
+void UPCGOctreeQueries::ForEachPointInsideSphere(const UPCGPointData* InPointData, const FVector& InCenter, const double InRadius, Func Callback)
+{
+	if (!InPointData)
+	{
+		return;
+	}
+
+	const double Extents = UE_DOUBLE_SQRT_2 * InRadius;
+	const double SquaredRadius = InRadius * InRadius;
+	FBoxCenterAndExtent SearchBounds(InCenter, FVector(Extents, Extents, Extents));
+
+	InPointData->GetOctree().FindElementsWithBoundsTest(SearchBounds, [&InCenter, SquaredRadius, Callback = std::move(Callback)](const FPCGPointRef& PointRef)
+	{
+		if (!PointRef.Point)
+		{
+			return;
+		}
+
+		const double SquaredDistance = FVector::DistSquared(InCenter, PointRef.Point->Transform.GetLocation());
+		if (SquaredDistance <= SquaredRadius)
+		{
+			Callback(*PointRef.Point, SquaredDistance);
+		}
+	});
+}
\ No newline at end of file
