[PCG] Build fix for Blur element, using features only available in Main

#rnx

diff --git a/Engine/Plugins/PCG/Source/PCG/Private/Elements/PCGBlurElement.cpp b/Engine/Plugins/PCG/Source/PCG/Private/Elements/PCGBlurElement.cpp
--- a/Engine/Plugins/PCG/Source/PCG/Private/Elements/PCGBlurElement.cpp
+++ b/Engine/Plugins/PCG/Source/PCG/Private/Elements/PCGBlurElement.cpp
@@ -97,7 +97,33 @@ bool FPCGBlurElement::PrepareDataInternal(FPCGContext* InContext) const
 
 		FPCGAttributePropertyOutputSelector OutputSelector = Settings->OutputTarget.CopyAndFixSource(&InputSelector, OutState.InputPointData);
 
-		OutState.OutputAccessor = PCGAttributeAccessorHelpers::CreateAccessorWithAttributeCreation(OutputPointData, OutputSelector, OutState.InputAccessor.Get());
+		// Create attribute if needed
+		if (OutputSelector.IsBasicAttribute())
+		{
+			check(OutputPointData->Metadata);
+			const FName AttributeName = OutputSelector.GetName();
+			if (OutputPointData->Metadata->HasAttribute(AttributeName))
+			{
+				OutputPointData->Metadata->DeleteAttribute(AttributeName);
+			}
+
+			auto CreateAttribute = [&InAccessor = OutState.InputAccessor, AttributeName, OutputPointData](auto&& Dummy) -> FPCGMetadataAttributeBase*
+			{
+				check(InAccessor.IsValid());
+				using AttributeType = std::decay_t<decltype(Dummy)>;
+				AttributeType DefaultValue = PCG::Private::MetadataTraits<AttributeType>::ZeroValue();
+				InAccessor->Get<AttributeType>(DefaultValue, FPCGAttributeAccessorKeysEntries(PCGInvalidEntryKey));
+				return OutputPointData->Metadata->CreateAttribute<AttributeType>(AttributeName, DefaultValue, /*bAllowInterpolation=*/true, /*bOverrideParent=*/false);
+			};
+
+			if (!PCGMetadataAttribute::CallbackWithRightType(OutState.InputAccessor->GetUnderlyingType(), std::move(CreateAttribute)))
+			{
+				PCGLog::LogErrorOnGraph(FText::Format(LOCTEXT("FailedToCreateAttr", "Could not create output attribute '{0}' for input {1}"), OutputSelector.GetDisplayText(), IterationIndex), InContext);
+				return EPCGTimeSliceInitResult::NoOperation;
+			}
+		}
+
+		OutState.OutputAccessor = PCGAttributeAccessorHelpers::CreateAccessor(OutputPointData, OutputSelector);
 		OutState.OutputKeys = PCGAttributeAccessorHelpers::CreateKeys(OutputPointData, OutputSelector);
 
 		if (!OutState.OutputAccessor || !OutState.OutputKeys)
@@ -260,7 +286,7 @@ bool FPCGBlurElement::ExecuteInternal(FPCGContext* InContext) const
 
 			// Copy back
 			const uint8* ReadBuffer = (IterState.CurrentIteration % 2 == 0) ? IterState.WorkingBuffer1.GetData() : IterState.WorkingBuffer2.GetData();
-			IterState.OutputAccessor->SetRange<AttributeType>(MakeArrayView(reinterpret_cast<const AttributeType*>(ReadBuffer), IterState.InputKeys->GetNum()), 0, *IterState.OutputKeys, EPCGAttributeAccessorFlags::AllowBroadcastAndConstructible);
+			IterState.OutputAccessor->SetRange<AttributeType>(MakeArrayView(reinterpret_cast<const AttributeType*>(ReadBuffer), IterState.InputKeys->GetNum()), 0, *IterState.OutputKeys, EPCGAttributeAccessorFlags::AllowBroadcast | EPCGAttributeAccessorFlags::AllowConstructible);
 
 			IterState.WorkingBuffer1.Reset();
 			IterState.WorkingBuffer2.Reset();
