Disabled attribute partition for non-compressed based attributes while we revisit performance characteristics.

#rnx

diff --git a/Engine/Plugins/PCG/Source/PCG/Private/Metadata/PCGMetadataPartitionCommon.cpp b/Engine/Plugins/PCG/Source/PCG/Private/Metadata/PCGMetadataPartitionCommon.cpp
--- a/Engine/Plugins/PCG/Source/PCG/Private/Metadata/PCGMetadataPartitionCommon.cpp
+++ b/Engine/Plugins/PCG/Source/PCG/Private/Metadata/PCGMetadataPartitionCommon.cpp
@@ -236,6 +236,12 @@ namespace PCGMetadataPartitionCommon
 			return {};
 		}
 
+		// Implementation note:
+		// We'll use the attribute partition only for compressed types here (+ needs to be basic attribute only)
+		// because otherwise we can run into issues where keeping track of the breadth of values is not great.
+		bool bUseAttributePartition = false;
+		const FPCGMetadataAttributeBase* Attribute = nullptr;
+
 		if (InSelector.IsBasicAttribute())
 		{
 			const UPCGMetadata* Metadata = InData->ConstMetadata();
@@ -245,7 +251,7 @@ namespace PCGMetadataPartitionCommon
 				return {};
 			}
 
-			const FPCGMetadataAttributeBase* Attribute = Metadata->GetConstAttribute(InSelector.GetName());
+			Attribute = Metadata->GetConstAttribute(InSelector.GetName());
 			if (!Attribute)
 			{
 				if (!bSilenceMissingAttributeErrors)
@@ -256,6 +262,12 @@ namespace PCGMetadataPartitionCommon
 				return {};
 			}
 
+			bUseAttributePartition = Attribute->UsesValueKeys();
+		}
+
+		if (bUseAttributePartition)
+		{
+			check(Attribute);
 			return AttributePartition<PartitionType>(Attribute, *Keys, InOptionalContext);
 		}
 		else
